Make particle properties into more general chaos properties

diff --git a/Engine/Source/Programs/HeadlessChaos/Private/HeadlessChaosTestSolverProxies.cpp b/Engine/Source/Programs/HeadlessChaos/Private/HeadlessChaosTestSolverProxies.cpp
--- a/Engine/Source/Programs/HeadlessChaos/Private/HeadlessChaosTestSolverProxies.cpp
+++ b/Engine/Source/Programs/HeadlessChaos/Private/HeadlessChaosTestSolverProxies.cpp
@@ -148,7 +148,8 @@ namespace ChaosTest {
 
 		// joint Proxy has been created, should be in dirty proxy list to be sent to physics thread.
 		// Note that we expect three dirty proxies, 2 for particles, 1 for joint.
-		EXPECT_EQ(Solver->NumDirtyProxies_External(), 3);
+		EXPECT_EQ(Solver->GetDirtyProxyBucketInfo_External().Num[EPhysicsProxyType::SingleParticleProxy], 2);
+		EXPECT_EQ(Solver->GetDirtyProxyBucketInfo_External().Num[EPhysicsProxyType::JointConstraintType], 1);
 
 		// Although proxy exists, it should not have been added to physics thread only proxy array yet, as doing so from game thread code is wrong.
 		// Confirm game thread did not add it.
diff --git a/Engine/Source/Runtime/Experimental/Chaos/Private/Chaos/ChaosMarshallingManager.cpp b/Engine/Source/Runtime/Experimental/Chaos/Private/Chaos/ChaosMarshallingManager.cpp
--- a/Engine/Source/Runtime/Experimental/Chaos/Private/Chaos/ChaosMarshallingManager.cpp
+++ b/Engine/Source/Runtime/Experimental/Chaos/Private/Chaos/ChaosMarshallingManager.cpp
@@ -227,24 +227,25 @@ TArray<FPushPhysicsData*> FChaosMarshallingManager::StealHistory_Internal(int32
 void FPushPhysicsData::CopySubstepData(const FPushPhysicsData& FirstStepData)
 {
 	const FDirtyPropertiesManager& FirstManager = FirstStepData.DirtyPropertiesManager;
-	DirtyPropertiesManager.SetNumParticles(FirstStepData.DirtyProxiesDataBuffer.NumDirtyProxies());
+	DirtyPropertiesManager.PrepareBuckets(FirstStepData.DirtyProxiesDataBuffer.GetDirtyProxyBucketInfo());
 	FirstStepData.DirtyProxiesDataBuffer.ForEachProxy([this, &FirstManager](int32 FirstDataIdx, const FDirtyProxy& Dirty)
 	{
-		if (Dirty.ParticleData.GetParticleBufferType() == EParticleType::Rigid)
+		//todo: use bucket type directly instead of iterating over each proxy
+		if (Dirty.Proxy->GetType() == EPhysicsProxyType::SingleParticleProxy && Dirty.PropertyData.GetParticleBufferType() == EParticleType::Rigid)
 		{
-			if (const FParticleDynamics* DynamicsData = Dirty.ParticleData.FindDynamics(FirstManager, FirstDataIdx))
+			if (const FParticleDynamics* DynamicsData = Dirty.PropertyData.FindDynamics(FirstManager, FirstDataIdx))
 			{
 				if (DynamicsData->F() != FVec3(0) || DynamicsData->Torque() != FVec3(0))	//don't bother interpolating 0. This is important because the input dirtys rewind data
 				{
 					DirtyProxiesDataBuffer.Add(Dirty.Proxy);
-					FParticleDynamics& SubsteppedDynamics = DirtyPropertiesManager.GetParticlePool<FParticleDynamics, EParticleProperty::Dynamics>().GetElement(Dirty.Proxy->GetDirtyIdx());
+					FParticleDynamics& SubsteppedDynamics = DirtyPropertiesManager.GetChaosPropertyPool<FParticleDynamics, EChaosProperty::Dynamics>().GetElement(Dirty.Proxy->GetDirtyIdx());
 					SubsteppedDynamics = *DynamicsData;
 					//we don't want to sub-step impulses so those are cleared in the sub-step
 					SubsteppedDynamics.SetAngularImpulse(FVec3(0));
 					SubsteppedDynamics.SetLinearImpulse(FVec3(0));
-					FDirtyProxy& NewDirtyProxy = DirtyProxiesDataBuffer.GetDirtyProxyAt(Dirty.Proxy->GetDirtyIdx());
-					NewDirtyProxy.ParticleData.DirtyFlag(EParticleFlags::Dynamics);
-					NewDirtyProxy.ParticleData.SetParticleBufferType(EParticleType::Rigid);
+					FDirtyProxy& NewDirtyProxy = DirtyProxiesDataBuffer.GetDirtyProxyAt(Dirty.Proxy->GetType(), Dirty.Proxy->GetDirtyIdx());
+					NewDirtyProxy.PropertyData.DirtyFlag(EChaosPropertyFlags::Dynamics);
+					NewDirtyProxy.PropertyData.SetParticleBufferType(EParticleType::Rigid);
 				}
 			}
 
diff --git a/Engine/Source/Runtime/Experimental/Chaos/Private/Chaos/ParticleHandle.cpp b/Engine/Source/Runtime/Experimental/Chaos/Private/Chaos/ParticleHandle.cpp
--- a/Engine/Source/Runtime/Experimental/Chaos/Private/Chaos/ParticleHandle.cpp
+++ b/Engine/Source/Runtime/Experimental/Chaos/Private/Chaos/ParticleHandle.cpp
@@ -223,7 +223,7 @@ namespace Chaos
 	template class CHAOS_API TPBDRigidParticle<FReal, 3>;
 
 	template <>
-	void Chaos::TGeometryParticle<FReal, 3>::MarkDirty(const EParticleFlags DirtyBits, bool bInvalidate )
+	void Chaos::TGeometryParticle<FReal, 3>::MarkDirty(const EChaosPropertyFlags DirtyBits, bool bInvalidate )
 	{
 		if (bInvalidate)
 		{
diff --git a/Engine/Source/Runtime/Experimental/Chaos/Private/PBDRigidsSolver.cpp b/Engine/Source/Runtime/Experimental/Chaos/Private/PBDRigidsSolver.cpp
--- a/Engine/Source/Runtime/Experimental/Chaos/Private/PBDRigidsSolver.cpp
+++ b/Engine/Source/Runtime/Experimental/Chaos/Private/PBDRigidsSolver.cpp
@@ -824,13 +824,13 @@ namespace Chaos
 		FDirtySet* DirtyProxiesData = &PushData->DirtyProxiesDataBuffer;
 		FDirtyPropertiesManager* Manager = &PushData->DirtyPropertiesManager;
 
-		Manager->SetNumParticles(DirtyProxiesData->NumDirtyProxies());
+		Manager->PrepareBuckets(DirtyProxiesData->GetDirtyProxyBucketInfo());
 		Manager->SetNumShapes(DirtyProxiesData->NumDirtyShapes());
 		FShapeDirtyData* ShapeDirtyData = DirtyProxiesData->GetShapesDirtyData();
 		auto ProcessProxyGT =[ShapeDirtyData, Manager, DirtyProxiesData](auto Proxy, int32 ParticleDataIdx, FDirtyProxy& DirtyProxy)
 		{
 			auto Particle = Proxy->GetParticle_LowLevel();
-			Particle->SyncRemoteData(*Manager,ParticleDataIdx,DirtyProxy.ParticleData,DirtyProxy.ShapeDataIndices,ShapeDirtyData);
+			Particle->SyncRemoteData(*Manager,ParticleDataIdx,DirtyProxy.PropertyData,DirtyProxy.ShapeDataIndices,ShapeDirtyData);
 			Proxy->ClearAccumulatedData();
 			Proxy->ResetDirtyIdx();
 		};
@@ -934,7 +934,7 @@ namespace Chaos
 			const bool bIsNew = !Proxy->IsInitialized();
 			if(bIsNew)
 			{
-				const auto* NonFrequentData = Dirty.ParticleData.FindNonFrequentData(*Manager,DataIdx);
+				const auto* NonFrequentData = Dirty.PropertyData.FindNonFrequentData(*Manager,DataIdx);
 				const FUniqueIdx* UniqueIdx = NonFrequentData ? &NonFrequentData->UniqueIdx() : nullptr;
 				Proxy->SetHandle(CreateHandleFunc(UniqueIdx));
 
@@ -988,7 +988,7 @@ namespace Chaos
 					auto Proxy = static_cast<FSingleParticlePhysicsProxy*>(Dirty.Proxy);
 					ProcessProxyPT(Proxy, DataIdx, Dirty, [this, &Dirty](const FUniqueIdx* UniqueIdx) -> TGeometryParticleHandle<FReal,3>*
 					{
-						switch (Dirty.ParticleData.GetParticleBufferType())
+						switch (Dirty.PropertyData.GetParticleBufferType())
 						{
 							case EParticleType::Static: return Particles.CreateStaticParticles(1, UniqueIdx)[0];
 							case EParticleType::Kinematic: return Particles.CreateKinematicParticles(1, UniqueIdx)[0];
diff --git a/Engine/Source/Runtime/Experimental/Chaos/Private/PhysicsProxy/SingleParticlePhysicsProxy.cpp b/Engine/Source/Runtime/Experimental/Chaos/Private/PhysicsProxy/SingleParticlePhysicsProxy.cpp
--- a/Engine/Source/Runtime/Experimental/Chaos/Private/PhysicsProxy/SingleParticlePhysicsProxy.cpp
+++ b/Engine/Source/Runtime/Experimental/Chaos/Private/PhysicsProxy/SingleParticlePhysicsProxy.cpp
@@ -47,7 +47,7 @@ void PushToPhysicsStateImp(const Chaos::FDirtyPropertiesManager& Manager, Chaos:
 	constexpr bool bHasDynamicData = ParticleType == EParticleType::Rigid;
 	auto KinematicHandle = bHasKinematicData ? static_cast<Chaos::FKinematicGeometryParticleHandle*>(Handle) : nullptr;
 	auto RigidHandle = bHasDynamicData ? static_cast<Chaos::FPBDRigidParticleHandle*>(Handle) : nullptr;
-	const FParticleDirtyData& ParticleData = Dirty.ParticleData;
+	const FDirtyChaosProperties& ParticleData = Dirty.PropertyData;
 
 	if (bResimInitialized)	//todo: assumes particles are always initialized as enabled. This is not true in future versions of code, so check PushData
 	{
@@ -188,7 +188,7 @@ void FSingleParticlePhysicsProxy::PushToPhysicsState(const Chaos::FDirtyProperti
 	const int32 CurFrame = static_cast<FPBDRigidsSolver*>(Solver)->GetCurrentFrame();
 	const FRewindData* RewindData = static_cast<FPBDRigidsSolver*>(Solver)->GetRewindData();
 	const bool bResimInitialized = RewindData && RewindData->IsResim() && CurFrame == InitializedOnStep;
-	switch(Dirty.ParticleData.GetParticleBufferType())
+	switch(Dirty.PropertyData.GetParticleBufferType())
 	{
 		
 	case EParticleType::Static: PushToPhysicsStateImp<EParticleType::Static>(Manager, Handle, DataIdx, Dirty, ShapesData, Evolution, bResimInitialized, ExternalDt); break;
diff --git a/Engine/Source/Runtime/Experimental/Chaos/Private/RewindData.cpp b/Engine/Source/Runtime/Experimental/Chaos/Private/RewindData.cpp
--- a/Engine/Source/Runtime/Experimental/Chaos/Private/RewindData.cpp
+++ b/Engine/Source/Runtime/Experimental/Chaos/Private/RewindData.cpp
@@ -15,11 +15,11 @@ FVec3 FGeometryParticleStateBase::ZeroVector = FVec3(0);
 
 void FGeometryParticleStateBase::SyncSimWritablePropsFromSim(FDirtyPropData Manager,const TPBDRigidParticleHandle<FReal,3>& Rigid)
 {
-	FParticleDirtyFlags Flags;
-	Flags.MarkDirty(EParticleFlags::XR);
-	Flags.MarkDirty(EParticleFlags::Velocities);
-	Flags.MarkDirty(EParticleFlags::DynamicMisc);
-	FParticleDirtyData Dirty;
+	FDirtyChaosPropertyFlags Flags;
+	Flags.MarkDirty(EChaosPropertyFlags::XR);
+	Flags.MarkDirty(EChaosPropertyFlags::Velocities);
+	Flags.MarkDirty(EChaosPropertyFlags::DynamicMisc);
+	FDirtyChaosProperties Dirty;
 	Dirty.SetFlags(Flags);
 
 #if 0
@@ -47,7 +47,7 @@ void FGeometryParticleStateBase::SyncSimWritablePropsFromSim(FDirtyPropData Mana
 #endif
 }
 
-void FGeometryParticleStateBase::SyncDirtyDynamics(FDirtyPropData& DestManager,const FParticleDirtyData& Dirty,const FConstDirtyPropData& SrcManager)
+void FGeometryParticleStateBase::SyncDirtyDynamics(FDirtyPropData& DestManager,const FDirtyChaosProperties& Dirty,const FConstDirtyPropData& SrcManager)
 {
 #if 0
 	FParticleDirtyData DirtyFlags;
@@ -441,33 +441,33 @@ void FRewindData::PushGTDirtyData(const FDirtyPropertiesManager& SrcManager,cons
 		if(Info.InitializedOnStep < CurFrame)
 		{
 			const FFrameAndPhase PrePushData = FFrameAndPhase{ CurFrame, FFrameAndPhase::PrePushData };
-			auto DirtyPropHelper = [this, &Dirty](auto& Property, const EParticleFlags PropName, const auto& Particle)
+			auto DirtyPropHelper = [this, &Dirty](auto& Property, const EChaosPropertyFlags PropName, const auto& Particle)
 			{
-				if (Dirty.ParticleData.IsDirty(PropName))
+				if (Dirty.PropertyData.IsDirty(PropName))
 				{
 					auto& Data = Property.WriteAccessMonotonic(FFrameAndPhase{ CurFrame, FFrameAndPhase::PrePushData }, PropertiesPool);
 					Data.CopyFrom(Particle);
 				}
 			};
 
-			DirtyPropHelper(Latest.ParticlePositionRotation, EParticleFlags::XR, *PTParticle);
-			DirtyPropHelper(Latest.NonFrequentData, EParticleFlags::NonFrequentData, *PTParticle);
+			DirtyPropHelper(Latest.ParticlePositionRotation, EChaosPropertyFlags::XR, *PTParticle);
+			DirtyPropHelper(Latest.NonFrequentData, EChaosPropertyFlags::NonFrequentData, *PTParticle);
 
 			if (auto Kinematic = PTParticle->CastToKinematicParticle())
 			{
-				DirtyPropHelper(Latest.Velocities, EParticleFlags::Velocities, *Kinematic);
-				DirtyPropHelper(Latest.KinematicTarget, EParticleFlags::KinematicTarget, *Kinematic);
+				DirtyPropHelper(Latest.Velocities, EChaosPropertyFlags::Velocities, *Kinematic);
+				DirtyPropHelper(Latest.KinematicTarget, EChaosPropertyFlags::KinematicTarget, *Kinematic);
 
 				if (auto Rigid = Kinematic->CastToRigidParticle())
 				{
-					DirtyPropHelper(Latest.DynamicsMisc, EParticleFlags::DynamicMisc, *Rigid);
-					DirtyPropHelper(Latest.MassProps, EParticleFlags::MassProps, *Rigid);
+					DirtyPropHelper(Latest.DynamicsMisc, EChaosPropertyFlags::DynamicMisc, *Rigid);
+					DirtyPropHelper(Latest.MassProps, EChaosPropertyFlags::MassProps, *Rigid);
 				}
 			}
 		}
 
 		//Dynamics are not available at head (sim zeroes them out), so we have to record them as PostPushData (since they're applied as part of PushData)
-		if (auto NewData = Dirty.ParticleData.FindDynamics(SrcManager, SrcDataIdx))
+		if (auto NewData = Dirty.PropertyData.FindDynamics(SrcManager, SrcDataIdx))
 		{
 			Latest.Dynamics.WriteAccessMonotonic(FFrameAndPhase{ CurFrame, FFrameAndPhase::PostPushData }, PropertiesPool) = *NewData;
 			Info.DirtyDynamics = CurFrame;	//Need to save the dirty dynamics into the next phase as well (it's possible a callback will stomp the dynamics value, so that's why it's pending)
diff --git a/Engine/Source/Runtime/Experimental/Chaos/Public/Chaos/ChaosMarshallingManager.h b/Engine/Source/Runtime/Experimental/Chaos/Public/Chaos/ChaosMarshallingManager.h
--- a/Engine/Source/Runtime/Experimental/Chaos/Public/Chaos/ChaosMarshallingManager.h
+++ b/Engine/Source/Runtime/Experimental/Chaos/Public/Chaos/ChaosMarshallingManager.h
@@ -19,7 +19,7 @@ class FPullPhysicsData;
 struct FDirtyProxy
 {
 	IPhysicsProxyBase* Proxy;
-	FParticleDirtyData ParticleData;
+	FDirtyChaosProperties PropertyData;
 	TArray<int32> ShapeDataIndices;
 
 	FDirtyProxy(IPhysicsProxyBase* InProxy)
@@ -39,7 +39,7 @@ struct FDirtyProxy
 
 	void Clear(FDirtyPropertiesManager& Manager,int32 DataIdx,FShapeDirtyData* ShapesData)
 	{
-		ParticleData.Clear(Manager,DataIdx);
+		PropertyData.Clear(Manager,DataIdx);
 		for(int32 ShapeDataIdx : ShapeDataIndices)
 		{
 			ShapesData[ShapeDataIdx].Clear(Manager,ShapeDataIdx);
@@ -47,6 +47,11 @@ struct FDirtyProxy
 	}
 };
 
+struct FDirtyProxiesBucket
+{
+	TArray<FDirtyProxy> ProxiesData;
+};
+
 class FDirtySet
 {
 public:
@@ -54,23 +59,37 @@ public:
 	{
 		if(Base->GetDirtyIdx() == INDEX_NONE)
 		{
-			const int32 Idx = ProxiesData.Num();
+			FDirtyProxiesBucket& Bucket = DirtyProxyBuckets[Base->GetType()];
+			++DirtyProxyBucketInfo.Num[Base->GetType()];
+			++DirtyProxyBucketInfo.TotalNum;
+
+			const int32 Idx = Bucket.ProxiesData.Num();
 			Base->SetDirtyIdx(Idx);
-			ProxiesData.Add(Base);
+			Bucket.ProxiesData.Add(Base);
 		}
 	}
 
 	// Batch proxy insertion, does not check DirtyIdx.
+	// Assumes proxies are the same type
 	template< typename TProxiesArray>
 	void AddMultipleUnsafe(TProxiesArray& ProxiesArray)
 	{
-		int32 Idx = ProxiesData.Num();
-		ProxiesData.Append(ProxiesArray);
-
-		for(IPhysicsProxyBase* Proxy : ProxiesArray)
+		if(ProxiesArray.Num())
 		{
-			Proxy->SetDirtyIdx(Idx++);
+			FDirtyProxiesBucket& Bucket = DirtyProxyBuckets[ProxiesArray[0]->GetType()];
+			int32 Idx = Bucket.ProxiesData.Num();
+			Bucket.ProxiesData.Append(ProxiesArray);
+
+			for (IPhysicsProxyBase* Proxy : ProxiesArray)
+			{
+				Proxy->SetDirtyIdx(Idx++);
+				ensure(ProxiesArray[0]->GetType() == Proxy->GetType());
+			}
+
+			DirtyProxyBucketInfo.Num[ProxiesArray[0]->GetType()] += ProxiesArray.Num();
+			DirtyProxyBucketInfo.TotalNum += ProxiesArray.Num();
 		}
+		
 	}
 
 
@@ -79,15 +98,21 @@ public:
 		const int32 Idx = Base->GetDirtyIdx();
 		if(Idx != INDEX_NONE)
 		{
-			if(Idx == ProxiesData.Num() - 1)
+			FDirtyProxiesBucket& Bucket = DirtyProxyBuckets[Base->GetType()];
+			if(Idx == Bucket.ProxiesData.Num() - 1)
 			{
 				//last element so just pop
-				ProxiesData.Pop(/*bAllowShrinking=*/false);
-			} else if(ProxiesData.IsValidIndex(Idx))
+				Bucket.ProxiesData.Pop(/*bAllowShrinking=*/false);
+				--DirtyProxyBucketInfo.Num[Base->GetType()];
+				--DirtyProxyBucketInfo.TotalNum;
+			}
+			else if(Bucket.ProxiesData.IsValidIndex(Idx))
 			{
 				//update other proxy's idx
-				ProxiesData.RemoveAtSwap(Idx);
-				ProxiesData[Idx].SetDirtyIdx(Idx);
+				Bucket.ProxiesData.RemoveAtSwap(Idx);
+				Bucket.ProxiesData[Idx].SetDirtyIdx(Idx);
+				--DirtyProxyBucketInfo.Num[Base->GetType()];
+				--DirtyProxyBucketInfo.TotalNum;
 			}
 
 			Base->ResetDirtyIdx();
@@ -96,58 +121,73 @@ public:
 
 	void Reset()
 	{
-		ProxiesData.Reset();
+		for(FDirtyProxiesBucket& Bucket : DirtyProxyBuckets)
+		{
+			Bucket.ProxiesData.Reset();
+		}
+		
+		DirtyProxyBucketInfo.Reset();
 		ShapesData.Reset();
 	}
 
-	int32 NumDirtyProxies() const { return ProxiesData.Num(); }
+	const FDirtyProxiesBucketInfo& GetDirtyProxyBucketInfo() const { return DirtyProxyBucketInfo; }
 	int32 NumDirtyShapes() const { return ShapesData.Num(); }
 
 	FShapeDirtyData* GetShapesDirtyData(){ return ShapesData.GetData(); }
-	FDirtyProxy& GetDirtyProxyAt(int32 Idx) { return ProxiesData[Idx]; }
+	FDirtyProxy& GetDirtyProxyAt(EPhysicsProxyType ProxyType, int32 Idx) { return DirtyProxyBuckets[ProxyType].ProxiesData[Idx]; }
 
 	template <typename Lambda>
 	void ParallelForEachProxy(const Lambda& Func)
 	{
-		::ParallelFor(ProxiesData.Num(),[this,&Func](int32 Idx)
+		::ParallelFor(DirtyProxyBucketInfo.TotalNum,[this,&Func](int32 Idx)
 		{
-			Func(Idx,ProxiesData[Idx]);
+			int32 BucketIdx, InnerIdx;
+			DirtyProxyBucketInfo.GetBucketIdx(Idx, BucketIdx, InnerIdx);
+			Func(InnerIdx, DirtyProxyBuckets[BucketIdx].ProxiesData[InnerIdx]);
 		});
 	}
 
 	template <typename Lambda>
 	void ParallelForEachProxy(const Lambda& Func) const
 	{
-		::ParallelFor(ProxiesData.Num(),[this,&Func](int32 Idx)
+		::ParallelFor(DirtyProxyBucketInfo.TotalNum,[this,&Func](int32 Idx)
 		{
-			Func(Idx,ProxiesData[Idx]);
+			int32 BucketIdx, InnerIdx;
+			DirtyProxyBucketInfo.GetBucketIdx(Idx, BucketIdx, InnerIdx);
+			Func(InnerIdx, DirtyProxyBuckets[BucketIdx].ProxiesData[InnerIdx]);
 		});
 	}
 
 	template <typename Lambda>
 	void ForEachProxy(const Lambda& Func)
 	{
-		int32 Idx = 0;
-		for(FDirtyProxy& Dirty : ProxiesData)
+		for(int32 BucketIdx = 0; BucketIdx < EPhysicsProxyType::Count; ++BucketIdx)
 		{
-			Func(Idx++,Dirty);
+			int32 Idx = 0;
+			for (FDirtyProxy& Dirty : DirtyProxyBuckets[BucketIdx].ProxiesData)
+			{
+				Func(Idx++, Dirty);
+			}
 		}
 	}
 
 	template <typename Lambda>
 	void ForEachProxy(const Lambda& Func) const
 	{
-		int32 Idx = 0;
-		for(const FDirtyProxy& Dirty : ProxiesData)
+		for (int32 BucketIdx = 0; BucketIdx < EPhysicsProxyType::Count; ++BucketIdx)
 		{
-			Func(Idx++,Dirty);
+			int32 Idx = 0;
+			for (const FDirtyProxy& Dirty : DirtyProxyBuckets[BucketIdx].ProxiesData)
+			{
+				Func(Idx++, Dirty);
+			}
 		}
 	}
 
 	void AddShape(IPhysicsProxyBase* Proxy,int32 ShapeIdx)
 	{
 		Add(Proxy);
-		FDirtyProxy& Dirty = ProxiesData[Proxy->GetDirtyIdx()];
+		FDirtyProxy& Dirty = DirtyProxyBuckets[Proxy->GetType()].ProxiesData[Proxy->GetDirtyIdx()];
 		for(int32 NewShapeIdx = Dirty.ShapeDataIndices.Num(); NewShapeIdx <= ShapeIdx; ++NewShapeIdx)
 		{
 			const int32 ShapeDataIdx = ShapesData.Add(FShapeDirtyData(NewShapeIdx));
@@ -158,7 +198,7 @@ public:
 	void SetNumDirtyShapes(IPhysicsProxyBase* Proxy,int32 NumShapes)
 	{
 		Add(Proxy);
-		FDirtyProxy& Dirty = ProxiesData[Proxy->GetDirtyIdx()];
+		FDirtyProxy& Dirty = DirtyProxyBuckets[Proxy->GetType()].ProxiesData[Proxy->GetDirtyIdx()];
 
 		if(NumShapes < Dirty.ShapeDataIndices.Num())
 		{
@@ -174,7 +214,8 @@ public:
 	}
 
 private:
-	TArray<FDirtyProxy> ProxiesData;
+	FDirtyProxiesBucketInfo DirtyProxyBucketInfo;
+	FDirtyProxiesBucket DirtyProxyBuckets[EPhysicsProxyType::Count];
 	TArray<FShapeDirtyData> ShapesData;
 };
 
diff --git a/Engine/Source/Runtime/Experimental/Chaos/Public/Chaos/Framework/PhysicsProxyBase.h b/Engine/Source/Runtime/Experimental/Chaos/Public/Chaos/Framework/PhysicsProxyBase.h
--- a/Engine/Source/Runtime/Experimental/Chaos/Public/Chaos/Framework/PhysicsProxyBase.h
+++ b/Engine/Source/Runtime/Experimental/Chaos/Public/Chaos/Framework/PhysicsProxyBase.h
@@ -6,7 +6,7 @@
 #include "UObject/GCObject.h"
 #include "Chaos/Core.h"
 
-enum class EPhysicsProxyType
+enum EPhysicsProxyType
 {
 	NoneType = 0,
 	StaticMeshType = 1,
@@ -15,7 +15,8 @@ enum class EPhysicsProxyType
 	SkeletalMeshType = 4,
 	JointConstraintType = 8,	//left gap when removed some types in case these numbers actually matter to someone, should remove
 	SuspensionConstraintType = 9,
-	SingleParticleProxy
+	SingleParticleProxy,
+	Count
 };
 
 namespace Chaos
diff --git a/Engine/Source/Runtime/Experimental/Chaos/Public/Chaos/Framework/PhysicsSolverBase.h b/Engine/Source/Runtime/Experimental/Chaos/Public/Chaos/Framework/PhysicsSolverBase.h
--- a/Engine/Source/Runtime/Experimental/Chaos/Public/Chaos/Framework/PhysicsSolverBase.h
+++ b/Engine/Source/Runtime/Experimental/Chaos/Public/Chaos/Framework/PhysicsSolverBase.h
@@ -127,9 +127,10 @@ namespace Chaos
 		{
 			MarshallingManager.GetProducerData_External()->DirtyProxiesDataBuffer.Remove(ProxyBaseIn);
 		}
-		int32 NumDirtyProxies_External()
+
+		const FDirtyProxiesBucketInfo& GetDirtyProxyBucketInfo_External()
 		{
-			return MarshallingManager.GetProducerData_External()->DirtyProxiesDataBuffer.NumDirtyProxies();
+			return MarshallingManager.GetProducerData_External()->DirtyProxiesDataBuffer.GetDirtyProxyBucketInfo();
 		}
 
 		// Batch dirty proxies without checking DirtyIdx.
diff --git a/Engine/Source/Runtime/Experimental/Chaos/Public/Chaos/KinematicGeometryParticleBuffer.h b/Engine/Source/Runtime/Experimental/Chaos/Public/Chaos/KinematicGeometryParticleBuffer.h
--- a/Engine/Source/Runtime/Experimental/Chaos/Public/Chaos/KinematicGeometryParticleBuffer.h
+++ b/Engine/Source/Runtime/Experimental/Chaos/Public/Chaos/KinematicGeometryParticleBuffer.h
@@ -42,10 +42,10 @@ public:
 	EObjectStateType ObjectState() const;
 
 private:
-	TParticleProperty<FParticleVelocities, EParticleProperty::Velocities> MVelocities;
+	TChaosProperty<FParticleVelocities, EChaosProperty::Velocities> MVelocities;
 
 protected:
-	virtual void SyncRemoteDataImp(FDirtyPropertiesManager& Manager, int32 DataIdx, const FParticleDirtyData& RemoteData) const
+	virtual void SyncRemoteDataImp(FDirtyPropertiesManager& Manager, int32 DataIdx, const FDirtyChaosProperties& RemoteData) const
 	{
 		FGeometryParticleBuffer::SyncRemoteDataImp(Manager, DataIdx, RemoteData);
 		MVelocities.SyncRemote(Manager, DataIdx, RemoteData);
diff --git a/Engine/Source/Runtime/Experimental/Chaos/Public/Chaos/ParticleDirtyFlags.h b/Engine/Source/Runtime/Experimental/Chaos/Public/Chaos/ParticleDirtyFlags.h
--- a/Engine/Source/Runtime/Experimental/Chaos/Public/Chaos/ParticleDirtyFlags.h
+++ b/Engine/Source/Runtime/Experimental/Chaos/Public/Chaos/ParticleDirtyFlags.h
@@ -15,6 +15,7 @@
 #include "UObject/PhysicsObjectVersion.h"
 #include "UObject/UE5MainStreamObjectVersion.h"
 #include "UObject/UE5ReleaseStreamObjectVersion.h"
+#include "Framework/PhysicsProxyBase.h"
 
 #ifndef CHAOS_DEBUG_NAME
 #define CHAOS_DEBUG_NAME 0
@@ -594,56 +595,32 @@ inline FChaosArchive& operator<<(FChaosArchive& Ar,FMaterialData& Data)
 	return Ar;
 }
 
-#define PARTICLE_PROPERTY(PropName, Type) PropName,
-	enum class EParticleProperty : uint32
+#define CHAOS_PROPERTY(PropName, Type, ProxyType) PropName,
+	enum class EChaosProperty : uint32
 	{
 #include "ParticleProperties.inl"
 		NumProperties
 	};
 
-#undef PARTICLE_PROPERTY
+#undef CHAOS_PROPERTY
 
-#define PROPERTY_TYPE(TypeName, Type) TypeName,
-	enum class EPropertyType: uint32
-	{
-#include "PropertiesTypes.inl"
-		NumTypes
-	};
-
-#undef PROPERTY_TYPE
-
-template <typename T>
-struct TPropertyTypeTrait
-{
-};
+#define CHAOS_PROPERTY(PropName, Type, ProxyType) PropName = (uint32)1 << (uint32)EChaosProperty::PropName,
 
-#define PROPERTY_TYPE(TypeName, Type) \
-template <>\
-struct TPropertyTypeTrait<Type>\
-{\
-	static constexpr EPropertyType PoolIdx = EPropertyType::TypeName;\
-};
-
-#include "PropertiesTypes.inl"
-#undef PROPERTY_TYPE
-
-#define PARTICLE_PROPERTY(PropName, Type) PropName = (uint32)1 << (uint32)EParticleProperty::PropName,
-
-	enum class EParticleFlags : uint32
+	enum class EChaosPropertyFlags : uint32
 	{
 		#include "ParticleProperties.inl"
 		DummyFlag
 	};
-#undef PARTICLE_PROPERTY
+#undef CHAOS_PROPERTY
 
-	constexpr EParticleFlags ParticlePropToFlag(EParticleProperty Prop)
+	constexpr EChaosPropertyFlags ChaosPropertyToFlag(EChaosProperty Prop)
 	{
 		switch(Prop)
 		{
-			#define PARTICLE_PROPERTY(PropName, Type) case EParticleProperty::PropName: return EParticleFlags::PropName;
+#define CHAOS_PROPERTY(PropName, Type, ProxyType) case EChaosProperty::PropName: return EChaosPropertyFlags::PropName;
 #include "ParticleProperties.inl"
-#undef PARTICLE_PROPERTY
-		default: return (EParticleFlags)0;
+#undef CHAOS_PROPERTY
+		default: return (EChaosPropertyFlags)0;
 		}
 	}
 
@@ -722,7 +699,7 @@ struct TPropertyTypeTrait<Type>\
 		int32 Bits;
 	};
 
-	using FParticleDirtyFlags = TDirtyFlags<EParticleFlags>;
+	using FDirtyChaosPropertyFlags = TDirtyFlags<EChaosPropertyFlags>;
 	using FShapeDirtyFlags = TDirtyFlags<EShapeFlags>;
 
 	struct FDirtyIdx
@@ -734,7 +711,6 @@ struct TPropertyTypeTrait<Type>\
 	template <typename T>
 	class TDirtyElementPool
 	{
-		static_assert(sizeof(TPropertyTypeTrait<T>::PoolIdx),"Property type must be registered. Is it in PropertiesTypes.inl?");
 	public:
 		const T& GetElement(int32 Idx) const { return Elements[Idx]; }
 		T& GetElement(int32 Idx){ return Elements[Idx]; }
@@ -759,321 +735,51 @@ struct TPropertyTypeTrait<Type>\
 		TArray<T> Elements;
 	};
 
-	//want this for sparse representation
-#if 0
-	template <typename T>
-	class TDirtyElementPool
-	{
-		static_assert(sizeof(TPropertyTypeTrait<T>::PoolIdx),"Property type must be registered. Is it in PropertiesTypes.inl?");
-	public:
-		const T& Read(int32 Idx) const
-		{
-			return Elements[Idx];
-		}
-
-		void Free(int32 Idx)
-		{
-			Elements[Idx].~T();
-			FreeIndices.Add(Idx);
-		}
-
-		T Pop(int32 Idx)
-		{
-			FreeIndices.Add(Idx);
-			T Result;
-			Swap(Result,Elements[Idx]);
-			Elements[Idx].~T();
-			return Result;
-		}
-
-		int32 Write(const T& Element)
-		{
-			const int32 Idx = GetFree();
-			Elements[Idx] = Element;
-			return Idx;
-		}
-
-		void Update(int32 Entry, const T& Element)
-		{
-			Elements[Entry] = Element;
-		}
-
-	private:
-
-		int32 GetFree()
-		{
-			//todo: can we avoid default constructors? maybe if std::is_trivially_copyable
-			if(FreeIndices.Num())
-			{
-				int32 NewIdx = FreeIndices.Pop(/*bAllowShrinking=*/false);
-				Elements[NewIdx] = T();
-				return NewIdx;
-			}
-			else
-			{
-				return Elements.AddDefaulted(1);
-			}
-		}
-
-		TArray<T> Elements;
-		TArray<int32> FreeIndices;
-	};
-
-class FDirtyPropertiesManager;
-
-template <typename T>
-class TRemoteProperty
-{
-public:
-	TRemoteProperty()
-	{
-		Idx.bHasEntry = false;
-	}
-
-	TRemoteProperty(const TRemoteProperty<T>& Rhs) = delete;
-	TRemoteProperty(TRemoteProperty<T>&& Rhs)
-	: Idx(Rhs.Idx)
-	{
-		Rhs.bHasEntry = false;
-	}
-
-	~TRemoteProperty()
-	{
-		ensure(!Idx.bHasEntry);	//leaking, make sure to call Pop
-	}
-
-	const T& Read(const FDirtyPropertiesManager& Manager) const;
-	void Clear(FDirtyPropertiesManager& Manager);
-	void Write(FDirtyPropertiesManager& Manager,const T& Val);
-	
-	bool IsSet() const
-	{
-		return Idx.bHasEntry;
-	}
-private:
-	FDirtyIdx Idx;
-
-	TRemoteProperty<T>& operator=(const TRemoteProperty<T>& Rhs){}
-};
-
-struct FParticlePropertiesData
-{
-	FParticlePropertiesData(FDirtyPropertiesManager* InManager = nullptr)
-		: Manager(InManager)
-	{
-	}
-
-	template <typename T, EParticleProperty PropertyIdx>
-	TRemoteProperty<T>& GetProperty()
-	{
-		switch(PropertyIdx)
-		{
-#define PARTICLE_PROPERTY(PropName, Type) case EParticleProperty::PropName: return (TRemoteProperty<T>&) PropName;
-#include "ParticleProperties.inl"
-#undef PARTICLE_PROPERTY
-		default: check(false);
-		}
-
-		static TRemoteProperty<T> Error;
-		return Error;
-	}
-
-	void Clear()
-	{
-		if(Manager)
-		{
-#define PARTICLE_PROPERTY(PropName, Type) PropName.Clear(*Manager);
-#include "ParticleProperties.inl"
-#undef PARTICLE_PROPERTY
-		}
-	}
-
-	inline void FreeToManager();
-	
-	~FParticlePropertiesData()
-	{
-		Clear();
-	}
-
-	FDirtyPropertiesManager* GetManager(){ return Manager; }
-	const FDirtyPropertiesManager* GetManager() const { return Manager; }
-
-#define PARTICLE_PROPERTY(PropName, Type)\
-Type const & Get##PropName() const { return PropName.Read(*Manager); }\
-bool Has##PropName() const { return PropName.IsSet(); }\
-Type const * Find##PropName() const { return Has##PropName() ? &Get##PropName() : nullptr; }
-
-#include "ParticleProperties.inl"
-#undef PARTICLE_PROPERTY
-
-private:
-#define PARTICLE_PROPERTY(PropName, Type) TRemoteProperty<Type> PropName;
-#include "ParticleProperties.inl"
-#undef PARTICLE_PROPERTY
-
-	FDirtyPropertiesManager* Manager;
-
-};
-
-struct FShapePropertiesData
-{
-	FShapePropertiesData(FDirtyPropertiesManager* InManager)
-		: Manager(InManager)
-	{
-
-	}
-	template <typename T,EShapeProperty PropertyIdx>
-	TRemoteProperty<T>& GetProperty()
-	{
-		switch(PropertyIdx)
-		{
-#define SHAPE_PROPERTY(PropName, Type) case EShapeProperty::PropName: return (TRemoteProperty<T>&) PropName;
-#include "ShapeProperties.inl"
-#undef SHAPE_PROPERTY
-		default: check(false);
-		}
-
-		static TRemoteProperty<T> Error;
-		return Error;
-	}
-
-	void Clear()
-	{
-		if(Manager)
-		{
-#define SHAPE_PROPERTY(PropName, Type) PropName.Clear(*Manager);
-#include "ShapeProperties.inl"
-#undef SHAPE_PROPERTY
-		}
-	}
-
-	~FShapePropertiesData()
-	{
-		Clear();
-	}
-
-	inline void FreeToManager();
-
-	FDirtyPropertiesManager* GetManager(){ return Manager; }
-	const FDirtyPropertiesManager* GetManager() const { return Manager; }
-
-#define SHAPE_PROPERTY(PropName, Type)\
-Type const & Get##PropName() const { return PropName.Read(*Manager); }\
-bool Has##PropName() const { return PropName.IsSet(); }\
-Type const * Find##PropName() const { return Has##PropName() ? &Get##PropName() : nullptr; }
-
-#include "ShapeProperties.inl"
-#undef SHAPE_PROPERTY
 
-private:
-#define SHAPE_PROPERTY(PropName, Type) TRemoteProperty<Type> PropName;
-#include "ShapeProperties.inl"
-#undef SHAPE_PROPERTY
-
-	FDirtyPropertiesManager* Manager;
-};
-
-class FPerShapeData;
-
-class CHAOS_API FShapeRemoteDataContainer
+/** Helper struct to let us know how many proxies are dirty per type,
+  * as well as how to go from a contiguous index into a per bucket index */
+struct FDirtyProxiesBucketInfo
 {
-public:
-	FShapeRemoteDataContainer(FDirtyPropertiesManager* InManager)
-	: Manager(InManager)
-	{
-	}
+	int32 Num[EPhysicsProxyType::Count] = {};
+	int32 TotalNum = 0;
 
-	~FShapeRemoteDataContainer()
+	void Reset()
 	{
-		Clear();
+		for (int32 Idx = 0; Idx < EPhysicsProxyType::Count; ++Idx) { Num[Idx] = 0; }
+		TotalNum = 0;
 	}
 
-	void SyncShapes(TArray<TUniquePtr<FPerShapeData>, TInlineAllocator<1>>& Shapes);
-	
-	void DetachRemoteData(TArray<TUniquePtr<FPerShapeData>,TInlineAllocator<1>>& Shapes);
-
-	inline void FreeToManager();
-
-	inline FShapePropertiesData* NewRemoteShapeProperties();
-	
-	void Clear()
+	void GetBucketIdx(int32 Idx, int32& OutBucketIdx, int32& InnerIdx) const
 	{
-		//todo: avoid iterating all remote data regardless of if dirty or not
-		for(FShapePropertiesData* RemoteData : RemoteDatas)
+		int32 Remaining = Idx;
+		for (int32 BucketIdx = 0; BucketIdx < EPhysicsProxyType::Count; ++BucketIdx)
 		{
-			if(RemoteData)
+			if (Remaining < Num[BucketIdx])
 			{
-				RemoteData->FreeToManager();
+				InnerIdx = Remaining;
+				OutBucketIdx = BucketIdx;
+				return;
+			}
+			else
+			{
+				Remaining -= Num[BucketIdx];
 			}
 		}
 
-		RemoteDatas.Reset();
-	}
-
-	const auto& GetRemoteDatas() const
-	{
-		return RemoteDatas;
-	}
-
-	auto& GetRemoteDatas()
-	{
-		return RemoteDatas;
+		check(false);	//couldn't find bucket for the given index
 	}
-	
-private:
-
-	FDirtyPropertiesManager* Manager;
-	TArray<FShapePropertiesData*,TInlineAllocator<4>> RemoteDatas;
 };
-#endif
+
 
 class FDirtyPropertiesManager
 {
 public:
 
-#if 0
-	FParticlePropertiesData* NewRemoteParticleProperties()
-	{
-		return RemoteParticlePool.NewEntry(this);
-	}
-
-	void FreeRemoteParticleProperties(FParticlePropertiesData* Entry)
+	void PrepareBuckets(const FDirtyProxiesBucketInfo& DirtyProxiesBucketInfo)
 	{
-		RemoteParticlePool.FreeEntry(Entry);
-	}
-
-	FShapePropertiesData* NewRemoteShapeProperties()
-	{
-		return RemoteShapePool.NewEntry(this);
-	}
-
-	void FreeRemoteShapeProperties(FShapePropertiesData* Entry)
-	{
-		RemoteShapePool.FreeEntry(Entry);
-	}
-
-	FShapeRemoteDataContainer* NewRemoteShapeContainer()
-	{
-		return RemoteShapeContainerPool.NewEntry(this);
-	}
-
-	void FreeRemoteShapeContainer(FShapeRemoteDataContainer* Entry)
-	{
-		RemoteShapeContainerPool.FreeEntry(Entry);
-	}
-#endif
-
-	void SetNumParticles(int32 NumParticles)
-	{
-#define PARTICLE_PROPERTY(PropName, Type) PropName##Pool.SetNum(NumParticles);
+#define CHAOS_PROPERTY(PropName, Type, ProxyType) PropName##Pool.SetNum(DirtyProxiesBucketInfo.Num[ProxyType]);
 #include "ParticleProperties.inl"
-#undef PARTICLE_PROPERTY
-	}
-
-	int32 GetNumParticles() const
-	{
-		//assume this property exists, if it gets renamed just pick any property
-		return XRPool.Num();
+#undef CHAOS_PROPERTY
 	}
 
 	void SetNumShapes(int32 NumShapes)
@@ -1083,14 +789,14 @@ public:
 #undef SHAPE_PROPERTY
 	}
 
-	template <typename T, EParticleProperty PropName>
-	TDirtyElementPool<T>& GetParticlePool()
+	template <typename T, EChaosProperty PropName>
+	TDirtyElementPool<T>& GetChaosPropertyPool()
 	{
 		switch(PropName)
 		{
-#define PARTICLE_PROPERTY(PropName, Type) case EParticleProperty::PropName: return (TDirtyElementPool<T>&)PropName##Pool;
+#define CHAOS_PROPERTY(PropName, Type, ProxyType) case EChaosProperty::PropName: return (TDirtyElementPool<T>&)PropName##Pool;
 #include "ParticleProperties.inl"
-#undef PARTICLE_PROPERTY
+#undef CHAOS_PROPERTY
 		default: check(false);
 		}
 
@@ -1098,14 +804,14 @@ public:
 		return ErrorPool;
 	}
 
-	template <typename T,EParticleProperty PropName>
-	const TDirtyElementPool<T>& GetParticlePool() const
+	template <typename T,EChaosProperty PropName>
+	const TDirtyElementPool<T>& GetChaosPropertyPool() const
 	{
 		switch(PropName)
 		{
-#define PARTICLE_PROPERTY(PropName, Type) case EParticleProperty::PropName: return (TDirtyElementPool<T>&)PropName##Pool;
+#define CHAOS_PROPERTY(PropName, Type, ProxyType) case EChaosProperty::PropName: return (TDirtyElementPool<T>&)PropName##Pool;
 #include "ParticleProperties.inl"
-#undef PARTICLE_PROPERTY
+#undef CHAOS_PROPERTY
 		default: check(false);
 		}
 
@@ -1145,58 +851,17 @@ public:
 
 private:
 
-#define PARTICLE_PROPERTY(PropName, Type) TDirtyElementPool<Type> PropName##Pool;
+#define CHAOS_PROPERTY(PropName, Type, ProxyType) TDirtyElementPool<Type> PropName##Pool;
 #include "ParticleProperties.inl"
-#undef PARTICLE_PROPERTY
+#undef CHAOS_PROPERTY
 
 #define SHAPE_PROPERTY(PropName, Type) TDirtyElementPool<Type> PropName##ShapePool;
 #include "ShapeProperties.inl"
 #undef SHAPE_PROPERTY
 
-#if 0
-	template <typename T>
-	class TRemotePropertiesPool
-	{
-	public:
-		T* NewEntry(FDirtyPropertiesManager* Manager)
-		{
-			if(Pool.Num())
-			{
-				return Pool.Pop(/*bAllowShrinking=*/false);
-			} else
-			{
-				return new T(Manager);
-			}
-		}
-
-		void FreeEntry(T* Entry)
-		{
-			Entry->Clear();
-			Pool.Add(Entry);
-		}
-
-		~TRemotePropertiesPool()
-		{
-			for(T* Entry : Pool)
-			{
-				delete Entry;
-			}
-		}
-
-		int32 NumInPool() const { return Pool.Num(); }
-
-	private:
-		TArray<T*> Pool;
-	};
-
-	TRemotePropertiesPool<FParticlePropertiesData> RemoteParticlePool;
-	TRemotePropertiesPool<FShapePropertiesData> RemoteShapePool;
-	TRemotePropertiesPool<FShapeRemoteDataContainer> RemoteShapeContainerPool;
-#endif
-
 };
 
-class FParticleDirtyData
+class FDirtyChaosProperties
 {
 public:
 	
@@ -1205,73 +870,75 @@ public:
 		ParticleBufferType = Type;
 	}
 
+	//NOTE: this is only valid if the proxy is a particle type and SetParticleBufferType was used
+	//TODO: remove this from API
 	EParticleType GetParticleBufferType() const
 	{
 		return ParticleBufferType;
 	}
 
-	void SetFlags(FParticleDirtyFlags InFlags)
+	void SetFlags(FDirtyChaosPropertyFlags InFlags)
 	{
 		Flags = InFlags;
 	}
 
-	FParticleDirtyFlags GetFlags() const
+	FDirtyChaosPropertyFlags GetFlags() const
 	{
 		return Flags;
 	}
 
-	void DirtyFlag(EParticleFlags Flag)
+	void DirtyFlag(EChaosPropertyFlags Flag)
 	{
 		Flags.MarkDirty(Flag);
 	}
 
-	template <typename T, EParticleProperty PropName>
+	template <typename T, EChaosProperty PropName>
 	void SyncRemote(FDirtyPropertiesManager& Manager, int32 Idx, const T& Val) const
 	{
-		if(Flags.IsDirty(ParticlePropToFlag(PropName)))
+		if(Flags.IsDirty(ChaosPropertyToFlag(PropName)))
 		{
-			Manager.GetParticlePool<T,PropName>().GetElement(Idx) = Val;
+			Manager.GetChaosPropertyPool<T,PropName>().GetElement(Idx) = Val;
 		}
 	}
 
 	void Clear(FDirtyPropertiesManager& Manager, int32 Idx)
 	{
-#define PARTICLE_PROPERTY(PropName, Type) ClearHelper<Type, EParticleProperty::PropName>(Manager, Idx);
+#define CHAOS_PROPERTY(PropName, Type, ProxyType) ClearHelper<Type, EChaosProperty::PropName>(Manager, Idx);
 #include "ParticleProperties.inl"
-#undef PARTICLE_PROPERTY
+#undef CHAOS_PROPERTY
 		Flags.Clear();
 	}
 
-	bool IsDirty(EParticleFlags InBits) const
+	bool IsDirty(EChaosPropertyFlags InBits) const
 	{
 		return Flags.IsDirty(InBits);
 	}
 
-#define PARTICLE_PROPERTY(PropName, Type)\
-Type const & Get##PropName(const FDirtyPropertiesManager& Manager, int32 Idx) const { return ReadImp<Type, EParticleProperty::PropName>(Manager, Idx); }\
-bool Has##PropName() const { return Flags.IsDirty(ParticlePropToFlag(EParticleProperty::PropName)); }\
+#define CHAOS_PROPERTY(PropName, Type, ProxyType)\
+Type const & Get##PropName(const FDirtyPropertiesManager& Manager, int32 Idx) const { return ReadImp<Type, EChaosProperty::PropName>(Manager, Idx); }\
+bool Has##PropName() const { return Flags.IsDirty(ChaosPropertyToFlag(EChaosProperty::PropName)); }\
 Type const * Find##PropName(const FDirtyPropertiesManager& Manager, int32 Idx) const { return Has##PropName() ? &Get##PropName(Manager, Idx) : nullptr; }
 
 #include "ParticleProperties.inl"
-#undef PARTICLE_PROPERTY
+#undef CHAOS_PROPERTY
 
 private:
-	FParticleDirtyFlags Flags;
+	FDirtyChaosPropertyFlags Flags;
 	EParticleType ParticleBufferType;
 
-	template <typename T,EParticleProperty PropName>
+	template <typename T,EChaosProperty PropName>
 	const T& ReadImp(const FDirtyPropertiesManager& Manager, int32 Idx) const
 	{
-		ensure(Flags.IsDirty(ParticlePropToFlag(PropName)));
-		return Manager.GetParticlePool<T,PropName>().GetElement(Idx);
+		ensure(Flags.IsDirty(ChaosPropertyToFlag(PropName)));
+		return Manager.GetChaosPropertyPool<T,PropName>().GetElement(Idx);
 	}
 
-	template <typename T, EParticleProperty PropName>
+	template <typename T, EChaosProperty PropName>
 	void ClearHelper(FDirtyPropertiesManager& Manager, int32 Idx)
 	{
-		if(Flags.IsDirty(ParticlePropToFlag(PropName)))
+		if(Flags.IsDirty(ChaosPropertyToFlag(PropName)))
 		{
-			Manager.GetParticlePool<T, PropName>().Reset(Idx);
+			Manager.GetChaosPropertyPool<T, PropName>().Reset(Idx);
 		}
 	}
 };
@@ -1322,7 +989,7 @@ bool Has##PropName() const { return Flags.IsDirty(ShapePropToFlag(EShapeProperty
 Type const * Find##PropName(const FDirtyPropertiesManager& Manager, int32 Idx) const { return Has##PropName() ? &Get##PropName(Manager, Idx) : nullptr; }
 
 #include "ShapeProperties.inl"
-#undef PARTICLE_PROPERTY
+#undef CHAOS_PROPERTY
 
 private:
 	int32 ShapeIdx;
@@ -1345,67 +1012,6 @@ private:
 	}
 };
 
-#if 0
-void FParticlePropertiesData::FreeToManager()
-{
-	if(Manager)
-	{
-		Manager->FreeRemoteParticleProperties(this);
-	}
-}
-
-void FShapePropertiesData::FreeToManager()
-{
-	if(Manager)
-	{
-		Manager->FreeRemoteShapeProperties(this);
-	}
-}
-
-void FShapeRemoteDataContainer::FreeToManager()
-{
-	if(Manager)
-	{
-		Manager->FreeRemoteShapeContainer(this);
-	}
-}
-
-FShapePropertiesData* FShapeRemoteDataContainer::NewRemoteShapeProperties()
-{
-	return Manager ? Manager->NewRemoteShapeProperties() : nullptr;
-}
-
-template <typename T>
-const T& TRemoteProperty<T>::Read(const FDirtyPropertiesManager& Manager) const
-{
-	ensure(Idx.bHasEntry);
-	return Manager.GetPool<T>().Read(Idx.Entry);
-}
-
-template <typename T>
-void TRemoteProperty<T>::Clear(FDirtyPropertiesManager& Manager)
-{
-	if(Idx.bHasEntry)
-	{
-		Idx.bHasEntry = false;
-		Manager.GetPool<T>().Pop(Idx.Entry);
-	}
-}
-
-template <typename T>
-void TRemoteProperty<T>::Write(FDirtyPropertiesManager& Manager,const T& Val)
-{
-	if(Idx.bHasEntry)
-	{
-		Manager.GetPool<T>().Update(Idx.Entry,Val);
-	} else
-	{
-		Idx.Entry = Manager.GetPool<T>().Write(Val);
-		Idx.bHasEntry = true;
-	}
-}
-#endif
-
 template <typename T>
 class TPropertyPool;
 
@@ -1414,7 +1020,6 @@ using FPropertyIdx = int32;
 template <typename T>
 class TPropertyPool
 {
-	static_assert(sizeof(TPropertyTypeTrait<T>::PoolIdx), "Property type must be registered. Is it in PropertiesTypes.inl?");
 public:
 
 	T& AddElement(FPropertyIdx& OutIdx)
@@ -1463,14 +1068,14 @@ private:
 class FDirtyPropertiesPool
 {
 public:
-	template <typename T, EParticleProperty PropName>
+	template <typename T, EChaosProperty PropName>
 	TPropertyPool<T>& GetPool()
 	{
 		switch (PropName)
 		{
-#define PARTICLE_PROPERTY(PropName, Type) case EParticleProperty::PropName: return (TPropertyPool<T>&)PropName##Pool;
+#define CHAOS_PROPERTY(PropName, Type, ProxyType) case EChaosProperty::PropName: return (TPropertyPool<T>&)PropName##Pool;
 #include "ParticleProperties.inl"
-#undef PARTICLE_PROPERTY
+#undef CHAOS_PROPERTY
 		default: check(false);
 		}
 
@@ -1478,14 +1083,14 @@ public:
 		return ErrorPool;
 	}
 
-	template <typename T, EParticleProperty PropName>
+	template <typename T, EChaosProperty PropName>
 	const TPropertyPool<T>& GetPool() const
 	{
 		switch (PropName)
 		{
-#define PARTICLE_PROPERTY(PropName, Type) case EParticleProperty::PropName: return (TPropertyPool<T>&)PropName##Pool;
+#define CHAOS_PROPERTY(PropName, Type, ProxyType) case EChaosProperty::PropName: return (TPropertyPool<T>&)PropName##Pool;
 #include "ParticleProperties.inl"
-#undef PARTICLE_PROPERTY
+#undef CHAOS_PROPERTY
 		default: check(false);
 		}
 
@@ -1495,9 +1100,8 @@ public:
 
 private:
 
-#define PARTICLE_PROPERTY(PropName, Type) TPropertyPool<Type> PropName##Pool;
+#define CHAOS_PROPERTY(PropName, Type, ProxyType) TPropertyPool<Type> PropName##Pool;
 #include "ParticleProperties.inl"
-#undef PARTICLE_PROPERTY
+#undef CHAOS_PROPERTY
 };
-
 }
diff --git a/Engine/Source/Runtime/Experimental/Chaos/Public/Chaos/ParticleHandle.h b/Engine/Source/Runtime/Experimental/Chaos/Public/Chaos/ParticleHandle.h
--- a/Engine/Source/Runtime/Experimental/Chaos/Public/Chaos/ParticleHandle.h
+++ b/Engine/Source/Runtime/Experimental/Chaos/Public/Chaos/ParticleHandle.h
@@ -1872,12 +1872,12 @@ public:
 		return MDirtyFlags.IsClean();
 	}
 
-	bool IsDirty(const EParticleFlags CheckBits) const
+	bool IsDirty(const EChaosPropertyFlags CheckBits) const
 	{
 		return MDirtyFlags.IsDirty(CheckBits);
 	}
 
-	const FParticleDirtyFlags& DirtyFlags() const
+	const FDirtyChaosPropertyFlags& DirtyFlags() const
 	{
 		return MDirtyFlags;
 	}
@@ -1908,7 +1908,7 @@ public:
 		return nullptr;
 	}
 
-	void SyncRemoteData(FDirtyPropertiesManager& Manager, int32 DataIdx, FParticleDirtyData& RemoteData, const TArray<int32>& ShapeDataIndices, FShapeDirtyData* ShapesRemoteData) const
+	void SyncRemoteData(FDirtyPropertiesManager& Manager, int32 DataIdx, FDirtyChaosProperties& RemoteData, const TArray<int32>& ShapeDataIndices, FShapeDirtyData* ShapesRemoteData) const
 	{
 		RemoteData.SetParticleBufferType(Type);
 		RemoteData.SetFlags(MDirtyFlags);
@@ -1992,8 +1992,8 @@ protected:
 	}
 private:
 
-	TParticleProperty<FParticlePositionRotation, EParticleProperty::XR> MXR;
-	TParticleProperty<FParticleNonFrequentData,EParticleProperty::NonFrequentData> MNonFrequentData;
+	TChaosProperty<FParticlePositionRotation, EChaosProperty::XR> MXR;
+	TChaosProperty<FParticleNonFrequentData,EChaosProperty::NonFrequentData> MNonFrequentData;
 	void* MUserData;
 
 	FShapesArray MShapesArray;
@@ -2017,9 +2017,9 @@ private:
 protected:
 
 	EParticleType Type;
-	FParticleDirtyFlags MDirtyFlags;
+	FDirtyChaosPropertyFlags MDirtyFlags;
 
-	void MarkDirty(const EParticleFlags DirtyBits, bool bInvalidate = true);
+	void MarkDirty(const EChaosPropertyFlags DirtyBits, bool bInvalidate = true);
 
 	void UpdateShapesArray()
 	{
@@ -2027,7 +2027,7 @@ protected:
 		MapImplicitShapes();
 	}
 
-	virtual void SyncRemoteDataImp(FDirtyPropertiesManager& Manager, int32 DataIdx, const FParticleDirtyData& RemoteData) const
+	virtual void SyncRemoteDataImp(FDirtyPropertiesManager& Manager, int32 DataIdx, const FDirtyChaosProperties& RemoteData) const
 	{
 		MXR.SyncRemote(Manager, DataIdx, RemoteData);
 		MNonFrequentData.SyncRemote(Manager, DataIdx, RemoteData);
@@ -2121,11 +2121,11 @@ public:
 
 
 private:
-	TParticleProperty<FParticleVelocities, EParticleProperty::Velocities> MVelocities;
-	TParticleProperty<FKinematicTarget, EParticleProperty::KinematicTarget> MKinematicTarget;
+	TChaosProperty<FParticleVelocities, EChaosProperty::Velocities> MVelocities;
+	TChaosProperty<FKinematicTarget, EChaosProperty::KinematicTarget> MKinematicTarget;
 
 protected:
-	virtual void SyncRemoteDataImp(FDirtyPropertiesManager& Manager, int32 DataIdx, const FParticleDirtyData& RemoteData) const
+	virtual void SyncRemoteDataImp(FDirtyPropertiesManager& Manager, int32 DataIdx, const FDirtyChaosProperties& RemoteData) const
 	{
 		Base::SyncRemoteDataImp(Manager, DataIdx, RemoteData);
 		MVelocities.SyncRemote(Manager, DataIdx, RemoteData);
@@ -2397,7 +2397,7 @@ public:
 			// have been buffered. If another force is added after the object is put to sleep, the old forces
 			// will remain and the new ones will accumulate and re-dirty the dynamic properties which will
 			// wake the body.
-			MDirtyFlags.MarkClean(ParticlePropToFlag(EParticleProperty::Dynamics));
+			MDirtyFlags.MarkClean(ChaosPropertyToFlag(EChaosProperty::Dynamics));
 		}
 
 		MMiscData.Modify(bInvalidate,MDirtyFlags,Proxy,[&InState](auto& Data){ Data.SetObjectState(InState);});
@@ -2433,9 +2433,9 @@ public:
 	}
 
 private:
-	TParticleProperty<FParticleMassProps,EParticleProperty::MassProps> MMassProps;
-	TParticleProperty<FParticleDynamics, EParticleProperty::Dynamics> MDynamics;
-	TParticleProperty<FParticleDynamicMisc,EParticleProperty::DynamicMisc> MMiscData;
+	TChaosProperty<FParticleMassProps,EChaosProperty::MassProps> MMassProps;
+	TChaosProperty<FParticleDynamics, EChaosProperty::Dynamics> MDynamics;
+	TChaosProperty<FParticleDynamicMisc,EChaosProperty::DynamicMisc> MMiscData;
 
 	int32 MIsland;
 	bool MToBeRemovedOnFracture;
@@ -2443,7 +2443,7 @@ private:
 	EWakeEventEntry MWakeEvent;
 
 protected:
-	virtual void SyncRemoteDataImp(FDirtyPropertiesManager& Manager, int32 DataIdx, const FParticleDirtyData& RemoteData) const
+	virtual void SyncRemoteDataImp(FDirtyPropertiesManager& Manager, int32 DataIdx, const FDirtyChaosProperties& RemoteData) const
 	{
 		Base::SyncRemoteDataImp(Manager,DataIdx,RemoteData);
 		MMassProps.SyncRemote(Manager,DataIdx,RemoteData);
@@ -2576,7 +2576,7 @@ TGeometryParticle<T, d>* TGeometryParticle<T, d>::SerializationFactory(FChaosArc
 }
 
 template <>
-CHAOS_API void Chaos::TGeometryParticle<FReal, 3>::MarkDirty(const EParticleFlags DirtyBits, bool bInvalidate);
+CHAOS_API void Chaos::TGeometryParticle<FReal, 3>::MarkDirty(const EChaosPropertyFlags DirtyBits, bool bInvalidate);
 
 FORCEINLINE_DEBUGGABLE FAccelerationStructureHandle::FAccelerationStructureHandle(FGeometryParticleHandle* InHandle)
 	: ExternalGeometryParticle(InHandle->GTGeometryParticle())
diff --git a/Engine/Source/Runtime/Experimental/Chaos/Public/Chaos/ParticleProperties.inl b/Engine/Source/Runtime/Experimental/Chaos/Public/Chaos/ParticleProperties.inl
--- a/Engine/Source/Runtime/Experimental/Chaos/Public/Chaos/ParticleProperties.inl
+++ b/Engine/Source/Runtime/Experimental/Chaos/Public/Chaos/ParticleProperties.inl
@@ -1,16 +1,18 @@
 // Copyright Epic Games, Inc. All Rights Reserved.
 
-#undef PARTICLE_PROPERTY_CHECKED
+//Make sure to include PhysicsProxyBase.h before using this inl file
+
+#undef CHAOS_PROPERTY_CHECKED
 #if CHAOS_CHECKED
-#define PARTICLE_PROPERTY_CHECKED(x, Type) PARTICLE_PROPERTY(x, Type)
+#define CHAOS_PROPERTY_CHECKED(x, Type, ProxyType) CHAOS_PROPERTY(x, Type, ProxyType)
 #else
-#define PARTICLE_PROPERTY_CHECKED(x, Type)
+#define CHAOS_PROPERTY_CHECKED(x, Type, ProxyType)
 #endif
 
-PARTICLE_PROPERTY(XR, FParticlePositionRotation)
-PARTICLE_PROPERTY(Velocities, FParticleVelocities)
-PARTICLE_PROPERTY(Dynamics, FParticleDynamics)
-PARTICLE_PROPERTY(DynamicMisc, FParticleDynamicMisc)
-PARTICLE_PROPERTY(NonFrequentData, FParticleNonFrequentData)
-PARTICLE_PROPERTY(MassProps, FParticleMassProps)
-PARTICLE_PROPERTY(KinematicTarget,FKinematicTarget)
\ No newline at end of file
+CHAOS_PROPERTY(XR, FParticlePositionRotation, EPhysicsProxyType::SingleParticleProxy)
+CHAOS_PROPERTY(Velocities, FParticleVelocities, EPhysicsProxyType::SingleParticleProxy)
+CHAOS_PROPERTY(Dynamics, FParticleDynamics, EPhysicsProxyType::SingleParticleProxy)
+CHAOS_PROPERTY(DynamicMisc, FParticleDynamicMisc, EPhysicsProxyType::SingleParticleProxy)
+CHAOS_PROPERTY(NonFrequentData, FParticleNonFrequentData, EPhysicsProxyType::SingleParticleProxy)
+CHAOS_PROPERTY(MassProps, FParticleMassProps, EPhysicsProxyType::SingleParticleProxy)
+CHAOS_PROPERTY(KinematicTarget,FKinematicTarget, EPhysicsProxyType::SingleParticleProxy)
\ No newline at end of file
diff --git a/Engine/Source/Runtime/Experimental/Chaos/Public/Chaos/Properties.h b/Engine/Source/Runtime/Experimental/Chaos/Public/Chaos/Properties.h
--- a/Engine/Source/Runtime/Experimental/Chaos/Public/Chaos/Properties.h
+++ b/Engine/Source/Runtime/Experimental/Chaos/Public/Chaos/Properties.h
@@ -9,39 +9,39 @@
 namespace Chaos
 {
 
-template <typename T, EParticleProperty PropName>
-class TParticleProperty
+template <typename T, EChaosProperty PropName>
+class TChaosProperty
 {
 public:
-	TParticleProperty() = default;
-	TParticleProperty(const T& Val)
+	TChaosProperty() = default;
+	TChaosProperty(const T& Val)
 		: Property(Val)
 	{
 	}
 
 	//we don't support this because it could lead to bugs with values not being properly written to remote
-	TParticleProperty(const TParticleProperty<T,PropName>& Rhs) = delete;
+	TChaosProperty(const TChaosProperty<T,PropName>& Rhs) = delete;
 
-	bool IsDirty(const FParticleDirtyFlags& Flags) const
+	bool IsDirty(const FDirtyChaosPropertyFlags& Flags) const
 	{
 		return Flags.IsDirty(PropertyFlag);
 	}
 
 	const T& Read() const { return Property; }
-	void Write(const T& Val, bool bInvalidate, FParticleDirtyFlags& Dirty, IPhysicsProxyBase* Proxy)
+	void Write(const T& Val, bool bInvalidate, FDirtyChaosPropertyFlags& Dirty, IPhysicsProxyBase* Proxy)
 	{
 		Property = Val;
 		MarkDirty(bInvalidate, Dirty, Proxy);
 	}
 
 	template <typename Lambda>
-	void Modify(bool bInvalidate, FParticleDirtyFlags& Dirty,IPhysicsProxyBase* Proxy,const Lambda& LambdaFunc)
+	void Modify(bool bInvalidate, FDirtyChaosPropertyFlags& Dirty,IPhysicsProxyBase* Proxy,const Lambda& LambdaFunc)
 	{
 		LambdaFunc(Property);
 		MarkDirty(bInvalidate, Dirty, Proxy);
 	}
 
-	void Clear(FParticleDirtyFlags& Dirty, IPhysicsProxyBase* Proxy)
+	void Clear(FDirtyChaosPropertyFlags& Dirty, IPhysicsProxyBase* Proxy)
 	{
 		Dirty.MarkClean(PropertyFlag);
 		if (Proxy && Dirty.IsClean())
@@ -53,7 +53,7 @@ public:
 		}
 	}
 
-	void SyncRemote(FDirtyPropertiesManager& Manager, int32 DataIdx, const FParticleDirtyData& Remote) const
+	void SyncRemote(FDirtyPropertiesManager& Manager, int32 DataIdx, const FDirtyChaosProperties& Remote) const
 	{
 		Remote.SyncRemote<T,PropName>(Manager, DataIdx, Property);
 	}
@@ -65,9 +65,9 @@ public:
 
 private:
 	T Property;
-	static constexpr EParticleFlags PropertyFlag = ParticlePropToFlag(PropName);
+	static constexpr EChaosPropertyFlags PropertyFlag = ChaosPropertyToFlag(PropName);
 
-	void MarkDirty(bool bInvalidate, FParticleDirtyFlags& Dirty, IPhysicsProxyBase* Proxy)
+	void MarkDirty(bool bInvalidate, FDirtyChaosPropertyFlags& Dirty, IPhysicsProxyBase* Proxy)
 	{
 		if(bInvalidate)
 		{
@@ -84,11 +84,11 @@ private:
 	}
 
 	//we don't support this because it could lead to bugs with values not being properly written to remote
-	TParticleProperty<T,PropName>& operator=(const TParticleProperty<T,PropName>& Rhs) {}
+	TChaosProperty<T,PropName>& operator=(const TChaosProperty<T,PropName>& Rhs) {}
 };
 
-template <typename T, EParticleProperty PropName>
-FChaosArchive& operator<<(FChaosArchive& Ar, TParticleProperty<T, PropName>& Prop)
+template <typename T, EChaosProperty PropName>
+FChaosArchive& operator<<(FChaosArchive& Ar, TChaosProperty<T, PropName>& Prop)
 {
 	//TODO: should this only work with dirty flag? Not sure if this path really matters at this point
 	Prop.Serialize(Ar);
diff --git a/Engine/Source/Runtime/Experimental/Chaos/Public/Chaos/PropertiesTypes.inl b/Engine/Source/Runtime/Experimental/Chaos/Public/Chaos/PropertiesTypes.inl
deleted file mode 100644
diff --git a/Engine/Source/Runtime/Experimental/Chaos/Public/RewindData.h b/Engine/Source/Runtime/Experimental/Chaos/Public/RewindData.h
--- a/Engine/Source/Runtime/Experimental/Chaos/Public/RewindData.h
+++ b/Engine/Source/Runtime/Experimental/Chaos/Public/RewindData.h
@@ -81,7 +81,7 @@ struct FPropertyInterval
 	FFrameAndPhase FrameAndPhase;
 };
 
-template <typename T, EParticleProperty PropName, bool bNoEntryIsHead = true>
+template <typename T, EChaosProperty PropName, bool bNoEntryIsHead = true>
 class TParticlePropertyBuffer
 {
 public:
@@ -684,17 +684,17 @@ public:
 	}
 
 	void SyncSimWritablePropsFromSim(FDirtyPropData Manager,const TPBDRigidParticleHandle<FReal,3>& Rigid);
-	void SyncDirtyDynamics(FDirtyPropData& DestManager,const FParticleDirtyData& Dirty,const FConstDirtyPropData& SrcManager);
+	void SyncDirtyDynamics(FDirtyPropData& DestManager,const FDirtyChaosProperties& Dirty,const FConstDirtyPropData& SrcManager);
 	
 private:
 
-	TParticlePropertyBuffer<FParticlePositionRotation,EParticleProperty::XR> ParticlePositionRotation;
-	TParticlePropertyBuffer<FParticleNonFrequentData,EParticleProperty::NonFrequentData> NonFrequentData;
-	TParticlePropertyBuffer<FParticleVelocities,EParticleProperty::Velocities> Velocities;
-	TParticlePropertyBuffer<FParticleDynamics,EParticleProperty::Dynamics, /*bNoEntryIsHead=*/false> Dynamics;
-	TParticlePropertyBuffer<FParticleDynamicMisc,EParticleProperty::DynamicMisc> DynamicsMisc;
-	TParticlePropertyBuffer<FParticleMassProps,EParticleProperty::MassProps> MassProps;
-	TParticlePropertyBuffer<FKinematicTarget, EParticleProperty::KinematicTarget> KinematicTarget;
+	TParticlePropertyBuffer<FParticlePositionRotation,EChaosProperty::XR> ParticlePositionRotation;
+	TParticlePropertyBuffer<FParticleNonFrequentData,EChaosProperty::NonFrequentData> NonFrequentData;
+	TParticlePropertyBuffer<FParticleVelocities,EChaosProperty::Velocities> Velocities;
+	TParticlePropertyBuffer<FParticleDynamics,EChaosProperty::Dynamics, /*bNoEntryIsHead=*/false> Dynamics;
+	TParticlePropertyBuffer<FParticleDynamicMisc,EChaosProperty::DynamicMisc> DynamicsMisc;
+	TParticlePropertyBuffer<FParticleMassProps,EChaosProperty::MassProps> MassProps;
+	TParticlePropertyBuffer<FKinematicTarget, EChaosProperty::KinematicTarget> KinematicTarget;
 
 	FShapesArrayStateBase ShapesArrayState;
 
