Fix CIS

diff --git a/Engine/Source/Programs/HeadlessChaos/Private/HeadlessChaosTestSolverProxies.cpp b/Engine/Source/Programs/HeadlessChaos/Private/HeadlessChaosTestSolverProxies.cpp
--- a/Engine/Source/Programs/HeadlessChaos/Private/HeadlessChaosTestSolverProxies.cpp
+++ b/Engine/Source/Programs/HeadlessChaos/Private/HeadlessChaosTestSolverProxies.cpp
@@ -148,8 +148,8 @@ namespace ChaosTest {
 
 		// joint Proxy has been created, should be in dirty proxy list to be sent to physics thread.
 		// Note that we expect three dirty proxies, 2 for particles, 1 for joint.
-		EXPECT_EQ(Solver->GetDirtyProxyBucketInfo_External().Num[EPhysicsProxyType::SingleParticleProxy], 2);
-		EXPECT_EQ(Solver->GetDirtyProxyBucketInfo_External().Num[EPhysicsProxyType::JointConstraintType], 1);
+		EXPECT_EQ(Solver->GetDirtyProxyBucketInfo_External().Num[(uint32)EPhysicsProxyType::SingleParticleProxy], 2);
+		EXPECT_EQ(Solver->GetDirtyProxyBucketInfo_External().Num[(uint32)EPhysicsProxyType::JointConstraintType], 1);
 
 		// Although proxy exists, it should not have been added to physics thread only proxy array yet, as doing so from game thread code is wrong.
 		// Confirm game thread did not add it.
diff --git a/Engine/Source/Runtime/Experimental/Chaos/Public/Chaos/ChaosMarshallingManager.h b/Engine/Source/Runtime/Experimental/Chaos/Public/Chaos/ChaosMarshallingManager.h
--- a/Engine/Source/Runtime/Experimental/Chaos/Public/Chaos/ChaosMarshallingManager.h
+++ b/Engine/Source/Runtime/Experimental/Chaos/Public/Chaos/ChaosMarshallingManager.h
@@ -59,8 +59,8 @@ public:
 	{
 		if(Base->GetDirtyIdx() == INDEX_NONE)
 		{
-			FDirtyProxiesBucket& Bucket = DirtyProxyBuckets[Base->GetType()];
-			++DirtyProxyBucketInfo.Num[Base->GetType()];
+			FDirtyProxiesBucket& Bucket = DirtyProxyBuckets[(uint32)Base->GetType()];
+			++DirtyProxyBucketInfo.Num[(uint32)Base->GetType()];
 			++DirtyProxyBucketInfo.TotalNum;
 
 			const int32 Idx = Bucket.ProxiesData.Num();
@@ -76,7 +76,7 @@ public:
 	{
 		if(ProxiesArray.Num())
 		{
-			FDirtyProxiesBucket& Bucket = DirtyProxyBuckets[ProxiesArray[0]->GetType()];
+			FDirtyProxiesBucket& Bucket = DirtyProxyBuckets[(uint32)ProxiesArray[0]->GetType()];
 			int32 Idx = Bucket.ProxiesData.Num();
 			Bucket.ProxiesData.Append(ProxiesArray);
 
@@ -86,7 +86,7 @@ public:
 				ensure(ProxiesArray[0]->GetType() == Proxy->GetType());
 			}
 
-			DirtyProxyBucketInfo.Num[ProxiesArray[0]->GetType()] += ProxiesArray.Num();
+			DirtyProxyBucketInfo.Num[(uint32)ProxiesArray[0]->GetType()] += ProxiesArray.Num();
 			DirtyProxyBucketInfo.TotalNum += ProxiesArray.Num();
 		}
 		
@@ -98,12 +98,12 @@ public:
 		const int32 Idx = Base->GetDirtyIdx();
 		if(Idx != INDEX_NONE)
 		{
-			FDirtyProxiesBucket& Bucket = DirtyProxyBuckets[Base->GetType()];
+			FDirtyProxiesBucket& Bucket = DirtyProxyBuckets[(uint32)Base->GetType()];
 			if(Idx == Bucket.ProxiesData.Num() - 1)
 			{
 				//last element so just pop
 				Bucket.ProxiesData.Pop(/*bAllowShrinking=*/false);
-				--DirtyProxyBucketInfo.Num[Base->GetType()];
+				--DirtyProxyBucketInfo.Num[(uint32)Base->GetType()];
 				--DirtyProxyBucketInfo.TotalNum;
 			}
 			else if(Bucket.ProxiesData.IsValidIndex(Idx))
@@ -111,7 +111,7 @@ public:
 				//update other proxy's idx
 				Bucket.ProxiesData.RemoveAtSwap(Idx);
 				Bucket.ProxiesData[Idx].SetDirtyIdx(Idx);
-				--DirtyProxyBucketInfo.Num[Base->GetType()];
+				--DirtyProxyBucketInfo.Num[(uint32)Base->GetType()];
 				--DirtyProxyBucketInfo.TotalNum;
 			}
 
@@ -134,7 +134,7 @@ public:
 	int32 NumDirtyShapes() const { return ShapesData.Num(); }
 
 	FShapeDirtyData* GetShapesDirtyData(){ return ShapesData.GetData(); }
-	FDirtyProxy& GetDirtyProxyAt(EPhysicsProxyType ProxyType, int32 Idx) { return DirtyProxyBuckets[ProxyType].ProxiesData[Idx]; }
+	FDirtyProxy& GetDirtyProxyAt(EPhysicsProxyType ProxyType, int32 Idx) { return DirtyProxyBuckets[(uint32)ProxyType].ProxiesData[Idx]; }
 
 	template <typename Lambda>
 	void ParallelForEachProxy(const Lambda& Func)
@@ -161,7 +161,7 @@ public:
 	template <typename Lambda>
 	void ForEachProxy(const Lambda& Func)
 	{
-		for(int32 BucketIdx = 0; BucketIdx < EPhysicsProxyType::Count; ++BucketIdx)
+		for(int32 BucketIdx = 0; BucketIdx < (uint32)EPhysicsProxyType::Count; ++BucketIdx)
 		{
 			int32 Idx = 0;
 			for (FDirtyProxy& Dirty : DirtyProxyBuckets[BucketIdx].ProxiesData)
@@ -174,7 +174,7 @@ public:
 	template <typename Lambda>
 	void ForEachProxy(const Lambda& Func) const
 	{
-		for (int32 BucketIdx = 0; BucketIdx < EPhysicsProxyType::Count; ++BucketIdx)
+		for (int32 BucketIdx = 0; BucketIdx < (uint32)EPhysicsProxyType::Count; ++BucketIdx)
 		{
 			int32 Idx = 0;
 			for (const FDirtyProxy& Dirty : DirtyProxyBuckets[BucketIdx].ProxiesData)
@@ -187,7 +187,7 @@ public:
 	void AddShape(IPhysicsProxyBase* Proxy,int32 ShapeIdx)
 	{
 		Add(Proxy);
-		FDirtyProxy& Dirty = DirtyProxyBuckets[Proxy->GetType()].ProxiesData[Proxy->GetDirtyIdx()];
+		FDirtyProxy& Dirty = DirtyProxyBuckets[(uint32)Proxy->GetType()].ProxiesData[(uint32)Proxy->GetDirtyIdx()];
 		for(int32 NewShapeIdx = Dirty.ShapeDataIndices.Num(); NewShapeIdx <= ShapeIdx; ++NewShapeIdx)
 		{
 			const int32 ShapeDataIdx = ShapesData.Add(FShapeDirtyData(NewShapeIdx));
@@ -198,7 +198,7 @@ public:
 	void SetNumDirtyShapes(IPhysicsProxyBase* Proxy,int32 NumShapes)
 	{
 		Add(Proxy);
-		FDirtyProxy& Dirty = DirtyProxyBuckets[Proxy->GetType()].ProxiesData[Proxy->GetDirtyIdx()];
+		FDirtyProxy& Dirty = DirtyProxyBuckets[(uint32)Proxy->GetType()].ProxiesData[Proxy->GetDirtyIdx()];
 
 		if(NumShapes < Dirty.ShapeDataIndices.Num())
 		{
@@ -215,7 +215,7 @@ public:
 
 private:
 	FDirtyProxiesBucketInfo DirtyProxyBucketInfo;
-	FDirtyProxiesBucket DirtyProxyBuckets[EPhysicsProxyType::Count];
+	FDirtyProxiesBucket DirtyProxyBuckets[(uint32)EPhysicsProxyType::Count];
 	TArray<FShapeDirtyData> ShapesData;
 };
 
diff --git a/Engine/Source/Runtime/Experimental/Chaos/Public/Chaos/Framework/PhysicsProxyBase.h b/Engine/Source/Runtime/Experimental/Chaos/Public/Chaos/Framework/PhysicsProxyBase.h
--- a/Engine/Source/Runtime/Experimental/Chaos/Public/Chaos/Framework/PhysicsProxyBase.h
+++ b/Engine/Source/Runtime/Experimental/Chaos/Public/Chaos/Framework/PhysicsProxyBase.h
@@ -6,7 +6,7 @@
 #include "UObject/GCObject.h"
 #include "Chaos/Core.h"
 
-enum EPhysicsProxyType
+enum class EPhysicsProxyType : uint32
 {
 	NoneType = 0,
 	StaticMeshType = 1,
diff --git a/Engine/Source/Runtime/Experimental/Chaos/Public/Chaos/ParticleDirtyFlags.h b/Engine/Source/Runtime/Experimental/Chaos/Public/Chaos/ParticleDirtyFlags.h
--- a/Engine/Source/Runtime/Experimental/Chaos/Public/Chaos/ParticleDirtyFlags.h
+++ b/Engine/Source/Runtime/Experimental/Chaos/Public/Chaos/ParticleDirtyFlags.h
@@ -16,6 +16,7 @@
 #include "UObject/UE5MainStreamObjectVersion.h"
 #include "UObject/UE5ReleaseStreamObjectVersion.h"
 #include "Framework/PhysicsProxyBase.h"
+#include "PBDJointConstraintTypes.h"
 
 #ifndef CHAOS_DEBUG_NAME
 #define CHAOS_DEBUG_NAME 0
@@ -740,19 +741,19 @@ inline FChaosArchive& operator<<(FChaosArchive& Ar,FMaterialData& Data)
   * as well as how to go from a contiguous index into a per bucket index */
 struct FDirtyProxiesBucketInfo
 {
-	int32 Num[EPhysicsProxyType::Count] = {};
+	int32 Num[(uint32)(EPhysicsProxyType::Count)] = {};
 	int32 TotalNum = 0;
 
 	void Reset()
 	{
-		for (int32 Idx = 0; Idx < EPhysicsProxyType::Count; ++Idx) { Num[Idx] = 0; }
+		for (int32 Idx = 0; Idx < (uint32)EPhysicsProxyType::Count; ++Idx) { Num[Idx] = 0; }
 		TotalNum = 0;
 	}
 
 	void GetBucketIdx(int32 Idx, int32& OutBucketIdx, int32& InnerIdx) const
 	{
 		int32 Remaining = Idx;
-		for (int32 BucketIdx = 0; BucketIdx < EPhysicsProxyType::Count; ++BucketIdx)
+		for (int32 BucketIdx = 0; BucketIdx < (uint32)EPhysicsProxyType::Count; ++BucketIdx)
 		{
 			if (Remaining < Num[BucketIdx])
 			{
@@ -777,7 +778,7 @@ public:
 
 	void PrepareBuckets(const FDirtyProxiesBucketInfo& DirtyProxiesBucketInfo)
 	{
-#define CHAOS_PROPERTY(PropName, Type, ProxyType) PropName##Pool.SetNum(DirtyProxiesBucketInfo.Num[ProxyType]);
+#define CHAOS_PROPERTY(PropName, Type, ProxyType) PropName##Pool.SetNum(DirtyProxiesBucketInfo.Num[(uint32)ProxyType]);
 #include "ParticleProperties.inl"
 #undef CHAOS_PROPERTY
 	}
