prevent rendering crash if mesh gets into a bad state where all geometry is hidden (by un-hiding everything if that case occurs)

#rnx

diff --git a/Engine/Source/Runtime/Experimental/GeometryCollectionEngine/Private/GeometryCollection/GeometryCollectionComponent.cpp b/Engine/Source/Runtime/Experimental/GeometryCollectionEngine/Private/GeometryCollection/GeometryCollectionComponent.cpp
--- a/Engine/Source/Runtime/Experimental/GeometryCollectionEngine/Private/GeometryCollection/GeometryCollectionComponent.cpp
+++ b/Engine/Source/Runtime/Experimental/GeometryCollectionEngine/Private/GeometryCollection/GeometryCollectionComponent.cpp
@@ -1435,6 +1435,8 @@ void UGeometryCollectionComponent::InitConstantData(FGeometryCollectionConstantD
 		{
 			bUsingHideArray = true;
 
+			bool bAllHidden = true;
+
 			const TManagedArray<bool>& Hide = Collection->GetAttribute<bool>("Hide", FGeometryCollection::TransformGroup);
 			for (int32 GeomIdx = 0; GeomIdx < NumGeom; ++GeomIdx)
 			{
@@ -1446,6 +1448,17 @@ void UGeometryCollectionComponent::InitConstantData(FGeometryCollectionConstantD
 						VisibleOverride[FaceStart[GeomIdx]+FaceIdxOffset] = false;
 					}
 				}
+				else
+				{
+					bAllHidden = false;
+				}
+			}
+			if (!ensure(!bAllHidden)) // if they're all hidden, rendering would crash -- unhide them
+			{
+				for (int32 FaceIdx = 0; FaceIdx < VisibleOverride.Num(); ++FaceIdx)
+				{
+					VisibleOverride[FaceIdx] = true;
+				}
 			}
 		}
 #endif
