Chaos : Collision Constraint
- Allow for particles to track their collision points.
- Added DoBufferColliosions flag to instruct the collision system to add the collision constraints to the buffered mapping.

diff --git a/Engine/Source/Runtime/Experimental/Chaos/Private/Chaos/PBDCollisionConstraints.cpp b/Engine/Source/Runtime/Experimental/Chaos/Private/Chaos/PBDCollisionConstraints.cpp
--- a/Engine/Source/Runtime/Experimental/Chaos/Private/Chaos/PBDCollisionConstraints.cpp
+++ b/Engine/Source/Runtime/Experimental/Chaos/Private/Chaos/PBDCollisionConstraints.cpp
@@ -79,6 +79,10 @@ namespace Chaos
 	int32 CollisionCanNeverDisableContacts = 0;
 	FAutoConsoleVariableRef CVarCollisionCanNeverDisableContacts(TEXT("p.CollisionCanNeverDisableContacts"), CollisionCanNeverDisableContacts, TEXT("Collision culling will never be able to permanently disable contacts"));
 
+	bool CollisionsAllowParticleTracking = true;
+	FAutoConsoleVariableRef CVarCollisionsAllowParticleTracking(TEXT("p.Chaos.Collision.AllowParticleTracking"), CollisionsAllowParticleTracking, TEXT("Allow particles to track their collisions constraints when their DoBufferCollisions flag is enable [def:true]"));
+
+
 #if INTEL_ISPC
 	bool bChaos_Collision_ISPC_Enabled = false;
 	FAutoConsoleVariableRef CVarChaosCollisionISPCEnabled(TEXT("p.Chaos.Collision.ISPC"), bChaos_Collision_ISPC_Enabled, TEXT("Whether to use ISPC optimizations in the Collision Solver"));
@@ -276,7 +280,7 @@ namespace Chaos
 
 			check(Handle != nullptr);
 			Handles.Add(Handle);
-
+				
 #if CHAOS_COLLISION_PERSISTENCE_ENABLED
 			check(!Manifolds.Contains(Handle->GetKey()));
 			Manifolds.Add(Handle->GetKey(), Handle);
@@ -359,6 +363,7 @@ namespace Chaos
 		}
 		Constraints.Reset();
 		Handles.Reset();
+		ParticleCollisionsMap.Reset();
 #endif
 
 		bUseCCD = false;
@@ -737,6 +742,35 @@ namespace Chaos
 
 			NewHandle->GetContact().Timestamp = -INT_MAX;
 			Constraints->SinglePointConstraints[NumBeginSingle + HandleIndex].SetConstraintHandle(NewHandle);
+
+			if (CollisionsAllowParticleTracking)
+			{
+				// Add reverse mappings for particles that require it. 
+				using FHandle = FPBDCollisionConstraintHandle;
+				using FParticle = TPBDRigidParticleHandleImp<FReal, 3, true>;
+
+				auto AddToParticleCollisionsMap = [](FHandle* InHandle, FGeometryParticleHandle* InParticle, TMap<FParticleHandle*, FHandles >& InParticleCollisionsMap)
+				{
+					if (InParticle)
+					{
+						if (FParticle* RigidParticle = InParticle->CastToRigidParticle())
+						{
+							if (RigidParticle->HasCollisionConstraintFlag(ECollisionConstraintFlags::CCF_DoBufferCollisions))
+							{
+								FHandles* HandleArray = InParticleCollisionsMap.Find(RigidParticle);
+								if (!HandleArray)
+								{
+									InParticleCollisionsMap.Add(RigidParticle, FHandles());
+									HandleArray = InParticleCollisionsMap.Find(RigidParticle);
+								}
+								HandleArray->Add(InHandle);
+							}
+						}
+					}
+				};
+				AddToParticleCollisionsMap(NewHandle, NewHandle->GetConstrainedParticles()[0], Owner->GetParticleCollisionsMap());
+				AddToParticleCollisionsMap(NewHandle, NewHandle->GetConstrainedParticles()[1], Owner->GetParticleCollisionsMap());
+			}
 		}
 		HandlesBeginIndex += NumAddedSingle;
 
diff --git a/Engine/Source/Runtime/Experimental/Chaos/Private/PhysicsProxy/JointConstraintProxy.cpp b/Engine/Source/Runtime/Experimental/Chaos/Private/PhysicsProxy/JointConstraintProxy.cpp
--- a/Engine/Source/Runtime/Experimental/Chaos/Private/PhysicsProxy/JointConstraintProxy.cpp
+++ b/Engine/Source/Runtime/Experimental/Chaos/Private/PhysicsProxy/JointConstraintProxy.cpp
@@ -195,6 +195,18 @@ void FJointConstraintPhysicsProxy::PushStateOnPhysicsThread(FPBDRigidsSolver* In
 						}
 					}
 				}
+
+				// Contact transfer enables buffered collision constraints on the child of the joint
+				if (Chaos::TGeometryParticleHandle<Chaos::FReal, 3>*Handle0 = GetParticleHandleFromProxy(Constraint->GetParticleProxies()[0]))
+				{
+					if (Chaos::TPBDRigidParticleHandle<FReal, 3>*RigidHandle0 = Handle0->CastToRigidParticle())
+					{
+						if (!FMath::IsNearlyZero(CurrentConstraintSettings.ContactTransferScale))
+						{
+							RigidHandle0->AddCollisionConstraintFlag(Chaos::ECollisionConstraintFlags::CCF_DoBufferCollisions);
+						}
+					}
+				}
 			}
 
 			Handle->SetSettings(JointSettingsBuffer);
diff --git a/Engine/Source/Runtime/Experimental/Chaos/Public/Chaos/Collision/CollisionConstraintFlags.h b/Engine/Source/Runtime/Experimental/Chaos/Public/Chaos/Collision/CollisionConstraintFlags.h
--- a/Engine/Source/Runtime/Experimental/Chaos/Public/Chaos/Collision/CollisionConstraintFlags.h
+++ b/Engine/Source/Runtime/Experimental/Chaos/Public/Chaos/Collision/CollisionConstraintFlags.h
@@ -18,6 +18,7 @@ namespace Chaos
 	{
 		CCF_None                       = 0x0,
 		CCF_BroadPhaseIgnoreCollisions = 0x1,
+		CCF_DoBufferCollisions         = 0x2,
 		CCF_DummyFlag
 	};
 
diff --git a/Engine/Source/Runtime/Experimental/Chaos/Public/Chaos/PBDCollisionConstraints.h b/Engine/Source/Runtime/Experimental/Chaos/Public/Chaos/PBDCollisionConstraints.h
--- a/Engine/Source/Runtime/Experimental/Chaos/Public/Chaos/PBDCollisionConstraints.h
+++ b/Engine/Source/Runtime/Experimental/Chaos/Public/Chaos/PBDCollisionConstraints.h
@@ -50,7 +50,7 @@ public:
 	// For use by dependent types
 	using FPointContactConstraint = FRigidBodyPointContactConstraint;
 	using FConstraintContainerHandle = FPBDCollisionConstraintHandle;
-
+	using FParticleHandle = TGeometryParticleHandle<FReal, 3>;
 
 	FPBDCollisionConstraints(const FPBDRigidsSOAs& InParticles, 
 		TArrayCollectionArray<bool>& Collided, 
@@ -340,6 +340,9 @@ public:
 	//Sort constraints based on particle indices so that we have a deterministic solve order
 	void SortConstraints();
 
+	TMap<FParticleHandle*, TArray<FPBDCollisionConstraintHandle*> >& GetParticleCollisionsMap() { return ParticleCollisionsMap; }
+	const TMap<FParticleHandle*, TArray<FPBDCollisionConstraintHandle*> >& GetParticleCollisionsMap() const { return ParticleCollisionsMap; }
+
 protected:
 	using Base::GetConstraintIndex;
 	using Base::SetConstraintIndex;
@@ -364,6 +367,9 @@ private:
 	TArray<FPBDCollisionConstraintHandle*> Handles;
 	FConstraintHandleAllocator HandleAllocator;
 
+	TMap<FParticleHandle*,TArray<FPBDCollisionConstraintHandle*> > ParticleCollisionsMap;
+
+
 	TArrayCollectionArray<bool>& MCollided;
 	const TArrayCollectionArray<TSerializablePtr<FChaosPhysicsMaterial>>& MPhysicsMaterials;
 	const TArrayCollectionArray<TUniquePtr<FChaosPhysicsMaterial>>& MPerParticlePhysicsMaterials;
diff --git a/Engine/Source/Runtime/Experimental/Chaos/Public/Chaos/ParticleHandle.h b/Engine/Source/Runtime/Experimental/Chaos/Public/Chaos/ParticleHandle.h
--- a/Engine/Source/Runtime/Experimental/Chaos/Public/Chaos/ParticleHandle.h
+++ b/Engine/Source/Runtime/Experimental/Chaos/Public/Chaos/ParticleHandle.h
@@ -566,6 +566,11 @@ public:
 		return GeometryParticles->ParticleConstraints(ParticleIdx);
 	}
 
+	const FConstraintHandleArray& ParticleConstraints() const
+	{
+		return GeometryParticles->ParticleConstraints(ParticleIdx);
+	}
+
 	void AddConstraintHandle(FConstraintHandle* InConstraintHandle )
 	{
 		return GeometryParticles->AddConstraintHandle(ParticleIdx, InConstraintHandle);
@@ -680,6 +685,7 @@ class TPBDRigidParticleHandleImp : public TKinematicGeometryParticleHandleImp<T,
 public:
 	using TGeometryParticleHandleImp<T, d, bPersistent>::ParticleIdx;
 	using TGeometryParticleHandleImp<T, d, bPersistent>::PBDRigidParticles;
+	using TGeometryParticleHandleImp<T, d, bPersistent>::ParticleConstraints;
 	using TKinematicGeometryParticleHandleImp<T, d, bPersistent>::V;
 	using TKinematicGeometryParticleHandleImp<T, d, bPersistent>::W;
 	using TGeometryParticleHandleImp<T, d, bPersistent>::Type;
