Chaos : - Fix CIS

diff --git a/Engine/Source/Runtime/Experimental/Chaos/Private/PhysicsProxy/JointConstraintProxy.cpp b/Engine/Source/Runtime/Experimental/Chaos/Private/PhysicsProxy/JointConstraintProxy.cpp
--- a/Engine/Source/Runtime/Experimental/Chaos/Private/PhysicsProxy/JointConstraintProxy.cpp
+++ b/Engine/Source/Runtime/Experimental/Chaos/Private/PhysicsProxy/JointConstraintProxy.cpp
@@ -197,7 +197,7 @@ void FJointConstraintPhysicsProxy::PushStateOnPhysicsThread(FPBDRigidsSolver* In
 				}
 
 				// Contact transfer enables buffered collision constraints on the child of the joint
-				if (Chaos::TGeometryParticleHandle<Chaos::FReal, 3>*Handle0 = GetParticleHandleFromProxy(Constraint->GetParticleProxies()[0]))
+				if (Handle0)
 				{
 					if (Chaos::TPBDRigidParticleHandle<FReal, 3>*RigidHandle0 = Handle0->CastToRigidParticle())
 					{
