GeometryCollectionProximityUtility
-- Caught edge case that defines identical triangles as overlapping

diff --git a/Engine/Source/Runtime/Experimental/Chaos/Private/GeometryCollection/GeometryCollectionProximityUtility.cpp b/Engine/Source/Runtime/Experimental/Chaos/Private/GeometryCollection/GeometryCollectionProximityUtility.cpp
--- a/Engine/Source/Runtime/Experimental/Chaos/Private/GeometryCollection/GeometryCollectionProximityUtility.cpp
+++ b/Engine/Source/Runtime/Experimental/Chaos/Private/GeometryCollection/GeometryCollectionProximityUtility.cpp
@@ -169,23 +169,55 @@ bool FGeometryCollectionProximityUtility::DoFacesOverlap(int32 Idx0, int32 Idx1)
 	Basis0.Normalize();
 	FVector3f Basis1 = FVector3f::CrossProduct(SurfaceNormals[Idx0], Basis0);
 	Basis1.Normalize();
+
+	FVector3f Origin = TransformedVertices[Indices[Idx0].X];
 	
 	TStaticArray<FVector2D,3> T0; 
 	// T0[0] is the origin of the system
 	T0[0] = FVector2D(0.f, 0.f);
-	T0[1] = FVector2D(FVector3f::DotProduct(TransformedVertices[Indices[Idx0].Y], Basis0), FVector3f::DotProduct(TransformedVertices[Indices[Idx0].Y], Basis1));
-	T0[2] = FVector2D(FVector3f::DotProduct(TransformedVertices[Indices[Idx0].Z], Basis0), FVector3f::DotProduct(TransformedVertices[Indices[Idx0].Z], Basis1));
+
+	T0[1] = FVector2D(FVector3f::DotProduct(TransformedVertices[Indices[Idx0].Y] - Origin, Basis0), FVector3f::DotProduct(TransformedVertices[Indices[Idx0].Y] - Origin, Basis1));
+	T0[2] = FVector2D(FVector3f::DotProduct(TransformedVertices[Indices[Idx0].Z] - Origin, Basis0), FVector3f::DotProduct(TransformedVertices[Indices[Idx0].Z] - Origin, Basis1));
 
 	// Project the second triangle into these coordinates. We reverse the winding order to flip the normal.
-	FVector3f Point0 = TransformedVertices[Indices[Idx1].Z] - TransformedVertices[Indices[Idx0].Z];
-	FVector3f Point1 = TransformedVertices[Indices[Idx1].Y] - TransformedVertices[Indices[Idx0].Y];
-	FVector3f Point2 = TransformedVertices[Indices[Idx1].X] - TransformedVertices[Indices[Idx0].X];
+	FVector3f Point0 = TransformedVertices[Indices[Idx1].Z] - Origin;
+	FVector3f Point1 = TransformedVertices[Indices[Idx1].Y] - Origin;
+	FVector3f Point2 = TransformedVertices[Indices[Idx1].X] - Origin;
 	TStaticArray<FVector2D,3> T1;
 	T1[0] = FVector2D(FVector3f::DotProduct(Point0, Basis0), FVector3f::DotProduct(Point0, Basis1));
 	T1[1] = FVector2D(FVector3f::DotProduct(Point1, Basis0), FVector3f::DotProduct(Point1, Basis1));
 	T1[2] = FVector2D(FVector3f::DotProduct(Point2, Basis0), FVector3f::DotProduct(Point2, Basis1));
 
-	return TrianglesIntersect(T0, T1);
+	return IdenticalTriangles(T0, T1) || TrianglesIntersect(T0, T1);
+}
+
+bool FGeometryCollectionProximityUtility::IdenticalTriangles(const TStaticArray<FVector2D, 3>& T0, const TStaticArray<FVector2D, 3>& T1)
+{
+	int32 FirstMatch = 0;
+	for (; FirstMatch < 3; ++FirstMatch)
+	{
+		if (T0[0] == T1[FirstMatch])
+		{
+			break;
+		}
+	}
+
+	if (FirstMatch == 3) // No match
+	{
+		return false;
+	}
+
+	if (T0[1] != T1[(FirstMatch + 1) % 3])
+	{
+		return false;
+	}
+
+	if (T0[2] != T1[(FirstMatch + 2) % 3])
+	{
+		return false;
+	}
+
+	return true;
 }
 
 bool FGeometryCollectionProximityUtility::TrianglesIntersect(const TStaticArray<FVector2D,3>& T0, const TStaticArray<FVector2D,3>& T1)
@@ -270,7 +302,8 @@ void FGeometryCollectionProximityUtility::ExtendFaceProximityToGeometry()
 {
 	if (!Collection->HasAttribute("Proximity", FGeometryCollection::GeometryGroup))
 	{
-		Collection->AddAttribute<TSet<int32>>("Proximity", FGeometryCollection::GeometryGroup);
+		const FManagedArrayCollection::FConstructionParameters GeometryDependency(FGeometryCollection::GeometryGroup);
+		Collection->AddAttribute<TSet<int32>>("Proximity", FGeometryCollection::GeometryGroup, GeometryDependency);
 	}
 
 	TManagedArray<TSet<int32>>& Proximity = Collection->GetAttribute<TSet<int32>>("Proximity", FGeometryCollection::GeometryGroup);
diff --git a/Engine/Source/Runtime/Experimental/Chaos/Public/GeometryCollection/GeometryCollectionProximityUtility.h b/Engine/Source/Runtime/Experimental/Chaos/Public/GeometryCollection/GeometryCollectionProximityUtility.h
--- a/Engine/Source/Runtime/Experimental/Chaos/Public/GeometryCollection/GeometryCollectionProximityUtility.h
+++ b/Engine/Source/Runtime/Experimental/Chaos/Public/GeometryCollection/GeometryCollectionProximityUtility.h
@@ -30,6 +30,7 @@ private:
 	bool AreNormalsOpposite(const FVector3f& Normal0, const FVector3f& Normal1) const;
 	bool AreFacesCoPlanar(int32 Idx0, int32 Idx1) const;
 	bool DoFacesOverlap(int32 Idx0, int32 Idx1) const;
+	static bool IdenticalTriangles(const TStaticArray<FVector2D, 3>& T0, const TStaticArray<FVector2D, 3>& T1);
 	static bool TrianglesIntersect(const TStaticArray<FVector2D,3>& T0, const TStaticArray<FVector2D,3>& T1);
 	static bool Cross(const TStaticArray<FVector2D,3>& Points, const FVector2D& B, const FVector2D& C, float Normal);
 
