FractureTool
-- Additional Undo converage
-- Embedded geometry works with explode
-- Selectable embedded geometry

diff --git a/Engine/Plugins/Experimental/ChaosEditor/Source/FractureEditor/Private/FractureEditorMode.cpp b/Engine/Plugins/Experimental/ChaosEditor/Source/FractureEditor/Private/FractureEditorMode.cpp
--- a/Engine/Plugins/Experimental/ChaosEditor/Source/FractureEditor/Private/FractureEditorMode.cpp
+++ b/Engine/Plugins/Experimental/ChaosEditor/Source/FractureEditor/Private/FractureEditorMode.cpp
@@ -26,6 +26,7 @@
 #include "GeometryCollection/GeometryCollectionAlgo.h"
 #include "EdModeInteractiveToolsContext.h"
 #include "BaseGizmos/TransformGizmoUtil.h"
+#include "Components/InstancedStaticMeshComponent.h"
 
 
 #include "UnrealEdGlobals.h"
@@ -178,7 +179,6 @@ bool UFractureEditorMode::HandleClick(FEditorViewportClient* InViewportClient, H
 		{
 			TArray<int32> BoneIndices({GeometryCollectionProxy->BoneIndex});
 
-			FScopedTransaction Transaction(FractureTransactionContexts::SelectBoneContext, LOCTEXT("SelectGeometryCollectionBoneTransaction", "Select Bone"), GeometryCollectionProxy->Component);
 			GeometryCollectionProxy->Component->Modify();
 			FFractureSelectionTools::ToggleSelectedBones(GeometryCollectionProxy->Component, BoneIndices, !(Click.IsControlDown()||Click.IsShiftDown()), Click.IsShiftDown());			
 
@@ -191,6 +191,33 @@ bool UFractureEditorMode::HandleClick(FEditorViewportClient* InViewportClient, H
 			return true;
 		}
 	}
+	else if (HitProxy && HitProxy->IsA(HInstancedStaticMeshInstance::StaticGetType()))
+	{
+		HInstancedStaticMeshInstance* ISMProxy = ((HInstancedStaticMeshInstance*)HitProxy);
+		if (ISMProxy->Component)
+		{
+			// Get the hit ISMComp's GeometryCollection 
+			if (UGeometryCollectionComponent* OwningGC = Cast<UGeometryCollectionComponent>(ISMProxy->Component->GetAttachParent()))
+			{
+				int32 TransformIdx = OwningGC->EmbeddedIndexToTransformIndex(ISMProxy->Component, ISMProxy->InstanceIndex);
+				if (TransformIdx > INDEX_NONE)
+				{
+					TArray<int32> BoneIndices({ TransformIdx });
+
+					OwningGC->Modify();
+					FFractureSelectionTools::ToggleSelectedBones(OwningGC, BoneIndices, !(Click.IsControlDown() || Click.IsShiftDown()), Click.IsShiftDown());
+
+					if (Toolkit.IsValid())
+					{
+						FFractureEditorModeToolkit* FractureToolkit = (FFractureEditorModeToolkit*)Toolkit.Get();
+						FractureToolkit->SetBoneSelection(OwningGC, OwningGC->GetSelectedBones(), true);
+					}
+				}
+			}
+
+			return true;
+		}
+	}
 
 	return false;
 
@@ -337,6 +364,8 @@ void UFractureEditorMode::OnActorSelectionChanged(const TArray<UObject*>& NewSel
 		
 		for(UGeometryCollectionComponent* GeometryCollectionComponent : GeometryCollectionComponents)
 		{
+			GeometryCollectionComponent->SetEmbeddedGeometrySelectable(true);
+			
 			// If collection does not already have guids, make them
 			FGeometryCollectionEdit RestCollectionEdit = GeometryCollectionComponent->EditRestCollection(GeometryCollection::EEditUpdate::None);
 			::GeometryCollection::GenerateTemporaryGuids(RestCollectionEdit.GetRestCollection()->GetGeometryCollection().Get());
@@ -359,6 +388,8 @@ void UFractureEditorMode::OnActorSelectionChanged(const TArray<UObject*>& NewSel
 			FScopedColorEdit ShowBoneColorsEdit(ExistingSelection);
 			ShowBoneColorsEdit.SetEnableBoneSelection(false);
 
+			ExistingSelection->SetEmbeddedGeometrySelectable(false);
+
 			// If we have a Hide array on the collection, remove it.
 			if (const UGeometryCollection* RestCollection = ExistingSelection->GetRestCollection())
 			{
diff --git a/Engine/Plugins/Experimental/ChaosEditor/Source/FractureEditor/Private/FractureEditorModeToolkit.cpp b/Engine/Plugins/Experimental/ChaosEditor/Source/FractureEditor/Private/FractureEditorModeToolkit.cpp
--- a/Engine/Plugins/Experimental/ChaosEditor/Source/FractureEditor/Private/FractureEditorModeToolkit.cpp
+++ b/Engine/Plugins/Experimental/ChaosEditor/Source/FractureEditor/Private/FractureEditorModeToolkit.cpp
@@ -1356,14 +1356,20 @@ void FFractureEditorModeToolkit::UpdateVolumes(FGeometryCollectionPtr GeometryCo
 	}
 }
 
-bool GetValidGeoCenter(const TManagedArray<int32>& TransformToGeometryIndex, const TArray<FTransform>& Transforms, const TManagedArray<TSet<int32>>& Children, const TManagedArray<FBox>& BoundingBox, int32 TransformIndex, FVector& OutGeoCenter )
+bool GetValidGeoCenter(FGeometryCollection* Collection, const TManagedArray<int32>& TransformToGeometryIndex, const TArray<FTransform>& Transforms, const TManagedArray<TSet<int32>>& Children, const TManagedArray<FBox>& BoundingBox, int32 TransformIndex, FVector& OutGeoCenter )
 {
-	if (TransformToGeometryIndex[TransformIndex] > -1)
+	if (Collection->IsRigid(TransformIndex))
 	{
 		OutGeoCenter = Transforms[TransformIndex].TransformPosition(BoundingBox[TransformToGeometryIndex[TransformIndex]].GetCenter());
 
 		return true;
 	}
+	else if (Collection->SimulationType[TransformIndex] == FGeometryCollection::ESimulationTypes::FST_None) // ie this is embedded geometry
+	{
+		OutGeoCenter = Transforms[Collection->Parent[TransformIndex]].TransformPosition(BoundingBox[TransformToGeometryIndex[Collection->Parent[TransformIndex]]].GetCenter());
+
+		return true;
+	}
 	else
 	{
 		FVector AverageCenter;
@@ -1371,7 +1377,7 @@ bool GetValidGeoCenter(const TManagedArray<int32>& TransformToGeometryIndex, con
 		for(int32 ChildIndex : Children[TransformIndex])
 		{
 
-			if (GetValidGeoCenter(TransformToGeometryIndex, Transforms, Children, BoundingBox, ChildIndex, OutGeoCenter))
+			if (GetValidGeoCenter(Collection, TransformToGeometryIndex, Transforms, Children, BoundingBox, ChildIndex, OutGeoCenter))
 			{
 				if (ValidVectors == 0)
 				{
@@ -1467,7 +1473,7 @@ void FFractureEditorModeToolkit::UpdateExplodedVectors(UGeometryCollectionCompon
 		{
 			ExplodedVectors[Idx] = FVector::ZeroVector;
 			FVector GeoCenter;
-			if (GetValidGeoCenter(TransformToGeometryIndex, Transforms, Children, BoundingBox, Idx, GeoCenter))
+			if (GetValidGeoCenter(GeometryCollection->GetGeometryCollection().Get(), TransformToGeometryIndex, Transforms, Children, BoundingBox, Idx, GeoCenter))
 			{
 				TransformedCenters[Idx] = GeoCenter;
 				if ((ViewFractureLevel < 0) || Levels[Idx] == ViewFractureLevel)
@@ -1498,6 +1504,8 @@ void FFractureEditorModeToolkit::UpdateExplodedVectors(UGeometryCollectionCompon
 			}
 		}
 	}
+
+	GeometryCollectionComponent->RefreshEmbeddedGeometry();
 }
 
 void FFractureEditorModeToolkit::RegenerateOutliner()
@@ -1515,8 +1523,6 @@ void FFractureEditorModeToolkit::OnOutlinerBoneSelectionChanged(UGeometryCollect
 	const UGeometryCollection* RestCollection = RootComponent->GetRestCollection();
 	if (RestCollection && !RestCollection->IsPendingKill())
 	{
-		FScopedTransaction Transaction(FractureTransactionContexts::SelectBoneContext, LOCTEXT("SelectGeometryCollectionBoneTransaction", "Select Bone"), RootComponent);
-
 		if (SelectedBones.Num())
 		{
 
@@ -1547,8 +1553,6 @@ void FFractureEditorModeToolkit::OnHistogramBoneSelectionChanged(UGeometryCollec
 	const UGeometryCollection* RestCollection = RootComponent->GetRestCollection();
 	if (RestCollection && !RestCollection->IsPendingKill())
 	{
-		FScopedTransaction Transaction(FractureTransactionContexts::SelectBoneContext, LOCTEXT("SelectGeometryCollectionBoneTransaction", "Select Bone"), RootComponent);
-
 		if (SelectedBones.Num())
 		{
 
diff --git a/Engine/Plugins/Experimental/ChaosEditor/Source/FractureEditor/Private/FractureToolClustering.cpp b/Engine/Plugins/Experimental/ChaosEditor/Source/FractureEditor/Private/FractureToolClustering.cpp
--- a/Engine/Plugins/Experimental/ChaosEditor/Source/FractureEditor/Private/FractureToolClustering.cpp
+++ b/Engine/Plugins/Experimental/ChaosEditor/Source/FractureEditor/Private/FractureToolClustering.cpp
@@ -36,6 +36,8 @@ void UFractureToolFlattenAll::Execute(TWeakPtr<FFractureEditorModeToolkit> InToo
 {
 	if (InToolkit.IsValid())
 	{
+		FScopedTransaction Transaction(LOCTEXT("FlattenAll", "Flatten All"));
+
 		FFractureEditorModeToolkit* Toolkit = InToolkit.Pin().Get();
 
 		TArray<FFractureToolContext> Contexts = GetFractureToolContexts();
@@ -91,6 +93,8 @@ void UFractureToolCluster::Execute(TWeakPtr<FFractureEditorModeToolkit> InToolki
 	
 	if (InToolkit.IsValid())
 	{
+		FScopedTransaction Transaction(LOCTEXT("Cluster", "Cluster"));
+		
 		FFractureEditorModeToolkit* Toolkit = InToolkit.Pin().Get();
 
 		int32 CurrentLevelView = Toolkit->GetLevelViewValue();
@@ -161,6 +165,8 @@ void UFractureToolUncluster::Execute(TWeakPtr<FFractureEditorModeToolkit> InTool
 {
 	if (InToolkit.IsValid())
 	{
+		FScopedTransaction Transaction(LOCTEXT("Uncluster", "Uncluster"));
+		
 		FFractureEditorModeToolkit* Toolkit = InToolkit.Pin().Get();
 
 		TArray<FFractureToolContext> Contexts = GetFractureToolContexts();
@@ -210,6 +216,8 @@ void UFractureToolMoveUp::Execute(TWeakPtr<FFractureEditorModeToolkit> InToolkit
 {
 	if (InToolkit.IsValid())
 	{
+		FScopedTransaction Transaction(LOCTEXT("MovelUp", "Level Up"));
+		
 		FFractureEditorModeToolkit* Toolkit = InToolkit.Pin().Get();
 
 		TArray<FFractureToolContext> Contexts = GetFractureToolContexts();
@@ -253,6 +261,8 @@ void UFractureToolClusterMerge::Execute(TWeakPtr<FFractureEditorModeToolkit> InT
 {
 	if (InToolkit.IsValid())
 	{
+		FScopedTransaction Transaction(LOCTEXT("ClusterMerge", "Cluster Merge"));
+		
 		FFractureEditorModeToolkit* Toolkit = InToolkit.Pin().Get();
 
 		TArray<FFractureToolContext> Contexts = GetFractureToolContexts();
diff --git a/Engine/Plugins/Experimental/ChaosEditor/Source/FractureEditor/Private/FractureToolEditing.cpp b/Engine/Plugins/Experimental/ChaosEditor/Source/FractureEditor/Private/FractureToolEditing.cpp
--- a/Engine/Plugins/Experimental/ChaosEditor/Source/FractureEditor/Private/FractureToolEditing.cpp
+++ b/Engine/Plugins/Experimental/ChaosEditor/Source/FractureEditor/Private/FractureToolEditing.cpp
@@ -38,6 +38,8 @@ void UFractureToolDeleteBranch::Execute(TWeakPtr<FFractureEditorModeToolkit> InT
 {
 	if (InToolkit.IsValid())
 	{
+		FScopedTransaction Transaction(LOCTEXT("DeleteBranch", "Prune"));
+
 		FFractureEditorModeToolkit* Toolkit = InToolkit.Pin().Get();
 
 		TArray<FFractureToolContext> Contexts = GetFractureToolContexts();
diff --git a/Engine/Plugins/Experimental/ChaosEditor/Source/FractureEditor/Private/GeometryCollectionOutlinerDragDrop.cpp b/Engine/Plugins/Experimental/ChaosEditor/Source/FractureEditor/Private/GeometryCollectionOutlinerDragDrop.cpp
--- a/Engine/Plugins/Experimental/ChaosEditor/Source/FractureEditor/Private/GeometryCollectionOutlinerDragDrop.cpp
+++ b/Engine/Plugins/Experimental/ChaosEditor/Source/FractureEditor/Private/GeometryCollectionOutlinerDragDrop.cpp
@@ -94,6 +94,8 @@ bool FGeometryCollectionBoneDragDrop::ReparentBones(const FGeometryCollection* O
 	bool bValid = ValidateDrop(OtherGeometryCollection, OtherBone, HoverText);
 	if (bValid)
 	{
+		FScopedTransaction Transaction(LOCTEXT("ReparentBones", "Reparent Bones"));
+		
 		// If we parent a rigid to another rigid, we convert the child rigid to embedded.
 		if (GeometryCollection->IsRigid(OtherBone))
 		{
diff --git a/Engine/Source/Runtime/Experimental/GeometryCollectionEngine/Private/GeometryCollection/GeometryCollectionComponent.cpp b/Engine/Source/Runtime/Experimental/GeometryCollectionEngine/Private/GeometryCollection/GeometryCollectionComponent.cpp
--- a/Engine/Source/Runtime/Experimental/GeometryCollectionEngine/Private/GeometryCollection/GeometryCollectionComponent.cpp
+++ b/Engine/Source/Runtime/Experimental/GeometryCollectionEngine/Private/GeometryCollection/GeometryCollectionComponent.cpp
@@ -844,10 +844,19 @@ void UGeometryCollectionComponent::RefreshEmbeddedGeometry()
 	{
 		HideArray = &RestCollection->GetGeometryCollection()->GetAttribute<bool>("Hide", FGeometryCollection::TransformGroup);
 	}
-	
+
+#if WITH_EDITOR	
+	EmbeddedInstanceIndex.Init(INDEX_NONE, RestCollection->GetGeometryCollection()->NumElements(FGeometryCollection::TransformGroup));
+#endif
+
 	const int32 ExemplarCount = EmbeddedGeometryComponents.Num();
 	for (int32 ExemplarIndex = 0; ExemplarIndex < ExemplarCount; ++ExemplarIndex)
 	{		
+#if WITH_EDITOR
+		EmbeddedBoneMaps[ExemplarIndex].Empty(TransformCount);
+		EmbeddedBoneMaps[ExemplarIndex].Reserve(TransformCount); // Allocate for worst case
+#endif
+		
 		TArray<FTransform> InstanceTransforms;
 		InstanceTransforms.Reserve(TransformCount); // Allocate for worst case
 
@@ -859,6 +868,10 @@ void UGeometryCollectionComponent::RefreshEmbeddedGeometry()
 				if (!HideArray || !(*HideArray)[Idx])
 				{ 
 					InstanceTransforms.Add(FTransform(GlobalMatrices[Idx]));
+#if WITH_EDITOR
+					int32 InstanceIndex = EmbeddedBoneMaps[ExemplarIndex].Add(Idx);
+					EmbeddedInstanceIndex[Idx] = InstanceIndex;
+#endif
 				}
 			}
 		}
@@ -890,6 +903,30 @@ void UGeometryCollectionComponent::RefreshEmbeddedGeometry()
 	}
 }
 
+#if WITH_EDITOR
+void UGeometryCollectionComponent::SetEmbeddedGeometrySelectable(bool bSelectableIn)
+{
+	for (TObjectPtr<UInstancedStaticMeshComponent> EmbeddedGeometryComponent : EmbeddedGeometryComponents)
+	{
+		EmbeddedGeometryComponent->bSelectable = bSelectable;
+		EmbeddedGeometryComponent->bHasPerInstanceHitProxies = bSelectable;
+	}
+}
+
+int32 UGeometryCollectionComponent::EmbeddedIndexToTransformIndex(const UInstancedStaticMeshComponent* ISMComponent, int32 InstanceIndex) const
+{
+	for (int32 ISMIdx = 0; ISMIdx < EmbeddedGeometryComponents.Num(); ++ISMIdx)
+	{
+		if (EmbeddedGeometryComponents[ISMIdx].Get() == ISMComponent)
+		{
+			return (EmbeddedBoneMaps[ISMIdx][InstanceIndex]);
+		}
+	}
+
+	return INDEX_NONE;
+}
+#endif
+
 void UGeometryCollectionComponent::SetRestState(TArray<FTransform>&& InRestTransforms)
 {	
 	RestTransforms = InRestTransforms;
@@ -2218,6 +2255,7 @@ void FScopedColorEdit::SetSelectedBones(const TArray<int32>& SelectedBonesIn)
 {
 	bUpdated = true;
 	Component->SelectedBones = SelectedBonesIn;
+	Component->SelectEmbeddedGeometry();
 }
 
 void FScopedColorEdit::AppendSelectedBones(const TArray<int32>& SelectedBonesIn)
@@ -2848,6 +2886,26 @@ UMaterialInterface* UGeometryCollectionComponent::GetMaterial(int32 MaterialInde
 	}
 }
 
+#if WITH_EDITOR
+void UGeometryCollectionComponent::SelectEmbeddedGeometry()
+{
+	// First reset the selections
+	for (TObjectPtr<UInstancedStaticMeshComponent> EmbeddedGeometryComponent : EmbeddedGeometryComponents)
+	{
+		EmbeddedGeometryComponent->ClearInstanceSelection();
+	}
+	
+	const TManagedArray<int32>& ExemplarIndex = GetExemplarIndexArray();
+	for (int32 SelectedBone : SelectedBones)
+	{
+		if (ExemplarIndex[SelectedBone] > INDEX_NONE)
+		{
+			EmbeddedGeometryComponents[ExemplarIndex[SelectedBone]]->SelectInstance(true, EmbeddedInstanceIndex[SelectedBone], 1);
+		}
+	}
+}
+#endif
+
 // #temp HACK for demo, When fracture happens (physics state changes to dynamic) then switch the visible render meshes in a blueprint/actor from static meshes to geometry collections
 void UGeometryCollectionComponent::SwitchRenderModels(const AActor* Actor)
 {
@@ -2986,6 +3044,11 @@ void UGeometryCollectionComponent::InitializeEmbeddedGeometry()
 			}
 		}
 
+#if WITH_EDITOR
+		EmbeddedBoneMaps.SetNum(RestCollection->EmbeddedGeometryExemplar.Num());
+		EmbeddedInstanceIndex.Init(INDEX_NONE,RestCollection->GetGeometryCollection()->NumElements(FGeometryCollection::TransformGroup));
+#endif
+
 		CalculateGlobalMatrices();
 		RefreshEmbeddedGeometry();
 	}
diff --git a/Engine/Source/Runtime/Experimental/GeometryCollectionEngine/Public/GeometryCollection/GeometryCollectionComponent.h b/Engine/Source/Runtime/Experimental/GeometryCollectionEngine/Public/GeometryCollection/GeometryCollectionComponent.h
--- a/Engine/Source/Runtime/Experimental/GeometryCollectionEngine/Public/GeometryCollection/GeometryCollectionComponent.h
+++ b/Engine/Source/Runtime/Experimental/GeometryCollectionEngine/Public/GeometryCollection/GeometryCollectionComponent.h
@@ -387,6 +387,9 @@ public:
 	FORCEINLINE FGeometryCollectionEdit EditRestCollection(GeometryCollection::EEditUpdate EditUpdate = GeometryCollection::EEditUpdate::RestPhysics) { return FGeometryCollectionEdit(this, EditUpdate); }
 #if WITH_EDITOR
 	FORCEINLINE FScopedColorEdit EditBoneSelection() { return FScopedColorEdit(this); }
+
+	/** Propagate bone selection to embedded geometry components. */
+	void SelectEmbeddedGeometry();
 #endif
 
 	/** API for getting at geometry collection data */
@@ -660,6 +663,12 @@ public:
 	/** Update instanced static mesh components to reflect internal embedded geometry state. */
 	void RefreshEmbeddedGeometry();
 
+#if WITH_EDITOR
+	void SetEmbeddedGeometrySelectable(bool bSelectableIn);
+	int32 EmbeddedIndexToTransformIndex(const UInstancedStaticMeshComponent* ISMComponent, int32 InstanceIndex) const;
+#endif
+
+	
 	// #todo should this only be available in editor?
 	void SetRestState(TArray<FTransform>&& InRestTransforms);
 
@@ -825,6 +834,11 @@ private:
 	UPROPERTY()
 	TArray<TObjectPtr<UInstancedStaticMeshComponent>> EmbeddedGeometryComponents;
 
+#if WITH_EDITORONLY_DATA
+	TArray<TArray<int32>> EmbeddedBoneMaps;
+	TArray<int32> EmbeddedInstanceIndex;
+#endif
+
 	bool IsEmbeddedGeometryValid() const;
 	void ClearEmbeddedGeometry();
 
