Chaos : Collision Constraints
- Add contact constraints to transfer joint collisions from child to parent.

diff --git a/Engine/Source/Runtime/Experimental/Chaos/Private/Chaos/Collision/CollisionSolver.cpp b/Engine/Source/Runtime/Experimental/Chaos/Private/Chaos/Collision/CollisionSolver.cpp
--- a/Engine/Source/Runtime/Experimental/Chaos/Private/Chaos/Collision/CollisionSolver.cpp
+++ b/Engine/Source/Runtime/Experimental/Chaos/Private/Chaos/Collision/CollisionSolver.cpp
@@ -51,12 +51,28 @@ namespace Chaos
 		FAutoConsoleVariableRef CVarChaos_Manifold_PushOut_MinStiffness(TEXT("p.Chaos.Collision.Manifold.PushOut.MinStiffness"), Chaos_Manifold_PushOut_MinStiffness, TEXT(""));
 		FAutoConsoleVariableRef CVarChaos_Manifold_PushOut_MaxStiffness(TEXT("p.Chaos.Collision.Manifold.PushOut.MaxStiffness"), Chaos_Manifold_PushOut_MaxStiffness, TEXT(""));
 
+		FReal Chaos_ConstraintStiffness_Modifier = 5.0;
+		FAutoConsoleVariableRef CVarChaos_ConstraintStiffness_Modifier(TEXT("p.Chaos.Collision.ConstraintStiffnessModifier"), Chaos_ConstraintStiffness_Modifier, TEXT(""));
+
 		FRealSingle Chaos_Manifold_Apply_ImpulseTolerance = 1.e-4f;
 		FAutoConsoleVariableRef CVarChaos_Manifold_Apply_ImpulseTolerance(TEXT("p.Chaos.Collision.Manifold.Apply.ImpulseTolerance"), Chaos_Manifold_Apply_ImpulseTolerance, TEXT(""));
 
 		float Chaos_Manifold_PushOut_PositionTolerance = 1.e-4f;
 		FAutoConsoleVariableRef CVarChaos_Manifold_PushOut_PositionTolerance(TEXT("p.Chaos.Collision.Manifold.PushOut.PositionTolerance"), Chaos_Manifold_PushOut_PositionTolerance, TEXT(""));
 
+		inline FReal GetConstraintStiffness(const FReal InStiffness, const FContactIterationParameters& IterationParameters)
+		{
+			FReal Stiffness = FMath::Clamp(InStiffness, (FReal)0.0, (FReal)1.0);
+			if (Stiffness < 1.0)
+			{
+				FReal Factor = Chaos_ConstraintStiffness_Modifier * (FReal)(IterationParameters.NumIterations * IterationParameters.NumPairIterations);
+				if (Factor)
+				{
+					return Stiffness / Factor;
+				}
+			}
+			return (FReal)1;			
+		}
 
 		inline FReal GetApplyStiffness(int32 It, int32 NumIts)
 		{
@@ -584,7 +600,8 @@ namespace Chaos
 
 			Constraint.AccumulatedImpulse = FVec3(0);
 
-			const FReal Stiffness = GetApplyStiffness(IterationParameters.Iteration, IterationParameters.NumIterations);
+			const FReal ConstraintStiffness = GetConstraintStiffness(Constraint.GetStiffness(), IterationParameters);
+			const FReal Stiffness = ConstraintStiffness*GetApplyStiffness(IterationParameters.Iteration, IterationParameters.NumIterations);
 
 			// Iterate over the manifold and accumulate velocity corrections - we will apply them after the loop
 			TArrayView<FManifoldPoint> ManifoldPoints = Constraint.GetManifoldPoints();
@@ -713,8 +730,9 @@ namespace Chaos
 			}
 
 			// Gradually increase position correction through iterations (optional based on cvars)
-			const FReal Stiffness = GetPushOutStiffness(IterationParameters.Iteration, IterationParameters.NumIterations);
-
+			const FReal ConstraintStiffness = GetConstraintStiffness(Constraint.GetStiffness(), IterationParameters);
+			const FReal Stiffness = ConstraintStiffness*GetPushOutStiffness(IterationParameters.Iteration, IterationParameters.NumIterations);
+			
 			FVec3 P0 = FParticleUtilities::GetCoMWorldPosition(Particle0);
 			FRotation3 Q0 = FParticleUtilities::GetCoMWorldRotation(Particle0);
 			FVec3 P1 = FParticleUtilities::GetCoMWorldPosition(Particle1);
diff --git a/Engine/Source/Runtime/Experimental/Chaos/Private/Chaos/PBDRigidsEvolutionGBF.cpp b/Engine/Source/Runtime/Experimental/Chaos/Private/Chaos/PBDRigidsEvolutionGBF.cpp
--- a/Engine/Source/Runtime/Experimental/Chaos/Private/Chaos/PBDRigidsEvolutionGBF.cpp
+++ b/Engine/Source/Runtime/Experimental/Chaos/Private/Chaos/PBDRigidsEvolutionGBF.cpp
@@ -22,8 +22,6 @@
 #include "Misc/Paths.h"
 #include "HAL/FileManager.h"
 
-//PRAGMA_DISABLE_OPTIMIZATION
-
 namespace Chaos
 {
 #if !UE_BUILD_SHIPPING
@@ -69,6 +67,11 @@ FAutoConsoleVariableRef CVarChaosSolverJointPriority(TEXT("p.Chaos.Solver.Joint.
 int32 ChaosSolverSuspensionPriority = 0;
 FAutoConsoleVariableRef CVarChaosSolverSuspensionPriority(TEXT("p.Chaos.Solver.Suspension.Priority"), ChaosSolverSuspensionPriority, TEXT("Set constraint priority. Larger values are evaluated later [def:0]"));
 
+bool DoTransferJointConstraintCollisions = true;
+FAutoConsoleVariableRef CVarDoTransferJointConstraintCollisions(TEXT("p.Chaos.Solver.Joint.TransferCollisions"), DoTransferJointConstraintCollisions, TEXT("Allows joints to apply collisions to the parent from the child when the Joints TransferCollisionScale is not 0 [def:true]"));
+
+
+
 DECLARE_CYCLE_STAT(TEXT("FPBDRigidsEvolutionGBF::AdvanceOneTimeStep"), STAT_Evolution_AdvanceOneTimeStep, STATGROUP_Chaos);
 DECLARE_CYCLE_STAT(TEXT("FPBDRigidsEvolutionGBF::UnclusterUnions"), STAT_Evolution_UnclusterUnions, STATGROUP_Chaos);
 DECLARE_CYCLE_STAT(TEXT("FPBDRigidsEvolutionGBF::Integrate"), STAT_Evolution_Integrate, STATGROUP_Chaos);
@@ -81,6 +84,7 @@ DECLARE_CYCLE_STAT(TEXT("FPBDRigidsEvolutionGBF::ApplyConstraints"), STAT_Evolut
 DECLARE_CYCLE_STAT(TEXT("FPBDRigidsEvolutionGBF::UpdateVelocities"), STAT_Evolution_UpdateVelocites, STATGROUP_Chaos);
 DECLARE_CYCLE_STAT(TEXT("FPBDRigidsEvolutionGBF::ApplyPushOut"), STAT_Evolution_ApplyPushOut, STATGROUP_Chaos);
 DECLARE_CYCLE_STAT(TEXT("FPBDRigidsEvolutionGBF::DetectCollisions"), STAT_Evolution_DetectCollisions, STATGROUP_Chaos);
+DECLARE_CYCLE_STAT(TEXT("FPBDRigidsEvolutionGBF::TransferJointCollisions"), STAT_Evolution_TransferJointCollisions, STATGROUP_Chaos);
 DECLARE_CYCLE_STAT(TEXT("FPBDRigidsEvolutionGBF::PostDetectCollisionsCallback"), STAT_Evolution_PostDetectCollisionsCallback, STATGROUP_Chaos);
 DECLARE_CYCLE_STAT(TEXT("FPBDRigidsEvolutionGBF::UpdateConstraintPositionBasedState"), STAT_Evolution_UpdateConstraintPositionBasedState, STATGROUP_Chaos);
 DECLARE_CYCLE_STAT(TEXT("FPBDRigidsEvolutionGBF::ComputeIntermediateSpatialAcceleration"), STAT_Evolution_ComputeIntermediateSpatialAcceleration, STATGROUP_Chaos);
@@ -231,6 +235,12 @@ void FPBDRigidsEvolutionGBF::AdvanceOneTimeStepImpl(const FReal Dt,const FSubSte
 		CHAOS_COLLISION_STAT(StatData.Print());
 	}
 
+	{
+		SCOPE_CYCLE_COUNTER(STAT_Evolution_TransferJointCollisions);
+
+		TransferJointConstraintCollisions();
+	}
+
 	if (PostDetectCollisionsCallback != nullptr)
 	{
 		SCOPE_CYCLE_COUNTER(STAT_Evolution_PostDetectCollisionsCallback);
@@ -520,5 +530,96 @@ void FPBDRigidsEvolutionGBF::SetCurrentStepResimCache(IResimCacheBase* InCurrent
 	CurrentStepResimCacheImp = static_cast<FEvolutionResimCache*>(InCurrentStepResimCache);
 }
 
+
+void FPBDRigidsEvolutionGBF::TransferJointConstraintCollisions()
+{
+	//
+	// Append transfer constraints. 
+	//
+	using FRigidHandle = TPBDRigidParticleHandleImp<FReal, 3, true>;
+	if (DoTransferJointConstraintCollisions)
+	{
+		TMap<TGeometryParticleHandle<FReal, 3>*, TArray<FPBDCollisionConstraintHandle*> >& CollisionMap = CollisionConstraints.GetParticleCollisionsMap();
+		for (auto& CollisionPair : CollisionMap)
+		{
+			if (FRigidHandle* ChildParticle = CollisionPair.Key->CastToRigidParticle())
+			{
+				for (FConstraintHandle* Constraint : ChildParticle->ParticleConstraints())
+				{
+					if (FPBDJointConstraintHandle* JointConstraint = Constraint->As<FPBDJointConstraintHandle>())
+					{
+						if (!FMath::IsNearlyZero(JointConstraint->GetSettings().ContactTransferScale))
+						{
+							if (FRigidHandle* ParentParticle = JointConstraint->GetConstrainedParticles()[1]->CastToRigidParticle()) // Parent
+							{
+								for (auto& CollisionHandel : CollisionPair.Value)
+								{
+									if (FPBDCollisionConstraintHandle* ContactHandle = CollisionHandel->As<FPBDCollisionConstraintHandle>())
+									{
+										if (ContactHandle->GetType() == FCollisionConstraintBase::FType::SinglePoint)
+										{
+											FRigidBodyPointContactConstraint CurrConstraint = ContactHandle->GetPointContact();
+											FTransform ParentTransform(ParentParticle->R(), ParentParticle->X()), ChildTransform(ChildParticle->R(), ChildParticle->X()),
+												ChildToParentTransform = ChildTransform.GetRelativeTransform(ParentTransform);
+
+											FRigidBodyPointContactConstraint TransferedConstraint;
+											if (ContactHandle->GetConstrainedParticles()[0] == ChildParticle)
+											{
+												TransferedConstraint = FRigidBodyPointContactConstraint(
+													ParentParticle,
+													CurrConstraint.GetManifold().Implicit[0],
+													CurrConstraint.GetManifold().Simplicial[0],
+													ChildToParentTransform.GetRelativeTransform(CurrConstraint.ImplicitTransform[0]),
+													CurrConstraint.Particle[1],
+													CurrConstraint.GetManifold().Implicit[1],
+													CurrConstraint.GetManifold().Simplicial[1],
+													CurrConstraint.ImplicitTransform[1],
+													CurrConstraint.GetCullDistance(),
+													CurrConstraint.GetManifold().ShapesType,
+													CurrConstraint.GetUseManifold()
+												);
+											}
+											else
+											{
+												TransferedConstraint = FRigidBodyPointContactConstraint(
+													CurrConstraint.Particle[0],
+													CurrConstraint.GetManifold().Implicit[0],
+													CurrConstraint.GetManifold().Simplicial[0],
+													CurrConstraint.ImplicitTransform[0],
+													ParentParticle,
+													CurrConstraint.GetManifold().Implicit[1],
+													CurrConstraint.GetManifold().Simplicial[1],
+													ChildToParentTransform.GetRelativeTransform(CurrConstraint.ImplicitTransform[1]),
+													CurrConstraint.GetCullDistance(),
+													CurrConstraint.GetManifold().ShapesType,
+													CurrConstraint.GetUseManifold()
+												);
+											}
+											TArray<FManifoldPoint> TransferManifolds;
+											for (FManifoldPoint& Manifold : CurrConstraint.GetManifoldPoints())
+											{
+												TransferManifolds.Add(Manifold);
+											}
+											TransferedConstraint.SetManifoldPoints(TransferManifolds);
+											for (FManifoldPoint& Manifold : CurrConstraint.GetManifoldPoints())
+											{
+												TransferedConstraint.UpdateManifoldPointFromContact(Manifold);
+											}
+											TransferedConstraint.SetStiffness(JointConstraint->GetSettings().ContactTransferScale);
+											CollisionConstraints.AddConstraint(TransferedConstraint);
+
+										}
+									}
+								}
+							}
+						}
+					}
+				}
+			}
+		}
+	}
+}
+
+
 }
 
diff --git a/Engine/Source/Runtime/Experimental/Chaos/Public/Chaos/Collision/PBDCollisionConstraint.h b/Engine/Source/Runtime/Experimental/Chaos/Public/Chaos/Collision/PBDCollisionConstraint.h
--- a/Engine/Source/Runtime/Experimental/Chaos/Public/Chaos/Collision/PBDCollisionConstraint.h
+++ b/Engine/Source/Runtime/Experimental/Chaos/Public/Chaos/Collision/PBDCollisionConstraint.h
@@ -154,6 +154,7 @@ namespace Chaos
 			, Timestamp(-INT_MAX)
 			, ConstraintHandle(nullptr)
 			, Type(InType)
+			, Stiffness(1.0)
 		{
 			ImplicitTransform[0] = FRigidTransform3::Identity; ImplicitTransform[1] = FRigidTransform3::Identity;
 			Manifold.Implicit[0] = nullptr; Manifold.Implicit[1] = nullptr;
@@ -177,6 +178,7 @@ namespace Chaos
 			, Timestamp(InTimestamp)
 			, ConstraintHandle(nullptr)
 			, Type(InType)
+			, Stiffness(1.0)
 		{
 			ImplicitTransform[0] = Transform0; 
 			ImplicitTransform[1] = Transform1;
@@ -235,6 +237,8 @@ namespace Chaos
 		void SetInvInertiaScale1(const FReal InInvInertiaScale) { Manifold.InvInertiaScale1 = InInvInertiaScale; }
 		FReal GetInvInertiaScale1() const { return Manifold.InvInertiaScale1; }
 
+		void SetStiffness(FReal InStiffness) { Stiffness = InStiffness; }
+		FReal GetStiffness() const { return Stiffness; }
 
 		FString ToString() const;
 
@@ -252,6 +256,7 @@ namespace Chaos
 	private:
 		FPBDCollisionConstraintHandle* ConstraintHandle;
 		FType Type;
+		FReal Stiffness;
 	};
 
 
@@ -296,6 +301,7 @@ namespace Chaos
 
 		bool GetUseManifold() const { return bUseManifold; }
 
+		void SetManifoldPoints(const TArray<FManifoldPoint>& InManifoldPoints) { ManifoldPoints=InManifoldPoints; }
 		TArrayView<FManifoldPoint> GetManifoldPoints() { return MakeArrayView(ManifoldPoints); }
 		TArrayView<const FManifoldPoint> GetManifoldPoints() const { return MakeArrayView(ManifoldPoints); }
 		FManifoldPoint& SetActiveManifoldPoint(
@@ -317,6 +323,9 @@ namespace Chaos
 		void UpdateManifoldContacts(FReal Dt);
 		void ClearManifold();
 
+		//@todo(make protected) : Was exposed to support contact copying. 
+		void UpdateManifoldPointFromContact(FManifoldPoint& ManifoldPoint);
+
 	protected:
 		// For use by derived types that can be used as point constraints in Update
 		FRigidBodyPointContactConstraint(typename Base::FType InType) : Base(InType) {}
@@ -347,7 +356,6 @@ namespace Chaos
 		int32 AddManifoldPoint(const FContactPoint& ContactPoint, const FReal Dt);
 		void InitManifoldPoint(FManifoldPoint& ManifoldPoint, FReal Dt);
 		void UpdateManifoldPoint(int32 ManifoldPointIndex, const FContactPoint& ContactPoint, const FReal Dt);
-		void UpdateManifoldPointFromContact(FManifoldPoint& ManifoldPoint);
 		void SetActiveContactPoint(const FContactPoint& ContactPoint);
 		void GetWorldSpaceManifoldPoint(const FManifoldPoint& ManifoldPoint, const FVec3& P0, const FRotation3& Q0, const FVec3& P1, const FRotation3& Q1, FVec3& OutContactLocation, FVec3& OutContactNormal, FReal& OutContactPhi);
 
diff --git a/Engine/Source/Runtime/Experimental/Chaos/Public/Chaos/PBDRigidsEvolutionGBF.h b/Engine/Source/Runtime/Experimental/Chaos/Public/Chaos/PBDRigidsEvolutionGBF.h
--- a/Engine/Source/Runtime/Experimental/Chaos/Public/Chaos/PBDRigidsEvolutionGBF.h
+++ b/Engine/Source/Runtime/Experimental/Chaos/Public/Chaos/PBDRigidsEvolutionGBF.h
@@ -258,6 +258,8 @@ namespace Chaos
 
 		CHAOS_API FSpatialAccelerationBroadPhase& GetBroadPhase() { return BroadPhase; }
 
+		CHAOS_API void TransferJointConstraintCollisions();
+
 	protected:
 
 		CHAOS_API void AdvanceOneTimeStepImpl(const FReal dt, const FSubStepInfo& SubStepInfo);
