Chaos - Fix CIS

diff --git a/Engine/Source/Runtime/Experimental/Chaos/Private/Chaos/Collision/CollisionSolver.cpp b/Engine/Source/Runtime/Experimental/Chaos/Private/Chaos/Collision/CollisionSolver.cpp
--- a/Engine/Source/Runtime/Experimental/Chaos/Private/Chaos/Collision/CollisionSolver.cpp
+++ b/Engine/Source/Runtime/Experimental/Chaos/Private/Chaos/Collision/CollisionSolver.cpp
@@ -65,8 +65,8 @@ namespace Chaos
 			FReal Stiffness = FMath::Clamp(InStiffness, (FReal)0.0, (FReal)1.0);
 			if (Stiffness < 1.0)
 			{
-				FReal Factor = Chaos_ConstraintStiffness_Modifier * (FReal)(IterationParameters.NumIterations * IterationParameters.NumPairIterations);
-				if (Factor)
+				FReal Factor = Chaos_ConstraintStiffness_Modifier * (FReal)IterationParameters.NumIterations * (FReal)IterationParameters.NumPairIterations;
+				if (!FMath::IsNearlyZero(Factor))
 				{
 					return Stiffness / Factor;
 				}
