Fixing CIS issue
#rnx

diff --git a/Engine/Source/Runtime/Experimental/Chaos/Private/Chaos/Collision/CollisionSolver.cpp b/Engine/Source/Runtime/Experimental/Chaos/Private/Chaos/Collision/CollisionSolver.cpp
--- a/Engine/Source/Runtime/Experimental/Chaos/Private/Chaos/Collision/CollisionSolver.cpp
+++ b/Engine/Source/Runtime/Experimental/Chaos/Private/Chaos/Collision/CollisionSolver.cpp
@@ -51,7 +51,7 @@ namespace Chaos
 		FAutoConsoleVariableRef CVarChaos_Manifold_PushOut_MinStiffness(TEXT("p.Chaos.Collision.Manifold.PushOut.MinStiffness"), Chaos_Manifold_PushOut_MinStiffness, TEXT(""));
 		FAutoConsoleVariableRef CVarChaos_Manifold_PushOut_MaxStiffness(TEXT("p.Chaos.Collision.Manifold.PushOut.MaxStiffness"), Chaos_Manifold_PushOut_MaxStiffness, TEXT(""));
 
-		FReal Chaos_ConstraintStiffness_Modifier = 5.0;
+		float Chaos_ConstraintStiffness_Modifier = 5.0;
 		FAutoConsoleVariableRef CVarChaos_ConstraintStiffness_Modifier(TEXT("p.Chaos.Collision.ConstraintStiffnessModifier"), Chaos_ConstraintStiffness_Modifier, TEXT(""));
 
 		FRealSingle Chaos_Manifold_Apply_ImpulseTolerance = 1.e-4f;
@@ -65,7 +65,7 @@ namespace Chaos
 			FReal Stiffness = FMath::Clamp(InStiffness, (FReal)0.0, (FReal)1.0);
 			if (Stiffness < 1.0)
 			{
-				FReal Factor = Chaos_ConstraintStiffness_Modifier * (FReal)IterationParameters.NumIterations * (FReal)IterationParameters.NumPairIterations;
+				FReal Factor = (FReal)Chaos_ConstraintStiffness_Modifier * (FReal)IterationParameters.NumIterations * (FReal)IterationParameters.NumPairIterations;
 				if (!FMath::IsNearlyZero(Factor))
 				{
 					return Stiffness / Factor;
