Joint uses push data to marshal data so that it is async thread safe

diff --git a/Engine/Source/Runtime/Experimental/Chaos/Private/Chaos/PBDJointConstraintData.cpp b/Engine/Source/Runtime/Experimental/Chaos/Private/Chaos/PBDJointConstraintData.cpp
--- a/Engine/Source/Runtime/Experimental/Chaos/Private/Chaos/PBDJointConstraintData.cpp
+++ b/Engine/Source/Runtime/Experimental/Chaos/Private/Chaos/PBDJointConstraintData.cpp
@@ -7,7 +7,6 @@ namespace Chaos
 {
 	FJointConstraint::FJointConstraint()
 		: FConstraintBase(EConstraintType::JointConstraintType)
-		, UserData(nullptr)
 		, KinematicEndPoint(nullptr)
 	{
 	}
@@ -19,14 +18,6 @@ namespace Chaos
 		Solver->RegisterObject(KinematicEndPoint);
 	}
 
-	const FJointConstraint::FTransformPair FJointConstraint::GetJointTransforms() const { return JointSettings.ConnectorTransforms; }
-	void FJointConstraint::SetJointTransforms(const Chaos::FJointConstraint::FTransformPair& InJointTransforms)
-	{
-		JointSettings.ConnectorTransforms = InJointTransforms;
-		MDirtyFlags.MarkDirty(EJointConstraintFlags::JointTransforms);
-		SetProxy(Proxy);
-	}
-
 	void FJointConstraint::SetLinearPositionDriveEnabled(TVector<bool,3> Enabled)
 	{
 		SetLinearPositionDriveXEnabled(Enabled.X);
diff --git a/Engine/Source/Runtime/Experimental/Chaos/Private/PBDRigidsSolver.cpp b/Engine/Source/Runtime/Experimental/Chaos/Private/PBDRigidsSolver.cpp
--- a/Engine/Source/Runtime/Experimental/Chaos/Private/PBDRigidsSolver.cpp
+++ b/Engine/Source/Runtime/Experimental/Chaos/Private/PBDRigidsSolver.cpp
@@ -601,7 +601,8 @@ namespace Chaos
 		GTConstraint->ReleaseKinematicEndPoint(this);
 
 		FParticlesType* InParticles = &GetParticles();
-
+		JointProxy->DestroyOnGameThread();	//destroy the game thread portion of the proxy
+		
 		// Finish registration on the physics thread...
 		EnqueueCommandImmediate([InParticles, JointProxy, this]()
 			{
@@ -827,28 +828,22 @@ namespace Chaos
 		Manager->PrepareBuckets(DirtyProxiesData->GetDirtyProxyBucketInfo());
 		Manager->SetNumShapes(DirtyProxiesData->NumDirtyShapes());
 		FShapeDirtyData* ShapeDirtyData = DirtyProxiesData->GetShapesDirtyData();
-		auto ProcessProxyGT =[ShapeDirtyData, Manager, DirtyProxiesData](auto Proxy, int32 ParticleDataIdx, FDirtyProxy& DirtyProxy)
-		{
-			auto Particle = Proxy->GetParticle_LowLevel();
-			Particle->SyncRemoteData(*Manager,ParticleDataIdx,DirtyProxy.PropertyData,DirtyProxy.ShapeDataIndices,ShapeDirtyData);
-			Proxy->ClearAccumulatedData();
-			Proxy->ResetDirtyIdx();
-		};
-
 
-		//todo: if we allocate remote data ahead of time we could go wide
-		DirtyProxiesData->ParallelForEachProxy([&ProcessProxyGT, this, DynamicsWeight](int32 DataIdx, FDirtyProxy& Dirty)
+		DirtyProxiesData->ParallelForEachProxy([this, DynamicsWeight, Manager, ShapeDirtyData](int32 DataIdx, FDirtyProxy& Dirty)
 		{
 			switch(Dirty.Proxy->GetType())
 			{
 			case EPhysicsProxyType::SingleParticleProxy:
 			{
 				auto Proxy = static_cast<FSingleParticlePhysicsProxy*>(Dirty.Proxy);
-				if(auto Rigid = Proxy->GetParticle_LowLevel()->CastToRigidParticle())
+				auto Particle = Proxy->GetParticle_LowLevel();
+				if(auto Rigid = Particle->CastToRigidParticle())
 				{
 					Rigid->ApplyDynamicsWeight(DynamicsWeight);
 				}
-				ProcessProxyGT(Proxy,DataIdx,Dirty);
+				Particle->SyncRemoteData(*Manager, DataIdx, Dirty.PropertyData, Dirty.ShapeDataIndices, ShapeDirtyData);
+				Proxy->ClearAccumulatedData();
+				Proxy->ResetDirtyIdx();
 				break;
 			}
 			case EPhysicsProxyType::GeometryCollectionType:
@@ -860,7 +855,8 @@ namespace Chaos
 			case EPhysicsProxyType::JointConstraintType:
 			{
 				auto Proxy = static_cast<FJointConstraintPhysicsProxy*>(Dirty.Proxy);
-				Proxy->PushStateOnGameThread(this);
+				Proxy->PushStateOnGameThread(*Manager, DataIdx, Dirty.PropertyData);
+				Proxy->ResetDirtyIdx();
 				break;
 			}
 			case EPhysicsProxyType::SuspensionConstraintType:
@@ -1024,7 +1020,7 @@ namespace Chaos
 		});
 
 		//need to create new constraint handles
-		DirtyProxiesData->ForEachProxy([this, &ProcessProxyPT](int32 DataIdx, FDirtyProxy& Dirty)
+		DirtyProxiesData->ForEachProxy([this, &ProcessProxyPT, Manager](int32 DataIdx, FDirtyProxy& Dirty)
 		{
 			switch (Dirty.Proxy->GetType())
 			{
@@ -1035,10 +1031,10 @@ namespace Chaos
 				if (bIsNew)
 				{
 					JointConstraintPhysicsProxies_Internal.Add(JointProxy);
-					JointProxy->InitializeOnPhysicsThread(this);
+					JointProxy->InitializeOnPhysicsThread(this, *Manager, DataIdx, Dirty.PropertyData);
 					JointProxy->SetInitialized();
 				}
-				JointProxy->PushStateOnPhysicsThread(this);
+				JointProxy->PushStateOnPhysicsThread(this, *Manager, DataIdx, Dirty.PropertyData);
 				Dirty.Proxy->ResetDirtyIdx();
 				break;
 			}
diff --git a/Engine/Source/Runtime/Experimental/Chaos/Private/PhysicsProxy/JointConstraintProxy.cpp b/Engine/Source/Runtime/Experimental/Chaos/Private/PhysicsProxy/JointConstraintProxy.cpp
--- a/Engine/Source/Runtime/Experimental/Chaos/Private/PhysicsProxy/JointConstraintProxy.cpp
+++ b/Engine/Source/Runtime/Experimental/Chaos/Private/PhysicsProxy/JointConstraintProxy.cpp
@@ -20,12 +20,11 @@ namespace Chaos
 
 FJointConstraintPhysicsProxy::FJointConstraintPhysicsProxy(FJointConstraint* InConstraint, FPBDJointConstraintHandle* InHandle, UObject* InOwner)
 	: Base(EPhysicsProxyType::JointConstraintType, InOwner)
-	, Constraint(InConstraint) // This proxy assumes ownership of the Constraint, and will free it during DestroyOnPhysicsThread
-	, Handle(InHandle)
+	, Constraint_GT(InConstraint) // This proxy assumes ownership of the Constraint, and will free it during DestroyOnPhysicsThread
+	, Constraint_PT(InHandle)
 {
-	check(Constraint!=nullptr);
-	Constraint->SetProxy(this);
-	JointSettingsBuffer = Constraint->GetJointSettings();
+	check(Constraint_GT !=nullptr);
+	Constraint_GT->SetProxy(this);
 }
 
 FGeometryParticleHandle*
@@ -45,109 +44,101 @@ FJointConstraintPhysicsProxy::GetParticleHandleFromProxy(IPhysicsProxyBase* Prox
 void FJointConstraintPhysicsProxy::BufferPhysicsResults(FDirtyJointConstraintData& Buffer)
 {
 	Buffer.SetProxy(*this);
-	if (Constraint != nullptr && Constraint->IsValid() )
+	if (Constraint_PT != nullptr && (Constraint_PT->IsValid() || Constraint_PT->IsConstraintBreaking() || Constraint_PT->IsDriveTargetChanged()))
 	{
-		if (Handle != nullptr && (Handle->IsValid() || Handle->IsConstraintBreaking() || Handle->IsDriveTargetChanged()))
-		{
-			Buffer.OutputData.bIsBreaking = Handle->IsConstraintBreaking();
-			Buffer.OutputData.bIsBroken = !Handle->IsConstraintEnabled();
-			Buffer.OutputData.bDriveTargetChanged = Handle->IsDriveTargetChanged();
-			Buffer.OutputData.Force = Handle->GetLinearImpulse();
-			Buffer.OutputData.Torque = Handle->GetAngularImpulse();
-
-			Handle->ClearConstraintBreaking(); // it's a single frame event, so reset
-			Handle->ClearDriveTargetChanged(); // it's a single frame event, so reset
-		}
+		Buffer.OutputData.bIsBreaking = Constraint_PT->IsConstraintBreaking();
+		Buffer.OutputData.bIsBroken = !Constraint_PT->IsConstraintEnabled();
+		Buffer.OutputData.bDriveTargetChanged = Constraint_PT->IsDriveTargetChanged();
+		Buffer.OutputData.Force = Constraint_PT->GetLinearImpulse();
+		Buffer.OutputData.Torque = Constraint_PT->GetAngularImpulse();
+
+		Constraint_PT->ClearConstraintBreaking(); // it's a single frame event, so reset
+		Constraint_PT->ClearDriveTargetChanged(); // it's a single frame event, so reset
 	}
 }
 
 /**/
 bool FJointConstraintPhysicsProxy::PullFromPhysicsState(const FDirtyJointConstraintData& Buffer, const int32 SolverSyncTimestamp)
 {
-	if (Constraint != nullptr && Constraint->IsValid())
+	if (Constraint_GT != nullptr && Constraint_GT->IsValid())
 	{
-		if (Handle != nullptr && (Handle->IsValid() || Buffer.OutputData.bIsBreaking || Buffer.OutputData.bDriveTargetChanged))
+		if (Buffer.OutputData.bIsBreaking || Buffer.OutputData.bDriveTargetChanged)
 		{
-			Constraint->GetOutputData().bIsBreaking = Buffer.OutputData.bIsBreaking;
-			Constraint->GetOutputData().bIsBroken = Buffer.OutputData.bIsBroken;
-			Constraint->GetOutputData().bDriveTargetChanged = Buffer.OutputData.bDriveTargetChanged;
-			Constraint->GetOutputData().Force = Buffer.OutputData.Force;
-			Constraint->GetOutputData().Torque = Buffer.OutputData.Torque;
+			Constraint_GT->GetOutputData().bIsBreaking = Buffer.OutputData.bIsBreaking;
+			Constraint_GT->GetOutputData().bIsBroken = Buffer.OutputData.bIsBroken;
+			Constraint_GT->GetOutputData().bDriveTargetChanged = Buffer.OutputData.bDriveTargetChanged;
+			Constraint_GT->GetOutputData().Force = Buffer.OutputData.Force;
+			Constraint_GT->GetOutputData().Torque = Buffer.OutputData.Torque;
 		}
 	}
 
 	return true;
 }
 
-void FJointConstraintPhysicsProxy::InitializeOnPhysicsThread(FPBDRigidsSolver* InSolver)
+void FJointConstraintPhysicsProxy::InitializeOnPhysicsThread(FPBDRigidsSolver* InSolver, FDirtyPropertiesManager& Manager, int32 DataIdx, FDirtyChaosProperties& RemoteData)
 {
 	auto& Handles = InSolver->GetParticles().GetParticleHandles();
-	if (Handles.Size() && IsValid())
+	if (Handles.Size())
 	{
 		auto& JointConstraints = InSolver->GetJointConstraints();
-		if (Constraint != nullptr)
+		if(const FProxyBasePair* BasePairs = RemoteData.FindJointParticleProxies(Manager, DataIdx))
 		{
-			FConstraintBase::FProxyBasePair& BasePairs = Constraint->GetParticleProxies();
-
-			FGeometryParticleHandle* Handle0 = GetParticleHandleFromProxy(BasePairs[0]);
-			FGeometryParticleHandle* Handle1 = GetParticleHandleFromProxy(BasePairs[1]);
+			FGeometryParticleHandle* Handle0 = GetParticleHandleFromProxy((*BasePairs)[0]);
+			FGeometryParticleHandle* Handle1 = GetParticleHandleFromProxy((*BasePairs)[1]);
 			if (Handle0 && Handle1)
 			{
-				Handle = JointConstraints.AddConstraint({ Handle0,Handle1 }, Constraint->GetJointTransforms());
-				Handle->SetSettings(JointSettingsBuffer);
+				if (const FPBDJointSettings* JointSettings = RemoteData.FindJointSettings(Manager, DataIdx))
+				{
+					Constraint_PT = JointConstraints.AddConstraint({ Handle0,Handle1 }, JointSettings->ConnectorTransforms);
 
-				Handle0->AddConstraintHandle(Handle);
-				Handle1->AddConstraintHandle(Handle);
+					Handle0->AddConstraintHandle(Constraint_PT);
+					Handle1->AddConstraintHandle(Constraint_PT);
+				}
 			}
-
 		}
 	}
 }
 
 void FJointConstraintPhysicsProxy::DestroyOnPhysicsThread(FPBDRigidsSolver* InSolver)
 {
-	if (Handle && Handle->IsValid())
+	if (Constraint_PT && Constraint_PT->IsValid())
 	{
 		auto& JointConstraints = InSolver->GetJointConstraints();
-		JointConstraints.RemoveConstraint(Handle->GetConstraintIndex());
+		JointConstraints.RemoveConstraint(Constraint_PT->GetConstraintIndex());
 
 	}
+}
 
-	if (Constraint != nullptr)
-	{
-		delete Constraint;
-		Constraint = nullptr;
-	}
+void FJointConstraintPhysicsProxy::DestroyOnGameThread()
+{
+	delete Constraint_GT;
+	Constraint_GT = nullptr;
 }
 
 
-void FJointConstraintPhysicsProxy::PushStateOnGameThread(FPBDRigidsSolver* InSolver)
+void FJointConstraintPhysicsProxy::PushStateOnGameThread(FDirtyPropertiesManager& Manager, int32 DataIdx, FDirtyChaosProperties& RemoteData)
 {
-	if (Constraint && Constraint->IsValid())
+	if (Constraint_GT && Constraint_GT->IsValid())
 	{
-		if (Constraint->IsDirty())
-		{
-			JointSettingsBuffer = Constraint->GetJointSettings();
-			DirtyFlagsBuffer = Constraint->GetDirtyFlags();
-			Constraint->ClearDirtyFlags();
-		}
+		Constraint_GT->SyncRemoteData(Manager, DataIdx, RemoteData);
 	}
 }
 
 
-void FJointConstraintPhysicsProxy::PushStateOnPhysicsThread(FPBDRigidsSolver* InSolver)
+void FJointConstraintPhysicsProxy::PushStateOnPhysicsThread(FPBDRigidsSolver* InSolver, const FDirtyPropertiesManager& Manager, int32 DataIdx, const FDirtyChaosProperties& RemoteData)
 {
-	if (Handle && Handle->IsValid())
+	if (Constraint_PT && Constraint_PT->IsValid())
 	{
-		if (DirtyFlagsBuffer.IsDirty())
+		if (const FPBDJointSettings* Data = RemoteData.FindJointSettings(Manager, DataIdx))
 		{
-			const FPBDJointSettings& CurrentConstraintSettings = Handle->GetSettings();
+			const FPBDJointSettings& JointSettingsBuffer = *Data;
+			const FPBDJointSettings& CurrentConstraintSettings = Constraint_PT->GetSettings();
 
 			if (CurrentConstraintSettings.bCollisionEnabled != JointSettingsBuffer.bCollisionEnabled)
 			{
-				FConstraintBase::FProxyBasePair& BasePairs = Constraint->GetParticleProxies();
-				FGeometryParticleHandle* Handle0 = GetParticleHandleFromProxy(BasePairs[0]);
-				FGeometryParticleHandle* Handle1 = GetParticleHandleFromProxy(BasePairs[1]);
+				const TVector<FGeometryParticleHandle*, 2>& BasePairs = Constraint_PT->GetConstrainedParticles();
+				FGeometryParticleHandle * Handle0 = BasePairs[0];
+				FGeometryParticleHandle* Handle1 = BasePairs[1];
 
 				// Three pieces of state to update on the physics thread. 
 				// .. Mask on the particle array
@@ -209,8 +200,7 @@ void FJointConstraintPhysicsProxy::PushStateOnPhysicsThread(FPBDRigidsSolver* In
 				}
 			}
 
-			Handle->SetSettings(JointSettingsBuffer);
-			DirtyFlagsBuffer.Clear();
+			Constraint_PT->SetSettings(JointSettingsBuffer);
 		}
 	}
 }
diff --git a/Engine/Source/Runtime/Experimental/Chaos/Private/PhysicsProxy/SuspensionConstraintProxy.cpp b/Engine/Source/Runtime/Experimental/Chaos/Private/PhysicsProxy/SuspensionConstraintProxy.cpp
--- a/Engine/Source/Runtime/Experimental/Chaos/Private/PhysicsProxy/SuspensionConstraintProxy.cpp
+++ b/Engine/Source/Runtime/Experimental/Chaos/Private/PhysicsProxy/SuspensionConstraintProxy.cpp
@@ -48,7 +48,7 @@ void FSuspensionConstraintPhysicsProxy::InitializeOnPhysicsThread(Chaos::FPBDRig
 		auto& SuspensionConstraints = InSolver->GetSuspensionConstraints();
 		if (Constraint != nullptr)
 		{
-			Chaos::FConstraintBase::FProxyBasePair& BasePairs = Constraint->GetParticleProxies();
+			Chaos::FProxyBasePair& BasePairs = Constraint->GetParticleProxies();
 
 			
 			Chaos::TGeometryParticleHandle<Chaos::FReal, 3>* Handle0 = GetParticleHandleFromProxy(BasePairs[0]);
diff --git a/Engine/Source/Runtime/Experimental/Chaos/Public/Chaos/PBDConstraintBaseData.h b/Engine/Source/Runtime/Experimental/Chaos/Public/Chaos/PBDConstraintBaseData.h
--- a/Engine/Source/Runtime/Experimental/Chaos/Public/Chaos/PBDConstraintBaseData.h
+++ b/Engine/Source/Runtime/Experimental/Chaos/Public/Chaos/PBDConstraintBaseData.h
@@ -10,6 +10,7 @@
 
 #include "Chaos/Framework/PhysicsProxyBase.h"
 #include "Chaos/Framework/PhysicsSolverBase.h"
+#include "Properties.h"
 
 namespace Chaos
 {
@@ -23,16 +24,20 @@ enum class EConstraintType
 };
 
 #define CONSTRAINT_JOINT_PROPERPETY_IMPL(TYPE, FNAME, ENAME, VNAME)\
-	void Set##FNAME(TYPE InValue){if (InValue != VNAME){VNAME = InValue;MDirtyFlags.MarkDirty(ENAME);SetProxy(Proxy);}}\
+	void Set##FNAME(TYPE InValue){ if (InValue != VNAME){VNAME = InValue;MDirtyFlags.MarkDirty(ENAME);SetProxy(Proxy);}}\
 	TYPE Get##FNAME() const{return VNAME;}\
 
+#define CONSTRAINT_JOINT_PROPERPETY_IMPL2(TYPE, FNAME, PROP, VNAME)\
+	void Set##FNAME(TYPE InValue){ PROP.Modify(/*bInvalidate=*/true, DirtyFlags, Proxy, [&InValue](auto& Data) { Data.VNAME = InValue; }); }\
+	TYPE Get##FNAME() const{return PROP.Read().VNAME;}\
+
+
 class CHAOS_API FConstraintBase
 {
 public:
 
 	virtual ~FConstraintBase() {}
 
-	typedef TVector<IPhysicsProxyBase*, 2> FProxyBasePair;
 	typedef TVector<TGeometryParticleHandle<FReal, 3>*, 2> FParticleHandlePair;
 
 	FConstraintBase(EConstraintType InType);
@@ -41,6 +46,11 @@ public:
 	bool IsType(EConstraintType InType) { return (Type == InType); }
 	bool IsValid() const;
 
+	bool IsDirty() const { return DirtyFlags.IsDirty(); }
+	bool IsDirty(const EChaosPropertyFlags CheckBits) const { return DirtyFlags.IsDirty(CheckBits); }
+	void ClearDirtyFlags() { DirtyFlags.Clear(); }
+	const FDirtyChaosPropertyFlags& GetDirtyFlags() const { return DirtyFlags; }
+
 	//template <typename CastType>
 	//static CastType CastTo<CastType>(EConstraintType InType) { (InType == Type) ? static_cast<CastType*> : nullptr; }
 
@@ -52,18 +62,39 @@ public:
 	{
 		ConnectedParticleProxy[0] = InJointParticles[0];
 		ConnectedParticleProxy[1] = InJointParticles[1];
+
+		ConnectedParticleProxy2.Modify(/*bInvalidate=*/true, DirtyFlags, Proxy, [&InJointParticles](FProxyBasePair& Data)
+		{
+			Data[0] = InJointParticles[0];
+			Data[1] = InJointParticles[1];
+		});
 	}
 
 
 	const FProxyBasePair& GetParticleProxies() const { return  ConnectedParticleProxy; }
 	FProxyBasePair& GetParticleProxies() { return  ConnectedParticleProxy; }
 
+	void SyncRemoteData(FDirtyPropertiesManager& Manager, int32 DataIdx, FDirtyChaosProperties& RemoteData)
+	{
+		RemoteData.SetFlags(DirtyFlags);
+		SyncRemoteDataImp(Manager, DataIdx, RemoteData);
+		ClearDirtyFlags();
+	}
+
 protected:
 	EConstraintType Type;
 	class IPhysicsProxyBase* Proxy;
 
-	FProxyBasePair ConnectedParticleProxy;
+	FProxyBasePair ConnectedParticleProxy;	//TODO: remove this once suspension uses chaos property
+
+	FDirtyChaosPropertyFlags DirtyFlags;
+	TChaosProperty<FProxyBasePair, EChaosProperty::JointParticleProxies> ConnectedParticleProxy2;
 
+	virtual void SyncRemoteDataImp(FDirtyPropertiesManager& Manager, int32 DataIdx, FDirtyChaosProperties& RemoteData)
+	{
+		ConnectedParticleProxy2.SyncRemote(Manager, DataIdx, RemoteData);
+	}
+	
 };
 
 } // Chaos
diff --git a/Engine/Source/Runtime/Experimental/Chaos/Public/Chaos/PBDJointConstraintData.h b/Engine/Source/Runtime/Experimental/Chaos/Public/Chaos/PBDJointConstraintData.h
--- a/Engine/Source/Runtime/Experimental/Chaos/Public/Chaos/PBDJointConstraintData.h
+++ b/Engine/Source/Runtime/Experimental/Chaos/Public/Chaos/PBDJointConstraintData.h
@@ -35,8 +35,7 @@ namespace Chaos
 	class CHAOS_API FJointConstraint : public FConstraintBase
 	{
 	public:
-		typedef FPBDJointSettings FData;
-		typedef FPBDJointConstraintHandle FHandle;
+		using Base = FConstraintBase;
 		typedef TVector<FTransform, 2> FTransformPair;
 
 		friend class FPBDRigidsSolver; // friend so we can call ReleaseKinematicEndPoint when unregistering joint.
@@ -44,261 +43,257 @@ namespace Chaos
 		FJointConstraint();
 		virtual ~FJointConstraint() override {}
 
-		bool IsDirty() const { return MDirtyFlags.IsDirty(); }
-		bool IsDirty(const EJointConstraintFlags CheckBits) const { return MDirtyFlags.IsDirty(CheckBits); }
-		void ClearDirtyFlags() { MDirtyFlags.Clear(); }
-		FJointConstraintDirtyFlags GetDirtyFlags() const { return MDirtyFlags; }
-
-		const FData& GetJointSettings()const { return JointSettings; }
+		const FPBDJointSettings& GetJointSettings()const { return JointSettings.Read(); }
 
 		// If we created particle to serve as kinematic endpoint, track so we can release later. This will add particle to solver.
 		void SetKinematicEndPoint(FSingleParticlePhysicsProxy* InParticle, FPBDRigidsSolver* Solver);
 
-		//CONSTRAINT_JOINT_PROPERPETY_IMPL(bool, JointTransforms, EJointConstraintFlags::Position, JointSettings.ConnectorTransforms);
-		void SetJointTransforms(const Chaos::FJointConstraint::FTransformPair& InJointTransforms);
-		const FTransformPair GetJointTransforms() const;
+		CONSTRAINT_JOINT_PROPERPETY_IMPL2(FTransformPair, JointTransforms, JointSettings, ConnectorTransforms);
+		//void SetJointTransforms(const Chaos::FJointConstraint::FTransformPair& InJointTransforms);
+		//const FTransformPair GetJointTransforms() const;
+
 
-		CONSTRAINT_JOINT_PROPERPETY_IMPL(bool, CollisionEnabled, EJointConstraintFlags::CollisionEnabled, JointSettings.bCollisionEnabled);
+		CONSTRAINT_JOINT_PROPERPETY_IMPL2(bool, CollisionEnabled, JointSettings, bCollisionEnabled);
 		//void SetCollisionEnabled(bool InValue);
 		//bool GetCollisionEnabled() const;
 
-		CONSTRAINT_JOINT_PROPERPETY_IMPL(bool, ProjectionEnabled, EJointConstraintFlags::Projection, JointSettings.bProjectionEnabled);
+		CONSTRAINT_JOINT_PROPERPETY_IMPL2(bool, ProjectionEnabled, JointSettings, bProjectionEnabled);
 		//void SetProjectionEnabled(bool bInProjectionEnabled);
 		//bool GetProjectionEnabled() const;
 
-		CONSTRAINT_JOINT_PROPERPETY_IMPL(FReal, ProjectionLinearAlpha, EJointConstraintFlags::Projection, JointSettings.LinearProjection);
+		CONSTRAINT_JOINT_PROPERPETY_IMPL2(FReal, ProjectionLinearAlpha, JointSettings, LinearProjection);
 		//void SetProjectionLinearAlpha(FReal InProjectionLinearAlpha);
 		//FReal GetProjectionLinearAlpha() const;
 
-		CONSTRAINT_JOINT_PROPERPETY_IMPL(FReal, ProjectionAngularAlpha, EJointConstraintFlags::Projection, JointSettings.AngularProjection);
+		CONSTRAINT_JOINT_PROPERPETY_IMPL2(FReal, ProjectionAngularAlpha, JointSettings, AngularProjection);
 		//void SetProjectionAngularAlpha(FReal InProjectionAngularAlpha);
 		//FReal GetProjectionAngularAlpha() const;
 
-		CONSTRAINT_JOINT_PROPERPETY_IMPL(FReal, ParentInvMassScale, EJointConstraintFlags::ParentInvMassScale, JointSettings.ParentInvMassScale);
+		CONSTRAINT_JOINT_PROPERPETY_IMPL2(FReal, ParentInvMassScale, JointSettings, ParentInvMassScale);
 		//void SetParentInvMassScale(FReal InParentInvMassScale);
 		//FReal GetParentInvMassScale() const
 
-		CONSTRAINT_JOINT_PROPERPETY_IMPL(FReal, LinearBreakForce, EJointConstraintFlags::LinearBreakForce, JointSettings.LinearBreakForce);
+		CONSTRAINT_JOINT_PROPERPETY_IMPL2(FReal, LinearBreakForce, JointSettings, LinearBreakForce);
 		//void SetLinearBreakForce(FReal InLinearBreakForce);
 		//FReal GetLinearBreakForce() const
 
-		CONSTRAINT_JOINT_PROPERPETY_IMPL(FReal, LinearPlasticityLimit, EJointConstraintFlags::LinearBreakForce, JointSettings.LinearPlasticityLimit);
+		CONSTRAINT_JOINT_PROPERPETY_IMPL2(FReal, LinearPlasticityLimit, JointSettings, LinearPlasticityLimit);
 		//void SetLinearPlasticityLimit(FReal InLinearPlasticityLimit);
 		//FReal GetLinearPlasticityLimit() const
 
-		CONSTRAINT_JOINT_PROPERPETY_IMPL(EPlasticityType, LinearPlasticityType, EJointConstraintFlags::LinearBreakForce, JointSettings.LinearPlasticityType);
+		CONSTRAINT_JOINT_PROPERPETY_IMPL2(EPlasticityType, LinearPlasticityType, JointSettings, LinearPlasticityType);
 		//void SetLinearPlasticityType(EPlasticityType InLinearPlasticityType);
 		//EPlasticityType GetLinearPlasticityType() const
 
-		CONSTRAINT_JOINT_PROPERPETY_IMPL(FReal, AngularBreakTorque, EJointConstraintFlags::AngularBreakTorque, JointSettings.AngularBreakTorque);
+		CONSTRAINT_JOINT_PROPERPETY_IMPL2(FReal, AngularBreakTorque, JointSettings, AngularBreakTorque);
 		//void SetAngularBreakTorque(FReal InAngularBreakTorque);
 		//FReal GetAngularBreakTorque() const
 
-		CONSTRAINT_JOINT_PROPERPETY_IMPL(FReal, AngularPlasticityLimit, EJointConstraintFlags::AngularBreakTorque, JointSettings.AngularPlasticityLimit);
+		CONSTRAINT_JOINT_PROPERPETY_IMPL2(FReal, AngularPlasticityLimit, JointSettings, AngularPlasticityLimit);
 		//void SetAngularPlasticityLimit(FReal InAngularPlasticityLimit);
 		//FReal GetAngularPlasticityLimit() const
 
-		CONSTRAINT_JOINT_PROPERPETY_IMPL(void*, UserData, EJointConstraintFlags::UserData, UserData);
+		CONSTRAINT_JOINT_PROPERPETY_IMPL2(void*, UserData, JointSettings, UserData);
 		//void SetUserData(void* InUserData);
 		//void* GetUserData() const
 
-		void SetLinearPositionDriveEnabled( TVector<bool,3> Enabled);
-		
-		CONSTRAINT_JOINT_PROPERPETY_IMPL(bool, LinearPositionDriveXEnabled, EJointConstraintFlags::LinearDrive, JointSettings.bLinearPositionDriveEnabled[0]);
+		void SetLinearPositionDriveEnabled(TVector<bool, 3> Enabled);
+
+		CONSTRAINT_JOINT_PROPERPETY_IMPL2(bool, LinearPositionDriveXEnabled, JointSettings, bLinearPositionDriveEnabled[0]);
 		//void SetLinearPositionDriveXEnabled(bool InLinearPositionDriveXEnabled);
 		//bool GetLinearPositionDriveXEnabled() const
 
-		CONSTRAINT_JOINT_PROPERPETY_IMPL(bool, LinearPositionDriveYEnabled, EJointConstraintFlags::LinearDrive, JointSettings.bLinearPositionDriveEnabled[1]);
+		CONSTRAINT_JOINT_PROPERPETY_IMPL2(bool, LinearPositionDriveYEnabled, JointSettings, bLinearPositionDriveEnabled[1]);
 		//void SetLinearPositionDriveYEnabled(bool InLinearPositionDriveYEnabled);
 		//bool GetLinearPositionDriveYEnabled() const
 
-		CONSTRAINT_JOINT_PROPERPETY_IMPL(bool, LinearPositionDriveZEnabled, EJointConstraintFlags::LinearDrive, JointSettings.bLinearPositionDriveEnabled[2]);
+		CONSTRAINT_JOINT_PROPERPETY_IMPL2(bool, LinearPositionDriveZEnabled, JointSettings, bLinearPositionDriveEnabled[2]);
 		//void SetLinearPositionDriveZEnabled(bool InLinearPositionDriveZEnabled);
 		//bool GetLinearPositionDriveZEnabled() const
 
-		CONSTRAINT_JOINT_PROPERPETY_IMPL(FVec3, LinearDrivePositionTarget, EJointConstraintFlags::LinearDrive, JointSettings.LinearDrivePositionTarget);
+		CONSTRAINT_JOINT_PROPERPETY_IMPL2(FVec3, LinearDrivePositionTarget, JointSettings, LinearDrivePositionTarget);
 		//void SetLinearDrivePositionTarget(FVec3 InLinearDrivePositionTarget);
 		//FVec3 GetLinearDrivePositionTarget() const
 
-		void SetLinearVelocityDriveEnabled(TVector<bool,3> Enabled);
+		void SetLinearVelocityDriveEnabled(TVector<bool, 3> Enabled);
 
-		CONSTRAINT_JOINT_PROPERPETY_IMPL(bool, LinearVelocityDriveXEnabled, EJointConstraintFlags::LinearDrive, JointSettings.bLinearVelocityDriveEnabled[0]);
+		CONSTRAINT_JOINT_PROPERPETY_IMPL2(bool, LinearVelocityDriveXEnabled, JointSettings, bLinearVelocityDriveEnabled[0]);
 		//void SetLinearVelocityDriveXEnabled(bool InLinearVelocityDriveXEnabled);
 		//bool GetLinearVelocityDriveXEnabled() const
 
-		CONSTRAINT_JOINT_PROPERPETY_IMPL(bool, LinearVelocityDriveYEnabled, EJointConstraintFlags::LinearDrive, JointSettings.bLinearVelocityDriveEnabled[1]);
+		CONSTRAINT_JOINT_PROPERPETY_IMPL2(bool, LinearVelocityDriveYEnabled, JointSettings, bLinearVelocityDriveEnabled[1]);
 		//void SetLinearVelocityDriveYEnabled(bool InLinearVelocityDriveYEnabled);
 		//bool GetLinearVelocityDriveYEnabled() const
 
-		CONSTRAINT_JOINT_PROPERPETY_IMPL(bool, LinearVelocityDriveZEnabled, EJointConstraintFlags::LinearDrive, JointSettings.bLinearVelocityDriveEnabled[2]);
+		CONSTRAINT_JOINT_PROPERPETY_IMPL2(bool, LinearVelocityDriveZEnabled, JointSettings, bLinearVelocityDriveEnabled[2]);
 		//void SetLinearVelocityDriveZEnabled(bool InLinearVelocityDriveZEnabled);
 		//bool GetLinearVelocityDriveZEnabled() const
 
-		CONSTRAINT_JOINT_PROPERPETY_IMPL(FVec3, LinearDriveVelocityTarget, EJointConstraintFlags::LinearDrive, JointSettings.LinearDriveVelocityTarget);
+		CONSTRAINT_JOINT_PROPERPETY_IMPL2(FVec3, LinearDriveVelocityTarget, JointSettings, LinearDriveVelocityTarget);
 		//void SetLinearDriveVelocityTarget(FVec3 InLinearDriveVelocityTarget);
 		//FVec3 GetLinearDriveVelocityTarget() const
 
-		CONSTRAINT_JOINT_PROPERPETY_IMPL(EJointForceMode, LinearDriveForceMode, EJointConstraintFlags::LinearDrive, JointSettings.LinearDriveForceMode);
+		CONSTRAINT_JOINT_PROPERPETY_IMPL2(EJointForceMode, LinearDriveForceMode, JointSettings, LinearDriveForceMode);
 		//void SetLinearDriveVelocityTarget(EJointForceMode InEJointForceMode);
 		//EJointForceMode GetLinearDriveVelocityTarget() const
 
-		CONSTRAINT_JOINT_PROPERPETY_IMPL(EJointMotionType, LinearMotionTypesX, EJointConstraintFlags::LinearDrive, JointSettings.LinearMotionTypes[0]);
+		CONSTRAINT_JOINT_PROPERPETY_IMPL2(EJointMotionType, LinearMotionTypesX, JointSettings, LinearMotionTypes[0]);
 		//void SetLinearMotionTypesX(EJointMotionType InLinearMotionTypesX);
 		//EJointMotionType GetLinearMotionTypesX() const
 
-		CONSTRAINT_JOINT_PROPERPETY_IMPL(EJointMotionType, LinearMotionTypesY, EJointConstraintFlags::LinearDrive, JointSettings.LinearMotionTypes[1]);
+		CONSTRAINT_JOINT_PROPERPETY_IMPL2(EJointMotionType, LinearMotionTypesY, JointSettings, LinearMotionTypes[1]);
 		//void SetLinearMotionTypesY(EJointMotionType InLinearMotionTypesY);
 		//EJointMotionType GetLinearMotionTypesY() const
 
-		CONSTRAINT_JOINT_PROPERPETY_IMPL(EJointMotionType, LinearMotionTypesZ, EJointConstraintFlags::LinearDrive, JointSettings.LinearMotionTypes[2]);
+		CONSTRAINT_JOINT_PROPERPETY_IMPL2(EJointMotionType, LinearMotionTypesZ, JointSettings, LinearMotionTypes[2]);
 		//void SetLinearMotionTypesZ(EJointMotionType InLinearMotionTypesZ);
 		//EJointMotionType GetLinearMotionTypesZ() const
 
-		CONSTRAINT_JOINT_PROPERPETY_IMPL(FReal, LinearDriveStiffness, EJointConstraintFlags::LinearDrive, JointSettings.LinearDriveStiffness);
+		CONSTRAINT_JOINT_PROPERPETY_IMPL2(FReal, LinearDriveStiffness, JointSettings, LinearDriveStiffness);
 		//void SetLinearDriveStiffness(FReal InLinearDriveStffness);
 		//FReal GetLinearDriveStiffness() const
 
-		CONSTRAINT_JOINT_PROPERPETY_IMPL(FReal, LinearDriveDamping, EJointConstraintFlags::LinearDrive, JointSettings.LinearDriveDamping);
+		CONSTRAINT_JOINT_PROPERPETY_IMPL2(FReal, LinearDriveDamping, JointSettings, LinearDriveDamping);
 		//void SetLinearDriveDamping(FReal InLinearDriveStffness);
 		//FReal GetLinearDriveDamping() const
 
-		CONSTRAINT_JOINT_PROPERPETY_IMPL(FReal, ContactTransferScale, EJointConstraintFlags::LinearDrive, JointSettings.ContactTransferScale);
+		CONSTRAINT_JOINT_PROPERPETY_IMPL2(FReal, ContactTransferScale, JointSettings, ContactTransferScale);
 		//void SetContactTransferScale(FReal InContactTransferScale);
 		//FReal GetContactTransferScale() const
 
-		CONSTRAINT_JOINT_PROPERPETY_IMPL(bool, AngularSLerpPositionDriveEnabled, EJointConstraintFlags::AngularDrive, JointSettings.bAngularSLerpPositionDriveEnabled);
+		CONSTRAINT_JOINT_PROPERPETY_IMPL2(bool, AngularSLerpPositionDriveEnabled, JointSettings, bAngularSLerpPositionDriveEnabled);
 		//void SetAngularSLerpPositionDriveEnabled(bool InAngularSLerpPositionDriveEnabled);
 		//bool GetAngularSLerpPositionDriveEnabled() const
 
-		CONSTRAINT_JOINT_PROPERPETY_IMPL(bool, AngularTwistPositionDriveEnabled, EJointConstraintFlags::AngularDrive, JointSettings.bAngularTwistPositionDriveEnabled);
+		CONSTRAINT_JOINT_PROPERPETY_IMPL2(bool, AngularTwistPositionDriveEnabled, JointSettings, bAngularTwistPositionDriveEnabled);
 		//void SetAngularTwistPositionDriveEnabled(bool InAngularTwistPositionDriveEnabled);
 		//bool GetAngularTwistPositionDriveEnabled() const
 
-		CONSTRAINT_JOINT_PROPERPETY_IMPL(bool, AngularSwingPositionDriveEnabled, EJointConstraintFlags::AngularDrive, JointSettings.bAngularSwingPositionDriveEnabled);
+		CONSTRAINT_JOINT_PROPERPETY_IMPL2(bool, AngularSwingPositionDriveEnabled, JointSettings, bAngularSwingPositionDriveEnabled);
 		//void SetAngularSwingPositionDriveEnabled(bool InAngularSwingPositionDriveEnabled);
 		//bool GetAngularSwingPositionDriveEnabled() const
 
-		CONSTRAINT_JOINT_PROPERPETY_IMPL(FRotation3, AngularDrivePositionTarget, EJointConstraintFlags::AngularDrive, JointSettings.AngularDrivePositionTarget);
+		CONSTRAINT_JOINT_PROPERPETY_IMPL2(FRotation3, AngularDrivePositionTarget, JointSettings, AngularDrivePositionTarget);
 		//void SetAngularDrivePositionTarget(FRotation3 InAngularDrivePositionTarget);
 		//FRotation3 GetAngularDrivePositionTarget() const
 
-		CONSTRAINT_JOINT_PROPERPETY_IMPL(bool, AngularSLerpVelocityDriveEnabled, EJointConstraintFlags::AngularDrive, JointSettings.bAngularSLerpVelocityDriveEnabled);
+		CONSTRAINT_JOINT_PROPERPETY_IMPL2(bool, AngularSLerpVelocityDriveEnabled, JointSettings, bAngularSLerpVelocityDriveEnabled);
 		//void SetAngularSLerpVelocityDriveEnabled(bool InAngularSLerpVelocityDriveEnabled);
 		//bool GetAngularSLerpVelocityDriveEnabled() const
 
-		CONSTRAINT_JOINT_PROPERPETY_IMPL(bool, AngularTwistVelocityDriveEnabled, EJointConstraintFlags::AngularDrive, JointSettings.bAngularTwistVelocityDriveEnabled);
+		CONSTRAINT_JOINT_PROPERPETY_IMPL2(bool, AngularTwistVelocityDriveEnabled, JointSettings, bAngularTwistVelocityDriveEnabled);
 		//void SetAngularTwistVelocityDriveEnabled(bool InAngularTwistVelocityDriveEnabled);
 		//bool GetAngularTwistVelocityDriveEnabled() const
 
-		CONSTRAINT_JOINT_PROPERPETY_IMPL(bool, AngularSwingVelocityDriveEnabled, EJointConstraintFlags::AngularDrive, JointSettings.bAngularSwingVelocityDriveEnabled);
+		CONSTRAINT_JOINT_PROPERPETY_IMPL2(bool, AngularSwingVelocityDriveEnabled, JointSettings, bAngularSwingVelocityDriveEnabled);
 		//void SetAngularSwingVelocityDriveEnabled(bool InAngularSwingVelocityDriveEnabled);
 		//bool GetAngularSwingVelocityDriveEnabled() const
 
-		CONSTRAINT_JOINT_PROPERPETY_IMPL(FVec3, AngularDriveVelocityTarget, EJointConstraintFlags::AngularDrive, JointSettings.AngularDriveVelocityTarget);
+		CONSTRAINT_JOINT_PROPERPETY_IMPL2(FVec3, AngularDriveVelocityTarget, JointSettings, AngularDriveVelocityTarget);
 		//void SetAngularDriveVelocityTarget(FVec3 InAngularDriveVelocityTarget);
 		//FVec3 GetAngularDriveVelocityTarget() const
 
-		CONSTRAINT_JOINT_PROPERPETY_IMPL(EJointForceMode, AngularDriveForceMode, EJointConstraintFlags::AngularDrive, JointSettings.AngularDriveForceMode);
+		CONSTRAINT_JOINT_PROPERPETY_IMPL2(EJointForceMode, AngularDriveForceMode, JointSettings, AngularDriveForceMode);
 		//void SetAngularDriveForceMode(EJointForceMode InAngularDriveForceMode);
 		//EJointForceMode GetAngularDriveForceMode() const
 
-		CONSTRAINT_JOINT_PROPERPETY_IMPL(EJointMotionType, AngularMotionTypesX, EJointConstraintFlags::AngularDrive, JointSettings.AngularMotionTypes[0]);
+		CONSTRAINT_JOINT_PROPERPETY_IMPL2(EJointMotionType, AngularMotionTypesX, JointSettings, AngularMotionTypes[0]);
 		//void SetAngularMotionTypesX(EJointMotionType InAngularMotionTypesX);
 		//EJointMotionType GetAngularMotionTypesX() const
 
-		CONSTRAINT_JOINT_PROPERPETY_IMPL(EJointMotionType, AngularMotionTypesY, EJointConstraintFlags::AngularDrive, JointSettings.AngularMotionTypes[1]);
+		CONSTRAINT_JOINT_PROPERPETY_IMPL2(EJointMotionType, AngularMotionTypesY, JointSettings, AngularMotionTypes[1]);
 		//void SetAngularMotionTypesY(EJointMotionType InAngularMotionTypesY);
 		//EJointMotionType GetAngularMotionTypesY() const
 
-		CONSTRAINT_JOINT_PROPERPETY_IMPL(EJointMotionType, AngularMotionTypesZ, EJointConstraintFlags::AngularDrive, JointSettings.AngularMotionTypes[2]);
+		CONSTRAINT_JOINT_PROPERPETY_IMPL2(EJointMotionType, AngularMotionTypesZ, JointSettings, AngularMotionTypes[2]);
 		//void SetAngularMotionTypesZ(EJointMotionType AngularMotionTypesZ);
 		//EJointMotionType GetAngularMotionTypesZ() const
 
-		CONSTRAINT_JOINT_PROPERPETY_IMPL(FReal, AngularDriveStiffness, EJointConstraintFlags::AngularDrive, JointSettings.AngularDriveStiffness);
+		CONSTRAINT_JOINT_PROPERPETY_IMPL2(FReal, AngularDriveStiffness, JointSettings, AngularDriveStiffness);
 		//void SetAngularDriveStiffness(FReal InAngularDriveStiffness);
 		//FReal GetAngularDriveStiffness() const
 
-		CONSTRAINT_JOINT_PROPERPETY_IMPL(FReal, AngularDriveDamping, EJointConstraintFlags::AngularDrive, JointSettings.AngularDriveDamping);
+		CONSTRAINT_JOINT_PROPERPETY_IMPL2(FReal, AngularDriveDamping, JointSettings, AngularDriveDamping);
 		//void SetAngularDriveDamping(FReal InAngularDriveDamping);
 		//FReal GetAngularDriveDamping() const
 
-		CONSTRAINT_JOINT_PROPERPETY_IMPL(FReal, Stiffness, EJointConstraintFlags::Stiffness, JointSettings.Stiffness);
+		CONSTRAINT_JOINT_PROPERPETY_IMPL2(FReal, Stiffness, JointSettings, Stiffness);
 		//void SetStiffness(FReal Stiffness);
 		//FReal GetStiffness() const
 
-		CONSTRAINT_JOINT_PROPERPETY_IMPL(bool, SoftLinearLimitsEnabled, EJointConstraintFlags::Limits, JointSettings.bSoftLinearLimitsEnabled);
+		CONSTRAINT_JOINT_PROPERPETY_IMPL2(bool, SoftLinearLimitsEnabled, JointSettings, bSoftLinearLimitsEnabled);
 		//void SetSoftLinearLimitsEnabled(bool bInSoftLinearLimitsEnabled);
 		//bool GetSoftLinearLimitsEnabled() const
 
-		CONSTRAINT_JOINT_PROPERPETY_IMPL(bool, SoftTwistLimitsEnabled, EJointConstraintFlags::Limits, JointSettings.bSoftTwistLimitsEnabled);
+		CONSTRAINT_JOINT_PROPERPETY_IMPL2(bool, SoftTwistLimitsEnabled, JointSettings, bSoftTwistLimitsEnabled);
 		//void SetSoftTwistLimitsEnabled(FReal bInSoftTwistLimitsEnabled);
 		//bool GetSoftTwistLimitsEnabled() const
 
-		CONSTRAINT_JOINT_PROPERPETY_IMPL(bool, SoftSwingLimitsEnabled, EJointConstraintFlags::Limits, JointSettings.bSoftSwingLimitsEnabled);
+		CONSTRAINT_JOINT_PROPERPETY_IMPL2(bool, SoftSwingLimitsEnabled, JointSettings, bSoftSwingLimitsEnabled);
 		//void SetSoftSwingLimitsEnabled(FReal bInSoftSwingLimitsEnabled);
 		//bool GetSoftSwingLimitsEnabled() const
 
-		CONSTRAINT_JOINT_PROPERPETY_IMPL(EJointForceMode, LinearSoftForceMode, EJointConstraintFlags::Limits, JointSettings.LinearSoftForceMode);
+		CONSTRAINT_JOINT_PROPERPETY_IMPL2(EJointForceMode, LinearSoftForceMode, JointSettings, LinearSoftForceMode);
 		//void SetLinearSoftForceMode(FRealIn LinearSoftForceMode);
 		//EJointForceMode GetLinearSoftForceMode() const
 
-		CONSTRAINT_JOINT_PROPERPETY_IMPL(EJointForceMode, AngularSoftForceMode, EJointConstraintFlags::Limits, JointSettings.AngularSoftForceMode);
+		CONSTRAINT_JOINT_PROPERPETY_IMPL2(EJointForceMode, AngularSoftForceMode, JointSettings, AngularSoftForceMode);
 		//void SetAngularSoftForceMode(FReal InAngularSoftForceMode);
 		//EJointForceMode GetAngularSoftForceMode() const
 
-		CONSTRAINT_JOINT_PROPERPETY_IMPL(FReal, SoftLinearStiffness, EJointConstraintFlags::Limits, JointSettings.SoftLinearStiffness);
+		CONSTRAINT_JOINT_PROPERPETY_IMPL2(FReal, SoftLinearStiffness, JointSettings, SoftLinearStiffness);
 		//void SetSoftLinearStiffness(FReal InSoftLinearStiffness);
 		//FReal GetSoftLinearStiffness() const
 
-		CONSTRAINT_JOINT_PROPERPETY_IMPL(FReal, SoftLinearDamping, EJointConstraintFlags::Limits, JointSettings.SoftLinearDamping);
+		CONSTRAINT_JOINT_PROPERPETY_IMPL2(FReal, SoftLinearDamping, JointSettings, SoftLinearDamping);
 		//void SetSoftLinearDamping(FReal InSoftLinearDamping);
 		//FReal GetSoftLinearDamping() const
 
-		CONSTRAINT_JOINT_PROPERPETY_IMPL(FReal, SoftTwistStiffness, EJointConstraintFlags::Limits, JointSettings.SoftTwistStiffness);
+		CONSTRAINT_JOINT_PROPERPETY_IMPL2(FReal, SoftTwistStiffness, JointSettings, SoftTwistStiffness);
 		//void SetSoftTwistStiffness(FReal InSoftTwistStiffness);
 		//FReal GetSoftTwistStiffness() const
 
-		CONSTRAINT_JOINT_PROPERPETY_IMPL(FReal, SoftTwistDamping, EJointConstraintFlags::Limits, JointSettings.SoftTwistDamping);
+		CONSTRAINT_JOINT_PROPERPETY_IMPL2(FReal, SoftTwistDamping, JointSettings, SoftTwistDamping);
 		//void SetSoftTwistDamping(FReal InSoftTwistDamping);
 		//FReal GetSoftTwistDamping() const
 
-		CONSTRAINT_JOINT_PROPERPETY_IMPL(FReal, SoftSwingStiffness, EJointConstraintFlags::Limits, JointSettings.SoftSwingStiffness);
+		CONSTRAINT_JOINT_PROPERPETY_IMPL2(FReal, SoftSwingStiffness, JointSettings, SoftSwingStiffness);
 		//void SetSoftSwingStiffness(FReal InSoftSwingStiffness);
 		//FReal GetSoftSwingStiffness() const
 
-		CONSTRAINT_JOINT_PROPERPETY_IMPL(FReal, SoftSwingDamping, EJointConstraintFlags::Limits, JointSettings.SoftSwingDamping);
+		CONSTRAINT_JOINT_PROPERPETY_IMPL2(FReal, SoftSwingDamping, JointSettings, SoftSwingDamping);
 		//void SetSoftSwingDamping(FReal InSoftSwingDamping);
 		//FReal GetSoftSwingDamping() const
 
-		CONSTRAINT_JOINT_PROPERPETY_IMPL(FReal, LinearLimit, EJointConstraintFlags::Limits, JointSettings.LinearLimit);
+		CONSTRAINT_JOINT_PROPERPETY_IMPL2(FReal, LinearLimit, JointSettings, LinearLimit);
 		//void SetLinearLimit(FReal InLinearLimit);
 		//FReal GetLinearLimit() const
 
-		CONSTRAINT_JOINT_PROPERPETY_IMPL(FVec3, AngularLimits, EJointConstraintFlags::Limits, JointSettings.AngularLimits);
+		CONSTRAINT_JOINT_PROPERPETY_IMPL2(FVec3, AngularLimits, JointSettings, AngularLimits);
 		//void SetAngularLimits(FVec3 AngularLimits);
 		//FVec3 GetAngularLimits() const
 
-		CONSTRAINT_JOINT_PROPERPETY_IMPL(FReal, LinearContactDistance, EJointConstraintFlags::Limits, JointSettings.LinearContactDistance);
+		CONSTRAINT_JOINT_PROPERPETY_IMPL2(FReal, LinearContactDistance, JointSettings, LinearContactDistance);
 		//void SetLinearContactDistance(FReal InLinearContactDistance);
 		//FReal GetLinearContactDistance() const
 
-		CONSTRAINT_JOINT_PROPERPETY_IMPL(FReal, TwistContactDistance, EJointConstraintFlags::Limits, JointSettings.TwistContactDistance);
+		CONSTRAINT_JOINT_PROPERPETY_IMPL2(FReal, TwistContactDistance, JointSettings, TwistContactDistance);
 		//void SetAngularTwistContactDistance(FReal InAngularTwistContactDistance);
 		//FReal GetAngularTwistContactDistance() const
 
-		CONSTRAINT_JOINT_PROPERPETY_IMPL(FReal, SwingContactDistance, EJointConstraintFlags::Limits, JointSettings.SwingContactDistance);
+		CONSTRAINT_JOINT_PROPERPETY_IMPL2(FReal, SwingContactDistance, JointSettings, SwingContactDistance);
 		//void SetAngularSwingContactDistance(FReal InAngularSwingContactDistance);
 		//FReal GetAngularSwingContactDistance() const
 
-		CONSTRAINT_JOINT_PROPERPETY_IMPL(FReal, LinearRestitution, EJointConstraintFlags::Limits, JointSettings.LinearRestitution);
+		CONSTRAINT_JOINT_PROPERPETY_IMPL2(FReal, LinearRestitution, JointSettings, LinearRestitution);
 		//void SetLinearRestitution(FReal InLinearRestitution);
 		//FReal GetLinearRestitution() const
 
-		CONSTRAINT_JOINT_PROPERPETY_IMPL(FReal, TwistRestitution, EJointConstraintFlags::Limits, JointSettings.TwistRestitution);
+		CONSTRAINT_JOINT_PROPERPETY_IMPL2(FReal, TwistRestitution, JointSettings, TwistRestitution);
 		//void SetAngularTwistRestitution(FReal InAngularTwistRestitution);
 		//FReal GetAngularTwistRestitution() const
 
-		CONSTRAINT_JOINT_PROPERPETY_IMPL(FReal, SwingRestitution, EJointConstraintFlags::Limits, JointSettings.SwingRestitution);
+		CONSTRAINT_JOINT_PROPERPETY_IMPL2(FReal, SwingRestitution, JointSettings, SwingRestitution);
 		//void SetAngularSwingRestitution(FReal InAngularSwingRestitution);
 		//FReal GetAngularSwingRestitution() const
 
@@ -313,14 +308,19 @@ namespace Chaos
 		};
 		FOutputData& GetOutputData() { return Output; }
 
+		virtual void SyncRemoteDataImp(FDirtyPropertiesManager& Manager, int32 DataIdx, FDirtyChaosProperties& RemoteData) override
+		{
+			Base::SyncRemoteDataImp(Manager, DataIdx, RemoteData);
+			JointSettings.SyncRemote(Manager, DataIdx, RemoteData);
+		}
+
 	protected:
 
 		void ReleaseKinematicEndPoint(FPBDRigidsSolver* Solver);
 		
-		FJointConstraintDirtyFlags MDirtyFlags;
-		FData JointSettings;
+		TChaosProperty<FPBDJointSettings, EChaosProperty::JointSettings> JointSettings;
+
 
-		void* UserData;
 		FOutputData Output;
 
 	private:
diff --git a/Engine/Source/Runtime/Experimental/Chaos/Public/Chaos/ParticleDirtyFlags.h b/Engine/Source/Runtime/Experimental/Chaos/Public/Chaos/ParticleDirtyFlags.h
--- a/Engine/Source/Runtime/Experimental/Chaos/Public/Chaos/ParticleDirtyFlags.h
+++ b/Engine/Source/Runtime/Experimental/Chaos/Public/Chaos/ParticleDirtyFlags.h
@@ -184,6 +184,8 @@ private:
 
 };
 
+typedef TVector<IPhysicsProxyBase*, 2> FProxyBasePair;
+
 inline FChaosArchive& operator<<(FChaosArchive& Ar, FParticleDynamics& Data)
 {
 	Data.Serialize(Ar);
diff --git a/Engine/Source/Runtime/Experimental/Chaos/Public/Chaos/ParticleProperties.inl b/Engine/Source/Runtime/Experimental/Chaos/Public/Chaos/ParticleProperties.inl
--- a/Engine/Source/Runtime/Experimental/Chaos/Public/Chaos/ParticleProperties.inl
+++ b/Engine/Source/Runtime/Experimental/Chaos/Public/Chaos/ParticleProperties.inl
@@ -15,4 +15,7 @@ CHAOS_PROPERTY(Dynamics, FParticleDynamics, EPhysicsProxyType::SingleParticlePro
 CHAOS_PROPERTY(DynamicMisc, FParticleDynamicMisc, EPhysicsProxyType::SingleParticleProxy)
 CHAOS_PROPERTY(NonFrequentData, FParticleNonFrequentData, EPhysicsProxyType::SingleParticleProxy)
 CHAOS_PROPERTY(MassProps, FParticleMassProps, EPhysicsProxyType::SingleParticleProxy)
-CHAOS_PROPERTY(KinematicTarget,FKinematicTarget, EPhysicsProxyType::SingleParticleProxy)
\ No newline at end of file
+CHAOS_PROPERTY(KinematicTarget,FKinematicTarget, EPhysicsProxyType::SingleParticleProxy)
+
+CHAOS_PROPERTY(JointSettings, FPBDJointSettings, EPhysicsProxyType::JointConstraintType)
+CHAOS_PROPERTY(JointParticleProxies, FProxyBasePair, EPhysicsProxyType::JointConstraintType)
\ No newline at end of file
diff --git a/Engine/Source/Runtime/Experimental/Chaos/Public/PhysicsProxy/JointConstraintProxy.h b/Engine/Source/Runtime/Experimental/Chaos/Public/PhysicsProxy/JointConstraintProxy.h
--- a/Engine/Source/Runtime/Experimental/Chaos/Public/PhysicsProxy/JointConstraintProxy.h
+++ b/Engine/Source/Runtime/Experimental/Chaos/Public/PhysicsProxy/JointConstraintProxy.h
@@ -27,8 +27,6 @@ class FJointConstraintPhysicsProxy : public IPhysicsProxyBase
 public:
 	FJointConstraintPhysicsProxy() = delete;
 	FJointConstraintPhysicsProxy(FJointConstraint* InConstraint, FPBDJointConstraintHandle* InHandle, UObject* InOwner = nullptr);
-		
-	bool IsValid() { return Constraint && Constraint->IsValid(); }
 
 	bool IsInitialized() const { return bInitialized; }
 	void SetInitialized() { bInitialized = true; }
@@ -39,29 +37,29 @@ public:
 	//  Lifespan Management
 	//
 
-	void CHAOS_API InitializeOnPhysicsThread(FPBDRigidsSolver* InSolver);
+	void CHAOS_API InitializeOnPhysicsThread(FPBDRigidsSolver* InSolver, FDirtyPropertiesManager& Manager, int32 DataIdx, FDirtyChaosProperties& RemoteData);
 
-	// Merge to perform a remote sync
-	void CHAOS_API PushStateOnGameThread(FPBDRigidsSolver* InSolver);
+	void CHAOS_API PushStateOnGameThread(FDirtyPropertiesManager& Manager, int32 DataIdx, FDirtyChaosProperties& RemoteData);
 
-	void CHAOS_API PushStateOnPhysicsThread(FPBDRigidsSolver* InSolver);
-	// Merge to perform a remote sync - END
+	void CHAOS_API PushStateOnPhysicsThread(FPBDRigidsSolver* InSolver, const FDirtyPropertiesManager& Manager, int32 DataIdx, const FDirtyChaosProperties& RemoteData);
 
 	void CHAOS_API DestroyOnPhysicsThread(FPBDRigidsSolver* InSolver);
 
+	void CHAOS_API DestroyOnGameThread();
+
 	//
 	// Member Access
 	//
 
-	FPBDJointConstraintHandle* GetHandle() { return Handle; }
-	const FPBDJointConstraintHandle* GetHandle() const { return Handle; }
+	FPBDJointConstraintHandle* GetHandle() { return Constraint_PT; }
+	const FPBDJointConstraintHandle* GetHandle() const { return Constraint_PT; }
 
-	virtual void* GetHandleUnsafe() const override { return Handle; }
+	virtual void* GetHandleUnsafe() const override { return Constraint_PT; }
 
-	void SetHandle(FPBDJointConstraintHandle* InHandle)	{ Handle = InHandle; }
+	void SetHandle(FPBDJointConstraintHandle* InHandle)	{ Constraint_PT = InHandle; }
 
-	FJointConstraint* GetConstraint(){ return Constraint; }
-	const FJointConstraint* GetConstraint() const { return Constraint; }
+	FJointConstraint* GetConstraint(){ return Constraint_GT; }
+	const FJointConstraint* GetConstraint() const { return Constraint_GT; }
 
 	//
 	// Threading API
@@ -72,18 +70,10 @@ public:
 
 	/**/
 	bool CHAOS_API PullFromPhysicsState(const FDirtyJointConstraintData& Buffer, const int32 SolverSyncTimestamp);
-
-	/**/
-	bool IsDirty() { return Constraint->IsDirty(); }
 	
 private:
-
-	// Input Buffer
-	FPBDJointSettings JointSettingsBuffer;
-	FJointConstraintDirtyFlags DirtyFlagsBuffer;
-
-	FJointConstraint* Constraint; 	// This proxy assumes ownership of the Constraint, and will free it during DestroyOnPhysicsThread
-	FPBDJointConstraintHandle* Handle;
+	FJointConstraint* Constraint_GT;
+	FPBDJointConstraintHandle* Constraint_PT;
 	bool bInitialized = false;
 };
 
diff --git a/Engine/Source/Runtime/PhysicsCore/Private/ChaosEngineInterface.cpp b/Engine/Source/Runtime/PhysicsCore/Private/ChaosEngineInterface.cpp
--- a/Engine/Source/Runtime/PhysicsCore/Private/ChaosEngineInterface.cpp
+++ b/Engine/Source/Runtime/PhysicsCore/Private/ChaosEngineInterface.cpp
@@ -1153,7 +1153,7 @@ FTransform FChaosEngineInterface::GetGlobalPose(const FPhysicsConstraintHandle&
 	{
 		if (Chaos::FJointConstraint* Constraint = static_cast<Chaos::FJointConstraint*>(InConstraintRef.Constraint))
 		{
-			Chaos::FConstraintBase::FProxyBasePair BasePairs = Constraint->GetParticleProxies();
+			Chaos::FProxyBasePair BasePairs = Constraint->GetParticleProxies();
 			const Chaos::FJointConstraint::FTransformPair& M = Constraint->GetJointTransforms();
 
 			if (InFrame == EConstraintFrame::Frame1)
