Chaos : GeometryCollection
- Changed array access to GetData off the internal Array
- Prevents ensures when accessing an empty array.

diff --git a/Engine/Source/Runtime/Experimental/Chaos/Public/GeometryCollection/ManagedArray.h b/Engine/Source/Runtime/Experimental/Chaos/Public/GeometryCollection/ManagedArray.h
--- a/Engine/Source/Runtime/Experimental/Chaos/Public/GeometryCollection/ManagedArray.h
+++ b/Engine/Source/Runtime/Experimental/Chaos/Public/GeometryCollection/ManagedArray.h
@@ -330,7 +330,7 @@ public:
 	*/
 	FORCEINLINE ElementType* GetData()
 	{
-		return &Array.operator[](0);
+		return Array.GetData();
 	}
 
 	/**
@@ -340,7 +340,7 @@ public:
 	*/
 	FORCEINLINE const ElementType * GetData() const
 	{
-		return &Array.operator[](0);
+		return Array.GetData();
 	}
 
 	/**
