fracture perf improvement: remove expensive validation from non-debug builds

#rnx

diff --git a/Engine/Plugins/Experimental/PlanarCutPlugin/Source/PlanarCut/Private/GeometryMeshConversion.cpp b/Engine/Plugins/Experimental/PlanarCutPlugin/Source/PlanarCut/Private/GeometryMeshConversion.cpp
--- a/Engine/Plugins/Experimental/PlanarCutPlugin/Source/PlanarCut/Private/GeometryMeshConversion.cpp
+++ b/Engine/Plugins/Experimental/PlanarCutPlugin/Source/PlanarCut/Private/GeometryMeshConversion.cpp
@@ -2472,7 +2472,11 @@ int32 FDynamicMeshCollection::CutWithMultiplePlanes(
 	if (GeometryForRemoval.Num() > 0)
 	{
 		GeometryForRemoval.Sort();
-		Collection->RemoveElements(FGeometryCollection::GeometryGroup, GeometryForRemoval);
+		FManagedArrayCollection::FProcessingParameters ProcessingParams;
+#if !UE_BUILD_DEBUG
+		ProcessingParams.bDoValidation = false;
+#endif
+		Collection->RemoveElements(FGeometryCollection::GeometryGroup, GeometryForRemoval, ProcessingParams);
 	}
 
 	return FirstCreatedIndex - GeometryForRemoval.Num();
@@ -2659,6 +2663,10 @@ int32 FDynamicMeshCollection::CutWithCellMeshes(const FInternalSurfaceMaterials&
 	if (GeometryForRemoval.Num() > 0)
 	{
 		GeometryForRemoval.Sort();
+		FManagedArrayCollection::FProcessingParameters ProcessingParams;
+#if !UE_BUILD_DEBUG
+		ProcessingParams.bDoValidation = false;
+#endif
 		Collection->RemoveElements(FGeometryCollection::GeometryGroup, GeometryForRemoval);
 	}
 
@@ -2805,7 +2813,11 @@ bool FDynamicMeshCollection::UpdateAllCollections(FGeometryCollection& Collectio
 		NewFaceCounts[GeomIdx] = MeshData.AugMesh.TriangleCount();
 		NewVertexCounts[GeomIdx] = MeshData.AugMesh.VertexCount();
 	}
-	GeometryCollectionAlgo::ResizeGeometries(&Collection, NewFaceCounts, NewVertexCounts);
+	bool bDoValidation = false;
+#if UE_BUILD_DEBUG
+	bDoValidation = true;
+#endif
+	GeometryCollectionAlgo::ResizeGeometries(&Collection, NewFaceCounts, NewVertexCounts, bDoValidation);
 
 	for (int MeshIdx = 0; MeshIdx < Meshes.Num(); MeshIdx++)
 	{
diff --git a/Engine/Source/Runtime/Experimental/Chaos/Private/GeometryCollection/GeometryCollectionAlgo.cpp b/Engine/Source/Runtime/Experimental/Chaos/Private/GeometryCollection/GeometryCollectionAlgo.cpp
--- a/Engine/Source/Runtime/Experimental/Chaos/Private/GeometryCollection/GeometryCollectionAlgo.cpp
+++ b/Engine/Source/Runtime/Experimental/Chaos/Private/GeometryCollection/GeometryCollectionAlgo.cpp
@@ -817,7 +817,7 @@ namespace GeometryCollectionAlgo
 		GeometryCollection->ReorderElements(FGeometryCollection::FacesGroup, ShiftedOrder);
 	}
 
-	void ResizeGeometries(FGeometryCollection* GeometryCollection, const TArray<int32>& FaceCounts, const TArray<int32>& VertexCounts)
+	void ResizeGeometries(FGeometryCollection* GeometryCollection, const TArray<int32>& FaceCounts, const TArray<int32>& VertexCounts, bool bDoValidation)
 	{
 		check(GeometryCollection);
 		int32 NumGeometries = GeometryCollection->NumElements(FGeometryCollection::GeometryGroup);
@@ -954,24 +954,29 @@ namespace GeometryCollectionAlgo
 			}
 		}
 
+		FManagedArrayCollection::FProcessingParameters ProcessingParams;
+		ProcessingParams.bDoValidation = bDoValidation;
+
 		// remove trailing elements if needed
 		if (NewNumVertices < OldNumVertices)
 		{
 			TArray<int32> ToDelete;
 			AddRange(ToDelete, NewNumVertices, OldNumVertices);
-			GeometryCollection->RemoveElements(FGeometryCollection::VerticesGroup, ToDelete);
+			GeometryCollection->RemoveElements(FGeometryCollection::VerticesGroup, ToDelete, ProcessingParams);
 		}
 		if (NewNumFaces < OldNumFaces)
 		{
 			TArray<int32> ToDelete;
 			AddRange(ToDelete, NewNumFaces, OldNumFaces);
-			GeometryCollection->RemoveElements(FGeometryCollection::FacesGroup, ToDelete);
+			GeometryCollection->RemoveElements(FGeometryCollection::FacesGroup, ToDelete, ProcessingParams);
 		}
 
-		// TODO: can remove these ensure()s after this has been tested some more
-		ensure(GeometryCollection->HasContiguousFaces());
-		ensure(GeometryCollection->HasContiguousVertices());
-		ensure(GeometryCollectionAlgo::HasValidGeometryReferences(GeometryCollection));
+		if (bDoValidation)
+		{
+			ensure(GeometryCollection->HasContiguousFaces());
+			ensure(GeometryCollection->HasContiguousVertices());
+			ensure(GeometryCollectionAlgo::HasValidGeometryReferences(GeometryCollection));
+		}
 	}
 
 	DEFINE_LOG_CATEGORY_STATIC(LogGeometryCollectionClean, Verbose, All);
diff --git a/Engine/Source/Runtime/Experimental/Chaos/Public/GeometryCollection/GeometryCollectionAlgo.h b/Engine/Source/Runtime/Experimental/Chaos/Public/GeometryCollection/GeometryCollectionAlgo.h
--- a/Engine/Source/Runtime/Experimental/Chaos/Public/GeometryCollection/GeometryCollectionAlgo.h
+++ b/Engine/Source/Runtime/Experimental/Chaos/Public/GeometryCollection/GeometryCollectionAlgo.h
@@ -191,7 +191,7 @@ namespace GeometryCollectionAlgo
 	 */
 	void
 	CHAOS_API
-	ResizeGeometries(FGeometryCollection* GeometryCollection, const TArray<int32>& FaceCounts, const TArray<int32>& VertexCounts);
+	ResizeGeometries(FGeometryCollection* GeometryCollection, const TArray<int32>& FaceCounts, const TArray<int32>& VertexCounts, bool bDoValidation = true);
 
 	void
 	CHAOS_API
