ChaosEventHandler
-- added ScopeLock include to fix non-unity build issue

diff --git a/Engine/Source/Runtime/Experimental/Chaos/Public/EventManager.h b/Engine/Source/Runtime/Experimental/Chaos/Public/EventManager.h
--- a/Engine/Source/Runtime/Experimental/Chaos/Public/EventManager.h
+++ b/Engine/Source/Runtime/Experimental/Chaos/Public/EventManager.h
@@ -6,6 +6,7 @@
 #include "Chaos/Framework/MultiBufferResource.h"
 #include "Chaos/Framework/PhysicsProxyBase.h"
 #include "Chaos/PBDRigidsEvolutionFwd.h"
+#include "Misc/ScopeLock.h"
 #include "Chaos/Defines.h"
 
 namespace Chaos
