Fixed Geometry collection rendering issues on mobile

diff --git a/Engine/Source/Runtime/Experimental/GeometryCollectionEngine/Private/GeometryCollection/GeometryCollectionSceneProxy.cpp b/Engine/Source/Runtime/Experimental/GeometryCollectionEngine/Private/GeometryCollection/GeometryCollectionSceneProxy.cpp
--- a/Engine/Source/Runtime/Experimental/GeometryCollectionEngine/Private/GeometryCollection/GeometryCollectionSceneProxy.cpp
+++ b/Engine/Source/Runtime/Experimental/GeometryCollectionEngine/Private/GeometryCollection/GeometryCollectionSceneProxy.cpp
@@ -616,8 +616,8 @@ void FGeometryCollectionSceneProxy::SetDynamicData_RenderThread(FGeometryCollect
 	#else
 					for (int32 i = IndexOffset; i < IndexOffset + ThisBatchSize; i++)
 					{
-						FVector Transformed = BoneTransformsPtr[ConstantData->BoneMap[i]].TransformPosition(ConstantData->Vertices[i]);
-						FMemory::Memcpy((uint8*)VertexBufferData + (i * VertexBuffer.GetStride()), &Transformed, sizeof(FVector));
+						FVector3f Transformed = BoneTransformsPtr[ConstantData->BoneMap[i]].TransformPosition(ConstantData->Vertices[i]);
+						FMemory::Memcpy((uint8*)VertexBufferData + (i * VertexBuffer.GetStride()), &Transformed, sizeof(FVector3f));
 					}
 	#endif
 				}
