Starting to expose joints to rewind data system

diff --git a/Engine/Source/Runtime/Experimental/Chaos/Private/PBDRigidsSolver.cpp b/Engine/Source/Runtime/Experimental/Chaos/Private/PBDRigidsSolver.cpp
--- a/Engine/Source/Runtime/Experimental/Chaos/Private/PBDRigidsSolver.cpp
+++ b/Engine/Source/Runtime/Experimental/Chaos/Private/PBDRigidsSolver.cpp
@@ -1034,7 +1034,7 @@ namespace Chaos
 				{
 					JointConstraintPhysicsProxies_Internal.Add(JointProxy);
 					JointProxy->InitializeOnPhysicsThread(this, *Manager, DataIdx, Dirty.PropertyData);
-					JointProxy->SetInitialized();
+					JointProxy->SetInitialized(GetCurrentFrame());
 				}
 				JointProxy->PushStateOnPhysicsThread(this, *Manager, DataIdx, Dirty.PropertyData);
 				Dirty.Proxy->ResetDirtyIdx();
diff --git a/Engine/Source/Runtime/Experimental/Chaos/Private/RewindData.cpp b/Engine/Source/Runtime/Experimental/Chaos/Private/RewindData.cpp
--- a/Engine/Source/Runtime/Experimental/Chaos/Private/RewindData.cpp
+++ b/Engine/Source/Runtime/Experimental/Chaos/Private/RewindData.cpp
@@ -4,6 +4,7 @@
 #include "PBDRigidsSolver.h"
 #include "PhysicsProxy/SingleParticlePhysicsProxy.h"
 #include "HAL/IConsoleManager.h"
+#include "Chaos/PBDJointConstraints.h"
 
 namespace Chaos
 {
@@ -11,7 +12,7 @@ namespace Chaos
 int32 EnableResimCache = 1;
 FAutoConsoleVariableRef CVarEnableEnableResimCache(TEXT("p.EnableResimCache"), EnableResimCache, TEXT("If enabled, provides a resim cache to speed up certain computations"));
 
-FVec3 FGeometryParticleStateBase::ZeroVector = FVec3(0);
+FVec3 FGeometryParticleState::ZeroVector = FVec3(0);
 
 void FGeometryParticleStateBase::SyncSimWritablePropsFromSim(FDirtyPropData Manager,const TPBDRigidParticleHandle<FReal,3>& Rigid)
 {
@@ -394,18 +395,31 @@ void FRewindData::AdvanceFrameImp(IResimCacheBase* ResimCache)
 #define REWIND_DESYNC 0
 #endif
 
+template <typename TData, typename TConcrete>
+void CopyDataFromConcrete(TData& Data, const TConcrete& Concrete)
+{
+	Data.CopyFrom(Concrete);
+}
+
+void CopyDataFromConcrete(FPBDJointSettings& Data, const FPBDJointConstraintHandle& Joint)
+{
+	Data = Joint.GetSettings();
+}
+
+
+
 template <bool bResim>
 void FRewindData::PushGTDirtyData(const FDirtyPropertiesManager& SrcManager,const int32 SrcDataIdx,const FDirtyProxy& Dirty, const FShapeDirtyData* ShapeDirtyData)
 {
 	//This records changes enqueued by GT.
 	bNeedsSave = true;
 
-	if(ensure(Dirty.Proxy->GetType() == EPhysicsProxyType::SingleParticleProxy))
-	{
-		auto Proxy = static_cast<FSingleParticlePhysicsProxy*>(Dirty.Proxy);
-		
-		FGeometryParticleHandle* PTParticle = Proxy->GetHandle_LowLevel();
+	IPhysicsProxyBase* Proxy = Dirty.Proxy;
 
+	//Helper to group most of the common logic about push data recording
+	//NOTE: when possible use passed in CopyFunc to do work, if lambda returns false you cannot record to history buffer
+	auto CopyHelper = [this, Proxy](auto Concrete, const auto& CopyFunc) -> bool
+	{
 		//Don't bother tracking static particles. We assume they stream in and out and don't need to be rewound
 		//TODO: find a way to skip statics that stream in and out - gameplay can technically spawn/destroy these so we can't just ignore statics
 		/*if(PTParticle->CastToKinematicParticle() == nullptr)
@@ -415,9 +429,9 @@ void FRewindData::PushGTDirtyData(const FDirtyPropertiesManager& SrcManager,cons
 
 		//During a resim the same exact push data comes from gt
 		//If the particle is already in sync, it will stay in sync so no need to touch history
-		if (bResim && PTParticle->SyncState() == ESyncState::InSync)
+		if (bResim && Concrete->SyncState() == ESyncState::InSync)
 		{
-			return;
+			return false;
 		}
 
 		if (bResim && Proxy->GetInitializedStep() == CurFrame)
@@ -425,31 +439,48 @@ void FRewindData::PushGTDirtyData(const FDirtyPropertiesManager& SrcManager,cons
 			//Particle is reinitialized, since it's out of sync it must be at a different time
 			//So make sure it's considered during resim
 			//TODO: should check if in bubble
-			PTParticle->SetEnabledDuringResim(true);
+			Concrete->SetEnabledDuringResim(true);
 		}
 
-		FDirtyParticleInfo& Info = FindOrAddParticle(*PTParticle, Proxy->IsInitialized() ? INDEX_NONE : CurFrame);
-		FGeometryParticleStateBase& Latest = Info.AddFrame(CurFrame);
+		auto& Info = FindOrAddConcrete(*Concrete, Proxy->IsInitialized() ? INDEX_NONE : CurFrame);
+		auto& Latest = Info.AddFrame(CurFrame);
 
 		//At this point all phases should be clean
 		ensure(Latest.IsClean(FFrameAndPhase{ CurFrame, FFrameAndPhase::PrePushData }));
 
-		//Most particles never change but may be created/destroyed often due to streaming
+		//Most objects never change but may be created/destroyed often due to streaming
 		//To avoid useless writes we call this function before PushData is processed.
-		//This means we will skip particles that are streamed in since they never change
+		//This means we will skip objects that are streamed in since they never change
 		//So if Proxy has initialized it means the particle isn't just streaming in, it's actually changing
-		if(Info.InitializedOnStep < CurFrame)
+		if (Info.InitializedOnStep < CurFrame)
 		{
-			const FFrameAndPhase PrePushData = FFrameAndPhase{ CurFrame, FFrameAndPhase::PrePushData };
-			auto DirtyPropHelper = [this, &Dirty](auto& Property, const EChaosPropertyFlags PropName, const auto& Particle)
-			{
-				if (Dirty.PropertyData.IsDirty(PropName))
-				{
-					auto& Data = Property.WriteAccessMonotonic(FFrameAndPhase{ CurFrame, FFrameAndPhase::PrePushData }, PropertiesPool);
-					Data.CopyFrom(Particle);
-				}
-			};
+			CopyFunc(Latest);
+		}
+
+		ensure(Latest.IsClean(FFrameAndPhase{ CurFrame, FFrameAndPhase::PostPushData }));   //PostPushData is untouched
+		ensure(Latest.IsClean(FFrameAndPhase{ CurFrame, FFrameAndPhase::PostCallbacks }));	//PostCallback is untouched
+
+		return true;
+	};
+
+	auto DirtyPropHelper = [this, &Dirty](auto& Property, const EChaosPropertyFlags PropName, const auto& Concrete)
+	{
+		if (Dirty.PropertyData.IsDirty(PropName))
+		{
+			auto& Data = Property.WriteAccessMonotonic(FFrameAndPhase{ CurFrame, FFrameAndPhase::PrePushData }, PropertiesPool);
+			CopyDataFromConcrete(Data, Concrete);
+		}
+	};
+
+	switch(Dirty.Proxy->GetType())
+	{
+	case EPhysicsProxyType::SingleParticleProxy:
+	{
+		auto ParticleProxy = static_cast<FSingleParticlePhysicsProxy*>(Dirty.Proxy);
+		FGeometryParticleHandle* PTParticle = ParticleProxy->GetHandle_LowLevel();
 
+		const bool bKeepRecording = CopyHelper(PTParticle, [PTParticle, &DirtyPropHelper](FGeometryParticleStateBase& Latest)
+		{
 			DirtyPropHelper(Latest.ParticlePositionRotation, EChaosPropertyFlags::XR, *PTParticle);
 			DirtyPropHelper(Latest.NonFrequentData, EChaosPropertyFlags::NonFrequentData, *PTParticle);
 
@@ -464,17 +495,39 @@ void FRewindData::PushGTDirtyData(const FDirtyPropertiesManager& SrcManager,cons
 					DirtyPropHelper(Latest.MassProps, EChaosPropertyFlags::MassProps, *Rigid);
 				}
 			}
-		}
+		});
 
-		//Dynamics are not available at head (sim zeroes them out), so we have to record them as PostPushData (since they're applied as part of PushData)
-		if (auto NewData = Dirty.PropertyData.FindDynamics(SrcManager, SrcDataIdx))
+		if(bKeepRecording)
 		{
-			Latest.Dynamics.WriteAccessMonotonic(FFrameAndPhase{ CurFrame, FFrameAndPhase::PostPushData }, PropertiesPool) = *NewData;
-			Info.DirtyDynamics = CurFrame;	//Need to save the dirty dynamics into the next phase as well (it's possible a callback will stomp the dynamics value, so that's why it's pending)
+			//Dynamics are not available at head (sim zeroes them out), so we have to record them as PostPushData (since they're applied as part of PushData)
+			if (auto NewData = Dirty.PropertyData.FindDynamics(SrcManager, SrcDataIdx))
+			{
+				FDirtyParticleInfo& Info = FindOrAddParticle(*PTParticle, ParticleProxy->IsInitialized() ? INDEX_NONE : CurFrame);
+				FGeometryParticleStateBase& Latest = Info.AddFrame(CurFrame);
+				const FFrameAndPhase PostPushData{ CurFrame, FFrameAndPhase::PostPushData };
+				Latest.Dynamics.WriteAccessMonotonic(PostPushData, PropertiesPool) = *NewData;
+				Info.DirtyDynamics = CurFrame;	//Need to save the dirty dynamics into the next phase as well (it's possible a callback will stomp the dynamics value, so that's why it's pending)
+
+				ensure(Latest.IsCleanExcludingDynamics(PostPushData)); //PostPushData is untouched except for dynamics
+			}
 		}
+		break;
+	}
+	case EPhysicsProxyType::JointConstraintType:
+	{
+		auto JointProxy = static_cast<FJointConstraintPhysicsProxy*>(Dirty.Proxy);
+		FPBDJointConstraintHandle* Joint = JointProxy->GetHandle();
 
-		ensure(Latest.IsCleanExcludingDynamics(FFrameAndPhase{ CurFrame, FFrameAndPhase::PostPushData })); //PostPushData is untouched except for dynamics
-		ensure(Latest.IsClean(FFrameAndPhase{ CurFrame, FFrameAndPhase::PostCallbacks }));	//PostCallback should be untouched
+		CopyHelper(Joint, [Joint, &DirtyPropHelper](FJointStateBase& Latest)
+		{
+			DirtyPropHelper(Latest.JointSettings, EChaosPropertyFlags::JointSettings, *Joint);
+		});
+		break;
+	}
+	default:
+	{
+		ensure(false);	//Unsupported proxy type
+	}
 	}
 }
 
@@ -592,41 +645,12 @@ FRewindData::FDirtyParticleInfo& FRewindData::FindOrAddParticle(TGeometryParticl
 	return AllDirtyParticles[DirtyIdx];
 }
 
-template <bool bResim>
-bool FRewindData::FSimWritableState::SyncSimWritablePropsFromSim(const TPBDRigidParticleHandle<FReal,3>& Rigid, const int32 Frame)
-{
-	FrameRecordedHack = Frame;
-	bool bDesynced = false;
-	if(bResim)
-	{
-		bDesynced |= Rigid.P() != MX;
-		bDesynced |= Rigid.Q() != MR;
-		bDesynced |= Rigid.V() != MV;
-		bDesynced |= Rigid.W() != MW;
-	}
-
-	MX = Rigid.P();
-	MR = Rigid.Q();
-	MV = Rigid.V();
-	MW = Rigid.W();
-
-	return bDesynced;
-}
-
-void FRewindData::FSimWritableState::SyncToParticle(TPBDRigidParticleHandle<FReal,3>& Rigid) const
-{
-	Rigid.SetX(MX);
-	Rigid.SetR(MR);
-	Rigid.SetV(MV);
-	Rigid.SetW(MW);
-}
-
-FRewindData::FDirtyParticleInfo::~FDirtyParticleInfo()
+FRewindData::FDirtyJointInfo& FRewindData::FindOrAddJoint(FPBDJointConstraintHandle& Handle, const int32 InitializedOnFrame)
 {
-	if(PropertiesPool)
-	{
-		History.Release(*PropertiesPool);
-}
+	//TODO: implement
+	ensure(false);
+	static FDirtyJointInfo TempReturn{ PropertiesPool, Handle, CurFrame, 0 };
+	return TempReturn;
 }
 
 template void FRewindData::PushGTDirtyData<true>(const FDirtyPropertiesManager& SrcManager,const int32 SrcDataIdx,const FDirtyProxy& Dirty, const FShapeDirtyData* DirtyShapesData);
diff --git a/Engine/Source/Runtime/Experimental/Chaos/Public/Chaos/Framework/PhysicsProxyBase.h b/Engine/Source/Runtime/Experimental/Chaos/Public/Chaos/Framework/PhysicsProxyBase.h
--- a/Engine/Source/Runtime/Experimental/Chaos/Public/Chaos/Framework/PhysicsProxyBase.h
+++ b/Engine/Source/Runtime/Experimental/Chaos/Public/Chaos/Framework/PhysicsProxyBase.h
@@ -73,6 +73,10 @@ public:
 
 	TSharedPtr<FProxyTimestamp,ESPMode::ThreadSafe> GetSyncTimestamp() const { return SyncTimestamp; }
 
+	bool IsInitialized() const { return InitializedOnStep != INDEX_NONE; }
+	void SetInitialized(const int32 InitializeStep) { InitializedOnStep = InitializeStep; }
+	int32 GetInitializedStep() const { return InitializedOnStep; }
+
 
 protected:
 	// Ensures that derived classes can successfully call this destructor
@@ -89,6 +93,7 @@ protected:
 	/** Proxy type */
 	EPhysicsProxyType Type;
 	TSharedPtr<FProxyTimestamp,ESPMode::ThreadSafe> SyncTimestamp;
+	int32 InitializedOnStep = INDEX_NONE;
 };
 
 struct PhysicsProxyWrapper
diff --git a/Engine/Source/Runtime/Experimental/Chaos/Public/Chaos/PBDJointConstraints.h b/Engine/Source/Runtime/Experimental/Chaos/Public/Chaos/PBDJointConstraints.h
--- a/Engine/Source/Runtime/Experimental/Chaos/Public/Chaos/PBDJointConstraints.h
+++ b/Engine/Source/Runtime/Experimental/Chaos/Public/Chaos/PBDJointConstraints.h
@@ -52,6 +52,10 @@ namespace Chaos
 		void SetSettings(const FPBDJointSettings& Settings);
 		TVec2<FGeometryParticleHandle*> GetConstrainedParticles() const;
 
+		//TODO: implement this for RewindData to work correctly
+		ESyncState SyncState() const { ensure(false); return ESyncState::InSync; }
+		void SetEnabledDuringResim(bool bEnabled) { ensure(false); }
+
 	protected:
 		using Base::ConstraintIndex;
 		using Base::ConcreteContainer;
diff --git a/Engine/Source/Runtime/Experimental/Chaos/Public/PhysicsProxy/JointConstraintProxy.h b/Engine/Source/Runtime/Experimental/Chaos/Public/PhysicsProxy/JointConstraintProxy.h
--- a/Engine/Source/Runtime/Experimental/Chaos/Public/PhysicsProxy/JointConstraintProxy.h
+++ b/Engine/Source/Runtime/Experimental/Chaos/Public/PhysicsProxy/JointConstraintProxy.h
@@ -28,9 +28,6 @@ public:
 	FJointConstraintPhysicsProxy() = delete;
 	FJointConstraintPhysicsProxy(FJointConstraint* InConstraint, FPBDJointConstraintHandle* InHandle, UObject* InOwner = nullptr);
 
-	bool IsInitialized() const { return bInitialized; }
-	void SetInitialized() { bInitialized = true; }
-
 	static FGeometryParticleHandle* GetParticleHandleFromProxy(IPhysicsProxyBase* ProxyBase);
 
 	//
diff --git a/Engine/Source/Runtime/Experimental/Chaos/Public/PhysicsProxy/SingleParticlePhysicsProxy.h b/Engine/Source/Runtime/Experimental/Chaos/Public/PhysicsProxy/SingleParticlePhysicsProxy.h
--- a/Engine/Source/Runtime/Experimental/Chaos/Public/PhysicsProxy/SingleParticlePhysicsProxy.h
+++ b/Engine/Source/Runtime/Experimental/Chaos/Public/PhysicsProxy/SingleParticlePhysicsProxy.h
@@ -139,9 +139,6 @@ public:
 	/**/
 	bool IsDirty();
 
-	bool IsInitialized() const { return InitializedOnStep != INDEX_NONE; }
-	void SetInitialized(const int32 InitializeStep) { InitializedOnStep = InitializeStep; }
-	int32 GetInitializedStep() const { return InitializedOnStep; }
 
 	/**/
 	Chaos::EWakeEventEntry GetWakeEvent() const;
@@ -170,9 +167,6 @@ public:
 	{
 		return static_cast<const Chaos::FPBDRigidParticle*>(GetParticle_LowLevel());
 	}
-	
-private:
-	int32 InitializedOnStep = INDEX_NONE;
 
 protected:
 	TUniquePtr<PARTICLE_TYPE> Particle;
diff --git a/Engine/Source/Runtime/Experimental/Chaos/Public/RewindData.h b/Engine/Source/Runtime/Experimental/Chaos/Public/RewindData.h
--- a/Engine/Source/Runtime/Experimental/Chaos/Public/RewindData.h
+++ b/Engine/Source/Runtime/Experimental/Chaos/Public/RewindData.h
@@ -6,6 +6,7 @@
 #include "Chaos/Framework/PhysicsSolverBase.h"
 #include "Containers/CircularBuffer.h"
 #include "Chaos/ResimCacheBase.h"
+#include "Chaos/PBDJointConstraints.h"
 
 #ifndef VALIDATE_REWIND_DATA
 #define VALIDATE_REWIND_DATA 0
@@ -423,17 +424,61 @@ private:
 	const FShapesArrayStateBase* State;
 };
 
-class FGeometryParticleStateBase
+#define REWIND_CHAOS_PARTICLE_PROPERTY(PROP, NAME)\
+		const auto Data = State ? State->PROP.Read(FrameAndPhase, Pool) : nullptr;\
+		return Data ? Data->NAME() : Head.NAME();\
+
+#define REWIND_CHAOS_ZERO_PARTICLE_PROPERTY(PROP, NAME)\
+		const auto Data = State ? State->PROP.Read(FrameAndPhase, Pool) : nullptr;\
+		return Data ? Data->NAME() : ZeroVector;\
+
+#define REWIND_PARTICLE_STATIC_PROPERTY(PROP, NAME)\
+	decltype(auto) NAME() const\
+	{\
+		auto& Head = Particle;\
+		REWIND_CHAOS_PARTICLE_PROPERTY(PROP, NAME);\
+	}\
+
+#define REWIND_PARTICLE_KINEMATIC_PROPERTY(PROP, NAME)\
+	decltype(auto) NAME() const\
+	{\
+		auto& Head = *Particle.CastToKinematicParticle();\
+		REWIND_CHAOS_PARTICLE_PROPERTY(PROP, NAME);\
+	}\
+
+#define REWIND_PARTICLE_RIGID_PROPERTY(PROP, NAME)\
+	decltype(auto) NAME() const\
+	{\
+		auto& Head = *Particle.CastToRigidParticle();\
+		REWIND_CHAOS_PARTICLE_PROPERTY(PROP, NAME);\
+	}\
+
+#define REWIND_PARTICLE_ZERO_PROPERTY(PROP, NAME)\
+	decltype(auto) NAME() const\
+	{\
+		auto& Head = *Particle.CastToRigidParticle();\
+		REWIND_CHAOS_ZERO_PARTICLE_PROPERTY(PROP, NAME);\
+	}\
+
+#define REWIND_JOINT_PROPERTY(PROP, FUNC_NAME, NAME)\
+	decltype(auto) FUNC_NAME() const\
+	{\
+		const auto Data = State ? State->PROP.Read(FrameAndPhase, Pool) : nullptr;\
+		return Data ? Data->NAME : Head.GetSettings().NAME;\
+	}\
+
+inline int32 ComputeCircularSize(int32 NumFrames) { return NumFrames * FFrameAndPhase::NumPhases; }
+
+struct FGeometryParticleStateBase
 {
-public:
 	explicit FGeometryParticleStateBase(int32 NumFrames)
-	: ParticlePositionRotation(NumFrames * FFrameAndPhase::NumPhases)
-	, NonFrequentData(NumFrames* FFrameAndPhase::NumPhases)
-	, Velocities(NumFrames* FFrameAndPhase::NumPhases)
-	, Dynamics(NumFrames * FFrameAndPhase::NumPhases)
-	, DynamicsMisc(NumFrames* FFrameAndPhase::NumPhases)
-	, MassProps(NumFrames* FFrameAndPhase::NumPhases)
-	, KinematicTarget(NumFrames* FFrameAndPhase::NumPhases)
+	: ParticlePositionRotation(ComputeCircularSize(NumFrames))
+	, NonFrequentData(ComputeCircularSize(NumFrames))
+	, Velocities(ComputeCircularSize(NumFrames))
+	, Dynamics(ComputeCircularSize(NumFrames))
+	, DynamicsMisc(ComputeCircularSize(NumFrames))
+	, MassProps(ComputeCircularSize(NumFrames))
+	, KinematicTarget(ComputeCircularSize(NumFrames))
 	{
 
 	}
@@ -492,191 +537,7 @@ public:
 
 	template <bool bSkipDynamics = false>
 	bool IsInSync(const FGeometryParticleHandle& Handle, const FFrameAndPhase FrameAndPhase, const FDirtyPropertiesPool& Pool) const;
-
-	template <typename TParticle>
-	static const FVec3& X(const FGeometryParticleStateBase* State, const TParticle& Particle, const FDirtyPropertiesPool& Manager, const FFrameAndPhase FrameAndPhase)
-	{
-		const auto Data = State ? State->ParticlePositionRotation.Read(FrameAndPhase, Manager) : nullptr;
-		return Data ? Data->X() : Particle.X();
-	}
-
-	template <typename TParticle>
-	static const FRotation3& R(const FGeometryParticleStateBase* State, const TParticle& Particle, const FDirtyPropertiesPool& Manager, const FFrameAndPhase FrameAndPhase)
-	{
-		const auto Data = State ? State->ParticlePositionRotation.Read(FrameAndPhase, Manager) : nullptr;
-		return Data ? Data->R() : Particle.R();
-	}
-	
-	template <typename TParticle>
-	static const FVec3& V(const FGeometryParticleStateBase* State, const TParticle& Particle, const FDirtyPropertiesPool& Manager, const FFrameAndPhase FrameAndPhase)
-	{
-		const auto Data = State ? State->Velocities.Read(FrameAndPhase, Manager) : nullptr;
-		return Data ? Data->V() : Particle.CastToKinematicParticle()->V();
-	}
-
-	template <typename TParticle>
-	static const FVec3& W(const FGeometryParticleStateBase* State, const TParticle& Particle, const FDirtyPropertiesPool& Manager, const FFrameAndPhase FrameAndPhase)
-	{
-		const auto Data = State ? State->Velocities.Read(FrameAndPhase, Manager) : nullptr;
-		return Data ? Data->W() : Particle.CastToKinematicParticle()->W();
-	}
-
-	template <typename TParticle>
-	static FReal LinearEtherDrag(const FGeometryParticleStateBase* State, const TParticle& Particle, const FDirtyPropertiesPool& Manager, const FFrameAndPhase FrameAndPhase)
-	{
-		const auto Data = State ? State->DynamicsMisc.Read(FrameAndPhase, Manager) : nullptr;
-		return Data ? Data->LinearEtherDrag() : Particle.CastToRigidParticle()->LinearEtherDrag();
-	}
-
-	template <typename TParticle>
-	static FReal AngularEtherDrag(const FGeometryParticleStateBase* State, const TParticle& Particle, const FDirtyPropertiesPool& Manager, const FFrameAndPhase FrameAndPhase)
-	{
-		const auto Data = State ? State->DynamicsMisc.Read(FrameAndPhase, Manager) : nullptr;
-		return Data ? Data->AngularEtherDrag() : Particle.CastToRigidParticle()->AngularEtherDrag();
-	}
-	
-	template <typename TParticle>
-	static FReal MaxLinearSpeedSq(const FGeometryParticleStateBase* State, const TParticle& Particle, const FDirtyPropertiesPool& Manager, const FFrameAndPhase FrameAndPhase)
-	{
-		const auto Data = State ? State->DynamicsMisc.Read(FrameAndPhase, Manager) : nullptr;
-		return Data ? Data->MaxLinearSpeedSq() : Particle.CastToRigidParticle()->MaxLinearSpeedSq();
-	}
-
-	template <typename TParticle>
-	static FReal MaxAngularSpeedSq(const FGeometryParticleStateBase* State, const TParticle& Particle, const FDirtyPropertiesPool& Manager, const FFrameAndPhase FrameAndPhase)
-	{
-		const auto Data = State ? State->DynamicsMisc.Read(FrameAndPhase, Manager) : nullptr;
-		return Data ? Data->MaxAngularSpeedSq() : Particle.CastToRigidParticle()->MaxAngularSpeedSq();
-	}
-
-	template <typename TParticle>
-	static EObjectStateType ObjectState(const FGeometryParticleStateBase* State, const TParticle& Particle, const FDirtyPropertiesPool& Manager, const FFrameAndPhase FrameAndPhase)
-	{
-		const auto Data = State ? State->DynamicsMisc.Read(FrameAndPhase, Manager) : nullptr;
-		return Data ? Data->ObjectState() : Particle.CastToRigidParticle()->ObjectState();
-	}
-
-	template <typename TParticle>
-	static bool GravityEnabled(const FGeometryParticleStateBase* State, const TParticle& Particle, const FDirtyPropertiesPool& Manager, const FFrameAndPhase FrameAndPhase)
-	{
-		const auto Data = State ? State->DynamicsMisc.Read(FrameAndPhase, Manager) : nullptr;
-		return Data ? Data->GravityEnabled() : Particle.CastToRigidParticle()->GravityEnabled();
-	}
-	
-	template <typename TParticle>
-	static bool CCDEnabled(const FGeometryParticleStateBase* State, const TParticle& Particle, const FDirtyPropertiesPool& Manager, const FFrameAndPhase FrameAndPhase)
-	{
-		const auto Data = State ? State->DynamicsMisc.Read(FrameAndPhase, Manager) : nullptr;
-		return Data ? Data->CCDEnabled() : Particle.CastToRigidParticle()->CCDEnabled();
-	}
-
-	template <typename TParticle>
-	static int32 CollisionGroup(const FGeometryParticleStateBase* State, const TParticle& Particle, const FDirtyPropertiesPool& Manager, const FFrameAndPhase FrameAndPhase)
-	{
-		const auto Data = State ? State->DynamicsMisc.Read(FrameAndPhase, Manager) : nullptr;
-		return Data ? Data->CollisionGroup() : Particle.CastToRigidParticle()->CollisionGroup();
-	}
-
-	template <typename TParticle>
-	static const FVec3& CenterOfMass(const FGeometryParticleStateBase* State, const TParticle& Particle, const FDirtyPropertiesPool& Manager, const FFrameAndPhase FrameAndPhase)
-	{
-		const auto Data = State ? State->MassProps.Read(FrameAndPhase, Manager) : nullptr;
-		return Data ? Data->CenterOfMass() : Particle.CastToRigidParticle()->CenterOfMass();
-	}
-
-	template <typename TParticle>
-	static const FRotation3& RotationOfMass(const FGeometryParticleStateBase* State, const TParticle& Particle, const FDirtyPropertiesPool& Manager, const FFrameAndPhase FrameAndPhase)
-	{
-		const auto Data = State ? State->MassProps.Read(FrameAndPhase, Manager) : nullptr;
-		return Data ? Data->RotationOfMass() : Particle.CastToRigidParticle()->RotationOfMass();
-	}
-
-	template <typename TParticle>
-	static const FMatrix33& I(const FGeometryParticleStateBase* State, const TParticle& Particle, const FDirtyPropertiesPool& Manager, const FFrameAndPhase FrameAndPhase)
-	{
-		const auto Data = State ? State->MassProps.Read(FrameAndPhase, Manager) : nullptr;
-		return Data ? Data->I() : Particle.CastToRigidParticle()->I();
-	}
-
-	template <typename TParticle>
-	static const FMatrix33& InvI(const FGeometryParticleStateBase* State, const TParticle& Particle, const FDirtyPropertiesPool& Manager, const FFrameAndPhase FrameAndPhase)
-	{
-		const auto Data = State ? State->MassProps.Read(FrameAndPhase, Manager) : nullptr;
-		return Data ? Data->InvI() : Particle.CastToRigidParticle()->InvI();
-	}
-
-	template <typename TParticle>
-	static FReal M(const FGeometryParticleStateBase* State, const TParticle& Particle, const FDirtyPropertiesPool& Manager, const FFrameAndPhase FrameAndPhase)
-	{
-		const auto Data = State ? State->MassProps.Read(FrameAndPhase, Manager) : nullptr;
-		return Data ? Data->M() : Particle.CastToRigidParticle()->M();
-	}
-
-	template <typename TParticle>
-	static FReal InvM(const FGeometryParticleStateBase* State, const TParticle& Particle, const FDirtyPropertiesPool& Manager, const FFrameAndPhase FrameAndPhase)
-	{
-		const auto Data = State ? State->MassProps.Read(FrameAndPhase, Manager) : nullptr;
-		return Data ? Data->InvM() : Particle.CastToRigidParticle()->InvM();
-	}
-
-	template <typename TParticle>
-	static TSerializablePtr<FImplicitObject> Geometry(const FGeometryParticleStateBase* State, const TParticle& Particle, const FDirtyPropertiesPool& Manager, const FFrameAndPhase FrameAndPhase)
-	{
-		const auto Data = State ? State->NonFrequentData.Read(FrameAndPhase, Manager) : nullptr;
-		return Data ? Data->Geometry() : Particle.Geometry();
-	}
-
-	template <typename TParticle>
-	static FUniqueIdx UniqueIdx(const FGeometryParticleStateBase* State, const TParticle& Particle, const FDirtyPropertiesPool& Manager, const FFrameAndPhase FrameAndPhase)
-	{
-		const auto Data = State ? State->NonFrequentData.Read(FrameAndPhase, Manager) : nullptr;
-		return Data ? Data->UniqueIdx() : Particle.UniqueIdx();
-	}
 	
-	template <typename TParticle>
-	static FSpatialAccelerationIdx SpatialIdx(const FGeometryParticleStateBase* State, const TParticle& Particle, const FDirtyPropertiesPool& Manager, const FFrameAndPhase FrameAndPhase)
-	{
-		const auto Data = State ? State->NonFrequentData.Read(FrameAndPhase, Manager) : nullptr;
-		return Data ? Data->SpatialIdx() : Particle.SpatialIdx();
-	}
-
-#if CHAOS_DEBUG_NAME
-	template <typename TParticle>
-	static const TSharedPtr<FString, ESPMode::ThreadSafe>& DebugName(const FGeometryParticleStateBase* State, const TParticle& Particle, const FDirtyPropertiesPool& Manager, const FFrameAndPhase FrameAndPhase)
-	{
-		const auto Data = State ? State->NonFrequentData.Read(FrameAndPhase, Manager) : nullptr;
-		return Data ? Data->DebugName() : Particle.DebugName();
-	}
-#endif
-
-	template <typename TParticle>
-	static const FVec3& F(const FGeometryParticleStateBase* State, const TParticle& Particle, const FDirtyPropertiesPool& Manager, const FFrameAndPhase FrameAndPhase)
-	{
-		const auto Data = State ? State->Dynamics.Read(FrameAndPhase, Manager) : nullptr;
-		return Data ? Data->F() : ZeroVector; //dynamics do not use delta writes, they always write their value immediately or it's 0
-	}
-
-	template <typename TParticle>
-	static const FVec3& Torque(const FGeometryParticleStateBase* State, const TParticle& Particle, const FDirtyPropertiesPool& Manager, const FFrameAndPhase FrameAndPhase)
-	{
-		const auto Data = State ? State->Dynamics.Read(FrameAndPhase, Manager) : nullptr;
-		return Data ? Data->Torque() : ZeroVector; //dynamics do not use delta writes, they always write their value immediately or it's 0
-	}
-
-	template <typename TParticle>
-	static const FVec3& LinearImpulse(const FGeometryParticleStateBase* State, const TParticle& Particle, const FDirtyPropertiesPool& Manager, const FFrameAndPhase FrameAndPhase)
-	{
-		const auto Data = State ? State->Dynamics.Read(FrameAndPhase, Manager) : nullptr;
-		return Data ? Data->LinearImpulse() : ZeroVector; //dynamics do not use delta writes, they always write their value immediately or it's 0
-	}
-
-	template <typename TParticle>
-	static const FVec3& AngularImpulse(const FGeometryParticleStateBase* State, const TParticle& Particle, const FDirtyPropertiesPool& Manager, const FFrameAndPhase FrameAndPhase)
-	{
-		const auto Data = State ? State->Dynamics.Read(FrameAndPhase, Manager) : nullptr;
-		return Data ? Data->AngularImpulse() : ZeroVector; //dynamics do not use delta writes, they always write their value immediately or it's 0
-	}
-
 	template <typename TParticle>
 	static TShapesArrayState<TParticle> ShapesArray(const FGeometryParticleStateBase* State, const TParticle& Particle)
 	{
@@ -685,9 +546,7 @@ public:
 
 	void SyncSimWritablePropsFromSim(FDirtyPropData Manager,const TPBDRigidParticleHandle<FReal,3>& Rigid);
 	void SyncDirtyDynamics(FDirtyPropData& DestManager,const FDirtyChaosProperties& Dirty,const FConstDirtyPropData& SrcManager);
-
-private:
-
+	
 	TParticlePropertyBuffer<FParticlePositionRotation,EChaosProperty::XR> ParticlePositionRotation;
 	TParticlePropertyBuffer<FParticleNonFrequentData,EChaosProperty::NonFrequentData> NonFrequentData;
 	TParticlePropertyBuffer<FParticleVelocities,EChaosProperty::Velocities> Velocities;
@@ -697,10 +556,6 @@ private:
 	TParticlePropertyBuffer<FKinematicTarget, EChaosProperty::KinematicTarget> KinematicTarget;
 
 	FShapesArrayStateBase ShapesArrayState;
-
-	CHAOS_API static FVec3 ZeroVector;
-
-	friend class FRewindData;
 };
 
 class FGeometryParticleState
@@ -722,158 +577,191 @@ public:
 	{
 	}
 
-	const FVec3& X() const
-	{
-		return FGeometryParticleStateBase::X(State, Particle, Pool, FrameAndPhase);
-	}
-
-	const FRotation3& R() const
-	{
-		return FGeometryParticleStateBase::R(State, Particle, Pool, FrameAndPhase);
-	}
-
-	const FVec3& V() const
-	{
-		return FGeometryParticleStateBase::V(State, Particle, Pool, FrameAndPhase);
-	}
-
-	const FVec3& W() const
-	{
-		return FGeometryParticleStateBase::W(State, Particle, Pool, FrameAndPhase);
-	}
-
-	FReal LinearEtherDrag() const
-	{
-		return FGeometryParticleStateBase::LinearEtherDrag(State, Particle, Pool, FrameAndPhase);
-	}
 
-	FReal AngularEtherDrag() const
-	{
-		return FGeometryParticleStateBase::AngularEtherDrag(State, Particle, Pool, FrameAndPhase);
-	}
+	REWIND_PARTICLE_STATIC_PROPERTY(ParticlePositionRotation, X)
+	REWIND_PARTICLE_STATIC_PROPERTY(ParticlePositionRotation, R)
 
-	FReal MaxLinearSpeedSq() const
-	{
-		return FGeometryParticleStateBase::MaxLinearSpeedSq(State, Particle, Pool, FrameAndPhase);
-	}
+	REWIND_PARTICLE_KINEMATIC_PROPERTY(Velocities, V)
+	REWIND_PARTICLE_KINEMATIC_PROPERTY(Velocities, W)
 
-	FReal MaxAngularSpeedSq() const
-	{
-		return FGeometryParticleStateBase::MaxAngularSpeedSq(State, Particle, Pool, FrameAndPhase);
-	}
+	REWIND_PARTICLE_RIGID_PROPERTY(DynamicsMisc, LinearEtherDrag)
+	REWIND_PARTICLE_RIGID_PROPERTY(DynamicsMisc, AngularEtherDrag)
+	REWIND_PARTICLE_RIGID_PROPERTY(DynamicsMisc, MaxLinearSpeedSq)
+	REWIND_PARTICLE_RIGID_PROPERTY(DynamicsMisc, MaxAngularSpeedSq)
+	REWIND_PARTICLE_RIGID_PROPERTY(DynamicsMisc, ObjectState)
+	REWIND_PARTICLE_RIGID_PROPERTY(DynamicsMisc, GravityEnabled)
+	REWIND_PARTICLE_RIGID_PROPERTY(DynamicsMisc, CCDEnabled)
+	REWIND_PARTICLE_RIGID_PROPERTY(DynamicsMisc, CollisionGroup)
 
-	EObjectStateType ObjectState() const
-	{
-		return FGeometryParticleStateBase::ObjectState(State, Particle, Pool, FrameAndPhase);
-	}
+	REWIND_PARTICLE_RIGID_PROPERTY(MassProps, CenterOfMass)
+	REWIND_PARTICLE_RIGID_PROPERTY(MassProps, RotationOfMass)
+	REWIND_PARTICLE_RIGID_PROPERTY(MassProps, I)
+	REWIND_PARTICLE_RIGID_PROPERTY(MassProps, M)
+	REWIND_PARTICLE_RIGID_PROPERTY(MassProps, InvM)
 
-	bool GravityEnabled() const
-	{
-		return FGeometryParticleStateBase::GravityEnabled(State, Particle, Pool, FrameAndPhase);
-	}
+	REWIND_PARTICLE_STATIC_PROPERTY(NonFrequentData, Geometry)
+	REWIND_PARTICLE_STATIC_PROPERTY(NonFrequentData, UniqueIdx)
+	REWIND_PARTICLE_STATIC_PROPERTY(NonFrequentData, SpatialIdx)
+#if CHAOS_DEBUG_NAME
+	REWIND_PARTICLE_STATIC_PROPERTY(NonFrequentData, DebugName)
+#endif
 
-	bool CCDEnabled() const
-	{
-		return FGeometryParticleStateBase::CCDEnabled(State, Particle, Pool, FrameAndPhase);
-	}
+	REWIND_PARTICLE_ZERO_PROPERTY(Dynamics, F)
+	REWIND_PARTICLE_ZERO_PROPERTY(Dynamics, Torque)
+	REWIND_PARTICLE_ZERO_PROPERTY(Dynamics, LinearImpulse)
+	REWIND_PARTICLE_ZERO_PROPERTY(Dynamics, AngularImpulse)
 
-	int32 CollisionGroup() const
+	TShapesArrayState<FGeometryParticleHandle> ShapesArray() const
 	{
-		return FGeometryParticleStateBase::CollisionGroup(State, Particle, Pool, FrameAndPhase);
+		return FGeometryParticleStateBase::ShapesArray(State, Particle);
 	}
 
-	const FVec3& CenterOfMass() const
+	const FGeometryParticleHandle& GetHandle() const
 	{
-		return FGeometryParticleStateBase::CenterOfMass(State, Particle, Pool, FrameAndPhase);
+		return Particle;
 	}
 
-	const FRotation3& RotationOfMass() const
+	void SetState(const FGeometryParticleStateBase* InState)
 	{
-		return FGeometryParticleStateBase::RotationOfMass(State, Particle, Pool, FrameAndPhase);
+		State = InState;
 	}
 
-	const FMatrix33& I() const
-	{
-		return FGeometryParticleStateBase::I(State, Particle, Pool, FrameAndPhase);
-	}
+private:
+	const FGeometryParticleHandle& Particle;
+	const FDirtyPropertiesPool& Pool;
+	const FGeometryParticleStateBase* State = nullptr;
+	const FFrameAndPhase FrameAndPhase;
 
-	const FMatrix33& InvI() const
-	{
-		return FGeometryParticleStateBase::InvI(State, Particle, Pool, FrameAndPhase);
-	}
+	CHAOS_API static FVec3 ZeroVector;
 
-	FReal M() const
-	{
-		return FGeometryParticleStateBase::M(State, Particle, Pool, FrameAndPhase);
-	}
+	};
 
-	FReal InvM() const
-	{
-		return FGeometryParticleStateBase::InvM(State, Particle, Pool, FrameAndPhase);
-	}
 
-	TSerializablePtr<FImplicitObject> Geometry() const
+struct FJointStateBase
+{
+	explicit FJointStateBase(int32 NumFrames)
+		: JointSettings(ComputeCircularSize(NumFrames))
+		, ProxyPair(ComputeCircularSize(NumFrames))
 	{
-		return FGeometryParticleStateBase::Geometry(State, Particle, Pool, FrameAndPhase);
 	}
 
-	FUniqueIdx UniqueIdx() const
-	{
-		return FGeometryParticleStateBase::UniqueIdx(State, Particle, Pool, FrameAndPhase);
-	}
+	FJointStateBase(const FJointStateBase& Other) = delete;
+	FJointStateBase(FJointStateBase&& Other) = default;
+	~FJointStateBase() = default;
 
-	FSpatialAccelerationIdx SpatialIdx() const
+	void Release(FDirtyPropertiesPool& Manager)
 	{
-		return FGeometryParticleStateBase::SpatialIdx(State, Particle, Pool, FrameAndPhase);
+		JointSettings.Release(Manager);
+		ProxyPair.Release(Manager);
 	}
 
-#if CHAOS_DEBUG_NAME
-	const TSharedPtr<FString, ESPMode::ThreadSafe>& DebugName() const
+	void Reset()
 	{
-		return FGeometryParticleStateBase::DebugName(State, Particle, Pool, FrameAndPhase);
+		JointSettings.Reset();
+		ProxyPair.Reset();
 	}
-#endif
 
-	const FVec3& F() const
+	void ClearEntryAndFuture(const FFrameAndPhase FrameAndPhase)
 	{
-		return FGeometryParticleStateBase::F(State, Particle, Pool, FrameAndPhase);
+		JointSettings.ClearEntryAndFuture(FrameAndPhase);
+		ProxyPair.ClearEntryAndFuture(FrameAndPhase);
 	}
 
-	const FVec3& Torque() const
+	bool IsClean(const FFrameAndPhase FrameAndPhase) const
 	{
-		return FGeometryParticleStateBase::Torque(State, Particle, Pool, FrameAndPhase);
+		return JointSettings.IsClean(FrameAndPhase) && ProxyPair.IsClean(FrameAndPhase);
 	}
 
-	const FVec3& LinearImpulse() const
-	{
-		return FGeometryParticleStateBase::LinearImpulse(State, Particle, Pool, FrameAndPhase);
-	}
+	bool IsInSync(const FPBDJointConstraintHandle& Handle, const FFrameAndPhase FrameAndPhase, const FDirtyPropertiesPool& Pool) const;
 
-	const FVec3& AngularImpulse() const
-	{
-		return FGeometryParticleStateBase::AngularImpulse(State, Particle, Pool, FrameAndPhase);
-	}
+	TParticlePropertyBuffer<FPBDJointSettings, EChaosProperty::JointSettings> JointSettings;
+	TParticlePropertyBuffer<FProxyBasePair, EChaosProperty::JointParticleProxies> ProxyPair;
+};
 
-	TShapesArrayState<FGeometryParticleHandle> ShapesArray() const
+class FJointState
+{
+public:
+	FJointState(const FPBDJointConstraintHandle& InJoint, const FDirtyPropertiesPool& InPool)
+	: Head(InJoint)
+	, Pool(InPool)
 	{
-		return FGeometryParticleStateBase::ShapesArray(State, Particle);
 	}
 
-	const FGeometryParticleHandle& GetHandle() const
+	FJointState(const FJointStateBase* InState, const FPBDJointConstraintHandle& InJoint, const FDirtyPropertiesPool& InPool, const FFrameAndPhase InFrameAndPhase)
+	: Head(InJoint)
+	, Pool(InPool)
+	, State(InState)
+	, FrameAndPhase(InFrameAndPhase)
 	{
-		return Particle;
 	}
 
-	void SetState(const FGeometryParticleStateBase* InState)
-	{
-		State = InState;
-	}
+	REWIND_JOINT_PROPERTY(JointSettings, JointTransforms, ConnectorTransforms);
+	REWIND_JOINT_PROPERTY(JointSettings, CollisionEnabled, bCollisionEnabled);
+	REWIND_JOINT_PROPERTY(JointSettings, ProjectionEnabled, bProjectionEnabled);
+	REWIND_JOINT_PROPERTY(JointSettings, ProjectionLinearAlpha, LinearProjection);
+	REWIND_JOINT_PROPERTY(JointSettings, ProjectionAngularAlpha, AngularProjection);
+	REWIND_JOINT_PROPERTY(JointSettings, ParentInvMassScale, ParentInvMassScale);
+	REWIND_JOINT_PROPERTY(JointSettings, LinearBreakForce, LinearBreakForce);
+	REWIND_JOINT_PROPERTY(JointSettings, LinearPlasticityLimit, LinearPlasticityLimit);
+	REWIND_JOINT_PROPERTY(JointSettings, LinearPlasticityType, LinearPlasticityType);
+	REWIND_JOINT_PROPERTY(JointSettings, AngularBreakTorque, AngularBreakTorque);
+	REWIND_JOINT_PROPERTY(JointSettings, AngularPlasticityLimit, AngularPlasticityLimit);
+	REWIND_JOINT_PROPERTY(JointSettings, UserData, UserData);
+	REWIND_JOINT_PROPERTY(JointSettings, LinearPositionDriveXEnabled, bLinearPositionDriveEnabled[0]);
+	REWIND_JOINT_PROPERTY(JointSettings, LinearPositionDriveYEnabled, bLinearPositionDriveEnabled[1]);
+	REWIND_JOINT_PROPERTY(JointSettings, LinearPositionDriveZEnabled, bLinearPositionDriveEnabled[2]);
+	REWIND_JOINT_PROPERTY(JointSettings, LinearDrivePositionTarget, LinearDrivePositionTarget);
+	REWIND_JOINT_PROPERTY(JointSettings, LinearVelocityDriveXEnabled, bLinearVelocityDriveEnabled[0]);
+	REWIND_JOINT_PROPERTY(JointSettings, LinearVelocityDriveYEnabled, bLinearVelocityDriveEnabled[1]);
+	REWIND_JOINT_PROPERTY(JointSettings, LinearVelocityDriveZEnabled, bLinearVelocityDriveEnabled[2]);
+	REWIND_JOINT_PROPERTY(JointSettings, LinearDriveVelocityTarget, LinearDriveVelocityTarget);
+	REWIND_JOINT_PROPERTY(JointSettings, LinearDriveForceMode, LinearDriveForceMode);
+	REWIND_JOINT_PROPERTY(JointSettings, LinearMotionTypesX, LinearMotionTypes[0]);
+	REWIND_JOINT_PROPERTY(JointSettings, LinearMotionTypesY, LinearMotionTypes[1]);
+	REWIND_JOINT_PROPERTY(JointSettings, LinearMotionTypesZ, LinearMotionTypes[2]);
+	REWIND_JOINT_PROPERTY(JointSettings, LinearDriveStiffness, LinearDriveStiffness);
+	REWIND_JOINT_PROPERTY(JointSettings, LinearDriveDamping, LinearDriveDamping);
+	REWIND_JOINT_PROPERTY(JointSettings, ContactTransferScale, ContactTransferScale);
+	REWIND_JOINT_PROPERTY(JointSettings, AngularSLerpPositionDriveEnabled, bAngularSLerpPositionDriveEnabled);
+	REWIND_JOINT_PROPERTY(JointSettings, AngularTwistPositionDriveEnabled, bAngularTwistPositionDriveEnabled);
+	REWIND_JOINT_PROPERTY(JointSettings, AngularSwingPositionDriveEnabled, bAngularSwingPositionDriveEnabled);
+	REWIND_JOINT_PROPERTY(JointSettings, AngularDrivePositionTarget, AngularDrivePositionTarget);
+	REWIND_JOINT_PROPERTY(JointSettings, AngularSLerpVelocityDriveEnabled, bAngularSLerpVelocityDriveEnabled);
+	REWIND_JOINT_PROPERTY(JointSettings, AngularTwistVelocityDriveEnabled, bAngularTwistVelocityDriveEnabled);
+	REWIND_JOINT_PROPERTY(JointSettings, AngularSwingVelocityDriveEnabled, bAngularSwingVelocityDriveEnabled);
+	REWIND_JOINT_PROPERTY(JointSettings, AngularDriveVelocityTarget, AngularDriveVelocityTarget);
+	REWIND_JOINT_PROPERTY(JointSettings, AngularDriveForceMode, AngularDriveForceMode);
+	REWIND_JOINT_PROPERTY(JointSettings, AngularMotionTypesX, AngularMotionTypes[0]);
+	REWIND_JOINT_PROPERTY(JointSettings, AngularMotionTypesY, AngularMotionTypes[1]);
+	REWIND_JOINT_PROPERTY(JointSettings, AngularMotionTypesZ, AngularMotionTypes[2]);
+	REWIND_JOINT_PROPERTY(JointSettings, AngularDriveStiffness, AngularDriveStiffness);
+	REWIND_JOINT_PROPERTY(JointSettings, AngularDriveDamping, AngularDriveDamping);
+	REWIND_JOINT_PROPERTY(JointSettings, Stiffness, Stiffness);
+	REWIND_JOINT_PROPERTY(JointSettings, SoftLinearLimitsEnabled, bSoftLinearLimitsEnabled);
+	REWIND_JOINT_PROPERTY(JointSettings, SoftTwistLimitsEnabled, bSoftTwistLimitsEnabled);
+	REWIND_JOINT_PROPERTY(JointSettings, SoftSwingLimitsEnabled, bSoftSwingLimitsEnabled);
+	REWIND_JOINT_PROPERTY(JointSettings, LinearSoftForceMode, LinearSoftForceMode);
+	REWIND_JOINT_PROPERTY(JointSettings, AngularSoftForceMode, AngularSoftForceMode);
+	REWIND_JOINT_PROPERTY(JointSettings, SoftLinearStiffness, SoftLinearStiffness);
+	REWIND_JOINT_PROPERTY(JointSettings, SoftLinearDamping, SoftLinearDamping);
+	REWIND_JOINT_PROPERTY(JointSettings, SoftTwistStiffness, SoftTwistStiffness);
+	REWIND_JOINT_PROPERTY(JointSettings, SoftTwistDamping, SoftTwistDamping);
+	REWIND_JOINT_PROPERTY(JointSettings, SoftSwingStiffness, SoftSwingStiffness);
+	REWIND_JOINT_PROPERTY(JointSettings, SoftSwingDamping, SoftSwingDamping);
+	REWIND_JOINT_PROPERTY(JointSettings, LinearLimit, LinearLimit);
+	REWIND_JOINT_PROPERTY(JointSettings, AngularLimits, AngularLimits);
+	REWIND_JOINT_PROPERTY(JointSettings, LinearContactDistance, LinearContactDistance);
+	REWIND_JOINT_PROPERTY(JointSettings, TwistContactDistance, TwistContactDistance);
+	REWIND_JOINT_PROPERTY(JointSettings, SwingContactDistance, SwingContactDistance);
+	REWIND_JOINT_PROPERTY(JointSettings, LinearRestitution, LinearRestitution);
+	REWIND_JOINT_PROPERTY(JointSettings, TwistRestitution, TwistRestitution);
+	REWIND_JOINT_PROPERTY(JointSettings, SwingRestitution, SwingRestitution);
 
 private:
-	const FGeometryParticleHandle& Particle;
+	const FPBDJointConstraintHandle& Head;
 	const FDirtyPropertiesPool& Pool;
-	const FGeometryParticleStateBase* State = nullptr;
-	const FFrameAndPhase FrameAndPhase;
+	const FJointStateBase* State = nullptr;
+	const FFrameAndPhase FrameAndPhase = { 0,0 };
 };
 
 extern CHAOS_API int32 EnableResimCache;
@@ -888,7 +776,6 @@ public:
 	, Solver(InSolver)
 	, CurFrame(InCurrentFrame)
 	, LatestFrame(-1)
-	, Buffer(0)
 	, FramesSaved(0)
 	, DataIdxOffset(0)
 	, bNeedsSave(false)
@@ -984,28 +871,6 @@ private:
 
 	friend class FPBDRigidsSolver;
 
-	struct FSimWritableState
-	{
-		template <bool bResim>
-		bool SyncSimWritablePropsFromSim(const TPBDRigidParticleHandle<FReal,3>& Rigid, const int32 Frame);
-		void SyncToParticle(TPBDRigidParticleHandle<FReal,3>& Rigid) const;
-
-		bool IsSimWritableDesynced(const TPBDRigidParticleHandle<FReal,3>& Rigid) const;
-		
-		const FVec3& X() const { return MX; }
-		const FQuat& R() const { return MR; }
-		const FVec3& V() const { return MV; }
-		const FVec3& W() const { return MW; }
-
-		int32 FrameRecordedHack = INDEX_NONE;
-
-	private:
-		FVec3 MX;
-		FQuat MR;
-		FVec3 MV;
-		FVec3 MW;
-	};
-
 	struct FDirtyFrameInfo
 	{
 		int32 Frame;	//needed to protect against stale entries in circular buffer
@@ -1025,7 +890,7 @@ private:
 	};
 
 	void CHAOS_API AdvanceFrameImp(IResimCacheBase* ResimCache);
-	
+
 	struct FFrameManagerInfo
 	{
 		TUniquePtr<IResimCacheBase> ExternalResimCache;
@@ -1037,50 +902,45 @@ private:
 		FReal DeltaTime;
 	};
 
-	struct FDirtyParticleInfo
+	template <typename HistoryType>
+	struct TDirtyConcreteInfo
 	{
 	private:
-		FGeometryParticleStateBase History;
-		TGeometryParticleHandle<FReal, 3>* PTParticle;
+		HistoryType History;
 		FDirtyPropertiesPool* PropertiesPool;
 	public:
-		FUniqueIdx CachedUniqueIdx;	//Needed when manipulating on physics thread and Particle data cannot be read
 		int32 LastDirtyFrame;	//Track how recently this was made dirty
 		int32 InitializedOnStep = INDEX_NONE;	//if not INDEX_NONE, it indicates we saw initialization during rewind history window
-		int32 DirtyDynamics = INDEX_NONE;
 		bool bResimAsSlave = true;	//Indicates the particle will always resim in the exact same way from game thread data
 
-		FDirtyParticleInfo(FDirtyPropertiesPool& InPropertiesPool, TGeometryParticleHandle<FReal, 3>& InPTParticle, const FUniqueIdx UniqueIdx, const int32 CurFrame, const int32 NumFrames)
-		: History(NumFrames)
-		, PTParticle(&InPTParticle)
-		, PropertiesPool(&InPropertiesPool)
-		, CachedUniqueIdx(UniqueIdx)
-		, LastDirtyFrame(CurFrame)
+		TDirtyConcreteInfo(FDirtyPropertiesPool& InPropertiesPool, const int32 CurFrame, const int32 NumFrames)
+			: History(NumFrames)
+			, PropertiesPool(&InPropertiesPool)
+			, LastDirtyFrame(CurFrame)
 		{
 		}
 
-		FDirtyParticleInfo(FDirtyParticleInfo&& Other)
-		: History(MoveTemp(Other.History))
-			, PTParticle(Other.PTParticle)
+		TDirtyConcreteInfo(TDirtyConcreteInfo&& Other)
+			: History(MoveTemp(Other.History))
 			, PropertiesPool(Other.PropertiesPool)
-			, CachedUniqueIdx(MoveTemp(Other.CachedUniqueIdx))
 			, LastDirtyFrame(Other.LastDirtyFrame)
 			, InitializedOnStep(Other.InitializedOnStep)
-		, bResimAsSlave(Other.bResimAsSlave)
+			, bResimAsSlave(Other.bResimAsSlave)
 		{
 			Other.PropertiesPool = nullptr;
 		}
 
-		~FDirtyParticleInfo();
-		
-		TGeometryParticleHandle<FReal,3>* GetPTParticle() const
+		~TDirtyConcreteInfo()
 		{
-			return PTParticle;
+			if (PropertiesPool)
+			{
+				History.Release(*PropertiesPool);
+			}
 		}
-		
-		FDirtyParticleInfo(const FDirtyParticleInfo& Other) = delete;
 
-		FGeometryParticleStateBase& AddFrame(const int32 Frame)
+		TDirtyConcreteInfo(const TDirtyConcreteInfo& Other) = delete;
+
+		HistoryType& AddFrame(const int32 Frame)
 		{
 			LastDirtyFrame = Frame;
 			return History;
@@ -1091,12 +951,48 @@ private:
 			History.ClearEntryAndFuture(FrameAndPhase);
 		}
 
-		const FGeometryParticleStateBase& GetHistory() const	//For non-const access use AddFrame
+		const HistoryType& GetHistory() const	//For non-const access use AddFrame
 		{
 			return History;
 		}
 	};
 
+	struct FDirtyParticleInfo : public TDirtyConcreteInfo<FGeometryParticleStateBase>
+	{
+	private:
+		using Base = TDirtyConcreteInfo<FGeometryParticleStateBase>;
+		TGeometryParticleHandle<FReal,3>* PTParticle;
+	public:
+		FUniqueIdx CachedUniqueIdx;	//Needed when manipulating on physics thread and Particle data cannot be read
+		int32 DirtyDynamics = INDEX_NONE;
+
+		FDirtyParticleInfo(FDirtyPropertiesPool& InPropertiesPool, TGeometryParticleHandle<FReal,3>& InPTParticle, const FUniqueIdx UniqueIdx,const int32 CurFrame,const int32 NumFrames)
+		: Base(InPropertiesPool, CurFrame, NumFrames)
+		, PTParticle(&InPTParticle)
+		, CachedUniqueIdx(UniqueIdx)
+		{
+		}
+		
+		TGeometryParticleHandle<FReal,3>* GetPTParticle() const
+		{
+			return PTParticle;
+		}
+	};
+
+	struct FDirtyJointInfo : public TDirtyConcreteInfo<FJointStateBase>
+	{
+	private:
+		using Base = TDirtyConcreteInfo<FJointStateBase>;
+		FPBDJointConstraintHandle* Handle;
+	public:
+
+		FDirtyJointInfo(FDirtyPropertiesPool& InPropertiesPool, FPBDJointConstraintHandle& InJoint, const int32 CurFrame, const int32 NumFrames)
+			: Base(InPropertiesPool, CurFrame, NumFrames)
+			, Handle(&InJoint)
+		{
+		}
+	};
+
 	const FDirtyParticleInfo& FindParticleChecked(const FUniqueIdx UniqueIdx) const
 	{
 		const int32 Idx = ParticleToAllDirtyIdx.FindChecked(UniqueIdx);
@@ -1130,6 +1026,18 @@ private:
 	}
 
 	FDirtyParticleInfo& FindOrAddParticle(TGeometryParticleHandle<FReal,3>& PTParticle, const int32 InitializedOnFrame = INDEX_NONE);
+	FDirtyJointInfo& FindOrAddJoint(FPBDJointConstraintHandle& Handle, const int32 InitializedOnFrame = INDEX_NONE);
+
+	FDirtyParticleInfo& FindOrAddConcrete(FGeometryParticleHandle& Handle, const int32 InitializedOnFrame = INDEX_NONE)
+	{
+		return FindOrAddParticle(Handle, InitializedOnFrame);
+	}
+
+	FDirtyJointInfo& FindOrAddConcrete(FPBDJointConstraintHandle& Handle, const int32 InitializedOnFrame = INDEX_NONE)
+	{
+		return FindOrAddJoint(Handle, InitializedOnFrame);
+	}
+
 	bool RewindToFrame(int32 Frame);
 
 	static void DesyncParticle(FDirtyParticleInfo& Info, const FFrameAndPhase FrameAndPhase);
@@ -1138,19 +1046,16 @@ private:
 	TCircularBuffer<FFrameManagerInfo> Managers;
 	FDirtyPropertiesPool PropertiesPool;	//must come before AllDirtyParticles since it relies on it (and used in destruction)
 	TArray<FDirtyParticleInfo> AllDirtyParticles;
+	TArray<FDirtyJointInfo> AllDirtyJoints;
 
 	FPBDRigidsSolver* Solver;
 	int32 CurFrame;
 	int32 LatestFrame;
-	uint8 Buffer;	//Used to flip flop between resim buffer and current buffer. Since resimmed becomes current we just use a double buffer
 	int32 FramesSaved;
 	int32 DataIdxOffset;
 	bool bNeedsSave;	//Indicates that some data is pointing at head and requires saving before a rewind
 	bool bResimOptimization;
 
-	uint8 CurrentBuffer() const { return Buffer; }
-	uint8 OriginalBufferIdx() const { return 1 - Buffer; }
-
 	bool IsResimAndInSync(const FGeometryParticleHandle& Handle) const { return IsResim() && Handle.SyncState() == ESyncState::InSync; }
 
 	template <bool bSkipDynamics = false>
