Make sure geometry collection proxy set particles dirty after creation on the PT to make sure the simData is properly updated on the acceleration structure

diff --git a/Engine/Source/Runtime/Experimental/Chaos/Private/PhysicsProxy/GeometryCollectionPhysicsProxy.cpp b/Engine/Source/Runtime/Experimental/Chaos/Private/PhysicsProxy/GeometryCollectionPhysicsProxy.cpp
--- a/Engine/Source/Runtime/Experimental/Chaos/Private/PhysicsProxy/GeometryCollectionPhysicsProxy.cpp
+++ b/Engine/Source/Runtime/Experimental/Chaos/Private/PhysicsProxy/GeometryCollectionPhysicsProxy.cpp
@@ -884,6 +884,15 @@ void FGeometryCollectionPhysicsProxy::InitializeBodiesPT(Chaos::FPBDRigidsSolver
 			}
 		}
 
+		// call DirtyParticle to make sure the acceleration structure is up to date with all the changes happening here
+		for (int32 TransformGroupIndex = 0; TransformGroupIndex < NumTransforms; ++TransformGroupIndex)
+		{
+			if (FClusterHandle* Handle = SolverParticleHandles[TransformGroupIndex])
+			{
+				RigidsSolver->GetEvolution()->DirtyParticle(*Handle);
+			}
+		}
+
 	} // end if simulating...
 
 }
