Defer particle destruction to make rewind data easier to manage.
Start recording joint data in rewind data

diff --git a/Engine/Source/Programs/HeadlessChaos/Private/HeadlessChaosTestSolverProxies.cpp b/Engine/Source/Programs/HeadlessChaos/Private/HeadlessChaosTestSolverProxies.cpp
--- a/Engine/Source/Programs/HeadlessChaos/Private/HeadlessChaosTestSolverProxies.cpp
+++ b/Engine/Source/Programs/HeadlessChaos/Private/HeadlessChaosTestSolverProxies.cpp
@@ -154,6 +154,8 @@ namespace ChaosTest {
 		// run PT again so that the cleanup is done
 		Solver->AdvanceAndDispatch_External(100.0f);
 		Solver->UpdateGameThreadStructures();
+		Solver->AdvanceAndDispatch_External(100.0f);
+		Solver->UpdateGameThreadStructures();
 
 		// Unique idx should be gone from the pending lists
 		EXPECT_FALSE(Solver->GetEvolution()->IsUniqueIndexPendingRelease(FirstIdx));
diff --git a/Engine/Source/Runtime/Experimental/Chaos/Private/PBDRigidsSolver.cpp b/Engine/Source/Runtime/Experimental/Chaos/Private/PBDRigidsSolver.cpp
--- a/Engine/Source/Runtime/Experimental/Chaos/Private/PBDRigidsSolver.cpp
+++ b/Engine/Source/Runtime/Experimental/Chaos/Private/PBDRigidsSolver.cpp
@@ -518,24 +518,14 @@ namespace Chaos
 			FGeometryParticleHandle* Handle = Proxy->GetHandle_LowLevel();
 			Proxy->SetHandle(nullptr);
 			const int32 OffsetForRewind = MRewindData ? MRewindData->Capacity() : 0;
-			PendingDestroyPhysicsProxy.Add(  FPendingDestroyInfo{Proxy, GetCurrentFrame() + OffsetForRewind});
+			PendingDestroyPhysicsProxy.Add(  FPendingDestroyInfo{Proxy, GetCurrentFrame() + OffsetForRewind, Handle, UniqueIdx});
 			
 			//If particle was created and destroyed before commands were enqueued just skip. I suspect we can skip entire lambda, but too much code to verify right now
-
 			if(Handle)
 			{
-				// Remove from rewind data
-				if(FRewindData* RewindData = GetRewindData())
-				{
-					RewindData->RemoveParticle(Handle->UniqueIdx());
-				}
-  
-				// Use the handle to destroy the particle data
-				GetEvolution()->DestroyParticle(Handle);
+				//Disable until particle is finally destroyed
+				GetEvolution()->DisableParticle(Handle);
 			}
-
-			// finally let's release the unique index
-			GetEvolution()->ReleaseUniqueIdx(UniqueIdx);
 		});
 	}
 
@@ -608,6 +598,13 @@ namespace Chaos
 		// Finish registration on the physics thread...
 		EnqueueCommandImmediate([InParticles, JointProxy, this]()
 			{
+				//TODO: consider deferring this so that rewind can resim with joint until moment of destruction
+				//For now we assume this always comes from server update
+				if(FRewindData* RewindData = GetRewindData())
+				{
+					RewindData->RemoveJoint(JointProxy->GetHandle());
+				}
+
 				JointProxy->DestroyOnPhysicsThread(this);
 				JointConstraintPhysicsProxies_Internal.RemoveSingle(JointProxy);
 				delete JointProxy;
@@ -712,6 +709,21 @@ namespace Chaos
 			FPendingDestroyInfo& Info = PendingDestroyPhysicsProxy[Idx];
 			if(Info.DestroyOnStep <= GetCurrentFrame() || IsShuttingDown())
 			{
+				// finally let's release the unique index
+				GetEvolution()->ReleaseUniqueIdx(Info.UniqueIdx);
+
+				if (Info.Handle)
+				{
+					// Remove from rewind data
+					if (FRewindData* RewindData = GetRewindData())
+					{
+						RewindData->RemoveParticle(Info.Handle);
+					}
+
+					// Use the handle to destroy the particle data
+					GetEvolution()->DestroyParticle(Info.Handle);
+				}
+
 				ensure(Info.Proxy->GetHandle_LowLevel() == nullptr);	//should have already cleared this out
 				delete Info.Proxy;
 				PendingDestroyPhysicsProxy.RemoveAtSwap(Idx);
@@ -944,15 +956,7 @@ namespace Chaos
 			{
 				if (RewindData)
 				{
-					//may want to remove branch by templatizing lambda
-					if (RewindData->IsResim())
-					{
-						RewindData->PushGTDirtyData<true>(*Manager, DataIdx, Dirty, ShapeDirtyData);
-					}
-					else
-					{
-						RewindData->PushGTDirtyData<false>(*Manager, DataIdx, Dirty, ShapeDirtyData);
-					}
+					RewindData->PushGTDirtyData(*Manager, DataIdx, Dirty, ShapeDirtyData);
 				}
 				Proxy->PushToPhysicsState(*Manager, DataIdx, Dirty, ShapeDirtyData, *GetEvolution(), ExternalDt);
 			}
@@ -1022,7 +1026,7 @@ namespace Chaos
 		});
 
 		//need to create new constraint handles
-		DirtyProxiesData->ForEachProxy([this, &ProcessProxyPT, Manager](int32 DataIdx, FDirtyProxy& Dirty)
+		DirtyProxiesData->ForEachProxy([this, &ProcessProxyPT, Manager, RewindData](int32 DataIdx, FDirtyProxy& Dirty)
 		{
 			switch (Dirty.Proxy->GetType())
 			{
@@ -1036,6 +1040,13 @@ namespace Chaos
 					JointProxy->InitializeOnPhysicsThread(this, *Manager, DataIdx, Dirty.PropertyData);
 					JointProxy->SetInitialized(GetCurrentFrame());
 				}
+
+				//TODO: if we support predicting creation / destruction of joints need to handle null joint case
+				if (RewindData)
+				{
+					RewindData->PushGTDirtyData(*Manager, DataIdx, Dirty, nullptr);
+				}
+				
 				JointProxy->PushStateOnPhysicsThread(this, *Manager, DataIdx, Dirty.PropertyData);
 				Dirty.Proxy->ResetDirtyIdx();
 				break;
@@ -1459,15 +1470,7 @@ namespace Chaos
 			int32 DataIdx = 0;
 			for(TPBDRigidParticleHandleImp<FReal,3,false>& DirtyParticle : DirtyParticles)
 			{
-				//may want to remove branch using templates outside loop
-				if (MRewindData->IsResim())
-				{
-					MRewindData->PushPTDirtyData<true>(*DirtyParticle.Handle(), DataIdx++);
-				}
-				else
-				{
-					MRewindData->PushPTDirtyData<false>(*DirtyParticle.Handle(), DataIdx++);
-				}
+				MRewindData->PushPTDirtyData(*DirtyParticle.Handle(), DataIdx++);
 			}
 		}
 	}
diff --git a/Engine/Source/Runtime/Experimental/Chaos/Private/RewindData.cpp b/Engine/Source/Runtime/Experimental/Chaos/Private/RewindData.cpp
--- a/Engine/Source/Runtime/Experimental/Chaos/Private/RewindData.cpp
+++ b/Engine/Source/Runtime/Experimental/Chaos/Private/RewindData.cpp
@@ -170,9 +170,9 @@ bool FRewindData::RewindToFrame(int32 Frame)
 		return false;
 	};
 			
-	for(FDirtyParticleInfo& DirtyParticleInfo : AllDirtyParticles)
+	for(FDirtyParticleInfo& DirtyParticleInfo : DirtyParticles)
 	{
-		FGeometryParticleHandle* PTParticle = DirtyParticleInfo.GetPTParticle();
+		FGeometryParticleHandle* PTParticle = DirtyParticleInfo.GetObjectPtr();
 
 		//rewind is about to start, all particles should be in sync at this point
 		ensure(PTParticle->SyncState() == ESyncState::InSync);
@@ -208,12 +208,12 @@ bool FRewindData::RewindToFrame(int32 Frame)
 				//particle actually changes not just created/streamed so need to update its state
 
 				//Data changes so send back to GT for interpolation. TODO: improve this in case data ends up being identical in resim
-				Solver->GetEvolution()->GetParticles().MarkTransientDirtyParticle(DirtyParticleInfo.GetPTParticle());
+				Solver->GetEvolution()->GetParticles().MarkTransientDirtyParticle(DirtyParticleInfo.GetObjectPtr());
 
 				DirtyParticleInfo.DirtyDynamics = INDEX_NONE;	//make sure to undo this as we want to record it again during resim
 
 				//for now just mark anything that changed as enabled during resim. TODO: use bubble
-				DirtyParticleInfo.GetPTParticle()->SetEnabledDuringResim(true);
+				DirtyParticleInfo.GetObjectPtr()->SetEnabledDuringResim(true);
 			}
 
 		if (DirtyParticleInfo.InitializedOnStep > Frame)
@@ -221,7 +221,7 @@ bool FRewindData::RewindToFrame(int32 Frame)
 			//hasn't initialized yet, so disable
 			//must do this after rewind because SetDynamicsMisc will re-enable
 			//(the disable is a temp way to ignore objects not spawned yet, they weren't really disabled which is why it gets re-enabled)
-			Solver->GetEvolution()->DisableParticle(DirtyParticleInfo.GetPTParticle());
+			Solver->GetEvolution()->DisableParticle(DirtyParticleInfo.GetObjectPtr());
 		}
 		}
 
@@ -234,33 +234,6 @@ bool FRewindData::RewindToFrame(int32 Frame)
 	return true;
 }
 
-void FRewindData::RemoveParticle(const FUniqueIdx UniqueIdx)
-{
-	if(const int32* IdxPtr = ParticleToAllDirtyIdx.Find(UniqueIdx))
-	{
-		const int32 Idx = *IdxPtr;
-		AllDirtyParticles.RemoveAtSwap(Idx);
-		if(Idx < AllDirtyParticles.Num())
-		{
-			//update particle in new position
-			ParticleToAllDirtyIdx.FindChecked(AllDirtyParticles[Idx].CachedUniqueIdx) = Idx;
-		}
-
-		ParticleToAllDirtyIdx.RemoveChecked(UniqueIdx);
-	}
-}
-
-/* Query the state of particles from the past. Once a rewind happens state captured must be queried using GetFutureStateAtFrame */
-FGeometryParticleState FRewindData::GetPastStateAtFrame(const FGeometryParticleHandle& Handle,int32 Frame, FFrameAndPhase::EParticleHistoryPhase Phase) const
-{
-	ensure(!IsResim());
-	ensure(Frame >= GetEarliestFrame_Internal());	//can't get state from before the frame we rewound to
-
-	const FDirtyParticleInfo* Info = FindParticle(Handle.UniqueIdx());
-	const FGeometryParticleStateBase* State = Info ? &Info->GetHistory() : nullptr;
-	return FGeometryParticleState(State, Handle, PropertiesPool, { Frame, Phase });
-}
-
 void FRewindData::FinishFrame()
 {
 	QUICK_SCOPE_CYCLE_COUNTER(RewindDataFinishFrame);
@@ -269,7 +242,7 @@ void FRewindData::FinishFrame()
 	{
 		FFrameAndPhase FutureFrame{ CurFrame + 1, FFrameAndPhase::PrePushData };
 
-		for(FDirtyParticleInfo& Info : AllDirtyParticles)
+		for(FDirtyParticleInfo& Info : DirtyParticles)
 		{
 			if (Info.bResimAsSlave)
 			{
@@ -277,7 +250,7 @@ void FRewindData::FinishFrame()
 				continue;
 			}
 
-			auto& Handle = *Info.GetPTParticle();
+			auto& Handle = *Info.GetObjectPtr();
 			if (auto Rigid = Handle.CastToRigidParticle())
 			{
 				if(Rigid->ResimType() == EResimType::FullResim)
@@ -313,7 +286,7 @@ FAutoConsoleVariableRef CVarSkipDesyncTest(TEXT("p.SkipDesyncTest"), SkipDesyncT
 void FRewindData::DesyncParticle(FDirtyParticleInfo& Info, const FFrameAndPhase FrameAndPhase)
 {
 	Info.ClearPhaseAndFuture(FrameAndPhase);
-	Info.GetPTParticle()->SetSyncState(ESyncState::HardDesync);
+	Info.GetObjectPtr()->SetSyncState(ESyncState::HardDesync);
 }
 
 template <bool bSkipDynamics>
@@ -321,7 +294,7 @@ void FRewindData::DesyncIfNecessary(FDirtyParticleInfo& Info, const FFrameAndPha
 {
 	ensure(IsResim());	//shouldn't bother with desync unless we're resimming
 
-	auto Handle = Info.GetPTParticle();
+	auto Handle = Info.GetObjectPtr();
 	if (Handle->SyncState() == ESyncState::InSync && !Info.GetHistory().IsInSync<bSkipDynamics>(*Handle, FrameAndPhase, PropertiesPool))
 	{
 		if (!SkipDesyncTest)
@@ -340,25 +313,26 @@ void FRewindData::AdvanceFrameImp(IResimCacheBase* ResimCache)
 	TArray<FGeometryParticleHandle*> DesyncedParticles;
 	if(IsResim())
 	{
-		DesyncedParticles.Reserve(AllDirtyParticles.Num());
+		DesyncedParticles.Reserve(DirtyParticles.Num());
 	}
 
 	const FFrameAndPhase FrameAndPhase{ CurFrame, FFrameAndPhase::PostCallbacks };
-	for (int32 DirtyIdx = AllDirtyParticles.Num() - 1; DirtyIdx >= 0; --DirtyIdx)
+	FDirtyParticleInfo* DirtyParticlesDense = DirtyParticles.GetDensePtr();
+	for (int32 DirtyIdx = DirtyParticles.Num() - 1; DirtyIdx >= 0; --DirtyIdx)
 	{
-		FDirtyParticleInfo& Info = AllDirtyParticles[DirtyIdx];
+		FDirtyParticleInfo& Info = DirtyParticlesDense[DirtyIdx];
 
-		ensure(IsResimAndInSync(*Info.GetPTParticle()) || Info.GetHistory().IsClean(FrameAndPhase));  //Sim hasn't run yet so PostCallbacks (sim results) should be clean
+		ensure(IsResimAndInSync(*Info.GetObjectPtr()) || Info.GetHistory().IsClean(FrameAndPhase));  //Sim hasn't run yet so PostCallbacks (sim results) should be clean
 
 		//if hasn't changed in a while stop tracking
 		if (Info.LastDirtyFrame < EarliestFrame)
 		{
-			RemoveParticle(AllDirtyParticles[DirtyIdx].CachedUniqueIdx);
+			RemoveParticle(DirtyParticlesDense[DirtyIdx].GetObjectPtr());
 		}
 		else
 		{
 
-			auto Handle = Info.GetPTParticle();
+			auto Handle = Info.GetObjectPtr();
 			Info.bResimAsSlave = Handle->ResimType() == EResimType::ResimAsSlave;
 
 			if (IsResim() && !Info.bResimAsSlave)
@@ -395,20 +369,19 @@ void FRewindData::AdvanceFrameImp(IResimCacheBase* ResimCache)
 #define REWIND_DESYNC 0
 #endif
 
-template <typename TData, typename TConcrete>
-void CopyDataFromConcrete(TData& Data, const TConcrete& Concrete)
+template <typename TData, typename TObj>
+void CopyDataFromObject(TData& Data, const TObj& Obj)
 {
-	Data.CopyFrom(Concrete);
+	Data.CopyFrom(Obj);
 }
 
-void CopyDataFromConcrete(FPBDJointSettings& Data, const FPBDJointConstraintHandle& Joint)
+void CopyDataFromObject(FPBDJointSettings& Data, const FPBDJointConstraintHandle& Joint)
 {
 	Data = Joint.GetSettings();
 }
 
 
 
-template <bool bResim>
 void FRewindData::PushGTDirtyData(const FDirtyPropertiesManager& SrcManager,const int32 SrcDataIdx,const FDirtyProxy& Dirty, const FShapeDirtyData* ShapeDirtyData)
 {
 	//This records changes enqueued by GT.
@@ -418,7 +391,7 @@ void FRewindData::PushGTDirtyData(const FDirtyPropertiesManager& SrcManager,cons
 
 	//Helper to group most of the common logic about push data recording
 	//NOTE: when possible use passed in CopyFunc to do work, if lambda returns false you cannot record to history buffer
-	auto CopyHelper = [this, Proxy](auto Concrete, const auto& CopyFunc) -> bool
+	auto CopyHelper = [this, Proxy](auto Object, const auto& CopyFunc) -> bool
 	{
 		//Don't bother tracking static particles. We assume they stream in and out and don't need to be rewound
 		//TODO: find a way to skip statics that stream in and out - gameplay can technically spawn/destroy these so we can't just ignore statics
@@ -429,20 +402,20 @@ void FRewindData::PushGTDirtyData(const FDirtyPropertiesManager& SrcManager,cons
 
 		//During a resim the same exact push data comes from gt
 		//If the particle is already in sync, it will stay in sync so no need to touch history
-		if (bResim && Concrete->SyncState() == ESyncState::InSync)
+		if (IsResim() && Object->SyncState() == ESyncState::InSync)
 		{
 			return false;
 		}
 
-		if (bResim && Proxy->GetInitializedStep() == CurFrame)
+		if (IsResim() && Proxy->GetInitializedStep() == CurFrame)
 		{
 			//Particle is reinitialized, since it's out of sync it must be at a different time
 			//So make sure it's considered during resim
 			//TODO: should check if in bubble
-			Concrete->SetEnabledDuringResim(true);
+			Object->SetEnabledDuringResim(true);
 		}
 
-		auto& Info = FindOrAddConcrete(*Concrete, Proxy->IsInitialized() ? INDEX_NONE : CurFrame);
+		auto& Info = FindOrAddDirtyObj(*Object, Proxy->IsInitialized() ? INDEX_NONE : CurFrame);
 		auto& Latest = Info.AddFrame(CurFrame);
 
 		//At this point all phases should be clean
@@ -463,12 +436,12 @@ void FRewindData::PushGTDirtyData(const FDirtyPropertiesManager& SrcManager,cons
 		return true;
 	};
 
-	auto DirtyPropHelper = [this, &Dirty](auto& Property, const EChaosPropertyFlags PropName, const auto& Concrete)
+	auto DirtyPropHelper = [this, &Dirty](auto& Property, const EChaosPropertyFlags PropName, const auto& Object)
 	{
 		if (Dirty.PropertyData.IsDirty(PropName))
 		{
 			auto& Data = Property.WriteAccessMonotonic(FFrameAndPhase{ CurFrame, FFrameAndPhase::PrePushData }, PropertiesPool);
-			CopyDataFromConcrete(Data, Concrete);
+			CopyDataFromObject(Data, Object);
 		}
 	};
 
@@ -502,7 +475,7 @@ void FRewindData::PushGTDirtyData(const FDirtyPropertiesManager& SrcManager,cons
 			//Dynamics are not available at head (sim zeroes them out), so we have to record them as PostPushData (since they're applied as part of PushData)
 			if (auto NewData = Dirty.PropertyData.FindDynamics(SrcManager, SrcDataIdx))
 			{
-				FDirtyParticleInfo& Info = FindOrAddParticle(*PTParticle, ParticleProxy->IsInitialized() ? INDEX_NONE : CurFrame);
+				FDirtyParticleInfo& Info = FindOrAddDirtyObj(*PTParticle, ParticleProxy->IsInitialized() ? INDEX_NONE : CurFrame);
 				FGeometryParticleStateBase& Latest = Info.AddFrame(CurFrame);
 				const FFrameAndPhase PostPushData{ CurFrame, FFrameAndPhase::PostPushData };
 				Latest.Dynamics.WriteAccessMonotonic(PostPushData, PropertiesPool) = *NewData;
@@ -536,7 +509,7 @@ void FRewindData::SpawnProxyIfNeeded(FSingleParticlePhysicsProxy& Proxy)
 	if(Proxy.GetInitializedStep() > CurFrame)
 	{
 		FGeometryParticleHandle* Handle = Proxy.GetHandle_LowLevel();
-		FDirtyParticleInfo& Info = FindOrAddParticle(*Handle, CurFrame);
+		FDirtyParticleInfo& Info = FindOrAddDirtyObj(*Handle, CurFrame);
 
 		Solver->GetEvolution()->EnableParticle(Handle, nullptr);
 		if(Proxy.GetInitializedStep() != CurFrame)
@@ -549,7 +522,7 @@ void FRewindData::SpawnProxyIfNeeded(FSingleParticlePhysicsProxy& Proxy)
 
 void FRewindData::MarkDirtyFromPT(FGeometryParticleHandle& Handle)
 {
-	FDirtyParticleInfo& Info = FindOrAddParticle(Handle);
+	FDirtyParticleInfo& Info = FindOrAddDirtyObj(Handle);
 
 	Info.DirtyDynamics = CurFrame;
 
@@ -596,12 +569,11 @@ void FRewindData::MarkDirtyFromPT(FGeometryParticleHandle& Handle)
 	}
 }
 
-template <bool bResim>
 void FRewindData::PushPTDirtyData(TPBDRigidParticleHandle<FReal,3>& Handle,const int32 SrcDataIdx)
 	{
 	const bool bRecordingHistory = !IsResimAndInSync(Handle);
 
-	FDirtyParticleInfo& Info = FindOrAddParticle(Handle);
+	FDirtyParticleInfo& Info = FindOrAddDirtyObj(Handle);
 	FGeometryParticleStateBase& Latest = Info.AddFrame(CurFrame);
 	const FFrameAndPhase FrameAndPhase{ CurFrame, FFrameAndPhase::PostCallbacks };
 
@@ -627,36 +599,4 @@ void FRewindData::PushPTDirtyData(TPBDRigidParticleHandle<FReal,3>& Handle,const
 	}
 }
 
-FRewindData::FDirtyParticleInfo& FRewindData::FindOrAddParticle(TGeometryParticleHandle<FReal,3>& PTParticle, const int32 InitializedOnFrame)
-{
-	const FUniqueIdx UniqueIdx = PTParticle.UniqueIdx();
-	if(FDirtyParticleInfo* Info = FindParticle(UniqueIdx))
-	{
-		return *Info;
-	}
-
-	const int32 DirtyIdx = AllDirtyParticles.Add(FDirtyParticleInfo(PropertiesPool,PTParticle,UniqueIdx,CurFrame,Managers.Capacity()));
-	ParticleToAllDirtyIdx.Add(UniqueIdx,DirtyIdx);
-	if(InitializedOnFrame != INDEX_NONE)
-	{
-		AllDirtyParticles[DirtyIdx].InitializedOnStep = InitializedOnFrame;
-	}
-
-	return AllDirtyParticles[DirtyIdx];
-}
-
-FRewindData::FDirtyJointInfo& FRewindData::FindOrAddJoint(FPBDJointConstraintHandle& Handle, const int32 InitializedOnFrame)
-{
-	//TODO: implement
-	ensure(false);
-	static FDirtyJointInfo TempReturn{ PropertiesPool, Handle, CurFrame, 0 };
-	return TempReturn;
-}
-
-template void FRewindData::PushGTDirtyData<true>(const FDirtyPropertiesManager& SrcManager,const int32 SrcDataIdx,const FDirtyProxy& Dirty, const FShapeDirtyData* DirtyShapesData);
-template void FRewindData::PushGTDirtyData<false>(const FDirtyPropertiesManager& SrcManager,const int32 SrcDataIdx,const FDirtyProxy& Dirty, const FShapeDirtyData* DirtyShapesData);
-
-template void FRewindData::PushPTDirtyData<true>(TPBDRigidParticleHandle<FReal,3>& Rigid,const int32 SrcDataIdx);
-template void FRewindData::PushPTDirtyData<false>(TPBDRigidParticleHandle<FReal,3>& Rigid,const int32 SrcDataIdx);
-
 }
diff --git a/Engine/Source/Runtime/Experimental/Chaos/Public/PBDRigidsSolver.h b/Engine/Source/Runtime/Experimental/Chaos/Public/PBDRigidsSolver.h
--- a/Engine/Source/Runtime/Experimental/Chaos/Public/PBDRigidsSolver.h
+++ b/Engine/Source/Runtime/Experimental/Chaos/Public/PBDRigidsSolver.h
@@ -349,6 +349,8 @@ namespace Chaos
 		{
 			FSingleParticlePhysicsProxy* Proxy;
 			int32 DestroyOnStep;
+			FGeometryParticleHandle* Handle;
+			FUniqueIdx UniqueIdx;
 		};
 
 		TArray<FPendingDestroyInfo> PendingDestroyPhysicsProxy;
diff --git a/Engine/Source/Runtime/Experimental/Chaos/Public/RewindData.h b/Engine/Source/Runtime/Experimental/Chaos/Public/RewindData.h
--- a/Engine/Source/Runtime/Experimental/Chaos/Public/RewindData.h
+++ b/Engine/Source/Runtime/Experimental/Chaos/Public/RewindData.h
@@ -764,8 +764,101 @@ private:
 	const FFrameAndPhase FrameAndPhase = { 0,0 };
 };
 
+template <typename T> 
+const T* ConstifyHelper(T* Ptr) { return Ptr; }
+
+template <typename T>
+T NoRefHelper(const T& Ref) { return Ref; }
+
 extern CHAOS_API int32 EnableResimCache;
 
+template <typename TVal>
+class TDirtyObjects
+{
+public:
+	using TKey = decltype(ConstifyHelper(
+		((TVal*)0)->GetObjectPtr()
+	));
+
+	TVal& Add(const TKey Key, TVal&& Val)
+	{
+		if(int32* ExistingIdx = KeyToIdx.Find(Key))
+		{
+			ensure(false);	//Item alread exists, shouldn't be adding again
+			return DenseVals[*ExistingIdx];
+		}
+		else
+		{
+			const int32 Idx = DenseVals.Emplace(MoveTemp(Val));
+			KeyToIdx.Add(Key, Idx);
+			return DenseVals[Idx];
+		}
+	}
+
+	const TVal& FindChecked(const TKey Key) const
+	{
+		const int32 Idx = KeyToIdx.FindChecked(Key);
+		return DenseVals[Idx];
+	}
+
+	TVal& FindChecked(const TKey Key)
+	{
+		const int32 Idx = KeyToIdx.FindChecked(Key);
+		return DenseVals[Idx];
+	}
+
+	const TVal* Find(const TKey Key) const
+	{
+		if (const int32* Idx = KeyToIdx.Find(Key))
+		{
+			return &DenseVals[*Idx];
+		}
+
+		return nullptr;
+	}
+
+	TVal* Find(const TKey Key)
+	{
+		if (const int32* Idx = KeyToIdx.Find(Key))
+		{
+			return &DenseVals[*Idx];
+		}
+
+		return nullptr;
+	}
+
+	void Remove(const TKey Key)
+	{
+		if (const int32* Idx = KeyToIdx.Find(Key))
+		{
+			DenseVals.RemoveAtSwap(*Idx);
+
+			if(*Idx < DenseVals.Num())
+			{
+				const TKey SwappedKey = DenseVals[*Idx].GetObjectPtr();
+				KeyToIdx.FindChecked(SwappedKey) = *Idx;
+			}
+
+			KeyToIdx.Remove(Key);
+		}
+	}
+
+	int32 Num() const { return DenseVals.Num(); }
+
+	auto begin() { return DenseVals.begin(); }
+	auto end() { return DenseVals.end(); }
+
+	auto cbegin() const { return DenseVals.begin(); }
+	auto cend() const { return DenseVals.end(); }
+
+	const TVal* GetDensePtr() const { return DenseVals.GetData(); }
+	TVal* GetDensePtr() { return DenseVals.GetData(); }
+
+private:
+	TMap<TKey, int32> KeyToIdx;
+	TArray<TVal> DenseVals;
+};
+
 class FPBDRigidsSolver;
 
 class FRewindData
@@ -793,12 +886,29 @@ public:
 		return Managers[Frame].DeltaTime;
 	}
 
-	void CHAOS_API RemoveParticle(const FUniqueIdx UniqueIdx);
+	void CHAOS_API RemoveParticle(const FGeometryParticleHandle* Particle)
+	{
+		DirtyParticles.Remove(Particle);
+	}
+
+	void CHAOS_API RemoveJoint(const FPBDJointConstraintHandle* Joint)
+	{
+		DirtyJoints.Remove(Joint);
+	}
 
 	int32 CHAOS_API GetEarliestFrame_Internal() const { return CurFrame - FramesSaved; }
 
-	/* Query the state of particles from the past. Once a rewind happens state captured must be queried using GetFutureStateAtFrame */
-	FGeometryParticleState CHAOS_API GetPastStateAtFrame(const FGeometryParticleHandle& Handle,int32 Frame, FFrameAndPhase::EParticleHistoryPhase Phase = FFrameAndPhase::EParticleHistoryPhase::PostPushData) const;
+	/* Query the state of particles from the past. Can only be used when not already resimming*/
+	FGeometryParticleState CHAOS_API GetPastStateAtFrame(const FGeometryParticleHandle& Handle,int32 Frame, FFrameAndPhase::EParticleHistoryPhase Phase = FFrameAndPhase::EParticleHistoryPhase::PostPushData) const
+	{
+		return GetPastStateAtFrameImp<FGeometryParticleState>(DirtyParticles, Handle, Frame, Phase);
+	}
+
+	/* Query the state of joints from the past. Can only be used when not already resimming*/
+	FJointState CHAOS_API GetPastJointStateAtFrame(const FPBDJointConstraintHandle& Handle, int32 Frame, FFrameAndPhase::EParticleHistoryPhase Phase = FFrameAndPhase::EParticleHistoryPhase::PostPushData) const
+	{
+		return GetPastStateAtFrameImp<FJointState>(DirtyJoints, Handle, Frame, Phase);
+	}
 
 	IResimCacheBase* GetCurrentStepResimCache() const
 	{
@@ -855,12 +965,10 @@ public:
 	}
 
 	//Number of particles that we're currently storing history for
-	int32 GetNumDirtyParticles() const { return AllDirtyParticles.Num(); }
+	int32 GetNumDirtyParticles() const { return DirtyParticles.Num(); }
 
-	template <bool bResim>
 	void PushGTDirtyData(const FDirtyPropertiesManager& SrcManager,const int32 SrcDataIdx,const FDirtyProxy& Dirty, const FShapeDirtyData* ShapeDirtyData);
 
-	template <bool bResim>
 	void PushPTDirtyData(TPBDRigidParticleHandle<FReal,3>& Rigid,const int32 SrcDataIdx);
 
 	void CHAOS_API MarkDirtyFromPT(FGeometryParticleHandle& Handle);
@@ -871,24 +979,6 @@ private:
 
 	friend class FPBDRigidsSolver;
 
-	struct FDirtyFrameInfo
-	{
-		int32 Frame;	//needed to protect against stale entries in circular buffer
-		uint8 Wave;
-
-		void SetWave(int32 InFrame, uint8 InWave)
-		{
-			Frame = InFrame;
-			Wave = InWave;
-		}
-
-		bool MissingWrite(int32 InFrame, uint8 InWave) const
-		{
-			//If this is not a stale entry and it was written to, but not during this latest sim
-			return (Wave != 0 && Frame == InFrame) && Wave != InWave;
-		}
-	};
-
 	void CHAOS_API AdvanceFrameImp(IResimCacheBase* ResimCache);
 
 	struct FFrameManagerInfo
@@ -902,26 +992,30 @@ private:
 		FReal DeltaTime;
 	};
 
-	template <typename HistoryType>
-	struct TDirtyConcreteInfo
+	template <typename THistoryType, typename TObj>
+	struct TDirtyObjectInfo
 	{
 	private:
-		HistoryType History;
+		THistoryType History;
+		TObj* ObjPtr;
 		FDirtyPropertiesPool* PropertiesPool;
 	public:
+		int32 DirtyDynamics = INDEX_NONE;	//Only used by particles, indicates the dirty properties was written to.
 		int32 LastDirtyFrame;	//Track how recently this was made dirty
 		int32 InitializedOnStep = INDEX_NONE;	//if not INDEX_NONE, it indicates we saw initialization during rewind history window
 		bool bResimAsSlave = true;	//Indicates the particle will always resim in the exact same way from game thread data
 
-		TDirtyConcreteInfo(FDirtyPropertiesPool& InPropertiesPool, const int32 CurFrame, const int32 NumFrames)
+		TDirtyObjectInfo(FDirtyPropertiesPool& InPropertiesPool, TObj& InObj, const int32 CurFrame, const int32 NumFrames)
 			: History(NumFrames)
+			, ObjPtr(&InObj)
 			, PropertiesPool(&InPropertiesPool)
 			, LastDirtyFrame(CurFrame)
 		{
 		}
 
-		TDirtyConcreteInfo(TDirtyConcreteInfo&& Other)
+		TDirtyObjectInfo(TDirtyObjectInfo&& Other)
 			: History(MoveTemp(Other.History))
+			, ObjPtr(Other.ObjPtr)
 			, PropertiesPool(Other.PropertiesPool)
 			, LastDirtyFrame(Other.LastDirtyFrame)
 			, InitializedOnStep(Other.InitializedOnStep)
@@ -930,7 +1024,7 @@ private:
 			Other.PropertiesPool = nullptr;
 		}
 
-		~TDirtyConcreteInfo()
+		~TDirtyObjectInfo()
 		{
 			if (PropertiesPool)
 			{
@@ -938,9 +1032,11 @@ private:
 			}
 		}
 
-		TDirtyConcreteInfo(const TDirtyConcreteInfo& Other) = delete;
+		TDirtyObjectInfo(const TDirtyObjectInfo& Other) = delete;
 
-		HistoryType& AddFrame(const int32 Frame)
+		TObj* GetObjectPtr() const { return ObjPtr; }
+
+		THistoryType& AddFrame(const int32 Frame)
 		{
 			LastDirtyFrame = Frame;
 			return History;
@@ -951,102 +1047,59 @@ private:
 			History.ClearEntryAndFuture(FrameAndPhase);
 		}
 
-		const HistoryType& GetHistory() const	//For non-const access use AddFrame
+		const THistoryType& GetHistory() const	//For non-const access use AddFrame
 		{
 			return History;
 		}
 	};
 
-	struct FDirtyParticleInfo : public TDirtyConcreteInfo<FGeometryParticleStateBase>
-	{
-	private:
-		using Base = TDirtyConcreteInfo<FGeometryParticleStateBase>;
-		TGeometryParticleHandle<FReal,3>* PTParticle;
-	public:
-		FUniqueIdx CachedUniqueIdx;	//Needed when manipulating on physics thread and Particle data cannot be read
-		int32 DirtyDynamics = INDEX_NONE;
+	using FDirtyParticleInfo = TDirtyObjectInfo<FGeometryParticleStateBase, FGeometryParticleHandle>;
+	using FDirtyJointInfo = TDirtyObjectInfo<FJointStateBase, FPBDJointConstraintHandle>;
 
-		FDirtyParticleInfo(FDirtyPropertiesPool& InPropertiesPool, TGeometryParticleHandle<FReal,3>& InPTParticle, const FUniqueIdx UniqueIdx,const int32 CurFrame,const int32 NumFrames)
-		: Base(InPropertiesPool, CurFrame, NumFrames)
-		, PTParticle(&InPTParticle)
-		, CachedUniqueIdx(UniqueIdx)
-		{
-		}
-		
-		TGeometryParticleHandle<FReal,3>* GetPTParticle() const
-		{
-			return PTParticle;
-		}
-	};
-
-	struct FDirtyJointInfo : public TDirtyConcreteInfo<FJointStateBase>
+	template <typename TDirtyObjs, typename TObj>
+	auto& FindOrAddDirtyObjImp(TDirtyObjs& DirtyObjs, TObj& Handle, const int32 InitializedOnFrame = INDEX_NONE)
 	{
-	private:
-		using Base = TDirtyConcreteInfo<FJointStateBase>;
-		FPBDJointConstraintHandle* Handle;
-	public:
-
-		FDirtyJointInfo(FDirtyPropertiesPool& InPropertiesPool, FPBDJointConstraintHandle& InJoint, const int32 CurFrame, const int32 NumFrames)
-			: Base(InPropertiesPool, CurFrame, NumFrames)
-			, Handle(&InJoint)
+		if (auto Info = DirtyObjs.Find(&Handle))
 		{
+			return *Info;
 		}
-	};
-
-	const FDirtyParticleInfo& FindParticleChecked(const FUniqueIdx UniqueIdx) const
-	{
-		const int32 Idx = ParticleToAllDirtyIdx.FindChecked(UniqueIdx);
-		return AllDirtyParticles[Idx];
-	}
 
-	FDirtyParticleInfo& FindParticleChecked(const FUniqueIdx UniqueIdx)
-	{
-		const int32 Idx = ParticleToAllDirtyIdx.FindChecked(UniqueIdx);
-		return AllDirtyParticles[Idx];
+		using TDirtyObj = decltype(NoRefHelper(*DirtyObjs.GetDensePtr()));
+		TDirtyObj& Info = DirtyObjs.Add(&Handle, TDirtyObj(PropertiesPool, Handle, CurFrame, Managers.Capacity()));
+		Info.InitializedOnStep = InitializedOnFrame;
+		return Info;
 	}
 
-	const FDirtyParticleInfo* FindParticle(const FUniqueIdx UniqueIdx) const
+	FDirtyParticleInfo& FindOrAddDirtyObj(FGeometryParticleHandle& Handle, const int32 InitializedOnFrame = INDEX_NONE)
 	{
-		if(const int32* Idx = ParticleToAllDirtyIdx.Find(UniqueIdx))
-		{
-			return &AllDirtyParticles[*Idx];
-		}
-
-		return nullptr;
+		return FindOrAddDirtyObjImp(DirtyParticles, Handle, InitializedOnFrame);
 	}
 
-	FDirtyParticleInfo* FindParticle(const FUniqueIdx UniqueIdx)
+	FDirtyJointInfo& FindOrAddDirtyObj(FPBDJointConstraintHandle& Handle, const int32 InitializedOnFrame = INDEX_NONE)
 	{
-		if(const int32* Idx = ParticleToAllDirtyIdx.Find(UniqueIdx))
-		{
-			return &AllDirtyParticles[*Idx];
-		}
-
-		return nullptr;
+		return FindOrAddDirtyObjImp(DirtyJoints, Handle, InitializedOnFrame);
 	}
 
-	FDirtyParticleInfo& FindOrAddParticle(TGeometryParticleHandle<FReal,3>& PTParticle, const int32 InitializedOnFrame = INDEX_NONE);
-	FDirtyJointInfo& FindOrAddJoint(FPBDJointConstraintHandle& Handle, const int32 InitializedOnFrame = INDEX_NONE);
-
-	FDirtyParticleInfo& FindOrAddConcrete(FGeometryParticleHandle& Handle, const int32 InitializedOnFrame = INDEX_NONE)
+	template <typename TObjState, typename TDirtyObjs, typename TObj>
+	auto GetPastStateAtFrameImp(const TDirtyObjs& DirtyObjs, const TObj& Handle, int32 Frame, FFrameAndPhase::EParticleHistoryPhase Phase) const
 	{
-		return FindOrAddParticle(Handle, InitializedOnFrame);
-	}
+		ensure(!IsResim());
+		ensure(Frame >= GetEarliestFrame_Internal());	//can't get state from before the frame we rewound to
 
-	FDirtyJointInfo& FindOrAddConcrete(FPBDJointConstraintHandle& Handle, const int32 InitializedOnFrame = INDEX_NONE)
-	{
-		return FindOrAddJoint(Handle, InitializedOnFrame);
+		const auto* Info = DirtyObjs.Find(&Handle);
+		const auto* State = Info ? &Info->GetHistory() : nullptr;
+		return TObjState(State, Handle, PropertiesPool, { Frame, Phase });
 	}
 
 	bool RewindToFrame(int32 Frame);
 
 	static void DesyncParticle(FDirtyParticleInfo& Info, const FFrameAndPhase FrameAndPhase);
 
-	TArrayAsMap<FUniqueIdx,int32> ParticleToAllDirtyIdx;
 	TCircularBuffer<FFrameManagerInfo> Managers;
-	FDirtyPropertiesPool PropertiesPool;	//must come before AllDirtyParticles since it relies on it (and used in destruction)
-	TArray<FDirtyParticleInfo> AllDirtyParticles;
-	TArray<FDirtyJointInfo> AllDirtyJoints;
+	FDirtyPropertiesPool PropertiesPool;	//must come before DirtyParticles since it relies on it (and used in destruction)
+
+	TDirtyObjects<FDirtyParticleInfo> DirtyParticles;
+	TDirtyObjects<FDirtyJointInfo> DirtyJoints;
 
 	FPBDRigidsSolver* Solver;
 	int32 CurFrame;
