Attempt to fix CIS

diff --git a/Engine/Source/Runtime/Experimental/Chaos/Private/RewindData.cpp b/Engine/Source/Runtime/Experimental/Chaos/Private/RewindData.cpp
--- a/Engine/Source/Runtime/Experimental/Chaos/Private/RewindData.cpp
+++ b/Engine/Source/Runtime/Experimental/Chaos/Private/RewindData.cpp
@@ -599,4 +599,15 @@ void FRewindData::PushPTDirtyData(TPBDRigidParticleHandle<FReal,3>& Handle,const
 	}
 }
 
+FGeometryParticleState FRewindData::GetPastStateAtFrame(const FGeometryParticleHandle& Handle, int32 Frame, FFrameAndPhase::EParticleHistoryPhase Phase) const
+{
+	return GetPastStateAtFrameImp<FGeometryParticleState>(DirtyParticles, Handle, Frame, Phase);
+}
+
+FJointState FRewindData::GetPastJointStateAtFrame(const FPBDJointConstraintHandle& Handle, int32 Frame, FFrameAndPhase::EParticleHistoryPhase Phase) const
+{
+	return GetPastStateAtFrameImp<FJointState>(DirtyJoints, Handle, Frame, Phase);
+}
+
+
 }
diff --git a/Engine/Source/Runtime/Experimental/Chaos/Public/RewindData.h b/Engine/Source/Runtime/Experimental/Chaos/Public/RewindData.h
--- a/Engine/Source/Runtime/Experimental/Chaos/Public/RewindData.h
+++ b/Engine/Source/Runtime/Experimental/Chaos/Public/RewindData.h
@@ -899,16 +899,10 @@ public:
 	int32 CHAOS_API GetEarliestFrame_Internal() const { return CurFrame - FramesSaved; }
 
 	/* Query the state of particles from the past. Can only be used when not already resimming*/
-	FGeometryParticleState CHAOS_API GetPastStateAtFrame(const FGeometryParticleHandle& Handle,int32 Frame, FFrameAndPhase::EParticleHistoryPhase Phase = FFrameAndPhase::EParticleHistoryPhase::PostPushData) const
-	{
-		return GetPastStateAtFrameImp<FGeometryParticleState>(DirtyParticles, Handle, Frame, Phase);
-	}
+	FGeometryParticleState CHAOS_API GetPastStateAtFrame(const FGeometryParticleHandle& Handle, int32 Frame, FFrameAndPhase::EParticleHistoryPhase Phase = FFrameAndPhase::EParticleHistoryPhase::PostPushData) const;
 
 	/* Query the state of joints from the past. Can only be used when not already resimming*/
-	FJointState CHAOS_API GetPastJointStateAtFrame(const FPBDJointConstraintHandle& Handle, int32 Frame, FFrameAndPhase::EParticleHistoryPhase Phase = FFrameAndPhase::EParticleHistoryPhase::PostPushData) const
-	{
-		return GetPastStateAtFrameImp<FJointState>(DirtyJoints, Handle, Frame, Phase);
-	}
+	FJointState CHAOS_API GetPastJointStateAtFrame(const FPBDJointConstraintHandle& Handle, int32 Frame, FFrameAndPhase::EParticleHistoryPhase Phase = FFrameAndPhase::EParticleHistoryPhase::PostPushData) const;
 
 	IResimCacheBase* GetCurrentStepResimCache() const
 	{
