GeometryCollectionComponent
-- Removal shrinks for bottom of geometry instead of CoM

#okforgitbug internal, public

diff --git a/Engine/Source/Runtime/Experimental/Chaos/Private/GeometryCollection/GeometryCollectionAlgo.cpp b/Engine/Source/Runtime/Experimental/Chaos/Private/GeometryCollection/GeometryCollectionAlgo.cpp
--- a/Engine/Source/Runtime/Experimental/Chaos/Private/GeometryCollection/GeometryCollectionAlgo.cpp
+++ b/Engine/Source/Runtime/Experimental/Chaos/Private/GeometryCollection/GeometryCollectionAlgo.cpp
@@ -297,21 +297,20 @@ namespace GeometryCollectionAlgo
 		}
 	}
 
-	void GlobalMatricesRecursive(const int32 Index, const TManagedArray<int32>& Parents, const TManagedArray<FTransform>& Transform, TArray<bool>& IsTransformComputed, const TManagedArray<float>* Scales, const TManagedArray<FTransform>* MassToLocals, TArray<FMatrix>& OutGlobalTransforms)
+	void GlobalMatricesRecursive(const int32 Index, const TManagedArray<int32>& Parents, const TManagedArray<FTransform>& Transform, TArray<bool>& IsTransformComputed, const TManagedArray<FTransform>* UniformScale, TArray<FMatrix>& OutGlobalTransforms)
 	{
 		if (!IsTransformComputed[Index])
 		{
 			FMatrix Result = Transform[Index].ToMatrixWithScale();
 			if (Parents[Index] != FGeometryCollection::Invalid)
 			{
-				GlobalMatricesRecursive(Parents[Index], Parents, Transform, IsTransformComputed, Scales, MassToLocals, OutGlobalTransforms);
+				GlobalMatricesRecursive(Parents[Index], Parents, Transform, IsTransformComputed, UniformScale, OutGlobalTransforms);
 				Result *= OutGlobalTransforms[Parents[Index]];
 			}
 
-			if (Scales && !FMath::IsNearlyEqual((*Scales)[Index], 1.0f))
+			if (UniformScale)
 			{
-				check (MassToLocals);
-				OutGlobalTransforms[Index] = (*MassToLocals)[Index].ToInverseMatrixWithScale() * FTransform(FQuat::Identity, FVector(), FVector((*Scales)[Index])).ToMatrixWithScale() * (*MassToLocals)[Index].ToMatrixWithScale() * Result;
+				OutGlobalTransforms[Index] = (*UniformScale)[Index].ToMatrixWithScale() * Result;
 			}
 			else
 			{ 
@@ -323,27 +322,20 @@ namespace GeometryCollectionAlgo
 	}
 
 	// #note: this verison returns an FTransform to support functionality for getting global matrices for an array of indices.
-	void GlobalMatricesRecursive(const int32 Index, const TManagedArray<int32>& Parents, const TManagedArray<FTransform>& Transform, TArray<bool>& IsTransformComputed, const TManagedArray<float>* Scales, const TManagedArray<FTransform>* MassToLocals, TArray<FTransform>& OutGlobalTransforms)
+	void GlobalMatricesRecursive(const int32 Index, const TManagedArray<int32>& Parents, const TManagedArray<FTransform>& Transform, TArray<bool>& IsTransformComputed, const TManagedArray<FTransform>* UniformScale, TArray<FTransform>& OutGlobalTransforms)
 	{
 		if (!IsTransformComputed[Index])
 		{
 			FTransform Result = Transform[Index];
 			if (Parents[Index] != FGeometryCollection::Invalid)
 			{
-				GlobalMatricesRecursive(Parents[Index], Parents, Transform, IsTransformComputed, Scales, MassToLocals, OutGlobalTransforms);
+				GlobalMatricesRecursive(Parents[Index], Parents, Transform, IsTransformComputed, UniformScale, OutGlobalTransforms);
 				Result *= OutGlobalTransforms[Parents[Index]];
 			}
-
-			if (Scales && !FMath::IsNearlyEqual((*Scales)[Index], 1.0f))
+			
+			if (UniformScale)
 			{
-				if(MassToLocals)
-				{ 
-					OutGlobalTransforms[Index] = (*MassToLocals)[Index].Inverse() * FTransform(FQuat::Identity, FVector(), FVector((*Scales)[Index])) * (*MassToLocals)[Index] * Result;
-				}
-				else
-				{
-					OutGlobalTransforms[Index] = FTransform(FQuat::Identity, FVector(), FVector((*Scales)[Index])) * Result;
-				}
+				OutGlobalTransforms[Index] = (*UniformScale)[Index] * Result;
 			}
 			else
 			{
@@ -354,20 +346,20 @@ namespace GeometryCollectionAlgo
 		}
 	}
 
-	FTransform GlobalMatricesRecursiveForIndices(const int32 Index, const TManagedArray<int32>& Parents, const TManagedArray<FTransform>& Transform, TArray<bool>& IsTransformComputed, const TManagedArray<float>* Scales, const TManagedArray<FTransform>* MassToLocals, TArray<FTransform>& TransformCache)
+	FTransform GlobalMatricesRecursiveForIndices(const int32 Index, const TManagedArray<int32>& Parents, const TManagedArray<FTransform>& Transform, TArray<bool>& IsTransformComputed, const TManagedArray<FTransform>* UniformScale, TArray<FTransform>& TransformCache)
 	{
 		if (!IsTransformComputed[Index])
 		{
 			FTransform Result = Transform[Index];
 			if (Parents[Index] != FGeometryCollection::Invalid)
 			{
-				GlobalMatricesRecursive(Parents[Index], Parents, Transform, IsTransformComputed, Scales, MassToLocals, TransformCache);
+				GlobalMatricesRecursive(Parents[Index], Parents, Transform, IsTransformComputed, UniformScale, TransformCache);
 				Result *= TransformCache[Parents[Index]];
 			}
-
-			if (Scales && !FMath::IsNearlyEqual((*Scales)[Index], 1.0f))
+			
+			if (UniformScale)
 			{
-				TransformCache[Index] = (*MassToLocals)[Index].Inverse() * FTransform(FQuat::Identity, FVector(), FVector((*Scales)[Index])) * (*MassToLocals)[Index] * Result;
+				TransformCache[Index] = (*UniformScale)[Index] * Result;
 			}
 			else
 			{
@@ -398,9 +390,8 @@ namespace GeometryCollectionAlgo
 		return Transform;
 	}
 
-	void GlobalMatrices(const TManagedArray<FTransform>& RelativeTransforms, const TManagedArray<int32>& Parents, const TArray<int32>& Indices, const TManagedArray<float>& Scales, const TManagedArray<FTransform>& MassToLocals, TArray<FTransform>& OutGlobalTransforms)
+	void GlobalMatrices(const TManagedArray<FTransform>& RelativeTransforms, const TManagedArray<int32>& Parents, const TArray<int32>& Indices, const TManagedArray<FTransform>& UniformScale, TArray<FTransform>& OutGlobalTransforms)
 	{
-		check(Scales.Num() == RelativeTransforms.Num());
 		TArray<bool> IsTransformComputed;
 		IsTransformComputed.AddDefaulted(RelativeTransforms.Num());
 
@@ -410,7 +401,7 @@ namespace GeometryCollectionAlgo
 		OutGlobalTransforms.SetNumUninitialized(Indices.Num(), false);
 		for (int Idx = 0; Idx < Indices.Num(); Idx++)
 		{
-			OutGlobalTransforms[Idx] = GlobalMatricesRecursiveForIndices(Indices[Idx], Parents, RelativeTransforms, IsTransformComputed, &Scales, &MassToLocals, TransformCache);
+			OutGlobalTransforms[Idx] = GlobalMatricesRecursiveForIndices(Indices[Idx], Parents, RelativeTransforms, IsTransformComputed, &UniformScale, TransformCache);
 		}
 	}
 
@@ -425,7 +416,7 @@ namespace GeometryCollectionAlgo
 		OutGlobalTransforms.SetNumUninitialized(Indices.Num(), false);
 		for (int Idx = 0; Idx < Indices.Num(); Idx++)
 		{
-			OutGlobalTransforms[Idx] = GlobalMatricesRecursiveForIndices(Indices[Idx], Parents, RelativeTransforms, IsTransformComputed, nullptr, nullptr, TransformCache);
+			OutGlobalTransforms[Idx] = GlobalMatricesRecursiveForIndices(Indices[Idx], Parents, RelativeTransforms, IsTransformComputed, nullptr, TransformCache);
 		}
 	}
 
@@ -444,12 +435,10 @@ namespace GeometryCollectionAlgo
 	}
 
 	template<typename MatrixType>
-	void GlobalMatrices(const TManagedArray<FTransform>& RelativeTransforms, const TManagedArray<int32>& Parents, const TManagedArray<float>& Scales, const TManagedArray<FTransform>& MassToLocals, TArray<MatrixType>& OutGlobalTransforms)
+	void GlobalMatrices(const TManagedArray<FTransform>& RelativeTransforms, const TManagedArray<int32>& Parents, const TManagedArray<FTransform>& UniformScale, TArray<MatrixType>& OutGlobalTransforms)
 	{
 		int32 NumTransforms = RelativeTransforms.Num();
-		check(Scales.Num() == NumTransforms);
-		check(MassToLocals.Num() == NumTransforms);
-
+		
 		TArray<bool> IsTransformComputed;
 		IsTransformComputed.AddDefaulted(NumTransforms);
 
@@ -457,7 +446,7 @@ namespace GeometryCollectionAlgo
 
 		for (int BoneIdx = 0; BoneIdx < NumTransforms; ++BoneIdx)
 		{
-			GlobalMatricesRecursive(BoneIdx, Parents, RelativeTransforms, IsTransformComputed, &Scales, &MassToLocals, OutGlobalTransforms);
+			GlobalMatricesRecursive(BoneIdx, Parents, RelativeTransforms, IsTransformComputed, &UniformScale, OutGlobalTransforms);
 		}
 	}
 
@@ -473,12 +462,12 @@ namespace GeometryCollectionAlgo
 
 		for (int BoneIdx = 0; BoneIdx < NumTransforms; ++BoneIdx)
 		{
-			GlobalMatricesRecursive(BoneIdx, Parents, RelativeTransforms, IsTransformComputed, nullptr, nullptr, OutGlobalTransforms);
+			GlobalMatricesRecursive(BoneIdx, Parents, RelativeTransforms, IsTransformComputed, nullptr, OutGlobalTransforms);
 		}
 	}
-	template void CHAOS_API GlobalMatrices<FTransform>(const TManagedArray<FTransform>&, const TManagedArray<int32>&, const TManagedArray<float>&, const TManagedArray<FTransform>&, TArray<FTransform>&);
+	template void CHAOS_API GlobalMatrices<FTransform>(const TManagedArray<FTransform>&, const TManagedArray<int32>&, const TManagedArray<FTransform>&, TArray<FTransform>&);
 	template void CHAOS_API GlobalMatrices<FTransform>(const TManagedArray<FTransform>&, const TManagedArray<int32>&, TArray<FTransform>&);
-	template void CHAOS_API GlobalMatrices<FMatrix>(const TManagedArray<FTransform>&, const TManagedArray<int32>&, const TManagedArray<float>&, const TManagedArray<FTransform>&, TArray<FMatrix>&);
+	template void CHAOS_API GlobalMatrices<FMatrix>(const TManagedArray<FTransform>&, const TManagedArray<int32>&, const TManagedArray<FTransform>&, TArray<FMatrix>&);
 	template void CHAOS_API GlobalMatrices<FMatrix>(const TManagedArray<FTransform>&, const TManagedArray<int32>&, TArray<FMatrix>&);
 
 	void FloodForOverlappedPairs(int Level, int32 BoneIndex, TMap<int32, int32> &BoneToGroup, const TManagedArray<int32>& Levels, const TMap<int32, FBox>& BoundingBoxes, TSet<TTuple<int32, int32>>& OutOverlappedPairs)
diff --git a/Engine/Source/Runtime/Experimental/Chaos/Public/GeometryCollection/GeometryCollectionAlgo.h b/Engine/Source/Runtime/Experimental/Chaos/Public/GeometryCollection/GeometryCollectionAlgo.h
--- a/Engine/Source/Runtime/Experimental/Chaos/Public/GeometryCollection/GeometryCollectionAlgo.h
+++ b/Engine/Source/Runtime/Experimental/Chaos/Public/GeometryCollection/GeometryCollectionAlgo.h
@@ -149,7 +149,7 @@ namespace GeometryCollectionAlgo
 	*  Global Matrices of the collection, transforms will be resized to fit
 	*/
 	template<typename MatrixType>
-	void CHAOS_API GlobalMatrices(const TManagedArray<FTransform>& RelativeTransforms, const TManagedArray<int32>& Parents, const TManagedArray<float>& Scales, const TManagedArray<FTransform>& MassToLocals, TArray<MatrixType>& Transforms);
+	void CHAOS_API GlobalMatrices(const TManagedArray<FTransform>& RelativeTransforms, const TManagedArray<int32>& Parents, const TManagedArray<FTransform>& UniformScale, TArray<MatrixType>& Transforms);
 
 	template<typename MatrixType>
 	void CHAOS_API GlobalMatrices(const TManagedArray<FTransform>& RelativeTransforms, const TManagedArray<int32>& Parents, TArray<MatrixType>& Transforms);
diff --git a/Engine/Source/Runtime/Experimental/GeometryCollectionEngine/Private/GeometryCollection/GeometryCollectionComponent.cpp b/Engine/Source/Runtime/Experimental/GeometryCollectionEngine/Private/GeometryCollection/GeometryCollectionComponent.cpp
--- a/Engine/Source/Runtime/Experimental/GeometryCollectionEngine/Private/GeometryCollection/GeometryCollectionComponent.cpp
+++ b/Engine/Source/Runtime/Experimental/GeometryCollectionEngine/Private/GeometryCollection/GeometryCollectionComponent.cpp
@@ -1837,8 +1837,30 @@ void UGeometryCollectionComponent::ResetDynamicCollection()
 
 			if (!DynamicCollection->HasAttribute("UniformScale", FGeometryCollection::TransformGroup))
 			{
-				TManagedArray<float>& UniformScale = DynamicCollection->AddAttribute<float>("UniformScale", FGeometryCollection::TransformGroup);
-				UniformScale.Fill(1.f);
+				TManagedArray<FTransform>& UniformScale = DynamicCollection->AddAttribute<FTransform>("UniformScale", FGeometryCollection::TransformGroup);
+				UniformScale.Fill(FTransform::Identity);
+			}
+
+			if (!DynamicCollection->HasAttribute("MaxSleepTime", FGeometryCollection::TransformGroup))
+			{
+				float MinTime = FMath::Max(0.0f, RestCollection->MaximumSleepTime.X);
+				float MaxTime = FMath::Max(MinTime, RestCollection->MaximumSleepTime.Y);
+				TManagedArray<float>& MaxSleepTime = DynamicCollection->AddAttribute<float>("MaxSleepTime", FGeometryCollection::TransformGroup);
+				for (int32 Idx = 0; Idx < MaxSleepTime.Num(); ++Idx)
+				{
+					MaxSleepTime[Idx] = FMath::RandRange(MinTime, MaxTime);
+				}
+			}
+
+			if (!DynamicCollection->HasAttribute("RemovalDuration", FGeometryCollection::TransformGroup))
+			{
+				float MinTime = FMath::Max(0.0f, RestCollection->RemovalDuration.X);
+				float MaxTime = FMath::Max(MinTime, RestCollection->RemovalDuration.Y);
+				TManagedArray<float>& RemovalDuration = DynamicCollection->AddAttribute<float>("RemovalDuration", FGeometryCollection::TransformGroup);
+				for (int32 Idx = 0; Idx < RemovalDuration.Num(); ++Idx)
+				{
+					RemovalDuration[Idx] = FMath::RandRange(MinTime, MaxTime);
+				}
 			}
 		}
 
@@ -2827,20 +2849,43 @@ void UGeometryCollectionComponent::CalculateGlobalMatrices()
 		else
 		{
 			// Have to fully rebuild
-			if (DynamicCollection && RestCollection->bRemoveOnMaxSleep && DynamicCollection->HasAttribute("SleepTimer", FGeometryCollection::TransformGroup) && DynamicCollection->HasAttribute("UniformScale", FGeometryCollection::TransformGroup))
+			if (DynamicCollection 
+				&& RestCollection->bRemoveOnMaxSleep 
+				&& DynamicCollection->HasAttribute("SleepTimer", FGeometryCollection::TransformGroup) 
+				&& DynamicCollection->HasAttribute("UniformScale", FGeometryCollection::TransformGroup)
+				&& DynamicCollection->HasAttribute("MaxSleepTime", FGeometryCollection::TransformGroup)
+				&& DynamicCollection->HasAttribute("RemovalDuration", FGeometryCollection::TransformGroup))
 			{
 				const TManagedArray<float>& SleepTimer = DynamicCollection->GetAttribute<float>("SleepTimer", FGeometryCollection::TransformGroup);
-				TManagedArray<float>& UniformScale = DynamicCollection->GetAttribute<float>("UniformScale", FGeometryCollection::TransformGroup);
-				
+				const TManagedArray<float>& MaxSleepTime = DynamicCollection->GetAttribute<float>("MaxSleepTime", FGeometryCollection::TransformGroup);
+				const TManagedArray<float>& RemovalDuration = DynamicCollection->GetAttribute<float>("RemovalDuration", FGeometryCollection::TransformGroup);
+				TManagedArray<FTransform>& UniformScale = DynamicCollection->GetAttribute<FTransform>("UniformScale", FGeometryCollection::TransformGroup);
+
 				for (int32 Idx = 0; Idx < GetTransformArray().Num(); ++Idx)
 				{
-					if (SleepTimer[Idx] > RestCollection->MaximumSleepTime)
+					if (SleepTimer[Idx] > MaxSleepTime[Idx])
 					{
-						UniformScale[Idx] = 1.0 - FMath::Min(1.0, (SleepTimer[Idx] - RestCollection->MaximumSleepTime) / RestCollection->RemovalDuration);
+						float Scale = 1.0 - FMath::Min(1.0, (SleepTimer[Idx] - MaxSleepTime[Idx]) / RemovalDuration[Idx]);
+						
+						if ((Scale < 1.0) && (Scale > 0.0))
+						{
+							float ShrinkRadius = 0.0f;
+							FSphere AccumulatedSphere;
+							if (CalculateInnerSphere(Idx, AccumulatedSphere))
+							{
+								ShrinkRadius = -AccumulatedSphere.W;
+							}
+							
+							FQuat LocalRotation = (GetComponentTransform().Inverse() * FTransform(GlobalMatrices[Idx]).Inverse()).GetRotation();
+							FTransform LocalDown(LocalRotation.RotateVector(FVector(0.f, 0.f, ShrinkRadius)));
+							FTransform ToCOM(DynamicCollection->MassToLocal[Idx].GetTranslation());
+							UniformScale[Idx] = ToCOM.Inverse() * LocalDown.Inverse() * FTransform(FQuat::Identity, FVector(0.f, 0.f, 0.f), FVector(Scale)) * LocalDown * ToCOM;
+						}
+	
 					}
 				}
 				
-				GeometryCollectionAlgo::GlobalMatrices(GetTransformArray(), GetParentArray(), UniformScale, DynamicCollection->MassToLocal, GlobalMatrices);
+				GeometryCollectionAlgo::GlobalMatrices(GetTransformArray(), GetParentArray(), UniformScale, GlobalMatrices);
 			}
 			else
 			{ 
@@ -3059,18 +3104,23 @@ void UGeometryCollectionComponent::InitializeEmbeddedGeometry()
 void UGeometryCollectionComponent::IncrementSleepTimer(float DeltaTime)
 {
 	// If a particle is sleeping, increment its sleep timer, otherwise reset it.
-	if (DynamicCollection && PhysicsProxy && DynamicCollection->HasAttribute("SleepTimer", FGeometryCollection::TransformGroup))
+	if (DynamicCollection && PhysicsProxy 
+		&& DynamicCollection->HasAttribute("SleepTimer", FGeometryCollection::TransformGroup) 
+		&& DynamicCollection->HasAttribute("MaxSleepTime", FGeometryCollection::TransformGroup)
+		&& DynamicCollection->HasAttribute("RemovalDuration", FGeometryCollection::TransformGroup))
 	{
 		TManagedArray<float>& SleepTimer = DynamicCollection->GetAttribute<float>("SleepTimer", FGeometryCollection::TransformGroup);
+		const TManagedArray<float>& RemovalDuration = DynamicCollection->GetAttribute<float>("RemovalDuration", FGeometryCollection::TransformGroup);
+		const TManagedArray<float>& MaxSleepTime = DynamicCollection->GetAttribute<float>("MaxSleepTime", FGeometryCollection::TransformGroup);
 		TArray<int32> ToDisable;
 		for (int32 TransformIdx = 0; TransformIdx < SleepTimer.Num(); ++TransformIdx)
 		{
-			bool PreviouslyAwake = SleepTimer[TransformIdx] < RestCollection->MaximumSleepTime;
-			if (SleepTimer[TransformIdx] < (RestCollection->MaximumSleepTime + RestCollection->RemovalDuration))
+			bool PreviouslyAwake = SleepTimer[TransformIdx] < MaxSleepTime[TransformIdx];
+			if (SleepTimer[TransformIdx] < (MaxSleepTime[TransformIdx] + RemovalDuration[TransformIdx]))
 			{
 				SleepTimer[TransformIdx] = (DynamicCollection->DynamicState[TransformIdx] == (int)EObjectStateTypeEnum::Chaos_Object_Sleeping) ? SleepTimer[TransformIdx] + DeltaTime : 0.0f;
 
-				if (SleepTimer[TransformIdx] > RestCollection->MaximumSleepTime)
+				if (SleepTimer[TransformIdx] > MaxSleepTime[TransformIdx])
 				{ 
 					DynamicCollection->MakeDirty();
 					if (PreviouslyAwake)
@@ -3089,3 +3139,47 @@ void UGeometryCollectionComponent::IncrementSleepTimer(float DeltaTime)
 	}
 }
 
+bool UGeometryCollectionComponent::CalculateInnerSphere(int32 TransformIndex, FSphere& SphereOut) const
+{
+	// Approximates the inscribed sphere. Returns false if no such sphere exists, if for instance the index is to an embedded geometry. 
+
+	const TManagedArray<int32>& TransformToGeometryIndex = RestCollection->GetGeometryCollection()->TransformToGeometryIndex;
+	const TManagedArray<Chaos::FRealSingle>& InnerRadius = RestCollection->GetGeometryCollection()->InnerRadius;
+	const TManagedArray<TSet<int32>>& Children = RestCollection->GetGeometryCollection()->Children;
+	const TManagedArray<FTransform>& MassToLocal = RestCollection->GetGeometryCollection()->GetAttribute<FTransform>("MassToLocal", FGeometryCollection::TransformGroup);
+
+	if (RestCollection->GetGeometryCollection()->IsRigid(TransformIndex))
+	{
+		// Sphere in component space, centered on body's COM.
+		FVector COM = MassToLocal[TransformIndex].GetLocation();
+		SphereOut = FSphere(COM, InnerRadius[TransformToGeometryIndex[TransformIndex]]);
+		return true;
+	}
+	else if (RestCollection->GetGeometryCollection()->IsClustered(TransformIndex))
+	{
+		// Recursively accumulate the cluster's child spheres. 
+		bool bSphereFound = false;
+		for (int32 ChildIndex: Children[TransformIndex])
+		{
+			FSphere LocalSphere;
+			if (CalculateInnerSphere(ChildIndex, LocalSphere))
+			{
+				if (!bSphereFound)
+				{
+					bSphereFound = true;
+					SphereOut = LocalSphere;
+				}
+				else
+				{
+					SphereOut += LocalSphere;
+				}
+			}
+		}
+		return bSphereFound;
+	}
+	else
+	{
+		// Likely an embedded geometry, which doesn't count towards volume.
+		return false;
+	}
+}
\ No newline at end of file
diff --git a/Engine/Source/Runtime/Experimental/GeometryCollectionEngine/Public/GeometryCollection/GeometryCollectionComponent.h b/Engine/Source/Runtime/Experimental/GeometryCollectionEngine/Public/GeometryCollection/GeometryCollectionComponent.h
--- a/Engine/Source/Runtime/Experimental/GeometryCollectionEngine/Public/GeometryCollection/GeometryCollectionComponent.h
+++ b/Engine/Source/Runtime/Experimental/GeometryCollectionEngine/Public/GeometryCollection/GeometryCollectionComponent.h
@@ -843,6 +843,7 @@ private:
 	void ClearEmbeddedGeometry();
 
 	void IncrementSleepTimer(float DeltaTime);
+	bool CalculateInnerSphere(int32 TransformIndex, FSphere& SphereOut) const;
 
 	/** True if GeometryCollection transforms have changed from previous tick. */
 	bool bIsMoving;
diff --git a/Engine/Source/Runtime/Experimental/GeometryCollectionEngine/Public/GeometryCollection/GeometryCollectionObject.h b/Engine/Source/Runtime/Experimental/GeometryCollectionEngine/Public/GeometryCollection/GeometryCollectionObject.h
--- a/Engine/Source/Runtime/Experimental/GeometryCollectionEngine/Public/GeometryCollection/GeometryCollectionObject.h
+++ b/Engine/Source/Runtime/Experimental/GeometryCollectionEngine/Public/GeometryCollection/GeometryCollectionObject.h
@@ -540,12 +540,12 @@ public:
 	bool bRemoveOnMaxSleep;
 	
 	/** How long may the particle sleep before initiating removal (in seconds). */
-	UPROPERTY(EditAnywhere, BlueprintReadOnly, Category = Removal, meta = (ClampMin = 0.001))
-	float MaximumSleepTime;
+	UPROPERTY(EditAnywhere, BlueprintReadOnly, Category = Removal)
+	FVector2D MaximumSleepTime;
 
 	/** How long does the removal process take (in seconds). */
-	UPROPERTY(EditAnywhere, BlueprintReadOnly, Category = Removal, meta = (ClampMin = 0.0))
-	float RemovalDuration;
+	UPROPERTY(EditAnywhere, BlueprintReadOnly, Category = Removal)
+	FVector2D RemovalDuration;
 
 	/*
 	* Size Specfic Data reflects the default geometry to bind to rigid bodies smaller
