Fix dangling ptr and API that leads to it

diff --git a/Engine/Source/Runtime/Experimental/Chaos/Private/RewindData.cpp b/Engine/Source/Runtime/Experimental/Chaos/Private/RewindData.cpp
--- a/Engine/Source/Runtime/Experimental/Chaos/Private/RewindData.cpp
+++ b/Engine/Source/Runtime/Experimental/Chaos/Private/RewindData.cpp
@@ -317,17 +317,16 @@ void FRewindData::AdvanceFrameImp(IResimCacheBase* ResimCache)
 	}
 
 	const FFrameAndPhase FrameAndPhase{ CurFrame, FFrameAndPhase::PostCallbacks };
-	FDirtyParticleInfo* DirtyParticlesDense = DirtyParticles.GetDensePtr();
 	for (int32 DirtyIdx = DirtyParticles.Num() - 1; DirtyIdx >= 0; --DirtyIdx)
 	{
-		FDirtyParticleInfo& Info = DirtyParticlesDense[DirtyIdx];
+		FDirtyParticleInfo& Info = DirtyParticles.GetDenseAt(DirtyIdx);
 
 		ensure(IsResimAndInSync(*Info.GetObjectPtr()) || Info.GetHistory().IsClean(FrameAndPhase));  //Sim hasn't run yet so PostCallbacks (sim results) should be clean
 
 		//if hasn't changed in a while stop tracking
 		if (Info.LastDirtyFrame < EarliestFrame)
 		{
-			RemoveParticle(DirtyParticlesDense[DirtyIdx].GetObjectPtr());
+			RemoveParticle(Info.GetObjectPtr());
 		}
 		else
 		{
diff --git a/Engine/Source/Runtime/Experimental/Chaos/Public/RewindData.h b/Engine/Source/Runtime/Experimental/Chaos/Public/RewindData.h
--- a/Engine/Source/Runtime/Experimental/Chaos/Public/RewindData.h
+++ b/Engine/Source/Runtime/Experimental/Chaos/Public/RewindData.h
@@ -851,8 +851,8 @@ public:
 	auto cbegin() const { return DenseVals.begin(); }
 	auto cend() const { return DenseVals.end(); }
 
-	const TVal* GetDensePtr() const { return DenseVals.GetData(); }
-	TVal* GetDensePtr() { return DenseVals.GetData(); }
+	const TVal& GetDenseAt(const int32 Idx) const { return DenseVals[Idx]; }
+	TVal& GetDenseAt(const int32 Idx) { return DenseVals[Idx]; }
 
 private:
 	TMap<TKey, int32> KeyToIdx;
@@ -1058,7 +1058,7 @@ private:
 			return *Info;
 		}
 
-		using TDirtyObj = decltype(NoRefHelper(*DirtyObjs.GetDensePtr()));
+		using TDirtyObj = decltype(NoRefHelper(DirtyObjs.GetDenseAt(0)));
 		TDirtyObj& Info = DirtyObjs.Add(&Handle, TDirtyObj(PropertiesPool, Handle, CurFrame, Managers.Capacity()));
 		Info.InitializedOnStep = InitializedOnFrame;
 		return Info;
