Create inner joint properties so we can easily replicate data to multiple structures.
Added an external and internal API for joints so we can write to them in sim callbacks
Refactor advance and finish in rewind data so joints are properly marked as desync etc...

diff --git a/Engine/Source/Runtime/Experimental/Chaos/Private/Chaos/PBDJointConstraints.cpp b/Engine/Source/Runtime/Experimental/Chaos/Private/Chaos/PBDJointConstraints.cpp
--- a/Engine/Source/Runtime/Experimental/Chaos/Private/Chaos/PBDJointConstraints.cpp
+++ b/Engine/Source/Runtime/Experimental/Chaos/Private/Chaos/PBDJointConstraints.cpp
@@ -104,6 +104,26 @@ namespace Chaos
 		return ConcreteContainer()->GetConstraintAngularImpulse(ConstraintIndex);
 	}
 
+	ESyncState FPBDJointConstraintHandle::SyncState() const
+	{
+		return ConcreteContainer()->GetConstraintSyncState(ConstraintIndex);
+	}
+
+	void FPBDJointConstraintHandle::SetSyncState(ESyncState SyncState)
+	{
+		return ConcreteContainer()->SetConstraintSyncState(ConstraintIndex, SyncState);
+	}
+
+	void FPBDJointConstraintHandle::SetEnabledDuringResim(bool bEnabled)
+	{
+		return ConcreteContainer()->SetConstraintEnabledDuringResim(ConstraintIndex, bEnabled);
+	}
+
+	EResimType FPBDJointConstraintHandle::ResimType() const
+	{
+		return ConcreteContainer()->GetConstraintResimType(ConstraintIndex);
+	}
+
 	
 	const FPBDJointSettings& FPBDJointConstraintHandle::GetSettings() const
 	{
@@ -697,6 +717,26 @@ namespace Chaos
 		return ConstraintStates[ConstraintIndex].AngularImpulse;
 	}
 
+	ESyncState FPBDJointConstraints::GetConstraintSyncState(int32 ConstraintIndex) const
+	{
+		return ConstraintStates[ConstraintIndex].SyncState;
+	}
+
+	void FPBDJointConstraints::SetConstraintSyncState(int32 ConstraintIndex, ESyncState SyncState)
+	{
+		ConstraintStates[ConstraintIndex].SyncState = SyncState;
+	}
+
+	void FPBDJointConstraints::SetConstraintEnabledDuringResim(int32 ConstraintIndex, bool bEnabled)
+	{
+		ConstraintStates[ConstraintIndex].bEnabledDuringResim = bEnabled;
+	}
+
+	EResimType FPBDJointConstraints::GetConstraintResimType(int32 ConstraintIndex) const
+	{
+		return ConstraintStates[ConstraintIndex].ResimType;
+	}
+
 
 	void FPBDJointConstraints::UpdatePositionBasedState(const FReal Dt)
 	{
diff --git a/Engine/Source/Runtime/Experimental/Chaos/Private/PBDRigidsSolver.cpp b/Engine/Source/Runtime/Experimental/Chaos/Private/PBDRigidsSolver.cpp
--- a/Engine/Source/Runtime/Experimental/Chaos/Private/PBDRigidsSolver.cpp
+++ b/Engine/Source/Runtime/Experimental/Chaos/Private/PBDRigidsSolver.cpp
@@ -602,7 +602,7 @@ namespace Chaos
 				//For now we assume this always comes from server update
 				if(FRewindData* RewindData = GetRewindData())
 				{
-					RewindData->RemoveJoint(JointProxy->GetHandle());
+					RewindData->RemoveObject(JointProxy->GetHandle());
 				}
 
 				JointProxy->DestroyOnPhysicsThread(this);
@@ -717,7 +717,7 @@ namespace Chaos
 					// Remove from rewind data
 					if (FRewindData* RewindData = GetRewindData())
 					{
-						RewindData->RemoveParticle(Info.Handle);
+						RewindData->RemoveObject(Info.Handle);
 					}
 
 					// Use the handle to destroy the particle data
diff --git a/Engine/Source/Runtime/Experimental/Chaos/Private/RewindData.cpp b/Engine/Source/Runtime/Experimental/Chaos/Private/RewindData.cpp
--- a/Engine/Source/Runtime/Experimental/Chaos/Private/RewindData.cpp
+++ b/Engine/Source/Runtime/Experimental/Chaos/Private/RewindData.cpp
@@ -125,6 +125,17 @@ bool FGeometryParticleStateBase::IsInSync(const FGeometryParticleHandle& Handle,
 	return true;
 }
 
+template <bool bSkipDynamics>
+bool FJointStateBase::IsInSync(const FPBDJointConstraintHandle& Handle, const FFrameAndPhase FrameAndPhase, const FDirtyPropertiesPool& Pool) const
+{
+	if (!JointSettings.IsInSync(Handle, FrameAndPhase, Pool))
+	{
+		return false;
+	}
+
+	return true;
+}
+
 bool FRewindData::RewindToFrame(int32 Frame)
 {
 	QUICK_SCOPE_CYCLE_COUNTER(RewindToFrame);
@@ -147,24 +158,26 @@ bool FRewindData::RewindToFrame(int32 Frame)
 	FFrameAndPhase RewindFrameAndPhase{ Frame, FFrameAndPhase::PostPushData };
 	FFrameAndPhase CurFrameAndPhase{ CurFrame, FFrameAndPhase::PrePushData };
 
-	//If we're rewinding a particle that doesn't need to save head (resim as slave never checks for desync so we don't care about head)
-	auto RewindNoSave = [RewindFrameAndPhase, this](auto Particle, auto& Property, const auto& RewindFunc)
+	auto RewindHelper = [RewindFrameAndPhase, CurFrameAndPhase, this](auto Obj, bool bResimAsSlave, auto& Property, const auto& RewindFunc) -> bool
 	{
-		if (auto Val = Property.Read(RewindFrameAndPhase, PropertiesPool))
+		if (bResimAsSlave)
 		{
-			RewindFunc(Particle, *Val);
+			//If we're rewinding a particle that doesn't need to save head (resim as slave never checks for desync so we don't care about head)
+			if (auto Val = Property.Read(RewindFrameAndPhase, PropertiesPool))
+			{
+				RewindFunc(Obj, *Val);
+			}
 		}
-	};
-
-	//If we're rewinding a particle that needs to save head (during resim when we get back to latest frame and phase we need to check for desync)
-	auto RewindAndSave = [RewindFrameAndPhase, CurFrameAndPhase, this](auto Particle, auto& Property, const auto& RewindFunc) -> bool
-	{
-		if (!Property.IsClean(RewindFrameAndPhase))
+		else
 		{
-			Property.WriteAccessMonotonic(CurFrameAndPhase, PropertiesPool).CopyFrom(*Particle);
-			RewindFunc(Particle, *Property.Read(RewindFrameAndPhase, PropertiesPool));
+			//If we're rewinding an object that needs to save head (during resim when we get back to latest frame and phase we need to check for desync)
+			if (!Property.IsClean(RewindFrameAndPhase))
+			{
+				CopyDataFromObject(Property.WriteAccessMonotonic(CurFrameAndPhase, PropertiesPool), *Obj);
+				RewindFunc(Obj, *Property.Read(RewindFrameAndPhase, PropertiesPool));
 
-			return true;
+				return true;
+			}
 		}
 
 		return false;
@@ -176,33 +189,20 @@ bool FRewindData::RewindToFrame(int32 Frame)
 
 		//rewind is about to start, all particles should be in sync at this point
 		ensure(PTParticle->SyncState() == ESyncState::InSync);
-
-		if(DirtyParticleInfo.bResimAsSlave)
-		{
-			//simple rewind with all data coming back exactly the same from gt
-			//this means we don't need to save head or do anything special
-			const FGeometryParticleStateBase& History = DirtyParticleInfo.GetHistory();
-
-			RewindNoSave(PTParticle, History.ParticlePositionRotation, [](auto Particle, const auto& Data) {Particle->SetXR(Data); });
-			RewindNoSave(PTParticle, History.NonFrequentData, [](auto Particle, const auto& Data) {Particle->SetNonFrequentData(Data); });
-			RewindNoSave(PTParticle->CastToKinematicParticle(), History.Velocities, [](auto Particle, const auto& Data) {Particle->SetVelocities(Data); });
-			RewindNoSave(PTParticle->CastToKinematicParticle(), History.KinematicTarget, [](auto Particle, const auto& Data) {Particle->SetKinematicTarget(Data); });
-			RewindNoSave(PTParticle->CastToRigidParticle(), History.Dynamics, [](auto Particle, const auto& Data) {Particle->SetDynamics(Data); });
-			RewindNoSave(PTParticle->CastToRigidParticle(), History.DynamicsMisc, [Evolution = Solver->GetEvolution()](auto Particle, const auto& Data) {Particle->SetDynamicMisc(Data, *Evolution); });
-			RewindNoSave(PTParticle->CastToRigidParticle(), History.MassProps, [](auto Particle, const auto& Data) {Particle->SetMassProps(Data); });
-		}
-		else
+		
+		FGeometryParticleStateBase& History = DirtyParticleInfo.AddFrame(CurFrame);	//non-const in case we need to record what's at head for a rewind (CurFrame has already been increased to the next frame)
+
+		const bool bResimAsSlave = DirtyParticleInfo.bResimAsSlave;
+		bool bAnyChange = RewindHelper(PTParticle, bResimAsSlave, History.ParticlePositionRotation, [](auto Particle, const auto& Data) {Particle->SetXR(Data); });
+		bAnyChange |= RewindHelper(PTParticle, bResimAsSlave, History.NonFrequentData, [](auto Particle, const auto& Data) {Particle->SetNonFrequentData(Data); });
+		bAnyChange |= RewindHelper(PTParticle->CastToKinematicParticle(), bResimAsSlave, History.Velocities, [](auto Particle, const auto& Data) {Particle->SetVelocities(Data); });
+		bAnyChange |= RewindHelper(PTParticle->CastToKinematicParticle(), bResimAsSlave, History.KinematicTarget, [](auto Particle, const auto& Data) {Particle->SetKinematicTarget(Data); });
+		bAnyChange |= RewindHelper(PTParticle->CastToRigidParticle(),bResimAsSlave,  History.Dynamics, [](auto Particle, const auto& Data) {Particle->SetDynamics(Data); });
+		bAnyChange |= RewindHelper(PTParticle->CastToRigidParticle(),bResimAsSlave,  History.DynamicsMisc, [Evolution = Solver->GetEvolution()](auto Particle, const auto& Data) {Particle->SetDynamicMisc(Data, *Evolution); });
+		bAnyChange |= RewindHelper(PTParticle->CastToRigidParticle(),bResimAsSlave,  History.MassProps, [](auto Particle, const auto& Data) {Particle->SetMassProps(Data); });
+
+		if (!bResimAsSlave)
 		{
-			FGeometryParticleStateBase& History = DirtyParticleInfo.AddFrame(CurFrame);	//non-const in case we need to record what's at head for a rewind (CurFrame has already been increased to the next frame)
-
-			bool bAnyChange = RewindAndSave(PTParticle, History.ParticlePositionRotation, [](auto Particle, const auto& Data) {Particle->SetXR(Data); });
-			bAnyChange |= RewindAndSave(PTParticle, History.NonFrequentData, [](auto Particle, const auto& Data) {Particle->SetNonFrequentData(Data); });
-			bAnyChange |= RewindAndSave(PTParticle->CastToKinematicParticle(), History.Velocities, [](auto Particle, const auto& Data) {Particle->SetVelocities(Data); });
-			bAnyChange |= RewindAndSave(PTParticle->CastToKinematicParticle(), History.KinematicTarget, [](auto Particle, const auto& Data) {Particle->SetKinematicTarget(Data); });
-			bAnyChange |= RewindAndSave(PTParticle->CastToRigidParticle(), History.Dynamics, [](auto Particle, const auto& Data) {Particle->SetDynamics(Data); });
-			bAnyChange |= RewindAndSave(PTParticle->CastToRigidParticle(), History.DynamicsMisc, [Evolution = Solver->GetEvolution()](auto Particle, const auto& Data) {Particle->SetDynamicMisc(Data, *Evolution); });
-			bAnyChange |= RewindAndSave(PTParticle->CastToRigidParticle(), History.MassProps, [](auto Particle, const auto& Data) {Particle->SetMassProps(Data); });
-
 			if (bAnyChange)
 			{
 				//particle actually changes not just created/streamed so need to update its state
@@ -216,15 +216,27 @@ bool FRewindData::RewindToFrame(int32 Frame)
 				DirtyParticleInfo.GetObjectPtr()->SetEnabledDuringResim(true);
 			}
 
-		if (DirtyParticleInfo.InitializedOnStep > Frame)
-		{
-			//hasn't initialized yet, so disable
-			//must do this after rewind because SetDynamicsMisc will re-enable
-			//(the disable is a temp way to ignore objects not spawned yet, they weren't really disabled which is why it gets re-enabled)
-			Solver->GetEvolution()->DisableParticle(DirtyParticleInfo.GetObjectPtr());
-		}
+			if (DirtyParticleInfo.InitializedOnStep > Frame)
+			{
+				//hasn't initialized yet, so disable
+				//must do this after rewind because SetDynamicsMisc will re-enable
+				//(the disable is a temp way to ignore objects not spawned yet, they weren't really disabled which is why it gets re-enabled)
+				Solver->GetEvolution()->DisableParticle(DirtyParticleInfo.GetObjectPtr());
+			}
 		}
+	}
+
+	for(FDirtyJointInfo& DirtyJointInfo : DirtyJoints)
+	{
+		FPBDJointConstraintHandle* Joint = DirtyJointInfo.GetObjectPtr();
+
+		//rewind is about to start, all particles should be in sync at this point
+		ensure(Joint->SyncState() == ESyncState::InSync);
 
+		FJointStateBase& History = DirtyJointInfo.AddFrame(CurFrame);	//non-const in case we need to record what's at head for a rewind (CurFrame has already been increased to the next frame)
+
+		const bool bResimAsSlave = DirtyJointInfo.bResimAsSlave;
+		bool bAnyChange = RewindHelper(Joint, bResimAsSlave, History.JointSettings, [](auto Joint, const auto& Data) {Joint->SetSettings(Data); });
 	}
 
 	CurFrame = Frame;
@@ -234,6 +246,23 @@ bool FRewindData::RewindToFrame(int32 Frame)
 	return true;
 }
 
+template <bool bSkipDynamics, typename TDirtyInfo>
+void FRewindData::DesyncIfNecessary(TDirtyInfo& Info, const FFrameAndPhase FrameAndPhase)
+{
+	ensure(IsResim());	//shouldn't bother with desync unless we're resimming
+
+	auto Handle = Info.GetObjectPtr();
+	const auto& History = Info.GetHistory();
+	if (Handle->SyncState() == ESyncState::InSync && !History.template IsInSync<bSkipDynamics>(*Handle, FrameAndPhase, PropertiesPool))
+	{
+		if (!SkipDesyncTest)
+		{
+			//first time desyncing so need to clear history from this point into the future
+			DesyncObject(Info, FrameAndPhase);
+		}
+	}
+}
+
 void FRewindData::FinishFrame()
 {
 	QUICK_SCOPE_CYCLE_COUNTER(RewindDataFinishFrame);
@@ -242,18 +271,18 @@ void FRewindData::FinishFrame()
 	{
 		FFrameAndPhase FutureFrame{ CurFrame + 1, FFrameAndPhase::PrePushData };
 
-		for(FDirtyParticleInfo& Info : DirtyParticles)
+		auto FinishHelper = [this, FutureFrame](auto& DirtyObjs)
 		{
-			if (Info.bResimAsSlave)
+			for (auto& Info : DirtyObjs)
 			{
-				//resim as slave means always in sync and no cleanup needed
-				continue;
-			}
+				if (Info.bResimAsSlave)
+				{
+					//resim as slave means always in sync and no cleanup needed
+					continue;
+				}
 
-			auto& Handle = *Info.GetObjectPtr();
-			if (auto Rigid = Handle.CastToRigidParticle())
-			{
-				if(Rigid->ResimType() == EResimType::FullResim)
+				auto& Handle = *Info.GetObjectPtr();
+				if (Handle.ResimType() == EResimType::FullResim)
 				{
 					if (IsFinalResim())
 					{
@@ -272,7 +301,10 @@ void FRewindData::FinishFrame()
 					}
 				}
 			}
-		}
+		};
+
+		FinishHelper(DirtyParticles);
+		FinishHelper(DirtyJoints);
 	}
 	
 
@@ -280,72 +312,66 @@ void FRewindData::FinishFrame()
 	LatestFrame = FMath::Max(LatestFrame,CurFrame);
 }
 
-int32 SkipDesyncTest = 0;
+CHAOS_API int32 SkipDesyncTest = 0;
 FAutoConsoleVariableRef CVarSkipDesyncTest(TEXT("p.SkipDesyncTest"), SkipDesyncTest, TEXT("Skips hard desync test, this means all particles will assume to be clean except spawning at different times. This is useful for a perf lower bound, not actually correct"));
 
-void FRewindData::DesyncParticle(FDirtyParticleInfo& Info, const FFrameAndPhase FrameAndPhase)
-{
-	Info.ClearPhaseAndFuture(FrameAndPhase);
-	Info.GetObjectPtr()->SetSyncState(ESyncState::HardDesync);
-}
-
-template <bool bSkipDynamics>
-void FRewindData::DesyncIfNecessary(FDirtyParticleInfo& Info, const FFrameAndPhase FrameAndPhase)
-{
-	ensure(IsResim());	//shouldn't bother with desync unless we're resimming
-
-	auto Handle = Info.GetObjectPtr();
-	if (Handle->SyncState() == ESyncState::InSync && !Info.GetHistory().IsInSync<bSkipDynamics>(*Handle, FrameAndPhase, PropertiesPool))
-	{
-		if (!SkipDesyncTest)
-		{
-			//first time desyncing so need to clear history from this point into the future
-			DesyncParticle(Info, FrameAndPhase);
-		}
-	}
-}
-
 void FRewindData::AdvanceFrameImp(IResimCacheBase* ResimCache)
 {
 	FramesSaved = FMath::Min(FramesSaved+1,static_cast<int32>(Managers.Capacity()-1));
 
 	const int32 EarliestFrame = CurFrame - 1 - FramesSaved;
-	TArray<FGeometryParticleHandle*> DesyncedParticles;
-	if(IsResim())
-	{
-		DesyncedParticles.Reserve(DirtyParticles.Num());
-	}
-
 	const FFrameAndPhase FrameAndPhase{ CurFrame, FFrameAndPhase::PostCallbacks };
-	for (int32 DirtyIdx = DirtyParticles.Num() - 1; DirtyIdx >= 0; --DirtyIdx)
-	{
-		FDirtyParticleInfo& Info = DirtyParticles.GetDenseAt(DirtyIdx);
-
-		ensure(IsResimAndInSync(*Info.GetObjectPtr()) || Info.GetHistory().IsClean(FrameAndPhase));  //Sim hasn't run yet so PostCallbacks (sim results) should be clean
 
-		//if hasn't changed in a while stop tracking
-		if (Info.LastDirtyFrame < EarliestFrame)
-		{
-			RemoveParticle(Info.GetObjectPtr());
-		}
-		else
+	auto AdvanceHelper = [this, EarliestFrame, FrameAndPhase](auto& DirtyObjects, const auto& DesyncFunc, const auto& AdvanceDirtyFunc)
+	{
+		for (int32 DirtyIdx = DirtyObjects.Num() - 1; DirtyIdx >= 0; --DirtyIdx)
 		{
+			auto& Info = DirtyObjects.GetDenseAt(DirtyIdx);
 
-			auto Handle = Info.GetObjectPtr();
-			Info.bResimAsSlave = Handle->ResimType() == EResimType::ResimAsSlave;
+			ensure(IsResimAndInSync(*Info.GetObjectPtr()) || Info.GetHistory().IsClean(FrameAndPhase));  //Sim hasn't run yet so PostCallbacks (sim results) should be clean
 
-			if (IsResim() && !Info.bResimAsSlave)
+			//if hasn't changed in a while stop tracking
+			if (Info.LastDirtyFrame < EarliestFrame)
 			{
-				DesyncIfNecessary(Info, FrameAndPhase);
+				RemoveObject(Info.GetObjectPtr());
 			}
-
-			if(IsResim() && Handle->SyncState() != ESyncState::InSync && !SkipDesyncTest)
+			else
 			{
-				Handle->SetEnabledDuringResim(true);	//for now just mark anything out of sync as resim enabled. TODO: use bubble
-				DesyncedParticles.Add(Handle);
+
+				auto Handle = Info.GetObjectPtr();
+				Info.bResimAsSlave = Handle->ResimType() == EResimType::ResimAsSlave;
+
+				if (IsResim() && !Info.bResimAsSlave)
+				{
+					DesyncIfNecessary</*bSkipDynamics=*/false>(Info, FrameAndPhase);
+				}
+
+				if (IsResim() && Handle->SyncState() != ESyncState::InSync && !SkipDesyncTest)
+				{
+					Handle->SetEnabledDuringResim(true);	//for now just mark anything out of sync as resim enabled. TODO: use bubble
+					DesyncFunc(Handle);
+				}
+
+				AdvanceDirtyFunc(Info, Handle);
 			}
+		}
+	};
+
+	TArray<FGeometryParticleHandle*> DesyncedParticles;
+	if (IsResim())
+	{
+		DesyncedParticles.Reserve(DirtyParticles.Num());
+	}
 
-			if(Info.DirtyDynamics == CurFrame && !IsResimAndInSync(*Handle))
+	AdvanceHelper(DirtyParticles,
+		[&DesyncedParticles](FGeometryParticleHandle* DesyncedHandle)
+		{
+			DesyncedParticles.Add(DesyncedHandle);
+		},
+
+		[this, FrameAndPhase](FDirtyParticleInfo& Info, FGeometryParticleHandle* Handle)
+		{
+			if (Info.DirtyDynamics == CurFrame && !IsResimAndInSync(*Handle))
 			{
 				//we only need to check the cast because right now there's no property system on PT, so any time a sim callback touches a particle we just mark it as dirty dynamics
 				if (auto Rigid = Handle->CastToRigidParticle())
@@ -355,8 +381,12 @@ void FRewindData::AdvanceFrameImp(IResimCacheBase* ResimCache)
 					Latest.Dynamics.WriteAccessMonotonic(FrameAndPhase, PropertiesPool).CopyFrom(*Rigid);
 				}
 			}
-		}
-	}
+		});
+
+	AdvanceHelper(DirtyJoints, [](const FPBDJointConstraintHandle*) {}, [](const FDirtyJointInfo&, const FPBDJointConstraintHandle*) {});
+
+	//TODO: if joint is desynced we should desync particles as well
+	//If particle of joint is desynced, we need to make sure the joint is reconsidered too for optimization, though maybe not "desynced"
 
 	if(IsResim() && ResimCache)
 	{
@@ -368,19 +398,6 @@ void FRewindData::AdvanceFrameImp(IResimCacheBase* ResimCache)
 #define REWIND_DESYNC 0
 #endif
 
-template <typename TData, typename TObj>
-void CopyDataFromObject(TData& Data, const TObj& Obj)
-{
-	Data.CopyFrom(Obj);
-}
-
-void CopyDataFromObject(FPBDJointSettings& Data, const FPBDJointConstraintHandle& Joint)
-{
-	Data = Joint.GetSettings();
-}
-
-
-
 void FRewindData::PushGTDirtyData(const FDirtyPropertiesManager& SrcManager,const int32 SrcDataIdx,const FDirtyProxy& Dirty, const FShapeDirtyData* ShapeDirtyData)
 {
 	//This records changes enqueued by GT.
@@ -513,7 +530,7 @@ void FRewindData::SpawnProxyIfNeeded(FSingleParticlePhysicsProxy& Proxy)
 		Solver->GetEvolution()->EnableParticle(Handle, nullptr);
 		if(Proxy.GetInitializedStep() != CurFrame)
 		{
-			DesyncParticle(Info, FFrameAndPhase{ Proxy.GetInitializedStep(), FFrameAndPhase::PrePushData });	//Spawned earlier so mark as desynced from that first frame
+			DesyncObject(Info, FFrameAndPhase{ Proxy.GetInitializedStep(), FFrameAndPhase::PrePushData });	//Spawned earlier so mark as desynced from that first frame
 			Proxy.SetInitialized(CurFrame);
 		}
 	}
@@ -568,6 +585,27 @@ void FRewindData::MarkDirtyFromPT(FGeometryParticleHandle& Handle)
 	}
 }
 
+void FRewindData::MarkDirtyJointFromPT(FPBDJointConstraintHandle& Handle)
+{
+	FDirtyJointInfo& Info = FindOrAddDirtyObj(Handle);
+
+	const bool bRecordingHistory = !IsResimAndInSync(Handle);
+
+	FJointStateBase& Latest = Info.AddFrame(CurFrame);
+
+	//TODO: use property system
+
+	const FFrameAndPhase FrameAndPhase{ CurFrame, FFrameAndPhase::PostPushData };
+
+	if (bRecordingHistory || Latest.JointSettings.IsClean(FrameAndPhase))
+	{
+		if (auto Data = Latest.JointSettings.WriteAccessNonDecreasing(FFrameAndPhase{ CurFrame, FFrameAndPhase::PostPushData }, PropertiesPool))
+		{
+			CopyDataFromObject(*Data, Handle);
+		}
+	}
+}
+
 void FRewindData::PushPTDirtyData(TPBDRigidParticleHandle<FReal,3>& Handle,const int32 SrcDataIdx)
 	{
 	const bool bRecordingHistory = !IsResimAndInSync(Handle);
diff --git a/Engine/Source/Runtime/Experimental/Chaos/Public/Chaos/JointProperties.inl b/Engine/Source/Runtime/Experimental/Chaos/Public/Chaos/JointProperties.inl
new file mode 100644
--- /dev/null
+++ b/Engine/Source/Runtime/Experimental/Chaos/Public/Chaos/JointProperties.inl
@@ -0,0 +1,66 @@
+// Copyright Epic Games, Inc. All Rights Reserved.
+
+//CHAOS_INNER_JOINT_PROPERTY(OuterProp, FuncName, InnerName, InnerType)							
+CHAOS_INNER_JOINT_PROPERTY(JointSettings, JointTransforms, ConnectorTransforms, FTransformPair)
+CHAOS_INNER_JOINT_PROPERTY(JointSettings, CollisionEnabled, bCollisionEnabled, bool)
+CHAOS_INNER_JOINT_PROPERTY(JointSettings, ProjectionEnabled, bProjectionEnabled, bool)
+CHAOS_INNER_JOINT_PROPERTY(JointSettings, ProjectionLinearAlpha, LinearProjection, FReal)
+CHAOS_INNER_JOINT_PROPERTY(JointSettings, ProjectionAngularAlpha, AngularProjection, FReal)
+CHAOS_INNER_JOINT_PROPERTY(JointSettings, ParentInvMassScale, ParentInvMassScale, FReal)
+CHAOS_INNER_JOINT_PROPERTY(JointSettings, LinearBreakForce, LinearBreakForce, FReal)
+CHAOS_INNER_JOINT_PROPERTY(JointSettings, LinearPlasticityLimit, LinearPlasticityLimit, FReal)
+CHAOS_INNER_JOINT_PROPERTY(JointSettings, LinearPlasticityType, LinearPlasticityType, EPlasticityType)
+CHAOS_INNER_JOINT_PROPERTY(JointSettings, AngularBreakTorque, AngularBreakTorque, FReal)
+CHAOS_INNER_JOINT_PROPERTY(JointSettings, AngularPlasticityLimit, AngularPlasticityLimit, FReal)
+CHAOS_INNER_JOINT_PROPERTY(JointSettings, UserData, UserData, void*)
+CHAOS_INNER_JOINT_PROPERTY(JointSettings, LinearPositionDriveXEnabled, bLinearPositionDriveEnabled[0], bool)
+CHAOS_INNER_JOINT_PROPERTY(JointSettings, LinearPositionDriveYEnabled, bLinearPositionDriveEnabled[1], bool)
+CHAOS_INNER_JOINT_PROPERTY(JointSettings, LinearPositionDriveZEnabled, bLinearPositionDriveEnabled[2], bool)
+CHAOS_INNER_JOINT_PROPERTY(JointSettings, LinearDrivePositionTarget, LinearDrivePositionTarget, FVec3)
+CHAOS_INNER_JOINT_PROPERTY(JointSettings, LinearVelocityDriveXEnabled, bLinearVelocityDriveEnabled[0], bool)
+CHAOS_INNER_JOINT_PROPERTY(JointSettings, LinearVelocityDriveYEnabled, bLinearVelocityDriveEnabled[1], bool)
+CHAOS_INNER_JOINT_PROPERTY(JointSettings, LinearVelocityDriveZEnabled, bLinearVelocityDriveEnabled[2], bool)
+CHAOS_INNER_JOINT_PROPERTY(JointSettings, LinearDriveVelocityTarget, LinearDriveVelocityTarget, FVec3)
+CHAOS_INNER_JOINT_PROPERTY(JointSettings, LinearDriveForceMode, LinearDriveForceMode, EJointForceMode)
+CHAOS_INNER_JOINT_PROPERTY(JointSettings, LinearMotionTypesX, LinearMotionTypes[0], EJointMotionType)
+CHAOS_INNER_JOINT_PROPERTY(JointSettings, LinearMotionTypesY, LinearMotionTypes[1], EJointMotionType)
+CHAOS_INNER_JOINT_PROPERTY(JointSettings, LinearMotionTypesZ, LinearMotionTypes[2], EJointMotionType)
+CHAOS_INNER_JOINT_PROPERTY(JointSettings, LinearDriveStiffness, LinearDriveStiffness, FReal)
+CHAOS_INNER_JOINT_PROPERTY(JointSettings, LinearDriveDamping, LinearDriveDamping, FReal)
+CHAOS_INNER_JOINT_PROPERTY(JointSettings, ContactTransferScale, ContactTransferScale, FReal)
+CHAOS_INNER_JOINT_PROPERTY(JointSettings, AngularSLerpPositionDriveEnabled, bAngularSLerpPositionDriveEnabled, bool)
+CHAOS_INNER_JOINT_PROPERTY(JointSettings, AngularTwistPositionDriveEnabled, bAngularTwistPositionDriveEnabled, bool)
+CHAOS_INNER_JOINT_PROPERTY(JointSettings, AngularSwingPositionDriveEnabled, bAngularSwingPositionDriveEnabled, bool)
+CHAOS_INNER_JOINT_PROPERTY(JointSettings, AngularDrivePositionTarget, AngularDrivePositionTarget, FRotation3)
+CHAOS_INNER_JOINT_PROPERTY(JointSettings, AngularSLerpVelocityDriveEnabled, bAngularSLerpVelocityDriveEnabled, bool)
+CHAOS_INNER_JOINT_PROPERTY(JointSettings, AngularTwistVelocityDriveEnabled, bAngularTwistVelocityDriveEnabled, bool)
+CHAOS_INNER_JOINT_PROPERTY(JointSettings, AngularSwingVelocityDriveEnabled, bAngularSwingVelocityDriveEnabled, bool)
+CHAOS_INNER_JOINT_PROPERTY(JointSettings, AngularDriveVelocityTarget, AngularDriveVelocityTarget, FVec3)
+CHAOS_INNER_JOINT_PROPERTY(JointSettings, AngularDriveForceMode, AngularDriveForceMode, EJointForceMode)
+CHAOS_INNER_JOINT_PROPERTY(JointSettings, AngularMotionTypesX, AngularMotionTypes[0], EJointMotionType)
+CHAOS_INNER_JOINT_PROPERTY(JointSettings, AngularMotionTypesY, AngularMotionTypes[1], EJointMotionType)
+CHAOS_INNER_JOINT_PROPERTY(JointSettings, AngularMotionTypesZ, AngularMotionTypes[2], EJointMotionType)
+CHAOS_INNER_JOINT_PROPERTY(JointSettings, AngularDriveStiffness, AngularDriveStiffness, FReal)
+CHAOS_INNER_JOINT_PROPERTY(JointSettings, AngularDriveDamping, AngularDriveDamping, FReal)
+CHAOS_INNER_JOINT_PROPERTY(JointSettings, Stiffness, Stiffness, FReal)
+CHAOS_INNER_JOINT_PROPERTY(JointSettings, SoftLinearLimitsEnabled, bSoftLinearLimitsEnabled, bool)
+CHAOS_INNER_JOINT_PROPERTY(JointSettings, SoftTwistLimitsEnabled, bSoftTwistLimitsEnabled, bool)
+CHAOS_INNER_JOINT_PROPERTY(JointSettings, SoftSwingLimitsEnabled, bSoftSwingLimitsEnabled, bool)
+CHAOS_INNER_JOINT_PROPERTY(JointSettings, LinearSoftForceMode, LinearSoftForceMode, EJointForceMode)
+CHAOS_INNER_JOINT_PROPERTY(JointSettings, AngularSoftForceMode, AngularSoftForceMode, EJointForceMode)
+CHAOS_INNER_JOINT_PROPERTY(JointSettings, SoftLinearStiffness, SoftLinearStiffness, FReal)
+CHAOS_INNER_JOINT_PROPERTY(JointSettings, SoftLinearDamping, SoftLinearDamping, FReal)
+CHAOS_INNER_JOINT_PROPERTY(JointSettings, SoftTwistStiffness, SoftTwistStiffness, FReal)
+CHAOS_INNER_JOINT_PROPERTY(JointSettings, SoftTwistDamping, SoftTwistDamping, FReal)
+CHAOS_INNER_JOINT_PROPERTY(JointSettings, SoftSwingStiffness, SoftSwingStiffness, FReal)
+CHAOS_INNER_JOINT_PROPERTY(JointSettings, SoftSwingDamping, SoftSwingDamping, FReal)
+CHAOS_INNER_JOINT_PROPERTY(JointSettings, LinearLimit, LinearLimit, FReal)
+CHAOS_INNER_JOINT_PROPERTY(JointSettings, AngularLimits, AngularLimits, FVec3)
+CHAOS_INNER_JOINT_PROPERTY(JointSettings, LinearContactDistance, LinearContactDistance, FReal)
+CHAOS_INNER_JOINT_PROPERTY(JointSettings, TwistContactDistance, TwistContactDistance, FReal)
+CHAOS_INNER_JOINT_PROPERTY(JointSettings, SwingContactDistance, SwingContactDistance, FReal)
+CHAOS_INNER_JOINT_PROPERTY(JointSettings, LinearRestitution, LinearRestitution, FReal)
+CHAOS_INNER_JOINT_PROPERTY(JointSettings, TwistRestitution, TwistRestitution, FReal)
+CHAOS_INNER_JOINT_PROPERTY(JointSettings, SwingRestitution, SwingRestitution, FReal)
+
+#undef CHAOS_INNER_JOINT_PROPERTY
\ No newline at end of file
diff --git a/Engine/Source/Runtime/Experimental/Chaos/Public/Chaos/PBDJointConstraintData.h b/Engine/Source/Runtime/Experimental/Chaos/Public/Chaos/PBDJointConstraintData.h
--- a/Engine/Source/Runtime/Experimental/Chaos/Public/Chaos/PBDJointConstraintData.h
+++ b/Engine/Source/Runtime/Experimental/Chaos/Public/Chaos/PBDJointConstraintData.h
@@ -31,12 +31,12 @@ namespace Chaos
 
 
 	using FJointConstraintDirtyFlags = TDirtyFlags<EJointConstraintFlags>;
+	typedef TVector<FTransform, 2> FTransformPair;
 
 	class CHAOS_API FJointConstraint : public FConstraintBase
 	{
 	public:
 		using Base = FConstraintBase;
-		typedef TVector<FTransform, 2> FTransformPair;
 
 		friend class FPBDRigidsSolver; // friend so we can call ReleaseKinematicEndPoint when unregistering joint.
 
@@ -49,254 +49,16 @@ namespace Chaos
 		void SetKinematicEndPoint(FSingleParticlePhysicsProxy* InParticle, FPBDRigidsSolver* Solver);
 
 		FSingleParticlePhysicsProxy* GetKinematicEndPoint() const;
-		CONSTRAINT_JOINT_PROPERPETY_IMPL2(FTransformPair, JointTransforms, JointSettings, ConnectorTransforms);
-		//void SetJointTransforms(const Chaos::FJointConstraint::FTransformPair& InJointTransforms);
-		//const FTransformPair GetJointTransforms() const;
 
-		CONSTRAINT_JOINT_PROPERPETY_IMPL2(bool, CollisionEnabled, JointSettings, bCollisionEnabled);
-		//void SetCollisionEnabled(bool InValue);
-		//bool GetCollisionEnabled() const;
-
-		CONSTRAINT_JOINT_PROPERPETY_IMPL2(bool, ProjectionEnabled, JointSettings, bProjectionEnabled);
-		//void SetProjectionEnabled(bool bInProjectionEnabled);
-		//bool GetProjectionEnabled() const;
-
-		CONSTRAINT_JOINT_PROPERPETY_IMPL2(FReal, ProjectionLinearAlpha, JointSettings, LinearProjection);
-		//void SetProjectionLinearAlpha(FReal InProjectionLinearAlpha);
-		//FReal GetProjectionLinearAlpha() const;
-
-		CONSTRAINT_JOINT_PROPERPETY_IMPL2(FReal, ProjectionAngularAlpha, JointSettings, AngularProjection);
-		//void SetProjectionAngularAlpha(FReal InProjectionAngularAlpha);
-		//FReal GetProjectionAngularAlpha() const;
-
-		CONSTRAINT_JOINT_PROPERPETY_IMPL2(FReal, ParentInvMassScale, JointSettings, ParentInvMassScale);
-		//void SetParentInvMassScale(FReal InParentInvMassScale);
-		//FReal GetParentInvMassScale() const
-
-		CONSTRAINT_JOINT_PROPERPETY_IMPL2(FReal, LinearBreakForce, JointSettings, LinearBreakForce);
-		//void SetLinearBreakForce(FReal InLinearBreakForce);
-		//FReal GetLinearBreakForce() const
-
-		CONSTRAINT_JOINT_PROPERPETY_IMPL2(FReal, LinearPlasticityLimit, JointSettings, LinearPlasticityLimit);
-		//void SetLinearPlasticityLimit(FReal InLinearPlasticityLimit);
-		//FReal GetLinearPlasticityLimit() const
-
-		CONSTRAINT_JOINT_PROPERPETY_IMPL2(EPlasticityType, LinearPlasticityType, JointSettings, LinearPlasticityType);
-		//void SetLinearPlasticityType(EPlasticityType InLinearPlasticityType);
-		//EPlasticityType GetLinearPlasticityType() const
-
-		CONSTRAINT_JOINT_PROPERPETY_IMPL2(FReal, AngularBreakTorque, JointSettings, AngularBreakTorque);
-		//void SetAngularBreakTorque(FReal InAngularBreakTorque);
-		//FReal GetAngularBreakTorque() const
-
-		CONSTRAINT_JOINT_PROPERPETY_IMPL2(FReal, AngularPlasticityLimit, JointSettings, AngularPlasticityLimit);
-		//void SetAngularPlasticityLimit(FReal InAngularPlasticityLimit);
-		//FReal GetAngularPlasticityLimit() const
-
-		CONSTRAINT_JOINT_PROPERPETY_IMPL2(void*, UserData, JointSettings, UserData);
-		//void SetUserData(void* InUserData);
-		//void* GetUserData() const
+		//See JointProperties for API
+		//Each CHAOS_INNER_JOINT_PROPERTY entry will have a Get* and Set*
+#define CHAOS_INNER_JOINT_PROPERTY(OuterProp, FuncName, Inner, InnerType) CONSTRAINT_JOINT_PROPERPETY_IMPL2(InnerType, FuncName, OuterProp, Inner)
+#include "Chaos/JointProperties.inl"
 
 		void SetLinearPositionDriveEnabled(TVector<bool, 3> Enabled);
 
-		CONSTRAINT_JOINT_PROPERPETY_IMPL2(bool, LinearPositionDriveXEnabled, JointSettings, bLinearPositionDriveEnabled[0]);
-		//void SetLinearPositionDriveXEnabled(bool InLinearPositionDriveXEnabled);
-		//bool GetLinearPositionDriveXEnabled() const
-
-		CONSTRAINT_JOINT_PROPERPETY_IMPL2(bool, LinearPositionDriveYEnabled, JointSettings, bLinearPositionDriveEnabled[1]);
-		//void SetLinearPositionDriveYEnabled(bool InLinearPositionDriveYEnabled);
-		//bool GetLinearPositionDriveYEnabled() const
-
-		CONSTRAINT_JOINT_PROPERPETY_IMPL2(bool, LinearPositionDriveZEnabled, JointSettings, bLinearPositionDriveEnabled[2]);
-		//void SetLinearPositionDriveZEnabled(bool InLinearPositionDriveZEnabled);
-		//bool GetLinearPositionDriveZEnabled() const
-
-		CONSTRAINT_JOINT_PROPERPETY_IMPL2(FVec3, LinearDrivePositionTarget, JointSettings, LinearDrivePositionTarget);
-		//void SetLinearDrivePositionTarget(FVec3 InLinearDrivePositionTarget);
-		//FVec3 GetLinearDrivePositionTarget() const
-
 		void SetLinearVelocityDriveEnabled(TVector<bool, 3> Enabled);
 
-		CONSTRAINT_JOINT_PROPERPETY_IMPL2(bool, LinearVelocityDriveXEnabled, JointSettings, bLinearVelocityDriveEnabled[0]);
-		//void SetLinearVelocityDriveXEnabled(bool InLinearVelocityDriveXEnabled);
-		//bool GetLinearVelocityDriveXEnabled() const
-
-		CONSTRAINT_JOINT_PROPERPETY_IMPL2(bool, LinearVelocityDriveYEnabled, JointSettings, bLinearVelocityDriveEnabled[1]);
-		//void SetLinearVelocityDriveYEnabled(bool InLinearVelocityDriveYEnabled);
-		//bool GetLinearVelocityDriveYEnabled() const
-
-		CONSTRAINT_JOINT_PROPERPETY_IMPL2(bool, LinearVelocityDriveZEnabled, JointSettings, bLinearVelocityDriveEnabled[2]);
-		//void SetLinearVelocityDriveZEnabled(bool InLinearVelocityDriveZEnabled);
-		//bool GetLinearVelocityDriveZEnabled() const
-
-		CONSTRAINT_JOINT_PROPERPETY_IMPL2(FVec3, LinearDriveVelocityTarget, JointSettings, LinearDriveVelocityTarget);
-		//void SetLinearDriveVelocityTarget(FVec3 InLinearDriveVelocityTarget);
-		//FVec3 GetLinearDriveVelocityTarget() const
-
-		CONSTRAINT_JOINT_PROPERPETY_IMPL2(EJointForceMode, LinearDriveForceMode, JointSettings, LinearDriveForceMode);
-		//void SetLinearDriveVelocityTarget(EJointForceMode InEJointForceMode);
-		//EJointForceMode GetLinearDriveVelocityTarget() const
-
-		CONSTRAINT_JOINT_PROPERPETY_IMPL2(EJointMotionType, LinearMotionTypesX, JointSettings, LinearMotionTypes[0]);
-		//void SetLinearMotionTypesX(EJointMotionType InLinearMotionTypesX);
-		//EJointMotionType GetLinearMotionTypesX() const
-
-		CONSTRAINT_JOINT_PROPERPETY_IMPL2(EJointMotionType, LinearMotionTypesY, JointSettings, LinearMotionTypes[1]);
-		//void SetLinearMotionTypesY(EJointMotionType InLinearMotionTypesY);
-		//EJointMotionType GetLinearMotionTypesY() const
-
-		CONSTRAINT_JOINT_PROPERPETY_IMPL2(EJointMotionType, LinearMotionTypesZ, JointSettings, LinearMotionTypes[2]);
-		//void SetLinearMotionTypesZ(EJointMotionType InLinearMotionTypesZ);
-		//EJointMotionType GetLinearMotionTypesZ() const
-
-		CONSTRAINT_JOINT_PROPERPETY_IMPL2(FReal, LinearDriveStiffness, JointSettings, LinearDriveStiffness);
-		//void SetLinearDriveStiffness(FReal InLinearDriveStffness);
-		//FReal GetLinearDriveStiffness() const
-
-		CONSTRAINT_JOINT_PROPERPETY_IMPL2(FReal, LinearDriveDamping, JointSettings, LinearDriveDamping);
-		//void SetLinearDriveDamping(FReal InLinearDriveStffness);
-		//FReal GetLinearDriveDamping() const
-
-		CONSTRAINT_JOINT_PROPERPETY_IMPL2(FReal, ContactTransferScale, JointSettings, ContactTransferScale);
-		//void SetContactTransferScale(FReal InContactTransferScale);
-		//FReal GetContactTransferScale() const
-
-		CONSTRAINT_JOINT_PROPERPETY_IMPL2(bool, AngularSLerpPositionDriveEnabled, JointSettings, bAngularSLerpPositionDriveEnabled);
-		//void SetAngularSLerpPositionDriveEnabled(bool InAngularSLerpPositionDriveEnabled);
-		//bool GetAngularSLerpPositionDriveEnabled() const
-
-		CONSTRAINT_JOINT_PROPERPETY_IMPL2(bool, AngularTwistPositionDriveEnabled, JointSettings, bAngularTwistPositionDriveEnabled);
-		//void SetAngularTwistPositionDriveEnabled(bool InAngularTwistPositionDriveEnabled);
-		//bool GetAngularTwistPositionDriveEnabled() const
-
-		CONSTRAINT_JOINT_PROPERPETY_IMPL2(bool, AngularSwingPositionDriveEnabled, JointSettings, bAngularSwingPositionDriveEnabled);
-		//void SetAngularSwingPositionDriveEnabled(bool InAngularSwingPositionDriveEnabled);
-		//bool GetAngularSwingPositionDriveEnabled() const
-
-		CONSTRAINT_JOINT_PROPERPETY_IMPL2(FRotation3, AngularDrivePositionTarget, JointSettings, AngularDrivePositionTarget);
-		//void SetAngularDrivePositionTarget(FRotation3 InAngularDrivePositionTarget);
-		//FRotation3 GetAngularDrivePositionTarget() const
-
-		CONSTRAINT_JOINT_PROPERPETY_IMPL2(bool, AngularSLerpVelocityDriveEnabled, JointSettings, bAngularSLerpVelocityDriveEnabled);
-		//void SetAngularSLerpVelocityDriveEnabled(bool InAngularSLerpVelocityDriveEnabled);
-		//bool GetAngularSLerpVelocityDriveEnabled() const
-
-		CONSTRAINT_JOINT_PROPERPETY_IMPL2(bool, AngularTwistVelocityDriveEnabled, JointSettings, bAngularTwistVelocityDriveEnabled);
-		//void SetAngularTwistVelocityDriveEnabled(bool InAngularTwistVelocityDriveEnabled);
-		//bool GetAngularTwistVelocityDriveEnabled() const
-
-		CONSTRAINT_JOINT_PROPERPETY_IMPL2(bool, AngularSwingVelocityDriveEnabled, JointSettings, bAngularSwingVelocityDriveEnabled);
-		//void SetAngularSwingVelocityDriveEnabled(bool InAngularSwingVelocityDriveEnabled);
-		//bool GetAngularSwingVelocityDriveEnabled() const
-
-		CONSTRAINT_JOINT_PROPERPETY_IMPL2(FVec3, AngularDriveVelocityTarget, JointSettings, AngularDriveVelocityTarget);
-		//void SetAngularDriveVelocityTarget(FVec3 InAngularDriveVelocityTarget);
-		//FVec3 GetAngularDriveVelocityTarget() const
-
-		CONSTRAINT_JOINT_PROPERPETY_IMPL2(EJointForceMode, AngularDriveForceMode, JointSettings, AngularDriveForceMode);
-		//void SetAngularDriveForceMode(EJointForceMode InAngularDriveForceMode);
-		//EJointForceMode GetAngularDriveForceMode() const
-
-		CONSTRAINT_JOINT_PROPERPETY_IMPL2(EJointMotionType, AngularMotionTypesX, JointSettings, AngularMotionTypes[0]);
-		//void SetAngularMotionTypesX(EJointMotionType InAngularMotionTypesX);
-		//EJointMotionType GetAngularMotionTypesX() const
-
-		CONSTRAINT_JOINT_PROPERPETY_IMPL2(EJointMotionType, AngularMotionTypesY, JointSettings, AngularMotionTypes[1]);
-		//void SetAngularMotionTypesY(EJointMotionType InAngularMotionTypesY);
-		//EJointMotionType GetAngularMotionTypesY() const
-
-		CONSTRAINT_JOINT_PROPERPETY_IMPL2(EJointMotionType, AngularMotionTypesZ, JointSettings, AngularMotionTypes[2]);
-		//void SetAngularMotionTypesZ(EJointMotionType AngularMotionTypesZ);
-		//EJointMotionType GetAngularMotionTypesZ() const
-
-		CONSTRAINT_JOINT_PROPERPETY_IMPL2(FReal, AngularDriveStiffness, JointSettings, AngularDriveStiffness);
-		//void SetAngularDriveStiffness(FReal InAngularDriveStiffness);
-		//FReal GetAngularDriveStiffness() const
-
-		CONSTRAINT_JOINT_PROPERPETY_IMPL2(FReal, AngularDriveDamping, JointSettings, AngularDriveDamping);
-		//void SetAngularDriveDamping(FReal InAngularDriveDamping);
-		//FReal GetAngularDriveDamping() const
-
-		CONSTRAINT_JOINT_PROPERPETY_IMPL2(FReal, Stiffness, JointSettings, Stiffness);
-		//void SetStiffness(FReal Stiffness);
-		//FReal GetStiffness() const
-
-		CONSTRAINT_JOINT_PROPERPETY_IMPL2(bool, SoftLinearLimitsEnabled, JointSettings, bSoftLinearLimitsEnabled);
-		//void SetSoftLinearLimitsEnabled(bool bInSoftLinearLimitsEnabled);
-		//bool GetSoftLinearLimitsEnabled() const
-
-		CONSTRAINT_JOINT_PROPERPETY_IMPL2(bool, SoftTwistLimitsEnabled, JointSettings, bSoftTwistLimitsEnabled);
-		//void SetSoftTwistLimitsEnabled(FReal bInSoftTwistLimitsEnabled);
-		//bool GetSoftTwistLimitsEnabled() const
-
-		CONSTRAINT_JOINT_PROPERPETY_IMPL2(bool, SoftSwingLimitsEnabled, JointSettings, bSoftSwingLimitsEnabled);
-		//void SetSoftSwingLimitsEnabled(FReal bInSoftSwingLimitsEnabled);
-		//bool GetSoftSwingLimitsEnabled() const
-
-		CONSTRAINT_JOINT_PROPERPETY_IMPL2(EJointForceMode, LinearSoftForceMode, JointSettings, LinearSoftForceMode);
-		//void SetLinearSoftForceMode(FRealIn LinearSoftForceMode);
-		//EJointForceMode GetLinearSoftForceMode() const
-
-		CONSTRAINT_JOINT_PROPERPETY_IMPL2(EJointForceMode, AngularSoftForceMode, JointSettings, AngularSoftForceMode);
-		//void SetAngularSoftForceMode(FReal InAngularSoftForceMode);
-		//EJointForceMode GetAngularSoftForceMode() const
-
-		CONSTRAINT_JOINT_PROPERPETY_IMPL2(FReal, SoftLinearStiffness, JointSettings, SoftLinearStiffness);
-		//void SetSoftLinearStiffness(FReal InSoftLinearStiffness);
-		//FReal GetSoftLinearStiffness() const
-
-		CONSTRAINT_JOINT_PROPERPETY_IMPL2(FReal, SoftLinearDamping, JointSettings, SoftLinearDamping);
-		//void SetSoftLinearDamping(FReal InSoftLinearDamping);
-		//FReal GetSoftLinearDamping() const
-
-		CONSTRAINT_JOINT_PROPERPETY_IMPL2(FReal, SoftTwistStiffness, JointSettings, SoftTwistStiffness);
-		//void SetSoftTwistStiffness(FReal InSoftTwistStiffness);
-		//FReal GetSoftTwistStiffness() const
-
-		CONSTRAINT_JOINT_PROPERPETY_IMPL2(FReal, SoftTwistDamping, JointSettings, SoftTwistDamping);
-		//void SetSoftTwistDamping(FReal InSoftTwistDamping);
-		//FReal GetSoftTwistDamping() const
-
-		CONSTRAINT_JOINT_PROPERPETY_IMPL2(FReal, SoftSwingStiffness, JointSettings, SoftSwingStiffness);
-		//void SetSoftSwingStiffness(FReal InSoftSwingStiffness);
-		//FReal GetSoftSwingStiffness() const
-
-		CONSTRAINT_JOINT_PROPERPETY_IMPL2(FReal, SoftSwingDamping, JointSettings, SoftSwingDamping);
-		//void SetSoftSwingDamping(FReal InSoftSwingDamping);
-		//FReal GetSoftSwingDamping() const
-
-		CONSTRAINT_JOINT_PROPERPETY_IMPL2(FReal, LinearLimit, JointSettings, LinearLimit);
-		//void SetLinearLimit(FReal InLinearLimit);
-		//FReal GetLinearLimit() const
-
-		CONSTRAINT_JOINT_PROPERPETY_IMPL2(FVec3, AngularLimits, JointSettings, AngularLimits);
-		//void SetAngularLimits(FVec3 AngularLimits);
-		//FVec3 GetAngularLimits() const
-
-		CONSTRAINT_JOINT_PROPERPETY_IMPL2(FReal, LinearContactDistance, JointSettings, LinearContactDistance);
-		//void SetLinearContactDistance(FReal InLinearContactDistance);
-		//FReal GetLinearContactDistance() const
-
-		CONSTRAINT_JOINT_PROPERPETY_IMPL2(FReal, TwistContactDistance, JointSettings, TwistContactDistance);
-		//void SetAngularTwistContactDistance(FReal InAngularTwistContactDistance);
-		//FReal GetAngularTwistContactDistance() const
-
-		CONSTRAINT_JOINT_PROPERPETY_IMPL2(FReal, SwingContactDistance, JointSettings, SwingContactDistance);
-		//void SetAngularSwingContactDistance(FReal InAngularSwingContactDistance);
-		//FReal GetAngularSwingContactDistance() const
-
-		CONSTRAINT_JOINT_PROPERPETY_IMPL2(FReal, LinearRestitution, JointSettings, LinearRestitution);
-		//void SetLinearRestitution(FReal InLinearRestitution);
-		//FReal GetLinearRestitution() const
-
-		CONSTRAINT_JOINT_PROPERPETY_IMPL2(FReal, TwistRestitution, JointSettings, TwistRestitution);
-		//void SetAngularTwistRestitution(FReal InAngularTwistRestitution);
-		//FReal GetAngularTwistRestitution() const
-
-		CONSTRAINT_JOINT_PROPERPETY_IMPL2(FReal, SwingRestitution, JointSettings, SwingRestitution);
-		//void SetAngularSwingRestitution(FReal InAngularSwingRestitution);
-		//FReal GetAngularSwingRestitution() const
-
 		struct FOutputData
 		{
 			// Output properties
diff --git a/Engine/Source/Runtime/Experimental/Chaos/Public/Chaos/PBDJointConstraintTypes.h b/Engine/Source/Runtime/Experimental/Chaos/Public/Chaos/PBDJointConstraintTypes.h
--- a/Engine/Source/Runtime/Experimental/Chaos/Public/Chaos/PBDJointConstraintTypes.h
+++ b/Engine/Source/Runtime/Experimental/Chaos/Public/Chaos/PBDJointConstraintTypes.h
@@ -121,6 +121,11 @@ namespace Chaos
 		// whenever settings change.
 		void Sanitize();
 
+		bool operator==(const FPBDJointSettings& Other) const
+		{
+			return !FMemory::Memcmp(this, &Other, sizeof(*this));
+		}
+
 		FTransformPair ConnectorTransforms;
 
 		FReal Stiffness;
diff --git a/Engine/Source/Runtime/Experimental/Chaos/Public/Chaos/PBDJointConstraints.h b/Engine/Source/Runtime/Experimental/Chaos/Public/Chaos/PBDJointConstraints.h
--- a/Engine/Source/Runtime/Experimental/Chaos/Public/Chaos/PBDJointConstraints.h
+++ b/Engine/Source/Runtime/Experimental/Chaos/Public/Chaos/PBDJointConstraints.h
@@ -16,6 +16,7 @@
 #include "Chaos/PBDConstraintContainer.h"
 #include "Chaos/PBDJointConstraintTypes.h"
 #include "Chaos/PBDJointConstraintData.h"
+#include "Chaos/GeometryParticles.h"
 
 namespace Chaos
 {
@@ -52,9 +53,11 @@ namespace Chaos
 		void SetSettings(const FPBDJointSettings& Settings);
 		TVec2<FGeometryParticleHandle*> GetConstrainedParticles() const;
 
-		//TODO: implement this for RewindData to work correctly
-		ESyncState SyncState() const { ensure(false); return ESyncState::InSync; }
-		void SetEnabledDuringResim(bool bEnabled) { ensure(false); }
+		ESyncState SyncState() const;
+		void SetSyncState(ESyncState SyncState);
+
+		void SetEnabledDuringResim(bool bEnabled);
+		EResimType ResimType() const;
 
 	protected:
 		using Base::ConstraintIndex;
@@ -75,6 +78,9 @@ namespace Chaos
 		bool bDriveTargetChanged;
 		FVec3 LinearImpulse;
 		FVec3 AngularImpulse;
+		EResimType ResimType = EResimType::FullResim;
+		ESyncState SyncState = ESyncState::InSync;
+		bool bEnabledDuringResim = true;
 	};
 
 	/**
@@ -233,6 +239,13 @@ namespace Chaos
 		FVec3 GetConstraintLinearImpulse(int32 ConstraintIndex) const;
 		FVec3 GetConstraintAngularImpulse(int32 ConstraintIndex) const;
 
+		ESyncState GetConstraintSyncState(int32 ConstraintIndex) const;
+		void SetConstraintSyncState(int32 ConstraintIndex, ESyncState SyncState);
+		
+		void SetConstraintEnabledDuringResim(int32 ConstraintIndex, bool bEnabled);
+		
+		EResimType GetConstraintResimType(int32 ConstraintIndex) const;
+		
 		//
 		// General Rule API
 		//
diff --git a/Engine/Source/Runtime/Experimental/Chaos/Public/Chaos/ParticleProperties.inl b/Engine/Source/Runtime/Experimental/Chaos/Public/Chaos/ParticleProperties.inl
--- a/Engine/Source/Runtime/Experimental/Chaos/Public/Chaos/ParticleProperties.inl
+++ b/Engine/Source/Runtime/Experimental/Chaos/Public/Chaos/ParticleProperties.inl
@@ -18,4 +18,6 @@ CHAOS_PROPERTY(MassProps, FParticleMassProps, EPhysicsProxyType::SingleParticleP
 CHAOS_PROPERTY(KinematicTarget,FKinematicTarget, EPhysicsProxyType::SingleParticleProxy)
 
 CHAOS_PROPERTY(JointSettings, FPBDJointSettings, EPhysicsProxyType::JointConstraintType)
-CHAOS_PROPERTY(JointParticleProxies, FProxyBasePair, EPhysicsProxyType::JointConstraintType)
\ No newline at end of file
+CHAOS_PROPERTY(JointParticleProxies, FProxyBasePair, EPhysicsProxyType::JointConstraintType)
+
+#undef CHAOS_PROPERTY
\ No newline at end of file
diff --git a/Engine/Source/Runtime/Experimental/Chaos/Public/PhysicsProxy/JointConstraintProxy.h b/Engine/Source/Runtime/Experimental/Chaos/Public/PhysicsProxy/JointConstraintProxy.h
--- a/Engine/Source/Runtime/Experimental/Chaos/Public/PhysicsProxy/JointConstraintProxy.h
+++ b/Engine/Source/Runtime/Experimental/Chaos/Public/PhysicsProxy/JointConstraintProxy.h
@@ -11,6 +11,9 @@
 #include "PhysicsCoreTypes.h"
 #include "Chaos/Defines.h"
 #include "JointConstraintProxyFwd.h"
+#include "Framework/Threading.h"
+#include "Chaos/PBDJointConstraintData.h"
+#include "RewindData.h"
 
 namespace Chaos
 {
@@ -20,6 +23,12 @@ class FPBDRigidsEvolutionGBF;
 
 struct FDirtyJointConstraintData;
 
+template <bool bExternal>
+class TThreadedJointConstraintPhysicsProxyBase;
+
+using FJointConstraintHandle_External = TThreadedJointConstraintPhysicsProxyBase<true>;
+using FJointConstraintHandle_Internal = TThreadedJointConstraintPhysicsProxyBase<false>;
+
 class FJointConstraintPhysicsProxy : public IPhysicsProxyBase
 {
 	using Base = IPhysicsProxyBase;
@@ -44,6 +53,29 @@ public:
 
 	void CHAOS_API DestroyOnGameThread();
 
+	FORCEINLINE Chaos::FJointConstraintHandle_External& GetGameThreadAPI()
+	{
+		return (Chaos::FJointConstraintHandle_External&)*this;
+	}
+
+	FORCEINLINE const Chaos::FJointConstraintHandle_External& GetGameThreadAPI() const
+	{
+		return (const Chaos::FJointConstraintHandle_External&)*this;
+	}
+
+	//Note this is a pointer because the internal handle may have already been deleted
+	FORCEINLINE Chaos::FJointConstraintHandle_Internal* GetPhysicsThreadAPI()
+	{
+		return GetHandle() == nullptr ? nullptr : (Chaos::FJointConstraintHandle_Internal*)this;
+	}
+
+	//Note this is a pointer because the internal handle may have already been deleted
+	FORCEINLINE const Chaos::FJointConstraintHandle_Internal* GetPhysicsThreadAPI() const
+	{
+		return GetHandle() == nullptr ? nullptr : (const Chaos::FJointConstraintHandle_Internal*)this;
+	}
+
+
 	//
 	// Member Access
 	//
@@ -74,4 +106,65 @@ private:
 	bool bInitialized = false;
 };
 
+
+/** Wrapper class that routes all reads and writes to the appropriate joint data. This is helpful for cases where we want to both write to a joint and a network buffer for example*/
+template <bool bExternal>
+class TThreadedJointConstraintPhysicsProxyBase : protected FJointConstraintPhysicsProxy
+{
+public:
+
+#define CHAOS_INNER_JOINT_PROPERTY(OuterProp, FuncName, Inner, InnerType)\
+	const InnerType& Get##FuncName() const { ReadRef([](const auto& Data){ return Data.Inner;}); }\
+	void Set##FuncName(const InnerType& Val) { Write([&Val](auto& Data){ Data.Inner = Val;}); }\
+
+#include "Chaos/JointProperties.inl"
+
+private:
+
+	void VerifyContext() const
+	{
+#if PHYSICS_THREAD_CONTEXT
+		//Are you using the wrong API type for the thread this code runs in?
+		//GetGameThreadAPI should be used for gamethread, GetPhysicsThreadAPI should be used for callbacks and internal physics thread
+		//Note if you are using a ParallelFor you must use PhysicsParallelFor to ensure the right context is inherited from parent thread
+		if (bExternal)
+		{
+			//if proxy is registered with solver, we need a lock
+			if (GetSolverBase() != nullptr)
+			{
+				ensure(IsInGameThreadContext());
+			}
+		}
+		else
+		{
+			ensure(IsInPhysicsThreadContext());
+		}
+#endif
+	}
+
+	template <typename TLambda>
+	const auto& ReadRef(const TLambda& Lambda) const { VerifyContext(); return bExternal ? Lambda(GetConstraint()) : Lambda(GetHandle()); }
+
+	template <typename TLambda>
+	void Write(const TLambda& Lambda)
+	{
+		VerifyContext();
+		if (bExternal)
+		{
+			Lambda(GetConstraint());
+		}
+		else
+		{
+			//Mark entire joint as dirty from PT. TODO: use property system
+			FPhysicsSolverBase* SolverBase = GetSolverBase();	//internal so must have solver already
+			if (FRewindData* RewindData = SolverBase->GetRewindData())
+			{
+				RewindData->MarkDirtyJointFromPT(*GetHandle());
+			}
+
+			Lambda(GetHandle());
+		}
+	}
+};
+
 }
\ No newline at end of file
diff --git a/Engine/Source/Runtime/Experimental/Chaos/Public/RewindData.h b/Engine/Source/Runtime/Experimental/Chaos/Public/RewindData.h
--- a/Engine/Source/Runtime/Experimental/Chaos/Public/RewindData.h
+++ b/Engine/Source/Runtime/Experimental/Chaos/Public/RewindData.h
@@ -82,6 +82,17 @@ struct FPropertyInterval
 	FFrameAndPhase FrameAndPhase;
 };
 
+template <typename TData, typename TObj>
+void CopyDataFromObject(TData& Data, const TObj& Obj)
+{
+	Data.CopyFrom(Obj);
+}
+
+inline void CopyDataFromObject(FPBDJointSettings& Data, const FPBDJointConstraintHandle& Joint)
+{
+	Data = Joint.GetSettings();
+}
+
 template <typename T, EChaosProperty PropName, bool bNoEntryIsHead = true>
 class TParticlePropertyBuffer
 {
@@ -177,9 +188,9 @@ public:
 		if (const T* Val = Read(FrameAndPhase, Pool))
 		{
 			T HeadVal;
-			HeadVal.CopyFrom(Handle);
+			CopyDataFromObject(HeadVal, Handle);
 			return *Val == HeadVal;
-	}
+		}
 
 		return NoEntryInSync<THandle, T, bNoEntryIsHead>::Helper(Handle);
 	}
@@ -461,7 +472,7 @@ private:
 	}\
 
 #define REWIND_JOINT_PROPERTY(PROP, FUNC_NAME, NAME)\
-	decltype(auto) FUNC_NAME() const\
+	decltype(auto) Get##FUNC_NAME() const\
 	{\
 		const auto Data = State ? State->PROP.Read(FrameAndPhase, Pool) : nullptr;\
 		return Data ? Data->NAME : Head.GetSettings().NAME;\
@@ -672,6 +683,7 @@ struct FJointStateBase
 		return JointSettings.IsClean(FrameAndPhase) && ProxyPair.IsClean(FrameAndPhase);
 	}
 
+	template <bool bSkipDynamics>
 	bool IsInSync(const FPBDJointConstraintHandle& Handle, const FFrameAndPhase FrameAndPhase, const FDirtyPropertiesPool& Pool) const;
 
 	TParticlePropertyBuffer<FPBDJointSettings, EChaosProperty::JointSettings> JointSettings;
@@ -695,67 +707,10 @@ public:
 	{
 	}
 
-	REWIND_JOINT_PROPERTY(JointSettings, JointTransforms, ConnectorTransforms);
-	REWIND_JOINT_PROPERTY(JointSettings, CollisionEnabled, bCollisionEnabled);
-	REWIND_JOINT_PROPERTY(JointSettings, ProjectionEnabled, bProjectionEnabled);
-	REWIND_JOINT_PROPERTY(JointSettings, ProjectionLinearAlpha, LinearProjection);
-	REWIND_JOINT_PROPERTY(JointSettings, ProjectionAngularAlpha, AngularProjection);
-	REWIND_JOINT_PROPERTY(JointSettings, ParentInvMassScale, ParentInvMassScale);
-	REWIND_JOINT_PROPERTY(JointSettings, LinearBreakForce, LinearBreakForce);
-	REWIND_JOINT_PROPERTY(JointSettings, LinearPlasticityLimit, LinearPlasticityLimit);
-	REWIND_JOINT_PROPERTY(JointSettings, LinearPlasticityType, LinearPlasticityType);
-	REWIND_JOINT_PROPERTY(JointSettings, AngularBreakTorque, AngularBreakTorque);
-	REWIND_JOINT_PROPERTY(JointSettings, AngularPlasticityLimit, AngularPlasticityLimit);
-	REWIND_JOINT_PROPERTY(JointSettings, UserData, UserData);
-	REWIND_JOINT_PROPERTY(JointSettings, LinearPositionDriveXEnabled, bLinearPositionDriveEnabled[0]);
-	REWIND_JOINT_PROPERTY(JointSettings, LinearPositionDriveYEnabled, bLinearPositionDriveEnabled[1]);
-	REWIND_JOINT_PROPERTY(JointSettings, LinearPositionDriveZEnabled, bLinearPositionDriveEnabled[2]);
-	REWIND_JOINT_PROPERTY(JointSettings, LinearDrivePositionTarget, LinearDrivePositionTarget);
-	REWIND_JOINT_PROPERTY(JointSettings, LinearVelocityDriveXEnabled, bLinearVelocityDriveEnabled[0]);
-	REWIND_JOINT_PROPERTY(JointSettings, LinearVelocityDriveYEnabled, bLinearVelocityDriveEnabled[1]);
-	REWIND_JOINT_PROPERTY(JointSettings, LinearVelocityDriveZEnabled, bLinearVelocityDriveEnabled[2]);
-	REWIND_JOINT_PROPERTY(JointSettings, LinearDriveVelocityTarget, LinearDriveVelocityTarget);
-	REWIND_JOINT_PROPERTY(JointSettings, LinearDriveForceMode, LinearDriveForceMode);
-	REWIND_JOINT_PROPERTY(JointSettings, LinearMotionTypesX, LinearMotionTypes[0]);
-	REWIND_JOINT_PROPERTY(JointSettings, LinearMotionTypesY, LinearMotionTypes[1]);
-	REWIND_JOINT_PROPERTY(JointSettings, LinearMotionTypesZ, LinearMotionTypes[2]);
-	REWIND_JOINT_PROPERTY(JointSettings, LinearDriveStiffness, LinearDriveStiffness);
-	REWIND_JOINT_PROPERTY(JointSettings, LinearDriveDamping, LinearDriveDamping);
-	REWIND_JOINT_PROPERTY(JointSettings, ContactTransferScale, ContactTransferScale);
-	REWIND_JOINT_PROPERTY(JointSettings, AngularSLerpPositionDriveEnabled, bAngularSLerpPositionDriveEnabled);
-	REWIND_JOINT_PROPERTY(JointSettings, AngularTwistPositionDriveEnabled, bAngularTwistPositionDriveEnabled);
-	REWIND_JOINT_PROPERTY(JointSettings, AngularSwingPositionDriveEnabled, bAngularSwingPositionDriveEnabled);
-	REWIND_JOINT_PROPERTY(JointSettings, AngularDrivePositionTarget, AngularDrivePositionTarget);
-	REWIND_JOINT_PROPERTY(JointSettings, AngularSLerpVelocityDriveEnabled, bAngularSLerpVelocityDriveEnabled);
-	REWIND_JOINT_PROPERTY(JointSettings, AngularTwistVelocityDriveEnabled, bAngularTwistVelocityDriveEnabled);
-	REWIND_JOINT_PROPERTY(JointSettings, AngularSwingVelocityDriveEnabled, bAngularSwingVelocityDriveEnabled);
-	REWIND_JOINT_PROPERTY(JointSettings, AngularDriveVelocityTarget, AngularDriveVelocityTarget);
-	REWIND_JOINT_PROPERTY(JointSettings, AngularDriveForceMode, AngularDriveForceMode);
-	REWIND_JOINT_PROPERTY(JointSettings, AngularMotionTypesX, AngularMotionTypes[0]);
-	REWIND_JOINT_PROPERTY(JointSettings, AngularMotionTypesY, AngularMotionTypes[1]);
-	REWIND_JOINT_PROPERTY(JointSettings, AngularMotionTypesZ, AngularMotionTypes[2]);
-	REWIND_JOINT_PROPERTY(JointSettings, AngularDriveStiffness, AngularDriveStiffness);
-	REWIND_JOINT_PROPERTY(JointSettings, AngularDriveDamping, AngularDriveDamping);
-	REWIND_JOINT_PROPERTY(JointSettings, Stiffness, Stiffness);
-	REWIND_JOINT_PROPERTY(JointSettings, SoftLinearLimitsEnabled, bSoftLinearLimitsEnabled);
-	REWIND_JOINT_PROPERTY(JointSettings, SoftTwistLimitsEnabled, bSoftTwistLimitsEnabled);
-	REWIND_JOINT_PROPERTY(JointSettings, SoftSwingLimitsEnabled, bSoftSwingLimitsEnabled);
-	REWIND_JOINT_PROPERTY(JointSettings, LinearSoftForceMode, LinearSoftForceMode);
-	REWIND_JOINT_PROPERTY(JointSettings, AngularSoftForceMode, AngularSoftForceMode);
-	REWIND_JOINT_PROPERTY(JointSettings, SoftLinearStiffness, SoftLinearStiffness);
-	REWIND_JOINT_PROPERTY(JointSettings, SoftLinearDamping, SoftLinearDamping);
-	REWIND_JOINT_PROPERTY(JointSettings, SoftTwistStiffness, SoftTwistStiffness);
-	REWIND_JOINT_PROPERTY(JointSettings, SoftTwistDamping, SoftTwistDamping);
-	REWIND_JOINT_PROPERTY(JointSettings, SoftSwingStiffness, SoftSwingStiffness);
-	REWIND_JOINT_PROPERTY(JointSettings, SoftSwingDamping, SoftSwingDamping);
-	REWIND_JOINT_PROPERTY(JointSettings, LinearLimit, LinearLimit);
-	REWIND_JOINT_PROPERTY(JointSettings, AngularLimits, AngularLimits);
-	REWIND_JOINT_PROPERTY(JointSettings, LinearContactDistance, LinearContactDistance);
-	REWIND_JOINT_PROPERTY(JointSettings, TwistContactDistance, TwistContactDistance);
-	REWIND_JOINT_PROPERTY(JointSettings, SwingContactDistance, SwingContactDistance);
-	REWIND_JOINT_PROPERTY(JointSettings, LinearRestitution, LinearRestitution);
-	REWIND_JOINT_PROPERTY(JointSettings, TwistRestitution, TwistRestitution);
-	REWIND_JOINT_PROPERTY(JointSettings, SwingRestitution, SwingRestitution);
+	//See JointProperties for API
+	//Each CHAOS_INNER_JOINT_PROPERTY entry will have a Get*
+#define CHAOS_INNER_JOINT_PROPERTY(OuterProp, FuncName, Inner, InnerType) REWIND_JOINT_PROPERTY(OuterProp, FuncName, Inner);
+#include "Chaos/JointProperties.inl"
 
 private:
 	const FPBDJointConstraintHandle& Head;
@@ -859,6 +814,8 @@ private:
 	TArray<TVal> DenseVals;
 };
 
+extern CHAOS_API int32 SkipDesyncTest;
+
 class FPBDRigidsSolver;
 
 class FRewindData
@@ -886,12 +843,12 @@ public:
 		return Managers[Frame].DeltaTime;
 	}
 
-	void CHAOS_API RemoveParticle(const FGeometryParticleHandle* Particle)
+	void CHAOS_API RemoveObject(const FGeometryParticleHandle* Particle)
 	{
 		DirtyParticles.Remove(Particle);
 	}
 
-	void CHAOS_API RemoveJoint(const FPBDJointConstraintHandle* Joint)
+	void CHAOS_API RemoveObject(const FPBDJointConstraintHandle* Joint)
 	{
 		DirtyJoints.Remove(Joint);
 	}
@@ -966,6 +923,7 @@ public:
 	void PushPTDirtyData(TPBDRigidParticleHandle<FReal,3>& Rigid,const int32 SrcDataIdx);
 
 	void CHAOS_API MarkDirtyFromPT(FGeometryParticleHandle& Handle);
+	void CHAOS_API MarkDirtyJointFromPT(FPBDJointConstraintHandle& Handle);
 
 	void CHAOS_API SpawnProxyIfNeeded(FSingleParticlePhysicsProxy& Proxy);
 
@@ -1087,7 +1045,12 @@ private:
 
 	bool RewindToFrame(int32 Frame);
 
-	static void DesyncParticle(FDirtyParticleInfo& Info, const FFrameAndPhase FrameAndPhase);
+	template <typename TDirtyInfo>
+	static void DesyncObject(TDirtyInfo& Info, const FFrameAndPhase FrameAndPhase)
+	{
+		Info.ClearPhaseAndFuture(FrameAndPhase);
+		Info.GetObjectPtr()->SetSyncState(ESyncState::HardDesync);
+	}
 
 	TCircularBuffer<FFrameManagerInfo> Managers;
 	FDirtyPropertiesPool PropertiesPool;	//must come before DirtyParticles since it relies on it (and used in destruction)
@@ -1103,11 +1066,11 @@ private:
 	bool bNeedsSave;	//Indicates that some data is pointing at head and requires saving before a rewind
 	bool bResimOptimization;
 
-	bool IsResimAndInSync(const FGeometryParticleHandle& Handle) const { return IsResim() && Handle.SyncState() == ESyncState::InSync; }
-
-	template <bool bSkipDynamics = false>
-	void DesyncIfNecessary(FDirtyParticleInfo& Info, const FFrameAndPhase FrameAndPhase);
+	template <typename TObj>
+	bool IsResimAndInSync(const TObj& Handle) const { return IsResim() && Handle.SyncState() == ESyncState::InSync; }
 
+	template <bool bSkipDynamics, typename TDirtyInfo>
+	void DesyncIfNecessary(TDirtyInfo& Info, const FFrameAndPhase FrameAndPhase);
 };
 
 /** Used by user code to determine when rewind should occur and gives it the opportunity to record any additional data */
diff --git a/Engine/Source/Runtime/PhysicsCore/Private/ChaosEngineInterface.cpp b/Engine/Source/Runtime/PhysicsCore/Private/ChaosEngineInterface.cpp
--- a/Engine/Source/Runtime/PhysicsCore/Private/ChaosEngineInterface.cpp
+++ b/Engine/Source/Runtime/PhysicsCore/Private/ChaosEngineInterface.cpp
@@ -1018,7 +1018,7 @@ FPhysicsConstraintHandle FChaosEngineInterface::CreateConstraint(const FPhysicsA
 
 				JointConstraint->SetParticleProxies({ ValidParticle, KinematicEndPoint });
 
-				Chaos::FJointConstraint::FTransformPair TransformPair = { InLocalFrame1, InLocalFrame2 };
+				Chaos::FTransformPair TransformPair = { InLocalFrame1, InLocalFrame2 };
 				if (bSwapped)
 				{
 					Swap(TransformPair[0], TransformPair[1]);
@@ -1124,7 +1124,7 @@ FTransform FChaosEngineInterface::GetLocalPose(const FPhysicsConstraintHandle& I
 	{
 		if (Chaos::FJointConstraint* Constraint = static_cast<Chaos::FJointConstraint*>(InConstraintRef.Constraint))
 		{
-			const Chaos::FJointConstraint::FTransformPair& M = Constraint->GetJointTransforms();
+			const Chaos::FTransformPair& M = Constraint->GetJointTransforms();
 			if (InFrame == EConstraintFrame::Frame1)
 			{
 				return M[0];
@@ -1159,7 +1159,7 @@ FTransform FChaosEngineInterface::GetGlobalPose(const FPhysicsConstraintHandle&
 		if (Chaos::FJointConstraint* Constraint = static_cast<Chaos::FJointConstraint*>(InConstraintRef.Constraint))
 		{
 			Chaos::FProxyBasePair BasePairs = Constraint->GetParticleProxies();
-			const Chaos::FJointConstraint::FTransformPair& M = Constraint->GetJointTransforms();
+			const Chaos::FTransformPair& M = Constraint->GetJointTransforms();
 
 			if (InFrame == EConstraintFrame::Frame1)
 			{
@@ -1337,7 +1337,7 @@ void FChaosEngineInterface::SetLocalPose(const FPhysicsConstraintHandle& InConst
 	{
 		if (Chaos::FJointConstraint* Constraint = static_cast<Chaos::FJointConstraint*>(InConstraintRef.Constraint))
 		{
-			Chaos::FJointConstraint::FTransformPair JointTransforms = Constraint->GetJointTransforms();
+			Chaos::FTransformPair JointTransforms = Constraint->GetJointTransforms();
 			if (InFrame == EConstraintFrame::Frame1)
 			{
 				JointTransforms[0] = InPose;
