Chaos: Downgraded relative size calculation warning - this is a tooling-calculated attribute that isn't available on all collections and was failing Headless runs. Make the Warning a Display insead.

diff --git a/Engine/Source/Runtime/Experimental/Chaos/Private/PhysicsProxy/GeometryCollectionPhysicsProxy.cpp b/Engine/Source/Runtime/Experimental/Chaos/Private/PhysicsProxy/GeometryCollectionPhysicsProxy.cpp
--- a/Engine/Source/Runtime/Experimental/Chaos/Private/PhysicsProxy/GeometryCollectionPhysicsProxy.cpp
+++ b/Engine/Source/Runtime/Experimental/Chaos/Private/PhysicsProxy/GeometryCollectionPhysicsProxy.cpp
@@ -1854,7 +1854,7 @@ void FGeometryCollectionPhysicsProxy::InitializeSharedCollisionStructures(
 	bool bUseRelativeSize = RestCollection.HasAttribute(TEXT("Size"), FTransformCollection::TransformGroup);
 	if (!bUseRelativeSize)
 	{
-		UE_LOG(LogChaos, Warning, TEXT("Relative Size not found on Rest Collection. Using bounds volume for SizeSpecificData indexing instead."));
+		UE_LOG(LogChaos, Display, TEXT("Relative Size not found on Rest Collection. Using bounds volume for SizeSpecificData indexing instead."));
 	}
 
 
