Make sure DestroyParticle removes from RewindData directly.

diff --git a/Engine/Source/Runtime/Experimental/Chaos/Private/PBDRigidsSolver.cpp b/Engine/Source/Runtime/Experimental/Chaos/Private/PBDRigidsSolver.cpp
--- a/Engine/Source/Runtime/Experimental/Chaos/Private/PBDRigidsSolver.cpp
+++ b/Engine/Source/Runtime/Experimental/Chaos/Private/PBDRigidsSolver.cpp
@@ -675,6 +675,7 @@ namespace Chaos
 		if(RewindCaptureNumFrames >= 0)
 		{
 			EnableRewindCapture(RewindCaptureNumFrames, bUseCollisionResimCache || !!UseResimCache);
+			MEvolution->SetRewindData(GetRewindData());
 		}
 
 		MEvolution->SetCaptureRewindDataFunction([this](const TParticleView<TPBDRigidParticles<FReal,3>>& ActiveParticles)
@@ -714,12 +715,6 @@ namespace Chaos
 
 				if (Info.Handle)
 				{
-					// Remove from rewind data
-					if (FRewindData* RewindData = GetRewindData())
-					{
-						RewindData->RemoveObject(Info.Handle);
-					}
-
 					// Use the handle to destroy the particle data
 					GetEvolution()->DestroyParticle(Info.Handle);
 				}
diff --git a/Engine/Source/Runtime/Experimental/Chaos/Public/Chaos/PBDRigidsEvolution.h b/Engine/Source/Runtime/Experimental/Chaos/Public/Chaos/PBDRigidsEvolution.h
--- a/Engine/Source/Runtime/Experimental/Chaos/Public/Chaos/PBDRigidsEvolution.h
+++ b/Engine/Source/Runtime/Experimental/Chaos/Public/Chaos/PBDRigidsEvolution.h
@@ -18,6 +18,7 @@
 #include "Chaos/Defines.h"
 #include "Chaos/PendingSpatialData.h"
 #include "ProfilingDebugging/CsvProfiler.h"
+#include "RewindData.h"
 
 
 extern int32 ChaosRigidsEvolutionApplyAllowEarlyOutCVar;
@@ -413,6 +414,11 @@ public:
 
 	CHAOS_API void DestroyParticle(FGeometryParticleHandle* Particle)
 	{
+		if(MRewindData)
+		{
+			MRewindData->RemoveObject(Particle);
+		}
+
 		RemoveParticleFromAccelerationStructure(*Particle);
 		DisconnectConstraints(TSet<FGeometryParticleHandle*>({ Particle }));
 		ConstraintGraph.RemoveParticle(Particle);
@@ -806,6 +812,11 @@ public:
 		bCanStartAsyncTasks = bInCanStartAsyncTasks;
 	}
 
+	void SetRewindData(FRewindData* RewindData)
+	{
+		MRewindData = RewindData;
+	}
+
 	CHAOS_API void DisableParticleWithRemovalEvent(FGeometryParticleHandle* Particle);
 	const TArray<FRemovalData>& GetAllRemovals() { return MAllRemovals; }
 	void ResetAllRemovals() { MAllRemovals.Reset(); }
@@ -924,6 +935,7 @@ protected:
 	FAccelerationStructure* InternalAcceleration;
 	FAccelerationStructure* AsyncInternalAcceleration;
 	FAccelerationStructure* AsyncExternalAcceleration;
+	FRewindData* MRewindData = nullptr;
 
 	//internal thread will push into this and external thread will consume
 	TQueue<FAccelerationStructure*,EQueueMode::Spsc> ExternalStructuresQueue;
