Fix evolution not getting the right rewind data when created from different contexts

diff --git a/Engine/Source/Runtime/Experimental/Chaos/Private/Chaos/Framework/PhysicsSolverBase.cpp b/Engine/Source/Runtime/Experimental/Chaos/Private/Chaos/Framework/PhysicsSolverBase.cpp
--- a/Engine/Source/Runtime/Experimental/Chaos/Private/Chaos/Framework/PhysicsSolverBase.cpp
+++ b/Engine/Source/Runtime/Experimental/Chaos/Private/Chaos/Framework/PhysicsSolverBase.cpp
@@ -300,16 +300,6 @@ namespace Chaos
 			UniqueIdxToGTParticles[Idx] = nullptr;
 		}
 	}
-	
-	void FPhysicsSolverBase::EnableRewindCapture(int32 NumFrames, bool InUseCollisionResimCache, TUniquePtr<IRewindCallback>&& RewindCallback)
-	{
-		//TODO: this function calls both internal and extrnal - sort of assumed during initialization. Should decide what thread it's called on and mark it as either external or internal
-		MRewindData = MakeUnique<FRewindData>(((FPBDRigidsSolver*)this), NumFrames, InUseCollisionResimCache, ((FPBDRigidsSolver*)this)->GetCurrentFrame()); // FIXME
-		bUseCollisionResimCache = InUseCollisionResimCache;
-		MRewindCallback = MoveTemp(RewindCallback);
-		MarshallingManager.SetHistoryLength_Internal(NumFrames);
-		PullResultsManager->SetHistoryLength_External(NumFrames);	
-	}
 
 	void FPhysicsSolverBase::SetRewindCallback(TUniquePtr<IRewindCallback>&& RewindCallback)
 	{
diff --git a/Engine/Source/Runtime/Experimental/Chaos/Private/PBDRigidsSolver.cpp b/Engine/Source/Runtime/Experimental/Chaos/Private/PBDRigidsSolver.cpp
--- a/Engine/Source/Runtime/Experimental/Chaos/Private/PBDRigidsSolver.cpp
+++ b/Engine/Source/Runtime/Experimental/Chaos/Private/PBDRigidsSolver.cpp
@@ -655,6 +655,17 @@ namespace Chaos
 	int32 UseResimCache = 0;
 	FAutoConsoleVariableRef CVarUseResimCache(TEXT("p.UseResimCache"),UseResimCache,TEXT("Whether resim uses cache to skip work, requires recreating world to take effect"));
 
+	void FPBDRigidsSolver::EnableRewindCapture(int32 NumFrames, bool InUseCollisionResimCache, TUniquePtr<IRewindCallback>&& RewindCallback)
+	{
+		//TODO: this function calls both internal and extrnal - sort of assumed during initialization. Should decide what thread it's called on and mark it as either external or internal
+		MRewindData = MakeUnique<FRewindData>(((FPBDRigidsSolver*)this), NumFrames, InUseCollisionResimCache, ((FPBDRigidsSolver*)this)->GetCurrentFrame()); // FIXME
+		bUseCollisionResimCache = InUseCollisionResimCache;
+		MRewindCallback = MoveTemp(RewindCallback);
+		MarshallingManager.SetHistoryLength_Internal(NumFrames);
+		PullResultsManager->SetHistoryLength_External(NumFrames);
+		MEvolution->SetRewindData(GetRewindData());
+	}
+
 	void FPBDRigidsSolver::Reset()
 	{
 		UE_LOG(LogPBDRigidsSolver, Verbose, TEXT("PBDRigidsSolver::Reset()"));
@@ -675,7 +686,6 @@ namespace Chaos
 		if(RewindCaptureNumFrames >= 0)
 		{
 			EnableRewindCapture(RewindCaptureNumFrames, bUseCollisionResimCache || !!UseResimCache);
-			MEvolution->SetRewindData(GetRewindData());
 		}
 
 		MEvolution->SetCaptureRewindDataFunction([this](const TParticleView<TPBDRigidParticles<FReal,3>>& ActiveParticles)
diff --git a/Engine/Source/Runtime/Experimental/Chaos/Public/Chaos/Framework/PhysicsSolverBase.h b/Engine/Source/Runtime/Experimental/Chaos/Public/Chaos/Framework/PhysicsSolverBase.h
--- a/Engine/Source/Runtime/Experimental/Chaos/Public/Chaos/Framework/PhysicsSolverBase.h
+++ b/Engine/Source/Runtime/Experimental/Chaos/Public/Chaos/Framework/PhysicsSolverBase.h
@@ -186,7 +186,6 @@ namespace Chaos
 			RegisterSimOneShotCallback(MoveTemp(Func));
 		}
 
-		void EnableRewindCapture(int32 NumFrames, bool InUseCollisionResimCache, TUniquePtr<IRewindCallback>&& RewindCallback = TUniquePtr<IRewindCallback>());
 		void SetRewindCallback(TUniquePtr<IRewindCallback>&& RewindCallback);
 
 		FRewindData* GetRewindData()
diff --git a/Engine/Source/Runtime/Experimental/Chaos/Public/PBDRigidsSolver.h b/Engine/Source/Runtime/Experimental/Chaos/Public/PBDRigidsSolver.h
--- a/Engine/Source/Runtime/Experimental/Chaos/Public/PBDRigidsSolver.h
+++ b/Engine/Source/Runtime/Experimental/Chaos/Public/PBDRigidsSolver.h
@@ -213,6 +213,8 @@ namespace Chaos
 		const FPBDSuspensionConstraints& GetSuspensionConstraints() const { return MEvolution->GetSuspensionConstraints(); }
 		void SetSuspensionTargetOnPhysicsThread(Chaos::FSuspensionConstraint* GTConstraint, const FVector& TargetPos, bool Enabled);
 
+		void EnableRewindCapture(int32 NumFrames, bool InUseCollisionResimCache, TUniquePtr<IRewindCallback>&& RewindCallback = TUniquePtr<IRewindCallback>());
+
 		/**/
 		FPBDRigidsEvolution* GetEvolution() { return MEvolution.Get(); }
 		FPBDRigidsEvolution* GetEvolution() const { return MEvolution.Get(); }
