Refactor suspension constraint to be async thread safe. Also moved FSingleParticlePhysicsProxy into the Chaos namespace

diff --git a/Engine/Plugins/Experimental/ChaosVehiclesPlugin/Source/ChaosVehicles/Public/ChaosVehicleManagerAsyncCallback.h b/Engine/Plugins/Experimental/ChaosVehiclesPlugin/Source/ChaosVehicles/Public/ChaosVehicleManagerAsyncCallback.h
--- a/Engine/Plugins/Experimental/ChaosVehiclesPlugin/Source/ChaosVehicles/Public/ChaosVehicleManagerAsyncCallback.h
+++ b/Engine/Plugins/Experimental/ChaosVehiclesPlugin/Source/ChaosVehicles/Public/ChaosVehicleManagerAsyncCallback.h
@@ -97,7 +97,7 @@ struct FChaosVehicleAsyncInput
 	const EChaosAsyncVehicleDataType Type;
 	UChaosVehicleMovementComponent* Vehicle;
 	
-	FSingleParticlePhysicsProxy* Proxy;
+	Chaos::FSingleParticlePhysicsProxy* Proxy;
 
 	/** 
 	* Vehicle simulation running on the Physics Thread
diff --git a/Engine/Plugins/Experimental/Water/Source/Runtime/Public/BuoyancyTypes.h b/Engine/Plugins/Experimental/Water/Source/Runtime/Public/BuoyancyTypes.h
--- a/Engine/Plugins/Experimental/Water/Source/Runtime/Public/BuoyancyTypes.h
+++ b/Engine/Plugins/Experimental/Water/Source/Runtime/Public/BuoyancyTypes.h
@@ -317,7 +317,7 @@ struct FBuoyancyComponentAsyncInput
 	const EAsyncBuoyancyComponentDataType Type;
 	const UBuoyancyComponent* BuoyancyComponent;
 
-	FSingleParticlePhysicsProxy* Proxy;
+	Chaos::FSingleParticlePhysicsProxy* Proxy;
 
 	virtual TUniquePtr<struct FBuoyancyComponentAsyncOutput> PreSimulate(UWorld* World, const float DeltaSeconds, const float TotalSeconds, FBuoyancyComponentAsyncAux* Aux, const TMap<UWaterBodyComponent*, TUniquePtr<FSolverSafeWaterBodyData>>& WaterBodyComponentData) const = 0;
 
diff --git a/Engine/Plugins/Runtime/NetworkPrediction/Source/NetworkPrediction/Private/NetworkPhysics.cpp b/Engine/Plugins/Runtime/NetworkPrediction/Source/NetworkPrediction/Private/NetworkPhysics.cpp
--- a/Engine/Plugins/Runtime/NetworkPrediction/Source/NetworkPrediction/Private/NetworkPhysics.cpp
+++ b/Engine/Plugins/Runtime/NetworkPrediction/Source/NetworkPrediction/Private/NetworkPhysics.cpp
@@ -245,7 +245,7 @@ struct FNetworkPhysicsRewindCallback : public Chaos::IRewindCallback
 
 			for (FNetworkPhysicsState& Obj : Snapshot.Objects)
 			{
-				FSingleParticlePhysicsProxy* Proxy = Obj.Proxy;
+				Chaos::FSingleParticlePhysicsProxy* Proxy = Obj.Proxy;
 				check(Proxy);
 
 				if (Obj.Frame < EarliestFrame || Obj.Frame > LastCompletedStep)
@@ -422,7 +422,7 @@ struct FNetworkPhysicsRewindCallback : public Chaos::IRewindCallback
 			FSnapshot Snapshot;
 			Snapshot.SimulationFrame = PhysicsStep - this->LastLocalOffset;
 			Snapshot.LocalFrameOffset = this->LastLocalOffset;
-			for (FSingleParticlePhysicsProxy* Proxy : Request.Proxies)
+			for (Chaos::FSingleParticlePhysicsProxy* Proxy : Request.Proxies)
 			{
 				if (auto* PT = Proxy->GetPhysicsThreadAPI())
 				{
@@ -552,7 +552,7 @@ struct FNetworkPhysicsRewindCallback : public Chaos::IRewindCallback
 	// GT -> PT. "These are the proxies I want you to populate DataFromPhysics with". Is there a better way to do it?
 	struct FRequest
 	{
-		TArray<FSingleParticlePhysicsProxy*> Proxies;
+		TArray<Chaos::FSingleParticlePhysicsProxy*> Proxies;
 	};
 
 	TQueue<FRequest> DataRequested;
@@ -579,7 +579,7 @@ struct FNetworkPhysicsRewindCallback : public Chaos::IRewindCallback
 	{
 		struct FCorrection
 		{
-			FSingleParticlePhysicsProxy* Proxy = nullptr;
+			Chaos::FSingleParticlePhysicsProxy* Proxy = nullptr;
 			FBasePhysicsState LocalState;
 			FBasePhysicsState AuthorityState;
 			int32 LocalFrame = 0;
diff --git a/Engine/Plugins/Runtime/NetworkPrediction/Source/NetworkPrediction/Private/NetworkPredictionUtil.cpp b/Engine/Plugins/Runtime/NetworkPrediction/Source/NetworkPrediction/Private/NetworkPredictionUtil.cpp
--- a/Engine/Plugins/Runtime/NetworkPrediction/Source/NetworkPrediction/Private/NetworkPredictionUtil.cpp
+++ b/Engine/Plugins/Runtime/NetworkPrediction/Source/NetworkPrediction/Private/NetworkPredictionUtil.cpp
@@ -7,7 +7,7 @@
 
 namespace UE_NP
 {
-	FSingleParticlePhysicsProxy* FindBestPhysicsProxy(AActor* Owner, FName NamedComponent)
+	Chaos::FSingleParticlePhysicsProxy* FindBestPhysicsProxy(AActor* Owner, FName NamedComponent)
 	{
 		// Test is component is valid for Network Physics. Needs a valid physics ActorHandle
 		auto ValidComponent = [](UActorComponent* Component)
@@ -35,7 +35,7 @@ namespace UE_NP
 			return Pc;
 		};
 
-		FSingleParticlePhysicsProxy* Proxy = nullptr;
+		Chaos::FSingleParticlePhysicsProxy* Proxy = nullptr;
 		UPrimitiveComponent* PrimitiveComponent = nullptr;
 		
 		// Explicitly tagged component
diff --git a/Engine/Plugins/Runtime/NetworkPrediction/Source/NetworkPrediction/Public/NetworkPhysics.h b/Engine/Plugins/Runtime/NetworkPrediction/Source/NetworkPrediction/Public/NetworkPhysics.h
--- a/Engine/Plugins/Runtime/NetworkPrediction/Source/NetworkPrediction/Public/NetworkPhysics.h
+++ b/Engine/Plugins/Runtime/NetworkPrediction/Source/NetworkPrediction/Public/NetworkPhysics.h
@@ -16,7 +16,10 @@ NETWORKPREDICTION_API DECLARE_LOG_CATEGORY_EXTERN(LogNetworkPhysics, Log, All);
 
 struct FNetworkPhysicsRewindCallback;
 class FMockObjectManager;
-class FSingleParticlePhysicsProxy;
+namespace Chaos
+{
+	class FSingleParticlePhysicsProxy;
+}
 
 namespace Chaos { struct FSimCallbackInputAndObject; }
 
@@ -55,7 +58,7 @@ struct FNetworkPhysicsState
 	GENERATED_BODY()
 	
 	// Local proxy we are working with. Must be set locally on initialization.
-	FSingleParticlePhysicsProxy* Proxy = nullptr;
+	Chaos::FSingleParticlePhysicsProxy* Proxy = nullptr;
 
 	// Local handle used by UNetworkPhysicsManager for registering/unregistering
 	int32 LocalManagedHandle = INDEX_NONE;
@@ -206,7 +209,7 @@ private:
 	int32 LastFreeIndex = 0;
 	int32 UniqueHandleCounter = 0;
 
-	TMap<FSingleParticlePhysicsProxy*, TUniqueFunction<void(const FDrawDebugParams&)>> DrawDebugMap;
+	TMap<Chaos::FSingleParticlePhysicsProxy*, TUniqueFunction<void(const FDrawDebugParams&)>> DrawDebugMap;
 	
 	TArray<TPair<FName, TUniquePtr<INetworkPhysicsSubsystem>>> SubSystems;
 
diff --git a/Engine/Plugins/Runtime/NetworkPrediction/Source/NetworkPrediction/Public/NetworkPhysicsComponent.h b/Engine/Plugins/Runtime/NetworkPrediction/Source/NetworkPrediction/Public/NetworkPhysicsComponent.h
--- a/Engine/Plugins/Runtime/NetworkPrediction/Source/NetworkPrediction/Public/NetworkPhysicsComponent.h
+++ b/Engine/Plugins/Runtime/NetworkPrediction/Source/NetworkPrediction/Public/NetworkPhysicsComponent.h
@@ -30,7 +30,7 @@ public:
 
 protected:
 	
-	FSingleParticlePhysicsProxy* GetManagedProxy() const { return NetworkPhysicsStateArray.Num() > 0 ? NetworkPhysicsStateArray[0].Proxy : nullptr; }
+	Chaos::FSingleParticlePhysicsProxy* GetManagedProxy() const { return NetworkPhysicsStateArray.Num() > 0 ? NetworkPhysicsStateArray[0].Proxy : nullptr; }
 	
 private:
 
diff --git a/Engine/Plugins/Runtime/NetworkPrediction/Source/NetworkPrediction/Public/NetworkPredictionUtil.h b/Engine/Plugins/Runtime/NetworkPrediction/Source/NetworkPrediction/Public/NetworkPredictionUtil.h
--- a/Engine/Plugins/Runtime/NetworkPrediction/Source/NetworkPrediction/Public/NetworkPredictionUtil.h
+++ b/Engine/Plugins/Runtime/NetworkPrediction/Source/NetworkPrediction/Public/NetworkPredictionUtil.h
@@ -52,9 +52,13 @@ void NpResizeForIndex(ArrayType& Array, int32 Index)
 }
 
 class AActor;
-class FSingleParticlePhysicsProxy;
+
+namespace Chaos
+{
+	class FSingleParticlePhysicsProxy;
+}
 
 namespace UE_NP
 {
-	NETWORKPREDICTION_API FSingleParticlePhysicsProxy* FindBestPhysicsProxy(AActor* Owning, FName NamedComponent = NAME_None);
+	NETWORKPREDICTION_API Chaos::FSingleParticlePhysicsProxy* FindBestPhysicsProxy(AActor* Owning, FName NamedComponent = NAME_None);
 }
\ No newline at end of file
diff --git a/Engine/Plugins/Runtime/NetworkPredictionExtras/Source/NetworkPredictionExtras/Public/PhysicsMovement.h b/Engine/Plugins/Runtime/NetworkPredictionExtras/Source/NetworkPredictionExtras/Public/PhysicsMovement.h
--- a/Engine/Plugins/Runtime/NetworkPredictionExtras/Source/NetworkPredictionExtras/Public/PhysicsMovement.h
+++ b/Engine/Plugins/Runtime/NetworkPredictionExtras/Source/NetworkPredictionExtras/Public/PhysicsMovement.h
@@ -204,7 +204,7 @@ struct FPhysicsMovementLocalState
 {
 	GENERATED_BODY()
 
-	FSingleParticlePhysicsProxy* Proxy = nullptr;
+	Chaos::FSingleParticlePhysicsProxy* Proxy = nullptr;
 
 	FCollisionQueryParams QueryParams = FCollisionQueryParams::DefaultQueryParam;
 };
diff --git a/Engine/Plugins/Runtime/NetworkPredictionExtras/Source/NetworkPredictionExtras/Public/PhysicsSimple.h b/Engine/Plugins/Runtime/NetworkPredictionExtras/Source/NetworkPredictionExtras/Public/PhysicsSimple.h
--- a/Engine/Plugins/Runtime/NetworkPredictionExtras/Source/NetworkPredictionExtras/Public/PhysicsSimple.h
+++ b/Engine/Plugins/Runtime/NetworkPredictionExtras/Source/NetworkPredictionExtras/Public/PhysicsSimple.h
@@ -103,7 +103,7 @@ struct FSimpleLocalState
 {
 	GENERATED_BODY()
 
-	FSingleParticlePhysicsProxy* Proxy = nullptr;
+	Chaos::FSingleParticlePhysicsProxy* Proxy = nullptr;
 };
 
 // Debug/Testing component
diff --git a/Engine/Source/Runtime/Engine/Private/PhysicsEngine/Experimental/PhysScene_Chaos.cpp b/Engine/Source/Runtime/Engine/Private/PhysicsEngine/Experimental/PhysScene_Chaos.cpp
--- a/Engine/Source/Runtime/Engine/Private/PhysicsEngine/Experimental/PhysScene_Chaos.cpp
+++ b/Engine/Source/Runtime/Engine/Private/PhysicsEngine/Experimental/PhysScene_Chaos.cpp
@@ -500,7 +500,7 @@ void FPhysScene_Chaos::AddObject(UPrimitiveComponent* Component, FStaticMeshPhys
 	ensure(false);
 }
 
-void FPhysScene_Chaos::AddObject(UPrimitiveComponent* Component, FSingleParticlePhysicsProxy* InObject)
+void FPhysScene_Chaos::AddObject(UPrimitiveComponent* Component, Chaos::FSingleParticlePhysicsProxy* InObject)
 {
 	AddToComponentMaps(Component, InObject);
 	ensure(false);
@@ -576,7 +576,7 @@ void FPhysScene_Chaos::RemoveObject(FStaticMeshPhysicsProxy* InObject)
 #endif
 }
 
-void FPhysScene_Chaos::RemoveObject(FSingleParticlePhysicsProxy* InObject)
+void FPhysScene_Chaos::RemoveObject(Chaos::FSingleParticlePhysicsProxy* InObject)
 {
 	ensure(false);
 #if 0
@@ -646,7 +646,7 @@ FBodyInstance* FPhysScene_Chaos::GetBodyInstanceFromProxy(const IPhysicsProxyBas
 	FBodyInstance* BodyInstance = nullptr;
 	if (PhysicsProxy && PhysicsProxy->GetType() == EPhysicsProxyType::SingleParticleProxy)
 	{
-		const Chaos::FRigidBodyHandle_External& RigidBodyHandle = static_cast<const FSingleParticlePhysicsProxy*>(PhysicsProxy)->GetGameThreadAPI();
+		const Chaos::FRigidBodyHandle_External& RigidBodyHandle = static_cast<const Chaos::FSingleParticlePhysicsProxy*>(PhysicsProxy)->GetGameThreadAPI();
 		BodyInstance = FPhysicsUserData::Get<FBodyInstance>(RigidBodyHandle.UserData());
 	}
 	// found none, let's see if there's an owning component in the scene
@@ -1618,7 +1618,7 @@ void FPhysScene_Chaos::OnSyncBodies(Chaos::FPhysicsSolverBase* Solver)
 
 	FPhysicsCommand::ExecuteWrite(this, [&Solver, &PendingTransforms](FPhysScene* PhysScene)
 	{
-		auto RigidLambda = [&PhysScene, &PendingTransforms](FSingleParticlePhysicsProxy* Proxy)
+		auto RigidLambda = [&PhysScene, &PendingTransforms](Chaos::FSingleParticlePhysicsProxy* Proxy)
 		{
 			FPBDRigidParticle* DirtyParticle = Proxy->GetRigidParticleUnsafe();
 
@@ -1626,7 +1626,7 @@ void FPhysScene_Chaos::OnSyncBodies(Chaos::FPhysicsSolverBase* Solver)
 			{
 				if(BodyInstance->OwnerComponent.IsValid())
 				{
-					if (SyncKinematicOnGameThread == 0 && BodyInstance->IsInstanceSimulatingPhysics() == false)
+					if (Chaos::SyncKinematicOnGameThread == 0 && BodyInstance->IsInstanceSimulatingPhysics() == false)
 					{
 						return;
 					}
diff --git a/Engine/Source/Runtime/Engine/Private/PhysicsEngine/PhysicsReplication.cpp b/Engine/Source/Runtime/Engine/Private/PhysicsEngine/PhysicsReplication.cpp
--- a/Engine/Source/Runtime/Engine/Private/PhysicsEngine/PhysicsReplication.cpp
+++ b/Engine/Source/Runtime/Engine/Private/PhysicsEngine/PhysicsReplication.cpp
@@ -353,7 +353,7 @@ bool FPhysicsReplication::ApplyRigidBodyState(float DeltaSeconds, FBodyInstance*
 				AsyncDesiredState.WorldTM = IdealWorldTM;
 				AsyncDesiredState.LinearVelocity = NewState.LinVel;
 				AsyncDesiredState.AngularVelocity = NewState.AngVel;
-				AsyncDesiredState.Proxy = static_cast<FSingleParticlePhysicsProxy*>(BI->GetPhysicsActorHandle());
+				AsyncDesiredState.Proxy = static_cast<Chaos::FSingleParticlePhysicsProxy*>(BI->GetPhysicsActorHandle());
 				AsyncDesiredState.ObjectState = AsyncDesiredState.Proxy->GetGameThreadAPI().ObjectState();
 				AsyncDesiredState.bShouldSleep = bShouldSleep;
 
@@ -579,7 +579,7 @@ void FPhysicsReplication::ApplyAsyncDesiredState(const float DeltaSeconds, const
 		for (const FAsyncPhysicsDesiredState& State : AsyncData->Buffer)
 		{
 			//Proxy should exist because we are using latest and any pending deletes would have been enqueued after
-			FSingleParticlePhysicsProxy* Proxy = State.Proxy;
+			Chaos::FSingleParticlePhysicsProxy* Proxy = State.Proxy;
 			auto* Handle = Proxy->GetPhysicsThreadAPI();
 
 
diff --git a/Engine/Source/Runtime/Engine/Public/Physics/Experimental/PhysScene_Chaos.h b/Engine/Source/Runtime/Engine/Public/Physics/Experimental/PhysScene_Chaos.h
--- a/Engine/Source/Runtime/Engine/Public/Physics/Experimental/PhysScene_Chaos.h
+++ b/Engine/Source/Runtime/Engine/Public/Physics/Experimental/PhysScene_Chaos.h
@@ -123,7 +123,7 @@ public:
 	 */
 	void AddObject(UPrimitiveComponent* Component, FSkeletalMeshPhysicsProxy* InObject);
 	void AddObject(UPrimitiveComponent* Component, FStaticMeshPhysicsProxy* InObject);
-	void AddObject(UPrimitiveComponent* Component, FSingleParticlePhysicsProxy* InObject);
+	void AddObject(UPrimitiveComponent* Component, Chaos::FSingleParticlePhysicsProxy* InObject);
 	void AddObject(UPrimitiveComponent* Component, FGeometryCollectionPhysicsProxy* InObject);
 	
 	void AddToComponentMaps(UPrimitiveComponent* Component, IPhysicsProxyBase* InObject);
@@ -135,7 +135,7 @@ public:
 	 */
 	void RemoveObject(FSkeletalMeshPhysicsProxy* InObject);
 	void RemoveObject(FStaticMeshPhysicsProxy* InObject);
-	void RemoveObject(FSingleParticlePhysicsProxy* InObject);
+	void RemoveObject(Chaos::FSingleParticlePhysicsProxy* InObject);
 	void RemoveObject(FGeometryCollectionPhysicsProxy* InObject);
 
 	FPhysicsReplication* GetPhysicsReplication();
diff --git a/Engine/Source/Runtime/Engine/Public/PhysicsReplication.h b/Engine/Source/Runtime/Engine/Public/PhysicsReplication.h
--- a/Engine/Source/Runtime/Engine/Public/PhysicsReplication.h
+++ b/Engine/Source/Runtime/Engine/Public/PhysicsReplication.h
@@ -56,7 +56,7 @@ struct FAsyncPhysicsDesiredState
 	FTransform WorldTM;
 	FVector LinearVelocity;
 	FVector AngularVelocity;
-	FSingleParticlePhysicsProxy* Proxy;
+	Chaos::FSingleParticlePhysicsProxy* Proxy;
 	Chaos::EObjectStateType ObjectState;
 	bool bShouldSleep;
 };
diff --git a/Engine/Source/Runtime/Experimental/Chaos/Private/Chaos/PBDConstraintBaseData.cpp b/Engine/Source/Runtime/Experimental/Chaos/Private/Chaos/PBDConstraintBaseData.cpp
--- a/Engine/Source/Runtime/Experimental/Chaos/Private/Chaos/PBDConstraintBaseData.cpp
+++ b/Engine/Source/Runtime/Experimental/Chaos/Private/Chaos/PBDConstraintBaseData.cpp
@@ -7,7 +7,6 @@ namespace Chaos
 	FConstraintBase::FConstraintBase(EConstraintType InType)
 		: Type(InType)
 		, Proxy(nullptr)
-		, ConnectedParticleProxy({ nullptr, nullptr })
 	{
 	}
 
diff --git a/Engine/Source/Runtime/Experimental/Chaos/Private/Chaos/PBDSuspensionConstraintData.cpp b/Engine/Source/Runtime/Experimental/Chaos/Private/Chaos/PBDSuspensionConstraintData.cpp
--- a/Engine/Source/Runtime/Experimental/Chaos/Private/Chaos/PBDSuspensionConstraintData.cpp
+++ b/Engine/Source/Runtime/Experimental/Chaos/Private/Chaos/PBDSuspensionConstraintData.cpp
@@ -6,7 +6,6 @@ namespace Chaos
 {
 	FSuspensionConstraint::FSuspensionConstraint()
 		: FConstraintBase(EConstraintType::SuspensionConstraintType)
-		, Location({ FVector::ZeroVector })
 
 	{
 	}
diff --git a/Engine/Source/Runtime/Experimental/Chaos/Private/PBDRigidsSolver.cpp b/Engine/Source/Runtime/Experimental/Chaos/Private/PBDRigidsSolver.cpp
--- a/Engine/Source/Runtime/Experimental/Chaos/Private/PBDRigidsSolver.cpp
+++ b/Engine/Source/Runtime/Experimental/Chaos/Private/PBDRigidsSolver.cpp
@@ -632,6 +632,7 @@ namespace Chaos
 		GTConstraint->SetProxy(static_cast<FSuspensionConstraintPhysicsProxy*>(nullptr));
 
 		FParticlesType* InParticles = &GetParticles();
+		SuspensionProxy->DestroyOnGameThread();	//destroy the game thread portion of the proxy
 
 		// Finish registration on the physics thread...
 		EnqueueCommandImmediate([InParticles, SuspensionProxy, this]()
@@ -881,7 +882,8 @@ namespace Chaos
 			case EPhysicsProxyType::SuspensionConstraintType:
 			{
 				auto Proxy = static_cast<FSuspensionConstraintPhysicsProxy*>(Dirty.Proxy);
-				Proxy->PushStateOnGameThread(this);
+				Proxy->PushStateOnGameThread(*Manager, DataIdx, Dirty.PropertyData);
+				Proxy->ResetDirtyIdx();
 				break;
 			}
 
@@ -1063,10 +1065,10 @@ namespace Chaos
 				const bool bIsNew = !SuspensionProxy->IsInitialized();
 				if (bIsNew)
 				{
-					SuspensionProxy->InitializeOnPhysicsThread(this);
+					SuspensionProxy->InitializeOnPhysicsThread(this, *Manager, DataIdx, Dirty.PropertyData);
 					SuspensionProxy->SetInitialized();
 				}
-				SuspensionProxy->PushStateOnPhysicsThread(this);
+				SuspensionProxy->PushStateOnPhysicsThread(this, *Manager, DataIdx, Dirty.PropertyData);
 				Dirty.Proxy->ResetDirtyIdx();
 				break;
 			}
diff --git a/Engine/Source/Runtime/Experimental/Chaos/Private/PhysicsProxy/JointConstraintProxy.cpp b/Engine/Source/Runtime/Experimental/Chaos/Private/PhysicsProxy/JointConstraintProxy.cpp
--- a/Engine/Source/Runtime/Experimental/Chaos/Private/PhysicsProxy/JointConstraintProxy.cpp
+++ b/Engine/Source/Runtime/Experimental/Chaos/Private/PhysicsProxy/JointConstraintProxy.cpp
@@ -81,10 +81,10 @@ void FJointConstraintPhysicsProxy::InitializeOnPhysicsThread(FPBDRigidsSolver* I
 	if (Handles.Size())
 	{
 		auto& JointConstraints = InSolver->GetJointConstraints();
-		if(const FProxyBasePair* BasePairs = RemoteData.FindJointParticleProxies(Manager, DataIdx))
+		if(const FProxyBasePairProperty* BasePairs = RemoteData.FindJointParticleProxies(Manager, DataIdx))
 		{
-			FGeometryParticleHandle* Handle0 = GetParticleHandleFromProxy((*BasePairs)[0]);
-			FGeometryParticleHandle* Handle1 = GetParticleHandleFromProxy((*BasePairs)[1]);
+			FGeometryParticleHandle* Handle0 = GetParticleHandleFromProxy(BasePairs->ParticleProxies[0]);
+			FGeometryParticleHandle* Handle1 = GetParticleHandleFromProxy(BasePairs->ParticleProxies[1]);
 			if (Handle0 && Handle1)
 			{
 				if (const FPBDJointSettings* JointSettings = RemoteData.FindJointSettings(Manager, DataIdx))
diff --git a/Engine/Source/Runtime/Experimental/Chaos/Private/PhysicsProxy/SingleParticlePhysicsProxy.cpp b/Engine/Source/Runtime/Experimental/Chaos/Private/PhysicsProxy/SingleParticlePhysicsProxy.cpp
--- a/Engine/Source/Runtime/Experimental/Chaos/Private/PhysicsProxy/SingleParticlePhysicsProxy.cpp
+++ b/Engine/Source/Runtime/Experimental/Chaos/Private/PhysicsProxy/SingleParticlePhysicsProxy.cpp
@@ -16,12 +16,14 @@
 #include "Chaos/PullPhysicsDataImp.h"
 #include "RewindData.h"
 
+namespace Chaos
+{
+
 // This is a temporary workaround to avoid GT copying position from physics results for kinematics, as they are already at target.
 // Velocity and such is still copied. This will be handled better in the future.
 int32 SyncKinematicOnGameThread = 0;
 FAutoConsoleVariableRef CVar_SyncKinematicOnGameThread(TEXT("P.Chaos.SyncKinematicOnGameThread"), SyncKinematicOnGameThread, TEXT("If set to 1, if a kinematic is flagged to send position back to game thread, move component, if 0, do not."));
 
-
 FSingleParticlePhysicsProxy::FSingleParticlePhysicsProxy(TUniquePtr<PARTICLE_TYPE>&& InParticle, FParticleHandle* InHandle, UObject* InOwner)
 	: IPhysicsProxyBase(EPhysicsProxyType::SingleParticleProxy, InOwner)
 	, Particle(MoveTemp(InParticle))
@@ -377,4 +379,5 @@ void FSingleParticlePhysicsProxy::ClearEvents()
 	{
 		Rigid->ClearEvents();
 	}
+}
 }
\ No newline at end of file
diff --git a/Engine/Source/Runtime/Experimental/Chaos/Private/PhysicsProxy/SuspensionConstraintProxy.cpp b/Engine/Source/Runtime/Experimental/Chaos/Private/PhysicsProxy/SuspensionConstraintProxy.cpp
--- a/Engine/Source/Runtime/Experimental/Chaos/Private/PhysicsProxy/SuspensionConstraintProxy.cpp
+++ b/Engine/Source/Runtime/Experimental/Chaos/Private/PhysicsProxy/SuspensionConstraintProxy.cpp
@@ -14,20 +14,20 @@
 #include "PhysicsProxy/SingleParticlePhysicsProxy.h"
 #include "PhysicsSolver.h"
 
+namespace Chaos
+{
 
-
-FSuspensionConstraintPhysicsProxy::FSuspensionConstraintPhysicsProxy(Chaos::FSuspensionConstraint* InConstraint, FConstraintHandle* InHandle, UObject* InOwner)
-	: Base(InOwner)
-	, Constraint(InConstraint)
-	, Handle(InHandle)
+FSuspensionConstraintPhysicsProxy::FSuspensionConstraintPhysicsProxy(FSuspensionConstraint* InConstraint, FPBDSuspensionConstraintHandle* InHandle, UObject* InOwner)
+	: Base(EPhysicsProxyType::SuspensionConstraintType, InOwner)
+	, Constraint_GT(InConstraint)
+	, Constraint_PT(InHandle)
 	, bInitialized(false)
 {
-	check(Constraint!=nullptr);
-	Constraint->SetProxy(this);
-	SuspensionSettingsBuffer = Constraint->GetSuspensionSettings();
+	check(Constraint_GT!=nullptr);
+	Constraint_GT->SetProxy(this);
 }
 
-Chaos::TGeometryParticleHandle<Chaos::FReal, 3>*
+TGeometryParticleHandle<FReal, 3>*
 FSuspensionConstraintPhysicsProxy::GetParticleHandleFromProxy(IPhysicsProxyBase* ProxyBase)
 {
 	if (ProxyBase)
@@ -40,173 +40,79 @@ FSuspensionConstraintPhysicsProxy::GetParticleHandleFromProxy(IPhysicsProxyBase*
 	return nullptr;
 }
 
-void FSuspensionConstraintPhysicsProxy::InitializeOnPhysicsThread(Chaos::FPBDRigidsSolver* InSolver)
+void FSuspensionConstraintPhysicsProxy::InitializeOnPhysicsThread(FPBDRigidsSolver* InSolver, FDirtyPropertiesManager& Manager, int32 DataIdx, FDirtyChaosProperties& RemoteData)
 {
 	auto& Handles = InSolver->GetParticles().GetParticleHandles();
-	if (Handles.Size() && IsValid())
+	if (Handles.Size())
 	{
 		auto& SuspensionConstraints = InSolver->GetSuspensionConstraints();
-		if (Constraint != nullptr)
+		if (const FParticleProxyProperty* ProxyProperty = RemoteData.FindSuspensionParticleProxy(Manager, DataIdx))
 		{
-			Chaos::FProxyBasePair& BasePairs = Constraint->GetParticleProxies();
-
-			
-			Chaos::TGeometryParticleHandle<Chaos::FReal, 3>* Handle0 = GetParticleHandleFromProxy(BasePairs[0]);
-			if (Handle0)
+			if (FGeometryParticleHandle* Handle0 = GetParticleHandleFromProxy(ProxyProperty->ParticleProxy))
 			{
-				Handle = SuspensionConstraints.AddConstraint(Handle0, Constraint->GetLocation(), SuspensionSettingsBuffer);
-				Handle0->AddConstraintHandle(Handle);
+				const FPBDSuspensionSettings* SuspensionSettings = RemoteData.FindSuspensionSettings(Manager, DataIdx);
+				const FSuspensionLocation* SuspensionLocation = RemoteData.FindSuspensionLocation(Manager, DataIdx);
+
+				if (SuspensionSettings && SuspensionLocation)
+				{
+					Constraint_PT = SuspensionConstraints.AddConstraint(Handle0, SuspensionLocation->Location, *SuspensionSettings);
+					Handle0->AddConstraintHandle(Constraint_PT);
+				}
 			}
-			
 		}
 	}
 }
 
 
-void FSuspensionConstraintPhysicsProxy::PushStateOnGameThread(Chaos::FPBDRigidsSolver* InSolver)
+void FSuspensionConstraintPhysicsProxy::PushStateOnGameThread(FDirtyPropertiesManager& Manager, int32 DataIdx, FDirtyChaosProperties& RemoteData)
 {
-	if (Constraint != nullptr)
+	if (Constraint_GT && Constraint_GT->IsValid())
 	{
-		if (Constraint->IsDirty())
-		{
-			if (Constraint->IsDirty(Chaos::ESuspensionConstraintFlags::Enabled))
-			{
-				SuspensionSettingsBuffer.Enabled = Constraint->GetEnabled();
-				DirtyFlagsBuffer.MarkDirty(Chaos::ESuspensionConstraintFlags::Enabled);
-			}
-
-			if (Constraint->IsDirty(Chaos::ESuspensionConstraintFlags::Target))
-			{
-				SuspensionSettingsBuffer.Target = Constraint->GetTarget();
-				DirtyFlagsBuffer.MarkDirty(Chaos::ESuspensionConstraintFlags::Target);
-			}
-
-			if (Constraint->IsDirty(Chaos::ESuspensionConstraintFlags::HardstopStiffness))
-			{
-				SuspensionSettingsBuffer.HardstopStiffness = Constraint->GetHardstopStiffness();
-				DirtyFlagsBuffer.MarkDirty(Chaos::ESuspensionConstraintFlags::HardstopStiffness);
-			}
-
-			if (Constraint->IsDirty(Chaos::ESuspensionConstraintFlags::HardstopVelocityCompensation))
-			{
-				SuspensionSettingsBuffer.HardstopVelocityCompensation = Constraint->GetHardstopVelocityCompensation();
-				DirtyFlagsBuffer.MarkDirty(Chaos::ESuspensionConstraintFlags::HardstopVelocityCompensation);
-			}
-
-			if (Constraint->IsDirty(Chaos::ESuspensionConstraintFlags::SpringPreload))
-			{
-				SuspensionSettingsBuffer.SpringPreload = Constraint->GetSpringPreload();
-				DirtyFlagsBuffer.MarkDirty(Chaos::ESuspensionConstraintFlags::SpringPreload);
-			}
-
-			if (Constraint->IsDirty(Chaos::ESuspensionConstraintFlags::SpringStiffness))
-			{
-				SuspensionSettingsBuffer.SpringStiffness = Constraint->GetSpringStiffness();
-				DirtyFlagsBuffer.MarkDirty(Chaos::ESuspensionConstraintFlags::SpringStiffness);
-			}
-
-			if (Constraint->IsDirty(Chaos::ESuspensionConstraintFlags::SpringDamping))
-			{
-				SuspensionSettingsBuffer.SpringDamping = Constraint->GetSpringDamping();
-				DirtyFlagsBuffer.MarkDirty(Chaos::ESuspensionConstraintFlags::SpringDamping);
-			}
-
-			if (Constraint->IsDirty(Chaos::ESuspensionConstraintFlags::MinLength))
-			{
-				SuspensionSettingsBuffer.MinLength = Constraint->GetMinLength();
-				DirtyFlagsBuffer.MarkDirty(Chaos::ESuspensionConstraintFlags::MinLength);
-			}
-
-			if (Constraint->IsDirty(Chaos::ESuspensionConstraintFlags::MaxLength))
-			{
-				SuspensionSettingsBuffer.MaxLength = Constraint->GetMaxLength();
-				DirtyFlagsBuffer.MarkDirty(Chaos::ESuspensionConstraintFlags::MaxLength);
-			}
-
-			Constraint->ClearDirtyFlags();
-		}
+		Constraint_GT->SyncRemoteData(Manager, DataIdx, RemoteData);
 	}
 }
 
 
-void FSuspensionConstraintPhysicsProxy::PushStateOnPhysicsThread(Chaos::FPBDRigidsSolver* InSolver)
+void FSuspensionConstraintPhysicsProxy::PushStateOnPhysicsThread(FPBDRigidsSolver* InSolver, const FDirtyPropertiesManager& Manager, int32 DataIdx, const FDirtyChaosProperties& RemoteData)
 {
-	typedef typename Chaos::FPBDRigidsSolver::FPBDRigidsEvolution::FCollisionConstraints FCollisionConstraints;
-	if (Handle)
+	if (Constraint_PT && Constraint_PT->IsValid())
 	{
-		if (DirtyFlagsBuffer.IsDirty())
+		if (const FPBDSuspensionSettings* Data = RemoteData.FindSuspensionSettings(Manager, DataIdx))
 		{
-			FConstraintData& ConstraintSettings = Handle->GetSettings();
-
-			if (DirtyFlagsBuffer.IsDirty(Chaos::ESuspensionConstraintFlags::Enabled))
-			{
-				ConstraintSettings.Enabled = SuspensionSettingsBuffer.Enabled;
-			}
-
-			if (DirtyFlagsBuffer.IsDirty(Chaos::ESuspensionConstraintFlags::Target))
-			{
-				ConstraintSettings.Target = SuspensionSettingsBuffer.Target;
-			}
-
-			if (DirtyFlagsBuffer.IsDirty(Chaos::ESuspensionConstraintFlags::HardstopStiffness))
-			{
-				ConstraintSettings.HardstopStiffness = SuspensionSettingsBuffer.HardstopStiffness;
-			}
-
-			if (DirtyFlagsBuffer.IsDirty(Chaos::ESuspensionConstraintFlags::HardstopVelocityCompensation))
-			{
-				ConstraintSettings.HardstopVelocityCompensation = SuspensionSettingsBuffer.HardstopVelocityCompensation;
-			}
-
-			if (DirtyFlagsBuffer.IsDirty(Chaos::ESuspensionConstraintFlags::SpringPreload))
-			{
-				ConstraintSettings.SpringPreload = SuspensionSettingsBuffer.SpringPreload;
-			}
-			
-			if (DirtyFlagsBuffer.IsDirty(Chaos::ESuspensionConstraintFlags::SpringStiffness))
-			{
-				ConstraintSettings.SpringStiffness = SuspensionSettingsBuffer.SpringStiffness;
-			}
-
-			if (DirtyFlagsBuffer.IsDirty(Chaos::ESuspensionConstraintFlags::SpringDamping))
-			{
-				ConstraintSettings.SpringDamping = SuspensionSettingsBuffer.SpringDamping;
-			}
-
-			if (DirtyFlagsBuffer.IsDirty(Chaos::ESuspensionConstraintFlags::MinLength))
-			{
-				ConstraintSettings.MinLength = SuspensionSettingsBuffer.MinLength;
-			}
-
-			if (DirtyFlagsBuffer.IsDirty(Chaos::ESuspensionConstraintFlags::MaxLength))
-			{
-				ConstraintSettings.MaxLength = SuspensionSettingsBuffer.MaxLength;
-			}
-
-			DirtyFlagsBuffer.Clear();
+			Constraint_PT->GetSettings() = *Data;
 		}
 	}
 }
 
 
-void FSuspensionConstraintPhysicsProxy::DestroyOnPhysicsThread(Chaos::FPBDRigidsSolver* RBDSolver)
+void FSuspensionConstraintPhysicsProxy::DestroyOnGameThread()
 {
-	if (Handle)
+	if (Constraint_GT)
 	{
-		auto& SuspensionConstraints = RBDSolver->GetSuspensionConstraints();
-		SuspensionConstraints.RemoveConstraint(Handle->GetConstraintIndex());
+		delete Constraint_GT; 
+		Constraint_GT = nullptr;
+	}
+}
 
-		delete Constraint; 
-		Constraint = nullptr;
+void FSuspensionConstraintPhysicsProxy::DestroyOnPhysicsThread(FPBDRigidsSolver* RBDSolver)
+{
+	if (Constraint_PT)
+	{
+		auto& SuspensionConstraints = RBDSolver->GetSuspensionConstraints();
+		SuspensionConstraints.RemoveConstraint(Constraint_PT->GetConstraintIndex());
+		Constraint_PT = nullptr;
 	}
 }
 
-void FSuspensionConstraintPhysicsProxy::UpdateTargetOnPhysicsThread(Chaos::FPBDRigidsSolver* RBDSolver, const FVector& TargetPos, bool Enabled)
+
+void FSuspensionConstraintPhysicsProxy::UpdateTargetOnPhysicsThread(FPBDRigidsSolver* RBDSolver, const FVector& TargetPos, bool Enabled)
 {
-	if (Handle)
+	if (Constraint_PT)
 	{
 		auto& SuspensionConstraints = RBDSolver->GetSuspensionConstraints();
-		SuspensionConstraints.GetSettings(Handle->GetConstraintIndex()).Target = TargetPos;
-		SuspensionConstraints.GetSettings(Handle->GetConstraintIndex()).Enabled = Enabled;
+		SuspensionConstraints.GetSettings(Constraint_PT->GetConstraintIndex()).Target = TargetPos;
+		SuspensionConstraints.GetSettings(Constraint_PT->GetConstraintIndex()).Enabled = Enabled;
 	}
 }
+
+}
\ No newline at end of file
diff --git a/Engine/Source/Runtime/Experimental/Chaos/Public/Chaos/PBDConstraintBaseData.h b/Engine/Source/Runtime/Experimental/Chaos/Public/Chaos/PBDConstraintBaseData.h
--- a/Engine/Source/Runtime/Experimental/Chaos/Public/Chaos/PBDConstraintBaseData.h
+++ b/Engine/Source/Runtime/Experimental/Chaos/Public/Chaos/PBDConstraintBaseData.h
@@ -31,11 +31,9 @@ enum class EConstraintType
 	void Set##FNAME(TYPE InValue){ PROP.Modify(/*bInvalidate=*/true, DirtyFlags, Proxy, [&InValue](auto& Data) { Data.VNAME = InValue; }); }\
 	TYPE Get##FNAME() const{return PROP.Read().VNAME;}\
 
-
 class CHAOS_API FConstraintBase
 {
 public:
-
 	virtual ~FConstraintBase() {}
 
 	typedef TVector<TGeometryParticleHandle<FReal, 3>*, 2> FParticleHandlePair;
@@ -58,22 +56,6 @@ public:
 
 	void SetProxy(IPhysicsProxyBase* InProxy);
 
-	void SetParticleProxies(const FProxyBasePair& InJointParticles)
-	{
-		ConnectedParticleProxy[0] = InJointParticles[0];
-		ConnectedParticleProxy[1] = InJointParticles[1];
-
-		ConnectedParticleProxy2.Modify(/*bInvalidate=*/true, DirtyFlags, Proxy, [&InJointParticles](FProxyBasePair& Data)
-		{
-			Data[0] = InJointParticles[0];
-			Data[1] = InJointParticles[1];
-		});
-	}
-
-
-	const FProxyBasePair& GetParticleProxies() const { return  ConnectedParticleProxy; }
-	FProxyBasePair& GetParticleProxies() { return  ConnectedParticleProxy; }
-
 	void SyncRemoteData(FDirtyPropertiesManager& Manager, int32 DataIdx, FDirtyChaosProperties& RemoteData)
 	{
 		RemoteData.SetFlags(DirtyFlags);
@@ -84,17 +66,9 @@ public:
 protected:
 	EConstraintType Type;
 	class IPhysicsProxyBase* Proxy;
-
-	FProxyBasePair ConnectedParticleProxy;	//TODO: remove this once suspension uses chaos property
-
 	FDirtyChaosPropertyFlags DirtyFlags;
-	TChaosProperty<FProxyBasePair, EChaosProperty::JointParticleProxies> ConnectedParticleProxy2;
 
-	virtual void SyncRemoteDataImp(FDirtyPropertiesManager& Manager, int32 DataIdx, FDirtyChaosProperties& RemoteData)
-	{
-		ConnectedParticleProxy2.SyncRemote(Manager, DataIdx, RemoteData);
-	}
-	
+	virtual void SyncRemoteDataImp(FDirtyPropertiesManager& Manager, int32 DataIdx, FDirtyChaosProperties& RemoteData) = 0;
 };
 
 } // Chaos
diff --git a/Engine/Source/Runtime/Experimental/Chaos/Public/Chaos/PBDJointConstraintData.h b/Engine/Source/Runtime/Experimental/Chaos/Public/Chaos/PBDJointConstraintData.h
--- a/Engine/Source/Runtime/Experimental/Chaos/Public/Chaos/PBDJointConstraintData.h
+++ b/Engine/Source/Runtime/Experimental/Chaos/Public/Chaos/PBDJointConstraintData.h
@@ -72,14 +72,27 @@ namespace Chaos
 
 		virtual void SyncRemoteDataImp(FDirtyPropertiesManager& Manager, int32 DataIdx, FDirtyChaosProperties& RemoteData) override
 		{
-			Base::SyncRemoteDataImp(Manager, DataIdx, RemoteData);
+			JointProxies.SyncRemote(Manager, DataIdx, RemoteData);
 			JointSettings.SyncRemote(Manager, DataIdx, RemoteData);
 		}
 
+		void SetParticleProxies(const FProxyBasePair& InJointParticles)
+		{
+			JointProxies.Modify(/*bInvalidate=*/true, DirtyFlags, Proxy, [&InJointParticles](FProxyBasePairProperty& Data)
+			{
+				Data.ParticleProxies[0] = InJointParticles[0];
+				Data.ParticleProxies[1] = InJointParticles[1];
+			});
+		}
+
+
+		const FProxyBasePair& GetParticleProxies() const { return  JointProxies.Read().ParticleProxies; }
+
 	protected:
 
 		void ReleaseKinematicEndPoint(FPBDRigidsSolver* Solver);
 		
+		TChaosProperty<FProxyBasePairProperty, EChaosProperty::JointParticleProxies> JointProxies;
 		TChaosProperty<FPBDJointSettings, EChaosProperty::JointSettings> JointSettings;
 
 
diff --git a/Engine/Source/Runtime/Experimental/Chaos/Public/Chaos/PBDJointConstraints.h b/Engine/Source/Runtime/Experimental/Chaos/Public/Chaos/PBDJointConstraints.h
--- a/Engine/Source/Runtime/Experimental/Chaos/Public/Chaos/PBDJointConstraints.h
+++ b/Engine/Source/Runtime/Experimental/Chaos/Public/Chaos/PBDJointConstraints.h
@@ -49,6 +49,7 @@ namespace Chaos
 		FVec3 GetAngularImpulse() const;
 
 		const FPBDJointSettings& GetSettings() const;
+		const FPBDJointSettings& GetJointSettings() const { return GetSettings(); }	//needed for property macros
 
 		void SetSettings(const FPBDJointSettings& Settings);
 		TVec2<FGeometryParticleHandle*> GetConstrainedParticles() const;
diff --git a/Engine/Source/Runtime/Experimental/Chaos/Public/Chaos/PBDSuspensionConstraintData.h b/Engine/Source/Runtime/Experimental/Chaos/Public/Chaos/PBDSuspensionConstraintData.h
--- a/Engine/Source/Runtime/Experimental/Chaos/Public/Chaos/PBDSuspensionConstraintData.h
+++ b/Engine/Source/Runtime/Experimental/Chaos/Public/Chaos/PBDSuspensionConstraintData.h
@@ -32,48 +32,29 @@ namespace Chaos
 	class CHAOS_API FSuspensionConstraint : public FConstraintBase
 	{
 	public:
-		typedef FPBDSuspensionSettings FData;
-		typedef FPBDSuspensionConstraintHandle FHandle;
-		friend FData;
-
 		FSuspensionConstraint();
 		virtual ~FSuspensionConstraint() override {}
 
-		bool IsDirty() const { return MDirtyFlags.IsDirty(); }
-		bool IsDirty(const ESuspensionConstraintFlags CheckBits) const { return MDirtyFlags.IsDirty(CheckBits); }
-		void ClearDirtyFlags() { MDirtyFlags.Clear(); }
-
-		const FData& GetSuspensionSettings()const { return SuspensionSettings; }
-
-		CONSTRAINT_JOINT_PROPERPETY_IMPL(bool, Enabled, ESuspensionConstraintFlags::Enabled, SuspensionSettings.Enabled);
-
-		CONSTRAINT_JOINT_PROPERPETY_IMPL(FVec3, Target, ESuspensionConstraintFlags::Target, SuspensionSettings.Target);
-
-		CONSTRAINT_JOINT_PROPERPETY_IMPL(FVec3, Location, ESuspensionConstraintFlags::Location, Location);
-
-		CONSTRAINT_JOINT_PROPERPETY_IMPL(FReal, HardstopStiffness, ESuspensionConstraintFlags::HardstopStiffness, SuspensionSettings.HardstopStiffness);
 
-		CONSTRAINT_JOINT_PROPERPETY_IMPL(FReal, HardstopVelocityCompensation, ESuspensionConstraintFlags::HardstopVelocityCompensation, SuspensionSettings.HardstopVelocityCompensation);
+		const FPBDSuspensionSettings& GetSuspensionSettings()const { return SuspensionSettings.Read(); }
 
-		CONSTRAINT_JOINT_PROPERPETY_IMPL(FReal, SpringPreload, ESuspensionConstraintFlags::SpringPreload, SuspensionSettings.SpringPreload);
-
-		CONSTRAINT_JOINT_PROPERPETY_IMPL(FReal, SpringStiffness, ESuspensionConstraintFlags::SpringStiffness, SuspensionSettings.SpringStiffness);
-
-		CONSTRAINT_JOINT_PROPERPETY_IMPL(FReal, SpringDamping, ESuspensionConstraintFlags::SpringDamping, SuspensionSettings.SpringDamping);
-
-		CONSTRAINT_JOINT_PROPERPETY_IMPL(FReal, MinLength, ESuspensionConstraintFlags::MinLength, SuspensionSettings.MinLength);
-
-		CONSTRAINT_JOINT_PROPERPETY_IMPL(FReal, MaxLength, ESuspensionConstraintFlags::MaxLength, SuspensionSettings.MaxLength);
-
-		CONSTRAINT_JOINT_PROPERPETY_IMPL(FVec3, Axis, ESuspensionConstraintFlags::Axis, SuspensionSettings.Axis);
+#define CHAOS_INNER_SUSP_PROPERTY(OuterProp, Name, InnerType)\
+	void Set##Name(InnerType Val){ OuterProp.Modify(/*bInvalidate=*/true, DirtyFlags, Proxy, [&Val](auto& Data) { Data.Name = Val; }); }\
+	InnerType Get##Name() const { return OuterProp.Read().Name;}\
 
+#include "SuspensionProperties.inl"
 
 	protected:
-		FSuspensionConstraintDirtyFlags MDirtyFlags;
-		FData SuspensionSettings;
-
-		FVec3 Location;	// spring local offset
-		FVec3 Target;	// target spring (wheel) end position
+		TChaosProperty<FPBDSuspensionSettings, EChaosProperty::SuspensionSettings> SuspensionSettings;
+		TChaosProperty<FParticleProxyProperty, EChaosProperty::SuspensionParticleProxy> SuspensionParticleProxy;
+		TChaosProperty<FSuspensionLocation, EChaosProperty::SuspensionLocation> SuspensionLocation; // location = spring local offset
+
+		virtual void SyncRemoteDataImp(FDirtyPropertiesManager& Manager, int32 DataIdx, FDirtyChaosProperties& RemoteData) override
+		{
+			SuspensionSettings.SyncRemote(Manager, DataIdx, RemoteData);
+			SuspensionParticleProxy.SyncRemote(Manager, DataIdx, RemoteData);
+			SuspensionLocation.SyncRemote(Manager, DataIdx, RemoteData);
+		}
 	};
 
 
diff --git a/Engine/Source/Runtime/Experimental/Chaos/Public/Chaos/PBDSuspensionConstraintTypes.h b/Engine/Source/Runtime/Experimental/Chaos/Public/Chaos/PBDSuspensionConstraintTypes.h
--- a/Engine/Source/Runtime/Experimental/Chaos/Public/Chaos/PBDSuspensionConstraintTypes.h
+++ b/Engine/Source/Runtime/Experimental/Chaos/Public/Chaos/PBDSuspensionConstraintTypes.h
@@ -13,6 +13,17 @@ namespace Chaos
 
 	class FPBDSuspensionConstraintHandle;
 
+	struct FSuspensionLocation
+	{
+		FVec3 Location = FVec3::ZeroVector;
+	};
+
+	class FSingleParticlePhysicsProxy;
+	struct FParticleProxyProperty
+	{
+		FSingleParticlePhysicsProxy* ParticleProxy = nullptr;
+	};
+
 	class CHAOS_API FPBDSuspensionSettings
 	{
 	public:
diff --git a/Engine/Source/Runtime/Experimental/Chaos/Public/Chaos/ParticleDirtyFlags.h b/Engine/Source/Runtime/Experimental/Chaos/Public/Chaos/ParticleDirtyFlags.h
--- a/Engine/Source/Runtime/Experimental/Chaos/Public/Chaos/ParticleDirtyFlags.h
+++ b/Engine/Source/Runtime/Experimental/Chaos/Public/Chaos/ParticleDirtyFlags.h
@@ -17,6 +17,7 @@
 #include "UObject/UE5ReleaseStreamObjectVersion.h"
 #include "Framework/PhysicsProxyBase.h"
 #include "PBDJointConstraintTypes.h"
+#include "PBDSuspensionConstraintTypes.h"
 
 #ifndef CHAOS_DEBUG_NAME
 #define CHAOS_DEBUG_NAME 0
@@ -186,6 +187,11 @@ private:
 
 typedef TVector<IPhysicsProxyBase*, 2> FProxyBasePair;
 
+struct FProxyBasePairProperty
+{
+	FProxyBasePair ParticleProxies = { nullptr, nullptr };
+};
+
 inline FChaosArchive& operator<<(FChaosArchive& Ar, FParticleDynamics& Data)
 {
 	Data.Serialize(Ar);
diff --git a/Engine/Source/Runtime/Experimental/Chaos/Public/Chaos/ParticleProperties.inl b/Engine/Source/Runtime/Experimental/Chaos/Public/Chaos/ParticleProperties.inl
--- a/Engine/Source/Runtime/Experimental/Chaos/Public/Chaos/ParticleProperties.inl
+++ b/Engine/Source/Runtime/Experimental/Chaos/Public/Chaos/ParticleProperties.inl
@@ -18,6 +18,10 @@ CHAOS_PROPERTY(MassProps, FParticleMassProps, EPhysicsProxyType::SingleParticleP
 CHAOS_PROPERTY(KinematicTarget,FKinematicTarget, EPhysicsProxyType::SingleParticleProxy)
 
 CHAOS_PROPERTY(JointSettings, FPBDJointSettings, EPhysicsProxyType::JointConstraintType)
-CHAOS_PROPERTY(JointParticleProxies, FProxyBasePair, EPhysicsProxyType::JointConstraintType)
+CHAOS_PROPERTY(JointParticleProxies, FProxyBasePairProperty, EPhysicsProxyType::JointConstraintType)
+
+CHAOS_PROPERTY(SuspensionSettings, FPBDSuspensionSettings, EPhysicsProxyType::SuspensionConstraintType)
+CHAOS_PROPERTY(SuspensionParticleProxy, FParticleProxyProperty, EPhysicsProxyType::SuspensionConstraintType)
+CHAOS_PROPERTY(SuspensionLocation, FSuspensionLocation, EPhysicsProxyType::SuspensionConstraintType)
 
 #undef CHAOS_PROPERTY
\ No newline at end of file
diff --git a/Engine/Source/Runtime/Experimental/Chaos/Public/Chaos/SuspensionProperties.inl b/Engine/Source/Runtime/Experimental/Chaos/Public/Chaos/SuspensionProperties.inl
new file mode 100644
--- /dev/null
+++ b/Engine/Source/Runtime/Experimental/Chaos/Public/Chaos/SuspensionProperties.inl
@@ -0,0 +1,18 @@
+// Copyright Epic Games, Inc. All Rights Reserved.
+
+//CHAOS_INNER_SUSP_PROPERTY(OuterProp, FuncName, InnerType)							
+
+CHAOS_INNER_SUSP_PROPERTY(SuspensionSettings, Enabled, bool);
+CHAOS_INNER_SUSP_PROPERTY(SuspensionSettings, Target, FVec3);
+CHAOS_INNER_SUSP_PROPERTY(SuspensionSettings, HardstopStiffness, FReal);
+CHAOS_INNER_SUSP_PROPERTY(SuspensionSettings, HardstopVelocityCompensation, FReal);
+CHAOS_INNER_SUSP_PROPERTY(SuspensionSettings, SpringPreload, FReal);
+CHAOS_INNER_SUSP_PROPERTY(SuspensionSettings, SpringStiffness, FReal);
+CHAOS_INNER_SUSP_PROPERTY(SuspensionSettings, SpringDamping, FReal);
+CHAOS_INNER_SUSP_PROPERTY(SuspensionSettings, MinLength, FReal);
+CHAOS_INNER_SUSP_PROPERTY(SuspensionSettings, MaxLength, FReal);
+CHAOS_INNER_SUSP_PROPERTY(SuspensionSettings, Axis, FVec3);
+CHAOS_INNER_SUSP_PROPERTY(SuspensionParticleProxy, ParticleProxy, FSingleParticlePhysicsProxy*)
+CHAOS_INNER_SUSP_PROPERTY(SuspensionLocation, Location, FVec3)
+
+#undef CHAOS_INNER_SUSP_PROPERTY
\ No newline at end of file
diff --git a/Engine/Source/Runtime/Experimental/Chaos/Public/PhysicsProxy/SingleParticlePhysicsProxy.h b/Engine/Source/Runtime/Experimental/Chaos/Public/PhysicsProxy/SingleParticlePhysicsProxy.h
--- a/Engine/Source/Runtime/Experimental/Chaos/Public/PhysicsProxy/SingleParticlePhysicsProxy.h
+++ b/Engine/Source/Runtime/Experimental/Chaos/Public/PhysicsProxy/SingleParticlePhysicsProxy.h
@@ -19,14 +19,13 @@
 
 
 
-extern CHAOS_API int32 SyncKinematicOnGameThread;
-
 namespace Chaos
 {
-	class FPBDRigidsEvolutionGBF;
+extern CHAOS_API int32 SyncKinematicOnGameThread;
 
-	struct FDirtyRigidParticleData;
-}
+class FPBDRigidsEvolutionGBF;
+
+struct FDirtyRigidParticleData;
 
 class FInitialState
 {
@@ -37,19 +36,19 @@ public:
 	    , InertiaTensor(1.f)
 	{}
 
-	FInitialState(Chaos::FReal MassIn, Chaos::FReal InvMassIn, FVector InertiaTensorIn)
+	FInitialState(FReal MassIn, FReal InvMassIn, FVector InertiaTensorIn)
 	    : Mass(MassIn)
 	    , InvMass(InvMassIn)
 	    , InertiaTensor(InertiaTensorIn)
 	{}
 
-	Chaos::FReal GetMass() const { return Mass; }
-	Chaos::FReal GetInverseMass() const { return InvMass; }
+	FReal GetMass() const { return Mass; }
+	FReal GetInverseMass() const { return InvMass; }
 	FVector GetInertiaTensor() const { return InertiaTensor; }
 
 private:
-	Chaos::FReal Mass;
-	Chaos::FReal InvMass;
+	FReal Mass;
+	FReal InvMass;
 	FVector InertiaTensor;
 };
 
@@ -59,10 +58,10 @@ class FRigidBodyHandle_Internal;
 class CHAOS_API FSingleParticlePhysicsProxy : public IPhysicsProxyBase
 {
 public:
-	using PARTICLE_TYPE = Chaos::FGeometryParticle;
-	using FParticleHandle = Chaos::FGeometryParticleHandle;
+	using PARTICLE_TYPE = FGeometryParticle;
+	using FParticleHandle = FGeometryParticleHandle;
 
-	static FSingleParticlePhysicsProxy* Create(TUniquePtr<Chaos::FGeometryParticle>&& Particle);
+	static FSingleParticlePhysicsProxy* Create(TUniquePtr<FGeometryParticle>&& Particle);
 
 	FSingleParticlePhysicsProxy() = delete;
 	FSingleParticlePhysicsProxy(const FSingleParticlePhysicsProxy&) = delete;
@@ -76,26 +75,26 @@ public:
 
 	int32 GetPullDataInterpIdx_External() const { return PullDataInterpIdx_External; }
 
-	FORCEINLINE Chaos::FRigidBodyHandle_External& GetGameThreadAPI()
+	FORCEINLINE FRigidBodyHandle_External& GetGameThreadAPI()
 	{
-		return (Chaos::FRigidBodyHandle_External&)*this;
+		return (FRigidBodyHandle_External&)*this;
 	}
 
-	FORCEINLINE const Chaos::FRigidBodyHandle_External& GetGameThreadAPI() const
+	FORCEINLINE const FRigidBodyHandle_External& GetGameThreadAPI() const
 	{
-		return (const Chaos::FRigidBodyHandle_External&)*this;
+		return (const FRigidBodyHandle_External&)*this;
 	}
 
 	//Note this is a pointer because the internal handle may have already been deleted
-	FORCEINLINE Chaos::FRigidBodyHandle_Internal* GetPhysicsThreadAPI()
+	FORCEINLINE FRigidBodyHandle_Internal* GetPhysicsThreadAPI()
 	{
-		return GetHandle_LowLevel() == nullptr ? nullptr : (Chaos::FRigidBodyHandle_Internal*)this;
+		return GetHandle_LowLevel() == nullptr ? nullptr : (FRigidBodyHandle_Internal*)this;
 	}
 
 	//Note this is a pointer because the internal handle may have already been deleted
-	FORCEINLINE const Chaos::FRigidBodyHandle_Internal* GetPhysicsThreadAPI() const
+	FORCEINLINE const FRigidBodyHandle_Internal* GetPhysicsThreadAPI() const
 	{
-		return GetHandle_LowLevel() == nullptr ? nullptr : (const Chaos::FRigidBodyHandle_Internal*)this;
+		return GetHandle_LowLevel() == nullptr ? nullptr : (const FRigidBodyHandle_Internal*)this;
 	}
 
 	//Returns the underlying physics thread particle. Note this should only be needed for internal book keeping type tasks. API may change, use GetPhysicsThreadAPI instead
@@ -122,26 +121,26 @@ public:
 
 	// Threading API
 
-	void PushToPhysicsState(const Chaos::FDirtyPropertiesManager& Manager,int32 DataIdx,const Chaos::FDirtyProxy& Dirty,Chaos::FShapeDirtyData* ShapesData, Chaos::FPBDRigidsEvolutionGBF& Evolution, Chaos::FReal ExternalDt);
+	void PushToPhysicsState(const FDirtyPropertiesManager& Manager,int32 DataIdx,const FDirtyProxy& Dirty,FShapeDirtyData* ShapesData, FPBDRigidsEvolutionGBF& Evolution, FReal ExternalDt);
 
 	/**/
 	void ClearAccumulatedData();
 
 	/**/
-	void BufferPhysicsResults(Chaos::FDirtyRigidParticleData&);
+	void BufferPhysicsResults(FDirtyRigidParticleData&);
 
 	/**/
-	void BufferPhysicsResults_External(Chaos::FDirtyRigidParticleData&);
+	void BufferPhysicsResults_External(FDirtyRigidParticleData&);
 
 	/**/
-	bool PullFromPhysicsState(const Chaos::FDirtyRigidParticleData& PullData, int32 SolverSyncTimestamp, const Chaos::FDirtyRigidParticleData* NextPullData = nullptr, const Chaos::FRealSingle* Alpha = nullptr, const Chaos::FRealSingle* LeashAlpha = nullptr);
+	bool PullFromPhysicsState(const FDirtyRigidParticleData& PullData, int32 SolverSyncTimestamp, const FDirtyRigidParticleData* NextPullData = nullptr, const FRealSingle* Alpha = nullptr, const FRealSingle* LeashAlpha = nullptr);
 
 	/**/
 	bool IsDirty();
 
 
 	/**/
-	Chaos::EWakeEventEntry GetWakeEvent() const;
+	EWakeEventEntry GetWakeEvent() const;
 
 	/**/
 	void ClearEvents();
@@ -158,14 +157,14 @@ public:
 		return Particle.Get();
 	}
 
-	Chaos::FPBDRigidParticle* GetRigidParticleUnsafe()
+	FPBDRigidParticle* GetRigidParticleUnsafe()
 	{
-		return static_cast<Chaos::FPBDRigidParticle*>(GetParticle_LowLevel());
+		return static_cast<FPBDRigidParticle*>(GetParticle_LowLevel());
 	}
 
-	const Chaos::FPBDRigidParticle* GetRigidParticleUnsafe() const
+	const FPBDRigidParticle* GetRigidParticleUnsafe() const
 	{
-		return static_cast<const Chaos::FPBDRigidParticle*>(GetParticle_LowLevel());
+		return static_cast<const FPBDRigidParticle*>(GetParticle_LowLevel());
 	}
 
 protected:
@@ -181,9 +180,6 @@ private:
 	FSingleParticlePhysicsProxy(TUniquePtr<PARTICLE_TYPE>&& InParticle, FParticleHandle* InHandle, UObject* InOwner = nullptr);
 };
 
-namespace Chaos
-{
-
 /** Wrapper class that routes all reads and writes to the appropriate particle data. This is helpful for cases where we want to both write to a particle and a network buffer for example*/
 template <bool bExternal>
 class TThreadedSingleParticlePhysicsProxyBase : protected FSingleParticlePhysicsProxy
@@ -1221,11 +1217,10 @@ public:
 
 static_assert(sizeof(FRigidBodyHandle_Internal) == sizeof(FSingleParticlePhysicsProxy), "Derived types only used to constrain API, all data lives in base class ");
 
-}
-
-inline FSingleParticlePhysicsProxy* FSingleParticlePhysicsProxy::Create(TUniquePtr<Chaos::FGeometryParticle>&& Particle)
+inline FSingleParticlePhysicsProxy* FSingleParticlePhysicsProxy::Create(TUniquePtr<FGeometryParticle>&& Particle)
 {
 	ensure(Particle->GetProxy() == nullptr);	//not already owned by another proxy. TODO: use TUniquePtr
 	auto Proxy = new FSingleParticlePhysicsProxy(MoveTemp(Particle), nullptr);
 	return Proxy;
+}
 }
\ No newline at end of file
diff --git a/Engine/Source/Runtime/Experimental/Chaos/Public/PhysicsProxy/SingleParticlePhysicsProxyFwd.h b/Engine/Source/Runtime/Experimental/Chaos/Public/PhysicsProxy/SingleParticlePhysicsProxyFwd.h
--- a/Engine/Source/Runtime/Experimental/Chaos/Public/PhysicsProxy/SingleParticlePhysicsProxyFwd.h
+++ b/Engine/Source/Runtime/Experimental/Chaos/Public/PhysicsProxy/SingleParticlePhysicsProxyFwd.h
@@ -15,6 +15,6 @@ namespace Chaos
 
 	class FRigidBodyHandle_External;
 	class FRigidBodyHandle_Internal;
-}
 
-class FSingleParticlePhysicsProxy;
\ No newline at end of file
+	class FSingleParticlePhysicsProxy;
+}
\ No newline at end of file
diff --git a/Engine/Source/Runtime/Experimental/Chaos/Public/PhysicsProxy/SuspensionConstraintProxy.h b/Engine/Source/Runtime/Experimental/Chaos/Public/PhysicsProxy/SuspensionConstraintProxy.h
--- a/Engine/Source/Runtime/Experimental/Chaos/Public/PhysicsProxy/SuspensionConstraintProxy.h
+++ b/Engine/Source/Runtime/Experimental/Chaos/Public/PhysicsProxy/SuspensionConstraintProxy.h
@@ -10,133 +10,82 @@
 #include "Chaos/ParticleHandle.h"
 #include "PhysicsCoreTypes.h"
 #include "Chaos/Defines.h"
+#include "Chaos/PBDConstraintBaseData.h"
 
 namespace Chaos
 {
 	class FSuspensionConstraint;
 
 	class FPBDRigidsEvolutionGBF;
-}
 
-class CHAOS_API FSuspensionConstraintPhysicsProxy : public TPhysicsProxy<FSuspensionConstraintPhysicsProxy,void>
+class CHAOS_API FSuspensionConstraintPhysicsProxy : public IPhysicsProxyBase
 {
-	typedef TPhysicsProxy<FSuspensionConstraintPhysicsProxy, void> Base;
-
 public:
-	using FReal = Chaos::FReal;
-	using FConstraintHandle = typename Chaos::FSuspensionConstraint::FHandle;
-	using FConstraintData = typename Chaos::FSuspensionConstraint::FData;
-	using FSuspensionConstraints = Chaos::FPBDSuspensionConstraints;
-	using FSuspensionConstraintDirtyFlags = Chaos::FSuspensionConstraintDirtyFlags;
-	using FParticlePair = Chaos::TVector<Chaos::FGeometryParticle*, 2>;
-	using FParticleHandlePair = Chaos::TVector<Chaos::FGeometryParticleHandle*, 2>;
-
-	FSuspensionConstraintPhysicsProxy() = delete;
-	FSuspensionConstraintPhysicsProxy(Chaos::FSuspensionConstraint* InConstraint, FConstraintHandle* InHandle, UObject* InOwner = nullptr);
-
-	EPhysicsProxyType ConcreteType() { return EPhysicsProxyType::SuspensionConstraintType; }
+	using Base = IPhysicsProxyBase;
 	
-	bool IsValid() { return Constraint != nullptr && Constraint->IsValid(); }
+	FSuspensionConstraintPhysicsProxy() = delete;
+	FSuspensionConstraintPhysicsProxy(FSuspensionConstraint* InConstraint, FPBDSuspensionConstraintHandle* InHandle, UObject* InOwner = nullptr);
 
 	bool IsInitialized() const { return bInitialized; }
 	void SetInitialized() { bInitialized = true; }
 
-	static Chaos::FGeometryParticleHandle* GetParticleHandleFromProxy(IPhysicsProxyBase* ProxyBase);
+	static FGeometryParticleHandle* GetParticleHandleFromProxy(IPhysicsProxyBase* ProxyBase);
 
 	//
 	//  Lifespan Management
 	//
 
-	void InitializeOnPhysicsThread(Chaos::FPBDRigidsSolver* InSolver);
+	void InitializeOnPhysicsThread(FPBDRigidsSolver* InSolver, FDirtyPropertiesManager& Manager, int32 DataIdx, FDirtyChaosProperties& RemoteData);
 
 	// Merge to perform a remote sync
-	void PushStateOnGameThread(Chaos::FPBDRigidsSolver* InSolver);
+	void PushStateOnGameThread(FDirtyPropertiesManager& Manager, int32 DataIdx, FDirtyChaosProperties& RemoteData);
 
-	void PushStateOnPhysicsThread(Chaos::FPBDRigidsSolver* InSolver);
+	void PushStateOnPhysicsThread(FPBDRigidsSolver* InSolver, const FDirtyPropertiesManager& Manager, int32 DataIdx, const FDirtyChaosProperties& RemoteData);
 	// Merge to perform a remote sync - END
 
-	void DestroyOnPhysicsThread(Chaos::FPBDRigidsSolver* InSolver);
-
-	void UpdateTargetOnPhysicsThread(Chaos::FPBDRigidsSolver* InSolver, const FVector& TargetPos, bool Enabled);
+	void DestroyOnGameThread();
+	void DestroyOnPhysicsThread(FPBDRigidsSolver* InSolver);
 
-	void SyncBeforeDestroy() {}
-	void OnRemoveFromScene() {}
+	void UpdateTargetOnPhysicsThread(FPBDRigidsSolver* InSolver, const FVector& TargetPos, bool Enabled);
 
 
 	//
 	// Member Access
 	//
 
-	FConstraintHandle* GetHandle()
+	FPBDSuspensionConstraintHandle* GetHandle()
 	{
-		return Handle;
+		return Constraint_PT;
 	}
 
-	const FConstraintHandle* GetHandle() const
+	const FPBDSuspensionConstraintHandle* GetHandle() const
 	{
-		return Handle;
+		return Constraint_PT;
 	}
 
 	virtual void* GetHandleUnsafe() const override
 	{
-		return Handle;
+		return Constraint_PT;
 	}
 
-	void SetHandle(FConstraintHandle* InHandle)
+	void SetHandle(FPBDSuspensionConstraintHandle* InHandle)
 	{
-		Handle = InHandle;
+		Constraint_PT = InHandle;
 	}
 
-	Chaos::FSuspensionConstraint* GetConstraint()
+	FSuspensionConstraint* GetConstraint()
 	{
-		return Constraint;
+		return Constraint_GT;
 	}
 
-	const Chaos::FSuspensionConstraint* GetConstraint() const
+	const FSuspensionConstraint* GetConstraint() const
 	{
-		return Constraint;
+		return Constraint_GT;
 	}
-
-	//
-	// Threading API
-	//
-
-	/**/
-	void FlipBuffer() { }
-
-	void PushToPhysicsState(Chaos::FPBDRigidsEvolutionGBF& Evolution) {}
-
-	/**/
-	void ClearAccumulatedData() {}
-
-	/**/
-	void BufferPhysicsResults() {}
-
-	/**/
-	void PullFromPhysicsState() {}
-
-	/**/
-	bool IsDirty() { return Constraint->IsDirty(); }
 	
 private:
-
-	FConstraintData SuspensionSettingsBuffer;
-	FSuspensionConstraintDirtyFlags DirtyFlagsBuffer;
-
-	Chaos::FSuspensionConstraint* Constraint;
-	FConstraintHandle* Handle;
+	FSuspensionConstraint* Constraint_GT;
+	FPBDSuspensionConstraintHandle* Constraint_PT;
 	bool bInitialized;
-
-
-public:
-	void AddForceCallback(FParticlesType& InParticles, const FReal InDt, const int32 InIndex) {}
-	void DisableCollisionsCallback(TSet<TTuple<int32, int32>>& InPairs) {}
-	void BindParticleCallbackMapping(Chaos::TArrayCollectionArray<PhysicsProxyWrapper>& PhysicsProxyReverseMap, Chaos::TArrayCollectionArray<int32>& ParticleIDReverseMap) {}
-	void EndFrameCallback(const FReal InDt) {}
-	void ParameterUpdateCallback(FParticlesType& InParticles, const FReal InTime) {}
-	void CreateRigidBodyCallback(FParticlesType& InOutParticles) {}
-	bool IsSimulating() const { return true; }
-	void UpdateKinematicBodiesCallback(const FParticlesType& InParticles, const FReal InDt, const FReal InTime, FKinematicProxy& InKinematicProxy) {}
-	void StartFrameCallback(const FReal InDt, const FReal InTime) {}
-
 };
+}
\ No newline at end of file
diff --git a/Engine/Source/Runtime/Experimental/Chaos/Public/RewindData.h b/Engine/Source/Runtime/Experimental/Chaos/Public/RewindData.h
--- a/Engine/Source/Runtime/Experimental/Chaos/Public/RewindData.h
+++ b/Engine/Source/Runtime/Experimental/Chaos/Public/RewindData.h
@@ -475,7 +475,7 @@ private:
 	decltype(auto) Get##FUNC_NAME() const\
 	{\
 		const auto Data = State ? State->PROP.Read(FrameAndPhase, Pool) : nullptr;\
-		return Data ? Data->NAME : Head.GetSettings().NAME;\
+		return Data ? Data->NAME : Head.Get##PROP().NAME;\
 	}\
 
 inline int32 ComputeCircularSize(int32 NumFrames) { return NumFrames * FFrameAndPhase::NumPhases; }
@@ -652,7 +652,7 @@ struct FJointStateBase
 {
 	explicit FJointStateBase(int32 NumFrames)
 		: JointSettings(ComputeCircularSize(NumFrames))
-		, ProxyPair(ComputeCircularSize(NumFrames))
+		, JointProxies(ComputeCircularSize(NumFrames))
 	{
 	}
 
@@ -663,31 +663,31 @@ struct FJointStateBase
 	void Release(FDirtyPropertiesPool& Manager)
 	{
 		JointSettings.Release(Manager);
-		ProxyPair.Release(Manager);
+		JointProxies.Release(Manager);
 	}
 
 	void Reset()
 	{
 		JointSettings.Reset();
-		ProxyPair.Reset();
+		JointProxies.Reset();
 	}
 
 	void ClearEntryAndFuture(const FFrameAndPhase FrameAndPhase)
 	{
 		JointSettings.ClearEntryAndFuture(FrameAndPhase);
-		ProxyPair.ClearEntryAndFuture(FrameAndPhase);
+		JointProxies.ClearEntryAndFuture(FrameAndPhase);
 	}
 
 	bool IsClean(const FFrameAndPhase FrameAndPhase) const
 	{
-		return JointSettings.IsClean(FrameAndPhase) && ProxyPair.IsClean(FrameAndPhase);
+		return JointSettings.IsClean(FrameAndPhase) && JointProxies.IsClean(FrameAndPhase);
 	}
 
 	template <bool bSkipDynamics>
 	bool IsInSync(const FPBDJointConstraintHandle& Handle, const FFrameAndPhase FrameAndPhase, const FDirtyPropertiesPool& Pool) const;
 
 	TParticlePropertyBuffer<FPBDJointSettings, EChaosProperty::JointSettings> JointSettings;
-	TParticlePropertyBuffer<FProxyBasePair, EChaosProperty::JointParticleProxies> ProxyPair;
+	TParticlePropertyBuffer<FProxyBasePairProperty, EChaosProperty::JointParticleProxies> JointProxies;
 };
 
 class FJointState
diff --git a/Engine/Source/Runtime/Experimental/ChaosSolverEngine/Private/Chaos/ChaosSolverActor.cpp b/Engine/Source/Runtime/Experimental/ChaosSolverEngine/Private/Chaos/ChaosSolverActor.cpp
--- a/Engine/Source/Runtime/Experimental/ChaosSolverEngine/Private/Chaos/ChaosSolverActor.cpp
+++ b/Engine/Source/Runtime/Experimental/ChaosSolverEngine/Private/Chaos/ChaosSolverActor.cpp
@@ -411,7 +411,7 @@ void AChaosSolverActor::MakeFloor()
 		FilterData.Word1 = 0xFFFF;
 		FilterData.Word3 = 0xFFFF;
 		FloorParticle->SetShapeSimData(0, FilterData);
-		Proxy = FSingleParticlePhysicsProxy::Create(MoveTemp(FloorParticle));
+		Proxy = Chaos::FSingleParticlePhysicsProxy::Create(MoveTemp(FloorParticle));
 		Solver->RegisterObject(Proxy);
 	}
 }
diff --git a/Engine/Source/Runtime/Experimental/ChaosSolverEngine/Public/Chaos/ChaosSolverActor.h b/Engine/Source/Runtime/Experimental/ChaosSolverEngine/Public/Chaos/ChaosSolverActor.h
--- a/Engine/Source/Runtime/Experimental/ChaosSolverEngine/Public/Chaos/ChaosSolverActor.h
+++ b/Engine/Source/Runtime/Experimental/ChaosSolverEngine/Public/Chaos/ChaosSolverActor.h
@@ -170,5 +170,5 @@ private:
 	TObjectPtr<UChaosGameplayEventDispatcher> GameplayEventDispatcherComponent;
 
 	/** If floor is enabled - this will point to the solver particle for it */
-	FSingleParticlePhysicsProxy* Proxy;
+	Chaos::FSingleParticlePhysicsProxy* Proxy;
 };
\ No newline at end of file
diff --git a/Engine/Source/Runtime/PhysicsCore/Private/ChaosEngineInterface.cpp b/Engine/Source/Runtime/PhysicsCore/Private/ChaosEngineInterface.cpp
--- a/Engine/Source/Runtime/PhysicsCore/Private/ChaosEngineInterface.cpp
+++ b/Engine/Source/Runtime/PhysicsCore/Private/ChaosEngineInterface.cpp
@@ -1050,8 +1050,8 @@ FPhysicsConstraintHandle FChaosEngineInterface::CreateSuspension(const FPhysicsA
 				auto* SuspensionConstraint = new Chaos::FSuspensionConstraint();
 				ConstraintRef.Constraint = SuspensionConstraint;
 
-				SuspensionConstraint->SetParticleProxies({ InActorRef,nullptr });
-			SuspensionConstraint->SetLocation(InLocalFrame);
+				SuspensionConstraint->SetParticleProxy(InActorRef);
+				SuspensionConstraint->SetLocation(InLocalFrame);
 
 				Chaos::FPhysicsSolver* Solver = InActorRef->GetSolver<Chaos::FPhysicsSolver>();
 				Solver->RegisterObject(SuspensionConstraint);
@@ -1145,7 +1145,7 @@ GetParticleFromProxy(IPhysicsProxyBase* ProxyBase)
 	{
 		if (ProxyBase->GetType() == EPhysicsProxyType::SingleParticleProxy)
 		{
-			return ((FSingleParticlePhysicsProxy*)ProxyBase)->GetParticle_LowLevel();
+			return ((Chaos::FSingleParticlePhysicsProxy*)ProxyBase)->GetParticle_LowLevel();
 		}
 	}
 	return nullptr;
@@ -1648,7 +1648,7 @@ void FChaosEngineInterface::CreateActor(const FActorCreationParams& InParams,FPh
 		Particle = MoveTemp(Rigid);
 	}
 
-	Handle = FSingleParticlePhysicsProxy::Create(MoveTemp(Particle));
+	Handle = Chaos::FSingleParticlePhysicsProxy::Create(MoveTemp(Particle));
 	Chaos::FRigidBodyHandle_External& Body_External = Handle->GetGameThreadAPI();
 
 	// Set up the new particle's game-thread data. This will be sent to physics-thread when
diff --git a/Engine/Source/Runtime/PhysicsCore/Public/PhysicsInterfaceDeclaresCore.h b/Engine/Source/Runtime/PhysicsCore/Public/PhysicsInterfaceDeclaresCore.h
--- a/Engine/Source/Runtime/PhysicsCore/Public/PhysicsInterfaceDeclaresCore.h
+++ b/Engine/Source/Runtime/PhysicsCore/Public/PhysicsInterfaceDeclaresCore.h
@@ -205,7 +205,7 @@ extern FPhysicsActorHandle DefaultPhysicsActorHandle();
 
 #elif WITH_CHAOS
 
-using FPhysicsActorHandle = FSingleParticlePhysicsProxy*;
+using FPhysicsActorHandle = Chaos::FSingleParticlePhysicsProxy*;
 
 class FChaosSceneId;
 class FPhysInterface_Chaos;
