Provide rewind callback with debug info about resim time (and other info in the future)

diff --git a/Engine/Source/Runtime/Experimental/Chaos/Private/PBDRigidsSolver.cpp b/Engine/Source/Runtime/Experimental/Chaos/Private/PBDRigidsSolver.cpp
--- a/Engine/Source/Runtime/Experimental/Chaos/Private/PBDRigidsSolver.cpp
+++ b/Engine/Source/Runtime/Experimental/Chaos/Private/PBDRigidsSolver.cpp
@@ -1155,6 +1155,7 @@ namespace Chaos
 			const int32 ResimStep = MRewindCallback->TriggerRewindIfNeeded_Internal(LastStep);
 			if(ResimStep != INDEX_NONE)
 			{
+				FResimDebugInfo DebugInfo;
 				QUICK_SCOPE_CYCLE_COUNTER(ChaosRewindAndResim);
 				if(ensure(MRewindData->RewindToFrame(ResimStep)))
 				{
@@ -1163,6 +1164,8 @@ namespace Chaos
 					const int32 NumResimSteps = LastStep - ResimStep + 1;
 					TArray<FPushPhysicsData*> RecordedPushData = MarshallingManager.StealHistory_Internal(NumResimSteps);
 					bool bFirst = true;
+
+					FDurationTimer ResimTimer(DebugInfo.ResimTime);
 					// Do rollback as necessary
 					for (int32 Step = ResimStep; Step <= LastStep; ++Step)
 					{
@@ -1181,8 +1184,12 @@ namespace Chaos
 					}
 
 					GetEvolution()->SetResim(false);
+
+					ResimTimer.Stop();
+					MRewindCallback->SetResimDebugInfo(DebugInfo);
 				}
 			}
+
 		}
 	}
 
diff --git a/Engine/Source/Runtime/Experimental/Chaos/Public/RewindData.h b/Engine/Source/Runtime/Experimental/Chaos/Public/RewindData.h
--- a/Engine/Source/Runtime/Experimental/Chaos/Public/RewindData.h
+++ b/Engine/Source/Runtime/Experimental/Chaos/Public/RewindData.h
@@ -1073,6 +1073,11 @@ private:
 	void DesyncIfNecessary(TDirtyInfo& Info, const FFrameAndPhase FrameAndPhase);
 };
 
+struct FResimDebugInfo
+{
+	double ResimTime = 0.0;
+};
+
 /** Used by user code to determine when rewind should occur and gives it the opportunity to record any additional data */
 class IRewindCallback
 {
@@ -1106,5 +1111,8 @@ public:
 	virtual void PostResimStep_Internal(int32 PhysicsStep){}
 
 	virtual void RegisterRewindableSimCallback_Internal(ISimCallbackObject* Callback) { ensure(false); }
+
+	/** Called When resim is finished with debug information about the resim */
+	virtual void SetResimDebugInfo(const FResimDebugInfo& ResimDebugInfo){}
 };
 }
