Make sure debug info callback is named with _Internal to indicate which thread it is triggered on

diff --git a/Engine/Source/Runtime/Experimental/Chaos/Private/PBDRigidsSolver.cpp b/Engine/Source/Runtime/Experimental/Chaos/Private/PBDRigidsSolver.cpp
--- a/Engine/Source/Runtime/Experimental/Chaos/Private/PBDRigidsSolver.cpp
+++ b/Engine/Source/Runtime/Experimental/Chaos/Private/PBDRigidsSolver.cpp
@@ -1186,7 +1186,7 @@ namespace Chaos
 					GetEvolution()->SetResim(false);
 
 					ResimTimer.Stop();
-					MRewindCallback->SetResimDebugInfo(DebugInfo);
+					MRewindCallback->SetResimDebugInfo_Internal(DebugInfo);
 				}
 			}
 
diff --git a/Engine/Source/Runtime/Experimental/Chaos/Public/RewindData.h b/Engine/Source/Runtime/Experimental/Chaos/Public/RewindData.h
--- a/Engine/Source/Runtime/Experimental/Chaos/Public/RewindData.h
+++ b/Engine/Source/Runtime/Experimental/Chaos/Public/RewindData.h
@@ -1113,6 +1113,6 @@ public:
 	virtual void RegisterRewindableSimCallback_Internal(ISimCallbackObject* Callback) { ensure(false); }
 
 	/** Called When resim is finished with debug information about the resim */
-	virtual void SetResimDebugInfo(const FResimDebugInfo& ResimDebugInfo){}
+	virtual void SetResimDebugInfo_Internal(const FResimDebugInfo& ResimDebugInfo){}
 };
 }
