Reduce memory foot print for Trimesh implicit object when using LWC
- change BVH to use FRealSingle instead of FReal
- added T parameter to AABBTree template
- changes TAABB, TBox and TPlane accordingly to compile

diff --git a/Engine/Source/Runtime/Experimental/Chaos/Private/Chaos/AABB.cpp b/Engine/Source/Runtime/Experimental/Chaos/Private/Chaos/AABB.cpp
--- a/Engine/Source/Runtime/Experimental/Chaos/Private/Chaos/AABB.cpp
+++ b/Engine/Source/Runtime/Experimental/Chaos/Private/Chaos/AABB.cpp
@@ -10,27 +10,27 @@
 namespace Chaos
 {
 template <typename T, int d>
-bool TAABB<T, d>::Raycast(const TVector<T, d>& StartPoint, const TVector<T, d>& Dir, const T Length, const T Thickness, T& OutTime, TVector<T, d>& OutPosition, TVector<T, d>& OutNormal, int32& OutFaceIndex) const
+bool TAABB<T, d>::Raycast(const TVector<FReal, d>& StartPoint, const TVector<FReal, d>& Dir, const FReal Length, const FReal Thickness, FReal& OutTime, TVector<FReal, d>& OutPosition, TVector<FReal, d>& OutNormal, int32& OutFaceIndex) const
 {
 	ensure(Length > 0);
 	ensure(FMath::IsNearlyEqual(Dir.SizeSquared(), (FReal)1, (FReal)KINDA_SMALL_NUMBER));
 
 	OutFaceIndex = INDEX_NONE;
-	const TVector<T, d> MinInflated = MMin - Thickness;
-	const TVector<T,d> StartToMin = MinInflated - StartPoint;
+	const TVector<FReal, d> MinInflated = TVector<FReal, d>(MMin) - Thickness;
+	const TVector<FReal,d> StartToMin = MinInflated - StartPoint;
 	
-	const TVector<T, d> MaxInflated = MMax + Thickness;
-	const TVector<T, d> StartToMax = MaxInflated - StartPoint;
+	const TVector<FReal, d> MaxInflated = TVector<FReal, d>(MMax) + Thickness;
+	const TVector<FReal, d> StartToMax = MaxInflated - StartPoint;
 
 	//For each axis record the start and end time when ray is in the box. If the intervals overlap the ray is inside the box
-	T LatestStartTime = 0;
-	T EarliestEndTime = FLT_MAX;
-	TVector<T, d> Normal(0);	//not needed but fixes compiler warning
+	FReal LatestStartTime = 0;
+	FReal EarliestEndTime = TNumericLimits<FReal>::Max();
+	TVector<FReal, d> Normal(0);	//not needed but fixes compiler warning
 
 	for (int Axis = 0; Axis < d; ++Axis)
 	{
 		const bool bParallel = FMath::IsNearlyZero(Dir[Axis]);
-		T Time1, Time2;
+		FReal Time1, Time2;
 		if (bParallel)
 		{
 			if (StartToMin[Axis] > 0 || StartToMax[Axis] < 0)
@@ -45,12 +45,12 @@ bool TAABB<T, d>::Raycast(const TVector<T, d>& StartPoint, const TVector<T, d>&
 		}
 		else
 		{
-			const T InvDir = (T)1 / Dir[Axis];
+			const FReal InvDir = (FReal)1 / Dir[Axis];
 			Time1 = StartToMin[Axis] * InvDir;
 			Time2 = StartToMax[Axis] * InvDir;
 		}
 
-		TVector<T, d> CurNormal = TVector<T, d>::AxisVector(Axis);
+		TVector<FReal, d> CurNormal = TVector<FReal, d>::AxisVector(Axis);
 
 		if (Time1 > Time2)
 		{
@@ -84,36 +84,36 @@ bool TAABB<T, d>::Raycast(const TVector<T, d>& StartPoint, const TVector<T, d>&
 		return false;
 	}
 	
-	const TVector<T, d> BoxIntersection = StartPoint + LatestStartTime * Dir;
+	const TVector<FReal, d> BoxIntersection = StartPoint + LatestStartTime * Dir;
 
 	//If the box is rounded we have to consider corners and edges.
 	//Break the box into voronoi regions based on features (corner, edge, face) and see which region the raycast hit
 
-	if (Thickness != (T)0)
+	if (Thickness != (FReal)0)
 	{
 		check(d == 3);
-		TVector<T, d> GeomStart;
-		TVector<T, d> GeomEnd;
+		TVector<FReal, d> GeomStart;
+		TVector<FReal, d> GeomEnd;
 		int32 NumAxes = 0;
 
 		for (int Axis = 0; Axis < d; ++Axis)
 		{
-			if (BoxIntersection[Axis] < MMin[Axis])
+			if (BoxIntersection[Axis] < (FReal)MMin[Axis])
 			{
-				GeomStart[Axis] = MMin[Axis];
-				GeomEnd[Axis] = MMin[Axis];
+				GeomStart[Axis] = (FReal)MMin[Axis];
+				GeomEnd[Axis] = (FReal)MMin[Axis];
 				++NumAxes;
 			}
-			else if (BoxIntersection[Axis] > MMax[Axis])
+			else if (BoxIntersection[Axis] > (FReal)MMax[Axis])
 			{
-				GeomStart[Axis] = MMax[Axis];
-				GeomEnd[Axis] = MMax[Axis];
+				GeomStart[Axis] = (FReal)MMax[Axis];
+				GeomEnd[Axis] = (FReal)MMax[Axis];
 				++NumAxes;
 			}
 			else
 			{
-				GeomStart[Axis] = MMin[Axis];
-				GeomEnd[Axis] = MMax[Axis];
+				GeomStart[Axis] = (FReal)MMin[Axis];
+				GeomEnd[Axis] = (FReal)MMax[Axis];
 			}
 		}
 
@@ -123,16 +123,16 @@ bool TAABB<T, d>::Raycast(const TVector<T, d>& StartPoint, const TVector<T, d>&
 			if (NumAxes == 3)
 			{
 				//hit a corner. For now just use 3 capsules, there's likely a better way to determine which capsule is needed
-				T CornerTimes[3];
-				TVector<T, d> CornerPositions[3];
-				TVector<T, d> CornerNormals[3];
+				FReal CornerTimes[3];
+				TVector<FReal, d> CornerPositions[3];
+				TVector<FReal, d> CornerNormals[3];
 				int32 HitIdx = INDEX_NONE;
-				T MinTime = 0;	//initialization just here for compiler warning
+				FReal MinTime = 0;	//initialization just here for compiler warning
 				for (int CurIdx = 0; CurIdx < 3; ++CurIdx)
 				{
-					TVector<T, d> End = GeomStart;
+					TVector<FReal, d> End = GeomStart;
 					End[CurIdx] = End[CurIdx] == MMin[CurIdx] ? MMax[CurIdx] : MMin[CurIdx];
-					TCapsule<T> Capsule(GeomStart, End, Thickness);
+					FCapsule Capsule(GeomStart, End, Thickness);
 					if (Capsule.Raycast(StartPoint, Dir, Length, 0, CornerTimes[CurIdx], CornerPositions[CurIdx], CornerNormals[CurIdx], OutFaceIndex))
 					{
 						if (HitIdx == INDEX_NONE || CornerTimes[CurIdx] < MinTime)
@@ -160,7 +160,7 @@ bool TAABB<T, d>::Raycast(const TVector<T, d>& StartPoint, const TVector<T, d>&
 			else
 			{
 				//capsule: todo(use a cylinder which is cheaper. Our current cylinder raycast implementation doesn't quite work for this setup)
-				TCapsule<T> CapsuleBorder(GeomStart, GeomEnd, Thickness);
+				FCapsule CapsuleBorder(GeomStart, GeomEnd, Thickness);
 				bHit = CapsuleBorder.Raycast(StartPoint, Dir, Length, 0, OutTime, OutPosition, OutNormal, OutFaceIndex);
 			}
 
@@ -197,7 +197,8 @@ inline TAABB<T, d> TransformedAABBHelper(const TAABB<T, d>& AABB, const TTRANSFO
 	return NewAABB;
 }
 
-inline TAABB<FReal, 3> TransformedAABBHelperISPC(const TAABB<FReal, 3>& AABB, const FTransform& SpaceTransform)
+template<typename T>
+inline TAABB<T, 3> TransformedAABBHelperISPC(const TAABB<T, 3>& AABB, const FTransform& SpaceTransform)
 {
 	check(bRealTypeCompatibleWithISPC);
 #if INTEL_ISPC
@@ -208,7 +209,7 @@ inline TAABB<FReal, 3> TransformedAABBHelperISPC(const TAABB<FReal, 3>& AABB, co
 	return NewAABB;
 #else
 	check(false);
-	return TAABB<FReal, 3>::EmptyAABB();
+	return TAABB<T, 3>::EmptyAABB();
 #endif
 }
 
@@ -252,7 +253,7 @@ TAABB<T, d> TAABB<T, d>::TransformedAABB(const FTransform& SpaceTransform) const
 }
 
 template<typename T, int d>
-inline TAABB<T, d> TAABB<T, d>::InverseTransformedAABB(const TRigidTransform<T, 3>& SpaceTransform) const
+inline TAABB<T, d> TAABB<T, d>::InverseTransformedAABB(const FRigidTransform3& SpaceTransform) const
 {
 	TVector<T, d> CurrentExtents = Extents();
 	int32 Idx = 0;
@@ -270,4 +271,8 @@ inline TAABB<T, d> TAABB<T, d>::InverseTransformedAABB(const TRigidTransform<T,
 }
 }
 
-template class Chaos::TAABB<Chaos::FReal, 3>;
+template class Chaos::TAABB<Chaos::FRealSingle, 3>;
+
+#if !UE_LARGE_WORLD_COORDINATES_DISABLED
+template class Chaos::TAABB<Chaos::FRealDouble, 3>;
+#endif
diff --git a/Engine/Source/Runtime/Experimental/Chaos/Public/Chaos/AABB.h b/Engine/Source/Runtime/Experimental/Chaos/Public/Chaos/AABB.h
--- a/Engine/Source/Runtime/Experimental/Chaos/Public/Chaos/AABB.h
+++ b/Engine/Source/Runtime/Experimental/Chaos/Public/Chaos/AABB.h
@@ -90,13 +90,14 @@ namespace Chaos
 		CHAOS_API TAABB<T, d> TransformedAABB(const FMatrix&) const;
 		CHAOS_API TAABB<T, d> TransformedAABB(const Chaos::PMatrix<FReal, 4, 4>&) const;
 
-		CHAOS_API TAABB<T, d> InverseTransformedAABB(const Chaos::TRigidTransform<T, 3>&) const;
+		CHAOS_API TAABB<T, d> InverseTransformedAABB(const Chaos::FRigidTransform3&) const;
 
-		FORCEINLINE bool Intersects(const TAABB<T, d>& Other) const
+		template <typename TReal>
+		FORCEINLINE bool Intersects(const TAABB<TReal, d>& Other) const
 		{
 			for (int32 i = 0; i < d; ++i)
 			{
-				if (Other.MMax[i] < MMin[i] || Other.MMin[i] > MMax[i])
+				if (Other.Max()[i] < MMin[i] || Other.Min()[i] > MMax[i])
 					return false;
 			}
 			return true;
@@ -136,21 +137,22 @@ namespace Chaos
 
 		FORCEINLINE uint16 GetMaterialIndex(uint32 HintIndex) const { return 0; }
 
-		FORCEINLINE T SignedDistance(const TVector<T, d>& x) const
+		FORCEINLINE FReal SignedDistance(const TVector<FReal, d>& x) const
 		{
-			TVector<T, d> Normal;
+			TVector<FReal, d> Normal;
 			return PhiWithNormal(x, Normal);
 		}
 
-
-		FORCEINLINE T PhiWithNormal(const TVector<T, d>& X, TVector<T, d>& Normal) const 
+		FORCEINLINE FReal PhiWithNormal(const TVector<FReal, d>& X, TVector<FReal, d>& Normal) const
 		{
-			const TVector<T, d> MaxDists = X - MMax;
-			const TVector<T, d> MinDists = MMin - X;
-			if (X <= MMax && X >= MMin)
+			const TVector<FReal, d> MaxDists = X - MMax;
+			const TVector<FReal, d> MinDists = MMin - X;
+			const TVector<FReal, d> MinAsTVecReal(MMin);
+			const TVector<FReal, d> MaxAsTVecReal(MMax);
+			if (X <= MaxAsTVecReal && X >= MinAsTVecReal)
 			{
-				const Pair<T, int32> MaxAndAxis = TVector<T, d>::MaxAndAxis(MinDists, MaxDists);
-				Normal = MaxDists[MaxAndAxis.Second] > MinDists[MaxAndAxis.Second] ? TVector<T, d>::AxisVector(MaxAndAxis.Second) : -TVector<T, d>::AxisVector(MaxAndAxis.Second);
+				const Pair<FReal, int32> MaxAndAxis = TVector<FReal, d>::MaxAndAxis(MinDists, MaxDists);
+				Normal = MaxDists[MaxAndAxis.Second] > MinDists[MaxAndAxis.Second] ? TVector<FReal, d>::AxisVector(MaxAndAxis.Second) : -TVector<FReal, d>::AxisVector(MaxAndAxis.Second);
 				return MaxAndAxis.First;
 			}
 			else
@@ -170,7 +172,7 @@ namespace Chaos
 						Normal[i] = 0;
 					}
 				}
-				T Phi = Normal.SafeNormalize();
+				FReal Phi = Normal.SafeNormalize();
 				if (Phi < KINDA_SMALL_NUMBER)
 				{
 					for (int i = 0; i < d; ++i)
@@ -190,21 +192,21 @@ namespace Chaos
 			}
 		}
 
-		bool CHAOS_API Raycast(const TVector<T, d>& StartPoint, const TVector<T, d>& Dir, const T Length, const T Thickness, T& OutTime, TVector<T, d>& OutPosition, TVector<T, d>& OutNormal, int32& OutFaceIndex) const;
+		bool CHAOS_API Raycast(const TVector<FReal, d>& StartPoint, const TVector<FReal, d>& Dir, const FReal Length, const FReal Thickness, FReal& OutTime, TVector<FReal, d>& OutPosition, TVector<FReal, d>& OutNormal, int32& OutFaceIndex) const;
 
 
-		FORCEINLINE bool RaycastFast(const TVector<T, d>& StartPoint, const TVector<T, d>& Dir, const TVector<T, d>& InvDir, const bool* bParallel, const T Length, const T InvLength, T& OutTime, TVector<T, d>& OutPosition) const
+		FORCEINLINE bool RaycastFast(const TVector<FReal, d>& StartPoint, const TVector<FReal, d>& Dir, const TVector<FReal, d>& InvDir, const bool* bParallel, const FReal Length, const FReal InvLength, FReal& OutTime, TVector<FReal, d>& OutPosition) const
 		{
-			const TVector<T, d> StartToMin = MMin - StartPoint;
-			const TVector<T, d> StartToMax = MMax - StartPoint;
+			const TVector<FReal, d> StartToMin = TVector<FReal, d>(MMin) - StartPoint;
+			const TVector<FReal, d> StartToMax = TVector<FReal, d>(MMax) - StartPoint;
 
 			//For each axis record the start and end time when ray is in the box. If the intervals overlap the ray is inside the box
-			T LatestStartTime = 0;
-			T EarliestEndTime = FLT_MAX;
+			FReal LatestStartTime = 0;
+			FReal EarliestEndTime = TNumericLimits<FReal>::Max();
 
 			for (int Axis = 0; Axis < d; ++Axis)
 			{
-				T Time1, Time2;
+				FReal Time1, Time2;
 				if (bParallel[Axis])
 				{
 					if (StartToMin[Axis] > 0 || StartToMax[Axis] < 0)
@@ -214,7 +216,7 @@ namespace Chaos
 					else
 					{
 						Time1 = 0;
-						Time2 = FLT_MAX;
+						Time2 = TNumericLimits<FReal>::Max();
 					}
 				}
 				else
@@ -281,13 +283,13 @@ namespace Chaos
 				for (int32 i = 0; i < d; ++i)
 				{
 					auto PlaneIntersection = TPlane<T, d>(MMin - Thickness, -TVector<T, d>::AxisVector(i)).FindClosestPoint(Result, 0);
-					Intersections.Add(MakePair((FReal)(PlaneIntersection - Result).Size(), -TVector<T, d>::AxisVector(i)));
+					Intersections.Add(MakePair((T)(PlaneIntersection - Result).Size(), -TVector<T, d>::AxisVector(i)));
 					PlaneIntersection = TPlane<T, d>(MMax + Thickness, TVector<T, d>::AxisVector(i)).FindClosestPoint(Result, 0);
-					Intersections.Add(MakePair((FReal)(PlaneIntersection - Result).Size(), TVector<T, d>::AxisVector(i)));
+					Intersections.Add(MakePair((T)(PlaneIntersection - Result).Size(), TVector<T, d>::AxisVector(i)));
 				}
 				Intersections.Sort([](const Pair<T, TVector<T, d>>& Elem1, const Pair<T, TVector<T, d>>& Elem2) { return Elem1.First < Elem2.First; });
 
-				if (!FMath::IsNearlyEqual(Intersections[0].First, (FReal)0.))
+				if (!FMath::IsNearlyEqual(Intersections[0].First, (T)0.))
 				{
 					T SmallestDistance = Intersections[0].First;
 					Result += Intersections[0].Second * Intersections[0].First;
@@ -300,19 +302,19 @@ namespace Chaos
 			return Result;
 		}
 
-		FORCEINLINE Pair<TVector<T, d>, bool> FindClosestIntersectionImp(const TVector<T, d>& StartPoint, const TVector<T, d>& EndPoint, const T Thickness) const 
+		FORCEINLINE Pair<TVector<FReal, d>, bool> FindClosestIntersectionImp(const TVector<FReal, d>& StartPoint, const TVector<FReal, d>& EndPoint, const FReal Thickness) const
 		{
-			TArray<Pair<T, TVector<T, d>>> Intersections;
+			TArray<Pair<FReal, TVector<FReal, d>>> Intersections;
 			for (int32 i = 0; i < d; ++i)
 			{
-				auto PlaneIntersection = TPlane<T, d>(MMin - Thickness, -TVector<T, d>::AxisVector(i)).FindClosestIntersection(StartPoint, EndPoint, 0);
+				auto PlaneIntersection = TPlane<FReal, d>(TVector<FReal, d>(MMin) - Thickness, -TVector<FReal, d>::AxisVector(i)).FindClosestIntersection(StartPoint, EndPoint, 0);
 				if (PlaneIntersection.Second)
 					Intersections.Add(MakePair((FReal)(PlaneIntersection.First - StartPoint).Size(), PlaneIntersection.First));
-				PlaneIntersection = TPlane<T, d>(MMax + Thickness, TVector<T, d>::AxisVector(i)).FindClosestIntersection(StartPoint, EndPoint, 0);
+				PlaneIntersection = TPlane<FReal, d>(TVector<FReal, d>(MMax) + Thickness, TVector<FReal, d>::AxisVector(i)).FindClosestIntersection(StartPoint, EndPoint, 0);
 				if (PlaneIntersection.Second)
 					Intersections.Add(MakePair((FReal)(PlaneIntersection.First - StartPoint).Size(), PlaneIntersection.First));
 			}
-			Intersections.Sort([](const Pair<T, TVector<T, d>>& Elem1, const Pair<T, TVector<T, d>>& Elem2) { return Elem1.First < Elem2.First; });
+			Intersections.Sort([](const Pair<FReal, TVector<FReal, d>>& Elem1, const Pair<FReal, TVector<FReal, d>>& Elem2) { return Elem1.First < Elem2.First; });
 			for (const auto& Elem : Intersections)
 			{
 				if (SignedDistance(Elem.Second) < (Thickness + 1e-4))
@@ -320,7 +322,7 @@ namespace Chaos
 					return MakePair(Elem.Second, true);
 				}
 			}
-			return MakePair(TVector<T, d>(0), false);
+			return MakePair(TVector<FReal, d>(0), false);
 		}
 
 		FORCEINLINE TVector<T, d> FindGeometryOpposingNormal(const TVector<T, d>& DenormDir, int32 FaceIndex, const TVector<T, d>& OriginalNormal) const 
@@ -384,9 +386,9 @@ namespace Chaos
 		}
 
 		// Support vertex in the specified direction, assuming each face has been moved inwards by InMargin
-		FORCEINLINE_DEBUGGABLE TVector<T, d> SupportCore(const TVector<T, d>& Direction, FReal InMargin) const
+		FORCEINLINE_DEBUGGABLE FVec3 SupportCore(const FVec3& Direction, FReal InMargin) const
 		{
-			TVector<T, d> ChosenPt;
+			FVec3 ChosenPt;
 			for (int Axis = 0; Axis < d; ++Axis)
 			{
 				ChosenPt[Axis] = Direction[Axis] < 0 ? MMin[Axis] + InMargin : MMax[Axis] - InMargin;
@@ -505,20 +507,20 @@ namespace Chaos
 			return HashCombine(UE::Math::GetTypeHash(MMin), UE::Math::GetTypeHash(MMax));
 		}
 
-		FORCEINLINE PMatrix<T, d, d> GetInertiaTensor(const T Mass) const { return GetInertiaTensor(Mass, Extents()); }
-		FORCEINLINE static PMatrix<T, 3, 3> GetInertiaTensor(const T Mass, const TVector<T, 3>& Dim)
+		FORCEINLINE PMatrix<FReal, d, d> GetInertiaTensor(const FReal Mass) const { return GetInertiaTensor(Mass, Extents()); }
+		FORCEINLINE static PMatrix<FReal, 3, 3> GetInertiaTensor(const FReal Mass, const TVector<FReal, 3>& Dim)
 		{
 			// https://www.wolframalpha.com/input/?i=cuboid
-			const T M = Mass / 12;
-			const T WW = Dim[0] * Dim[0];
-			const T HH = Dim[1] * Dim[1];
-			const T DD = Dim[2] * Dim[2];
-			return PMatrix<T, 3, 3>(M * (HH + DD), M * (WW + DD), M * (WW + HH));
+			const FReal M = Mass / 12;
+			const FReal WW = Dim[0] * Dim[0];
+			const FReal HH = Dim[1] * Dim[1];
+			const FReal DD = Dim[2] * Dim[2];
+			return PMatrix<FReal, 3, 3>(M * (HH + DD), M * (WW + DD), M * (WW + HH));
 		}
 
-		FORCEINLINE static TRotation<T, d> GetRotationOfMass()
+		FORCEINLINE static TRotation<FReal, d> GetRotationOfMass()
 		{
-			return TRotation<T, d>::FromIdentity();
+			return TRotation<FReal, d>::FromIdentity();
 		}
 
 		FORCEINLINE constexpr bool IsConvex() const { return true; }
diff --git a/Engine/Source/Runtime/Experimental/Chaos/Public/Chaos/AABBTree.h b/Engine/Source/Runtime/Experimental/Chaos/Public/Chaos/AABBTree.h
--- a/Engine/Source/Runtime/Experimental/Chaos/Public/Chaos/AABBTree.h
+++ b/Engine/Source/Runtime/Experimental/Chaos/Public/Chaos/AABBTree.h
@@ -140,18 +140,18 @@ struct TAABBTreeIntersectionHelper<FQueryFastDataVoid, EAABBQueryType::Overlap>
 	}
 };
 
-template <typename TPayloadType, bool bComputeBounds>
+template <typename TPayloadType, typename T, bool bComputeBounds>
 struct TBoundsWrapperHelper
 {
 
 };
 
-template <typename TPayloadType>
-struct TBoundsWrapperHelper<TPayloadType, true>
+template <typename TPayloadType, typename T>
+struct TBoundsWrapperHelper<TPayloadType, T, true>
 {
-	void ComputeBounds(const TArray<TPayloadBoundsElement<TPayloadType, FReal>>& Elems)
+	void ComputeBounds(const TArray<TPayloadBoundsElement<TPayloadType, T>>& Elems)
 	{
-		Bounds = FAABB3::EmptyAABB();
+		Bounds = TAABB<T, 3>::EmptyAABB();
 
 		for (const auto& Elem : Elems)
 		{
@@ -159,37 +159,37 @@ struct TBoundsWrapperHelper<TPayloadType, true>
 		}
 	}
 
-	const FAABB3& GetBounds() const { return Bounds; }
+	const TAABB<T, 3>& GetBounds() const { return Bounds; }
 
 private:
-	FAABB3 Bounds;
+	TAABB<T, 3> Bounds;
 };
 
-template <typename TPayloadType>
-struct TBoundsWrapperHelper<TPayloadType, false>
+template <typename TPayloadType, typename T>
+struct TBoundsWrapperHelper<TPayloadType, T, false>
 {
-	void ComputeBounds(const TArray<TPayloadBoundsElement<TPayloadType, FReal>>&)
+	void ComputeBounds(const TArray<TPayloadBoundsElement<TPayloadType, T>>&)
 	{
 	}
 
-	const FAABB3 GetBounds() const
+	const TAABB<T, 3> GetBounds() const
 	{
-		return FAABB3::EmptyAABB();
+		return TAABB<T, 3>::EmptyAABB();
 	}
 };
 
-template <typename TPayloadType, bool bComputeBounds = true>
-struct TAABBTreeLeafArray : public TBoundsWrapperHelper<TPayloadType, bComputeBounds>
+template <typename TPayloadType, bool bComputeBounds = true, typename T = FReal>
+struct TAABBTreeLeafArray : public TBoundsWrapperHelper<TPayloadType, T, bComputeBounds>
 {
 	TAABBTreeLeafArray() {}
 	//todo: avoid copy?
-	TAABBTreeLeafArray(const TArray<TPayloadBoundsElement<TPayloadType, FReal>>& InElems)
+	TAABBTreeLeafArray(const TArray<TPayloadBoundsElement<TPayloadType, T>>& InElems)
 		: Elems(InElems)
 	{
 		this->ComputeBounds(Elems);
 	}
 
-	void GatherElements(TArray<TPayloadBoundsElement<TPayloadType, FReal>>& OutElements)
+	void GatherElements(TArray<TPayloadBoundsElement<TPayloadType, T>>& OutElements)
 	{
 		OutElements.Append(Elems);
 	}
@@ -201,14 +201,14 @@ struct TAABBTreeLeafArray : public TBoundsWrapperHelper<TPayloadType, bComputeBo
 	}
 
 	template <typename TSQVisitor, typename TQueryFastData>
-	FORCEINLINE_DEBUGGABLE bool RaycastFast(const FVec3& Start, TQueryFastData& QueryFastData, TSQVisitor& Visitor, const FVec3& Dir, const FVec3 InvDir, const bool bParallel[3]) const
+	FORCEINLINE_DEBUGGABLE bool RaycastFast(const TVec3<T>& Start, TQueryFastData& QueryFastData, TSQVisitor& Visitor, const TVec3<T>& Dir, const TVec3<T> InvDir, const bool bParallel[3]) const
 	{
-		return RaycastSweepImp</*bSweep=*/false>(Start, QueryFastData, FVec3(), Visitor, Dir, InvDir, bParallel);
+		return RaycastSweepImp</*bSweep=*/false>(Start, QueryFastData, TVec3<T>(), Visitor, Dir, InvDir, bParallel);
 	}
 
 	template <typename TSQVisitor, typename TQueryFastData>
-	FORCEINLINE_DEBUGGABLE bool SweepFast(const FVec3& Start, TQueryFastData& QueryFastData, const FVec3& QueryHalfExtents, TSQVisitor& Visitor,
-			const FVec3& Dir, const FVec3 InvDir, const bool bParallel[3]) const
+	FORCEINLINE_DEBUGGABLE bool SweepFast(const TVec3<T>& Start, TQueryFastData& QueryFastData, const TVec3<T>& QueryHalfExtents, TSQVisitor& Visitor,
+			const TVec3<T>& Dir, const TVec3<T> InvDir, const bool bParallel[3]) const
 	{
 		return RaycastSweepImp</*bSweep=*/true>(Start, QueryFastData, QueryHalfExtents, Visitor, Dir, InvDir, bParallel);
 	}
@@ -227,7 +227,7 @@ struct TAABBTreeLeafArray : public TBoundsWrapperHelper<TPayloadType, bComputeBo
 
 			if (Elem.Bounds.Intersects(QueryBounds))
 			{
-				TSpatialVisitorData<TPayloadType> VisitData(Elem.Payload, true, Elem.Bounds);
+				TSpatialVisitorData<TPayloadType> VisitData(Elem.Payload, true, FAABB3(Elem.Bounds.Min(), Elem.Bounds.Max()));
 				if (Visitor.VisitOverlap(VisitData) == false)
 				{
 					return false;
@@ -239,7 +239,7 @@ struct TAABBTreeLeafArray : public TBoundsWrapperHelper<TPayloadType, bComputeBo
 	}
 
 	template <bool bSweep, typename TQueryFastData, typename TSQVisitor>
-	FORCEINLINE_DEBUGGABLE bool RaycastSweepImp(const FVec3& Start, TQueryFastData& QueryFastData, const FVec3& QueryHalfExtents, TSQVisitor& Visitor, const FVec3& Dir, const FVec3 InvDir, const bool bParallel[3]) const
+	FORCEINLINE_DEBUGGABLE bool RaycastSweepImp(const TVec3<T>& Start, TQueryFastData& QueryFastData, const TVec3<T>& QueryHalfExtents, TSQVisitor& Visitor, const TVec3<T>& Dir, const TVec3<T> InvDir, const bool bParallel[3]) const
 	{
 		FVec3 TmpPosition;
 		FReal TOI;
@@ -252,7 +252,7 @@ struct TAABBTreeLeafArray : public TBoundsWrapperHelper<TPayloadType, bComputeBo
 				continue;
 			}
 
-			const auto& InstanceBounds = Elem.Bounds;
+			const FAABB3 InstanceBounds(Elem.Bounds.Min(), Elem.Bounds.Max());
 			if (TAABBTreeIntersectionHelper<TQueryFastData, bSweep ? EAABBQueryType::Sweep :
 					EAABBQueryType::Raycast>::Intersects(Start, QueryFastData, TOI, TmpPosition, InstanceBounds, FAABB3(), QueryHalfExtents, Dir, InvDir, bParallel))
 			{
@@ -280,7 +280,7 @@ struct TAABBTreeLeafArray : public TBoundsWrapperHelper<TPayloadType, bComputeBo
 		}
 	}
 
-	void UpdateElement(const TPayloadType& Payload, const FAABB3& NewBounds, bool bHasBounds)
+	void UpdateElement(const TPayloadType& Payload, const TAABB<T, 3>& NewBounds, bool bHasBounds)
 	{
 		if (!bHasBounds)
 			return;
@@ -301,24 +301,25 @@ struct TAABBTreeLeafArray : public TBoundsWrapperHelper<TPayloadType, bComputeBo
 		Ar << Elems;
 	}
 
-	TArray<TPayloadBoundsElement<TPayloadType, FReal>> Elems;
+	TArray<TPayloadBoundsElement<TPayloadType, T>> Elems;
 };
 
-template <typename TPayloadType, bool bComputeBounds>
-FChaosArchive& operator<<(FChaosArchive& Ar, TAABBTreeLeafArray<TPayloadType, bComputeBounds>& LeafArray)
+template <typename TPayloadType, bool bComputeBounds, typename T>
+FChaosArchive& operator<<(FChaosArchive& Ar, TAABBTreeLeafArray<TPayloadType, bComputeBounds, T>& LeafArray)
 {
 	LeafArray.Serialize(Ar);
 	return Ar;
 }
 
-struct FAABBTreeNode
+template <typename T>
+struct TAABBTreeNode
 {
-	FAABBTreeNode()
+	TAABBTreeNode()
 	{
-		ChildrenBounds[0] = FAABB3();
-		ChildrenBounds[1] = FAABB3();
+		ChildrenBounds[0] = TAABB<T, 3>();
+		ChildrenBounds[1] = TAABB<T, 3>();
 	}
-	FAABB3 ChildrenBounds[2];
+	TAABB<T, 3> ChildrenBounds[2];
 	int32 ChildrenNodes[2] = { 0, 0 };
 	bool bLeaf = false;
 
@@ -326,7 +327,7 @@ struct FAABBTreeNode
 	{
 		for (auto& Bounds : ChildrenBounds)
 		{
-			TBox<FReal, 3>::SerializeAsAABB(Ar, Bounds);
+			TBox<T, 3>::SerializeAsAABB(Ar, Bounds);
 		}
 
 		for (auto& Node : ChildrenNodes)
@@ -338,7 +339,8 @@ struct FAABBTreeNode
 	}
 };
 
-FORCEINLINE FChaosArchive& operator<<(FChaosArchive& Ar, FAABBTreeNode& Node)
+template <typename T>
+FORCEINLINE FChaosArchive& operator<<(FChaosArchive& Ar, TAABBTreeNode<T>& Node)
 {
 	Node.Serialize(Ar);
 	return Ar;
@@ -393,21 +395,21 @@ struct DirtyGridHashEntry
 	int32 Count;  // Number of valid entries from Index in FlattenedCellArrayOfDirtyIndices
 };
 
-template <typename TPayloadType, typename TLeafType, bool bMutable = true>
-class TAABBTree final : public ISpatialAcceleration<TPayloadType, FReal, 3>
+template <typename TPayloadType, typename TLeafType, bool bMutable = true, typename T = FReal>
+class TAABBTree final : public ISpatialAcceleration<TPayloadType, T, 3> 
 {
 public:
 	using PayloadType = TPayloadType;
 	static constexpr int D = 3;
-	using TType = FReal;
-	static constexpr FReal DefaultMaxPayloadBounds = 100000;
+	using TType = T;
+	static constexpr T DefaultMaxPayloadBounds = 100000;
 	static constexpr int32 DefaultMaxChildrenInLeaf = 12;
 	static constexpr int32 DefaultMaxTreeDepth = 16;
 	static constexpr int32 DefaultMaxNumToProcess = 0; // 0 special value for processing all without timeslicing
 	static constexpr ESpatialAcceleration StaticType = TIsSame<TAABBTreeLeafArray<TPayloadType>, TLeafType>::Value ? ESpatialAcceleration::AABBTree : 
 		(TIsSame<TBoundingVolume<TPayloadType>, TLeafType>::Value ? ESpatialAcceleration::AABBTreeBV : ESpatialAcceleration::Unknown);
 	TAABBTree()
-		: ISpatialAcceleration<TPayloadType, FReal, 3>(StaticType)
+		: ISpatialAcceleration<TPayloadType, T, 3>(StaticType)
 		, MaxChildrenInLeaf(DefaultMaxChildrenInLeaf)
 		, MaxTreeDepth(DefaultMaxTreeDepth)
 		, MaxPayloadBounds(DefaultMaxPayloadBounds)
@@ -452,8 +454,8 @@ public:
 	}
 
 	template <typename TParticles>
-	TAABBTree(const TParticles& Particles, int32 InMaxChildrenInLeaf = DefaultMaxChildrenInLeaf, int32 InMaxTreeDepth = DefaultMaxTreeDepth, FReal InMaxPayloadBounds = DefaultMaxPayloadBounds, int32 InMaxNumToProcess = DefaultMaxNumToProcess )
-		: ISpatialAcceleration<TPayloadType, FReal, 3>(StaticType)
+	TAABBTree(const TParticles& Particles, int32 InMaxChildrenInLeaf = DefaultMaxChildrenInLeaf, int32 InMaxTreeDepth = DefaultMaxTreeDepth, T InMaxPayloadBounds = DefaultMaxPayloadBounds, int32 InMaxNumToProcess = DefaultMaxNumToProcess )
+		: ISpatialAcceleration<TPayloadType, T, 3>(StaticType)
 		, MaxChildrenInLeaf(InMaxChildrenInLeaf)
 		, MaxTreeDepth(InMaxTreeDepth)
 		, MaxPayloadBounds(InMaxPayloadBounds)
@@ -471,7 +473,7 @@ public:
 
 	virtual TArray<TPayloadType> FindAllIntersections(const FAABB3& Box) const override { return FindAllIntersectionsImp(Box); }
 
-	bool GetAsBoundsArray(TArray<FAABB3>& AllBounds, int32 NodeIdx, int32 ParentNode, FAABB3& Bounds)
+	bool GetAsBoundsArray(TArray<TAABB<T, 3>>& AllBounds, int32 NodeIdx, int32 ParentNode, TAABB<T, 3>& Bounds)
 	{
 		if (Nodes[NodeIdx].bLeaf)
 		{
@@ -488,14 +490,14 @@ public:
 
 	virtual ~TAABBTree() {}
 
-	void CopyFrom(const TAABBTree<TPayloadType, TLeafType, bMutable>& Other)
+	void CopyFrom(const TAABBTree<TPayloadType, TLeafType, bMutable, T>& Other)
 	{
 		(*this) = Other;
 	}
 
-	virtual TUniquePtr<ISpatialAcceleration<TPayloadType, FReal, 3>> Copy() const override
+	virtual TUniquePtr<ISpatialAcceleration<TPayloadType, T, 3>> Copy() const override
 	{
-		return TUniquePtr<ISpatialAcceleration<TPayloadType, FReal, 3>>(new TAABBTree<TPayloadType, TLeafType, bMutable>(*this));
+		return TUniquePtr<ISpatialAcceleration<TPayloadType, T, 3>>(new TAABBTree<TPayloadType, TLeafType, bMutable, T>(*this));
 	}
 
 	virtual void Raycast(const FVec3& Start, const FVec3& Dir, const FReal Length, ISpatialVisitor<TPayloadType, FReal>& Visitor) const override
@@ -514,7 +516,7 @@ public:
 	template <typename SQVisitor>
 	bool RaycastFast(const FVec3& Start, FQueryFastData& CurData, SQVisitor& Visitor, const FVec3& Dir, const FVec3 InvDir, const bool bParallel[3]) const
 	{
-		return QueryImp<EAABBQueryType::Raycast>(Start, CurData, FVec3(), FAABB3(), Visitor, Dir, InvDir, bParallel);
+		return QueryImp<EAABBQueryType::Raycast>(Start, CurData, TVec3<T>(), TAABB<T, 3>(), Visitor, Dir, InvDir, bParallel);
 	}
 
 	void Sweep(const FVec3& Start, const FVec3& Dir, const FReal Length, const FVec3 QueryHalfExtents, ISpatialVisitor<TPayloadType, FReal>& Visitor) const override
@@ -559,7 +561,7 @@ public:
 	// This is to make sure important parameters are not changed during inopportune times
 	void GetCVars()
 	{
-		DirtyElementGridCellSize = (FReal) FAABBTreeDirtyGridCVars::DirtyElementGridCellSize;
+		DirtyElementGridCellSize = (T) FAABBTreeDirtyGridCVars::DirtyElementGridCellSize;
 		if (DirtyElementGridCellSize > SMALL_NUMBER)
 		{
 			DirtyElementGridCellSizeInv = 1.0f / DirtyElementGridCellSize;
@@ -693,7 +695,7 @@ public:
 		DirtyElements.RemoveAtSwap(DeleteDirtyParticleIdx);
 	}
 
-	FORCEINLINE_DEBUGGABLE int32 AddDirtyElementToGrid(const TAABB<FReal, 3>& NewBounds, int32 NewDirtyElement)
+	FORCEINLINE_DEBUGGABLE int32 AddDirtyElementToGrid(const TAABB<T, 3>& NewBounds, int32 NewDirtyElement)
 	{
 		bool bAddToGrid = !TooManyOverlapQueryCells(NewBounds, DirtyElementGridCellSizeInv, DirtyElementMaxPhysicalSizeInCells);
 		if (!bAddToGrid)
@@ -729,11 +731,11 @@ public:
 		return INDEX_NONE;
 	}
 
-	FORCEINLINE_DEBUGGABLE int32 UpdateDirtyElementInGrid(const TAABB<FReal, 3>& NewBounds, int32 DirtyElementIndex, int32 DirtyGridOverflowIdx)
+	FORCEINLINE_DEBUGGABLE int32 UpdateDirtyElementInGrid(const TAABB<T, 3>& NewBounds, int32 DirtyElementIndex, int32 DirtyGridOverflowIdx)
 	{
 		if (DirtyGridOverflowIdx == INDEX_NONE)
 		{
-			const TAABB<FReal, 3>& OldBounds = DirtyElements[DirtyElementIndex].Bounds;
+			const TAABB<T, 3>& OldBounds = DirtyElements[DirtyElementIndex].Bounds;
 
 			// Delete element in cells that are no longer overlapping
 			DoForOverlappedCellsExclude(OldBounds, NewBounds, DirtyElementGridCellSize, DirtyElementGridCellSizeInv, [&](int32 Hash) -> bool {
@@ -803,7 +805,7 @@ public:
 		}
 	}
 
-	virtual void UpdateElement(const TPayloadType& Payload, const FAABB3& NewBounds, bool bHasBounds) override
+	virtual void UpdateElement(const TPayloadType& Payload, const TAABB<T, 3>& NewBounds, bool bHasBounds) override
 	{
 		//CSV_SCOPED_TIMING_STAT(ChaosPhysicsTimers, AABBTreeUpdateElement)
 		if (ensure(bMutable))
@@ -814,9 +816,9 @@ public:
 				if (PayloadInfo->LeafIdx != INDEX_NONE)
 				{
 					//If we are still within the same leaf bounds, do nothing
-					if (bHasBounds)
+					if (bHasBounds) 
 					{
-						const FAABB3& LeafBounds = Leaves[PayloadInfo->LeafIdx].GetBounds();
+						const TAABB<T, 3>& LeafBounds = Leaves[PayloadInfo->LeafIdx].GetBounds();
 						if (LeafBounds.Contains(NewBounds.Min()) && LeafBounds.Contains(NewBounds.Max()))
 						{
 							// We still need to update the constituent bounds
@@ -882,7 +884,7 @@ public:
 			}
 			else
 			{
-				FAABB3 GlobalBounds = bTooBig ? NewBounds : FAABB3(FVec3(TNumericLimits<FReal>::Lowest()), FVec3(TNumericLimits<FReal>::Max()));
+				TAABB<T, 3> GlobalBounds = bTooBig ? NewBounds : TAABB<T, 3>(TVec3<T>(TNumericLimits<T>::Lowest()), TVec3<T>(TNumericLimits<T>::Max()));
 				if (PayloadInfo->GlobalPayloadIdx == INDEX_NONE)
 				{
 					PayloadInfo->GlobalPayloadIdx = GlobalPayloads.Add(FElement{ Payload, GlobalBounds });
@@ -936,7 +938,7 @@ public:
 		return TreeStats;
 	}
 
-	const TArray<TPayloadBoundsElement<TPayloadType, FReal>>& GlobalObjects() const
+	const TArray<TPayloadBoundsElement<TPayloadType, T>>& GlobalObjects() const
 	{
 		return GlobalPayloads;
 	}
@@ -948,8 +950,8 @@ public:
 		if (Ar.CustomVer(FExternalPhysicsCustomObjectVersion::GUID) < FExternalPhysicsCustomObjectVersion::RemovedAABBTreeFullBounds)
 		{
 			// Serialize out unused aabb for earlier versions
-			FAABB3 Dummy(FVec3((FReal)0), FVec3((FReal)0));
-			TBox<FReal, 3>::SerializeAsAABB(Ar, Dummy);
+			TAABB<T, 3> Dummy(TVec3<T>((T)0), TVec3<T>((T)0));
+			TBox<T, 3>::SerializeAsAABB(Ar, Dummy);
 		}
 		Ar << Nodes;
 		Ar << Leaves;
@@ -990,8 +992,8 @@ public:
 
 private:
 
-	using FElement = TPayloadBoundsElement<TPayloadType, FReal>;
-	using FNode = FAABBTreeNode;
+	using FElement = TPayloadBoundsElement<TPayloadType, T>;
+	using FNode = TAABBTreeNode<T>;
 
 	void ReoptimizeTree()
 	{
@@ -1014,7 +1016,7 @@ private:
 			Leaf.GatherElements(AllElements);
 		}
 
-		TAABBTree<TPayloadType,TLeafType,bMutable> NewTree(AllElements);
+		TAABBTree<TPayloadType,TLeafType,bMutable, T> NewTree(AllElements);
 		*this = NewTree;
 	}
 
@@ -1098,7 +1100,7 @@ private:
 	
 
 	template <EAABBQueryType Query, typename TQueryFastData, typename SQVisitor>
-	bool QueryImp(const FVec3& RESTRICT Start, TQueryFastData& CurData, const FVec3 QueryHalfExtents, const FAABB3& QueryBounds, SQVisitor& Visitor, const FVec3& Dir, const FVec3 InvDir, const bool bParallel[3]) const
+	bool QueryImp(const FVec3& RESTRICT Start, TQueryFastData& CurData, const FVec3& QueryHalfExtents, const FAABB3& QueryBounds, SQVisitor& Visitor, const FVec3& Dir, const FVec3& InvDir, const bool bParallel[3]) const
 	{
 		//QUICK_SCOPE_CYCLE_COUNTER(AABBTreeQueryImp);
 		FVec3 TmpPosition;
@@ -1115,10 +1117,10 @@ private:
 					continue;
 				}
 
-				const auto& InstanceBounds = Elem.Bounds;
+				const FAABB3 InstanceBounds(Elem.Bounds.Min(), Elem.Bounds.Max());
 				if(TAABBTreeIntersectionHelper<TQueryFastData,Query>::Intersects(Start, CurData, TOI, TmpPosition, InstanceBounds, QueryBounds, QueryHalfExtents, Dir, InvDir, bParallel))
 				{
-					TSpatialVisitorData<TPayloadType> VisitData(Elem.Payload,true);
+					TSpatialVisitorData<TPayloadType> VisitData(Elem.Payload,true, InstanceBounds);
 					bool bContinue;
 					if(Query == EAABBQueryType::Overlap)
 					{
@@ -1140,7 +1142,7 @@ private:
 		{	// Returns true if we should continue
 			auto IntersectAndVisit = [&](const FElement& Elem) -> bool
 			{
-				const auto& InstanceBounds = Elem.Bounds;
+				const FAABB3 InstanceBounds(Elem.Bounds.Min(), Elem.Bounds.Max());
 				if (PrePreFilterHelper(Elem.Payload, QueryData, SimData))
 				{
 					return true;
@@ -1285,9 +1287,9 @@ private:
 			else
 			{
 				int32 Idx = 0;
-				for (const FAABB3& AABB : Node.ChildrenBounds)
+				for (const TAABB<T, 3>& AABB : Node.ChildrenBounds)
 				{
-					if(TAABBTreeIntersectionHelper<TQueryFastData, Query>::Intersects(Start, CurData, TOI, TmpPosition, AABB, QueryBounds, QueryHalfExtents, Dir, InvDir, bParallel))
+					if(TAABBTreeIntersectionHelper<TQueryFastData, Query>::Intersects(Start, CurData, TOI, TmpPosition, FAABB3(AABB.Min(), AABB.Max()), QueryBounds, QueryHalfExtents, Dir, InvDir, bParallel))
 					{
 						NodeStack.Add(FNodeQueueEntry{ Node.ChildrenNodes[Idx], TOI });
 					}
@@ -1349,7 +1351,7 @@ private:
 		NumProcessedThisSlice = 0;
 		GetCVars();  // Safe to copy CVARS here
 
-		WorkSnapshot.Bounds = FAABB3::EmptyAABB();
+		WorkSnapshot.Bounds = TAABB<T, 3>::EmptyAABB();
 
 		{
 			SCOPE_CYCLE_COUNTER(STAT_AABBTreeTimeSliceSetup);
@@ -1357,13 +1359,13 @@ private:
 			int32 Idx = 0;
 
 			//TODO: we need a better way to time-slice this case since there can be a huge number of Particles. Can't do it right now without making full copy
-			FVec3 CenterSum(0);
+			TVec3<T> CenterSum(0);
 
 			for (auto& Particle : Particles)
 			{
 				bool bHasBoundingBox = HasBoundingBox(Particle);
 				auto Payload = Particle.template GetPayload<TPayloadType>(Idx);
-				FAABB3 ElemBounds = ComputeWorldSpaceBoundingBox(Particle, false, (FReal)0);
+				TAABB<T, 3> ElemBounds = ComputeWorldSpaceBoundingBox(Particle, false, (T)0);
 				if (bHasBoundingBox)
 				{
 					if (ElemBounds.Extents().Max() > MaxPayloadBounds)
@@ -1379,7 +1381,7 @@ private:
 				}
 				else
 				{
-					ElemBounds = FAABB3(FVec3(TNumericLimits<FReal>::Lowest()), FVec3(TNumericLimits<FReal>::Max()));
+					ElemBounds = TAABB<T, 3>(TVec3<T>(TNumericLimits<T>::Lowest()), TVec3<T>(TNumericLimits<T>::Max()));
 				}
 
 				if (!bHasBoundingBox)
@@ -1397,11 +1399,11 @@ private:
 
 			if(WorkSnapshot.Elems.Num())
 			{
-				WorkSnapshot.AverageCenter = CenterSum * ((FReal)1 / (FReal)WorkSnapshot.Elems.Num());
+				WorkSnapshot.AverageCenter = CenterSum * ((T)1 / (T)WorkSnapshot.Elems.Num());
 			}
 			else
 			{
-				WorkSnapshot.AverageCenter = FVec3(0);
+				WorkSnapshot.AverageCenter = TVec3<T>(0);
 			}
 		}
 
@@ -1456,10 +1458,10 @@ private:
 
 	struct FSplitInfo
 	{
-		FAABB3 SplitBounds;	//Even split of parent bounds
-		FAABB3 RealBounds;	//Actual bounds as children are added
+		TAABB<T, 3> SplitBounds;	//Even split of parent bounds
+		TAABB<T, 3> RealBounds;	//Actual bounds as children are added
 		int32 WorkSnapshotIdx;	//Idx into work snapshot pool
-		FReal SplitBoundsSize2;
+		T SplitBoundsSize2;
 	};
 
 	struct FWorkSnapshot
@@ -1472,8 +1474,8 @@ private:
 
 		eTimeSlicePhase TimeslicePhase;
 
-		FAABB3 Bounds;
-		FVec3 AverageCenter;
+		TAABB<T, 3> Bounds;
+		TVec3<T> AverageCenter;
 		TArray<FElement> Elems;
 
 		int32 NodeLevel;
@@ -1491,13 +1493,13 @@ private:
 		{
 			const FElement& Elem = CurrentSnapshot.Elems[ElemIdx];
 			int32 MinBoxIdx = INDEX_NONE;
-			FReal MinDelta2 = TNumericLimits<FReal>::Max();
+			T MinDelta2 = TNumericLimits<T>::Max();
 			int32 BoxIdx = 0;
 			for (const FSplitInfo& SplitInfo : CurrentSnapshot.SplitInfos)
 			{
-				FAABB3 NewBox = SplitInfo.SplitBounds;
+				TAABB<T, 3> NewBox = SplitInfo.SplitBounds;
 				NewBox.GrowToInclude(Elem.Bounds);
-				const FReal Delta2 = NewBox.Extents().SizeSquared() - SplitInfo.SplitBoundsSize2;
+				const T Delta2 = NewBox.Extents().SizeSquared() - SplitInfo.SplitBoundsSize2;
 				if (Delta2 < MinDelta2)
 				{
 					MinDelta2 = Delta2;
@@ -1580,7 +1582,7 @@ private:
 
 			if (WorkPool[CurIdx].TimeslicePhase == eTimeSlicePhase::PreFindBestBounds)
 			{
-				const FVec3 Extents = WorkPool[CurIdx].Bounds.Extents();
+				const TVec3<T> Extents = WorkPool[CurIdx].Bounds.Extents();
 				const int32 MaxAxis = WorkPool[CurIdx].Bounds.LargestAxis();	//todo: use variance instead
 
 				//Add two children, remember this invalidates any pointers to current snapshot
@@ -1592,18 +1594,18 @@ private:
 				WorkPool[CurIdx].SplitInfos[1].WorkSnapshotIdx = SecondChildIdx;
 
 				//these are the hypothetical bounds for the perfect 50/50 split
-				WorkPool[CurIdx].SplitInfos[0].SplitBounds = FAABB3(WorkPool[CurIdx].Bounds.Min(), WorkPool[CurIdx].Bounds.Min());
-				WorkPool[CurIdx].SplitInfos[1].SplitBounds = FAABB3(WorkPool[CurIdx].Bounds.Max(), WorkPool[CurIdx].Bounds.Max());
+				WorkPool[CurIdx].SplitInfos[0].SplitBounds = TAABB<T, 3>(WorkPool[CurIdx].Bounds.Min(), WorkPool[CurIdx].Bounds.Min());
+				WorkPool[CurIdx].SplitInfos[1].SplitBounds = TAABB<T, 3>(WorkPool[CurIdx].Bounds.Max(), WorkPool[CurIdx].Bounds.Max());
 
-				const FVec3 Center = WorkPool[CurIdx].AverageCenter;
+				const TVec3<T> Center = WorkPool[CurIdx].AverageCenter;
 				for (FSplitInfo& SplitInfo : WorkPool[CurIdx].SplitInfos)
 				{
-					SplitInfo.RealBounds = FAABB3::EmptyAABB();
+					SplitInfo.RealBounds = TAABB<T, 3>::EmptyAABB();
 
 					for (int32 Axis = 0; Axis < 3; ++Axis)
 					{
-						FVec3 NewPt0 = Center;
-						FVec3 NewPt1 = Center;
+						TVec3<T> NewPt0 = Center;
+						TVec3<T> NewPt1 = Center;
 						if (Axis != MaxAxis)
 						{
 							NewPt0[Axis] = WorkPool[CurIdx].Bounds.Min()[Axis];
@@ -1623,8 +1625,8 @@ private:
 					WorkPool[FirstChildIdx].Elems.Reserve(ExpectedNumPerChild);
 					WorkPool[SecondChildIdx].Elems.Reserve(ExpectedNumPerChild);
 
-					WorkPool[FirstChildIdx].AverageCenter = FVec3(0);
-					WorkPool[SecondChildIdx].AverageCenter = FVec3(0);
+					WorkPool[FirstChildIdx].AverageCenter = TVec3<T>(0);
+					WorkPool[SecondChildIdx].AverageCenter = TVec3<T>(0);
 				}
 			}
 
@@ -1666,8 +1668,8 @@ private:
 				WorkPool[FirstChildIdx].NewNodeIdx = Nodes[NewNodeIdx].ChildrenNodes[0];
 				WorkPool[SecondChildIdx].NewNodeIdx = Nodes[NewNodeIdx].ChildrenNodes[1];
 
-				WorkPool[FirstChildIdx].AverageCenter *= ((FReal)1 / (FReal)WorkPool[FirstChildIdx].Elems.Num());
-				WorkPool[SecondChildIdx].AverageCenter *= ((FReal)1 / (FReal)WorkPool[SecondChildIdx].Elems.Num());
+				WorkPool[FirstChildIdx].AverageCenter *= ((T)1 / (T)WorkPool[FirstChildIdx].Elems.Num());
+				WorkPool[SecondChildIdx].AverageCenter *= ((T)1 / (T)WorkPool[SecondChildIdx].Elems.Num());
 
 				//push these two new nodes onto the stack
 				WorkStack.Add(SecondChildIdx);
@@ -1728,8 +1730,8 @@ private:
 	}
 
 
-	TAABBTree(const TAABBTree<TPayloadType, TLeafType, bMutable>& Other)
-		: ISpatialAcceleration<TPayloadType, FReal, 3>(StaticType)
+	TAABBTree(const TAABBTree<TPayloadType, TLeafType, bMutable, T>& Other)
+		: ISpatialAcceleration<TPayloadType, T, 3>(StaticType)
 		, Nodes(Other.Nodes)
 		, Leaves(Other.Leaves)
 		, DirtyElements(Other.DirtyElements)
@@ -1753,7 +1755,7 @@ private:
 
 	}
 
-	TAABBTree<TPayloadType,TLeafType, bMutable>& operator=(const TAABBTree<TPayloadType,TLeafType,bMutable>& Rhs)
+	TAABBTree<TPayloadType,TLeafType, bMutable, T>& operator=(const TAABBTree<TPayloadType,TLeafType,bMutable, T>& Rhs)
 	{
 		ensure(Rhs.WorkStack.Num() == 0);
 		//ensure(Rhs.WorkPool.Num() == 0);
@@ -1799,8 +1801,8 @@ private:
 	TArray<int32> FlattenedCellArrayOfDirtyIndices; // Array of indices into dirty Elements array, indices are always sorted within a 2D cell
 	TArray<int32> DirtyElementsGridOverflow; // Array of indices of DirtyElements that is not in the grid for some reason
 	// Copy of CVARS
-	FReal DirtyElementGridCellSize;
-	FReal DirtyElementGridCellSizeInv;
+	T DirtyElementGridCellSize;
+	T DirtyElementGridCellSizeInv;
 	int32 DirtyElementMaxGridCellQueryCount;
 	int32 DirtyElementMaxPhysicalSizeInCells;
 	int32 DirtyElementMaxCellCapacity;
@@ -1812,7 +1814,7 @@ private:
 
 	int32 MaxChildrenInLeaf;
 	int32 MaxTreeDepth;
-	FReal MaxPayloadBounds;
+	T MaxPayloadBounds;
 	int32 MaxNumToProcess;
 
 	int32 NumProcessedThisSlice;
@@ -1821,8 +1823,8 @@ private:
 	TArray<FWorkSnapshot> WorkPool;
 };
 
-template<typename TPayloadType, typename TLeafType, bool bMutable>
-FArchive& operator<<(FChaosArchive& Ar, TAABBTree<TPayloadType, TLeafType, bMutable>& AABBTree)
+template<typename TPayloadType, typename TLeafType, bool bMutable, typename T>
+FArchive& operator<<(FChaosArchive& Ar, TAABBTree<TPayloadType, TLeafType, bMutable, T>& AABBTree)
 {
 	AABBTree.Serialize(Ar);
 	return Ar;
diff --git a/Engine/Source/Runtime/Experimental/Chaos/Public/Chaos/AABBTreeDirtyGridUtils.h b/Engine/Source/Runtime/Experimental/Chaos/Public/Chaos/AABBTreeDirtyGridUtils.h
--- a/Engine/Source/Runtime/Experimental/Chaos/Public/Chaos/AABBTreeDirtyGridUtils.h
+++ b/Engine/Source/Runtime/Experimental/Chaos/Public/Chaos/AABBTreeDirtyGridUtils.h
@@ -41,10 +41,11 @@ namespace Chaos
 		return HashCell(X, Y);
 	}
 
-	FORCEINLINE_DEBUGGABLE bool TooManyOverlapQueryCells(const TAABB<FReal, 3>& AABB, FReal DirtyElementGridCellSizeInv, int32 MaximumOverlap)
+	template <typename T>
+	FORCEINLINE_DEBUGGABLE bool TooManyOverlapQueryCells(const TAABB<T, 3>& AABB, FReal DirtyElementGridCellSizeInv, int32 MaximumOverlap)
 	{
-		int32 XsampleCount = GetDirtyCellIndexFromWorldCoordinate(AABB.Max().X, DirtyElementGridCellSizeInv) - GetDirtyCellIndexFromWorldCoordinate(AABB.Min().X, DirtyElementGridCellSizeInv) + 1;
-		int32 YsampleCount = GetDirtyCellIndexFromWorldCoordinate(AABB.Max().Y, DirtyElementGridCellSizeInv) - GetDirtyCellIndexFromWorldCoordinate(AABB.Min().Y, DirtyElementGridCellSizeInv) + 1;
+		int32 XsampleCount = GetDirtyCellIndexFromWorldCoordinate((FReal)AABB.Max().X, DirtyElementGridCellSizeInv) - GetDirtyCellIndexFromWorldCoordinate((FReal)AABB.Min().X, DirtyElementGridCellSizeInv) + 1;
+		int32 YsampleCount = GetDirtyCellIndexFromWorldCoordinate((FReal)AABB.Max().Y, DirtyElementGridCellSizeInv) - GetDirtyCellIndexFromWorldCoordinate((FReal)AABB.Min().Y, DirtyElementGridCellSizeInv) + 1;
 
 		if (XsampleCount * YsampleCount <= MaximumOverlap)
 		{
@@ -53,8 +54,8 @@ namespace Chaos
 		return true;
 	}
 
-	template <typename FunctionType>
-	FORCEINLINE_DEBUGGABLE bool DoForOverlappedCells(const TAABB<FReal, 3>& AABB, FReal DirtyElementGridCellSize, FReal DirtyElementGridCellSizeInv, FunctionType Function)
+	template <typename T, typename FunctionType>
+	FORCEINLINE_DEBUGGABLE bool DoForOverlappedCells(const TAABB<T, 3>& AABB, FReal DirtyElementGridCellSize, FReal DirtyElementGridCellSizeInv, FunctionType Function)
 	{
 		int32 CellStartX = GetDirtyCellIndexFromWorldCoordinate(AABB.Min().X, DirtyElementGridCellSizeInv);
 		int32 CellStartY = GetDirtyCellIndexFromWorldCoordinate(AABB.Min().Y, DirtyElementGridCellSizeInv);
@@ -76,21 +77,21 @@ namespace Chaos
 	}
 
 	// Only execute function for new Cells not covered in old (Set difference: {Cells spanned by AABB} - { Cells spanned by AABBExclude})
-	template <typename FunctionType>
-	FORCEINLINE_DEBUGGABLE bool DoForOverlappedCellsExclude(const TAABB<FReal, 3>& AABB, const TAABB<FReal, 3>& AABBExclude, FReal DirtyElementGridCellSize, FReal DirtyElementGridCellSizeInv, FunctionType Function)
+	template <typename T, typename FunctionType>
+	FORCEINLINE_DEBUGGABLE bool DoForOverlappedCellsExclude(const TAABB<T, 3>& AABB, const TAABB<T, 3>& AABBExclude, FReal DirtyElementGridCellSize, FReal DirtyElementGridCellSizeInv, FunctionType Function)
 	{
 
-		int32 NewCellStartX = GetDirtyCellIndexFromWorldCoordinate(AABB.Min().X, DirtyElementGridCellSizeInv);
-		int32 NewCellStartY = GetDirtyCellIndexFromWorldCoordinate(AABB.Min().Y, DirtyElementGridCellSizeInv);
+		int32 NewCellStartX = GetDirtyCellIndexFromWorldCoordinate((FReal)AABB.Min().X, DirtyElementGridCellSizeInv);
+		int32 NewCellStartY = GetDirtyCellIndexFromWorldCoordinate((FReal)AABB.Min().Y, DirtyElementGridCellSizeInv);
 
-		int32 NewCellEndX = GetDirtyCellIndexFromWorldCoordinate(AABB.Max().X, DirtyElementGridCellSizeInv);
-		int32 NewCellEndY = GetDirtyCellIndexFromWorldCoordinate(AABB.Max().Y, DirtyElementGridCellSizeInv);
+		int32 NewCellEndX = GetDirtyCellIndexFromWorldCoordinate((FReal)AABB.Max().X, DirtyElementGridCellSizeInv);
+		int32 NewCellEndY = GetDirtyCellIndexFromWorldCoordinate((FReal)AABB.Max().Y, DirtyElementGridCellSizeInv);
 
-		int32 OldCellStartX = GetDirtyCellIndexFromWorldCoordinate(AABBExclude.Min().X, DirtyElementGridCellSizeInv);
-		int32 OldCellStartY = GetDirtyCellIndexFromWorldCoordinate(AABBExclude.Min().Y, DirtyElementGridCellSizeInv);
+		int32 OldCellStartX = GetDirtyCellIndexFromWorldCoordinate((FReal)AABBExclude.Min().X, DirtyElementGridCellSizeInv);
+		int32 OldCellStartY = GetDirtyCellIndexFromWorldCoordinate((FReal)AABBExclude.Min().Y, DirtyElementGridCellSizeInv);
 
-		int32 OldCellEndX = GetDirtyCellIndexFromWorldCoordinate(AABBExclude.Max().X, DirtyElementGridCellSizeInv);
-		int32 OldCellEndY = GetDirtyCellIndexFromWorldCoordinate(AABBExclude.Max().Y, DirtyElementGridCellSizeInv);
+		int32 OldCellEndX = GetDirtyCellIndexFromWorldCoordinate((FReal)AABBExclude.Max().X, DirtyElementGridCellSizeInv);
+		int32 OldCellEndY = GetDirtyCellIndexFromWorldCoordinate((FReal)AABBExclude.Max().Y, DirtyElementGridCellSizeInv);
 
 		// Early out here
 		if (OldCellStartX <= NewCellStartX &&
@@ -412,7 +413,7 @@ namespace Chaos
 	}
 
 	template <typename FunctionType>
-	FORCEINLINE_DEBUGGABLE void DoForSweepIntersectCells(const TVec3<FReal>& QueryHalfExtents, const FVec3& StartPoint, const FVec3& Dir, FReal Length, FReal DirtyElementGridCellSize, FReal DirtyElementGridCellSizeInv, FunctionType InFunction)
+	FORCEINLINE_DEBUGGABLE void DoForSweepIntersectCells(const FVec3 QueryHalfExtents, const FVec3& StartPoint, const FVec3& Dir, FReal Length, FReal DirtyElementGridCellSize, FReal DirtyElementGridCellSizeInv, FunctionType InFunction)
 	{
 		FReal AbsDx = FMath::Abs(Dir.X);
 		FReal AbsDy = FMath::Abs(Dir.Y);
@@ -426,7 +427,7 @@ namespace Chaos
 		else
 		{
 			// Swap Y and X
-			DoForSweepIntersectCellsImp(QueryHalfExtents.Y, QueryHalfExtents.X, StartPoint.Y, StartPoint.X, Dir.Y * Length, Dir.X * Length, DirtyElementGridCellSize, DirtyElementGridCellSizeInv, [&](FReal X, FReal Y) {InFunction(Y, X); });
+			DoForSweepIntersectCellsImp(QueryHalfExtents.Y, QueryHalfExtents.X, StartPoint.Y, StartPoint.X, Dir.Y * Length, Dir.X * Length, DirtyElementGridCellSize, DirtyElementGridCellSizeInv, [&](auto X, auto Y) {InFunction(Y, X); });
 		}
 
 	}
diff --git a/Engine/Source/Runtime/Experimental/Chaos/Public/Chaos/Box.h b/Engine/Source/Runtime/Experimental/Chaos/Public/Chaos/Box.h
--- a/Engine/Source/Runtime/Experimental/Chaos/Public/Chaos/Box.h
+++ b/Engine/Source/Runtime/Experimental/Chaos/Public/Chaos/Box.h
@@ -146,7 +146,8 @@ namespace Chaos
 		// Extents
 		FORCEINLINE const TAABB<FReal, 3> BoundingBox() const
 		{
-			return AABB;
+			// LWC - necessary if T is different from FReal
+			return TAABB<FReal, 3>(AABB.Min(), AABB.Max());
 		}
 
 		// Apply a limit to the specified margin that prevents the box inverting
@@ -157,12 +158,12 @@ namespace Chaos
 		}
 
 		// Return the distance and normal is the closest point on the surface to Pos. Negative for penetration.
-		virtual T PhiWithNormal(const TVector<T, d>& Pos, TVector<T, d>& Normal) const override
+		virtual FReal PhiWithNormal(const FVec3& Pos, FVec3& Normal) const override
 		{
 			return AABB.PhiWithNormal(Pos, Normal);
 		}
 
-		virtual T PhiWithNormalScaled(const TVector<T, d>& Pos, const TVector<T, d>& Scale, TVector<T, d>& Normal) const override
+		virtual FReal PhiWithNormalScaled(const FVec3& Pos, const FVec3& Scale, FVec3& Normal) const override
 		{
 			return TAABB<T, d>(Scale * AABB.Min(), Scale * AABB.Max()).PhiWithNormal(Pos, Normal);
 		}
@@ -172,7 +173,7 @@ namespace Chaos
 			return TAABB<T, d>(InMin, InMax).RaycastFast(StartPoint, Dir, InvDir, bParallel, Length, InvLength, OutTime, OutPosition);
 		}
 
-		virtual bool CHAOS_API Raycast(const TVector<T, d>& StartPoint, const TVector<T, d>& Dir, const T Length, const T Thickness, T& OutTime, TVector<T, d>& OutPosition, TVector<T, d>& OutNormal, int32& OutFaceIndex) const override
+		virtual bool CHAOS_API Raycast(const FVec3& StartPoint, const FVec3& Dir, const FReal Length, const FReal Thickness, FReal& OutTime, FVec3& OutPosition, FVec3& OutNormal, int32& OutFaceIndex) const override
 		{
 			if (AABB.Raycast(StartPoint, Dir, Length, Thickness, OutTime, OutPosition, OutNormal, OutFaceIndex))
 			{
@@ -186,12 +187,12 @@ namespace Chaos
 			return AABB.FindClosestPoint(StartPoint, Thickness);
 		}
 
-		virtual Pair<TVector<T, d>, bool> FindClosestIntersectionImp(const TVector<T, d>& StartPoint, const TVector<T, d>& EndPoint, const T Thickness) const override
+		virtual Pair<FVec3, bool> FindClosestIntersectionImp(const FVec3& StartPoint, const FVec3& EndPoint, const FReal Thickness) const override
 		{
 			return AABB.FindClosestIntersectionImp(StartPoint, EndPoint, Thickness);
 		}
 
-		virtual TVector<T, d> FindGeometryOpposingNormal(const TVector<T, d>& DenormDir, int32 FaceIndex, const TVector<T, d>& OriginalNormal) const override
+		virtual FVec3 FindGeometryOpposingNormal(const FVec3& DenormDir, int32 FaceIndex, const FVec3& OriginalNormal) const override
 		{
 			return AABB.FindGeometryOpposingNormal(DenormDir, FaceIndex, OriginalNormal);
 		}
@@ -328,11 +329,11 @@ namespace Chaos
 		int32 NumVertices() const { return SVertices.Num(); }
 
 		// Get the plane at the specified index (e.g., indices from FindVertexPlanes)
-		const TPlaneConcrete<FReal, 3> GetPlane(int32 FaceIndex) const
+		const TPlaneConcrete<FReal> GetPlane(int32 FaceIndex) const
 		{
 			const FVec3& PlaneN = SNormals[FaceIndex];
 			const FVec3 PlaneX = AABB.Center() + 0.5f * (PlaneN * AABB.Extents());
-			return TPlaneConcrete<FReal, 3>(PlaneX, PlaneN);
+			return TPlaneConcrete<FReal>(PlaneX, PlaneN);
 		}
 
 		// Get the vertex at the specified index (e.g., indices from GetPlaneVertexs)
diff --git a/Engine/Source/Runtime/Experimental/Chaos/Public/Chaos/ISpatialAcceleration.h b/Engine/Source/Runtime/Experimental/Chaos/Public/Chaos/ISpatialAcceleration.h
--- a/Engine/Source/Runtime/Experimental/Chaos/Public/Chaos/ISpatialAcceleration.h
+++ b/Engine/Source/Runtime/Experimental/Chaos/Public/Chaos/ISpatialAcceleration.h
@@ -259,11 +259,12 @@ public:
 
 	virtual bool IsAsyncTimeSlicingComplete() { return AsyncTimeSlicingComplete; }
 	virtual void ProgressAsyncTimeSlicing(bool ForceBuildCompletion = false) {}
-	virtual TArray<TPayloadType> FindAllIntersections(const TAABB<T, d>& Box) const { check(false); return TArray<TPayloadType>(); }
 
-	virtual void Raycast(const TVector<T, d>& Start, const TVector<T, d>& Dir, const T Length, ISpatialVisitor<TPayloadType, T>& Visitor) const { check(false); }
-	virtual void Sweep(const TVector<T, d>& Start, const TVector<T, d>& Dir, const T Length, const TVector<T, d> QueryHalfExtents, ISpatialVisitor<TPayloadType, T>& Visitor) const { check(false);}
-	virtual void Overlap(const TAABB<T, d>& QueryBounds, ISpatialVisitor<TPayloadType, T>& Visitor) const { check(false); }
+	// IMPORTANT : (LWC) this API should be typed on Freal not T, as we want the query API to be using the highest precision while maintaining arbitrary internal precision for the acceleration structure ( based on T )
+	virtual TArray<TPayloadType> FindAllIntersections(const FAABB3& Box) const { check(false); return TArray<TPayloadType>(); }
+	virtual void Raycast(const FVec3& Start, const FVec3& Dir, const FReal Length, ISpatialVisitor<TPayloadType, FReal>& Visitor) const { check(false); }
+	virtual void Sweep(const FVec3& Start, const FVec3& Dir, const FReal Length, const FVec3 QueryHalfExtents, ISpatialVisitor<TPayloadType, FReal>& Visitor) const { check(false);}
+	virtual void Overlap(const FAABB3& QueryBounds, ISpatialVisitor<TPayloadType, FReal>& Visitor) const { check(false); }
 
 	virtual void Reset()
 	{
diff --git a/Engine/Source/Runtime/Experimental/Chaos/Public/Chaos/Plane.h b/Engine/Source/Runtime/Experimental/Chaos/Public/Chaos/Plane.h
--- a/Engine/Source/Runtime/Experimental/Chaos/Public/Chaos/Plane.h
+++ b/Engine/Source/Runtime/Experimental/Chaos/Public/Chaos/Plane.h
@@ -8,15 +8,15 @@
 namespace Chaos
 {
 
-template <typename T, int d>
+template <typename T, int d = 3>
 class TPlaneConcrete
 {
 public:
 
 	// Scale the plane and assume that any of the scale components could be zero
-	static TPlaneConcrete<T, d> MakeScaledSafe(const TPlaneConcrete<T, d>& Plane, const TVector<T, d>& Scale)
+	static TPlaneConcrete<T> MakeScaledSafe(const TPlaneConcrete<T>& Plane, const TVec3<T>& Scale)
 	{
-		const FVec3 ScaledX = Plane.MX * Scale;
+		const TVec3<T> ScaledX = Plane.MX * Scale;
 		
 		// If all 3 scale components are non-zero we can just inverse-scale the normal
 		// If 1 scale component is zero, the normal will point in that direction of the zero scale
@@ -26,16 +26,16 @@ public:
 		const int32 ZeroY = FMath::IsNearlyZero(Scale.Y) ? 1 : 0;
 		const int32 ZeroZ = FMath::IsNearlyZero(Scale.Z) ? 1 : 0;
 		const int32 NumZeros = ZeroX + ZeroY + ZeroZ;
-		FVec3 ScaledN;
+		TVec3<T> ScaledN;
 		if (NumZeros == 0)
 		{
 			// All 3 scale components non-zero
-			ScaledN = FVec3(Plane.MNormal.X / Scale.X, Plane.MNormal.Y / Scale.Y, Plane.MNormal.Z / Scale.Z);
+			ScaledN = TVec3<T>(Plane.MNormal.X / Scale.X, Plane.MNormal.Y / Scale.Y, Plane.MNormal.Z / Scale.Z);
 		}
 		else if (NumZeros == 1)
 		{
 			// Exactly one Scale component is zero
-			ScaledN = FVec3(
+			ScaledN = TVec3<T>(
 				(ZeroX) ? 1.0f : 0.0f,
 				(ZeroY) ? 1.0f : 0.0f,
 				(ZeroZ) ? 1.0f : 0.0f);
@@ -43,7 +43,7 @@ public:
 		else if (NumZeros == 2)
 		{
 			// Exactly two Scale components is zero
-			ScaledN = FVec3(
+			ScaledN = TVec3<T>(
 				(ZeroX) ? Plane.MNormal.X : 0.0f,
 				(ZeroY) ? Plane.MNormal.Y : 0.0f,
 				(ZeroZ) ? Plane.MNormal.Z : 0.0f);
@@ -55,7 +55,7 @@ public:
 		}
 
 		// Even after all the above, we may still get a zero normal (e.g., we scale N=(1,0,0) by S=(0,1,0))
-		const FReal ScaleN2 = ScaledN.SizeSquared();
+		const T ScaleN2 = ScaledN.SizeSquared();
 		if (ScaleN2 > SMALL_NUMBER)
 		{
 			ScaledN = ScaledN * FMath::InvSqrt(ScaleN2);
@@ -65,17 +65,17 @@ public:
 			ScaledN = Plane.MNormal;
 		}
 		
-		return TPlaneConcrete<FReal, 3>(ScaledX, ScaledN);
+		return TPlaneConcrete<T>(ScaledX, ScaledN);
 	}
 
 	// Scale the plane and assume that none of the scale components are zero
-	static TPlaneConcrete<T, d> MakeScaledUnsafe(const TPlaneConcrete<T, d>& Plane, const TVector<T, d>& Scale)
+	static TPlaneConcrete<T> MakeScaledUnsafe(const TPlaneConcrete<T>& Plane, const TVec3<T>& Scale)
 	{
-		const FVec3 ScaledX = Plane.MX * Scale;
-		FVec3 ScaledN = Plane.MNormal / Scale;
+		const TVec3<T> ScaledX = Plane.MX * Scale;
+		TVec3<T> ScaledN = Plane.MNormal / Scale;
 
 		// We don't handle zero scales, but we could still end up with a small normal
-		const FReal ScaleN2 = ScaledN.SizeSquared();
+		const T ScaleN2 = ScaledN.SizeSquared();
 		if (ScaleN2 > SMALL_NUMBER)
 		{
 			ScaledN =  ScaledN * FMath::InvSqrt(ScaleN2);
@@ -85,47 +85,48 @@ public:
 			ScaledN = Plane.MNormal;
 		}
 
-		return TPlaneConcrete<FReal, 3>(ScaledX, ScaledN);
+		return TPlaneConcrete<T>(ScaledX, ScaledN);
 	}
 
 
 	TPlaneConcrete() = default;
-	TPlaneConcrete(const TVector<T, d>& InX, const TVector<T, d>& InNormal)
+	TPlaneConcrete(const TVec3<T>& InX, const TVec3<T>& InNormal)
 	    : MX(InX)
 	    , MNormal(InNormal)
 	{
+		static_assert(d == 3, "Only dimension 3 is supported");
 	}
 
 	/**
 	 * Phi is positive on the side of the normal, and negative otherwise.
 	 */
-	T SignedDistance(const TVector<T, d>& x) const
+	FReal SignedDistance(const FVec3& x) const
 	{
-		return TVector<T, d>::DotProduct(x - MX, MNormal);
+		return FVec3::DotProduct(x - (FVec3)MX, (FVec3)MNormal);
 	}
 
 	/**
 	 * Phi is positive on the side of the normal, and negative otherwise.
 	 */
-	T PhiWithNormal(const TVector<T, d>& x, TVector<T, d>& Normal) const
+	FReal PhiWithNormal(const FVec3& x, FVec3& Normal) const
 	{
 		Normal = MNormal;
-		return TVector<T, d>::DotProduct(x - MX, MNormal);
+		return FVec3::DotProduct(x - (FVec3)MX, (FVec3)MNormal);
 	}
 
-	TVector<T, d> FindClosestPoint(const TVector<T, d>& x, const T Thickness = (T)0) const
+	FVec3 FindClosestPoint(const FVec3& x, const FReal Thickness = (FReal)0) const
 	{
-		auto Dist = TVector<T, d>::DotProduct(x - MX, MNormal) - Thickness;
-		return x - TVector<T, d>(Dist * MNormal);
+		auto Dist = FVec3::DotProduct(x - (FVec3)MX, (FVec3)MNormal) - Thickness;
+		return x - FVec3(Dist * MNormal);
 	}
 
-	bool Raycast(const TVector<T, d>& StartPoint, const TVector<T, d>& Dir, const T Length, const T Thickness, T& OutTime, TVector<T, d>& OutPosition, TVector<T, d>& OutNormal, int32& OutFaceIndex) const
+	bool Raycast(const FVec3& StartPoint, const FVec3& Dir, const FReal Length, const FReal Thickness, FReal& OutTime, FVec3& OutPosition, FVec3& OutNormal, int32& OutFaceIndex) const
 	{
-		ensure(FMath::IsNearlyEqual(Dir.SizeSquared(), (T)1, (T)KINDA_SMALL_NUMBER));
+		ensure(FMath::IsNearlyEqual(Dir.SizeSquared(), (FReal)1, (FReal)KINDA_SMALL_NUMBER));
 		CHAOS_ENSURE(Length > 0);
 		OutFaceIndex = INDEX_NONE;
 
-		const T SignedDist = TVector<T, d>::DotProduct(StartPoint - MX, MNormal);
+		const FReal SignedDist = FVec3::DotProduct(StartPoint - (FVec3)MX, (FVec3)MNormal);
 		if (FMath::Abs(SignedDist) < Thickness)
 		{
 			//initial overlap so stop
@@ -133,19 +134,19 @@ public:
 			return true;
 		}
 
-		const TVector<T, d> DirTowardsPlane = SignedDist < 0 ? MNormal : -MNormal;
-		const T RayProjectedTowardsPlane = TVector<T, d>::DotProduct(Dir, DirTowardsPlane);
-		const T Epsilon = 1e-7f;
+		const FVec3 DirTowardsPlane = SignedDist < 0 ? MNormal : -MNormal;
+		const FReal RayProjectedTowardsPlane = FVec3::DotProduct(Dir, DirTowardsPlane);
+		const FReal Epsilon = 1e-7f;
 		if (RayProjectedTowardsPlane < Epsilon)	//moving parallel or away
 		{
 			return false;
 		}
 
 		//No initial overlap so we are outside the thickness band of the plane. So translate the plane to account for thickness	
-		const TVector<T, d> TranslatedPlaneX = MX - Thickness * DirTowardsPlane;
-		const TVector<T, d> StartToTranslatedPlaneX = TranslatedPlaneX - StartPoint;
-		const T LengthTowardsPlane = TVector<T, d>::DotProduct(StartToTranslatedPlaneX, DirTowardsPlane);
-		const T LengthAlongRay = LengthTowardsPlane / RayProjectedTowardsPlane;
+		const FVec3 TranslatedPlaneX = (FVec3)MX - Thickness * DirTowardsPlane;
+		const FVec3 StartToTranslatedPlaneX = TranslatedPlaneX - StartPoint;
+		const FReal LengthTowardsPlane = FVec3::DotProduct(StartToTranslatedPlaneX, DirTowardsPlane);
+		const FReal LengthAlongRay = LengthTowardsPlane / RayProjectedTowardsPlane;
 		
 		if (LengthAlongRay > Length)
 		{
@@ -158,35 +159,35 @@ public:
 		return true;
 	}
 
-	Pair<TVector<T, d>, bool> FindClosestIntersection(const TVector<T, d>& StartPoint, const TVector<T, d>& EndPoint, const T Thickness) const
+	Pair<FVec3, bool> FindClosestIntersection(const FVec3& StartPoint, const FVec3& EndPoint, const FReal Thickness) const
  	{
-		TVector<T, d> Direction = EndPoint - StartPoint;
-		T Length = Direction.Size();
+		FVec3 Direction = EndPoint - StartPoint;
+		FReal Length = Direction.Size();
 		Direction = Direction.GetSafeNormal();
-		TVector<T, d> XPos = MX + MNormal * Thickness;
-		TVector<T, d> XNeg = MX - MNormal * Thickness;
-		TVector<T, d> EffectiveX = ((XNeg - StartPoint).Size() < (XPos - StartPoint).Size()) ? XNeg : XPos;
-		TVector<T, d> PlaneToStart = EffectiveX - StartPoint;
-		T Denominator = TVector<T, d>::DotProduct(Direction, MNormal);
+		FVec3 XPos = (FVec3)MX + (FVec3)MNormal * Thickness;
+		FVec3 XNeg = (FVec3)MX - (FVec3)MNormal * Thickness;
+		FVec3 EffectiveX = ((XNeg - StartPoint).Size() < (XPos - StartPoint).Size()) ? XNeg : XPos;
+		FVec3 PlaneToStart = EffectiveX - StartPoint;
+		FReal Denominator = FVec3::DotProduct(Direction, MNormal);
 		if (Denominator == 0)
 		{
-			if (TVector<T, d>::DotProduct(PlaneToStart, MNormal) == 0)
+			if (FVec3::DotProduct(PlaneToStart, MNormal) == 0)
 			{
 				return MakePair(EndPoint, true);
 			}
-			return MakePair(TVector<T, d>(0), false);
+			return MakePair(FVec3(0), false);
 		}
-		T Root = TVector<T, d>::DotProduct(PlaneToStart, MNormal) / Denominator;
+		FReal Root = FVec3::DotProduct(PlaneToStart, MNormal) / Denominator;
 		if (Root < 0 || Root > Length)
 		{
-			return MakePair(TVector<T, d>(0), false);
+			return MakePair(FVec3(0), false);
 		}
-		return MakePair(TVector<T, d>(Root * Direction + StartPoint), true);
+		return MakePair(FVec3(Root * Direction + StartPoint), true);
 	}
 
-	const TVector<T, d>& X() const { return MX; }
-	const TVector<T, d>& Normal() const { return MNormal; }
-	const TVector<T, d>& Normal(const TVector<T, d>&) const { return MNormal; }
+	const TVec3<T>& X() const { return MX; }
+	const TVec3<T>& Normal() const { return MNormal; }
+	const TVec3<T>& Normal(const TVec3<T>&) const { return MNormal; }
 
 	FORCEINLINE void Serialize(FArchive& Ar)
 	{
@@ -199,12 +200,12 @@ public:
 	}
 
   private:
-	TVector<T, d> MX;
-	TVector<T, d> MNormal;
+	TVec3<T> MX;
+	TVec3<T> MNormal;
 };
 
-template <typename T, int d>
-FArchive& operator<<(FArchive& Ar, TPlaneConcrete<T,d>& PlaneConcrete)
+template <typename T>
+FArchive& operator<<(FArchive& Ar, TPlaneConcrete<T>& PlaneConcrete)
 {
 	PlaneConcrete.Serialize(Ar);
 	return Ar;
@@ -256,7 +257,7 @@ class TPlane final : public FImplicitObject
 	/**
 	 * Phi is positive on the side of the normal, and negative otherwise.
 	 */
-	virtual T PhiWithNormal(const TVector<T, d>& x, TVector<T, d>& Normal) const override
+	virtual FReal PhiWithNormal(const FVec3& x, FVec3& Normal) const override
 	{
 		return MPlaneConcrete.PhiWithNormal(x,Normal);
 	}
@@ -266,12 +267,12 @@ class TPlane final : public FImplicitObject
 		return MPlaneConcrete.FindClosestPoint(x,Thickness);
 	}
 
-	virtual bool Raycast(const TVector<T, d>& StartPoint, const TVector<T, d>& Dir, const T Length, const T Thickness, T& OutTime, TVector<T, d>& OutPosition, TVector<T, d>& OutNormal, int32& OutFaceIndex) const override
+	virtual bool Raycast(const FVec3& StartPoint, const FVec3& Dir, const FReal Length, const FReal Thickness, FReal& OutTime, FVec3& OutPosition, FVec3& OutNormal, int32& OutFaceIndex) const override
 	{
 		return MPlaneConcrete.Raycast(StartPoint,Dir,Length,Thickness,OutTime,OutPosition,OutNormal,OutFaceIndex);
 	}
 
-	virtual Pair<TVector<T, d>, bool> FindClosestIntersectionImp(const TVector<T, d>& StartPoint, const TVector<T, d>& EndPoint, const T Thickness) const override
+	virtual Pair<FVec3, bool> FindClosestIntersectionImp(const FVec3& StartPoint, const FVec3& EndPoint, const FReal Thickness) const override
  	{
 		return MPlaneConcrete.FindClosestIntersection(StartPoint,EndPoint,Thickness);
 	}
@@ -302,10 +303,10 @@ class TPlane final : public FImplicitObject
 		return MPlaneConcrete.GetTypeHash();
 	}
 
-	const TPlaneConcrete<T,d>& PlaneConcrete() const { return MPlaneConcrete; }
+	const TPlaneConcrete<T>& PlaneConcrete() const { return MPlaneConcrete; }
 
   private:
-	  TPlaneConcrete<T,d> MPlaneConcrete;
+	  TPlaneConcrete<T> MPlaneConcrete;
 };
 
 template<typename T, int d>
diff --git a/Engine/Source/Runtime/Experimental/Chaos/Public/Chaos/TriangleMeshImplicitObject.h b/Engine/Source/Runtime/Experimental/Chaos/Public/Chaos/TriangleMeshImplicitObject.h
--- a/Engine/Source/Runtime/Experimental/Chaos/Public/Chaos/TriangleMeshImplicitObject.h
+++ b/Engine/Source/Runtime/Experimental/Chaos/Public/Chaos/TriangleMeshImplicitObject.h
@@ -307,7 +307,7 @@ namespace Chaos
 		TUniquePtr<TArray<int32>> ExternalVertexIndexMap;
 		bool bCullsBackFaceRaycast;
 
-		using BVHType = TAABBTree<int32, TAABBTreeLeafArray<int32, /*bComputeBounds=*/false>, /*bMutable=*/false>;
+		using BVHType = TAABBTree<int32, TAABBTreeLeafArray<int32, /*bComputeBounds=*/false, FRealSingle>, /*bMutable=*/false, FRealSingle>;
 
 		// Initialising constructor privately declared for use in CopySlow to copy the underlying BVH
 		template <typename IdxType>
@@ -339,11 +339,11 @@ namespace Chaos
 
 			bool HasBoundingBox() const { return true; }
 
-			FAABB3 BoundingBox() const
+			TAABB<FRealSingle, 3> BoundingBox() const
 			{
 				auto LambdaHelper = [&](const auto& Elements)
 				{
-					TAABB<FReal,3> Bounds(TmData->MParticles.X(Elements[Index][0]), TmData->MParticles.X(Elements[Index][0]));
+					TAABB<FRealSingle,3> Bounds(TmData->MParticles.X(Elements[Index][0]), TmData->MParticles.X(Elements[Index][0]));
 
 					Bounds.GrowToInclude(TmData->MParticles.X(Elements[Index][1]));
 					Bounds.GrowToInclude(TmData->MParticles.X(Elements[Index][2]));
