Make constraint sort stable so that multiple constraints of same pair are in the same order during resim

diff --git a/Engine/Source/Runtime/Experimental/Chaos/Private/Chaos/PBDCollisionConstraints.cpp b/Engine/Source/Runtime/Experimental/Chaos/Private/Chaos/PBDCollisionConstraints.cpp
--- a/Engine/Source/Runtime/Experimental/Chaos/Private/Chaos/PBDCollisionConstraints.cpp
+++ b/Engine/Source/Runtime/Experimental/Chaos/Private/Chaos/PBDCollisionConstraints.cpp
@@ -27,6 +27,7 @@
 #include "ProfilingDebugging/ScopedTimers.h"
 #include "ProfilingDebugging/CsvProfiler.h"
 #include "Algo/Sort.h"
+#include "Algo/StableSort.h"
 
 #if INTEL_ISPC
 #include "PBDCollisionConstraints.ispc.generated.h"
@@ -616,7 +617,12 @@ namespace Chaos
 	{
 		check(!bInAppendOperation);
 
-		Algo::Sort(Handles, [](const FPBDCollisionConstraintHandle* A, const FPBDCollisionConstraintHandle* B)
+		//We have to use StableSort so that constraints of the same pair stay in the same order
+		//Otherwise the order within each pair can change due to where they start out in the array
+		//TODO: we should label each contact (and shape) for things like warm starting GJK
+		//and so we could use that label as part of the key
+		//and then we could use regular Sort (which is faster)
+		Algo::StableSort(Handles, [](const FPBDCollisionConstraintHandle* A, const FPBDCollisionConstraintHandle* B)
 		{
 			if(A->GetType() == B->GetType())
 			{
