don't ensure() on failed voronoi cell computes; these just indicate we had overlapping voronoi sites and don't cause problems generally

#rnx

diff --git a/Engine/Source/Runtime/Experimental/Voronoi/Private/Voronoi/Voronoi.cpp b/Engine/Source/Runtime/Experimental/Voronoi/Private/Voronoi/Voronoi.cpp
--- a/Engine/Source/Runtime/Experimental/Voronoi/Private/Voronoi/Voronoi.cpp
+++ b/Engine/Source/Runtime/Experimental/Voronoi/Private/Voronoi/Voronoi.cpp
@@ -176,7 +176,6 @@ bool GetVoronoiEdges(const TArrayView<const FVector> &Sites, const FBox& Bounds,
 		do
 		{
 			bool bCouldComputeCell = Container->compute_cell(cell, CellIterator);
-			ensureMsgf(bCouldComputeCell, TEXT("Failed to compute a Voronoi cell -- this may indicate sites positioned directly on top of other sites, which is not valid for a Voronoi diagram"));
 			if (bCouldComputeCell)
 			{
 				const double *pp = con.p[CellIterator.ijk] + con.ps * CellIterator.q;
