Change resim interpolation to be distance based instead of time based

diff --git a/Engine/Source/Programs/HeadlessChaos/Private/HeadlessChaosTestRewind.cpp b/Engine/Source/Programs/HeadlessChaos/Private/HeadlessChaosTestRewind.cpp
--- a/Engine/Source/Programs/HeadlessChaos/Private/HeadlessChaosTestRewind.cpp
+++ b/Engine/Source/Programs/HeadlessChaos/Private/HeadlessChaosTestRewind.cpp
@@ -645,10 +645,6 @@ namespace ChaosTest {
 
 				auto FloorGeom = TSharedPtr<FImplicitObject, ESPMode::ThreadSafe>(new TBox<FReal, 3>(FVec3(-100, -100, -1), FVec3(100, 100, 0)));
 
-				const FReal LeashTime = 3;
-
-				Solver->GetResultsManager().SetResimInterpTime(LeashTime);
-
 				auto& Particle = Proxy->GetGameThreadAPI();
 				Particle.SetGravityEnabled(false);
 				Particle.SetV(FVec3(0, 0, 0));
@@ -659,7 +655,6 @@ namespace ChaosTest {
 				const FReal StartMovingTimeDiscrete = FMath::FloorToInt(StartMovingTime / SimDt) * SimDt;
 				const FReal GTDt = 1;
 				const FReal InterpStartTime = ResimTime + SimDt;
-				const FReal InterpEndTime = ResimTime + LeashTime + SimDt;
 				const FReal SleepLocation = SleepTime - StartMovingTimeDiscrete;
 				const FReal CorrectedLocation = SimDt <= 1 ? 0 : 4;
 				FReal PrevZDuringInterp = SleepLocation;	//shouldn't be further then this because already interpolating back to 0
@@ -700,18 +695,10 @@ namespace ChaosTest {
 					}
 					else
 					{
-						if (InterpolatedTime >= InterpEndTime)
-						{
-							//not moving and no longer in leash mode
-							EXPECT_NEAR(Particle.X()[2], CorrectedLocation, 1e-2);
-						}
-						else
-						{
-							//leash mode
-							EXPECT_GT(Particle.X()[2], CorrectedLocation);
-							EXPECT_LT(Particle.X()[2], PrevZDuringInterp);
-							PrevZDuringInterp = Particle.X()[2];
-						}
+						//leash mode
+						EXPECT_GE(Particle.X()[2], CorrectedLocation);
+						EXPECT_LE(Particle.X()[2], PrevZDuringInterp);
+						PrevZDuringInterp = Particle.X()[2];
 					}
 				}
 
@@ -1961,11 +1948,6 @@ namespace ChaosTest {
 					}
 				};
 
-				const FReal InterpTime = 10;
-				const FReal InterpStrength = 0.25;
-				Solver->GetResultsManager().SetResimInterpTime(InterpTime);
-				Solver->GetResultsManager().SetResimInterpStrength(InterpStrength);
-
 				const int32 LastGameStep = 20;
 				const int32 NumPhysSteps = FMath::TruncToInt(LastGameStep / SimDt);
 
@@ -2021,11 +2003,10 @@ namespace ChaosTest {
 						//expected interpolation from pt to gt
 						const int32 NextSimStep = FMath::CeilToInt(InterpolatedTime / SimDt);
 						const FReal NextSimStepTime = NextSimStep * SimDt;
-						const FReal InterpEndTime = 20 + InterpTime;
 						const FReal ExpectedValue = ZStart + ZVel * InterpolatedTime;
 						const FReal TargetValue = ZStart + ZVel * NextSimStepTime;
 
-						if(InterpolatedTime >= InterpEndTime || InterpolatedTime <= 20)
+						if(!Proxy->IsResimSmoothing())
 						{
 							//no resim interpolation, just simple value interpolation
 							EXPECT_NEAR(Particle.X()[2], ExpectedValue, 1e-2);
@@ -2033,7 +2014,7 @@ namespace ChaosTest {
 						else
 						{
 							//exponential decay from current state to target
-							EXPECT_NEAR(Particle.X()[2], FMath::Lerp(PrevZ, ExpectedValue, InterpStrength), 1e-2);
+							EXPECT_NEAR(Particle.X()[2], FMath::Lerp(PrevZ, ExpectedValue, Chaos::ResimInterpStrength), 1e-2);
 						}
 					}
 				}
diff --git a/Engine/Source/Runtime/Experimental/Chaos/Private/Chaos/Framework/ChaosResultsManager.cpp b/Engine/Source/Runtime/Experimental/Chaos/Private/Chaos/Framework/ChaosResultsManager.cpp
--- a/Engine/Source/Runtime/Experimental/Chaos/Private/Chaos/Framework/ChaosResultsManager.cpp
+++ b/Engine/Source/Runtime/Experimental/Chaos/Private/Chaos/Framework/ChaosResultsManager.cpp
@@ -62,9 +62,9 @@ namespace Chaos
 				}
 
 				//update leash target
-				if(FResimParticleInfo* ResimInfo = ParticleToResimInfo.Find(Proxy))
+				if(FDirtyRigidParticleData* ResimTarget = ParticleToResimTarget.Find(Proxy))
 				{
-					ResimInfo->Next = Data;
+					*ResimTarget = Data;
 				}
 			}
 		}
@@ -84,7 +84,7 @@ namespace Chaos
 			//newer result exists so prev is overwritten
 			if(Results.Prev)
 			{
-				FreeToHistory_External(Results.Prev);
+				MarshallingManager.FreePullData_External(Results.Prev);
 			}
 
 			Results.Prev = Results.Next;
@@ -130,7 +130,7 @@ namespace Chaos
 		//sync mode doesn't use prev results, but if we were async previously we need to clean it up
 		if (Results.Prev)
 		{
-			FreeToHistory_External(Results.Prev);
+			MarshallingManager.FreePullData_External(Results.Prev);
 			Results.Prev = nullptr;
 		}
 
@@ -164,108 +164,46 @@ namespace Chaos
 		return 1;	//if 0 dt just use 1 as alpha
 	}
 
-	FRealSingle DefaultResimInterpTime = 1.f;
-	FAutoConsoleVariableRef CVarResimInterpTime(TEXT("p.ResimInterpTime"), DefaultResimInterpTime, TEXT("How long to interpolate between original sim and resim results. 0 means no interpolation, the larget the value the smoother and longer interpolation takes. Restart game to see affect"));
-
-	FRealSingle DefaultResimInterpStrength = 0.2f;
-	FAutoConsoleVariableRef CVarResimInterpStrength(TEXT("p.ResimInterpStrength"), DefaultResimInterpStrength, TEXT("How strong the resim interp leash is. 1 means immediately snap to new target, 0 means do not interpolate at all"));
-
 	FChaosResultsManager::FChaosResultsManager(FChaosMarshallingManager& InMarshallingManager)
 		: MarshallingManager(InMarshallingManager)
 	{
-		SetResimInterpTime(DefaultResimInterpTime);
-		SetResimInterpStrength(DefaultResimInterpStrength);
-	}
-
-	void FChaosResultsManager::SetResimInterpTime(const FReal InterpTime)
-	{
-		if(InterpTime <= 0)
-		{
-			InvResimInterpTime = FLT_MAX;
-		}
-		else
-		{
-			InvResimInterpTime = 1 / InterpTime;
-		}
 	}
 
 	const FChaosInterpolationResults& FChaosResultsManager::UpdateInterpAlpha_External(const FReal ResultsTime, const FReal GlobalAlpha)
 	{
-		auto ComputeLeashAlphaHelper = [this, ResultsTime](const FReal LatestDivergeTime)
-		{
-			
-			return FMath::Min((FReal)1, (ResultsTime - LatestDivergeTime) * InvResimInterpTime);
-		};
-
 		Results.Alpha = (float)GlobalAlpha;	 // LWC_TODO: Precision loss
 
-		if(ParticleToResimInfo.Num() == 0)	//no resim interpolation so just exit
-		{
-			return Results;
-		}
-
-		//first update all interpolations that were dirtied by sim results
-		for(FChaosRigidInterpolationData& RigidData : Results.RigidInterpolations)
-		{
-			if(FSingleParticlePhysicsProxy* Proxy = RigidData.Prev.GetProxy())	//don't care about pending deleted proxies
-			{
-				RigidData.LeashAlpha = 1;	//if no resim interp just use global alpha
-				if (FResimParticleInfo* LeashInfo = ParticleToResimInfo.Find(RigidData.Prev.GetProxy()))
-				{
-					const FReal LeashAlpha = ComputeLeashAlphaHelper(LeashInfo->LeashStartTime);
-					if (LeashAlpha >= 1)
-					{
-						ParticleToResimInfo.Remove(Proxy);	//no longer interpolating
-					}
-					else
-					{
-						//still in leash mode so let interpolator know
-						if(LeashInfo->bDiverged)
-						{
-							//RigidData.LeashAlpha = LeashAlpha;
-							RigidData.LeashAlpha = (float)ResimInterpStrength;		 // LWC_TODO: Precision loss
-						}
-					}
-				}
-			}
-		}
-		
 		//make sure any resim interpolated bodies are still in the results array.
 		//It's possible the body stopped moving after the resim and is not dirty, but we still want to interpolate to final place
-		TArray<FSingleParticlePhysicsProxy*> NoLongerInterpolating;
-		for (const auto& Itr : ParticleToResimInfo)
+		TArray<FSingleParticlePhysicsProxy*> FinishedSmoothing;
+		for (const auto& Itr : ParticleToResimTarget)
 		{
 			FSingleParticlePhysicsProxy* Proxy = Itr.Key;
-			if (Proxy->GetPullDataInterpIdx_External() == INDEX_NONE)	//not in results array
-			{
-				const FReal LeashAlpha = ComputeLeashAlphaHelper(Itr.Value.LeashStartTime);
 
-				if(Itr.Value.bDiverged)
+			if(Proxy->IsResimSmoothing())
+			{
+				if (Proxy->GetPullDataInterpIdx_External() == INDEX_NONE)	//not in results array
 				{
 					//still need to interpolate, so add to results array
-					//even if alpha > 1 we want to snap to end of interpolation (i.e. the leash target)
 					const int32 DataIdx = Results.RigidInterpolations.AddDefaulted(1);
 					Proxy->SetPullDataInterpIdx_External(DataIdx);
 					FChaosRigidInterpolationData& RigidData = Results.RigidInterpolations[DataIdx];
 
-					RigidData.Next = Itr.Value.Next;	//not dirty from sim, so just use whatever last next was
+					RigidData.Next = Itr.Value;			//not dirty from sim, so just use whatever last next was
 					RigidData.Prev = RigidData.Next;	//prev same as next since we're just using leash
-					RigidData.LeashAlpha = LeashAlpha < 1 ? (float)ResimInterpStrength : 1;	//if last step with leash just snap to end		 // LWC_TODO: Precision loss
-				}
-				
-				if (LeashAlpha >= 1)	//leash will snap to end so not needed after this
-				{
-					NoLongerInterpolating.Add(Proxy);
 				}
 			}
+			else
+			{
+				FinishedSmoothing.Add(Proxy);
+			}
 		}
-		
-		for(FSingleParticlePhysicsProxy* Proxy : NoLongerInterpolating)
+
+		for(FSingleParticlePhysicsProxy* Proxy : FinishedSmoothing)
 		{
-			ParticleToResimInfo.Remove(Proxy);
+			RemoveProxy_External(Proxy);
 		}
 
-
 		return Results;
 	}
 
@@ -331,93 +269,32 @@ namespace Chaos
 
 	void FChaosResultsManager::ProcessResimResult_External()
 	{
-		//find original pull data to compare to
-		FPullPhysicsData* OriginalData = nullptr;
-		if(Results.Next->ExternalEndTime == Results.Prev->ExternalEndTime)
-		{
-			OriginalData = Results.Prev;
-		}
-		else
-		{
-			for (int32 Idx = 0; Idx < ResultsHistory.Num(); ++Idx)
-			{
-				if(Results.Next->ExternalEndTime == ResultsHistory[Idx]->ExternalEndTime)
-				{
-					OriginalData = ResultsHistory[Idx];
-					break;
-				}
-			}
-		}
-		ensure(OriginalData != nullptr);	//rewind happened, but our history buffer doesn't contain the original result
-
-		if(OriginalData)
+		//make sure any proxy in the resim data is marked as resimming
+		for(const FDirtyRigidParticleData& ResimDirty : Results.Next->DirtyRigids)
 		{
-			//First record the potential target for anything that is dirty from resim
-			//During the first frame of the rewind, everything that changed must be dirty
-			for (const FDirtyRigidParticleData& ResimDirty : Results.Next->DirtyRigids)
-			{
-				if (FSingleParticlePhysicsProxy* ResimProxy = ResimDirty.GetProxy())
-				{
-					FResimParticleInfo& ResimInfo = ParticleToResimInfo.FindOrAdd(ResimProxy);
-					ResimInfo.bDiverged = true;	//If not in original data then diverged, otherwise will be reset below
-					ResimInfo.LeashStartTime = LatestTimeSeen;
-					ResimInfo.EntryTime = Results.Next->ExternalEndTime;
-					ResimInfo.Next = ResimDirty;
-				}
-			}
-			
-			//Then check if anything diverged from original data
-			//If so, turn leash mode on (i.e. mark it as diverged)
-			for (const FDirtyRigidParticleData& OriginalDirty : OriginalData->DirtyRigids)
+			if (FSingleParticlePhysicsProxy* ResimProxy = ResimDirty.GetProxy())
 			{
-				if (FSingleParticlePhysicsProxy* OriginalProxy = OriginalDirty.GetProxy())
-				{
-					//During rewind we have to update game thread of original rewind
-					//Because of that, anything in the original data must also be in the first frame
-					//of the resim dirty data
-					//Without this we don't have the target data which means we don't know what to interpolate to
-					FResimParticleInfo* ResimInfo = ParticleToResimInfo.Find(OriginalProxy);
-					if(ensure(ResimInfo))
-					{
-						ResimInfo->bDiverged = (ResimInfo->EntryTime == OriginalData->ExternalEndTime) ? StateDiverged(ResimInfo->Next, OriginalDirty) : true;
-						if(ResimInfo->bDiverged)
-						{
-							//We still use resim's latest target (i.e. maybe it stopped moving a few frames ago)
-							//But the time is updated since we want the leash to be turned on for x seconds
-							//(We want to record latest moment of divergence, not moment when resim data was recorded)
-							ResimInfo->EntryTime = Results.Next->ExternalEndTime;
-							ResimInfo->LeashStartTime = LatestTimeSeen;
-						
-						}
-					}
-				}
+				ParticleToResimTarget.FindOrAdd(ResimProxy) = ResimDirty;
+				ResimProxy->SetResimSmoothing(true);
 			}
 		}
 	}
 
-	void FChaosResultsManager::FreeToHistory_External(FPullPhysicsData* PullData)
-	{
-		ResultsHistory.Insert(PullData, 0);	//store in reverse order (most recent data first)
-		SetHistoryLength_External(HistoryLength);
-	}
-
-	void FChaosResultsManager::SetHistoryLength_External(int32 InLength)
+	FChaosResultsManager::~FChaosResultsManager()
 	{
-		while(ResultsHistory.Num() > InLength)
+		if(Results.Prev)
 		{
-			MarshallingManager.FreePullData_External(ResultsHistory.Pop());	//oldest data is last so can just pop it
+			MarshallingManager.FreePullData_External(Results.Prev);
 		}
 
-		HistoryLength = InLength;
-	}
-
-	FChaosResultsManager::~FChaosResultsManager()
-	{
-		SetHistoryLength_External(0);
+		if(Results.Next)
+		{
+			MarshallingManager.FreePullData_External(Results.Next);
+		}
 	}
 
 	void FChaosResultsManager::RemoveProxy_External(FSingleParticlePhysicsProxy* Proxy)
 	{
-		ParticleToResimInfo.Remove(Proxy);
+		ParticleToResimTarget.Remove(Proxy);
 	}
 }
diff --git a/Engine/Source/Runtime/Experimental/Chaos/Private/Chaos/Framework/PhysicsSolverBase.cpp b/Engine/Source/Runtime/Experimental/Chaos/Private/Chaos/Framework/PhysicsSolverBase.cpp
--- a/Engine/Source/Runtime/Experimental/Chaos/Private/Chaos/Framework/PhysicsSolverBase.cpp
+++ b/Engine/Source/Runtime/Experimental/Chaos/Private/Chaos/Framework/PhysicsSolverBase.cpp
@@ -178,8 +178,6 @@ namespace Chaos
 
 	FPhysicsSolverBase::~FPhysicsSolverBase()
 	{
-		//free pull results before MarshallingManager is destroyed (technically should be safe because constructed after, but this makes it more explicit)
-		PullResultsManager->SetHistoryLength_External(0);	//ok to use _External because solver going away means external thread released it
 		//reset history buffer before freeing any unremoved callback objects
 		MarshallingManager.SetHistoryLength_Internal(0);
 
diff --git a/Engine/Source/Runtime/Experimental/Chaos/Private/PBDRigidsSolver.cpp b/Engine/Source/Runtime/Experimental/Chaos/Private/PBDRigidsSolver.cpp
--- a/Engine/Source/Runtime/Experimental/Chaos/Private/PBDRigidsSolver.cpp
+++ b/Engine/Source/Runtime/Experimental/Chaos/Private/PBDRigidsSolver.cpp
@@ -354,11 +354,6 @@ namespace Chaos
 				MSolver->GetEvolution()->EndFrame(MDeltaTime);
 			}
 
-			if(FRewindData* RewindData = MSolver->GetRewindData())
-			{
-				RewindData->FinishFrame();
-			}
-
 			MSolver->FinalizeCallbackData_Internal();
 
 			MSolver->GetSolverTime() += MDeltaTime;
@@ -373,6 +368,11 @@ namespace Chaos
 			{
 				MSolver->CompleteSceneSimulation();
 			}
+
+			if (FRewindData* RewindData = MSolver->GetRewindData())
+			{
+				RewindData->FinishFrame();
+			}
 		}
 
 	protected:
@@ -663,7 +663,6 @@ namespace Chaos
 		bUseCollisionResimCache = InUseCollisionResimCache;
 		MRewindCallback = MoveTemp(RewindCallback);
 		MarshallingManager.SetHistoryLength_Internal(NumFrames);
-		PullResultsManager->SetHistoryLength_External(NumFrames);
 		MEvolution->SetRewindData(GetRewindData());
 	}
 
@@ -1212,6 +1211,7 @@ namespace Chaos
 		FPullPhysicsData* PullData = MarshallingManager.GetCurrentPullData_Internal();
 
 		TParticleView<FPBDRigidParticles>& DirtyParticles = GetParticles().GetDirtyParticlesView();
+		const bool bIsResim = GetEvolution()->IsResimming();
 
 		//todo: should be able to go wide just add defaulted etc...
 		{
@@ -1221,25 +1221,27 @@ namespace Chaos
 
 			for (Chaos::TPBDRigidParticleHandleImp<FReal, 3, false>& DirtyParticle : DirtyParticles)
 			{
-				IPhysicsProxyBase* Proxy = DirtyParticle.Handle()->PhysicsProxy();
-							if(Proxy != nullptr)
+				if(IPhysicsProxyBase* Proxy = DirtyParticle.Handle()->PhysicsProxy())
+				{
+					switch(DirtyParticle.GetParticleType())
+					{
+						case Chaos::EParticleType::Rigid:
 						{
-								switch(DirtyParticle.GetParticleType())
+							if(!bIsResim || DirtyParticle.SyncState() == ESyncState::HardDesync)
 							{
-							case Chaos::EParticleType::Rigid:
-								{
-									PullData->DirtyRigids.AddDefaulted();
-									((FSingleParticlePhysicsProxy*)(Proxy))->BufferPhysicsResults(PullData->DirtyRigids.Last());
-									break;
-								}
-							case Chaos::EParticleType::Kinematic:
-							case Chaos::EParticleType::Static:
-								ensure(false);
-								break;
-							case Chaos::EParticleType::GeometryCollection:
-								ActiveGC.AddUnique((FGeometryCollectionPhysicsProxy*)(Proxy));
-								break;
-							case Chaos::EParticleType::Clustered:
+								PullData->DirtyRigids.AddDefaulted();
+								((FSingleParticlePhysicsProxy*)(Proxy))->BufferPhysicsResults(PullData->DirtyRigids.Last());
+							}
+							break;
+						}
+						case Chaos::EParticleType::Kinematic:
+						case Chaos::EParticleType::Static:
+							ensure(false);
+							break;
+						case Chaos::EParticleType::GeometryCollection:
+							ActiveGC.AddUnique((FGeometryCollectionPhysicsProxy*)(Proxy));
+							break;
+						case Chaos::EParticleType::Clustered:
 							if (auto ClusterParticle = DirtyParticle.CastToClustered())
 							{
 								if (ClusterParticle->InternalCluster())
@@ -1252,13 +1254,13 @@ namespace Chaos
 								}
 								else
 								{
-								ActiveGC.AddUnique((FGeometryCollectionPhysicsProxy*)(Proxy));
+									ActiveGC.AddUnique((FGeometryCollectionPhysicsProxy*)(Proxy));
 								}
 							}
-								break;
-							default:
-								check(false);
-							}
+							break;
+						default:
+							check(false);
+					}
 				}	
 			}
 		}
diff --git a/Engine/Source/Runtime/Experimental/Chaos/Private/PhysicsProxy/SingleParticlePhysicsProxy.cpp b/Engine/Source/Runtime/Experimental/Chaos/Private/PhysicsProxy/SingleParticlePhysicsProxy.cpp
--- a/Engine/Source/Runtime/Experimental/Chaos/Private/PhysicsProxy/SingleParticlePhysicsProxy.cpp
+++ b/Engine/Source/Runtime/Experimental/Chaos/Private/PhysicsProxy/SingleParticlePhysicsProxy.cpp
@@ -244,7 +244,16 @@ void FSingleParticlePhysicsProxy::BufferPhysicsResults_External(Chaos::FDirtyRig
 	}
 }
 
-bool FSingleParticlePhysicsProxy::PullFromPhysicsState(const Chaos::FDirtyRigidParticleData& PullData,int32 SolverSyncTimestamp, const Chaos::FDirtyRigidParticleData* NextPullData, const Chaos::FRealSingle* Alpha, const Chaos::FRealSingle* LeashAlpha)
+FRealSingle ResimInterpStrength = 0.2f;
+FAutoConsoleVariableRef CVarResimInterpStrength(TEXT("p.ResimInterpStrength"), ResimInterpStrength, TEXT("How strong the resim interp leash is. 1 means immediately snap to new target, 0 means do not interpolate at all"));
+
+FRealSingle MinLinError2ForResimInterp = 1.f;
+FAutoConsoleVariableRef CVarMinLinError2ForResimInterp(TEXT("p.MinLinError2ForResimInterp"), MinLinError2ForResimInterp, TEXT("The minimum squared error needed to continue interpolation during a resim"));
+
+FRealSingle MinRotErrorForResimInterp = 0.1f;
+FAutoConsoleVariableRef CVarMinRotErrorForResimInterp(TEXT("p.MinRotErrorForResimInterp"), MinRotErrorForResimInterp, TEXT("The minimum rotation error needed to continue interpolation during a resim"));
+
+bool FSingleParticlePhysicsProxy::PullFromPhysicsState(const Chaos::FDirtyRigidParticleData& PullData,int32 SolverSyncTimestamp, const Chaos::FDirtyRigidParticleData* NextPullData, const Chaos::FRealSingle* Alpha)
 {
 	using namespace Chaos;
 	// Move buffered data into the TPBDRigidParticle without triggering invalidation of the physics state.
@@ -269,12 +278,19 @@ bool FSingleParticlePhysicsProxy::PullFromPhysicsState(const Chaos::FDirtyRigidP
 
 			if (bSyncXR)
 			{
+				bool bKeepSmoothing = false;
 				if (const FVec3* Prev = LerpHelper(ProxyTimestamp->XTimestamp, PullData.X, ProxyTimestamp->OverWriteX))
 				{
 					FVec3 Target = FMath::Lerp(*Prev, NextPullData->X, *Alpha);
-					if (LeashAlpha)
+					if (IsResimSmoothing())
 					{
-						Target = FMath::Lerp(Rigid->X(), Target, *LeashAlpha);
+						const FVec3 SmoothedTarget = FMath::Lerp(Rigid->X(), Target, ResimInterpStrength);
+						if((SmoothedTarget - Target).SizeSquared() > MinLinError2ForResimInterp)
+						{
+							bKeepSmoothing = true;
+							Target = SmoothedTarget;
+						}
+
 					}
 					Rigid->SetX(Target, false);
 				}
@@ -282,31 +298,30 @@ bool FSingleParticlePhysicsProxy::PullFromPhysicsState(const Chaos::FDirtyRigidP
 				if (const FQuat* Prev = LerpHelper(ProxyTimestamp->RTimestamp, PullData.R, ProxyTimestamp->OverWriteR))
 				{
 					FQuat Target = FMath::Lerp(*Prev, NextPullData->R, *Alpha);
-					if (LeashAlpha)
+					if (IsResimSmoothing())
 					{
-						Target = FMath::Lerp<FQuat>(Rigid->R(), Target, *LeashAlpha);
+						const FQuat SmoothedTarget = FMath::Lerp<FQuat>(Rigid->R(), Target, ResimInterpStrength);
+						if(FQuat::ErrorAutoNormalize(SmoothedTarget, Target) > MinRotErrorForResimInterp)
+						{
+							bKeepSmoothing = true;
+							Target = SmoothedTarget;
+						}
 					}
 					Rigid->SetR(Target, false);
 				}
+
+				SetResimSmoothing(bKeepSmoothing);
 			}
 
 			if (const FVec3* Prev = LerpHelper(ProxyTimestamp->VTimestamp, PullData.V, ProxyTimestamp->OverWriteV))
 			{
 				FVec3 Target = FMath::Lerp(*Prev, NextPullData->V, *Alpha);
-				if (LeashAlpha)
-				{
-					Target = FMath::Lerp(Rigid->V(), Target, *LeashAlpha);
-				}
 				Rigid->SetV(Target, false);
 			}
 
 			if (const FVec3* Prev = LerpHelper(ProxyTimestamp->WTimestamp, PullData.W, ProxyTimestamp->OverWriteW))
 			{
 				FVec3 Target = FMath::Lerp(*Prev, NextPullData->W, *Alpha);
-				if (LeashAlpha)
-				{
-					Target = FMath::Lerp(Rigid->W(), Target, *LeashAlpha);
-				}
 				Rigid->SetW(Target, false);
 			}
 
diff --git a/Engine/Source/Runtime/Experimental/Chaos/Public/Chaos/Framework/ChaosResultsManager.h b/Engine/Source/Runtime/Experimental/Chaos/Public/Chaos/Framework/ChaosResultsManager.h
--- a/Engine/Source/Runtime/Experimental/Chaos/Public/Chaos/Framework/ChaosResultsManager.h
+++ b/Engine/Source/Runtime/Experimental/Chaos/Public/Chaos/Framework/ChaosResultsManager.h
@@ -7,19 +7,10 @@ namespace Chaos
 {
 	class FChaosMarshallingManager;
 	
-	struct FResimParticleInfo
-	{
-		FDirtyRigidParticleData Next;
-		FReal LeashStartTime = -1;
-		FReal EntryTime = -1;
-		bool bDiverged;
-	};
-
 	struct CHAOS_API FChaosRigidInterpolationData
 	{
 		FDirtyRigidParticleData Prev;
 		FDirtyRigidParticleData Next;
-		FRealSingle LeashAlpha = 1;
 	};
 
 	struct CHAOS_API FChaosInterpolationResults
@@ -48,16 +39,11 @@ namespace Chaos
 		const FChaosInterpolationResults& PullSyncPhysicsResults_External();
 		const FChaosInterpolationResults& PullAsyncPhysicsResults_External(const FReal ResultsTime);
 
-		void SetHistoryLength_External(int32 InLength);
 		void RemoveProxy_External(FSingleParticlePhysicsProxy* Proxy);
 		void SetLastExternalDt_External(const FReal InExternalDt) { LastExternalDt = InExternalDt; }
-		void SetResimInterpTime(const FReal InterpTime);
-		void SetResimInterpStrength(const FReal InterpStrength) { ResimInterpStrength = InterpStrength; }
-
 	private:
 
 		const FChaosInterpolationResults& UpdateInterpAlpha_External(const FReal ResultsTime, const FReal GlobalAlpha);
-		void FreeToHistory_External(FPullPhysicsData* PullData);
 		void ProcessResimResult_External();
 		bool AdvanceResult();
 		void CollapseResultsToLatest();
@@ -72,13 +58,9 @@ namespace Chaos
 		void SetPrevNextDataHelper(const FPullPhysicsData& PullData);
 
 		FChaosInterpolationResults Results;
-		TArray<FPullPhysicsData*> ResultsHistory;
 		FReal LatestTimeSeen = 0;	//we use this to know when resim results are being pushed
-		int32 HistoryLength = 0;
 		FChaosMarshallingManager& MarshallingManager;
 		FReal LastExternalDt = 0;
-		FReal InvResimInterpTime = 0;
-		FReal ResimInterpStrength = 1;
-		TMap<FSingleParticlePhysicsProxy*, FResimParticleInfo> ParticleToResimInfo;
+		TMap<FSingleParticlePhysicsProxy*, FDirtyRigidParticleData> ParticleToResimTarget;
 	};
 }
diff --git a/Engine/Source/Runtime/Experimental/Chaos/Public/Chaos/PhysicsSolverBaseImpl.h b/Engine/Source/Runtime/Experimental/Chaos/Public/Chaos/PhysicsSolverBaseImpl.h
--- a/Engine/Source/Runtime/Experimental/Chaos/Public/Chaos/PhysicsSolverBaseImpl.h
+++ b/Engine/Source/Runtime/Experimental/Chaos/Public/Chaos/PhysicsSolverBaseImpl.h
@@ -40,7 +40,7 @@ namespace Chaos
 			{
 				if(FSingleParticlePhysicsProxy* Proxy = RigidInterp.Prev.GetProxy())
 				{
-					if (Proxy->PullFromPhysicsState(RigidInterp.Prev, SolverTimestamp, &RigidInterp.Next, &Results.Alpha, &RigidInterp.LeashAlpha))
+					if (Proxy->PullFromPhysicsState(RigidInterp.Prev, SolverTimestamp, &RigidInterp.Next, &Results.Alpha))
 					{
 						RigidFunc(Proxy);
 					}
diff --git a/Engine/Source/Runtime/Experimental/Chaos/Public/PhysicsProxy/SingleParticlePhysicsProxy.h b/Engine/Source/Runtime/Experimental/Chaos/Public/PhysicsProxy/SingleParticlePhysicsProxy.h
--- a/Engine/Source/Runtime/Experimental/Chaos/Public/PhysicsProxy/SingleParticlePhysicsProxy.h
+++ b/Engine/Source/Runtime/Experimental/Chaos/Public/PhysicsProxy/SingleParticlePhysicsProxy.h
@@ -55,6 +55,8 @@ private:
 class FRigidBodyHandle_External;
 class FRigidBodyHandle_Internal;
 
+extern CHAOS_API FRealSingle ResimInterpStrength;
+
 class CHAOS_API FSingleParticlePhysicsProxy : public IPhysicsProxyBase
 {
 public:
@@ -73,6 +75,9 @@ public:
 		PullDataInterpIdx_External = Idx;
 	}
 
+	void SetResimSmoothing(bool ResimSmoothing) { bResimSmoothing = ResimSmoothing; }
+	bool IsResimSmoothing() const { return bResimSmoothing; }
+
 	int32 GetPullDataInterpIdx_External() const { return PullDataInterpIdx_External; }
 
 	FORCEINLINE FRigidBodyHandle_External& GetGameThreadAPI()
@@ -133,7 +138,7 @@ public:
 	void BufferPhysicsResults_External(FDirtyRigidParticleData&);
 
 	/**/
-	bool PullFromPhysicsState(const FDirtyRigidParticleData& PullData, int32 SolverSyncTimestamp, const FDirtyRigidParticleData* NextPullData = nullptr, const FRealSingle* Alpha = nullptr, const FRealSingle* LeashAlpha = nullptr);
+	bool PullFromPhysicsState(const FDirtyRigidParticleData& PullData, int32 SolverSyncTimestamp, const FDirtyRigidParticleData* NextPullData = nullptr, const FRealSingle* Alpha = nullptr);
 
 	/**/
 	bool IsDirty();
@@ -176,6 +181,8 @@ private:
 	//Used by interpolation code
 	int32 PullDataInterpIdx_External;
 
+	bool bResimSmoothing = false;
+
 	//use static Create
 	FSingleParticlePhysicsProxy(TUniquePtr<PARTICLE_TYPE>&& InParticle, FParticleHandle* InHandle, UObject* InOwner = nullptr);
 };
