Remove constraint from particle whether the constraint is enabled or not

diff --git a/Engine/Source/Runtime/Experimental/Chaos/Private/Chaos/PBDSuspensionConstraints.cpp b/Engine/Source/Runtime/Experimental/Chaos/Private/Chaos/PBDSuspensionConstraints.cpp
--- a/Engine/Source/Runtime/Experimental/Chaos/Private/Chaos/PBDSuspensionConstraints.cpp
+++ b/Engine/Source/Runtime/Experimental/Chaos/Private/Chaos/PBDSuspensionConstraints.cpp
@@ -60,7 +60,7 @@ namespace Chaos
 		FConstraintContainerHandle* ConstraintHandle = Handles[ConstraintIndex];
 		if (ConstraintHandle != nullptr)
 		{
-			if (ConstraintEnabledStates[ConstraintIndex] && ConstrainedParticles[ConstraintIndex])
+			if (ConstrainedParticles[ConstraintIndex])
 			{				
 				ConstrainedParticles[ConstraintIndex]->RemoveConstraintHandle(ConstraintHandle);
 			}
