Chaos : Heightfield bugs
- Fix issue with heighfield Phi computation when sample point is outside of the XY bounds of the heightfield
- Fix out of range issue if a point exactly on the far X  edge of the heightfield is queried for normals
- fix search radius in PhiWithNormal

diff --git a/Engine/Source/Runtime/Experimental/Chaos/Private/Chaos/HeightField.cpp b/Engine/Source/Runtime/Experimental/Chaos/Private/Chaos/HeightField.cpp
--- a/Engine/Source/Runtime/Experimental/Chaos/Private/Chaos/HeightField.cpp
+++ b/Engine/Source/Runtime/Experimental/Chaos/Private/Chaos/HeightField.cpp
@@ -601,55 +601,59 @@ namespace Chaos
 		Chaos::FVec3 Normal = Chaos::FVec3(0);
 	};
 
+	/**
+	* GetHeightNormalAt always return a valid normal
+	* if the point is outside of the grid, the edge is extended infinitely 
+	**/
 	template<bool bFillHeight, bool bFillNormal>
 	FHeightNormalResult GetHeightNormalAt(const FVec2& InGridLocationLocal, const Chaos::FHeightField::FDataType& InGeomData, const TUniformGrid<Chaos::FReal, 2>& InGrid)
 	{
 		FHeightNormalResult Result;
 
-		if(CHAOS_ENSURE(InGridLocationLocal == InGrid.Clamp(InGridLocationLocal)))
-		{
-			TVec2<int32> CellCoord = InGrid.Cell(InGridLocationLocal);
+		const FVec2 ClampedGridLocationLocal = InGrid.Clamp(InGridLocationLocal);
+		
+		TVec2<int32> CellCoord = InGrid.Cell(ClampedGridLocationLocal);
+		CellCoord = InGrid.ClampIndex(CellCoord);
 
-			const int32 SingleIndex = CellCoord[1] * (InGeomData.NumCols) + CellCoord[0];
-			FVec3 Pts[4];
-			InGeomData.GetPointsScaled(SingleIndex, Pts);
+		const int32 SingleIndex = CellCoord[1] * (InGeomData.NumCols) + CellCoord[0];
+		FVec3 Pts[4];
+		InGeomData.GetPointsScaled(SingleIndex, Pts);
 
-			float FractionX = FMath::Frac(InGridLocationLocal[0]);
-			float FractionY = FMath::Frac(InGridLocationLocal[1]);
+		const float FractionX = FMath::Frac(ClampedGridLocationLocal[0]);
+		const float FractionY = FMath::Frac(ClampedGridLocationLocal[1]);
 
-			if(FractionX > FractionY)
+		if(FractionX > FractionY)
+		{
+			if(bFillHeight)
 			{
-				if(bFillHeight)
-				{
-					const static FVector Tri[3] = {FVector(0.0f, 0.0f, 0.0f), FVector(1.0f, 1.0f, 0.0f), FVector(0.0f, 1.0f, 0.0f)};
-					FVector Bary = FMath::GetBaryCentric2D({ FractionX, FractionY, 0.0f }, Tri[0], Tri[1], Tri[2]);
+				const static FVector Tri[3] = {FVector(0.0f, 0.0f, 0.0f), FVector(1.0f, 1.0f, 0.0f), FVector(0.0f, 1.0f, 0.0f)};
+				FVector Bary = FMath::GetBaryCentric2D ({ FractionX, FractionY, 0.0f }, Tri[0], Tri[1], Tri[2]);
 
-					Result.Height = Pts[0].Z * Bary[0] + Pts[3].Z * Bary[1] + Pts[2].Z * Bary[2];
-				}
+				Result.Height = Pts[0].Z * Bary[0] + Pts[3].Z * Bary[1] + Pts[2].Z * Bary[2];
+			}
 
-				if(bFillNormal)
-				{
-					const FVec3 AB = Pts[3] - Pts[0];
-					const FVec3 AC = Pts[2] - Pts[0];
-					Result.Normal = FVec3::CrossProduct(AB, AC).GetUnsafeNormal();
-				}
+			if(bFillNormal)
+			{
+				const FVec3 AB = Pts[3] - Pts[0];
+				const FVec3 AC = Pts[2] - Pts[0];
+				Result.Normal = FVec3::CrossProduct(AB, AC).GetUnsafeNormal();
 			}
-			else
+		}
+		else
+		{
+			if(bFillHeight)
 			{
-				if(bFillHeight)
-				{
-					const static FVector Tri[3] = {FVector(0.0f, 0.0f, 0.0f), FVector(1.0f, 0.0f, 0.0f), FVector(1.0f, 1.0f, 0.0f)};
-					FVector Bary = FMath::GetBaryCentric2D({ FractionX, FractionY, 0.0f }, Tri[0], Tri[1], Tri[2]);
+				const static FVector Tri[3] = {FVector(0.0f, 0.0f, 0.0f), FVector(1.0f, 0.0f, 0.0f), FVector(1.0f, 1.0f, 0.0f)};
+				FVector Bary = FMath::GetBaryCentric2D({ FractionX, FractionY, 0.0f }, Tri[0], Tri[1], Tri[2]);
 
-					Result.Height = Pts[0].Z * Bary[0] + Pts[1].Z * Bary[1] + Pts[3].Z * Bary[2];
-				}
+				Result.Height = Pts[0].Z * Bary[0] + Pts[1].Z * Bary[1] + Pts[3].Z * Bary[2];
+			}
 
-				if(bFillNormal)
-				{
-					const FVec3 AB = Pts[1] - Pts[0];
-					const FVec3 AC = Pts[3] - Pts[0];
-					Result.Normal = FVec3::CrossProduct(AB, AC).GetUnsafeNormal();
-				}
+			if(bFillNormal)
+			{
+				const FVec3 AB = Pts[1] - Pts[0];
+				const FVec3 AC = Pts[3] - Pts[0];
+				Result.Normal = FVec3::CrossProduct(AB, AC).GetUnsafeNormal();
 			}
 		}
 
@@ -1792,15 +1796,16 @@ namespace Chaos
 		Chaos::FVec2 HeightField2DPosition(x.X / GeomData.Scale.X, x.Y / GeomData.Scale.Y);
 
 		FHeightNormalResult HeightNormal = GetHeightNormalAt<true, true>(HeightField2DPosition, GeomData, FlatGrid);
+		ensure(!HeightNormal.Normal.IsZero());
 
 		// Assume the cell below us is the correct result initially
 		const Chaos::FReal& HeightAtPoint = HeightNormal.Height;
 		Normal = HeightNormal.Normal;
 		FReal OutPhi = x.Z - HeightAtPoint;
-
-		// Need to sample all cells in a HeightAtPoint radius circle on the 2D grid. Large cliffs mean that the actual
+		
+		// Need to sample all cells in a Phi radius circle on the 2D grid. Large cliffs mean that the actual
 		// Phi could be in an entirely different cell.
-		const FClosestFaceData ClosestFace = FindClosestFace(x, HeightAtPoint);
+		const FClosestFaceData ClosestFace = FindClosestFace(x, FMath::Abs(OutPhi));
 
 		if(ClosestFace.FaceIndex > INDEX_NONE)
 		{
