Add Joint data to np2.DumpOnRewind.

diff --git a/Engine/Source/Runtime/Experimental/Chaos/Private/RewindData.cpp b/Engine/Source/Runtime/Experimental/Chaos/Private/RewindData.cpp
--- a/Engine/Source/Runtime/Experimental/Chaos/Private/RewindData.cpp
+++ b/Engine/Source/Runtime/Experimental/Chaos/Private/RewindData.cpp
@@ -326,6 +326,14 @@ void FRewindData::DumpHistory_Internal(const int32 FramePrintOffset, const FStri
 				FGeometryParticleState State = GetPastStateAtFrame(*Info.GetObjectPtr(), Frame, (FFrameAndPhase::EParticleHistoryPhase)Phase);
 				Out.Logf(TEXT("%s\n"), *State.ToString());
 			}
+
+			for (const FDirtyJointInfo& Info : DirtyJoints)
+			{
+				Out.Logf(TEXT("Frame:%d Phase:%d\n"), Frame + FramePrintOffset, Phase);
+
+				FJointState State = GetPastJointStateAtFrame(*Info.GetObjectPtr(), Frame, (FFrameAndPhase::EParticleHistoryPhase)Phase);
+				Out.Logf(TEXT("%s\n"), *State.ToString()); 
+			}
 		}
 	}
 
diff --git a/Engine/Source/Runtime/Experimental/Chaos/Public/RewindData.h b/Engine/Source/Runtime/Experimental/Chaos/Public/RewindData.h
--- a/Engine/Source/Runtime/Experimental/Chaos/Public/RewindData.h
+++ b/Engine/Source/Runtime/Experimental/Chaos/Public/RewindData.h
@@ -426,6 +426,19 @@ FString ToStringHelper(const T& Val)
 	return Val.ToString();
 }
 
+template <typename T>
+FString ToStringHelper(const TVector<T, 2>& Val)
+{
+	return FString::Printf(TEXT("(%s, %s)"), *Val[0].ToString(), *Val[1].ToString());
+}
+
+inline FString ToStringHelper(void* Val)
+{
+	// We don't print pointers because they will always be different in diff, need this function so we will compile
+	// when using property .inl macros.
+	return FString();
+}
+
 inline FString ToStringHelper(const FReal Val)
 {
 	return FString::Printf(TEXT("%f"), Val);
@@ -436,6 +449,21 @@ inline FString ToStringHelper(const EObjectStateType Val)
 	return FString::Printf(TEXT("%d"), Val);
 }
 
+inline FString ToStringHelper(const EPlasticityType Val)
+{
+	return FString::Printf(TEXT("%d"), Val);
+}
+
+inline FString ToStringHelper(const EJointForceMode Val)
+{
+	return FString::Printf(TEXT("%d"), Val);
+}
+
+inline FString ToStringHelper(const EJointMotionType Val)
+{
+	return FString::Printf(TEXT("%d"), Val);
+}
+
 inline FString ToStringHelper(const bool Val)
 {
 	return FString::Printf(TEXT("%d"), Val);
@@ -669,7 +697,7 @@ public:
 #undef REWIND_PARTICLE_TO_STR
 #define REWIND_PARTICLE_TO_STR(PropName) Out += FString::Printf(TEXT(#PropName":%s\n"), *ToStringHelper(PropName()));
 		//TODO: use macro to define api and the to string
-		FString Out = FString::Printf(TEXT("ParticleID:[%d%d]\n"), Particle.ParticleID().GlobalID, Particle.ParticleID().LocalID);
+		FString Out = FString::Printf(TEXT("ParticleID:[Global: %d Local: %d]\n"), Particle.ParticleID().GlobalID, Particle.ParticleID().LocalID);
 
 		REWIND_PARTICLE_TO_STR(X)
 		REWIND_PARTICLE_TO_STR(R)
@@ -784,6 +812,19 @@ public:
 #define CHAOS_INNER_JOINT_PROPERTY(OuterProp, FuncName, Inner, InnerType) REWIND_JOINT_PROPERTY(OuterProp, FuncName, Inner);
 #include "Chaos/JointProperties.inl"
 
+
+	FString ToString() const
+	{
+		TVector<FGeometryParticleHandle*, 2> Particles = Head.GetConstrainedParticles();
+		FString Out = FString::Printf(TEXT("Joint: Particle0 ID:[Global: %d Local: %d] Particle1 ID:[Global: %d Local: %d]\n"), Particles[0]->ParticleID().GlobalID, Particles[0]->ParticleID().LocalID, Particles[1]->ParticleID().GlobalID, Particles[1]->ParticleID().LocalID);
+
+#define CHAOS_INNER_JOINT_PROPERTY(OuterProp, FuncName, Inner, InnerType) Out += FString::Printf(TEXT(#FuncName":%s\n"), *ToStringHelper(Get##FuncName()));
+#include "Chaos/JointProperties.inl"
+#undef CHAOS_INNER_JOINT_PROPERTY
+
+		return Out;
+	}
+
 private:
 	const FPBDJointConstraintHandle& Head;
 	const FDirtyPropertiesPool& Pool;
