Results manager can use multiple channels with different time delays

diff --git a/Engine/Source/Programs/HeadlessChaos/Private/HeadlessChaosTestDataMarshalling.cpp b/Engine/Source/Programs/HeadlessChaos/Private/HeadlessChaosTestDataMarshalling.cpp
--- a/Engine/Source/Programs/HeadlessChaos/Private/HeadlessChaosTestDataMarshalling.cpp
+++ b/Engine/Source/Programs/HeadlessChaos/Private/HeadlessChaosTestDataMarshalling.cpp
@@ -223,9 +223,9 @@ namespace ChaosTest
 				const FReal StartTime = ExternalTime;	//external time we would have kicked the sim task off with
 				ExternalTime += ExternalDt;
 				MarshallingManager.FinalizePullData_Internal(Step, StartTime, ExternalDt);
-				const FChaosInterpolationResults& Results = ResultsManager.PullAsyncPhysicsResults_External(ExternalTime);
-				EXPECT_EQ(Results.Alpha, 1);	//async mode but no buffer so no interpolation
-				EXPECT_EQ(Results.Next->ExternalStartTime, StartTime);	//async mode but no buffer so should appear the same as sync
+				TArray<const FChaosInterpolationResults*> Results = ResultsManager.PullAsyncPhysicsResults_External(ExternalTime);
+				EXPECT_EQ(Results[0]->Alpha, 1);	//async mode but no buffer so no interpolation
+				EXPECT_EQ(Results[0]->Next->ExternalStartTime, StartTime);	//async mode but no buffer so should appear the same as sync
 			}
 		}
 
@@ -244,18 +244,18 @@ namespace ChaosTest
 				ExternalTime += ExternalDt;
 				const FReal RenderTime = ExternalTime - Delay;
 				MarshallingManager.FinalizePullData_Internal(Step, StartTime, ExternalDt);
-				FChaosInterpolationResults Results = ResultsManager.PullAsyncPhysicsResults_External(RenderTime);
+				TArray<const FChaosInterpolationResults*> Results = ResultsManager.PullAsyncPhysicsResults_External(RenderTime);
 				if(RenderTime < 0)
 				{
 					EXPECT_LT(Step, 2);	//first two frames treat as sync mode since we don't have enough delay
-					EXPECT_EQ(Results.Alpha, 1);
-					EXPECT_EQ(Results.Next, nullptr);
+					EXPECT_EQ(Results[0]->Alpha, 1);
+					EXPECT_EQ(Results[0]->Next, nullptr);
 				}
 				else
 				{
 					//after first two frames we have enough to interpolate
 					EXPECT_GE(Step, 2);
-					EXPECT_GT(Results.Next->ExternalEndTime, RenderTime);
+					EXPECT_GT(Results[0]->Next->ExternalEndTime, RenderTime);
 				}
 
 				PreTime = StartTime;
@@ -285,18 +285,18 @@ namespace ChaosTest
 					MarshallingManager.FinalizePullData_Internal(InnerStepTotal++, StartTime + InnerDt * InnerStep, InnerDt);
 				}
 
-				FChaosInterpolationResults Results = ResultsManager.PullAsyncPhysicsResults_External(RenderTime);
+				TArray<const FChaosInterpolationResults*> Results = ResultsManager.PullAsyncPhysicsResults_External(RenderTime);
 				if (RenderTime < 0)
 				{
 					EXPECT_LT(Step, 2);	//first two frames treat as sync mode since we don't have enough delay
-					EXPECT_EQ(Results.Alpha, 1);
-					EXPECT_EQ(Results.Next, nullptr);	//until we have enough results we just use the first result
+					EXPECT_EQ(Results[0]->Alpha, 1);
+					EXPECT_EQ(Results[0]->Next, nullptr);	//until we have enough results we just use the first result
 				}
 				else
 				{
 					//after first two frames we have enough to interpolate
 					EXPECT_GE(Step, 2);
-					EXPECT_GT(Results.Next->ExternalEndTime, RenderTime);
+					EXPECT_GT(Results[0]->Next->ExternalEndTime, RenderTime);
 				}
 
 				PreTime = StartTime;
diff --git a/Engine/Source/Programs/HeadlessChaos/Private/HeadlessChaosTestRewind.cpp b/Engine/Source/Programs/HeadlessChaos/Private/HeadlessChaosTestRewind.cpp
--- a/Engine/Source/Programs/HeadlessChaos/Private/HeadlessChaosTestRewind.cpp
+++ b/Engine/Source/Programs/HeadlessChaos/Private/HeadlessChaosTestRewind.cpp
@@ -4201,4 +4201,63 @@ namespace ChaosTest {
 		EXPECT_EQ(MaxAngularError, 0);
 	}
 
+
+	GTEST_TEST(AllTraits, RewindTest_InterpolatedTwoChannels)
+	{
+		int32 PrevNumActiveChannels = Chaos::DefaultNumActiveChannels;
+		Chaos::DefaultNumActiveChannels = 2;
+		//Have two moving particles, one in each channel to see that there's a delay in time on second channel
+		TRewindHelper::TestDynamicSphere([](auto* Solver, FReal SimDt, int32 Optimization, auto Proxy, auto Sphere)
+		{
+			if (!Solver->IsUsingAsyncResults()) { return; }
+			auto& Particle = Proxy->GetGameThreadAPI();
+			Particle.SetV(FVec3(0, 0, 1));
+			Particle.SetGravityEnabled(false);
+
+			auto Proxy2 = FSingleParticlePhysicsProxy::Create(Chaos::FPBDRigidParticle::CreateParticle());
+			auto& Particle2 = Proxy2->GetGameThreadAPI();
+
+			auto Sphere2 = TSharedPtr<FImplicitObject, ESPMode::ThreadSafe>(new TSphere<FReal, 3>(FVec3(0), 10));
+			Particle2.SetGeometry(Sphere2);
+			Particle2.SetV(FVec3(0, 0, 1));
+			Particle2.SetGravityEnabled(false);
+			Proxy2->SetInterpChannel_External(1);
+			Solver->RegisterObject(Proxy2);
+
+			FReal Time = 0;
+			const FReal GtDt = 1;
+			for (int Step = 0; Step < 32; ++Step)
+			{
+				TickSolverHelper(Solver);
+
+				Time += GtDt;
+				const FReal InterpolatedTime0 = Time - SimDt * Chaos::AsyncInterpolationMultiplier;
+				const FReal InterpolatedTime1 = InterpolatedTime0 - Chaos::SecondChannelDelay;
+
+				if (InterpolatedTime0 < 0)
+				{
+					//No movement yet
+					EXPECT_NEAR(Particle.X()[2], 0, 1e-2);
+				}
+				else
+				{
+					EXPECT_NEAR(Particle.X()[2], InterpolatedTime0, 1e-2);
+				}
+
+				if (InterpolatedTime1 < 0)
+				{
+					//No movement yet
+					EXPECT_NEAR(Particle2.X()[2], 0, 1e-2);
+				}
+				else
+				{
+					EXPECT_NEAR(Particle2.X()[2], InterpolatedTime1, 1e-2);
+				}
+			}
+		});
+
+		Chaos::DefaultNumActiveChannels = PrevNumActiveChannels;
+	}
+
+
 }
diff --git a/Engine/Source/Runtime/Experimental/Chaos/Private/Chaos/Framework/ChaosResultsManager.cpp b/Engine/Source/Runtime/Experimental/Chaos/Private/Chaos/Framework/ChaosResultsManager.cpp
--- a/Engine/Source/Runtime/Experimental/Chaos/Private/Chaos/Framework/ChaosResultsManager.cpp
+++ b/Engine/Source/Runtime/Experimental/Chaos/Private/Chaos/Framework/ChaosResultsManager.cpp
@@ -6,6 +6,46 @@
 
 namespace Chaos
 {	
+	enum class ESetPrevNextDataMode
+	{
+		Prev,
+		Next,
+	};
+
+	/** Helper class used to interpolate results per channel */
+	struct FChaosResultsChannel
+	{
+		FChaosResultsChannel(FChaosResultsManager& InResultsManager, int32 InChannelIdx)
+		: ResultsManager(InResultsManager)
+		, ChannelIdx(InChannelIdx)
+		{
+		}
+
+		~FChaosResultsChannel();
+		FChaosResultsChannel(const FChaosResultsChannel&) = delete;
+		FChaosResultsChannel(FChaosResultsChannel&& Other) = default;
+
+		const FChaosInterpolationResults& UpdateInterpAlpha_External(const FReal ResultsTime, const FReal GlobalAlpha);
+		void ProcessResimResult_External();
+		bool AdvanceResult();
+		void CollapseResultsToLatest();
+		const FChaosInterpolationResults& PullSyncPhysicsResults_External();
+		const FChaosInterpolationResults& PullAsyncPhysicsResults_External(const FReal ResultsTime);
+
+		template <ESetPrevNextDataMode Mode>
+		void SetPrevNextDataHelper(const FPullPhysicsData& PullData);
+		void RemoveProxy_External(FSingleParticlePhysicsProxy* Proxy)
+		{
+			ParticleToResimTarget.Remove(Proxy);
+		}
+
+		FChaosInterpolationResults Results;
+		FReal LatestTimeSeen = 0;	//we use this to know when resim results are being pushed
+		TMap<FSingleParticlePhysicsProxy*, FDirtyRigidParticleData> ParticleToResimTarget;
+		FChaosResultsManager& ResultsManager;
+		int32 ChannelIdx;
+	};
+
 	void FChaosInterpolationResults::Reset()
 	{
 		for (FChaosRigidInterpolationData& Data : RigidInterpolations)
@@ -20,14 +60,8 @@ namespace Chaos
 		//purposely leave Prev and Next alone as we use those for rebuild
 	}
 
-	enum class ESetPrevNextDataMode
-	{
-		Prev,
-		Next,
-	};
-
-	template <FChaosResultsManager::ESetPrevNextDataMode Mode>
-	void FChaosResultsManager::SetPrevNextDataHelper(const FPullPhysicsData& PullData)
+	template <ESetPrevNextDataMode Mode>
+	void FChaosResultsChannel::SetPrevNextDataHelper(const FPullPhysicsData& PullData)
 	{
 		//clear results
 		const int32 Timestamp = PullData.SolverTimestamp;
@@ -35,6 +69,12 @@ namespace Chaos
 		{
 			if (FSingleParticlePhysicsProxy* Proxy = Data.GetProxy())
 			{
+				//If proxy is not associated with this channel, do nothing
+				if (Proxy->GetInterpChannel_External() != ChannelIdx)
+				{
+					continue;
+				}
+
 				int32 DataIdx = Proxy->GetPullDataInterpIdx_External();
 				if(DataIdx == INDEX_NONE)
 				{
@@ -77,14 +117,14 @@ namespace Chaos
 	 @param Results Results to update with advanced state
 	 @return whether an advance occurred
 	*/
-	bool FChaosResultsManager::AdvanceResult()
+	bool FChaosResultsChannel::AdvanceResult()
 	{
-		if(FPullPhysicsData* PotentialNext = MarshallingManager.PopPullData_External())
+		if(FPullPhysicsData* PotentialNext = ResultsManager.PopPullData_External(ChannelIdx))
 		{
 			//newer result exists so prev is overwritten
 			if(Results.Prev)
 			{
-				MarshallingManager.FreePullData_External(Results.Prev);
+				ResultsManager.FreePullData_External(Results.Prev, ChannelIdx);
 			}
 
 			Results.Prev = Results.Next;
@@ -113,12 +153,12 @@ namespace Chaos
 	 @param MarshallingManager Manger to advance
 	 @param Results Results to update with advanced state
 	*/
-	void FChaosResultsManager::CollapseResultsToLatest()
+	void FChaosResultsChannel::CollapseResultsToLatest()
 	{
 		if(Results.Next == nullptr)
 		{
 			//nothing in Next (first time), so get latest if possible
-			Results.Next = MarshallingManager.PopPullData_External();
+			Results.Next = ResultsManager.PopPullData_External(ChannelIdx);
 		}
 
 		while(AdvanceResult())
@@ -126,11 +166,16 @@ namespace Chaos
 	}
 
 	const FChaosInterpolationResults& FChaosResultsManager::PullSyncPhysicsResults_External()
+	{
+		return Channels[0]->PullSyncPhysicsResults_External();
+	}
+
+	const FChaosInterpolationResults& FChaosResultsChannel::PullSyncPhysicsResults_External()
 	{
 		//sync mode doesn't use prev results, but if we were async previously we need to clean it up
 		if (Results.Prev)
 		{
-			MarshallingManager.FreePullData_External(Results.Prev);
+			ResultsManager.FreePullData_External(Results.Prev, ChannelIdx);
 			Results.Prev = nullptr;
 		}
 
@@ -164,12 +209,26 @@ namespace Chaos
 		return 1;	//if 0 dt just use 1 as alpha
 	}
 
+	FRealSingle SecondChannelDelay = 0.05f;
+	FAutoConsoleVariableRef CVarSecondChannelDelay(TEXT("p.SecondChannelDelay"), SecondChannelDelay, TEXT(""));
+
+	int32 DefaultNumActiveChannels = 1;
+	FAutoConsoleVariableRef CVarNumActiveChannels(TEXT("p.NumActiveChannels"), DefaultNumActiveChannels, TEXT(""));
+
 	FChaosResultsManager::FChaosResultsManager(FChaosMarshallingManager& InMarshallingManager)
 		: MarshallingManager(InMarshallingManager)
+		, NumActiveChannels(DefaultNumActiveChannels)
 	{
+		for(int32 Channel = 0; Channel < NumActiveChannels; ++Channel)
+		{
+			Channels.Emplace(new FChaosResultsChannel(*this, Channel));
+		}
+
+		PerChannelTimeDelay.Add(0);
+		PerChannelTimeDelay.Add(SecondChannelDelay);
 	}
 
-	const FChaosInterpolationResults& FChaosResultsManager::UpdateInterpAlpha_External(const FReal ResultsTime, const FReal GlobalAlpha)
+	const FChaosInterpolationResults& FChaosResultsChannel::UpdateInterpAlpha_External(FReal ResultsTime, const FReal GlobalAlpha)
 	{
 		Results.Alpha = (float)GlobalAlpha;	 // LWC_TODO: Precision loss
 
@@ -207,7 +266,18 @@ namespace Chaos
 		return Results;
 	}
 
-	const FChaosInterpolationResults& FChaosResultsManager::PullAsyncPhysicsResults_External(const FReal ResultsTime)
+	TArray<const FChaosInterpolationResults*> FChaosResultsManager::PullAsyncPhysicsResults_External(const FReal ResultsTime)
+	{
+		TArray<const FChaosInterpolationResults*> InterpResults;
+		for(int32 ChannelIdx = 0; ChannelIdx < NumActiveChannels; ++ChannelIdx)
+		{
+			InterpResults.Add(&Channels[ChannelIdx]->PullAsyncPhysicsResults_External(ResultsTime - PerChannelTimeDelay[ChannelIdx]));
+		}
+
+		return InterpResults;
+	}
+
+	const FChaosInterpolationResults& FChaosResultsChannel::PullAsyncPhysicsResults_External(const FReal ResultsTime)
 	{
 		//in async mode we must interpolate between Start and End of a particular sim step, where ResultsTime is in the inclusive interval [Start, End]
 		//to do this we need to keep the results of the previous sim step, which ends exactly when the next one starts
@@ -233,7 +303,7 @@ namespace Chaos
 		if (Results.Next == nullptr)
 		{
 			//nothing in Next (first time), so get latest if possible
-			Results.Next = MarshallingManager.PopPullData_External();
+			Results.Next = ResultsManager.PopPullData_External(ChannelIdx);
 		}
 
 		if(Results.Next)
@@ -267,34 +337,127 @@ namespace Chaos
 		return A.X != B.X || A.R != B.R || A.V != B.V || A.W != B.W || A.ObjectState != B.ObjectState;
 	}
 
-	void FChaosResultsManager::ProcessResimResult_External()
+	void FChaosResultsChannel::ProcessResimResult_External()
 	{
 		//make sure any proxy in the resim data is marked as resimming
 		for(const FDirtyRigidParticleData& ResimDirty : Results.Next->DirtyRigids)
 		{
 			if (FSingleParticlePhysicsProxy* ResimProxy = ResimDirty.GetProxy())
 			{
-				ParticleToResimTarget.FindOrAdd(ResimProxy) = ResimDirty;
-				ResimProxy->SetResimSmoothing(true);
+				//Mark as resim only if proxy is owned by this channel
+				if(ResimProxy->GetInterpChannel_External() == ChannelIdx)
+				{
+					ParticleToResimTarget.FindOrAdd(ResimProxy) = ResimDirty;
+					ResimProxy->SetResimSmoothing(true);
+				}
 			}
 		}
 	}
 
-	FChaosResultsManager::~FChaosResultsManager()
+	FChaosResultsChannel::~FChaosResultsChannel()
 	{
 		if(Results.Prev)
 		{
-			MarshallingManager.FreePullData_External(Results.Prev);
+			ResultsManager.FreePullData_External(Results.Prev, ChannelIdx);
 		}
 
 		if(Results.Next)
 		{
-			MarshallingManager.FreePullData_External(Results.Next);
+			ResultsManager.FreePullData_External(Results.Next, ChannelIdx);
 		}
 	}
 
 	void FChaosResultsManager::RemoveProxy_External(FSingleParticlePhysicsProxy* Proxy)
 	{
-		ParticleToResimTarget.Remove(Proxy);
+		for(FChaosResultsChannel* Channel : Channels)
+		{
+			Channel->RemoveProxy_External(Proxy);
+		}
+	}
+
+	FPullPhysicsData* FChaosResultsManager::PopPullData_External(int32 ChannelIdx)
+	{
+		int32 FoundToPop = INDEX_NONE;
+		for(int32 Idx = InternalQueue.Num() - 1; Idx >= 0; --Idx)
+		{
+			FPullDataQueueInfo& Info = InternalQueue[Idx];
+			if(!Info.bHasPopped[ChannelIdx])
+			{
+				FoundToPop = Idx;
+			}
+			else
+			{
+				//Found an entry that was popped so stop searching
+				break;
+			}
+		}
+
+		if(FoundToPop == INDEX_NONE)
+		{
+			//Need to pop from main queue
+			if(FPullPhysicsData* PullData = MarshallingManager.PopPullData_External())
+			{
+				FoundToPop = InternalQueue.Add(FPullDataQueueInfo(PullData, NumActiveChannels));
+			}
+			else
+			{
+				//Need to pop but nothing at head so return null
+				return nullptr;
+			}
+		}
+
+		FPullDataQueueInfo& Info = InternalQueue[FoundToPop];
+		Info.bHasPopped[ChannelIdx] = true;
+		return Info.PullData;
+	}
+
+	void FChaosResultsManager::FreePullData_External(FPullPhysicsData* PullData, int32 GivenChannelIdx)
+	{
+		for(int32 Idx = 0; Idx < InternalQueue.Num(); ++Idx)
+		{
+			FPullDataQueueInfo& Info = InternalQueue[Idx];
+			if(Info.PullData == PullData)
+			{
+				ensure(Info.bHasPopped[GivenChannelIdx]);	//free before popped?
+				ensure(Info.bPendingFree[GivenChannelIdx] == false);	//double free?
+
+				Info.bPendingFree[GivenChannelIdx] = true;
+
+				bool bFree = true;
+				for(int32 ChannelIdx = 0; ChannelIdx < MaxNumChannels; ++ChannelIdx)
+				{
+					if(!Info.bPendingFree[ChannelIdx])
+					{
+						bFree = false;
+						break;
+					}
+				}
+
+				if(bFree)
+				{
+					MarshallingManager.FreePullData_External(Info.PullData);
+					InternalQueue.RemoveAt(Idx);
+				}
+
+				return;
+			}
+		}
+
+		ensure(false);	//didn't find queue entry, double free?
+	}
+
+	FChaosResultsManager::~FChaosResultsManager()
+	{
+		//first clean up channels
+		for(FChaosResultsChannel* Channel : Channels)
+		{
+			delete Channel;
+		}
+
+		//if anything is left in the internal queue (for example channel was falling behind and never popped or freed) clear it
+		for(FPullDataQueueInfo& Info : InternalQueue)
+		{
+			MarshallingManager.FreePullData_External(Info.PullData);
+		}
 	}
 }
diff --git a/Engine/Source/Runtime/Experimental/Chaos/Private/Chaos/Framework/PhysicsSolverBase.cpp b/Engine/Source/Runtime/Experimental/Chaos/Private/Chaos/Framework/PhysicsSolverBase.cpp
--- a/Engine/Source/Runtime/Experimental/Chaos/Private/Chaos/Framework/PhysicsSolverBase.cpp
+++ b/Engine/Source/Runtime/Experimental/Chaos/Private/Chaos/Framework/PhysicsSolverBase.cpp
@@ -310,7 +310,6 @@ namespace Chaos
 		const bool bSubstepping = MMaxSubSteps > 1;
 		SetSolverSubstep_External(bSubstepping);
 		const FReal DtWithPause = bPaused_External ? 0.0f : InDt;
-		PullResultsManager->SetLastExternalDt_External(DtWithPause);
 		FReal InternalDt = DtWithPause;
 		int32 NumSteps = 1;
 
diff --git a/Engine/Source/Runtime/Experimental/Chaos/Public/Chaos/Framework/ChaosResultsManager.h b/Engine/Source/Runtime/Experimental/Chaos/Public/Chaos/Framework/ChaosResultsManager.h
--- a/Engine/Source/Runtime/Experimental/Chaos/Public/Chaos/Framework/ChaosResultsManager.h
+++ b/Engine/Source/Runtime/Experimental/Chaos/Public/Chaos/Framework/ChaosResultsManager.h
@@ -7,6 +7,9 @@ namespace Chaos
 {
 	class FChaosMarshallingManager;
 	
+	extern FRealSingle SecondChannelDelay;
+	extern int32 DefaultNumActiveChannels;
+
 	struct CHAOS_API FChaosRigidInterpolationData
 	{
 		FDirtyRigidParticleData Prev;
@@ -21,6 +24,17 @@ namespace Chaos
 		{
 		}
 
+		FChaosInterpolationResults(const FChaosInterpolationResults& Other) = delete;
+		FChaosInterpolationResults(FChaosInterpolationResults&& Other)
+			: RigidInterpolations(MoveTemp(Other.RigidInterpolations))
+			, Prev(Other.Prev)
+			, Next(Other.Next)
+			, Alpha(Other.Alpha)
+		{
+			Other.Prev = nullptr;
+			Other.Next = nullptr;
+		}
+
 		void Reset();
 		
 		TArray<FChaosRigidInterpolationData> RigidInterpolations;
@@ -29,38 +43,63 @@ namespace Chaos
 		FRealSingle Alpha;
 	};
 
+	struct FChaosResultsChannel;
+
 	class CHAOS_API FChaosResultsManager
 	{
 	public:
 		FChaosResultsManager(FChaosMarshallingManager& InMarshallingManager);
+		FChaosResultsManager(const FChaosResultsManager& Other) = delete;
+		FChaosResultsManager(FChaosResultsManager&& Other) = default;
 		
 		~FChaosResultsManager();
 
 		const FChaosInterpolationResults& PullSyncPhysicsResults_External();
-		const FChaosInterpolationResults& PullAsyncPhysicsResults_External(const FReal ResultsTime);
+		TArray<const FChaosInterpolationResults*> PullAsyncPhysicsResults_External(const FReal ResultsTime);
 
 		void RemoveProxy_External(FSingleParticlePhysicsProxy* Proxy);
-		void SetLastExternalDt_External(const FReal InExternalDt) { LastExternalDt = InExternalDt; }
+
 	private:
 
-		const FChaosInterpolationResults& UpdateInterpAlpha_External(const FReal ResultsTime, const FReal GlobalAlpha);
-		void ProcessResimResult_External();
-		bool AdvanceResult();
-		void CollapseResultsToLatest();
+		friend FChaosResultsChannel;
 
-		enum class ESetPrevNextDataMode
+		FPullPhysicsData* PopPullData_External(int32 ChannelIdx);
+		void FreePullData_External(FPullPhysicsData* PullData, int32 ChannelIdx);
+
+		TArray<FChaosResultsChannel*> Channels;
+		FChaosMarshallingManager& MarshallingManager;
+
+		static constexpr int32 MaxNumChannels = 2;
+
+		struct FPullDataQueueInfo
 		{
-			Prev,
-			Next,
-		};
+			FPullDataQueueInfo(FPullPhysicsData* Data, int32 NumActiveChannels)
+			: PullData(Data)
+			{
+				for(int32 Idx = 0; Idx < NumActiveChannels; ++Idx)
+				{
+					bHasPopped[Idx] = false;
+					bPendingFree[Idx] = false;
+				}
 
-		template <ESetPrevNextDataMode Mode>
-		void SetPrevNextDataHelper(const FPullPhysicsData& PullData);
+				//inactive channels are treated like they've already consumed the data
+				for (int32 Idx = NumActiveChannels; Idx < MaxNumChannels; ++Idx)
+				{
+					bHasPopped[Idx] = true;
+					bPendingFree[Idx] = true;
+				}
+			}
 
-		FChaosInterpolationResults Results;
-		FReal LatestTimeSeen = 0;	//we use this to know when resim results are being pushed
-		FChaosMarshallingManager& MarshallingManager;
-		FReal LastExternalDt = 0;
-		TMap<FSingleParticlePhysicsProxy*, FDirtyRigidParticleData> ParticleToResimTarget;
+			FPullDataQueueInfo(FPullDataQueueInfo& Other) = delete;
+			FPullDataQueueInfo(FPullDataQueueInfo&& Other) = default;
+
+			FPullPhysicsData* PullData;
+			bool bHasPopped[MaxNumChannels];
+			bool bPendingFree[MaxNumChannels];
+		};
+
+		TArray<FPullDataQueueInfo> InternalQueue;
+		int32 NumActiveChannels;
+		TArray<FRealSingle> PerChannelTimeDelay;
 	};
 }
diff --git a/Engine/Source/Runtime/Experimental/Chaos/Public/Chaos/PhysicsSolverBaseImpl.h b/Engine/Source/Runtime/Experimental/Chaos/Public/Chaos/PhysicsSolverBaseImpl.h
--- a/Engine/Source/Runtime/Experimental/Chaos/Public/Chaos/PhysicsSolverBaseImpl.h
+++ b/Engine/Source/Runtime/Experimental/Chaos/Public/Chaos/PhysicsSolverBaseImpl.h
@@ -32,21 +32,26 @@ namespace Chaos
 			//case 4: prev has no dirty data and next does. In this case interpolate from gt data to next
 			//case 5: prev has no dirty data and next was overwritten. In this case do nothing as the overwritten data wins, and also particle may be deleted
 
-			const FChaosInterpolationResults& Results = PullResultsManager->PullAsyncPhysicsResults_External(ResultsTime);
-			LatestData = Results.Next;
-			//todo: go wide
-			const int32 SolverTimestamp = Results.Next ? Results.Next->SolverTimestamp : INDEX_NONE;
-			for(const FChaosRigidInterpolationData& RigidInterp : Results.RigidInterpolations)
+			TArray<const FChaosInterpolationResults*> ResultsPerChannel = PullResultsManager->PullAsyncPhysicsResults_External(ResultsTime);
+			for(int32 ChannelIdx = 0; ChannelIdx < ResultsPerChannel.Num(); ++ChannelIdx)
 			{
-				if(FSingleParticlePhysicsProxy* Proxy = RigidInterp.Prev.GetProxy())
+				const FChaosInterpolationResults& Results = *ResultsPerChannel[ChannelIdx];
+				LatestData = Results.Next;
+				//todo: go wide
+				const int32 SolverTimestamp = Results.Next ? Results.Next->SolverTimestamp : INDEX_NONE;
+				for (const FChaosRigidInterpolationData& RigidInterp : Results.RigidInterpolations)
 				{
-					if (Proxy->PullFromPhysicsState(RigidInterp.Prev, SolverTimestamp, &RigidInterp.Next, &Results.Alpha))
+					if (FSingleParticlePhysicsProxy* Proxy = RigidInterp.Prev.GetProxy())
 					{
-						RigidFunc(Proxy);
+						if (Proxy->PullFromPhysicsState(RigidInterp.Prev, SolverTimestamp, &RigidInterp.Next, &Results.Alpha))
+						{
+							RigidFunc(Proxy);
+						}
 					}
+
 				}
-				
 			}
+			
 		}
 		else
 		{
diff --git a/Engine/Source/Runtime/Experimental/Chaos/Public/PhysicsProxy/SingleParticlePhysicsProxy.h b/Engine/Source/Runtime/Experimental/Chaos/Public/PhysicsProxy/SingleParticlePhysicsProxy.h
--- a/Engine/Source/Runtime/Experimental/Chaos/Public/PhysicsProxy/SingleParticlePhysicsProxy.h
+++ b/Engine/Source/Runtime/Experimental/Chaos/Public/PhysicsProxy/SingleParticlePhysicsProxy.h
@@ -78,6 +78,9 @@ public:
 	void SetResimSmoothing(bool ResimSmoothing) { bResimSmoothing = ResimSmoothing; }
 	bool IsResimSmoothing() const { return bResimSmoothing; }
 
+	int32 GetInterpChannel_External() const { return InterpChannel_External; }
+	void SetInterpChannel_External(int32 Channel) { InterpChannel_External = Channel; }
+
 	int32 GetPullDataInterpIdx_External() const { return PullDataInterpIdx_External; }
 
 	FORCEINLINE FRigidBodyHandle_External& GetGameThreadAPI()
@@ -180,7 +183,7 @@ private:
 
 	//Used by interpolation code
 	int32 PullDataInterpIdx_External;
-
+	int32 InterpChannel_External = 0;
 	bool bResimSmoothing = false;
 
 	//use static Create
