Make sure to ignore PushData when initialization step changes

diff --git a/Engine/Source/Programs/HeadlessChaos/Private/HeadlessChaosTestRewind.cpp b/Engine/Source/Programs/HeadlessChaos/Private/HeadlessChaosTestRewind.cpp
--- a/Engine/Source/Programs/HeadlessChaos/Private/HeadlessChaosTestRewind.cpp
+++ b/Engine/Source/Programs/HeadlessChaos/Private/HeadlessChaosTestRewind.cpp
@@ -587,6 +587,108 @@ namespace ChaosTest {
 			});
 	}
 
+	GTEST_TEST(AllTraits, RewindTest_SpawnEarlierCorrection2)
+	{
+		// Test resim when object spawned earlier as part of correction
+		TRewindHelper::TestDynamicSphere([](auto* Solver, FReal SimDt, int32 Optimization, auto Proxy, auto Sphere)
+			{
+				FSingleParticlePhysicsProxy* SpawnedProxy = nullptr;
+				FSingleParticlePhysicsProxy* SpawnedProxyNoCorrection = nullptr;
+
+				bool bHasResimmed = false;
+
+				FRewindCallbackTestHelper* Helper = RegisterCallbackHelper(Solver);
+				Helper->TriggerRewindFunc = [&SpawnedProxy, &SpawnedProxyNoCorrection, &bHasResimmed](int32 PhysicsStep) -> int32
+				{
+					
+					if(!bHasResimmed && SpawnedProxy)
+					{
+						bHasResimmed = true;
+						return 0;
+					}
+					
+					return INDEX_NONE;
+				};
+
+				Helper->ProcessInputsFunc = [RewindData = Solver->GetRewindData(), SimDt, &SpawnedProxy, &SpawnedProxyNoCorrection, &bHasResimmed](int32 PhysicsStep, bool bIsResimming)
+				{
+					const FReal Time = PhysicsStep * SimDt;
+
+					if (bIsResimming && SpawnedProxy)
+					{
+						if (PhysicsStep == 10)
+						{
+							RewindData->SpawnProxyIfNeeded(*SpawnedProxy);
+							SpawnedProxy->GetPhysicsThreadAPI()->SetX(FVec3(500, 0, 100.0));
+							
+						}
+					}
+
+					if (SpawnedProxy && SpawnedProxyNoCorrection)
+					{
+						auto PT0 = SpawnedProxy->GetPhysicsThreadAPI();
+						auto PT1 = SpawnedProxyNoCorrection->GetPhysicsThreadAPI();
+						
+						// After we've applied the correction and simmed past frame 10, we expect the first proxy to be "ahead" of the second one that didn't get corrected
+						if (bHasResimmed && PhysicsStep > 10)
+						{
+							EXPECT_LT(PT0->X()[2], PT1->X()[2]);
+						}
+					}					
+				};
+
+				const int32 LastGameStep = 32;
+				
+				auto SphereGeom = TSharedPtr<FImplicitObject, ESPMode::ThreadSafe>(new TSphere<FReal, 3>(FVec3(0), 10));
+
+				auto& Particle = Proxy->GetGameThreadAPI();
+				Particle.SetGravityEnabled(false);
+				Particle.SetV(FVec3(0, 0, -1));
+				Particle.SetX(FVec3(0, 0, 14.5));
+
+				ChaosTest::SetParticleSimDataToCollide({ Proxy->GetParticle_LowLevel() });
+
+				for (int Step = 0; Step <= LastGameStep; ++Step)
+				{
+					TickSolverHelper(Solver);
+
+					//spawn floor way late to ensure no collision on first run
+					if (Step == 12)
+					{
+
+						// Make particles. E.g:
+						//	we just found out from the server these were spawned and these are the latest positions replicated from the server.
+						//	but actually these positions are the server positions from frame 10. Lets see if we can correct this.
+						{
+							SpawnedProxy = FSingleParticlePhysicsProxy::Create(Chaos::FPBDRigidParticle::CreateParticle());
+							auto& GT = SpawnedProxy->GetGameThreadAPI();
+
+							GT.SetV(FVec3(0, 0, 0.0));
+							GT.SetX(FVec3(500, 0, 100.0));
+
+							GT.SetGeometry(SphereGeom);
+							Solver->RegisterObject(SpawnedProxy);
+							ChaosTest::SetParticleSimDataToCollide({ SpawnedProxy->GetParticle_LowLevel() });
+						}
+
+						{
+							SpawnedProxyNoCorrection = FSingleParticlePhysicsProxy::Create(Chaos::FPBDRigidParticle::CreateParticle());
+							auto& GT = SpawnedProxyNoCorrection->GetGameThreadAPI();
+
+							GT.SetV(FVec3(0, 0, 0.0));
+							GT.SetX(FVec3(100, 0, 100.0));
+
+							GT.SetGeometry(SphereGeom);
+							Solver->RegisterObject(SpawnedProxyNoCorrection);
+							ChaosTest::SetParticleSimDataToCollide({ SpawnedProxyNoCorrection->GetParticle_LowLevel() });
+						}
+						
+
+					}
+				}
+			});
+	}
+
 	GTEST_TEST(AllTraits, RewindTest_MovingToNotMovingInterpolation)
 	{
 		//
diff --git a/Engine/Source/Runtime/Experimental/Chaos/Private/PBDRigidsSolver.cpp b/Engine/Source/Runtime/Experimental/Chaos/Private/PBDRigidsSolver.cpp
--- a/Engine/Source/Runtime/Experimental/Chaos/Private/PBDRigidsSolver.cpp
+++ b/Engine/Source/Runtime/Experimental/Chaos/Private/PBDRigidsSolver.cpp
@@ -994,44 +994,47 @@ namespace Chaos
 		//need to create new particle handles
 		DirtyProxiesData->ForEachProxy([this,&ProcessProxyPT](int32 DataIdx,FDirtyProxy& Dirty)
 		{
-			switch(Dirty.Proxy->GetType())
+			if(Dirty.Proxy->GetIgnoreDataOnStep_Internal() != CurrentFrame)
 			{
-				case EPhysicsProxyType::SingleParticleProxy:
+				switch(Dirty.Proxy->GetType())
 				{
-					auto Proxy = static_cast<FSingleParticlePhysicsProxy*>(Dirty.Proxy);
-					ProcessProxyPT(Proxy, DataIdx, Dirty, [this, &Dirty](const FUniqueIdx* UniqueIdx) -> TGeometryParticleHandle<FReal,3>*
+					case EPhysicsProxyType::SingleParticleProxy:
 					{
-						switch (Dirty.PropertyData.GetParticleBufferType())
+						auto Proxy = static_cast<FSingleParticlePhysicsProxy*>(Dirty.Proxy);
+						ProcessProxyPT(Proxy, DataIdx, Dirty, [this, &Dirty](const FUniqueIdx* UniqueIdx) -> TGeometryParticleHandle<FReal,3>*
 						{
-							case EParticleType::Static: return Particles.CreateStaticParticles(1, UniqueIdx)[0];
-							case EParticleType::Kinematic: return Particles.CreateKinematicParticles(1, UniqueIdx)[0];
-							case EParticleType::Rigid: return Particles.CreateDynamicParticles(1, UniqueIdx)[0];
-							default: check(false); return nullptr;
-						}
-					});
-					break;
-				}
+							switch (Dirty.PropertyData.GetParticleBufferType())
+							{
+								case EParticleType::Static: return Particles.CreateStaticParticles(1, UniqueIdx)[0];
+								case EParticleType::Kinematic: return Particles.CreateKinematicParticles(1, UniqueIdx)[0];
+								case EParticleType::Rigid: return Particles.CreateDynamicParticles(1, UniqueIdx)[0];
+								default: check(false); return nullptr;
+							}
+						});
+						break;
+					}
 			
-				case EPhysicsProxyType::GeometryCollectionType:
-				{
-					auto Proxy = static_cast<FGeometryCollectionPhysicsProxy*>(Dirty.Proxy);
-					Proxy->PushToPhysicsState();
-					// Currently no push needed for geometry collections and they handle the particle creation internally
-					// #TODO This skips the rewind data push so GC will not be rewindable until resolved.
-					Dirty.Proxy->ResetDirtyIdx();
-					break;
-				}
-				case EPhysicsProxyType::JointConstraintType:
-				case EPhysicsProxyType::SuspensionConstraintType:
-				{
-					// Pass until after all bodies are created. 
-					break;
-				}
-				default:
-				{
-					ensure(0 && TEXT("Unknown proxy type in physics solver."));
-					//Can't use, but we can still mark as "clean"
-					Dirty.Proxy->ResetDirtyIdx();
+					case EPhysicsProxyType::GeometryCollectionType:
+					{
+						auto Proxy = static_cast<FGeometryCollectionPhysicsProxy*>(Dirty.Proxy);
+						Proxy->PushToPhysicsState();
+						// Currently no push needed for geometry collections and they handle the particle creation internally
+						// #TODO This skips the rewind data push so GC will not be rewindable until resolved.
+						Dirty.Proxy->ResetDirtyIdx();
+						break;
+					}
+					case EPhysicsProxyType::JointConstraintType:
+					case EPhysicsProxyType::SuspensionConstraintType:
+					{
+						// Pass until after all bodies are created. 
+						break;
+					}
+					default:
+					{
+						ensure(0 && TEXT("Unknown proxy type in physics solver."));
+						//Can't use, but we can still mark as "clean"
+						Dirty.Proxy->ResetDirtyIdx();
+					}
 				}
 			}
 		});
@@ -1039,44 +1042,47 @@ namespace Chaos
 		//need to create new constraint handles
 		DirtyProxiesData->ForEachProxy([this, &ProcessProxyPT, Manager, RewindData](int32 DataIdx, FDirtyProxy& Dirty)
 		{
-			switch (Dirty.Proxy->GetType())
-			{
-			case EPhysicsProxyType::JointConstraintType:
+			if (Dirty.Proxy->GetIgnoreDataOnStep_Internal() != CurrentFrame)
 			{
-				auto JointProxy = static_cast<FJointConstraintPhysicsProxy*>(Dirty.Proxy);
-				const bool bIsNew = !JointProxy->IsInitialized();
-				if (bIsNew)
+				switch (Dirty.Proxy->GetType())
 				{
-					JointConstraintPhysicsProxies_Internal.Add(JointProxy);
-					JointProxy->InitializeOnPhysicsThread(this, *Manager, DataIdx, Dirty.PropertyData);
-					JointProxy->SetInitialized(GetCurrentFrame());
-				}
-
-				//TODO: if we support predicting creation / destruction of joints need to handle null joint case
-				if (RewindData)
+				case EPhysicsProxyType::JointConstraintType:
 				{
-					RewindData->PushGTDirtyData(*Manager, DataIdx, Dirty, nullptr);
-				}
+					auto JointProxy = static_cast<FJointConstraintPhysicsProxy*>(Dirty.Proxy);
+					const bool bIsNew = !JointProxy->IsInitialized();
+					if (bIsNew)
+					{
+						JointConstraintPhysicsProxies_Internal.Add(JointProxy);
+						JointProxy->InitializeOnPhysicsThread(this, *Manager, DataIdx, Dirty.PropertyData);
+						JointProxy->SetInitialized(GetCurrentFrame());
+					}
+
+					//TODO: if we support predicting creation / destruction of joints need to handle null joint case
+					if (RewindData)
+					{
+						RewindData->PushGTDirtyData(*Manager, DataIdx, Dirty, nullptr);
+					}
 				
-				JointProxy->PushStateOnPhysicsThread(this, *Manager, DataIdx, Dirty.PropertyData);
-				Dirty.Proxy->ResetDirtyIdx();
-				break;
-			}
+					JointProxy->PushStateOnPhysicsThread(this, *Manager, DataIdx, Dirty.PropertyData);
+					Dirty.Proxy->ResetDirtyIdx();
+					break;
+				}
 
-			case EPhysicsProxyType::SuspensionConstraintType:
-			{
-				auto SuspensionProxy = static_cast<FSuspensionConstraintPhysicsProxy*>(Dirty.Proxy);
-				const bool bIsNew = !SuspensionProxy->IsInitialized();
-				if (bIsNew)
+				case EPhysicsProxyType::SuspensionConstraintType:
 				{
-					SuspensionProxy->InitializeOnPhysicsThread(this, *Manager, DataIdx, Dirty.PropertyData);
-					SuspensionProxy->SetInitialized();
+					auto SuspensionProxy = static_cast<FSuspensionConstraintPhysicsProxy*>(Dirty.Proxy);
+					const bool bIsNew = !SuspensionProxy->IsInitialized();
+					if (bIsNew)
+					{
+						SuspensionProxy->InitializeOnPhysicsThread(this, *Manager, DataIdx, Dirty.PropertyData);
+						SuspensionProxy->SetInitialized();
+					}
+					SuspensionProxy->PushStateOnPhysicsThread(this, *Manager, DataIdx, Dirty.PropertyData);
+					Dirty.Proxy->ResetDirtyIdx();
+					break;
 				}
-				SuspensionProxy->PushStateOnPhysicsThread(this, *Manager, DataIdx, Dirty.PropertyData);
-				Dirty.Proxy->ResetDirtyIdx();
-				break;
-			}
 
+				}
 			}
 		});
 
diff --git a/Engine/Source/Runtime/Experimental/Chaos/Public/Chaos/Framework/PhysicsProxyBase.h b/Engine/Source/Runtime/Experimental/Chaos/Public/Chaos/Framework/PhysicsProxyBase.h
--- a/Engine/Source/Runtime/Experimental/Chaos/Public/Chaos/Framework/PhysicsProxyBase.h
+++ b/Engine/Source/Runtime/Experimental/Chaos/Public/Chaos/Framework/PhysicsProxyBase.h
@@ -74,9 +74,19 @@ public:
 	TSharedPtr<FProxyTimestamp,ESPMode::ThreadSafe> GetSyncTimestamp() const { return SyncTimestamp; }
 
 	bool IsInitialized() const { return InitializedOnStep != INDEX_NONE; }
-	void SetInitialized(const int32 InitializeStep) { InitializedOnStep = InitializeStep; }
+	void SetInitialized(const int32 InitializeStep)
+	{
+		//If changed initialization, ignore the very first initialization push data
+		if(InitializedOnStep != InitializeStep && InitializedOnStep != INDEX_NONE)
+		{
+			IgnoreDataOnStep_Internal = InitializedOnStep;
+		}
+
+		InitializedOnStep = InitializeStep;
+	}
 	int32 GetInitializedStep() const { return InitializedOnStep; }
 
+	int32 GetIgnoreDataOnStep_Internal() const { return IgnoreDataOnStep_Internal; }
 
 protected:
 	// Ensures that derived classes can successfully call this destructor
@@ -94,6 +104,7 @@ protected:
 	EPhysicsProxyType Type;
 	TSharedPtr<FProxyTimestamp,ESPMode::ThreadSafe> SyncTimestamp;
 	int32 InitializedOnStep = INDEX_NONE;
+	int32 IgnoreDataOnStep_Internal = INDEX_NONE;
 };
 
 struct PhysicsProxyWrapper
