Network Physics

Update InitializedOnStep when SpawnProxyIfNeeded is called, otherwise particle can be stuck in an inactive state after receiving multiple corrections

diff --git a/Engine/Source/Runtime/Experimental/Chaos/Private/RewindData.cpp b/Engine/Source/Runtime/Experimental/Chaos/Private/RewindData.cpp
--- a/Engine/Source/Runtime/Experimental/Chaos/Private/RewindData.cpp
+++ b/Engine/Source/Runtime/Experimental/Chaos/Private/RewindData.cpp
@@ -562,6 +562,7 @@ void FRewindData::SpawnProxyIfNeeded(FSingleParticlePhysicsProxy& Proxy)
 		{
 			DesyncObject(Info, FFrameAndPhase{ Proxy.GetInitializedStep(), FFrameAndPhase::PrePushData });	//Spawned earlier so mark as desynced from that first frame
 			Proxy.SetInitialized(CurFrame);
+			Info.InitializedOnStep = CurFrame;
 		}
 	}
 }
