Network Physics
-Add second cvar for interp channel tuning

diff --git a/Engine/Source/Runtime/Experimental/Chaos/Private/PhysicsProxy/SingleParticlePhysicsProxy.cpp b/Engine/Source/Runtime/Experimental/Chaos/Private/PhysicsProxy/SingleParticlePhysicsProxy.cpp
--- a/Engine/Source/Runtime/Experimental/Chaos/Private/PhysicsProxy/SingleParticlePhysicsProxy.cpp
+++ b/Engine/Source/Runtime/Experimental/Chaos/Private/PhysicsProxy/SingleParticlePhysicsProxy.cpp
@@ -247,6 +247,9 @@ void FSingleParticlePhysicsProxy::BufferPhysicsResults_External(Chaos::FDirtyRig
 FRealSingle ResimInterpStrength = 0.2f;
 FAutoConsoleVariableRef CVarResimInterpStrength(TEXT("p.ResimInterpStrength"), ResimInterpStrength, TEXT("How strong the resim interp leash is. 1 means immediately snap to new target, 0 means do not interpolate at all"));
 
+FRealSingle ResimInterpStrength2 = 0.05f;
+FAutoConsoleVariableRef CVarResimInterpStrength2(TEXT("p.ResimInterpStrength2"), ResimInterpStrength2, TEXT("How strong the resim interp leash is for object in channel 2. 1 means immediately snap to new target, 0 means do not interpolate at all"));
+
 FRealSingle MinLinError2ForResimInterp = 1.f;
 FAutoConsoleVariableRef CVarMinLinError2ForResimInterp(TEXT("p.MinLinError2ForResimInterp"), MinLinError2ForResimInterp, TEXT("The minimum squared error needed to continue interpolation during a resim"));
 
@@ -278,13 +281,15 @@ bool FSingleParticlePhysicsProxy::PullFromPhysicsState(const Chaos::FDirtyRigidP
 
 			if (bSyncXR)
 			{
+				const float UseResimInterpStrength = GetInterpChannel_External() == 0 ? ResimInterpStrength : ResimInterpStrength2;
+
 				bool bKeepSmoothing = false;
 				if (const FVec3* Prev = LerpHelper(ProxyTimestamp->XTimestamp, PullData.X, ProxyTimestamp->OverWriteX))
 				{
 					FVec3 Target = FMath::Lerp(*Prev, NextPullData->X, *Alpha);
 					if (IsResimSmoothing())
 					{
-						const FVec3 SmoothedTarget = FMath::Lerp(Rigid->X(), Target, ResimInterpStrength);
+						const FVec3 SmoothedTarget = FMath::Lerp(Rigid->X(), Target, UseResimInterpStrength);
 						if((SmoothedTarget - Target).SizeSquared() > MinLinError2ForResimInterp)
 						{
 							bKeepSmoothing = true;
@@ -300,7 +305,7 @@ bool FSingleParticlePhysicsProxy::PullFromPhysicsState(const Chaos::FDirtyRigidP
 					FQuat Target = FMath::Lerp(*Prev, NextPullData->R, *Alpha);
 					if (IsResimSmoothing())
 					{
-						const FQuat SmoothedTarget = FMath::Lerp<FQuat>(Rigid->R(), Target, ResimInterpStrength);
+						const FQuat SmoothedTarget = FMath::Lerp<FQuat>(Rigid->R(), Target, UseResimInterpStrength);
 						if(FQuat::ErrorAutoNormalize(SmoothedTarget, Target) > MinRotErrorForResimInterp)
 						{
 							bKeepSmoothing = true;
