improve point ordering algorithm for chaos trianglemesh

#rnx

diff --git a/Engine/Source/Runtime/Experimental/Chaos/Private/Chaos/TriangleMesh.cpp b/Engine/Source/Runtime/Experimental/Chaos/Private/Chaos/TriangleMesh.cpp
--- a/Engine/Source/Runtime/Experimental/Chaos/Private/Chaos/TriangleMesh.cpp
+++ b/Engine/Source/Runtime/Experimental/Chaos/Private/Chaos/TriangleMesh.cpp
@@ -745,7 +745,7 @@ TMap<int32, int32> FTriangleMesh::FindCoincidentVertexRemappings(
 
 	// Move the points to the origin to avoid floating point aliasing far away
 	// from the origin.
-	FAABB3 Bbox(Points[0], Points[0]);
+	FAABB3 Bbox(Points[TestIndices[0]], Points[TestIndices[0]]);
 	for (int i = 1; i < NumPoints; i++)
 	{
 		Bbox.GrowToInclude(Points[TestIndices[i]]);
@@ -996,7 +996,7 @@ TArray<int32> FTriangleMesh::GetVertexImportanceOrdering(
 
 	// Move the points to the origin to avoid floating point aliasing far away
 	// from the origin.
-	FAABB3 Bbox(Points[0], Points[0]);
+	FAABB3 Bbox(Points[Offset], Points[Offset]);
 	for (int i = 1; i < NumPoints; i++)
 	{
 		Bbox.GrowToInclude(Points[i + Offset]);
@@ -1013,9 +1013,21 @@ TArray<int32> FTriangleMesh::GetVertexImportanceOrdering(
 		LocalBBox.GrowToInclude(LocalPoints[i]);
 	}
 	LocalBBox.Thicken(1.0e-3f);
+	// Center of the local bounding box (should always be close to zero)
 	const FVec3 LocalCenter = LocalBBox.Center();
 	const FVec3& LocalMin = LocalBBox.Min();
 
+	auto ToFlatIdx = [&LocalPoints, Offset](int32 PointIdx, FVec3 Center, FReal CellSizeIn, int64 Res) -> int64
+	{
+		const FVec3& Pos = LocalPoints[PointIdx - Offset];
+		// grid center co-located at bbox center:
+		const TVec3<int64> Coord(
+			static_cast<int64>(FMath::Floor((Pos[0] - Center[0]) / CellSizeIn)) + Res / 2,
+			static_cast<int64>(FMath::Floor((Pos[1] - Center[1]) / CellSizeIn)) + Res / 2,
+			static_cast<int64>(FMath::Floor((Pos[2] - Center[2]) / CellSizeIn)) + Res / 2);
+		return ((Coord[0] * Res + Coord[1]) * Res) + Coord[2];
+	};
+
 	// Bias towards points further away from the center of the bounding box.
 	// Send points that are the furthest away to the front of the list.
 	TArray<FReal> Dist;
@@ -1024,8 +1036,6 @@ TArray<int32> FTriangleMesh::GetVertexImportanceOrdering(
 	{
 		Dist[i] = (LocalPoints[i] - LocalCenter).SizeSquared();
 	}
-	DescendingPredicate<FReal> DescendingDistPred(Dist); // high to low
-	Sort(&PointOrder[0], NumPoints, DescendingDistPred);
 
 	// If all points are coincident, return early.
 	const FReal MaxBBoxDim = LocalBBox.Extents().Max();
@@ -1062,22 +1072,18 @@ TArray<int32> FTriangleMesh::GetVertexImportanceOrdering(
 			// Shift the grid by 1/2 a grid cell the second iteration so that
 			// we don't miss slightly adjacent coincident points across cell
 			// boundaries.
+			// (Note this could still miss some across-cell coincident points, but that's ok because 
+			//  the coincident vert array is just used to reduce the number of points used for collisions)
 			const FVec3 GridCenter = LocalCenter - FVec3(static_cast<FReal>(i) * CellSize / 2);
 			const int NumCoincidentPrev = NumCoincident;
 			for (int j = 0; j < NumPoints - NumCoincidentPrev; j++)
 			{
 				const int32 Idx = PointOrder[j];
-				const FVec3& Pos = LocalPoints[Idx - Offset];
-				const TVec3<int64> Coord(
-				    static_cast<int64>(FMath::Floor((Pos[0] - GridCenter[0]) / CellSize + static_cast<double>(Resolution) / 2)),
-				    static_cast<int64>(FMath::Floor((Pos[1] - GridCenter[1]) / CellSize + static_cast<double>(Resolution) / 2)),
-				    static_cast<int64>(FMath::Floor((Pos[2] - GridCenter[2]) / CellSize + static_cast<double>(Resolution) / 2)));
-				const int64 FlatIdx =
-				    ((Coord[0] * Resolution + Coord[1]) * Resolution) + Coord[2];
-
-				bool AlreadyInSet = false;
-				OccupiedCells.Add(FlatIdx, &AlreadyInSet);
-				if (AlreadyInSet)
+				const int64 FlatIdx = ToFlatIdx(Idx, GridCenter, CellSize, Resolution + i*2);
+
+				bool bAlreadyInSet = false;
+				OccupiedCells.Add(FlatIdx, &bAlreadyInSet);
+				if (bAlreadyInSet)
 				{
 					Rank[Idx - Offset] = 1;
 					if (CoincidentVertices)
@@ -1089,65 +1095,198 @@ TArray<int32> FTriangleMesh::GetVertexImportanceOrdering(
 			}
 			if (NumCoincident > NumCoincidentPrev)
 			{
-				StableSort(&PointOrder[0], NumPoints - NumCoincidentPrev, AscendingRankPred);
+				Sort(&PointOrder[0], NumPoints - NumCoincidentPrev, AscendingRankPred);
 			}
 		}
 	}
 	check(NumCoincident < NumPoints);
+	
+	// track the best points in a region, by multiple metrics (distance from center, and curvature at point)
+	// note the best-curvature and best-distance points are allowed to be the same point
+	struct FBestPtData
+	{
+		int32 FarthestIdx = -1; // Index of point farthest from center
+		FReal FarthestDistSq = 0; // Squared distance of farthest point from center
+		int32 MaxCurveIdx = -1; // Index of point with largest curvature
+		FReal MaxCurve = (FReal)-MAX_FLT; // Largest curvature value
+	};
+	// Update FBestPointData with a new point; return true if it's the first point in the cell, false otherwise
+	auto UpdateBestPtData = [&Dist, &PointCurvatures, &Offset, &PointOrder](FBestPtData& BestData, int32 OrderIdx) -> bool
+	{
+		int32 PtIdx = PointOrder[OrderIdx];
+		int32 OffsetPtIdx = PtIdx - Offset;
+		if (BestData.FarthestIdx == -1)
+		{
+			BestData.FarthestIdx = OrderIdx;
+			BestData.FarthestDistSq = Dist[OffsetPtIdx];
+			BestData.MaxCurveIdx = OrderIdx;
+			BestData.MaxCurve = PointCurvatures[OffsetPtIdx];
+			return true;
+		}
 
-	// Use spatial hashing to a grid of variable resolution to distribute points
-	// evenly across the volume.
-	for (int i = 2; i <= 1024; i += 2) // coarse to fine
+		if (Dist[OffsetPtIdx] > BestData.FarthestDistSq)
+		{
+			BestData.FarthestDistSq = Dist[OffsetPtIdx];
+			BestData.FarthestIdx = OrderIdx;
+		}
+
+		FReal Curvature = PointCurvatures[OffsetPtIdx];
+		if (Curvature > BestData.MaxCurve)
+		{
+			BestData.MaxCurve = Curvature;
+			BestData.MaxCurveIdx = OrderIdx;
+		}
+		return false;
+	};
+
+	// Points before this offset have already been moved forward, and don't need further consideration
+	int32 MovedPtsOffset = 0;
+	// traverse power of 2 cell resolutions (~octree levels), moving forward "best" points from the cells at each level
+	for (int32 Resolution = 2; MovedPtsOffset < NumPoints - NumCoincident && Resolution <= 1024; Resolution *= 2)
 	{
-		OccupiedCells.Reset();
+		TMap<int64, FBestPtData> BestPoints;
+		TSet<int64> CoveredCells, OffsetCoveredCells;
+		const FReal CellSize = MaxBBoxDim / static_cast<FReal>(Resolution);
+
+		// make smaller cells w/ a half-cell-width offset to help detect points that are close, but across cell boundaries
+		int32 OffsetResolution = Resolution * 4;
+		const FReal OffsetCellSize = MaxBBoxDim / static_cast<FReal>(OffsetResolution);
+		OffsetResolution += 2; // allow for offset center
+		FVec3 OffsetCenter = LocalCenter - FVec3(OffsetCellSize / 2);
+		int32 FoundPointsCount = 0;
+
+		// Use the Rank array to mark which points should move to the front of the remaining list in this iteration
 		Rank.Reset();
 		Rank.AddZeroed(NumPoints);
 
-		const int32 Resolution = i;
-		check(Resolution > 0);
-		check(Resolution % 2 == 0);
-		const FReal CellSize = MaxBBoxDim / static_cast<FReal>(Resolution);
+		auto ToMainFlatIdx = [&ToFlatIdx, &LocalCenter, &CellSize, &Resolution](int32 PointIdx)
+		{
+			return ToFlatIdx(PointIdx, LocalCenter, CellSize, Resolution);
+		};
+		auto ToOffsetFlatIdx = [&ToFlatIdx, &OffsetCenter, &OffsetCellSize, &OffsetResolution](int32 PointIdx)
+		{
+			return ToFlatIdx(PointIdx, OffsetCenter, OffsetCellSize, OffsetResolution);
+		};
 
-		// The order in which we process these points matters.  Must do
-		// the current highest rank first.
-		for (int j = 0; j < NumPoints - NumCoincident; j++)
+		// Mark cells that are already covered by points in the already-considered front of the list
+		for (int32 OrderIdx = 0; OrderIdx < MovedPtsOffset; OrderIdx++)
 		{
-			const int32 Idx = PointOrder[j];
-			const FVec3& Pos = LocalPoints[Idx - Offset];
-			// grid center co-located at bbox center:
-			const TVec3<int64> Coord(
-			    static_cast<int64>(FMath::Floor((Pos[0] - LocalCenter[0]) / CellSize)) + Resolution / 2,
-			    static_cast<int64>(FMath::Floor((Pos[1] - LocalCenter[1]) / CellSize)) + Resolution / 2,
-			    static_cast<int64>(FMath::Floor((Pos[2] - LocalCenter[2]) / CellSize)) + Resolution / 2);
-			const int64 FlatIdx =
-			    ((Coord[0] * Resolution + Coord[1]) * Resolution) + Coord[2];
+			int32 PointIdx = PointOrder[OrderIdx];
+			const int64 FlatIdx = ToMainFlatIdx(PointIdx);
+			CoveredCells.Add(FlatIdx);
+			OffsetCoveredCells.Add(ToOffsetFlatIdx(PointIdx));
+		}
 
-			bool AlreadyInSet = false;
-			OccupiedCells.Add(FlatIdx, &AlreadyInSet);
-			Rank[Idx - Offset] = AlreadyInSet ? 1 : 0;
+		// find the best points in cells that aren't already covered
+		bool bAllSeparate = true;
+		for (int32 OrderIdx = MovedPtsOffset; OrderIdx < NumPoints - NumCoincident; OrderIdx++)
+		{
+			int32 PointIdx = PointOrder[OrderIdx];
+			const int64 FlatIdx = ToMainFlatIdx(PointIdx);
+			if (CoveredCells.Contains(FlatIdx) || OffsetCoveredCells.Contains(ToOffsetFlatIdx(PointIdx)))
+			{
+				bAllSeparate = false;
+				continue;
+			}
+			FBestPtData& Best = BestPoints.FindOrAdd(FlatIdx);
+			bAllSeparate = UpdateBestPtData(Best, OrderIdx) & bAllSeparate;
 		}
 
-		// If every particle mapped to its own cell, we're done.
-		if (OccupiedCells.Num() == NumPoints)
+		if (bAllSeparate)
 		{
+			// all the points were in separate cells, no need to continue promoting the 'best' for each cell beyond here
 			break;
 		}
-		// If every particle mapped to 1 cell, don't bother sorting.
-		if (OccupiedCells.Num() == 1)
+
+		auto ConsiderMovingPt = [&Rank, Offset, MovedPtsOffset, &OffsetCoveredCells, &ToOffsetFlatIdx](int32 PointIdx) -> bool
 		{
-			continue;
+			uint8& R = Rank[PointIdx - Offset];
+			if (!R)
+			{
+				// try to avoid clumping by skipping points whose offset cells were already covered
+				int64 OffsetFlatIdx = ToOffsetFlatIdx(PointIdx);
+				bool bAlreadyInSet = false;
+				OffsetCoveredCells.Add(OffsetFlatIdx, &bAlreadyInSet);
+				if (bAlreadyInSet)
+				{
+					return false;
+				}
+
+				R = 1;
+				return true;
+			}
+			else
+			{
+				return false;
+			}
+		};
+
+		// Decide which points to move into the 'front' section
+		int32 NumToMoveToFront = 0;
+		TArray<int32> MoveToFront;
+		// Add points favored by the (more reliable) distance-from-center metric first
+		for (const TPair<int64, FBestPtData>& BestDataPair : BestPoints)
+		{
+			int PtIdx = PointOrder[BestDataPair.Value.FarthestIdx];
+			if (ConsiderMovingPt(PtIdx))
+			{
+				NumToMoveToFront++;
+				MoveToFront.Add(BestDataPair.Value.FarthestIdx);
+			}
+		}
+		// Add points favored by the curvature metric second
+		//  (so they may be skipped if distance-prioritized ones already covered the region)
+		for (const TPair<int64, FBestPtData>& BestDataPair : BestPoints)
+		{
+			int PtIdx = PointOrder[BestDataPair.Value.MaxCurveIdx];
+			if (ConsiderMovingPt(PtIdx))
+			{
+				NumToMoveToFront++;
+				MoveToFront.Add(BestDataPair.Value.MaxCurveIdx);
+			}
 		}
 
-		// Stable sort by rank.  When the resolution is high, stable sort will
-		// do nothing as we'll have nothing but rank 0's.  But then as the grid
-		// gets coarser, stable sort will get more and more selective about
-		// which particles get promoted.
-		//
-		// Since the initial ordering was biased by distance from
-		// the center, each rank is similarly ordered. That is, the first vertex
-		// to land in a cell will be the most distant.
-		StableSort(&PointOrder[0], NumPoints - NumCoincident, AscendingRankPred);
-	} // end for
+		// Do the move-to-front operations by swapping
+		int MoveFrontIdx = 0;
+		for (int MoveBackIdx = 0; MoveBackIdx < NumToMoveToFront; MoveBackIdx++)
+		{
+			if (Rank[PointOrder[MovedPtsOffset + MoveBackIdx] - Offset])
+			{
+				// Point was marked for moving to front, and is already in front; skip it
+				continue;
+			}
+			// Swap point with one that was marked for the front section
+			bool bDidSwap = false;
+			for (; MoveFrontIdx < MoveToFront.Num(); MoveFrontIdx++)
+			{
+				if (MoveToFront[MoveFrontIdx] < MovedPtsOffset + NumToMoveToFront)
+				{
+					// Don't swap this one back because it's already in the 'front' zone
+					continue;
+				}
+				int32 PtIdx = PointOrder[MoveToFront[MoveFrontIdx]];
+				checkSlow(Rank[PtIdx - Offset]);
+				Swap(PointOrder[MovedPtsOffset + MoveBackIdx], PointOrder[MoveToFront[MoveFrontIdx]]);
+				bDidSwap = true;
+				MoveFrontIdx++;
+				break;
+			}
+			// The above for loop should always find a point to swap back
+			checkSlow(bDidSwap);
+		}
+
+		// Note the above swapping method could have been done via a Sort() using the Rank array, but swapping is more direct / faster
+
+		// Sort the just-added points by curvature (they were in arbitrary order before)
+		DescendingPredicate<FReal> DescendingCurvaturePred(PointCurvatures);
+		Sort(&PointOrder[MovedPtsOffset], NumToMoveToFront, DescendingCurvaturePred);
+
+		MovedPtsOffset += NumToMoveToFront;
+	}
+
+	// Sort remaining non-coincident points by curvature
+	DescendingPredicate<FReal> DescendingCurvaturePred(PointCurvatures);
+	Sort(&PointOrder[MovedPtsOffset], NumPoints - MovedPtsOffset - NumCoincident, DescendingCurvaturePred);
 
 	return PointOrder;
 }
