Make sure active and transient list are mutually exclusive. Prevents duplicate updates on dirty particles

diff --git a/Engine/Source/Runtime/Experimental/Chaos/Private/RewindData.cpp b/Engine/Source/Runtime/Experimental/Chaos/Private/RewindData.cpp
--- a/Engine/Source/Runtime/Experimental/Chaos/Private/RewindData.cpp
+++ b/Engine/Source/Runtime/Experimental/Chaos/Private/RewindData.cpp
@@ -943,8 +943,7 @@ void FRewindData::PushPTDirtyData(TPBDRigidParticleHandle<FReal,3>& Handle,const
 	FDirtyParticleInfo& Info = FindOrAddParticle(Handle);
 	Info.LastDirtyFrame = CurFrame;
 	FParticleHistoryEntry& Latest = Info.AddFrame(CurFrame, Buffer, PropertiesPool);
-	//TODO: bring this ensure back one duplicate dirty particles are fixed
-	//ensure(Latest.GetStateChecked(FParticleHistoryEntry::PostCallbacks, CurFrame, Buffer).IsClean());	//PostCallbacks should be clean before we write sim results
+	ensure(Latest.GetStateChecked(FParticleHistoryEntry::PostCallbacks, CurFrame, Buffer).IsClean());	//PostCallbacks should be clean before we write sim results
 	Latest.GetStateChecked(FParticleHistoryEntry::PostCallbacks, CurFrame, Buffer).RecordSimResults(Handle, PropertiesPool);
 	CoalesceBack(Info.Frames, FParticleHistoryEntry::PostCallbacks);
 #if 0
diff --git a/Engine/Source/Runtime/Experimental/Chaos/Public/Chaos/PBDRigidsSOAs.h b/Engine/Source/Runtime/Experimental/Chaos/Public/Chaos/PBDRigidsSOAs.h
--- a/Engine/Source/Runtime/Experimental/Chaos/Public/Chaos/PBDRigidsSOAs.h
+++ b/Engine/Source/Runtime/Experimental/Chaos/Public/Chaos/PBDRigidsSOAs.h
@@ -235,7 +235,7 @@ public:
 
 		if (!Params.bStartSleeping)
 		{
-			ActiveParticlesMapArray.Insert(Results);
+			InsertToMapAndArray(Results, ActiveParticlesToIndex, ActiveParticlesArray);
 		}
 		UpdateViews();
 		return Results;
@@ -273,7 +273,7 @@ public:
 
 		if (!Params.bStartSleeping)
 		{
-			ActiveParticlesMapArray.Insert(reinterpret_cast<TArray<FPBDRigidParticleHandle*>&>(NewClustered));
+			AddToActiveArray(reinterpret_cast<TArray<FPBDRigidParticleHandle*>&>(NewClustered));
 		}
 
 		UpdateViews();
@@ -308,8 +308,7 @@ public:
 		auto PBDRigid = Particle->CastToRigidParticle();
 		if(PBDRigid)
 		{
-			ActiveParticlesMapArray.Remove(PBDRigid);
-			TransientDirtyMapArray.Remove(PBDRigid);
+			RemoveFromActiveArray(PBDRigid, /*bStillDirty=*/ false);
 
 			if (auto PBDRigidClustered = Particle->CastToClustered())
 			{
@@ -384,7 +383,7 @@ public:
 
 			// All active particles RIGID particles
 			{
-				ActiveParticlesMapArray.Remove(PBDRigid);
+				RemoveFromActiveArray(PBDRigid, /*bStillDirty=*/false);
 			}
 		}
 		else if (Particle->CastToKinematicParticle())
@@ -428,7 +427,7 @@ public:
 
 			if (!PBDRigid->Sleeping() && Particle->ObjectState() == EObjectStateType::Dynamic)
 			{
-				ActiveParticlesMapArray.Insert(PBDRigid);
+				AddToActiveArray(PBDRigid);
 			}
 		}
 		else if (Particle->CastToKinematicParticle())
@@ -471,7 +470,7 @@ public:
 					else
 					{
 						// Non clustered rigid particles:
-						ActiveParticlesMapArray.Insert(PBDRigid);
+						AddToActiveArray(PBDRigid);
 					}
 
 					UpdateViews();
@@ -522,8 +521,6 @@ public:
 					PBDRigid->SetSleeping(true);
 					PBDRigid->SetObjectStateLowLevel(EObjectStateType::Sleeping);
 
-					TransientDirtyMapArray.Insert(PBDRigid);
-
 					if (auto PBDRigidClustered = Particle->CastToClustered())
 					{
 						if (Particle->GetParticleType() == Chaos::EParticleType::GeometryCollection)
@@ -533,7 +530,7 @@ public:
 					}
 					else
 					{
-						ActiveParticlesMapArray.Remove(PBDRigid);
+						RemoveFromActiveArray(PBDRigid, /*bStillDirty=*/true);
 					}
 
 					if (!DeferUpdateViews)
@@ -571,11 +568,11 @@ public:
 		{
 			if (Particle->ObjectState() != EObjectStateType::Dynamic)
 			{
-				ActiveParticlesMapArray.Remove(Particle->CastToRigidParticle());
+				RemoveFromActiveArray(Particle->CastToRigidParticle(), /*bStillDirty=*/true);
 			}
 			else
 			{
-				ActiveParticlesMapArray.Insert(Particle->CastToRigidParticle());
+				AddToActiveArray(Particle->CastToRigidParticle());
 			}
 
 			// Move to appropriate dynamic SOA
@@ -864,6 +861,42 @@ private:
 		return ReturnHandles;
 	}
 	
+	void AddToActiveArray(const TArray<FPBDRigidParticleHandle*>& Particles)
+	{
+		ActiveParticlesMapArray.Insert(Particles);
+		
+		//dirty contains Active so make sure no duplicates
+		for(FPBDRigidParticleHandle* Particle : Particles)
+		{
+			TransientDirtyMapArray.Remove(Particle);
+		}
+	}
+
+	void AddToActiveArray(FPBDRigidParticleHandle* Particle)
+	{
+		ActiveParticlesMapArray.Insert(Particle);
+
+		//dirty contains Active so make sure no duplicates
+		TransientDirtyMapArray.Remove(Particle);
+	}
+
+	void RemoveFromActiveArray(FPBDRigidParticleHandle* Particle, bool bStillDirty)
+	{
+		ActiveParticlesMapArray.Remove(Particle);
+
+		if(bStillDirty)
+		{
+			//no longer active, but still dirty
+			ActiveParticlesMapArray.Insert(Particle);
+		}
+		else
+		{
+			//might have already been removed from active from a previous call
+			//but now removing and don't want it dirty either
+			TransientDirtyMapArray.Remove(Particle);
+		}
+	}
+	
 	//should be called whenever particles are added / removed / reordered
 	void UpdateViews()
 	{
