Chaos : Turn off generic convex ccd by default.

diff --git a/Engine/Source/Runtime/Experimental/Chaos/Private/Chaos/CollisionResolution.cpp b/Engine/Source/Runtime/Experimental/Chaos/Private/Chaos/CollisionResolution.cpp
--- a/Engine/Source/Runtime/Experimental/Chaos/Private/Chaos/CollisionResolution.cpp
+++ b/Engine/Source/Runtime/Experimental/Chaos/Private/Chaos/CollisionResolution.cpp
@@ -82,7 +82,7 @@ FAutoConsoleVariableRef  CVarCCDEnableThresholdBoundsScale(TEXT("p.Chaos.CCD.Ena
 Chaos::FRealSingle CCDAllowedDepthBoundsScale = 0.1f;
 FAutoConsoleVariableRef CVarCCDAllowedDepthBoundsScale(TEXT("p.Chaos.CCD.AllowedDepthBoundsScale"), CCDAllowedDepthBoundsScale, TEXT("When rolling back to TOI, allow (smallest bound's extent) * AllowedDepthBoundsScale, instead of rolling back to exact TOI w/ penetration = 0."));
 
-int32 CCDUseGenericSweptConvexConstraints = 1;
+int32 CCDUseGenericSweptConvexConstraints = 0;
 FAutoConsoleVariableRef CVarUseGenericSweptConvexConstraints(TEXT("p.Chaos.CCD.UseGenericSweptConvexConstraints"), CCDUseGenericSweptConvexConstraints, TEXT("Use generic convex convex swept constraint generation for convex shape pairs which don't have specialized implementations."));
 
 int32 CCDOnlyConsiderDynamicStatic = 0;
