Fix failing build with recent Chaos changes

diff --git a/Engine/Source/Runtime/Experimental/Chaos/Public/Chaos/PBDRigidsSOAs.h b/Engine/Source/Runtime/Experimental/Chaos/Public/Chaos/PBDRigidsSOAs.h
--- a/Engine/Source/Runtime/Experimental/Chaos/Public/Chaos/PBDRigidsSOAs.h
+++ b/Engine/Source/Runtime/Experimental/Chaos/Public/Chaos/PBDRigidsSOAs.h
@@ -235,7 +235,7 @@ public:
 
 		if (!Params.bStartSleeping)
 		{
-			InsertToMapAndArray(Results, ActiveParticlesToIndex, ActiveParticlesArray);
+			AddToActiveArray(Results);
 		}
 		UpdateViews();
 		return Results;
