Chaos
- add half-edge index-based structure data to convex shapes. The index size is 1, 2, or 4 bytes depending on the structure.
- move old structure data into a legacy namespace and add a custom loader/conversion step for serializing old data

diff --git a/Engine/Source/Programs/HeadlessChaos/Private/HeadlessChaosTestConvex.cpp b/Engine/Source/Programs/HeadlessChaos/Private/HeadlessChaosTestConvex.cpp
--- a/Engine/Source/Programs/HeadlessChaos/Private/HeadlessChaosTestConvex.cpp
+++ b/Engine/Source/Programs/HeadlessChaos/Private/HeadlessChaosTestConvex.cpp
@@ -109,10 +109,14 @@ namespace ChaosTest
 		// Check all per-vertex data
 		for (int32 VertexIndex = 0; VertexIndex < Convex.NumVertices(); ++VertexIndex)
 		{
-			// All planes should pass through the vertex
-			for (int32 VertexPlaneIndex = 0; VertexPlaneIndex < Convex.NumVertexPlanes(VertexIndex); ++VertexPlaneIndex)
+			// Get all the planes for the vertex
+			TArray<int32> PlaneIndices;
+			PlaneIndices.SetNum(128);
+			int32 NumPlanes = Convex.FindVertexPlanes(VertexIndex, PlaneIndices.GetData(), PlaneIndices.Num());
+			PlaneIndices.SetNum(NumPlanes);
+
+			for (int32 PlaneIndex : PlaneIndices)
 			{
-				const int32 PlaneIndex = Convex.GetVertexPlane(VertexIndex, VertexPlaneIndex);
 				const TPlaneConcrete<FReal, 3> Plane = Convex.GetPlane(PlaneIndex);
 				const FVec3 Vertex = Convex.GetVertex(VertexIndex);
 				const FReal VertexDistance = FVec3::DotProduct(Plane.Normal(), Vertex - Plane.X());
@@ -158,6 +162,7 @@ namespace ChaosTest
 	GTEST_TEST(ConvexStructureTests, TestConvexStructureData2)
 	{
 		const FVec3 Vertices[] =
+
 		{
 			FVec3(0, 0, 12.0f),
 			FVec3(-0.707f, -0.707f, 10.0f),
@@ -203,7 +208,7 @@ namespace ChaosTest
 	void TestConvexStructureDataMapping(const T_STRUCTUREDATA& StructureData)
 	{
 		// For each plane, get the list of vertices that make its edges.
-		// Then check that we list of planes used by that vertex contains the original plane
+		// Then check that the list of planes used by that vertex contains the original plane
 		for (int32 PlaneIndex = 0; PlaneIndex < StructureData.NumPlanes(); ++PlaneIndex)
 		{
 			for (int32 PlaneVertexIndex = 0; PlaneVertexIndex < StructureData.NumPlaneVertices(PlaneIndex); ++PlaneVertexIndex)
@@ -211,16 +216,12 @@ namespace ChaosTest
 				const int32 VertexIndex = StructureData.GetPlaneVertex(PlaneIndex, PlaneVertexIndex);
 
 				// Check that the plane's vertex has the plane in its list
-				bool bFoundPlane = false;
-				for (int32 VertexPlaneIndex = 0; VertexPlaneIndex < StructureData.NumVertexPlanes(VertexIndex); ++VertexPlaneIndex)
-				{
-					const int32 SearchPlaneIndex = StructureData.GetVertexPlane(VertexIndex, VertexPlaneIndex);
-					if (PlaneIndex == SearchPlaneIndex)
-					{
-						bFoundPlane = true;
-						break;
-					}
-				}
+				TArray<int32> PlaneIndices;
+				PlaneIndices.SetNum(128);
+				const int32 NumPlanes = StructureData.FindVertexPlanes(VertexIndex, PlaneIndices.GetData(), PlaneIndices.Num());
+				PlaneIndices.SetNum(NumPlanes);
+
+				const bool bFoundPlane = PlaneIndices.Contains(PlaneIndex);
 				EXPECT_TRUE(bFoundPlane);
 			}
 		}
@@ -248,7 +249,7 @@ namespace ChaosTest
 
 		FConvex Convex(Particles, 0.0f);
 
-		const FConvexStructureDataU8& StructureData = Convex.GetStructureData().Data8();
+		const FConvexStructureData::FConvexStructureDataMedium& StructureData = Convex.GetStructureData().DataM();
 		TestConvexStructureDataMapping(StructureData);
 		TestConvexStructureDataImpl(Convex);
 	}
@@ -256,11 +257,12 @@ namespace ChaosTest
 
 	// Check that the structure data is good for convex shapes that have faces merged during construction
 	// This test uses the large index size in StructureData.
-	GTEST_TEST(ConvexStructureTests, TestLargeIndexStructureData2)
+	// This test is disabled - the convex building is too slow for this many verts
+	GTEST_TEST(ConvexStructureTests, DISABLED_TestLargeIndexStructureData2)
 	{
 		FMath::RandInit(53799058);
-		const FReal Radius = 1000.0f;
-		const int32 NumVertices = 1000;
+		const FReal Radius = 10000.0f;
+		const int32 NumVertices = 50000;
 
 		// Make a convex with points on a sphere.
 		TArray<FVec3> Vertices;
@@ -277,7 +279,7 @@ namespace ChaosTest
 		EXPECT_GT(Convex.NumVertices(), 800);
 		EXPECT_GT(Convex.NumPlanes(), 500);
 
-		const FConvexStructureDataS32& StructureData = Convex.GetStructureData().Data32();
+		const FConvexStructureData::FConvexStructureDataLarge& StructureData = Convex.GetStructureData().DataL();
 		TestConvexStructureDataMapping(StructureData);
 		TestConvexStructureDataImpl(Convex);
 	}
@@ -372,4 +374,88 @@ namespace ChaosTest
 		// Check that we've failed to build a 3D convex hull and safely returned
 		EXPECT_EQ(Planes.Num(), 0);
 	}
+
+	GTEST_TEST(ConvexStructureTests, TestConvexHalfEdgeStructureData_Box)
+	{
+		const TArray<FVec3> InputVertices =
+		{
+			FVec3(-50,		-50,	-50),
+			FVec3(-50,		-50,	50),
+			FVec3(-50,		50,		-50),
+			FVec3(-50,		50,		50),
+			FVec3(50,		-50,	-50),
+			FVec3(50,		-50,	50),
+			FVec3(50,		50,		-50),
+			FVec3(50,		50,		50),
+		};
+
+		TArray<TPlaneConcrete<FReal, 3>> Planes;
+		TArray<TArray<int32>> FaceVertices;
+		TArray<Chaos::FVec3> Vertices;
+		TAABB<FReal, 3> LocalBounds;
+		FConvexBuilder::Build(InputVertices, Planes, FaceVertices, Vertices, LocalBounds);
+		FConvexBuilder::MergeFaces(Planes, FaceVertices, Vertices, 1.0f);
+
+		FConvex Convex(Vertices, 0.0f);
+
+		const FConvexStructureData::FConvexStructureDataSmall& StructureData = Convex.GetStructureData().DataS();
+
+		EXPECT_EQ(StructureData.NumPlanes(), 6);
+		EXPECT_EQ(StructureData.NumHalfEdges(), 24);
+		EXPECT_EQ(StructureData.NumVertices(), 8);
+
+		// Count how many times each vertex and edge is referenced
+		TArray<int32> VertexIndexCount;
+		TArray<int32> EdgeIndexCount;
+		VertexIndexCount.SetNumZeroed(StructureData.NumVertices());
+		EdgeIndexCount.SetNumZeroed(StructureData.NumHalfEdges());
+		for (int32 PlaneIndex = 0; PlaneIndex < StructureData.NumPlanes(); ++PlaneIndex)
+		{
+			EXPECT_EQ(StructureData.NumPlaneHalfEdges(PlaneIndex), 4);
+			for (int32 PlaneEdgeIndex = 0; PlaneEdgeIndex < StructureData.NumPlaneHalfEdges(PlaneIndex); ++PlaneEdgeIndex)
+			{
+				const int32 EdgeIndex = StructureData.GetPlaneHalfEdge(PlaneIndex, PlaneEdgeIndex);
+				const int32 VertexIndex = StructureData.GetHalfEdgeVertex(EdgeIndex);
+				EdgeIndexCount[EdgeIndex]++;
+				VertexIndexCount[VertexIndex]++;
+			}
+		}
+
+		// Every vertex is used by 3 half-edges (and planes)
+		for (int32 VertexCount : VertexIndexCount)
+		{
+			EXPECT_EQ(VertexCount, 3);
+		}
+
+		// Each half edge is used by a single plane
+		for (int32 EdgeCount : EdgeIndexCount)
+		{
+			EXPECT_EQ(EdgeCount, 1);
+		}
+
+		// Vertex Plane iterator generates 3 planes and all the edges have the same primary vertex
+		for (int32 VertexIndex = 0; VertexIndex < StructureData.NumVertices(); ++VertexIndex)
+		{
+			int32 PlaneCount = 0;
+
+			TArray<int32> VertexPlanes;
+			VertexPlanes.SetNum(128);
+			const int32 NumPlanes = StructureData.FindVertexPlanes(VertexIndex, VertexPlanes.GetData(), VertexPlanes.Num());
+			VertexPlanes.SetNum(NumPlanes);
+
+			for (int32 PlaneIndex : VertexPlanes)
+			{
+				EXPECT_NE(PlaneIndex, INDEX_NONE);
+
+				++PlaneCount;
+			}
+
+			// Everty vertex belongs to 3 planes
+			EXPECT_EQ(PlaneCount, 3);
+
+			// Every vertex's first edge should have that vertex as its root vertex
+			const int32 VertexHalfEdgeIndex = StructureData.GetVertexFirstHalfEdge(VertexIndex);
+			EXPECT_EQ(VertexIndex, StructureData.GetHalfEdgeVertex(VertexHalfEdgeIndex));
+		}
+	}
 }
diff --git a/Engine/Source/Runtime/Core/Public/UObject/PhysicsObjectVersion.h b/Engine/Source/Runtime/Core/Public/UObject/PhysicsObjectVersion.h
--- a/Engine/Source/Runtime/Core/Public/UObject/PhysicsObjectVersion.h
+++ b/Engine/Source/Runtime/Core/Public/UObject/PhysicsObjectVersion.h
@@ -45,6 +45,9 @@ struct CORE_API FPhysicsObjectVersion
 		// Fix corrupted LOD transition maps
 		ChaosClothFixLODTransitionMaps,
 
+		// Convex structure data is now an index-based half-edge structure
+		ConvexUsesHalfEdges,
+
 		// -----<new versions can be added above this line>-------------------------------------------------
 		VersionPlusOne,
 		LatestVersion = VersionPlusOne - 1
diff --git a/Engine/Source/Runtime/Experimental/Chaos/Private/Chaos/Box.cpp b/Engine/Source/Runtime/Experimental/Chaos/Private/Chaos/Box.cpp
--- a/Engine/Source/Runtime/Experimental/Chaos/Private/Chaos/Box.cpp
+++ b/Engine/Source/Runtime/Experimental/Chaos/Private/Chaos/Box.cpp
@@ -5,6 +5,26 @@
 
 namespace Chaos
 {
+	template<typename T, int D> TBox<T, D>::FStaticInitializer::FStaticInitializer()
+	{
+		TBox<FReal, 3>::InitializeStructureData();
+	}
+
+	template<typename T, int D> void TBox<T, D>::InitializeStructureData()
+	{
+		TArray<TArray<int32>> PlaneVertices
+		{
+			{6, 7, 3, 2},	// X
+			{1, 5, 6, 2},	// Y
+			{7, 6, 5, 4},	// Z,
+			{1, 0, 4, 5},	// -X
+			{0, 3, 7, 4},	// -Y
+			{0, 1, 2, 3},	// -Z,
+		};
+
+		SStructureData.SetPlaneVertices(PlaneVertices, 8);
+	}
+
 	template<typename T, int D> TArray<FVec3> TBox<T, D>::SNormals =
 	{
 		FVec3(1,0,0),	// X
@@ -27,27 +47,9 @@ namespace Chaos
 		FVec3(1,-1,1),
 	};
 
-	template<typename T, int D> TArray<TArray<int32>> TBox<T, D>::SPlaneVertices
-	{
-		{6, 7, 3, 2},	// X
-		{1, 5, 6, 2},	// Y
-		{7, 6, 5, 4},	// Z,
-		{1, 0, 4, 5},	// -X
-		{0, 3, 7, 4},	// -Y
-		{0, 1, 2, 3},	// -Z,
-	};
+	template<typename T, int D> FConvexHalfEdgeStructureDataS16 TBox<T, D>::SStructureData;
 
-	template<typename T, int D> TArray<TArray<int32>> TBox<T, D>::SVertexPlanes =
-	{
-		{3, 4, 5},
-		{1, 3, 5},
-		{0, 1, 5},
-		{0, 4, 5},
-		{2, 3, 4},
-		{1, 2, 3},
-		{0, 1, 2},
-		{0, 2, 4},
-	};
+	template<typename T, int D> typename TBox<T, D>::FStaticInitializer TBox<T, D>::SInitializer;
 
 	template class TBox<FReal, 3>;
 }
\ No newline at end of file
diff --git a/Engine/Source/Runtime/Experimental/Chaos/Private/Chaos/Convex.cpp b/Engine/Source/Runtime/Experimental/Chaos/Private/Chaos/Convex.cpp
--- a/Engine/Source/Runtime/Experimental/Chaos/Private/Chaos/Convex.cpp
+++ b/Engine/Source/Runtime/Experimental/Chaos/Private/Chaos/Convex.cpp
@@ -144,33 +144,6 @@ namespace Chaos
 		return MostOpposingIdx;
 	}
 
-
-
-	int32 FConvex::GetMostOpposingPlaneWithVertex(int32 VertexIndex, const FVec3& Normal) const
-	{
-		const int32 VertexPlaneNum = NumVertexPlanes(VertexIndex);
-		if ((VertexIndex == INDEX_NONE) || (VertexPlaneNum == 0))
-		{
-			return GetMostOpposingPlane(Normal);
-		}
-
-		int32 MostOpposingIdx = INDEX_NONE;
-		FReal MostOpposingDot = TNumericLimits<FReal>::Max();
-		for (int32 VertexPlaneIndex = 0; VertexPlaneIndex < VertexPlaneNum; ++VertexPlaneIndex)
-		{
-			const int32 PlaneIndex = GetVertexPlane(VertexIndex, VertexPlaneIndex);
-			const TPlaneConcrete<FReal, 3>& Plane = Planes[PlaneIndex];
-			const FReal Dot = FVec3::DotProduct(Plane.Normal(), Normal);
-			if (Dot < MostOpposingDot)
-			{
-				MostOpposingDot = Dot;
-				MostOpposingIdx = PlaneIndex;
-			}
-		}
-		CHAOS_ENSURE(MostOpposingIdx != INDEX_NONE);
-		return MostOpposingIdx;
-	}
-
 	FVec3 FConvex::GetClosestEdgePosition(int32 PlaneIndex, const FVec3& Position) const
 	{
 		FVec3 ClosestEdgePosition = FVec3(0);
@@ -236,25 +209,15 @@ namespace Chaos
 		return false;
 	}
 
-
-	int32 FConvex::NumVertexPlanes(int32 VertexIndex) const
+	int32 FConvex::FindVertexPlanes(int32 VertexIndex, int32* OutVertexPlanes, int32 MaxVertexPlanes) const
 	{
 		if (StructureData.IsValid())
 		{
-			return StructureData.NumVertexPlanes(VertexIndex);
+			return StructureData.FindVertexPlanes(VertexIndex, OutVertexPlanes, MaxVertexPlanes);
 		}
 		return 0;
 	}
 
-	int32 FConvex::GetVertexPlane(int32 VertexIndex, int32 VertexPlaneIndex) const
-	{
-		if (StructureData.IsValid())
-		{
-			return StructureData.GetVertexPlane(VertexIndex, VertexPlaneIndex);
-		}
-		return INDEX_NONE;
-	}
-
 	int32 FConvex::NumPlaneVertices(int32 PlaneIndex) const
 	{
 		if (StructureData.IsValid())
diff --git a/Engine/Source/Runtime/Experimental/Chaos/Public/Chaos/Box.h b/Engine/Source/Runtime/Experimental/Chaos/Public/Chaos/Box.h
--- a/Engine/Source/Runtime/Experimental/Chaos/Public/Chaos/Box.h
+++ b/Engine/Source/Runtime/Experimental/Chaos/Public/Chaos/Box.h
@@ -5,6 +5,7 @@
 
 #include "Chaos/ImplicitObject.h"
 #include "Chaos/AABB.h"
+#include "Chaos/ConvexHalfEdgeStructureData.h"
 #include "Chaos/Transform.h"
 #include "ChaosArchive.h"
 #include "UObject/ExternalPhysicsCustomObjectVersion.h"
@@ -207,12 +208,6 @@ namespace Chaos
 			return GetMostOpposingPlane(Normal);
 		}
 
-		// Get the index of the plane that most opposes the normal (VertexIndex is ignored)
-		int32 GetMostOpposingPlaneWithVertex(int32 VertexIndex, const TVector<T, d>& Normal) const
-		{
-			return GetMostOpposingPlane(Normal);
-		}
-
 		// Get the nearest point on an edge
 		TVector<T, d> GetClosestEdgePosition(int32 PlaneIndexHint, const TVector<T, d>& Position) const
 		{
@@ -292,36 +287,30 @@ namespace Chaos
 			return false;
 		}
 
-
-		// The number of planes that use the specified vertex
-		int32 NumVertexPlanes(int32 VertexIndex) const
+		// Get an array of all the plane indices that belong to a vertex (up to MaxVertexPlanes).
+		// Returns the number of planes found.
+		int32 FindVertexPlanes(int32 VertexIndex, int32* OutVertexPlanes, int32 MaxVertexPlanes) const
 		{
-			return 3;
-		}
-
-		// Get the plane index of one of the planes that uses the specified vertex
-		int32 GetVertexPlane(int32 VertexIndex, int32 VertexPlaneIndex) const
-		{
-			return SVertexPlanes[VertexIndex][VertexPlaneIndex];
+			return SStructureData.FindVertexPlanes(VertexIndex, OutVertexPlanes, MaxVertexPlanes);
 		}
 
 		// The number of vertices that make up the corners of the specified face
 		int32 NumPlaneVertices(int32 PlaneIndex) const
 		{
-			return 4;
+			return SStructureData.NumPlaneVertices(PlaneIndex);
 		}
 
 		// Get the vertex index of one of the vertices making up the corners of the specified face
 		int32 GetPlaneVertex(int32 PlaneIndex, int32 PlaneVertexIndex) const
 		{
-			return SPlaneVertices[PlaneIndex][PlaneVertexIndex];
+			return SStructureData.GetPlaneVertex(PlaneIndex, PlaneVertexIndex);
 		}
 
-		int32 NumPlanes() const { return 6; }
+		int32 NumPlanes() const { return SNormals.Num(); }
 
-		int32 NumVertices() const { return 8; }
+		int32 NumVertices() const { return SVertices.Num(); }
 
-		// Get the plane at the specified index (e.g., indices from GetVertexPlane)
+		// Get the plane at the specified index (e.g., indices from FindVertexPlanes)
 		const TPlaneConcrete<FReal, 3> GetPlane(int32 FaceIndex) const
 		{
 			const FVec3& PlaneN = SNormals[FaceIndex];
@@ -329,7 +318,7 @@ namespace Chaos
 			return TPlaneConcrete<FReal, 3>(PlaneX, PlaneN);
 		}
 
-		// Get the vertex at the specified index (e.g., indices from GetPlaneVertex)
+		// Get the vertex at the specified index (e.g., indices from GetPlaneVertexs)
 		const FVec3 GetVertex(int32 VertexIndex) const
 		{
 			const FVec3& Vertex = SVertices[VertexIndex];
@@ -480,14 +469,21 @@ namespace Chaos
 			return AABB.GetTypeHash();
 		}
 
+		static void InitializeStructureData();
+
 	private:
 		TAABB<T, d> AABB;
 
 		// Structure data shared by all boxes and used for manifold creation
-		static TArray<TArray<int32>> SVertexPlanes;
-		static TArray<TArray<int32>> SPlaneVertices;
 		static TArray<FVec3> SNormals;
 		static TArray<FVec3> SVertices;
+		static FConvexHalfEdgeStructureDataS16 SStructureData;
+
+		struct FStaticInitializer
+		{
+			FStaticInitializer();
+		};
+		static FStaticInitializer SInitializer;
 	};
 
 
diff --git a/Engine/Source/Runtime/Experimental/Chaos/Public/Chaos/Convex.h b/Engine/Source/Runtime/Experimental/Chaos/Public/Chaos/Convex.h
--- a/Engine/Source/Runtime/Experimental/Chaos/Public/Chaos/Convex.h
+++ b/Engine/Source/Runtime/Experimental/Chaos/Public/Chaos/Convex.h
@@ -355,19 +355,14 @@ namespace Chaos
 		// Get the index of the plane that most opposes the normal
 		int32 GetMostOpposingPlaneScaled(const FVec3& Normal, const FVec3& Scale) const;
 
-		// Get the index of the plane that most opposes the normal, assuming it passes through the specified vertex
-		int32 GetMostOpposingPlaneWithVertex(int32 VertexIndex, const FVec3& Normal) const;
-
 		// Get the nearest point on an edge of the specified face
 		FVec3 GetClosestEdgePosition(int32 PlaneIndex, const FVec3& Position) const;
 
 		bool GetClosestEdgeVertices(int32 PlaneIndex, const FVec3& Position, int32& OutVertexIndex0, int32& OutVertexIndex1) const;
 
-		// The number of planes that use the specified vertex
-		int32 NumVertexPlanes(int32 VertexIndex) const;
-
-		// Get the plane index of one of the planes that uses the specified vertex
-		int32 GetVertexPlane(int32 VertexIndex, int32 VertexPlaneIndex) const;
+		// Get an array of all the plane indices that belong to a vertex (up to MaxVertexPlanes).
+		// Returns the number of planes found.
+		int32 FindVertexPlanes(int32 VertexIndex, int32* OutVertexPlanes, int32 MaxVertexPlanes) const;
 
 		// The number of vertices that make up the corners of the specified face
 		int32 NumPlaneVertices(int32 PlaneIndex) const;
@@ -385,13 +380,13 @@ namespace Chaos
 			return (int32)Vertices.Num();
 		}
 
-		// Get the plane at the specified index (e.g., indices from GetVertexPlane)
+		// Get the plane at the specified index (e.g., indices from FindVertexPlanes)
 		const TPlaneConcrete<FReal, 3>& GetPlane(int32 FaceIndex) const
 		{
 			return Planes[FaceIndex];
 		}
 
-		// Get the vertex at the specified index (e.g., indices from GetPlaneVertex)
+		// Get the vertex at the specified index (e.g., indices from GetPlaneVertexs)
 		const FVec3& GetVertex(int32 VertexIndex) const
 		{
 			return Vertices[VertexIndex];
@@ -461,14 +456,17 @@ namespace Chaos
 			}
 
 			// Get any 3 planes that contribute to this vertex
-			if (NumVertexPlanes(VertexIndex) >= 3)
+			int32 PlaneIndices[3];
+			int32 NumVertexPlanes = FindVertexPlanes(VertexIndex, PlaneIndices, 3);
+
+			// Move the planes by the margin and recalculate the interection
+			// @todo(chaos): calculate dV/dm per vertex and store it in StructureData
+			if (NumVertexPlanes >= 3)
 			{
-				const int32 PlaneIndex0 = GetVertexPlane(VertexIndex, 0);
-				const int32 PlaneIndex1 = GetVertexPlane(VertexIndex, 1);
-				const int32 PlaneIndex2 = GetVertexPlane(VertexIndex, 2);
+				const int32 PlaneIndex0 = PlaneIndices[0];
+				const int32 PlaneIndex1 = PlaneIndices[1];
+				const int32 PlaneIndex2 = PlaneIndices[2];
 
-				// Move the planes by the margin and recalculate the interection
-				// @todo(chaos): calculate dV/dm per vertex and store it in StructureData
 				FVec3 PlanesPos;
 				UE::Math::TPlane<FReal> NewPlanes[3] =
 				{
@@ -485,10 +483,10 @@ namespace Chaos
 			// If we get here, the convex hull is malformed. Try to handle it anyway 
 			// @todo(chaos): track down the invalid hull issue
 
-			if (NumVertexPlanes(VertexIndex) == 2)
+			if (NumVertexPlanes == 2)
 			{
-				const int32 PlaneIndex0 = GetVertexPlane(VertexIndex, 0);
-				const int32 PlaneIndex1 = GetVertexPlane(VertexIndex, 1);
+				const int32 PlaneIndex0 = PlaneIndices[0];
+				const int32 PlaneIndex1 = PlaneIndices[1];
 				const FVec3 NewPlaneX = GetVertex(VertexIndex);
 				const FVec3 NewPlaneN0 = Planes[PlaneIndex0].Normal();
 				const FVec3 NewPlaneN1 = Planes[PlaneIndex1].Normal();
@@ -496,9 +494,9 @@ namespace Chaos
 				return NewPlaneX - (InMargin * NewPlaneN);
 			}
 
-			if (NumVertexPlanes(VertexIndex) == 1)
+			if (NumVertexPlanes == 1)
 			{
-				const int32 PlaneIndex0 = GetVertexPlane(VertexIndex, 0);
+				const int32 PlaneIndex0 = PlaneIndices[0];
 				const FVec3 NewPlaneX = GetVertex(VertexIndex);
 				const FVec3 NewPlaneN = Planes[PlaneIndex0].Normal();
 				return NewPlaneX - (InMargin * NewPlaneN);
@@ -516,14 +514,17 @@ namespace Chaos
 			}
 
 			// Get any 3 planes that contribute to this vertex
-			if (NumVertexPlanes(VertexIndex) >= 3)
+			int32 PlaneIndices[3];
+			int32 NumVertexPlanes = FindVertexPlanes(VertexIndex, PlaneIndices, 3);
+
+			// Move the planes by the margin and recalculate the interection
+			// @todo(chaos): calculate dV/dm per vertex and store it in StructureData (but see todo above)
+			if (NumVertexPlanes >= 3)
 			{
-				const int32 PlaneIndex0 = GetVertexPlane(VertexIndex, 0);
-				const int32 PlaneIndex1 = GetVertexPlane(VertexIndex, 1);
-				const int32 PlaneIndex2 = GetVertexPlane(VertexIndex, 2);
+				const int32 PlaneIndex0 = PlaneIndices[0];
+				const int32 PlaneIndex1 = PlaneIndices[1];
+				const int32 PlaneIndex2 = PlaneIndices[2];
 
-				// Move the planes by the margin and recalculate the interection
-				// @todo(chaos): calculate dV/dm per vertex and store it in StructureData (but see todo above)
 				const FVec3 NewPlaneX = Scale * GetVertex(VertexIndex);
 				const FVec3 NewPlaneNs[3] = 
 				{
@@ -554,10 +555,10 @@ namespace Chaos
 			// If we get here, the convex hull is malformed. Try to handle it anyway 
 			// @todo(chaos): track down the invalid hull issue
 
-			if (NumVertexPlanes(VertexIndex) == 2)
+			if (NumVertexPlanes == 2)
 			{
-				const int32 PlaneIndex0 = GetVertexPlane(VertexIndex, 0);
-				const int32 PlaneIndex1 = GetVertexPlane(VertexIndex, 1);
+				const int32 PlaneIndex0 = PlaneIndices[0];
+				const int32 PlaneIndex1 = PlaneIndices[1];
 				const FVec3 NewPlaneX = Scale * GetVertex(VertexIndex);
 				const FVec3 NewPlaneN0 = (Planes[PlaneIndex0].Normal() / Scale).GetUnsafeNormal();
 				const FVec3 NewPlaneN1 = (Planes[PlaneIndex1].Normal() / Scale).GetUnsafeNormal();
@@ -565,9 +566,9 @@ namespace Chaos
 				return NewPlaneX - (InMargin * NewPlaneN);
 			}
 
-			if (NumVertexPlanes(VertexIndex) == 1)
+			if (NumVertexPlanes == 1)
 			{
-				const int32 PlaneIndex0 = GetVertexPlane(VertexIndex, 0);
+				const int32 PlaneIndex0 = PlaneIndices[0];
 				const FVec3 NewPlaneX = Scale * GetVertex(VertexIndex);
 				const FVec3 NewPlaneN = (Planes[PlaneIndex0].Normal() / Scale).GetUnsafeNormal();
 				return NewPlaneX - (InMargin * NewPlaneN);
diff --git a/Engine/Source/Runtime/Experimental/Chaos/Public/Chaos/ConvexFlattenedArrayStructureData.h b/Engine/Source/Runtime/Experimental/Chaos/Public/Chaos/ConvexFlattenedArrayStructureData.h
new file mode 100644
--- /dev/null
+++ b/Engine/Source/Runtime/Experimental/Chaos/Public/Chaos/ConvexFlattenedArrayStructureData.h
@@ -0,0 +1,285 @@
+// Copyright Epic Games, Inc. All Rights Reserved.
+#pragma once
+
+#include "Chaos/Core.h"
+#include "Chaos/ConvexHalfEdgeStructureData.h"
+#include "ChaosArchive.h"
+#include "ChaosCheck.h"
+#include "ChaosLog.h"
+#include "UObject/FortniteMainBranchObjectVersion.h"
+#include "UObject/PhysicsObjectVersion.h"
+
+//
+//
+// LEGACY CODE
+// Kept around for serialization of old data. Most of the API has been removed, 
+// except what is needed to convert to the new format. 
+// See ConvexHalfEdgeStructureData.h for the replacement.
+//
+//
+
+namespace Chaos
+{
+	namespace Legacy
+	{
+		class FLegacyConvexStructureDataLoader;
+
+		// Base class for TConvexFlattenedArrayStructureData.
+		// NOTE: Deliberately no-virtual destructor so it should never be deleted from a 
+		// base class pointer.
+		class CHAOS_API FConvexFlattenedArrayStructureData
+		{
+		public:
+			FConvexFlattenedArrayStructureData() {}
+			~FConvexFlattenedArrayStructureData() {}
+		};
+
+		// A container for the convex structure data arrays.
+		// This is templated on the index type required, which needs to
+		// be large enough to hold the max of the number of vertices or 
+		// planes on the convex.
+		// 
+		// T_INDEX:			the type used to index into the convex vertices 
+		//					and planes array (in the outer convex). Must be 
+		//					able to contain max(NumPlanes, NumVerts).
+		//
+		// T_OFFSETINDEX:	the type used to index the flattened array of
+		//					indices. Must be able to contain 
+		//					Max(NumPlanes*AverageVertsPerPlane)
+		//
+		template<typename T_INDEX, typename T_OFFSETINDEX>
+		class CHAOS_API TConvexFlattenedArrayStructureData : public FConvexFlattenedArrayStructureData
+		{
+		public:
+			// Everything public - not used outside of legacy loader and unit tests
+
+			using FIndex = T_INDEX;
+			using FOffsetIndex = T_OFFSETINDEX;
+
+			FORCEINLINE int32 NumVertices() const
+			{
+				return VertexPlanesOffsetCount.Num();
+			}
+
+			int32 NumPlanes() const
+			{
+				return PlaneVerticesOffsetCount.Num();
+			}
+
+			// The number of planes that use the specified vertex
+			int32 NumVertexPlanes(int32 VertexIndex) const
+			{
+				return VertexPlanesOffsetCount[VertexIndex].Value;
+			}
+
+			// Get the plane index (in the outer convex container) of one of the planes that uses the specified vertex
+			int32 GetVertexPlane(int32 VertexIndex, int32 VertexPlaneIndex) const
+			{
+				check(VertexPlaneIndex < NumVertexPlanes(VertexIndex));
+
+				const int32 VertexPlaneFlatArrayIndex = VertexPlanesOffsetCount[VertexIndex].Key + VertexPlaneIndex;
+				return (int32)VertexPlanes[VertexPlaneFlatArrayIndex];
+			}
+
+			// The number of vertices that make up the corners of the specified face
+			int32 NumPlaneVertices(int32 PlaneIndex) const
+			{
+				return PlaneVerticesOffsetCount[PlaneIndex].Value;
+			}
+
+			// Get the vertex index (in the outer convex container) of one of the vertices making up the corners of the specified face
+			int32 GetPlaneVertex(int32 PlaneIndex, int32 PlaneVertexIndex) const
+			{
+				check(PlaneVertexIndex < PlaneVerticesOffsetCount[PlaneIndex].Value);
+
+				const int32 PlaneVertexFlatArrayIndex = PlaneVerticesOffsetCount[PlaneIndex].Key + PlaneVertexIndex;
+				return (int32)PlaneVertices[PlaneVertexFlatArrayIndex];
+			}
+
+			void SetPlaneVertices(const TArray<TArray<int32>>& InPlaneVertices, int32 NumVerts)
+			{
+				// We flatten the ragged arrays into a single array, and store a seperate
+				// arrray of [index,count] tuples to reproduce the ragged arrays.
+				Reset();
+
+				// Generate the ragged array [offset,count] tuples
+				PlaneVerticesOffsetCount.SetNumZeroed(InPlaneVertices.Num());
+				VertexPlanesOffsetCount.SetNumZeroed(NumVerts);
+
+				// Count the number of planes for each vertex (store it in the tuple)
+				int FlatArrayIndexCount = 0;
+				for (int32 PlaneIndex = 0; PlaneIndex < InPlaneVertices.Num(); ++PlaneIndex)
+				{
+					FlatArrayIndexCount += InPlaneVertices[PlaneIndex].Num();
+
+					for (int32 PlaneVertexIndex = 0; PlaneVertexIndex < InPlaneVertices[PlaneIndex].Num(); ++PlaneVertexIndex)
+					{
+						const int32 VertexIndex = InPlaneVertices[PlaneIndex][PlaneVertexIndex];
+						VertexPlanesOffsetCount[VertexIndex].Value++;
+					}
+				}
+
+				// Initialize the flattened arrary offsets and reset the count (we re-increment it when copying the data below)
+				int VertexPlanesArrayStart = 0;
+				for (int32 VertexIndex = 0; VertexIndex < NumVerts; ++VertexIndex)
+				{
+					VertexPlanesOffsetCount[VertexIndex].Key = VertexPlanesArrayStart;
+					VertexPlanesArrayStart += VertexPlanesOffsetCount[VertexIndex].Value;
+					VertexPlanesOffsetCount[VertexIndex].Value = 0;
+				}
+
+				// Allocate space for the flattened arrays
+				PlaneVertices.SetNumZeroed(FlatArrayIndexCount);
+				VertexPlanes.SetNumZeroed(FlatArrayIndexCount);
+
+				// Copy the indices into the flattened arrays
+				int32 PlaneVerticesArrayStart = 0;
+				for (int32 PlaneIndex = 0; PlaneIndex < InPlaneVertices.Num(); ++PlaneIndex)
+				{
+					PlaneVerticesOffsetCount[PlaneIndex].Key = PlaneVerticesArrayStart;
+					PlaneVerticesOffsetCount[PlaneIndex].Value = InPlaneVertices[PlaneIndex].Num();
+					PlaneVerticesArrayStart += InPlaneVertices[PlaneIndex].Num();
+
+					for (int32 PlaneVertexIndex = 0; PlaneVertexIndex < InPlaneVertices[PlaneIndex].Num(); ++PlaneVertexIndex)
+					{
+						const int32 VertexIndex = InPlaneVertices[PlaneIndex][PlaneVertexIndex];
+
+						const int32 PlaneVertexFlatArrayIndex = PlaneVerticesOffsetCount[PlaneIndex].Key + PlaneVertexIndex;
+						PlaneVertices[PlaneVertexFlatArrayIndex] = VertexIndex;
+
+						const int32 VertexPlaneFlatArrayIndex = VertexPlanesOffsetCount[VertexIndex].Key + VertexPlanesOffsetCount[VertexIndex].Value;
+						VertexPlanesOffsetCount[VertexIndex].Value++;
+						VertexPlanes[VertexPlaneFlatArrayIndex] = PlaneIndex;
+					}
+				}
+			}
+
+			void Reset()
+			{
+				PlaneVerticesOffsetCount.Reset();
+				VertexPlanesOffsetCount.Reset();
+				PlaneVertices.Reset();
+				VertexPlanes.Reset();
+			}
+
+			void Serialize(FArchive& Ar)
+			{
+				Ar << PlaneVerticesOffsetCount;
+				Ar << VertexPlanesOffsetCount;
+				Ar << PlaneVertices;
+				Ar << VertexPlanes;
+			}
+
+			friend FArchive& operator<<(FArchive& Ar, TConvexFlattenedArrayStructureData<T_INDEX, T_OFFSETINDEX>& Value)
+			{
+				Value.Serialize(Ar);
+				return Ar;
+			}
+
+
+			// Array of [offset, count] for each plane that gives the set of indices in the PlaneVertices flattened array
+			TArray<TPair<FOffsetIndex, FIndex>> PlaneVerticesOffsetCount;
+
+			// Array of [offset, count] for each vertex that gives the set of indices in the VertexPlanes flattened array
+			TArray<TPair<FOffsetIndex, FIndex>> VertexPlanesOffsetCount;
+
+			// A flattened ragged array. For each plane: the set of vertex indices that form the corners of the face in counter-clockwise order
+			TArray<FIndex> PlaneVertices;
+
+			// A flattened ragged array. For each vertex: the set of plane indices that use the vertex
+			TArray<FIndex> VertexPlanes;
+		};
+
+		using FConvexFlattenedArrayStructureDataS32 = TConvexFlattenedArrayStructureData<int32, int32>;
+		using FConvexFlattenedArrayStructureDataU8 = TConvexFlattenedArrayStructureData<uint8, uint16>;
+
+		//
+		// A utility for loading old convex structure data
+		//
+		class FLegacyConvexStructureDataLoader
+		{
+		public:
+			// Load data from an asset saved before we had a proper half-edge data structure
+			static void Load(FArchive& Ar, TArray<TArray<int32>>& OutPlaneVertices, int32& OutNumVertices)
+			{
+				bool bUseVariableSizeStructureDataUE4 = Ar.CustomVer(FPhysicsObjectVersion::GUID) >= FPhysicsObjectVersion::VariableConvexStructureData;
+				bool bUseVariableSizeStructureDataFN = Ar.CustomVer(FFortniteMainBranchObjectVersion::GUID) >= FFortniteMainBranchObjectVersion::ChaosConvexVariableStructureDataAndVerticesArray;
+				bool bUseVariableSizeStructureData = bUseVariableSizeStructureDataUE4 || bUseVariableSizeStructureDataFN;
+
+				if (!bUseVariableSizeStructureData)
+				{
+					LoadFixedSizeRaggedArrays(Ar, OutPlaneVertices, OutNumVertices);
+				}
+				else
+				{
+					LoadVariableSizePackedArrays(Ar, OutPlaneVertices, OutNumVertices);
+				}
+			}
+
+		private:
+			enum class EIndexType : int8
+			{
+				None,
+				S32,
+				U8,
+			};
+
+			// Orginal structure used 32bit indices regardless of max index, and they were store in ragged TArray<TArray<int32>> of vertex indices per plane
+			static void LoadFixedSizeRaggedArrays(FArchive& Ar, TArray<TArray<int32>>& OutPlaneVertices, int32& OutNumVertices)
+			{
+				TArray<TArray<int32>> OldPlaneVertices;
+				TArray<TArray<int32>> OldVertexPlanes;
+				Ar << OldPlaneVertices;
+				Ar << OldVertexPlanes;
+
+				OutPlaneVertices = MoveTemp(OldPlaneVertices);
+				OutNumVertices = OldVertexPlanes.Num();
+			}
+
+			// The second data structure packed the indices into flat arrays, and used an 8bit index if possible (and 32bit otherwise)
+			static void LoadVariableSizePackedArrays(FArchive& Ar, TArray<TArray<int32>>& OutPlaneVertices, int32& OutNumVertices)
+			{
+				EIndexType OldIndexType;
+				TArray<TArray<int32>> OldPlaneVertices;
+				int32 OldNumVertices = 0;
+
+				Ar << OldIndexType;
+
+				// Read the old structure and pull out the data we need to generate the new one (the output params)
+				if (OldIndexType == EIndexType::S32)
+				{
+					FConvexFlattenedArrayStructureDataS32 OldData;
+					Ar << OldData;
+					ExtractData(OldData, OldPlaneVertices, OldNumVertices);
+				}
+				else if (OldIndexType == EIndexType::U8)
+				{
+					FConvexFlattenedArrayStructureDataU8 OldData;
+					Ar << OldData;
+					ExtractData(OldData, OldPlaneVertices, OldNumVertices);
+				}
+
+				OutPlaneVertices = MoveTemp(OldPlaneVertices);
+				OutNumVertices = OldNumVertices;
+			}
+
+			// Extract ragged arrays of vertices per plane from the flattened array structure
+			// For use by LoadVariableSizePackedArrays to avoid code dupe based on index size
+			template<typename T_OLDCONTAINER>
+			static void ExtractData(const T_OLDCONTAINER& OldData, TArray<TArray<int32>>& OldPlaneVertices, int32& OldNumVertices)
+			{
+				OldPlaneVertices.SetNum(OldData.NumPlanes());
+				for (int32 PlaneIndex = 0; PlaneIndex < OldData.NumPlanes(); ++PlaneIndex)
+				{
+					OldPlaneVertices[PlaneIndex].SetNum(OldData.NumPlaneVertices(PlaneIndex));
+					for (int32 PlaneVertexIndex = 0; PlaneVertexIndex < OldData.NumPlaneVertices(PlaneIndex); ++PlaneVertexIndex)
+					{
+						OldPlaneVertices[PlaneIndex][PlaneVertexIndex] = (int32)OldData.GetPlaneVertex(PlaneIndex, PlaneVertexIndex);
+					}
+				}
+				OldNumVertices = OldData.NumVertices();
+			}
+		};
+	}
+
+}
\ No newline at end of file
diff --git a/Engine/Source/Runtime/Experimental/Chaos/Public/Chaos/ConvexHalfEdgeStructureData.h b/Engine/Source/Runtime/Experimental/Chaos/Public/Chaos/ConvexHalfEdgeStructureData.h
new file mode 100644
--- /dev/null
+++ b/Engine/Source/Runtime/Experimental/Chaos/Public/Chaos/ConvexHalfEdgeStructureData.h
@@ -0,0 +1,410 @@
+// Copyright Epic Games, Inc. All Rights Reserved.
+#pragma once
+
+#include "Chaos/Core.h"
+#include "ChaosArchive.h"
+#include "ChaosCheck.h"
+#include "ChaosLog.h"
+#include "Math/NumericLimits.h"
+#include "UObject/PhysicsObjectVersion.h"
+#include "UObject/FortniteMainBranchObjectVersion.h"
+
+namespace Chaos
+{
+	// Default convex structure index traits - assumes signed
+	template<typename T_INDEX>
+	struct TConvexStructureIndexTraits
+	{
+		using FIndex = T_INDEX;
+		static const FIndex InvalidIndex = TNumericLimits<FIndex>::Lowest();
+		static const FIndex MaxIndex = TNumericLimits<FIndex>::Max();
+
+		static_assert(TIsSigned<T_INDEX>::Value, "The default TConvexStructureIndexTraits implementation is only valid for signed T_INDEX");
+	};
+
+	// uint8 uses 255 as InvalidIndex, and therefore supports elements with indices 0...254
+	template<>
+	struct TConvexStructureIndexTraits<uint8>
+	{
+		using FIndex = uint8;
+		static const FIndex InvalidIndex = TNumericLimits<FIndex>::Max();
+		static const FIndex MaxIndex = TNumericLimits<FIndex>::Max() - 1;
+	};
+
+	// Convex half-edge structure data. Supports different index sizes.
+	// Uses indices into packed arrays rather than pointers. Avoids prev/next indices by keeping a plane's edges in order and sequential.
+	template<typename T_INDEX>
+	class TConvexHalfEdgeStructureData
+	{
+	public:
+		using FIndex = T_INDEX;
+		using FIndexTraits = TConvexStructureIndexTraits<T_INDEX>;
+		using FConvexHalfEdgeStructureData = TConvexHalfEdgeStructureData<T_INDEX>;
+
+		static const FIndex InvalidIndex = FIndexTraits::InvalidIndex;
+		static const int32 MaxIndex = (int32)FIndexTraits::MaxIndex;
+
+		friend class FVertexPlaneIterator;
+
+		// A plane of a convex hull. Each plane has an array of half edges, stored
+		// as an index into the edge list and a count.
+		struct FPlaneData
+		{
+			FIndex FirstHalfEdgeIndex;	// index into HalfEdges
+			FIndex NumHalfEdges;
+
+			friend FArchive& operator<<(FArchive& Ar, FPlaneData& Value)
+			{
+				return Ar << Value.FirstHalfEdgeIndex << Value.NumHalfEdges;
+			}
+		};
+
+		// Every plane is bounded by a sequence of edges, and every edge should be shared 
+		// by two planes. The edges that bound a plane are stored as a sequence of half-edges. 
+		// Each half-edge references the starting vertex of the edge, and the half-edge 
+		// pointing in the opposite direction (belonging to the plane that shares the edge).
+		struct FHalfEdgeData
+		{
+			FIndex PlaneIndex;			// index into Planes
+			FIndex VertexIndex;			// index into Vertices
+			FIndex TwinHalfEdgeIndex;	// index into HalfEdges
+
+			friend FArchive& operator<<(FArchive& Ar, FHalfEdgeData& Value)
+			{
+				return Ar << Value.PlaneIndex << Value.VertexIndex << Value.TwinHalfEdgeIndex;
+			}
+		};
+
+		// A vertex of a convex hull. We just store one edge that uses the vertex - the others
+		// can be found via the half-edge links.
+		struct FVertexData
+		{
+			FIndex FirstHalfEdgeIndex;	// index into HalfEdges
+
+			friend FArchive& operator<<(FArchive& Ar, FVertexData& Value)
+			{
+				return Ar << Value.FirstHalfEdgeIndex;
+			}
+		};
+
+		// Initialize the structure data from the array of vertex indices per plane (in CW or CCW order - it is retained in structure)
+		// If this fails for some reason, the structure data will be invalid (check IsValid())
+		static FConvexHalfEdgeStructureData MakePlaneVertices(const TArray<TArray<int32>>& InPlaneVertices, int32 InNumVertices)
+		{
+			FConvexHalfEdgeStructureData StructureData;
+			StructureData.SetPlaneVertices(InPlaneVertices, InNumVertices);
+			return StructureData;
+		}
+
+		// Return true if we can support this convex, based on number of features and maximum index size
+		static bool CanMake(const TArray<TArray<int32>>& InPlaneVertices, int32 InNumVertices)
+		{
+			int32 HalfEdgeCount = 0;
+			for (int32 PlaneIndex = 0; PlaneIndex < InPlaneVertices.Num(); ++PlaneIndex)
+			{
+				HalfEdgeCount += InPlaneVertices[PlaneIndex].Num();
+			}
+
+			// For a well-formed convex HalfEdgeCount must be larger than NumPlanes and NumVerts, but check them all anyway just in case...
+			return ((HalfEdgeCount <= MaxIndex) && (InPlaneVertices.Num() <= MaxIndex) && (InNumVertices <= MaxIndex));
+		}
+
+
+		bool IsValid() const { return Planes.Num() > 0; }
+		int32 NumPlanes() const { return Planes.Num(); }
+		int32 NumHalfEdges() const { return HalfEdges.Num(); }
+		int32 NumVertices() const { return Vertices.Num(); }
+
+		FPlaneData& GetPlane(int32 PlaneIndex) { return Planes[PlaneIndex]; }
+		const FPlaneData& GetPlane(int32 PlaneIndex) const { return Planes[PlaneIndex]; }
+		FHalfEdgeData& GetHalfEdge(int32 EdgeIndex) { return HalfEdges[EdgeIndex]; }
+		const FHalfEdgeData& GetHalfEdge(int32 EdgeIndex) const { return HalfEdges[EdgeIndex]; }
+		FVertexData& GetVertex(int32 VertexIndex) { return Vertices[VertexIndex]; }
+		const FVertexData& GetVertex(int32 VertexIndex) const { return Vertices[VertexIndex]; }
+
+		// The number of edges bounding the specified plane
+		int32 NumPlaneHalfEdges(int32 PlaneIndex) const
+		{
+			return GetPlane(PlaneIndex).NumHalfEdges;
+		}
+
+		// The edge index of one of the bounding edges of a plane
+		// PlaneIndex must be in range [0, NumPlanes())
+		// PlaneEdgeIndex must be in range [0, NumPlaneHalfEdges(PlaneIndex))
+		// return value is in range [0, NumHalfEdges())
+		int32 GetPlaneHalfEdge(int32 PlaneIndex, int32 PlaneEdgeIndex) const
+		{
+			check(PlaneEdgeIndex >= 0);
+			check(PlaneEdgeIndex < NumPlaneHalfEdges(PlaneIndex));
+			return GetPlane(PlaneIndex).FirstHalfEdgeIndex + PlaneEdgeIndex;
+		}
+
+		// The number of vertices that bound the specified plane (same as number of half edges)
+		// PlaneIndex must be in range [0, NumPlaneHalfEdges(PlaneIndex))
+		int32 NumPlaneVertices(int32 PlaneIndex) const
+		{
+			return GetPlane(PlaneIndex).NumHalfEdges;
+		}
+
+		// Get the index of one of the vertices bounding the specified plane
+		// PlaneIndex must be in range [0, NumPlanes())
+		// PlaneVertexIndex must be in [0, NumPlaneVertices(PlaneIndex))
+		// return value is in [0, NumVertices())
+		int32 GetPlaneVertex(int32 PlaneIndex, int32 PlaneVertexIndex) const
+		{
+			const int32 EdgeIndex = GetPlaneHalfEdge(PlaneIndex, PlaneVertexIndex);
+			return GetHalfEdge(EdgeIndex).VertexIndex;
+		}
+
+		// EdgeIndex must be in range [0, NumHalfEdges())
+		// return value is in range [0, NumPlanes())
+		int32 GetHalfEdgePlane(int32 EdgeIndex) const
+		{
+			return GetHalfEdge(EdgeIndex).PlaneIndex;
+		}
+
+		// EdgeIndex must be in range [0, NumHalfEdges())
+		// return value is in range [0, NumVertices())
+		int32 GetHalfEdgeVertex(int32 EdgeIndex) const
+		{
+			return GetHalfEdge(EdgeIndex).VertexIndex;
+		}
+
+		// EdgeIndex must be in range [0, NumHalfEdges())
+		// return value is in range [0, NumHalfEdges())
+		int32 GetTwinHalfEdge(int32 EdgeIndex) const
+		{
+			return GetHalfEdge(EdgeIndex).TwinHalfEdgeIndex;
+		}
+
+		// Get the previous half edge on the same plane
+		// EdgeIndex must be in range [0, NumHalfEdges())
+		// return value is in range [0, NumHalfEdges())
+		int32 GetPrevHalfEdge(int32 EdgeIndex) const
+		{
+			// Calculate the edge index on the plane
+			const int32 PlaneIndex = GetHalfEdge(EdgeIndex).PlaneIndex;
+			const int32 PlaneEdgeIndex = EdgeIndex - GetPlane(PlaneIndex).FirstHalfEdgeIndex;
+			return GetPrevPlaneHalfEdge(PlaneIndex, PlaneEdgeIndex);
+		}
+
+		// Get the next half edge on the same plane
+		// EdgeIndex must be in range [0, NumHalfEdges())
+		// return value is in range [0, NumHalfEdges())
+		int32 GetNextHalfEdge(int32 EdgeIndex) const
+		{
+			// Calculate the edge index on the plane
+			const int32 PlaneIndex = GetHalfEdge(EdgeIndex).PlaneIndex;
+			const int32 PlaneEdgeIndex = EdgeIndex - GetPlane(PlaneIndex).FirstHalfEdgeIndex;
+			return GetNextPlaneHalfEdge(PlaneIndex, PlaneEdgeIndex);
+		}
+
+		// VertexIndex must be in range [0, NumVertices())
+		// return value is in range [0, NumHalfEdges())
+		int32 GetVertexFirstHalfEdge(int32 VertexIndex) const
+		{
+			return GetVertex(VertexIndex).FirstHalfEdgeIndex;
+		}
+
+		// Iterate over the planes assiciated with a vertex.
+		// Visitor should return false to halt iteration.
+		void VisitVertexPlanes(int32 VertexIndex, const TFunction<bool(int32 PlaneIndex)>& Visitor) const
+		{
+			const int32 FirstEdgeIndex = GetVertexFirstHalfEdge(VertexIndex);
+			int32 EdgeIndex = FirstEdgeIndex;
+			while (EdgeIndex != InvalidIndex)
+			{
+				// Send the plane to the visitor
+				const int32 PlaneIndex = GetHalfEdgePlane(EdgeIndex);
+				const bool bContinue = Visitor(PlaneIndex);
+				
+				// Stop if the vistor wants no more planes
+				if (!bContinue)
+				{
+					break;
+				}
+
+				const int32 TwinEdgeIndex = GetTwinHalfEdge(EdgeIndex);
+				if (TwinEdgeIndex == InvalidIndex)
+				{
+					// Malformed convex, but we need to handle it
+					break;
+				}
+
+				EdgeIndex = GetNextHalfEdge(TwinEdgeIndex);
+				if (EdgeIndex == FirstEdgeIndex)
+				{
+					// We have looped back to the first edge
+					break;
+				}
+			}
+		}
+
+		// Fill an array with plane indices for the specified vertex. Return the number of planes found.
+		int32 FindVertexPlanes(int32 VertexIndex, int32* PlaneIndices, int32 MaxVertexPlanes) const
+		{
+			int32 NumPlanesFound = 0;
+
+			if (MaxVertexPlanes > 0)
+			{
+				VisitVertexPlanes(VertexIndex, 
+					[PlaneIndices, MaxVertexPlanes, &NumPlanesFound](int32 PlaneIndex)
+					{
+						PlaneIndices[NumPlanesFound++] = PlaneIndex;
+						return (NumPlanesFound < MaxVertexPlanes);
+					});
+			}
+
+			return NumPlanesFound;
+		}
+
+		// Initialize the structure data from the set of vertices associated with each plane.
+		// The vertex indices are assumed to be in CCW order (or CW order - doesn't matter here
+		// as long as it is sequential).
+		bool SetPlaneVertices(const TArray<TArray<int32>>& InPlaneVertices, int32 InNumVertices)
+		{
+			// Count the edges
+			int32 HalfEdgeCount = 0;
+			for (int32 PlaneIndex = 0; PlaneIndex < InPlaneVertices.Num(); ++PlaneIndex)
+			{
+				HalfEdgeCount += InPlaneVertices[PlaneIndex].Num();
+			}
+
+			if ((InPlaneVertices.Num() > MaxIndex) || (HalfEdgeCount > MaxIndex) || (InNumVertices > MaxIndex))
+			{
+				return false;
+			}
+
+			Planes.SetNum(InPlaneVertices.Num());
+			HalfEdges.SetNum(HalfEdgeCount);
+			Vertices.SetNum(InNumVertices);
+
+			// Initialize the vertex list - it will be filled in as we build the edge list
+			for (int32 VertexIndex = 0; VertexIndex < Vertices.Num(); ++VertexIndex)
+			{
+				GetVertex(VertexIndex).FirstHalfEdgeIndex = InvalidIndex;
+			}
+
+			// Build the planes and edges. The edges for a plane are stored sequentially in the half-edge array.
+			// On the first pass, the edges contain 2 vertex indices, rather than a vertex index and a twin edge index.
+			// We fix this up on a second pass.
+			int32 NextEdgeIndex = 0;
+			for (int32 PlaneIndex = 0; PlaneIndex < InPlaneVertices.Num(); ++PlaneIndex)
+			{
+				const TArray<int32>& PlaneVertices = InPlaneVertices[PlaneIndex];
+
+				GetPlane(PlaneIndex) =
+				{
+					(FIndex)NextEdgeIndex,
+					(FIndex)PlaneVertices.Num()
+				};
+
+				for (int32 PlaneVertexIndex = 0; PlaneVertexIndex < PlaneVertices.Num(); ++PlaneVertexIndex)
+				{
+					// Add a new edge
+					const int32 VertexIndex0 = PlaneVertices[PlaneVertexIndex];
+					const int32 VertexIndex1 = PlaneVertices[(PlaneVertexIndex + 1) % PlaneVertices.Num()];
+					GetHalfEdge(NextEdgeIndex) =
+					{
+						(FIndex)PlaneIndex,
+						(FIndex)VertexIndex0,
+						(FIndex)VertexIndex1,	// Will get converted to a half-edge index later
+					};
+
+					// If this is the first time Vertex0 has showed up, set its edge index
+					if (Vertices[VertexIndex0].FirstHalfEdgeIndex == InvalidIndex)
+					{
+						Vertices[VertexIndex0].FirstHalfEdgeIndex = NextEdgeIndex;
+					}
+
+					++NextEdgeIndex;
+				}
+			}
+
+			// Find the twin half edge for each edge
+			// @todo(chaos): could use a map of vertex-index-pair to half edge to eliminate O(N^2) algorithm
+			TArray<FIndex> TwinHalfEdgeIndices;
+			TwinHalfEdgeIndices.SetNum(HalfEdges.Num());
+			for (int32 EdgeIndex = 0; EdgeIndex < TwinHalfEdgeIndices.Num(); ++EdgeIndex)
+			{
+				TwinHalfEdgeIndices[EdgeIndex] = InvalidIndex;
+			}
+			for (int32 EdgeIndex0 = 0; EdgeIndex0 < HalfEdges.Num(); ++EdgeIndex0)
+			{
+				const int32 VertexIndex0 = HalfEdges[EdgeIndex0].VertexIndex;
+				const int32 VertexIndex1 = HalfEdges[EdgeIndex0].TwinHalfEdgeIndex;	// Actually a vertex index for now...
+
+				// Find the edge with the vertices the other way round
+				for (int32 EdgeIndex1 = 0; EdgeIndex1 < HalfEdges.Num(); ++EdgeIndex1)
+				{
+					if ((HalfEdges[EdgeIndex1].VertexIndex == VertexIndex1) && (HalfEdges[EdgeIndex1].TwinHalfEdgeIndex == VertexIndex0))
+					{
+						TwinHalfEdgeIndices[EdgeIndex0] = (FIndex)EdgeIndex1;
+						break;
+					}
+				}
+			}
+
+			// Set the twin edge indices
+			for (int32 EdgeIndex = 0; EdgeIndex < HalfEdges.Num(); ++EdgeIndex)
+			{
+				GetHalfEdge(EdgeIndex).TwinHalfEdgeIndex = (FIndex)TwinHalfEdgeIndices[EdgeIndex];
+			}
+
+			return true;
+		}
+
+		void Serialize(FArchive& Ar)
+		{
+			Ar << Planes;
+			Ar << HalfEdges;
+			Ar << Vertices;
+		}
+
+		friend FArchive& operator<<(FArchive& Ar, FConvexHalfEdgeStructureData& Value)
+		{
+			Value.Serialize(Ar);
+			return Ar;
+		}
+
+	private:
+
+		// The edge index of the previous edge on the plane (loops)
+		// PlaneIndex must be in range [0, NumPlanes())
+		// PlaneEdgeIndex must be in range [0, NumPlaneHalfEdges(PlaneIndex))
+		// return value is in range [0, NumHalfEdges())
+		int32 GetPrevPlaneHalfEdge(int32 PlaneIndex, int32 PlaneEdgeIndex) const
+		{
+			// A plane's edges are sequential and loop
+			check(PlaneEdgeIndex >= 0);
+			check(PlaneEdgeIndex < NumPlaneHalfEdges(PlaneIndex));
+			const int32 PlaneHalfEdgeCount = NumPlaneHalfEdges(PlaneIndex);
+			const int32 PrevPlaneEdgeIndex = (PlaneEdgeIndex + PlaneHalfEdgeCount - 1) % PlaneHalfEdgeCount;
+			return GetPlaneHalfEdge(PlaneIndex, PrevPlaneEdgeIndex);
+		}
+
+		// The edge index of the next edge on the plane (loops)
+		// PlaneIndex must be in range [0, NumPlanes())
+		// PlaneEdgeIndex must be in range [0, NumPlaneHalfEdges(PlaneIndex))
+		// return value is in range [0, NumHalfEdges())
+		int32 GetNextPlaneHalfEdge(int32 PlaneIndex, int32 PlaneEdgeIndex) const
+		{
+			// A plane's edges are sequential and loop
+			check(PlaneEdgeIndex >= 0);
+			check(PlaneEdgeIndex < NumPlaneHalfEdges(PlaneIndex));
+			const int32 PlaneHalfEdgeCount = NumPlaneHalfEdges(PlaneIndex);
+			const int32 NextPlaneEdgeIndex = (PlaneEdgeIndex + 1) % PlaneHalfEdgeCount;
+			return GetPlaneHalfEdge(PlaneIndex, NextPlaneEdgeIndex);
+		}
+
+		TArray<FPlaneData> Planes;
+		TArray<FHalfEdgeData> HalfEdges;
+		TArray<FVertexData> Vertices;
+	};
+
+
+	// Typedefs for the supported index sizes
+	using FConvexHalfEdgeStructureDataS32 = TConvexHalfEdgeStructureData<int32>;
+	using FConvexHalfEdgeStructureDataS16 = TConvexHalfEdgeStructureData<int16>;
+	using FConvexHalfEdgeStructureDataU8 = TConvexHalfEdgeStructureData<uint8>;
+}
diff --git a/Engine/Source/Runtime/Experimental/Chaos/Public/Chaos/ConvexStructureData.h b/Engine/Source/Runtime/Experimental/Chaos/Public/Chaos/ConvexStructureData.h
--- a/Engine/Source/Runtime/Experimental/Chaos/Public/Chaos/ConvexStructureData.h
+++ b/Engine/Source/Runtime/Experimental/Chaos/Public/Chaos/ConvexStructureData.h
@@ -1,10 +1,9 @@
 // Copyright Epic Games, Inc. All Rights Reserved.
 #pragma once
 
-#include "Chaos/ImplicitObject.h"
-#include "Chaos/AABB.h"
-#include "Chaos/MassProperties.h"
-#include "CollisionConvexMesh.h"
+#include "Chaos/Core.h"
+#include "Chaos/ConvexFlattenedArrayStructureData.h"
+#include "Chaos/ConvexHalfEdgeStructureData.h"
 #include "ChaosArchive.h"
 #include "ChaosCheck.h"
 #include "ChaosLog.h"
@@ -13,196 +12,81 @@
 
 namespace Chaos
 {
-	// Base class for TConvexStructureDataImp.
-	// NOTE: Deliberately no-virtual destructor so it should never be deleted from a 
-	// base class pointer.
-	class CHAOS_API FConvexStructureDataImp
-	{
-	public:
-		FConvexStructureDataImp() {}
-		~FConvexStructureDataImp() {}
-	};
 
-	// A container for the convex structure data arrays.
-	// This is templated on the index type required, which needs to
-	// be large enough to hold the max of the number of vertices or 
-	// planes on the convex.
-	// 
-	// T_INDEX:			the type used to index into the convex vertices 
-	//					and planes array (in the outer convex). Must be 
-	//					able to contain max(NumPlanes, NumVerts).
-	//
-	// T_OFFSETINDEX:	the type used to index the flattened array of
-	//					indices. Must be able to contain 
-	//					Max(NumPlanes*AverageVertsPerPlane)
-	//
-	template<typename T_INDEX, typename T_OFFSETINDEX>
-	class CHAOS_API TConvexStructureDataImp : public FConvexStructureDataImp
+	// Metadata for a convex shape used by the manifold generation system and anything
+	// else that can benefit from knowing which vertices are associated with the faces.
+	class CHAOS_API FConvexStructureData
 	{
 	public:
-		using FIndex = T_INDEX;
-		using FOffsetIndex = T_OFFSETINDEX;
-
-		FORCEINLINE bool IsValid() const
-		{
-			return PlaneVertices.Num() > 0;
-		}
-
-		FORCEINLINE int32 NumVertices() const
-		{
-			return VertexPlanesOffsetCount.Num();
-		}
-
-		// The number of planes that use the specified vertex
-		FORCEINLINE int32 NumVertexPlanes(int32 VertexIndex) const
-		{
-			return VertexPlanesOffsetCount[VertexIndex].Value;
-		}
-
-		// Get the plane index (in the outer convex container) of one of the planes that uses the specified vertex
-		FORCEINLINE int32 GetVertexPlane(int32 VertexIndex, int32 VertexPlaneIndex) const
-		{
-			check(VertexPlaneIndex < NumVertexPlanes(VertexIndex));
+		using FConvexStructureDataLarge = FConvexHalfEdgeStructureDataS32;
+		using FConvexStructureDataMedium = FConvexHalfEdgeStructureDataS16;
+		using FConvexStructureDataSmall = FConvexHalfEdgeStructureDataU8;
 
-			const int32 VertexPlaneFlatArrayIndex = VertexPlanesOffsetCount[VertexIndex].Key + VertexPlaneIndex;
-			return (int32)VertexPlanes[VertexPlaneFlatArrayIndex];
-		}
-
-		FORCEINLINE int32 NumPlanes() const
-		{
-			return PlaneVerticesOffsetCount.Num();
-		}
-
-		// The number of vertices that make up the corners of the specified face
-		FORCEINLINE int32 NumPlaneVertices(int32 PlaneIndex) const
-		{
-			return PlaneVerticesOffsetCount[PlaneIndex].Value;
-		}
-
-		// Get the vertex index (in the outer convex container) of one of the vertices making up the corners of the specified face
-		FORCEINLINE int32 GetPlaneVertex(int32 PlaneIndex, int32 PlaneVertexIndex) const
+		// Note: this is serialized (do not change order without adding a new ObjectVersion and legacy load path)
+		enum class EIndexType : int8
 		{
-			check(PlaneVertexIndex < PlaneVerticesOffsetCount[PlaneIndex].Value);
+			None,
+			Small,
+			Medium,
+			Large,
+		};
 
-			const int32 PlaneVertexFlatArrayIndex = PlaneVerticesOffsetCount[PlaneIndex].Key + PlaneVertexIndex;
-			return (int32)PlaneVertices[PlaneVertexFlatArrayIndex];
-		}
+		// Cats the inner data container to the correct index size
+		// Do not use - only public for unit tests
+		const FConvexStructureDataLarge& DataL() const { checkSlow(IndexType == EIndexType::Large); return *Data.DataL; }
+		const FConvexStructureDataMedium& DataM() const { checkSlow(IndexType == EIndexType::Medium); return *Data.DataM; }
+		const FConvexStructureDataSmall& DataS() const { checkSlow(IndexType == EIndexType::Small); return *Data.DataS; }
 
-		void SetPlaneVertices(const TArray<TArray<int32>>& InPlaneVertices, int32 NumVerts)
+	private:
+		FConvexStructureDataLarge& DataL() { checkSlow(IndexType == EIndexType::Large); return *Data.DataL; }
+		FConvexStructureDataMedium& DataM() { checkSlow(IndexType == EIndexType::Medium); return *Data.DataM; }
+		FConvexStructureDataSmall& DataS() { checkSlow(IndexType == EIndexType::Small); return *Data.DataS; }
+
+		// Call a function on the read-only data (cast to Small, Medium or Large indices as appropriate).
+		// It will call "Op" on the down-casted data container if we have a valid container. It will
+		// call EmptyOp if the container has not been set up (IndexType is None).
+		// Note: This expands into a switch statement, calling the lambda with a container of each index size.
+		//
+		// Must be used with an auto lambda to handle the downcasting, e.g.:
+		//		int32 Plane0VertexCount = ConstDataOp(
+		//			[](const auto& ConcreteData) { return ConcreteData.NumPlaneVertices(0); },
+		//			[]() { return 0; });
+		//
+		template<typename T_OP, typename T_EMPTYOP>
+		auto ConstDataOp(const T_OP& Op, const T_EMPTYOP& EmptyOp) const
 		{
-			// We flatten the ragged arrays into a single array, and store a seperate
-			// arrray of [index,count] tuples to reproduce the ragged arrays.
-			Reset();
-
-			// Generate the ragged array [offset,count] tuples
-			PlaneVerticesOffsetCount.SetNumZeroed(InPlaneVertices.Num());
-			VertexPlanesOffsetCount.SetNumZeroed(NumVerts);
-
-			// Count the number of planes for each vertex (store it in the tuple)
-			int FlatArrayIndexCount = 0;
-			for (int32 PlaneIndex = 0; PlaneIndex < InPlaneVertices.Num(); ++PlaneIndex)
-			{
-				FlatArrayIndexCount += InPlaneVertices[PlaneIndex].Num();
-
-				for (int32 PlaneVertexIndex = 0; PlaneVertexIndex < InPlaneVertices[PlaneIndex].Num(); ++PlaneVertexIndex)
-				{
-					const int32 VertexIndex = InPlaneVertices[PlaneIndex][PlaneVertexIndex];
-					VertexPlanesOffsetCount[VertexIndex].Value++;
-				}
-			}
-
-			// Initialize the flattened arrary offsets and reset the count (we re-increment it when copying the data below)
-			int VertexPlanesArrayStart = 0;
-			for (int32 VertexIndex = 0; VertexIndex < NumVerts; ++VertexIndex)
-			{
-				VertexPlanesOffsetCount[VertexIndex].Key = static_cast<FOffsetIndex>(VertexPlanesArrayStart);
-				VertexPlanesArrayStart += VertexPlanesOffsetCount[VertexIndex].Value;
-				VertexPlanesOffsetCount[VertexIndex].Value = 0;
-			}
-
-			// Allocate space for the flattened arrays
-			PlaneVertices.SetNumZeroed(FlatArrayIndexCount);
-			VertexPlanes.SetNumZeroed(FlatArrayIndexCount);
-
-			// Copy the indices into the flattened arrays
-			int32 PlaneVerticesArrayStart = 0;
-			for (int32 PlaneIndex = 0; PlaneIndex < InPlaneVertices.Num(); ++PlaneIndex)
+			switch (IndexType)
 			{
-				PlaneVerticesOffsetCount[PlaneIndex].Key = static_cast<FOffsetIndex>(PlaneVerticesArrayStart);
-				PlaneVerticesOffsetCount[PlaneIndex].Value = static_cast<FIndex>(InPlaneVertices[PlaneIndex].Num());
-				PlaneVerticesArrayStart += InPlaneVertices[PlaneIndex].Num();
-
-				for (int32 PlaneVertexIndex = 0; PlaneVertexIndex < InPlaneVertices[PlaneIndex].Num(); ++PlaneVertexIndex)
-				{
-					const int32 VertexIndex = InPlaneVertices[PlaneIndex][PlaneVertexIndex];
-
-					const int32 PlaneVertexFlatArrayIndex = PlaneVerticesOffsetCount[PlaneIndex].Key + PlaneVertexIndex;
-					PlaneVertices[PlaneVertexFlatArrayIndex] = static_cast<FIndex>(VertexIndex);
-
-					const int32 VertexPlaneFlatArrayIndex = VertexPlanesOffsetCount[VertexIndex].Key + VertexPlanesOffsetCount[VertexIndex].Value;
-					VertexPlanesOffsetCount[VertexIndex].Value++;
-					VertexPlanes[VertexPlaneFlatArrayIndex] = static_cast<FIndex>(PlaneIndex);
-				}
+			case EIndexType::Small:
+				return Op(DataS());
+			case EIndexType::Medium:
+				return Op(DataM());
+			case EIndexType::Large:
+				return Op(DataL());
 			}
+			return EmptyOp();
 		}
 
-		void Reset()
-		{
-			PlaneVerticesOffsetCount.Reset();
-			VertexPlanesOffsetCount.Reset();
-			PlaneVertices.Reset();
-			VertexPlanes.Reset();
-		}
-
-		void Serialize(FArchive& Ar)
-		{
-			Ar << PlaneVerticesOffsetCount;
-			Ar << VertexPlanesOffsetCount;
-			Ar << PlaneVertices;
-			Ar << VertexPlanes;
-		}
-
-		friend FArchive& operator<<(FArchive& Ar, TConvexStructureDataImp<T_INDEX, T_OFFSETINDEX>& Value)
+		// A write-enabled version of ConstDataOp
+		template<typename T_OP, typename T_EMPTYOP>
+		auto NonConstDataOp(const T_OP& Op, const T_EMPTYOP& EmptyOp)
 		{
-			Value.Serialize(Ar);
-			return Ar;
+			switch (IndexType)
+			{
+			case EIndexType::Small:
+				return Op(DataS());
+			case EIndexType::Medium:
+				return Op(DataM());
+			case EIndexType::Large:
+				return Op(DataL());
+			}
+			return EmptyOp();
 		}
 
-
-		// Array of [offset, count] for each plane that gives the set of indices in the PlaneVertices flattened array
-		TArray<TPair<FOffsetIndex, FIndex>> PlaneVerticesOffsetCount;
-
-		// Array of [offset, count] for each vertex that gives the set of indices in the VertexPlanes flattened array
-		TArray<TPair<FOffsetIndex, FIndex>> VertexPlanesOffsetCount;
-
-		// A flattened ragged array. For each plane: the set of vertex indices that form the corners of the face in counter-clockwise order
-		TArray<FIndex> PlaneVertices;
-
-		// A flattened ragged array. For each vertex: the set of plane indices that use the vertex
-		TArray<FIndex> VertexPlanes;
-	};
-
-	using FConvexStructureDataS32 = TConvexStructureDataImp<int32, int32>;
-	using FConvexStructureDataU8 = TConvexStructureDataImp<uint8, uint16>;
-
-	// Metadata for a convex shape used by the manifold generation system and anything
-	// else that can benefit from knowing which vertices are associated with the faces.
-	class CHAOS_API FConvexStructureData
-	{
-	private:
-		FORCEINLINE FConvexStructureDataS32& Data32() { return static_cast<FConvexStructureDataS32&>(*Data); }
-		FORCEINLINE FConvexStructureDataU8& Data8() { return static_cast<FConvexStructureDataU8&>(*Data); }
-
 	public:
-		enum class EIndexType : int8
-		{
-			None,
-			S32,
-			U8,
-		};
 
 		FConvexStructureData()
-			: Data(nullptr)
+			: Data{ nullptr }
 			, IndexType(EIndexType::None)
 		{
 		}
@@ -223,116 +107,60 @@ namespace Chaos
 
 		FConvexStructureData& operator=(FConvexStructureData&& Other)
 		{
-			Data = Other.Data;
+			Data.Ptr = Other.Data.Ptr;
 			IndexType = Other.IndexType;
 
-			Other.Data = nullptr;
+			Other.Data.Ptr = nullptr;
 			Other.IndexType = EIndexType::None;
 
 			return *this;
 		}
 
 
-		FORCEINLINE bool IsValid() const
+		bool IsValid() const
 		{
-			return (Data != nullptr);
+			return (Data.Ptr != nullptr);
 		}
 
-		FORCEINLINE EIndexType GetIndexType() const
+		EIndexType GetIndexType() const
 		{
 			return IndexType;
 		}
 
-		// Only exposed for unit tests
-		FORCEINLINE const FConvexStructureDataS32& Data32() const
+		int32 FindVertexPlanes(int32 VertexIndex, int32* VertexPlanes, int32 MaxVertexPlanes) const
 		{
-			checkSlow(IndexType == EIndexType::S32);
-			return static_cast<FConvexStructureDataS32&>(*Data);
-		}
-
-		// Only exposed for unit tests
-		FORCEINLINE const FConvexStructureDataU8& Data8() const
-		{
-			checkSlow(IndexType == EIndexType::U8);
-			return static_cast<FConvexStructureDataU8&>(*Data);
-		}
-
-		// The number of planes that use the specified vertex
-		FORCEINLINE int32 NumVertexPlanes(int32 VertexIndex) const
-		{
-			if (IndexType == EIndexType::S32)
-			{
-				return Data32().NumVertexPlanes(VertexIndex);
-			}
-			else if (IndexType == EIndexType::U8)
-			{
-				return Data8().NumVertexPlanes(VertexIndex);
-			}
-			else
-			{
-				return 0;
-			}
-		}
-
-		// Get the plane index (in the outer convex container) of one of the planes that uses the specified vertex
-		FORCEINLINE int32 GetVertexPlane(int32 VertexIndex, int32 VertexPlaneIndex) const
-		{
-			checkSlow(IsValid());
-			if (IndexType == EIndexType::S32)
-			{
-				return Data32().GetVertexPlane(VertexIndex, VertexPlaneIndex);
-			}
-			else
-			{
-				return Data8().GetVertexPlane(VertexIndex, VertexPlaneIndex);
-			}
+			return ConstDataOp(
+				[VertexIndex, VertexPlanes, MaxVertexPlanes](const auto& ConcreteData) { return ConcreteData.FindVertexPlanes(VertexIndex, VertexPlanes, MaxVertexPlanes); },
+				[]() { return 0; });
 		}
 
 		// The number of vertices that make up the corners of the specified face
-		FORCEINLINE int32 NumPlaneVertices(int32 PlaneIndex) const
+		int32 NumPlaneVertices(int32 PlaneIndex) const
 		{
-			if (IndexType == EIndexType::S32)
-			{
-				return Data32().NumPlaneVertices(PlaneIndex);
-			}
-			else if (IndexType == EIndexType::U8)
-			{
-				return Data8().NumPlaneVertices(PlaneIndex);
-			}
-			else
-			{
-				return 0;
-			}
+			return ConstDataOp(
+				[PlaneIndex](const auto& ConcreteData) { return ConcreteData.NumPlaneVertices(PlaneIndex); },
+				[]() { return 0; });
 		}
 
 		// Get the vertex index (in the outer convex container) of one of the vertices making up the corners of the specified face
-		FORCEINLINE int32 GetPlaneVertex(int32 PlaneIndex, int32 PlaneVertexIndex) const
+		int32 GetPlaneVertex(int32 PlaneIndex, int32 PlaneVertexIndex) const
 		{
 			checkSlow(IsValid());
-			if (IndexType == EIndexType::S32)
-			{
-				return Data32().GetPlaneVertex(PlaneIndex, PlaneVertexIndex);
-			}
-			else
-			{
-				return Data8().GetPlaneVertex(PlaneIndex, PlaneVertexIndex);
-			}
+
+			return ConstDataOp(
+				[PlaneIndex, PlaneVertexIndex](const auto& ConcreteData) { return ConcreteData.GetPlaneVertex(PlaneIndex, PlaneVertexIndex); },
+				[]() { return (int32)INDEX_NONE; });
 		}
 
 		// Initialize the structure data from the set of vertices for each face of the convex
 		void SetPlaneVertices(const TArray<TArray<int32>>& InPlaneVertices, int32 NumVerts)
 		{
-			const EIndexType NewIndexType = GetRequiredIndexType(InPlaneVertices.Num(), NumVerts);
+			const EIndexType NewIndexType = GetRequiredIndexType(InPlaneVertices, NumVerts);
 			CreateDataContainer(NewIndexType);
 
-			if (IndexType == EIndexType::S32)
-			{
-				Data32().SetPlaneVertices(InPlaneVertices, NumVerts);
-			}
-			else if (IndexType == EIndexType::U8)
-			{
-				Data8().SetPlaneVertices(InPlaneVertices, NumVerts);
-			}
+			NonConstDataOp(
+				[InPlaneVertices, NumVerts](auto& ConcreteData) { ConcreteData.SetPlaneVertices(InPlaneVertices, NumVerts); },
+				[]() {});
 		}
 
 		void Serialize(FArchive& Ar)
@@ -340,23 +168,18 @@ namespace Chaos
 			Ar.UsingCustomVersion(FPhysicsObjectVersion::GUID);
 			Ar.UsingCustomVersion(FFortniteMainBranchObjectVersion::GUID);
 
-			bool bUseVariableSizeStructureDataUE4 = Ar.CustomVer(FPhysicsObjectVersion::GUID) >= FPhysicsObjectVersion::VariableConvexStructureData;
-			bool bUseVariableSizeStructureDataFN = Ar.CustomVer(FFortniteMainBranchObjectVersion::GUID) >= FFortniteMainBranchObjectVersion::ChaosConvexVariableStructureDataAndVerticesArray;
-			bool bUseVariableSizeStructureData = bUseVariableSizeStructureDataUE4 || bUseVariableSizeStructureDataFN;
-
-			if (!bUseVariableSizeStructureData)
+			const bool bUseHalfEdgeStructureData = Ar.CustomVer(FPhysicsObjectVersion::GUID) >= FPhysicsObjectVersion::ConvexUsesHalfEdges;
+			
+			// Load and convert the legacy structure if necessary
+			if (Ar.IsLoading() && !bUseHalfEdgeStructureData)
 			{
-				// Orginal structure used 32bit indices regardless of max index, and they were store in ragged TArray<TArray<int32>>
-				TArray<TArray<int32>> OldPlaneVertices;
-				TArray<TArray<int32>> OldVertexPlanes;
-				Ar << OldPlaneVertices;
-				Ar << OldVertexPlanes;
-				SetPlaneVertices(OldPlaneVertices, OldVertexPlanes.Num());
+				LoadLegacyData(Ar);
 				return;
 			}
 
 			if (Ar.IsLoading())
 			{
+				// Create the container with the appropriate index size
 				EIndexType NewIndexType;
 				Ar << NewIndexType;
 				CreateDataContainer(NewIndexType);
@@ -366,14 +189,10 @@ namespace Chaos
 				Ar << IndexType;
 			}
 
-			if (IndexType == EIndexType::S32)
-			{
-				Ar << Data32();
-			}
-			else if (IndexType == EIndexType::U8)
-			{
-				Ar << Data8();
-			}
+			// Serialize the container with the correct index type
+			NonConstDataOp(
+				[&Ar](auto& ConcreteData) { Ar << ConcreteData; },
+				[]() {});
 		}
 
 		friend FArchive& operator<<(FArchive& Ar, FConvexStructureData& Value)
@@ -383,16 +202,32 @@ namespace Chaos
 		}
 
 	private:
+
+		// Load data from an asset saved before we had a proper half-edge data structure
+		void LoadLegacyData(FArchive& Ar)
+		{
+			TArray<TArray<int32>> OldPlaneVertices;
+			int32 OldNumVertices = 0;
+
+			Legacy::FLegacyConvexStructureDataLoader::Load(Ar, OldPlaneVertices, OldNumVertices);
+
+			SetPlaneVertices(OldPlaneVertices, OldNumVertices);
+		}
+
 		// Determine the minimum index size we need for the specified convex size
-		EIndexType GetRequiredIndexType(int32 NumPlanes, int32 NumVerts)
+		EIndexType GetRequiredIndexType(const TArray<TArray<int32>>& InPlaneVertices, int32 NumVerts)
 		{
-			if ((NumPlanes > 255) || (NumVerts > 255))
+			if (FConvexStructureDataSmall::CanMake(InPlaneVertices, NumVerts))
+			{
+				return EIndexType::Small;
+			}
+			else if (FConvexStructureDataMedium::CanMake(InPlaneVertices, NumVerts))
 			{
-				return EIndexType::S32;
+				return EIndexType::Medium;
 			}
 			else
 			{
-				return EIndexType::U8;
+				return EIndexType::Large;
 			}
 		}
 
@@ -401,16 +236,20 @@ namespace Chaos
 		{
 			DestroyDataContainer();
 
-			check(Data == nullptr);
+			check(Data.Ptr == nullptr);
 			check(IndexType == EIndexType::None);
 
-			if (InIndexType == EIndexType::S32)
+			if (InIndexType == EIndexType::Large)
 			{
-				Data = new TConvexStructureDataImp<int32, int32>();
+				Data.DataL = new FConvexStructureDataLarge();
 			}
-			else if (InIndexType == EIndexType::U8)
+			else if (InIndexType == EIndexType::Medium)
 			{
-				Data = new TConvexStructureDataImp<uint8, uint16>();
+				Data.DataM = new FConvexStructureDataMedium();
+			}
+			else if (InIndexType == EIndexType::Small)
+			{
+				Data.DataS = new FConvexStructureDataSmall();
 			}
 
 			IndexType = InIndexType;
@@ -419,28 +258,30 @@ namespace Chaos
 		// Destroy the container we created in CreateDataContainer
 		void DestroyDataContainer()
 		{
-			if (Data != nullptr)
+			if (Data.Ptr != nullptr)
 			{
 				check(IndexType != EIndexType::None);
 
-				if (IndexType == EIndexType::S32)
-				{
-					delete &Data32();
-				}
-				else if (IndexType == EIndexType::U8)
-				{
-					delete &Data8();
-				}
+				NonConstDataOp(
+					[](auto& ConcreteData) 
+					{ 
+						delete &ConcreteData;
+					},
+					[]() {});
 
-				Data = nullptr;
 				IndexType = EIndexType::None;
+				Data.Ptr = nullptr;
 			}
 		}
 
-		// A pointer to the data base class which has no API. Must be downcast 
-		// using Data32() or Data8() depending on IndexType.
-		// Note: not using TUniquePtr here because we don't want a vptr in FConvexStructureDataImp
-		FConvexStructureDataImp* Data;
+		union FStructureData
+		{
+			void* Ptr;
+			FConvexStructureDataLarge* DataL;
+			FConvexStructureDataMedium* DataM;
+			FConvexStructureDataSmall* DataS;
+		};
+		FStructureData Data;
 
 		// The index type we require for the structure data
 		EIndexType IndexType;
diff --git a/Engine/Source/Runtime/Experimental/Chaos/Public/Chaos/ImplicitObjectScaled.h b/Engine/Source/Runtime/Experimental/Chaos/Public/Chaos/ImplicitObjectScaled.h
--- a/Engine/Source/Runtime/Experimental/Chaos/Public/Chaos/ImplicitObjectScaled.h
+++ b/Engine/Source/Runtime/Experimental/Chaos/Public/Chaos/ImplicitObjectScaled.h
@@ -226,12 +226,6 @@ public:
 		return MObject->GetMostOpposingPlane(Normal);
 	}
 
-	// Get the index of the plane that most opposes the normal, assuming it passes through the specified vertex
-	int32 GetMostOpposingPlaneWithVertex(int32 VertexIndex, const FVec3& Normal) const
-	{
-		return MObject->GetMostOpposingPlane(VertexIndex, Normal);
-	}
-
 	// Get the nearest point on an edge of the specified face
 	FVec3 GetClosestEdgePosition(int32 PlaneIndex, const FVec3& Position) const
 	{
@@ -243,16 +237,11 @@ public:
 		return MObject->GetClosestEdgeVertices(PlaneIndexHint, Position, OutVertexIndex0, OutVertexIndex1);
 	}
 
-	// The number of planes that use the specified vertex
-	int32 NumVertexPlanes(int32 VertexIndex) const
+	// Get an array of all the plane indices that belong to a vertex (up to MaxVertexPlanes).
+	// Returns the number of planes found.
+	int32 FindVertexPlanes(int32 VertexIndex, int32* OutVertexPlanes, int32 MaxVertexPlanes) const
 	{
-		return MObject->NumVertexPlanes(VertexIndex);
-	}
-
-	// Get the plane index of one of the planes that uses the specified vertex
-	int32 GetVertexPlane(int32 VertexIndex, int32 VertexPlaneIndex) const
-	{
-		return MObject->GetVertexPlane(VertexIndex, VertexPlaneIndex);
+		return MObject->FindVertexPlanes(VertexIndex, OutVertexPlanes, MaxVertexPlanes);
 	}
 
 	// The number of vertices that make up the corners of the specified face
@@ -277,13 +266,13 @@ public:
 		return MObject->NumVertices();
 	}
 
-	// Get the plane at the specified index (e.g., indices from GetVertexPlane)
+	// Get the plane at the specified index (e.g., indices from FindVertexPlanes)
 	const TPlaneConcrete<FReal, 3> GetPlane(int32 FaceIndex) const
 	{
 		return MObject->GetPlane(FaceIndex);
 	}
 
-	// Get the vertex at the specified index (e.g., indices from GetPlaneVertex)
+	// Get the vertex at the specified index (e.g., indices from GetPlaneVertexs)
 	const FVec3 GetVertex(int32 VertexIndex) const
 	{
 		return MObject->GetVertex(VertexIndex);
@@ -627,12 +616,6 @@ public:
 		return MObject->GetMostOpposingPlaneScaled(Normal, MScale);
 	}
 
-	// Get the index of the plane that most opposes the normal, assuming it passes through the specified vertex
-	int32 GetMostOpposingPlaneWithVertex(int32 VertexIndex, const FVec3& Normal) const
-	{
-		return MObject->GetMostOpposingPlane(VertexIndex, GetInnerUnscaledNormal(Normal));
-	}
-
 	// Get the nearest point on an edge of the specified face
 	FVec3 GetClosestEdgePosition(int32 PlaneIndex, const FVec3& Position) const
 	{
@@ -644,16 +627,11 @@ public:
 		return MObject->GetClosestEdgeVertices(PlaneIndex, MInvScale * Position, OutVertexIndex0, OutVertexIndex1);
 	}
 
-	// The number of planes that use the specified vertex
-	int32 NumVertexPlanes(int32 VertexIndex) const
-	{
-		return MObject->NumVertexPlanes(VertexIndex);
-	}
-
-	// Get the plane index of one of the planes that uses the specified vertex
-	int32 GetVertexPlane(int32 VertexIndex, int32 VertexPlaneIndex) const
+	// Get an array of all the plane indices that belong to a vertex (up to MaxVertexPlanes).
+	// Returns the number of planes found.
+	int32 FindVertexPlanes(int32 VertexIndex, int32* OutVertexPlanes, int32 MaxVertexPlanes) const
 	{
-		return MObject->GetVertexPlane(VertexIndex, VertexPlaneIndex);
+		return MObject->FindVertexPlanes(VertexIndex, OutVertexPlanes, MaxVertexPlanes);
 	}
 
 	// The number of vertices that make up the corners of the specified face
@@ -678,7 +656,7 @@ public:
 		return MObject->NumVertices();
 	}
 
-	// Get the plane at the specified index (e.g., indices from GetVertexPlane)
+	// Get the plane at the specified index (e.g., indices from FindVertexPlanes)
 	const TPlaneConcrete<FReal, 3> GetPlane(int32 FaceIndex) const
 	{
 		const TPlaneConcrete<FReal, 3> InnerPlane = MObject->GetPlane(FaceIndex);
diff --git a/Engine/Source/Runtime/Experimental/ChaosCore/Public/Chaos/Core.h b/Engine/Source/Runtime/Experimental/ChaosCore/Public/Chaos/Core.h
--- a/Engine/Source/Runtime/Experimental/ChaosCore/Public/Chaos/Core.h
+++ b/Engine/Source/Runtime/Experimental/ChaosCore/Public/Chaos/Core.h
@@ -35,5 +35,5 @@ namespace Chaos
 	// NOTE: if you get a merge conflict on the GUID, you must replace it with a new GUID - do not accept the source or target
 	// or you will likely get DDC version conflicts resulting in crashes during load.
 	// Core version string for Chaos data. Any DDC builder dependent on Chaos for serialization should depend on this version
-	static const FString ChaosVersionString = TEXT("6315F0B9C62B42B1A7929463A7375317");
+	static const FString ChaosVersionString = TEXT("D5F568F4-C41C-4474-9040-0004049088DC");
 }
