Make sure MarkTransientDirtyParticle only inserts to transient dirty if not already active

diff --git a/Engine/Source/Runtime/Experimental/Chaos/Public/Chaos/PBDRigidsSOAs.h b/Engine/Source/Runtime/Experimental/Chaos/Public/Chaos/PBDRigidsSOAs.h
--- a/Engine/Source/Runtime/Experimental/Chaos/Public/Chaos/PBDRigidsSOAs.h
+++ b/Engine/Source/Runtime/Experimental/Chaos/Public/Chaos/PBDRigidsSOAs.h
@@ -99,6 +99,11 @@ public:
 		ParticleToIndex.Reset();
 		ParticleArray.Reset();
 	}
+	
+	bool Contains(const TParticleType* Particle) const
+	{
+		return ParticleToIndex.Contains(Particle);
+	}
 
 	template <typename TParticle1>
 	void Insert(const TArray<TParticle1*>& ParticlesToInsert)
@@ -288,13 +293,15 @@ public:
 
 	void MarkTransientDirtyParticle(FGeometryParticleHandle* Particle)
 	{
-		FPBDRigidParticleHandle* Rigid =  Particle->CastToRigidParticle();
-		if(Rigid)
+		if(FPBDRigidParticleHandle* Rigid =  Particle->CastToRigidParticle())
 		{
-			TransientDirtyMapArray.Insert(Rigid);
+			//if active it's already in the dirty view
+			if (!ActiveParticlesMapArray.Contains(Rigid))
+			{
+				TransientDirtyMapArray.Insert(Rigid);
+				UpdateViews();
+			}
 		}
-
-		UpdateViews();
 	}
 
 	// WARNING
