Chaos - contact graph debug draw

diff --git a/Engine/Source/Runtime/Experimental/Chaos/Private/Chaos/ChaosDebugDraw.cpp b/Engine/Source/Runtime/Experimental/Chaos/Private/Chaos/ChaosDebugDraw.cpp
--- a/Engine/Source/Runtime/Experimental/Chaos/Private/Chaos/ChaosDebugDraw.cpp
+++ b/Engine/Source/Runtime/Experimental/Chaos/Private/Chaos/ChaosDebugDraw.cpp
@@ -52,6 +52,34 @@ namespace Chaos
 		float ChaosDebugDrawConvexExplodeDistance = 0.0f;
 		FAutoConsoleVariableRef CVarChaosDebugDrawConvexExplodeDistance(TEXT("p.Chaos.DebugDraw.ConvexExplodeDistance"), ChaosDebugDrawConvexExplodeDistance, TEXT("Explode convex edges by this amount (useful for looking at convex integrity)"));
 
+
+		// NOTE: These settings should never really be used - they are the fallback defaults
+		// if the user does not specify settings in the debug draw call.
+		// See PBDRigidsColver.cpp and ImmediatePhysicsSimulation_Chaos.cpp for example.
+		FChaosDebugDrawSettings ChaosDefaultDebugDebugDrawSettings(
+			/* ArrowSize =					*/ 1.5f,
+			/* BodyAxisLen =				*/ 4.0f,
+			/* ContactLen =					*/ 4.0f,
+			/* ContactWidth =				*/ 2.0f,
+			/* ContactPhiWidth =			*/ 0.0f,
+			/* ContactOwnerWidth =			*/ 0.0f,
+			/* ConstraintAxisLen =			*/ 5.0f,
+			/* JointComSize =				*/ 2.0f,
+			/* LineThickness =				*/ 0.15f,
+			/* DrawScale =					*/ 1.0f,
+			/* FontHeight =					*/ 10.0f,
+			/* FontScale =					*/ 1.5f,
+			/* ShapeThicknesScale =			*/ 1.0f,
+			/* PointSize =					*/ 2.0f,
+			/* VelScale =					*/ 0.0f,
+			/* AngVelScale =				*/ 0.0f,
+			/* ImpulseScale =				*/ 0.0f,
+			/* DrawPriority =				*/ 10.0f,
+			/* bShowSimple =				*/ true,
+			/* bShowComplex =				*/ false,
+			/* bInShowLevelSetCollision =	*/ false
+		);
+
 		const FChaosDebugDrawSettings& GetChaosDebugDrawSettings(const FChaosDebugDrawSettings* InSettings)
 		{
 			if (InSettings != nullptr)
@@ -65,31 +93,6 @@ namespace Chaos
 			}
 		}
 
-		// NOTE: These settings should never really be used - they are the fallback defaults
-		// if the user does not specify settings in the debug draw call.
-		// See PBDRigidsColver.cpp and ImmediatePhysicsSimulation_Chaos.cpp for example.
-		FChaosDebugDrawSettings::FChaosDebugDrawSettings()
-			: ArrowSize(1.5f)
-			, BodyAxisLen(4.0f)
-			, ContactLen(4.0f)
-			, ContactWidth(2.0f)
-			, ContactPhiWidth(0.0f)
-			, ContactOwnerWidth(0.0f)
-			, ConstraintAxisLen(5.0f)
-			, JointComSize(2.0f)
-			, LineThickness(0.15f)
-			, DrawScale(1.0f)
-			, FontHeight(10.0f)
-			, FontScale(1.5f)
-			, ShapeThicknesScale(1.0f)
-			, PointSize(2.0f)
-			, VelScale(0.0f)
-			, AngVelScale(0.0f)
-			, ImpulseScale(0.0f)
-			, DrawPriority(10)
-		{
-		}
-
 		//
 		//
 		//
