Chaos
- SAT implementation (currently unused except for unit tests)
- add unique edge list to structure data (half of the half-edges)

diff --git a/Engine/Source/Programs/HeadlessChaos/Private/HeadlessChaosTestSAT.cpp b/Engine/Source/Programs/HeadlessChaos/Private/HeadlessChaosTestSAT.cpp
new file mode 100644
--- /dev/null
+++ b/Engine/Source/Programs/HeadlessChaos/Private/HeadlessChaosTestSAT.cpp
@@ -0,0 +1,86 @@
+// Copyright Epic Games, Inc. All Rights Reserved.
+
+#include "HeadlessChaos.h"
+#include "HeadlessChaosTestUtility.h"
+#include "Chaos/Core.h"
+#include "Chaos/SAT.h"
+#include "Chaos/Convex.h"
+#include "Chaos/ImplicitObjectScaled.h"
+#include "Chaos/Particles.h"
+
+namespace ChaosTest
+{
+	using namespace Chaos;
+
+	// Two boxes of same size separated along the Z axis with no rotation.
+	// Check that we get the correct separation
+	GTEST_TEST(SATTests, TestSATBox)
+	{
+		const FVec3 BoxSize = FVec3(100, 100, 100);
+		const FReal BoxMargin = 0;
+		const FReal CullDistance = 100;
+		FSATSettings SATSettings;
+
+		const FRigidTransform3 Convex1Transform = FRigidTransform3(FVec3(0, 0, 0), FRotation3::FromAxisAngle(FVec3(1, 0, 0), 0.0f));
+		const FRigidTransform3 Convex2Transform = FRigidTransform3(FVec3(0, 0, 150), FRotation3::FromAxisAngle(FVec3(1, 0, 0), 0.0f));
+		const FImplicitBox3 Convex1 = FImplicitBox3(-0.5f * BoxSize, 0.5f * BoxSize, BoxMargin);
+		const FImplicitBox3 Convex2 = FImplicitBox3(-0.5f * BoxSize, 0.5f * BoxSize, BoxMargin);
+
+		FSATResult SATResult = SATPenetration(Convex1, Convex1Transform, Convex2, Convex2Transform, CullDistance, SATSettings);
+
+		EXPECT_TRUE(SATResult.IsValid());
+
+		const FReal DistanceTolerance = 1.e-3f;
+		const FReal ExpectedSeparation = 50;
+		EXPECT_NEAR(SATResult.SignedDistance, ExpectedSeparation, DistanceTolerance);
+	}
+
+
+	// Two convex boxes of same size separated along the Z axis with no rotation.
+	// Check that we get the correct separation
+	GTEST_TEST(SATTests, TestSATConvexBox)
+	{
+		const FVec3 BoxSize = FVec3(100, 100, 100);
+		const FReal BoxMargin = 0;
+		const FReal CullDistance = 100;
+		FSATSettings SATSettings;
+
+		const FRigidTransform3 Convex1Transform = FRigidTransform3(FVec3(0, 0, 0), FRotation3::FromAxisAngle(FVec3(1, 0, 0), 0.0f));
+		const FRigidTransform3 Convex2Transform = FRigidTransform3(FVec3(0, 0, 150), FRotation3::FromAxisAngle(FVec3(1, 0, 0), 0.0f));
+		const FImplicitConvex3 Convex1 = CreateConvexBox(BoxSize, BoxMargin);
+		const FImplicitConvex3 Convex2 = CreateConvexBox(BoxSize, BoxMargin);
+
+		FSATResult SATResult = SATPenetration(Convex1, Convex1Transform, Convex2, Convex2Transform, CullDistance, SATSettings);
+
+		EXPECT_TRUE(SATResult.IsValid());
+
+		const FReal DistanceTolerance = 1.e-3f;
+		const FReal ExpectedSeparation = 50;
+		EXPECT_NEAR(SATResult.SignedDistance, ExpectedSeparation, DistanceTolerance);
+	}
+
+	// Two scaled convex boxes of same size separated along the Z axis with no rotation.
+	// Check that we get the correct separation
+	GTEST_TEST(SATTests, TestSATScaledConvexBox)
+	{
+		const FVec3 BoxSize = FVec3(50, 100, 50);
+		const FVec3 BoxScale = FVec3(2, 1, 2);
+		const FReal BoxMargin = 0;
+		const FReal CullDistance = 100;
+		FSATSettings SATSettings;
+
+		const FRigidTransform3 Convex1Transform = FRigidTransform3(FVec3(0, 0, 0), FRotation3::FromAxisAngle(FVec3(1, 0, 0), 0.0f));
+		const FRigidTransform3 Convex2Transform = FRigidTransform3(FVec3(0, 0, 150), FRotation3::FromAxisAngle(FVec3(1, 0, 0), 0.0f));
+		const TImplicitObjectScaled<FImplicitConvex3> Convex1 = CreateScaledConvexBox(BoxSize, BoxScale, BoxMargin);
+		const TImplicitObjectScaled<FImplicitConvex3> Convex2 = CreateScaledConvexBox(BoxSize, BoxScale, BoxMargin);
+
+		FSATResult SATResult = SATPenetration(Convex1, Convex1Transform, Convex2, Convex2Transform, CullDistance, SATSettings);
+
+		EXPECT_TRUE(SATResult.IsValid());
+
+		const FReal DistanceTolerance = 1.e-3f;
+		const FReal ExpectedSeparation = 50;
+		EXPECT_NEAR(SATResult.SignedDistance, ExpectedSeparation, DistanceTolerance);
+	}
+
+}
\ No newline at end of file
diff --git a/Engine/Source/Runtime/Core/Public/UObject/PhysicsObjectVersion.h b/Engine/Source/Runtime/Core/Public/UObject/PhysicsObjectVersion.h
--- a/Engine/Source/Runtime/Core/Public/UObject/PhysicsObjectVersion.h
+++ b/Engine/Source/Runtime/Core/Public/UObject/PhysicsObjectVersion.h
@@ -46,7 +46,10 @@ struct CORE_API FPhysicsObjectVersion
 		ChaosClothFixLODTransitionMaps,
 
 		// Convex structure data is now an index-based half-edge structure
-		ConvexUsesHalfEdges,
+		ChaosConvexUsesHalfEdges,
+
+		// Convex structure data has a list of unique edges (half of the half edges)
+		ChaosConvexHasUniqueEdgeSet,
 
 		// -----<new versions can be added above this line>-------------------------------------------------
 		VersionPlusOne,
diff --git a/Engine/Source/Runtime/Experimental/Chaos/Private/Chaos/Convex.cpp b/Engine/Source/Runtime/Experimental/Chaos/Private/Chaos/Convex.cpp
--- a/Engine/Source/Runtime/Experimental/Chaos/Private/Chaos/Convex.cpp
+++ b/Engine/Source/Runtime/Experimental/Chaos/Private/Chaos/Convex.cpp
@@ -236,6 +236,33 @@ namespace Chaos
 		return INDEX_NONE;
 	}
 
+	int32 FConvex::GetEdgeVertex(int32 EdgeIndex, int32 EdgeVertexIndex) const
+	{
+		if (StructureData.IsValid())
+		{
+			return StructureData.GetEdgeVertex(EdgeIndex, EdgeVertexIndex);
+		}
+		return INDEX_NONE;
+	}
+
+	int32 FConvex::GetEdgePlane(int32 EdgeIndex, int32 EdgePlaneIndex) const
+	{
+		if (StructureData.IsValid())
+		{
+			return StructureData.GetEdgePlane(EdgeIndex, EdgePlaneIndex);
+		}
+		return INDEX_NONE;
+	}
+
+	int32 FConvex::NumEdges() const
+	{
+		if (StructureData.IsValid())
+		{
+			return StructureData.NumEdges();
+		}
+		return 0;
+	}
+
 	// Store the structure data with the convex. This is used by manifold generation, for example
 	void FConvex::CreateStructureData(TArray<TArray<int32>>&& PlaneVertexIndices)
 	{
diff --git a/Engine/Source/Runtime/Experimental/Chaos/Public/Chaos/Box.h b/Engine/Source/Runtime/Experimental/Chaos/Public/Chaos/Box.h
--- a/Engine/Source/Runtime/Experimental/Chaos/Public/Chaos/Box.h
+++ b/Engine/Source/Runtime/Experimental/Chaos/Public/Chaos/Box.h
@@ -306,8 +306,20 @@ namespace Chaos
 			return SStructureData.GetPlaneVertex(PlaneIndex, PlaneVertexIndex);
 		}
 
+		int32 GetEdgeVertex(int32 EdgeIndex, int32 EdgeVertexIndex) const
+		{
+			return SStructureData.GetEdgeVertex(EdgeIndex, EdgeVertexIndex);
+		}
+
+		int32 GetEdgePlane(int32 EdgeIndex, int32 EdgePlaneIndex) const
+		{
+			return SStructureData.GetEdgePlane(EdgeIndex, EdgePlaneIndex);
+		}
+
 		int32 NumPlanes() const { return SNormals.Num(); }
 
+		int32 NumEdges() const { return SStructureData.NumEdges(); }
+
 		int32 NumVertices() const { return SVertices.Num(); }
 
 		// Get the plane at the specified index (e.g., indices from FindVertexPlanes)
diff --git a/Engine/Source/Runtime/Experimental/Chaos/Public/Chaos/Convex.h b/Engine/Source/Runtime/Experimental/Chaos/Public/Chaos/Convex.h
--- a/Engine/Source/Runtime/Experimental/Chaos/Public/Chaos/Convex.h
+++ b/Engine/Source/Runtime/Experimental/Chaos/Public/Chaos/Convex.h
@@ -370,11 +370,17 @@ namespace Chaos
 		// Get the vertex index of one of the vertices making up the corners of the specified face
 		int32 GetPlaneVertex(int32 PlaneIndex, int32 PlaneVertexIndex) const;
 
+		int32 GetEdgeVertex(int32 EdgeIndex, int32 EdgeVertexIndex) const;
+
+		int32 GetEdgePlane(int32 EdgeIndex, int32 EdgePlaneIndex) const;
+
 		int32 NumPlanes() const
 		{
 			return Planes.Num();
 		}
 
+		int32 NumEdges() const;
+
 		int32 NumVertices() const
 		{
 			return (int32)Vertices.Num();
diff --git a/Engine/Source/Runtime/Experimental/Chaos/Public/Chaos/ConvexHalfEdgeStructureData.h b/Engine/Source/Runtime/Experimental/Chaos/Public/Chaos/ConvexHalfEdgeStructureData.h
--- a/Engine/Source/Runtime/Experimental/Chaos/Public/Chaos/ConvexHalfEdgeStructureData.h
+++ b/Engine/Source/Runtime/Experimental/Chaos/Public/Chaos/ConvexHalfEdgeStructureData.h
@@ -7,7 +7,6 @@
 #include "ChaosLog.h"
 #include "Math/NumericLimits.h"
 #include "UObject/PhysicsObjectVersion.h"
-#include "UObject/FortniteMainBranchObjectVersion.h"
 
 namespace Chaos
 {
@@ -31,8 +30,9 @@ namespace Chaos
 		static const FIndex MaxIndex = TNumericLimits<FIndex>::Max() - 1;
 	};
 
-	// Convex half-edge structure data. Supports different index sizes.
-	// Uses indices into packed arrays rather than pointers. Avoids prev/next indices by keeping a plane's edges in order and sequential.
+	// Convex half-edge structure data.
+	// Uses indices rather than pointers.
+	// Supports different index sizes.
 	template<typename T_INDEX>
 	class TConvexHalfEdgeStructureData
 	{
@@ -61,8 +61,9 @@ namespace Chaos
 
 		// Every plane is bounded by a sequence of edges, and every edge should be shared 
 		// by two planes. The edges that bound a plane are stored as a sequence of half-edges. 
-		// Each half-edge references the starting vertex of the edge, and the half-edge 
-		// pointing in the opposite direction (belonging to the plane that shares the edge).
+		// Each half-edge references the starting vertex of the edge, the half-edge 
+		// pointing in the opposite direction (belonging to the plane that shares the edge),
+		// and the next half-edge on the same plane.
 		struct FHalfEdgeData
 		{
 			FIndex PlaneIndex;			// index into Planes
@@ -115,10 +116,13 @@ namespace Chaos
 		int32 NumHalfEdges() const { return HalfEdges.Num(); }
 		int32 NumVertices() const { return Vertices.Num(); }
 
+		// Number of unique half-edges (no half edge's twin is in also the list). Should be NumHalfEdges/2
+		int32 NumEdges() const { return Edges.Num(); }		
+
 		FPlaneData& GetPlane(int32 PlaneIndex) { return Planes[PlaneIndex]; }
 		const FPlaneData& GetPlane(int32 PlaneIndex) const { return Planes[PlaneIndex]; }
-		FHalfEdgeData& GetHalfEdge(int32 EdgeIndex) { return HalfEdges[EdgeIndex]; }
-		const FHalfEdgeData& GetHalfEdge(int32 EdgeIndex) const { return HalfEdges[EdgeIndex]; }
+		FHalfEdgeData& GetHalfEdge(int32 HalfEdgeIndex) { return HalfEdges[HalfEdgeIndex]; }
+		const FHalfEdgeData& GetHalfEdge(int32 HalfEdgeIndex) const { return HalfEdges[HalfEdgeIndex]; }
 		FVertexData& GetVertex(int32 VertexIndex) { return Vertices[VertexIndex]; }
 		const FVertexData& GetVertex(int32 VertexIndex) const { return Vertices[VertexIndex]; }
 
@@ -152,51 +156,83 @@ namespace Chaos
 		// return value is in [0, NumVertices())
 		int32 GetPlaneVertex(int32 PlaneIndex, int32 PlaneVertexIndex) const
 		{
-			const int32 EdgeIndex = GetPlaneHalfEdge(PlaneIndex, PlaneVertexIndex);
-			return GetHalfEdge(EdgeIndex).VertexIndex;
+			const int32 HalfEdgeIndex = GetPlaneHalfEdge(PlaneIndex, PlaneVertexIndex);
+			return GetHalfEdge(HalfEdgeIndex).VertexIndex;
 		}
 
-		// EdgeIndex must be in range [0, NumHalfEdges())
+		// HalfEdgeIndex must be in range [0, NumHalfEdges())
 		// return value is in range [0, NumPlanes())
-		int32 GetHalfEdgePlane(int32 EdgeIndex) const
+		int32 GetHalfEdgePlane(int32 HalfEdgeIndex) const
 		{
-			return GetHalfEdge(EdgeIndex).PlaneIndex;
+			return GetHalfEdge(HalfEdgeIndex).PlaneIndex;
 		}
 
-		// EdgeIndex must be in range [0, NumHalfEdges())
+		// HalfEdgeIndex must be in range [0, NumHalfEdges())
 		// return value is in range [0, NumVertices())
-		int32 GetHalfEdgeVertex(int32 EdgeIndex) const
+		int32 GetHalfEdgeVertex(int32 HalfEdgeIndex) const
 		{
-			return GetHalfEdge(EdgeIndex).VertexIndex;
+			return GetHalfEdge(HalfEdgeIndex).VertexIndex;
 		}
 
-		// EdgeIndex must be in range [0, NumHalfEdges())
+		// HalfEdgeIndex must be in range [0, NumHalfEdges())
 		// return value is in range [0, NumHalfEdges())
-		int32 GetTwinHalfEdge(int32 EdgeIndex) const
+		int32 GetTwinHalfEdge(int32 HalfEdgeIndex) const
 		{
-			return GetHalfEdge(EdgeIndex).TwinHalfEdgeIndex;
+			return GetHalfEdge(HalfEdgeIndex).TwinHalfEdgeIndex;
 		}
 
 		// Get the previous half edge on the same plane
-		// EdgeIndex must be in range [0, NumHalfEdges())
+		// HalfEdgeIndex must be in range [0, NumHalfEdges())
 		// return value is in range [0, NumHalfEdges())
-		int32 GetPrevHalfEdge(int32 EdgeIndex) const
+		int32 GetPrevHalfEdge(int32 HalfEdgeIndex) const
 		{
 			// Calculate the edge index on the plane
-			const int32 PlaneIndex = GetHalfEdge(EdgeIndex).PlaneIndex;
-			const int32 PlaneEdgeIndex = EdgeIndex - GetPlane(PlaneIndex).FirstHalfEdgeIndex;
-			return GetPrevPlaneHalfEdge(PlaneIndex, PlaneEdgeIndex);
+			const int32 PlaneIndex = GetHalfEdge(HalfEdgeIndex).PlaneIndex;
+			const int32 PlaneHalfEdgeIndex = HalfEdgeIndex - GetPlane(PlaneIndex).FirstHalfEdgeIndex;
+			return GetPrevPlaneHalfEdge(PlaneIndex, PlaneHalfEdgeIndex);
 		}
 
 		// Get the next half edge on the same plane
-		// EdgeIndex must be in range [0, NumHalfEdges())
+		// HalfEdgeIndex must be in range [0, NumHalfEdges())
 		// return value is in range [0, NumHalfEdges())
-		int32 GetNextHalfEdge(int32 EdgeIndex) const
+		int32 GetNextHalfEdge(int32 HalfEdgeIndex) const
 		{
 			// Calculate the edge index on the plane
-			const int32 PlaneIndex = GetHalfEdge(EdgeIndex).PlaneIndex;
-			const int32 PlaneEdgeIndex = EdgeIndex - GetPlane(PlaneIndex).FirstHalfEdgeIndex;
-			return GetNextPlaneHalfEdge(PlaneIndex, PlaneEdgeIndex);
+			const int32 PlaneIndex = GetHalfEdge(HalfEdgeIndex).PlaneIndex;
+			const int32 PlaneHalfEdgeIndex = HalfEdgeIndex - GetPlane(PlaneIndex).FirstHalfEdgeIndex;
+			return GetNextPlaneHalfEdge(PlaneIndex, PlaneHalfEdgeIndex);
+		}
+
+		// Get a vertex in the specified Edge (NOTE: edge index, not half-edge index)
+		// EdgeIndex must be in range [0, NumEdges())
+		// EdgeVertexIndex must be 0 or 1
+		// return value is in range [0, NumVertices())
+		int32 GetEdgeVertex(int32 EdgeIndex, int32 EdgeVertexIndex) const
+		{
+			if (EdgeVertexIndex == 0)
+			{
+				return GetHalfEdgeVertex(Edges[EdgeIndex]);
+			}
+			else
+			{
+				return GetHalfEdgeVertex(GetTwinHalfEdge(Edges[EdgeIndex]));
+			}
+		}
+
+		// Get a plane for the specified Edge (NOTE: edge index, not half-edge index)
+		// EdgeIndex must be in range [0, NumEdges())
+		// EdgePlaneIndex must be 0 or 1
+		// return value is in range [0, NumPlanes())
+		int32 GetEdgePlane(int32 EdgeIndex, int32 EdgePlaneIndex) const
+		{
+			if (EdgePlaneIndex == 0)
+			{
+				return GetHalfEdgePlane(Edges[EdgeIndex]);
+			}
+			else
+			{
+				return GetHalfEdgePlane(GetTwinHalfEdge(Edges[EdgeIndex]));
+			}
 		}
 
 		// VertexIndex must be in range [0, NumVertices())
@@ -206,37 +242,46 @@ namespace Chaos
 			return GetVertex(VertexIndex).FirstHalfEdgeIndex;
 		}
 
-		// Iterate over the planes assiciated with a vertex.
+		// Iterate over the edges associated with a plane. These edges form the boundary of the plane.
 		// Visitor should return false to halt iteration.
-		void VisitVertexPlanes(int32 VertexIndex, const TFunction<bool(int32 PlaneIndex)>& Visitor) const
+		void VisitPlaneEdges(int32 PlaneIndex, const TFunction<bool(int32 HalfEdgeIndex, int32 NextHalfEdgeIndex)>& Visitor) const
 		{
-			const int32 FirstEdgeIndex = GetVertexFirstHalfEdge(VertexIndex);
-			int32 EdgeIndex = FirstEdgeIndex;
-			while (EdgeIndex != InvalidIndex)
+			const int32 FirstHalfEdgeIndex = GetPlane(PlaneIndex).FirstHalfEdgeIndex;
+			int32 HalfEdgeIndex0 = FirstHalfEdgeIndex;
+			if (HalfEdgeIndex0 != InvalidIndex)
 			{
-				// Send the plane to the visitor
-				const int32 PlaneIndex = GetHalfEdgePlane(EdgeIndex);
-				const bool bContinue = Visitor(PlaneIndex);
-				
-				// Stop if the vistor wants no more planes
-				if (!bContinue)
+				bool bContinue = true;
+				do
 				{
-					break;
-				}
-
-				const int32 TwinEdgeIndex = GetTwinHalfEdge(EdgeIndex);
-				if (TwinEdgeIndex == InvalidIndex)
-				{
-					// Malformed convex, but we need to handle it
-					break;
-				}
+					const int32 HalfEdgeIndex1 = GetNextHalfEdge(HalfEdgeIndex0);
+					if (HalfEdgeIndex1 != InvalidIndex)
+					{
+						bContinue = Visitor(HalfEdgeIndex0, HalfEdgeIndex1);
+					}
+					HalfEdgeIndex0 = HalfEdgeIndex1;
+				} while (bContinue && (HalfEdgeIndex0 != FirstHalfEdgeIndex) && (HalfEdgeIndex0 != InvalidIndex));
+			}
+		}
 
-				EdgeIndex = GetNextHalfEdge(TwinEdgeIndex);
-				if (EdgeIndex == FirstEdgeIndex)
+		// Iterate over the half-edges associated with a vertex (leading out from the vertex, so all half edges have the vertex as the root).
+		// Visitor should return false to halt iteration.
+		void VisitVertexHalfEdges(int32 VertexIndex, const TFunction<bool(int32 HalfEdgeIndex)>& Visitor) const
+		{
+			const int32 FirstHalfEdgeIndex = GetVertex(VertexIndex).FirstHalfEdgeIndex;
+			int32 HalfEdgeIndex = FirstHalfEdgeIndex;
+			if (HalfEdgeIndex != InvalidIndex)
+			{
+				bool bContinue = true;
+				do
 				{
-					// We have looped back to the first edge
-					break;
-				}
+					bContinue = Visitor(HalfEdgeIndex);
+					const int32 TwinHalfEdgeIndex = GetTwinHalfEdge(HalfEdgeIndex);
+					if (TwinHalfEdgeIndex == InvalidIndex)
+					{
+						break;
+					}
+					HalfEdgeIndex = GetNextHalfEdge(TwinHalfEdgeIndex);
+				} while (bContinue && (HalfEdgeIndex != FirstHalfEdgeIndex) && (HalfEdgeIndex != InvalidIndex));
 			}
 		}
 
@@ -247,10 +292,10 @@ namespace Chaos
 
 			if (MaxVertexPlanes > 0)
 			{
-				VisitVertexPlanes(VertexIndex, 
-					[PlaneIndices, MaxVertexPlanes, &NumPlanesFound](int32 PlaneIndex)
+				VisitVertexHalfEdges(VertexIndex,
+					[this, PlaneIndices, MaxVertexPlanes, &NumPlanesFound](int32 HalfEdgeIndex)
 					{
-						PlaneIndices[NumPlanesFound++] = PlaneIndex;
+						PlaneIndices[NumPlanesFound++] = GetHalfEdgePlane(HalfEdgeIndex);
 						return (NumPlanesFound < MaxVertexPlanes);
 					});
 			}
@@ -288,14 +333,14 @@ namespace Chaos
 			// Build the planes and edges. The edges for a plane are stored sequentially in the half-edge array.
 			// On the first pass, the edges contain 2 vertex indices, rather than a vertex index and a twin edge index.
 			// We fix this up on a second pass.
-			int32 NextEdgeIndex = 0;
+			int32 NextHalfEdgeIndex = 0;
 			for (int32 PlaneIndex = 0; PlaneIndex < InPlaneVertices.Num(); ++PlaneIndex)
 			{
 				const TArray<int32>& PlaneVertices = InPlaneVertices[PlaneIndex];
 
 				GetPlane(PlaneIndex) =
 				{
-					(FIndex)NextEdgeIndex,
+					(FIndex)NextHalfEdgeIndex,
 					(FIndex)PlaneVertices.Num()
 				};
 
@@ -304,7 +349,7 @@ namespace Chaos
 					// Add a new edge
 					const int32 VertexIndex0 = PlaneVertices[PlaneVertexIndex];
 					const int32 VertexIndex1 = PlaneVertices[(PlaneVertexIndex + 1) % PlaneVertices.Num()];
-					GetHalfEdge(NextEdgeIndex) =
+					GetHalfEdge(NextHalfEdgeIndex) =
 					{
 						(FIndex)PlaneIndex,
 						(FIndex)VertexIndex0,
@@ -314,10 +359,10 @@ namespace Chaos
 					// If this is the first time Vertex0 has showed up, set its edge index
 					if (Vertices[VertexIndex0].FirstHalfEdgeIndex == InvalidIndex)
 					{
-						Vertices[VertexIndex0].FirstHalfEdgeIndex = NextEdgeIndex;
+						Vertices[VertexIndex0].FirstHalfEdgeIndex = NextHalfEdgeIndex;
 					}
 
-					++NextEdgeIndex;
+					++NextHalfEdgeIndex;
 				}
 			}
 
@@ -325,30 +370,30 @@ namespace Chaos
 			// @todo(chaos): could use a map of vertex-index-pair to half edge to eliminate O(N^2) algorithm
 			TArray<FIndex> TwinHalfEdgeIndices;
 			TwinHalfEdgeIndices.SetNum(HalfEdges.Num());
-			for (int32 EdgeIndex = 0; EdgeIndex < TwinHalfEdgeIndices.Num(); ++EdgeIndex)
+			for (int32 HalfEdgeIndex = 0; HalfEdgeIndex < TwinHalfEdgeIndices.Num(); ++HalfEdgeIndex)
 			{
-				TwinHalfEdgeIndices[EdgeIndex] = InvalidIndex;
+				TwinHalfEdgeIndices[HalfEdgeIndex] = InvalidIndex;
 			}
-			for (int32 EdgeIndex0 = 0; EdgeIndex0 < HalfEdges.Num(); ++EdgeIndex0)
+			for (int32 HalfEdgeIndex0 = 0; HalfEdgeIndex0 < HalfEdges.Num(); ++HalfEdgeIndex0)
 			{
-				const int32 VertexIndex0 = HalfEdges[EdgeIndex0].VertexIndex;
-				const int32 VertexIndex1 = HalfEdges[EdgeIndex0].TwinHalfEdgeIndex;	// Actually a vertex index for now...
+				const int32 VertexIndex0 = HalfEdges[HalfEdgeIndex0].VertexIndex;
+				const int32 VertexIndex1 = HalfEdges[HalfEdgeIndex0].TwinHalfEdgeIndex;	// Actually a vertex index for now...
 
 				// Find the edge with the vertices the other way round
-				for (int32 EdgeIndex1 = 0; EdgeIndex1 < HalfEdges.Num(); ++EdgeIndex1)
+				for (int32 HalfEdgeIndex1 = 0; HalfEdgeIndex1 < HalfEdges.Num(); ++HalfEdgeIndex1)
 				{
-					if ((HalfEdges[EdgeIndex1].VertexIndex == VertexIndex1) && (HalfEdges[EdgeIndex1].TwinHalfEdgeIndex == VertexIndex0))
+					if ((HalfEdges[HalfEdgeIndex1].VertexIndex == VertexIndex1) && (HalfEdges[HalfEdgeIndex1].TwinHalfEdgeIndex == VertexIndex0))
 					{
-						TwinHalfEdgeIndices[EdgeIndex0] = (FIndex)EdgeIndex1;
+						TwinHalfEdgeIndices[HalfEdgeIndex0] = (FIndex)HalfEdgeIndex1;
 						break;
 					}
 				}
 			}
 
 			// Set the twin edge indices
-			for (int32 EdgeIndex = 0; EdgeIndex < HalfEdges.Num(); ++EdgeIndex)
+			for (int32 HalfEdgeIndex = 0; HalfEdgeIndex < HalfEdges.Num(); ++HalfEdgeIndex)
 			{
-				GetHalfEdge(EdgeIndex).TwinHalfEdgeIndex = (FIndex)TwinHalfEdgeIndices[EdgeIndex];
+				GetHalfEdge(HalfEdgeIndex).TwinHalfEdgeIndex = (FIndex)TwinHalfEdgeIndices[HalfEdgeIndex];
 			}
 
 			return true;
@@ -356,9 +401,21 @@ namespace Chaos
 
 		void Serialize(FArchive& Ar)
 		{
+			Ar.UsingCustomVersion(FPhysicsObjectVersion::GUID);
+
 			Ar << Planes;
 			Ar << HalfEdges;
 			Ar << Vertices;
+
+			const bool bHasUniqueEdgeList = Ar.CustomVer(FPhysicsObjectVersion::GUID) >= FPhysicsObjectVersion::ChaosConvexHasUniqueEdgeSet;
+			if (Ar.IsLoading() && !bHasUniqueEdgeList)
+			{
+				BuildUniqueEdgeList();
+			}
+			else
+			{
+				Ar << Edges;
+			}
 		}
 
 		friend FArchive& operator<<(FArchive& Ar, FConvexHalfEdgeStructureData& Value)
@@ -371,35 +428,60 @@ namespace Chaos
 
 		// The edge index of the previous edge on the plane (loops)
 		// PlaneIndex must be in range [0, NumPlanes())
-		// PlaneEdgeIndex must be in range [0, NumPlaneHalfEdges(PlaneIndex))
+		// PlaneHalfEdgeIndex must be in range [0, NumPlaneHalfEdges(PlaneIndex))
 		// return value is in range [0, NumHalfEdges())
-		int32 GetPrevPlaneHalfEdge(int32 PlaneIndex, int32 PlaneEdgeIndex) const
+		int32 GetPrevPlaneHalfEdge(int32 PlaneIndex, int32 PlaneHalfEdgeIndex) const
 		{
 			// A plane's edges are sequential and loop
-			check(PlaneEdgeIndex >= 0);
-			check(PlaneEdgeIndex < NumPlaneHalfEdges(PlaneIndex));
+			check(PlaneHalfEdgeIndex >= 0);
+			check(PlaneHalfEdgeIndex < NumPlaneHalfEdges(PlaneIndex));
 			const int32 PlaneHalfEdgeCount = NumPlaneHalfEdges(PlaneIndex);
-			const int32 PrevPlaneEdgeIndex = (PlaneEdgeIndex + PlaneHalfEdgeCount - 1) % PlaneHalfEdgeCount;
-			return GetPlaneHalfEdge(PlaneIndex, PrevPlaneEdgeIndex);
+			const int32 PrevPlaneHalfEdgeIndex = (PlaneHalfEdgeIndex + PlaneHalfEdgeCount - 1) % PlaneHalfEdgeCount;
+			return GetPlaneHalfEdge(PlaneIndex, PrevPlaneHalfEdgeIndex);
 		}
 
 		// The edge index of the next edge on the plane (loops)
 		// PlaneIndex must be in range [0, NumPlanes())
-		// PlaneEdgeIndex must be in range [0, NumPlaneHalfEdges(PlaneIndex))
+		// PlaneHalfEdgeIndex must be in range [0, NumPlaneHalfEdges(PlaneIndex))
 		// return value is in range [0, NumHalfEdges())
-		int32 GetNextPlaneHalfEdge(int32 PlaneIndex, int32 PlaneEdgeIndex) const
+		int32 GetNextPlaneHalfEdge(int32 PlaneIndex, int32 PlaneHalfEdgeIndex) const
 		{
 			// A plane's edges are sequential and loop
-			check(PlaneEdgeIndex >= 0);
-			check(PlaneEdgeIndex < NumPlaneHalfEdges(PlaneIndex));
+			check(PlaneHalfEdgeIndex >= 0);
+			check(PlaneHalfEdgeIndex < NumPlaneHalfEdges(PlaneIndex));
 			const int32 PlaneHalfEdgeCount = NumPlaneHalfEdges(PlaneIndex);
-			const int32 NextPlaneEdgeIndex = (PlaneEdgeIndex + 1) % PlaneHalfEdgeCount;
-			return GetPlaneHalfEdge(PlaneIndex, NextPlaneEdgeIndex);
+			const int32 NextPlaneHalfEdgeIndex = (PlaneHalfEdgeIndex + 1) % PlaneHalfEdgeCount;
+			return GetPlaneHalfEdge(PlaneIndex, NextPlaneHalfEdgeIndex);
+		}
+
+		// Generate the set of half-edges where none of the edge twins are also in the list.
+		// this is effectively the set of full edges. This will also exclude any malformed edges if they exist.
+		void BuildUniqueEdgeList()
+		{
+			Edges.Empty();
+			Edges.Reserve(NumHalfEdges() / 2);
+
+			for (int32 HalfEdgeIndex = 0; HalfEdgeIndex < NumHalfEdges(); ++HalfEdgeIndex)
+			{
+				const FHalfEdgeData& Edge = GetHalfEdge(HalfEdgeIndex);
+				if (Edge.TwinHalfEdgeIndex != InvalidIndex)
+				{
+					const FHalfEdgeData& TwinEdge = GetHalfEdge(Edge.TwinHalfEdgeIndex);
+					if ((Edge.VertexIndex != InvalidIndex) && (TwinEdge.VertexIndex != InvalidIndex))
+					{
+						if (Edge.VertexIndex < TwinEdge.VertexIndex)
+						{
+							Edges.Add(HalfEdgeIndex);
+						}
+					}
+				}
+			}
 		}
 
 		TArray<FPlaneData> Planes;
 		TArray<FHalfEdgeData> HalfEdges;
 		TArray<FVertexData> Vertices;
+		TArray<FIndex> Edges;
 	};
 
 
diff --git a/Engine/Source/Runtime/Experimental/Chaos/Public/Chaos/ConvexStructureData.h b/Engine/Source/Runtime/Experimental/Chaos/Public/Chaos/ConvexStructureData.h
--- a/Engine/Source/Runtime/Experimental/Chaos/Public/Chaos/ConvexStructureData.h
+++ b/Engine/Source/Runtime/Experimental/Chaos/Public/Chaos/ConvexStructureData.h
@@ -152,6 +152,33 @@ namespace Chaos
 				[]() { return (int32)INDEX_NONE; });
 		}
 
+		int32 NumEdges() const
+		{
+			return ConstDataOp(
+				[](const auto& ConcreteData) { return ConcreteData.NumEdges(); },
+				[]() { return 0; });
+		}
+
+		int32 GetEdgeVertex(int32 EdgeIndex, int32 EdgeVertexIndex) const
+		{
+			return ConstDataOp(
+				[EdgeIndex, EdgeVertexIndex](const auto& ConcreteData)
+				{
+					return ConcreteData.GetEdgeVertex(EdgeIndex, EdgeVertexIndex);
+				},
+				[]() { return (int32)INDEX_NONE; });
+		}
+
+		int32 GetEdgePlane(int32 EdgeIndex, int32 EdgePlaneIndex) const
+		{
+			return ConstDataOp(
+				[EdgeIndex, EdgePlaneIndex](const auto& ConcreteData)
+				{
+					return ConcreteData.GetEdgePlane(EdgeIndex, EdgePlaneIndex);
+				},
+				[]() { return (int32)INDEX_NONE; });
+		}
+
 		// Initialize the structure data from the set of vertices for each face of the convex
 		void SetPlaneVertices(const TArray<TArray<int32>>& InPlaneVertices, int32 NumVerts)
 		{
@@ -168,7 +195,7 @@ namespace Chaos
 			Ar.UsingCustomVersion(FPhysicsObjectVersion::GUID);
 			Ar.UsingCustomVersion(FFortniteMainBranchObjectVersion::GUID);
 
-			const bool bUseHalfEdgeStructureData = Ar.CustomVer(FPhysicsObjectVersion::GUID) >= FPhysicsObjectVersion::ConvexUsesHalfEdges;
+			const bool bUseHalfEdgeStructureData = Ar.CustomVer(FPhysicsObjectVersion::GUID) >= FPhysicsObjectVersion::ChaosConvexUsesHalfEdges;
 			
 			// Load and convert the legacy structure if necessary
 			if (Ar.IsLoading() && !bUseHalfEdgeStructureData)
diff --git a/Engine/Source/Runtime/Experimental/Chaos/Public/Chaos/ImplicitObjectScaled.h b/Engine/Source/Runtime/Experimental/Chaos/Public/Chaos/ImplicitObjectScaled.h
--- a/Engine/Source/Runtime/Experimental/Chaos/Public/Chaos/ImplicitObjectScaled.h
+++ b/Engine/Source/Runtime/Experimental/Chaos/Public/Chaos/ImplicitObjectScaled.h
@@ -256,11 +256,26 @@ public:
 		return MObject->GetPlaneVertex(PlaneIndex, PlaneVertexIndex);
 	}
 
+	int32 GetEdgeVertex(int32 EdgeIndex, int32 EdgeVertexIndex) const
+	{
+		return MObject->GetEdgeVertex(EdgeIndex, EdgeVertexIndex);
+	}
+
+	int32 GetEdgePlane(int32 EdgeIndex, int32 EdgePlaneIndex) const
+	{
+		return MObject->GetEdgePlane(EdgeIndex, EdgePlaneIndex);
+	}
+
 	int32 NumPlanes() const
 	{
 		return MObject->NumPlanes();
 	}
 
+	int32 NumEdges() const
+	{
+		return MObject->NumEdges();
+	}
+
 	int32 NumVertices() const
 	{
 		return MObject->NumVertices();
@@ -646,11 +661,26 @@ public:
 		return MObject->GetPlaneVertex(PlaneIndex, PlaneVertexIndex);
 	}
 
+	int32 GetEdgeVertex(int32 EdgeIndex, int32 EdgeVertexIndex) const
+	{
+		return MObject->GetEdgeVertex(EdgeIndex, EdgeVertexIndex);
+	}
+
+	int32 GetEdgePlane(int32 EdgeIndex, int32 EdgePlaneIndex) const
+	{
+		return MObject->GetEdgePlane(EdgeIndex, EdgePlaneIndex);
+	}
+
 	int32 NumPlanes() const
 	{
 		return MObject->NumPlanes();
 	}
 
+	int32 NumEdges() const
+	{
+		return MObject->NumEdges();
+	}
+
 	int32 NumVertices() const
 	{
 		return MObject->NumVertices();
diff --git a/Engine/Source/Runtime/Experimental/Chaos/Public/Chaos/SAT.h b/Engine/Source/Runtime/Experimental/Chaos/Public/Chaos/SAT.h
new file mode 100644
--- /dev/null
+++ b/Engine/Source/Runtime/Experimental/Chaos/Public/Chaos/SAT.h
@@ -0,0 +1,285 @@
+// Copyright Epic Games, Inc. All Rights Reserved.
+#pragma once
+
+#include "CoreMinimal.h"
+#include "Chaos/Core.h"
+#include "Chaos/Plane.h"
+#include "Chaos/Vector.h"
+#include "Chaos/Transform.h"
+
+namespace Chaos
+{
+	// The feature type returned by SATPenetration
+	enum class ESATFeatureType
+	{
+		None,
+		Plane,
+		Edge,
+		Vertex,
+	};
+
+	// The results from SATPenetration
+	struct FSATResult
+	{
+		FSATResult()
+			: FeatureTypes{ ESATFeatureType::None, ESATFeatureType::None }
+			, FeatureIndices{ INDEX_NONE, INDEX_NONE }
+			, SignedDistance(TNumericLimits<FReal>::Lowest())
+		{}
+
+		bool IsValid() const
+		{
+			return FeatureTypes[0] != ESATFeatureType::None;	// No need to check feature 1
+		}
+
+		bool IsEdgeContact() const
+		{
+			return (FeatureTypes[0] == ESATFeatureType::Edge);	// No need to check feature 1
+		}
+
+		bool IsPlaneContact() const
+		{
+			return (FeatureTypes[0] == ESATFeatureType::Plane) || (FeatureTypes[0] == ESATFeatureType::Vertex);	// No need to check feature 1
+		}
+
+		FSATResult& SwapShapes()
+		{
+			Swap(FeatureTypes[0], FeatureTypes[1]);
+			Swap(FeatureIndices[0], FeatureIndices[1]);
+			return *this;
+		}
+
+		ESATFeatureType FeatureTypes[2];
+		int32 FeatureIndices[2];
+		FReal SignedDistance;
+	};
+
+	// Parameters for SATPenetartion
+	struct FSATSettings
+	{
+		FSATSettings()
+			: PlaneBias(0)
+			, ObjectBias(0)
+		{}
+
+		// Bias to select Plane-Vertex contacts over Edge-Edge contacts with similar separation
+		FReal PlaneBias;
+
+		// Bias to select first (+ve) or second (-ve) object as the Plane owner when both report similar Plane-Vertex distances
+		FReal ObjectBias;
+	};
+
+	// Check whether the two edges of two convex shapes contribute to the Minkowski sum.
+	// A and B are the face normals for the faces of the edge convex 1
+	// C and D are the face normals for the faces of the edge convex 2
+	bool IsMinkowskiSumFace(const FVec3& A, const FVec3& B, const FVec3& C, const FVec3& D)
+	{
+		const FVec3 BA = FVec3::CrossProduct(B, A);
+		const FVec3 DC = FVec3::CrossProduct(D, C);
+		const FReal CBA = FVec3::DotProduct(C, BA);
+		const FReal DBA = FVec3::DotProduct(D, BA);
+		const FReal ADC = FVec3::DotProduct(A, DC);
+		const FReal BDC = FVec3::DotProduct(B, DC);
+
+		const FReal Tolerance = 1.e-2f;
+		return ((CBA * DBA) < -Tolerance) && ((ADC * BDC) < -Tolerance) && ((CBA * BDC) > Tolerance);
+	}
+
+	// Find the nearest Plane-Vertex pair by looking at the Vertices of Convex1 and the Planes of Convex2
+	template <typename ConvexImplicitType1, typename ConvexImplicitType2>
+	FSATResult SATPlaneVertex(
+		const ConvexImplicitType1& Convex1,
+		const FRigidTransform3& Convex1Transform,
+		const ConvexImplicitType2& Convex2,
+		const FRigidTransform3& Convex2Transform,
+		const FReal CullDistance)
+	{
+		FSATResult Result;
+
+		const FRigidTransform3 Convex2ToConvex1Transform = Convex2Transform.GetRelativeTransformNoScale(Convex1Transform);
+
+		const int32 NumPlanes2 = Convex2.NumPlanes();
+		for (int32 PlaneIndex2 = 0; PlaneIndex2 < NumPlanes2; ++PlaneIndex2)
+		{
+			const TPlaneConcrete<FReal, 3> Plane2 = Convex2.GetPlane(PlaneIndex2);
+			const FVec3 PlaneN2In1 = Convex2ToConvex1Transform.TransformVectorNoScale(Plane2.Normal());
+			const FVec3 PlaneX2In1 = Convex2ToConvex1Transform.TransformPositionNoScale(Plane2.X());
+
+			FReal NearestVertexDistance = TNumericLimits<FReal>::Max();
+			int32 NearestVertexIndex1 = INDEX_NONE;
+
+			// @todo(chaos): Use support method that returns vertex index. Use hill climbing
+			const int32 NumVertices1 = Convex1.NumVertices();
+			for (int32 VertexIndex1 = 0; VertexIndex1 < NumVertices1; ++VertexIndex1)
+			{
+				const FVec3 VertexX1 = Convex1.GetVertex(VertexIndex1);
+				const FReal VertexDistance = FVec3::DotProduct(VertexX1 - PlaneX2In1, PlaneN2In1);
+				if (VertexDistance < NearestVertexDistance)
+				{
+					NearestVertexDistance = VertexDistance;
+					NearestVertexIndex1 = VertexIndex1;
+				}
+			}
+
+			// We can stop if all verts are farther than CullDistance from any plane
+			if (NearestVertexDistance > CullDistance)
+			{
+				return FSATResult();
+			}
+
+			// Is this the new best separating axis?
+			if (NearestVertexDistance > Result.SignedDistance)
+			{
+				Result.FeatureTypes[0] = ESATFeatureType::Vertex;
+				Result.FeatureTypes[1] = ESATFeatureType::Plane;
+				Result.FeatureIndices[0] = NearestVertexIndex1;
+				Result.FeatureIndices[1] = PlaneIndex2;
+				Result.SignedDistance = NearestVertexDistance;
+			}
+		}
+
+		return Result;
+	}
+
+	// Find the nearest Edge-Edge pair that contributes to the minkowski surface
+	template <typename ConvexImplicitType1, typename ConvexImplicitType2>
+	FSATResult SATEdgeEdge(
+		const ConvexImplicitType1& Convex1,
+		const FRigidTransform3& Convex1Transform,
+		const ConvexImplicitType2& Convex2,
+		const FRigidTransform3& Convex2Transform,
+		const FReal CullDistance)
+	{
+		FSATResult Result;
+
+		const FRigidTransform3 Convex2ToConvex1Transform = Convex2Transform.GetRelativeTransformNoScale(Convex1Transform);
+
+		// Center of the convex shape, used to enforce correct normal direction
+		const FVec3 Centroid2 = Convex2.GetCenterOfMass();
+		const FVec3 Centroid2In1 = Convex2ToConvex1Transform.TransformPositionNoScale(Centroid2);
+
+		const int32 NumEdges1 = Convex1.NumEdges();
+		const int32 NumEdges2 = Convex2.NumEdges();
+
+		// Loop over the edges in Convex2
+		for (int32 EdgeIndex2 = 0; EdgeIndex2 < NumEdges2; ++EdgeIndex2)
+		{
+			// Edge Vertices in other object space
+			const int32 EdgeVertexIndex2A = Convex2.GetEdgeVertex(EdgeIndex2, 0);
+			const int32 EdgeVertexIndex2B = Convex2.GetEdgeVertex(EdgeIndex2, 1);
+			const FVec3 EdgeVertex2AIn1 = Convex2ToConvex1Transform.TransformPositionNoScale(Convex2.GetVertex(EdgeVertexIndex2A));
+			const FVec3 EdgeVertex2BIn1 = Convex2ToConvex1Transform.TransformPositionNoScale(Convex2.GetVertex(EdgeVertexIndex2B));
+
+			// Planes that use the edge
+			const int32 EdgePlaneIndex2A = Convex2.GetEdgePlane(EdgeIndex2, 0);
+			const int32 EdgePlaneIndex2B = Convex2.GetEdgePlane(EdgeIndex2, 1);
+			const FVec3 EdgePlaneNormal2AIn1 = Convex2ToConvex1Transform.TransformVectorNoScale(Convex2.GetPlane(EdgePlaneIndex2A).Normal());
+			const FVec3 EdgePlaneNormal2BIn1 = Convex2ToConvex1Transform.TransformVectorNoScale(Convex2.GetPlane(EdgePlaneIndex2B).Normal());
+
+			// Loop over the edges in Convex1
+			for (int32 EdgeIndex1 = 0; EdgeIndex1 < NumEdges1; ++EdgeIndex1)
+			{
+				// Edge Vertices
+				const int32 EdgeVertexIndex1A = Convex1.GetEdgeVertex(EdgeIndex1, 0);
+				const int32 EdgeVertexIndex1B = Convex1.GetEdgeVertex(EdgeIndex1, 1);
+				const FVec3 EdgeVertex1A = Convex1.GetVertex(EdgeVertexIndex1A);
+				const FVec3 EdgeVertex1B = Convex1.GetVertex(EdgeVertexIndex1B);
+
+				// Planes that use the edge
+				const int32 EdgePlaneIndex1A = Convex1.GetEdgePlane(EdgeIndex1, 0);
+				const int32 EdgePlaneIndex1B = Convex1.GetEdgePlane(EdgeIndex1, 1);
+				const FVec3 EdgePlaneNormal1A = Convex1.GetPlane(EdgePlaneIndex1A).Normal();
+				const FVec3 EdgePlaneNormal1B = Convex1.GetPlane(EdgePlaneIndex1B).Normal();
+
+				// Does this edge pair contribute to the Minkowski sum?
+				if (!IsMinkowskiSumFace(EdgePlaneNormal1A, EdgePlaneNormal1B, -EdgePlaneNormal2AIn1, -EdgePlaneNormal2BIn1))
+				{
+					continue;
+				}
+
+				// Separating normal (always points away from convex 2)
+				FVec3 EdgeNormal = FVec3::CrossProduct(EdgeVertex1B - EdgeVertex1A, EdgeVertex2BIn1 - EdgeVertex2AIn1);
+				if (!Utilities::NormalizeSafe(EdgeNormal))
+				{
+					continue;
+				}
+				if (FVec3::DotProduct(EdgeNormal, EdgeVertex2AIn1 - Centroid2In1) < 0)
+				{
+					EdgeNormal = -EdgeNormal;
+				}
+
+				// Signed separating distance
+				const FReal EdgeDistance = FVec3::DotProduct(EdgeVertex1A - EdgeVertex2AIn1, EdgeNormal);
+
+				// We can stop if any edge pair on the Minkowski surface is a separating axis
+				if (EdgeDistance > CullDistance)
+				{
+					return FSATResult();
+				}
+
+				// Should we use this edge pair?
+				if (EdgeDistance > Result.SignedDistance)
+				{
+					Result.FeatureTypes[0] = ESATFeatureType::Edge;
+					Result.FeatureTypes[1] = ESATFeatureType::Edge;
+					Result.FeatureIndices[0] = EdgeIndex1;
+					Result.FeatureIndices[1] = EdgeIndex2;
+					Result.SignedDistance = EdgeDistance;
+				}
+			}
+		}
+
+		return Result;
+	}
+
+	// Separating Axis Test
+	// Find the pair of features with the minimum separation distance or the minimum depenetration distance
+	template <typename ConvexImplicitType1, typename ConvexImplicitType2>
+	FSATResult SATPenetration(
+		const ConvexImplicitType1& Convex1,
+		const FRigidTransform3& Convex1Transform,
+		const ConvexImplicitType2& Convex2,
+		const FRigidTransform3& Convex2Transform,
+		const FReal CullDistance,
+		const FSATSettings& Settings)
+	{
+		// Find the closest vertex of Convex1 to a plane of Convex2
+		FSATResult PlaneResult1 = SATPlaneVertex(Convex1, Convex1Transform, Convex2, Convex2Transform, CullDistance);
+		if (PlaneResult1.SignedDistance > CullDistance)
+		{
+			return FSATResult();
+		}
+
+		// Find the closest vertex of Convex2 to a plane of Convex1
+		FSATResult PlaneResult2 = SATPlaneVertex(Convex2, Convex2Transform, Convex1, Convex1Transform, CullDistance).SwapShapes();
+		if (PlaneResult2.SignedDistance > CullDistance)
+		{
+			return FSATResult();
+		}
+
+		// Find the closest edge pair
+		FSATResult EdgeResult = SATEdgeEdge(Convex1, Convex1Transform, Convex2, Convex2Transform, CullDistance);
+		if (EdgeResult.SignedDistance > CullDistance)
+		{
+			return FSATResult();
+		}
+
+		// Select the best contact. Prefer face contacts to edge contacts (for +ve bias)
+		const FReal MaxPlaneDistance = FMath::Max(PlaneResult1.SignedDistance, PlaneResult2.SignedDistance);
+		const bool bUseEdgeResult = (EdgeResult.SignedDistance > MaxPlaneDistance + Settings.PlaneBias);
+		if (bUseEdgeResult)
+		{
+			return EdgeResult;
+		}
+
+		// Prefer planes on Convex2 over Convex1 (for +ve bias, and vice-versa for -ve bias) to prevent flip-flopping
+		const bool bUsePlaneResult1 = (PlaneResult1.SignedDistance > PlaneResult2.SignedDistance - Settings.ObjectBias);
+		if (bUsePlaneResult1)
+		{
+			return PlaneResult1;
+		}
+
+		return PlaneResult2;
+	}
+
+}
\ No newline at end of file
diff --git a/Engine/Source/Runtime/Experimental/ChaosCore/Public/Chaos/Core.h b/Engine/Source/Runtime/Experimental/ChaosCore/Public/Chaos/Core.h
--- a/Engine/Source/Runtime/Experimental/ChaosCore/Public/Chaos/Core.h
+++ b/Engine/Source/Runtime/Experimental/ChaosCore/Public/Chaos/Core.h
@@ -35,5 +35,5 @@ namespace Chaos
 	// NOTE: if you get a merge conflict on the GUID, you must replace it with a new GUID - do not accept the source or target
 	// or you will likely get DDC version conflicts resulting in crashes during load.
 	// Core version string for Chaos data. Any DDC builder dependent on Chaos for serialization should depend on this version
-	static const FString ChaosVersionString = TEXT("D5F568F4-C41C-4474-9040-0004049088DC");
+	static const FString ChaosVersionString = TEXT("9C6B5D3F-C884-4564-BF7C-68B28AFFDD33");
 }
