Chaos - Fix merge mistake

diff --git a/Engine/Source/Runtime/Experimental/Chaos/Private/Chaos/ChaosDebugDraw.cpp b/Engine/Source/Runtime/Experimental/Chaos/Private/Chaos/ChaosDebugDraw.cpp
--- a/Engine/Source/Runtime/Experimental/Chaos/Private/Chaos/ChaosDebugDraw.cpp
+++ b/Engine/Source/Runtime/Experimental/Chaos/Private/Chaos/ChaosDebugDraw.cpp
@@ -87,10 +87,7 @@ namespace Chaos
 				return *InSettings;
 			}
 
-			{
-				static FChaosDebugDrawSettings ChaosDefaultDebugDebugDrawSettings;
-				return ChaosDefaultDebugDebugDrawSettings;
-			}
+			return ChaosDefaultDebugDebugDrawSettings;
 		}
 
 		//
