Added the ability to turn off no collision optimization.

diff --git a/Engine/Source/Runtime/Experimental/Chaos/Private/PhysicsProxy/SingleParticlePhysicsProxy.cpp b/Engine/Source/Runtime/Experimental/Chaos/Private/PhysicsProxy/SingleParticlePhysicsProxy.cpp
--- a/Engine/Source/Runtime/Experimental/Chaos/Private/PhysicsProxy/SingleParticlePhysicsProxy.cpp
+++ b/Engine/Source/Runtime/Experimental/Chaos/Private/PhysicsProxy/SingleParticlePhysicsProxy.cpp
@@ -32,6 +32,9 @@ FSingleParticlePhysicsProxy::~FSingleParticlePhysicsProxy()
 {
 }
 
+CHAOS_API int32 ForceNoCollisionIntoSQ = 0;
+FAutoConsoleVariableRef CVarForceNoCollisionIntoSQ(TEXT("p.ForceNoCollisionIntoSQ"), ForceNoCollisionIntoSQ, TEXT("When enabled, all particles end up in sq structure, even ones with no collision"));
+
 template <Chaos::EParticleType ParticleType, typename TEvolution>
 void PushToPhysicsStateImp(const Chaos::FDirtyPropertiesManager& Manager, Chaos::FGeometryParticleHandle* Handle, int32 DataIdx, const Chaos::FDirtyProxy& Dirty, Chaos::FShapeDirtyData* ShapesData, TEvolution& Evolution, bool bResimInitialized)
 {
@@ -161,7 +164,7 @@ void PushToPhysicsStateImp(const Chaos::FDirtyPropertiesManager& Manager, Chaos:
 		}
 		
 
-		if(bUpdateCollisionData)
+		if(bUpdateCollisionData && !ForceNoCollisionIntoSQ)
 		{
 			//Some shapes were not dirty and may have collision - so have to iterate them all. TODO: find a better way to handle this case
 			if(!bHasCollision && Dirty.ShapeDataIndices.Num() != Handle->ShapesArray().Num())
diff --git a/Engine/Source/Runtime/Experimental/Chaos/Public/Chaos/PBDRigidsEvolution.h b/Engine/Source/Runtime/Experimental/Chaos/Public/Chaos/PBDRigidsEvolution.h
--- a/Engine/Source/Runtime/Experimental/Chaos/Public/Chaos/PBDRigidsEvolution.h
+++ b/Engine/Source/Runtime/Experimental/Chaos/Public/Chaos/PBDRigidsEvolution.h
@@ -24,6 +24,9 @@ extern int32 ChaosRigidsEvolutionApplyPushoutAllowEarlyOutCVar;
 extern int32 ChaosNumPushOutIterationsOverride;
 extern int32 ChaosNumContactIterationsOverride;
 
+extern CHAOS_API int32 ForceNoCollisionIntoSQ;
+
+
 namespace Chaos
 {
 
@@ -342,6 +345,7 @@ public:
 		}
 	}
 
+
 	template <bool bPersistent>
 	FORCEINLINE_DEBUGGABLE void DirtyParticle(TGeometryParticleHandleImp<FReal, 3, bPersistent>& Particle)
 	{
@@ -374,7 +378,7 @@ public:
 		}
 
 		//only add to acceleration structure if it has collision
-		if (Particle.HasCollision())
+		if (Particle.HasCollision() || ForceNoCollisionIntoSQ)
 		{
 			//TODO: distinguish between new particles and dirty particles
 			const FUniqueIdx UniqueIdx = Particle.UniqueIdx();
