fixing compile error in chaos debug draw

diff --git a/Engine/Source/Runtime/Experimental/Chaos/Private/Chaos/ChaosDebugDraw.cpp b/Engine/Source/Runtime/Experimental/Chaos/Private/Chaos/ChaosDebugDraw.cpp
--- a/Engine/Source/Runtime/Experimental/Chaos/Private/Chaos/ChaosDebugDraw.cpp
+++ b/Engine/Source/Runtime/Experimental/Chaos/Private/Chaos/ChaosDebugDraw.cpp
@@ -74,6 +74,7 @@ namespace Chaos
 			/* VelScale =					*/ 0.0f,
 			/* AngVelScale =				*/ 0.0f,
 			/* ImpulseScale =				*/ 0.0f,
+			/* InertiaScale =				*/ 0.0f,
 			/* DrawPriority =				*/ 10.0f,
 			/* bShowSimple =				*/ true,
 			/* bShowComplex =				*/ false,
