Chaos : Add option to ignore sleeping contacts in broadphase to improve perf.

diff --git a/Engine/Source/Runtime/Experimental/Chaos/Private/Chaos/PBDConstraintGraph.cpp b/Engine/Source/Runtime/Experimental/Chaos/Private/Chaos/PBDConstraintGraph.cpp
--- a/Engine/Source/Runtime/Experimental/Chaos/Private/Chaos/PBDConstraintGraph.cpp
+++ b/Engine/Source/Runtime/Experimental/Chaos/Private/Chaos/PBDConstraintGraph.cpp
@@ -29,6 +29,12 @@ FAutoConsoleVariableRef CVarChaosSolverCollisionDefaultLinearSleepThreshold(TEXT
 Chaos::FRealSingle ChaosSolverCollisionDefaultAngularSleepThresholdCVar = 0.0087f;  //~1/2 unit mass degree
 FAutoConsoleVariableRef CVarChaosSolverCollisionDefaultAngularSleepThreshold(TEXT("p.ChaosSolverCollisionDefaultAngularSleepThreshold"), ChaosSolverCollisionDefaultAngularSleepThresholdCVar, TEXT("Default angular threshold for sleeping.[def:0.0087]"));
 
+namespace Chaos
+{
+	bool ChaosPerfHackIgnoreSleepingContacts = false;
+	FAutoConsoleVariableRef CVarChaosPerfHackDisableSleepingConstraints(TEXT("p.Chaos.PerfHackIgnoreSleepingContacts"), ChaosPerfHackIgnoreSleepingContacts, TEXT("Hack to improve perf by not generating all contacts for sleeping particles"));
+}
+
 FPBDConstraintGraph::FPBDConstraintGraph() : VisitToken(0)
 {
 }
@@ -584,6 +590,21 @@ void FPBDConstraintGraph::ComputeIslands(const TParticleView<FPBDRigidParticles>
 			}
 			else
 			{
+				int32 NumRigidsInNewIsland = 0;
+				if (ChaosPerfHackIgnoreSleepingContacts && NewIslandParticles.IsValidIndex(OtherIsland))
+				{
+					for (FGeometryParticleHandle* Particle : NewIslandParticles[OtherIsland])
+					{
+						if (Particle)
+						{
+							if (FPBDRigidParticleHandle* PBDRigid = Particle->CastToRigidParticle())
+							{
+								NumRigidsInNewIsland++;
+							}
+						}
+					}
+				}
+
 				for (TGeometryParticleHandle<FReal, 3>* Particle : IslandToParticles[Island])
 				{
 					if (CHAOS_ENSURE(Particle))
@@ -591,7 +612,11 @@ void FPBDConstraintGraph::ComputeIslands(const TParticleView<FPBDRigidParticles>
 						TPBDRigidParticleHandle<FReal, 3>* PBDRigid = Particle->CastToRigidParticle();
 						if (PBDRigid && PBDRigid->ObjectState() != EObjectStateType::Kinematic)
 						{
-							Particles.ActivateParticle(Particle);
+							const bool bShouldActivateParticle = ChaosPerfHackIgnoreSleepingContacts ? !(NumRigidsInNewIsland == 1 && PBDRigid->Sleeping()) : true;
+							if (bShouldActivateParticle)
+							{
+								Particles.ActivateParticle(Particle);
+							}
 						}
 					}
 				}
diff --git a/Engine/Source/Runtime/Experimental/Chaos/Public/Chaos/Collision/SpatialAccelerationBroadPhase.h b/Engine/Source/Runtime/Experimental/Chaos/Public/Chaos/Collision/SpatialAccelerationBroadPhase.h
--- a/Engine/Source/Runtime/Experimental/Chaos/Public/Chaos/Collision/SpatialAccelerationBroadPhase.h
+++ b/Engine/Source/Runtime/Experimental/Chaos/Public/Chaos/Collision/SpatialAccelerationBroadPhase.h
@@ -20,7 +20,9 @@ namespace Chaos
 	DECLARE_CYCLE_STAT_EXTERN(TEXT("Collisions::Filtering"), STAT_Collisions_Filtering, STATGROUP_ChaosCollision, CHAOS_API);
 	DECLARE_CYCLE_STAT_EXTERN(TEXT("Collisions::GenerateCollisions"), STAT_Collisions_GenerateCollisions, STATGROUP_ChaosCollision, CHAOS_API);
 	DECLARE_CYCLE_STAT_EXTERN(TEXT("Collisions::ReceiveCollisions"), STAT_Collisions_ReceiveCollisions, STATGROUP_ChaosCollision, CHAOS_API);
-	
+
+	extern bool ChaosPerfHackIgnoreSleepingContacts;
+
 	template <typename TPayloadType, typename T, int d>
 	class ISpatialAcceleration;
 
@@ -286,7 +288,23 @@ namespace Chaos
 					const bool bSecondParticleWillHaveAnswer = !bIsResimming || Particle2.SyncState() == ESyncState::HardDesync;
 					// Sleeping vs dynamic gets picked up by the other direction.
 					const bool bIsParticle2Dynamic = Particle2.CastToRigidParticle() && Particle2.ObjectState() == EObjectStateType::Dynamic;
-					if (Particle1.ObjectState() == EObjectStateType::Sleeping && bIsParticle2Dynamic && bSecondParticleWillHaveAnswer)
+
+					bool bSkipSleepingContact = bSecondParticleWillHaveAnswer;
+					if (ChaosPerfHackIgnoreSleepingContacts)
+					{
+						const bool bIsParticle2Kinematic = Particle2.CastToKinematicParticle() &&
+							(Particle2.ObjectState() == EObjectStateType::Kinematic &&
+								(Particle2.CastToKinematicParticle()->V().SizeSquared() > 1e-4 ||
+									Particle2.Geometry()->GetType() == FCapsule::StaticType()));
+
+						bSkipSleepingContact &= !bIsParticle2Kinematic;
+					}
+					else
+					{
+						bSkipSleepingContact &= bIsParticle2Dynamic;
+					}
+
+					if (Particle1.ObjectState() == EObjectStateType::Sleeping && bSkipSleepingContact)
 					{
 						//question: if !bSecondParticleWillHaveAnswer do we need to reorder constraint?
 						continue;
