Merging 16132794

Chaos
- debug draw colliding shapes
- debug draw exploded comvex

diff --git a/Engine/Source/Runtime/Experimental/Chaos/Private/PBDRigidsSolver.cpp b/Engine/Source/Runtime/Experimental/Chaos/Private/PBDRigidsSolver.cpp
--- a/Engine/Source/Runtime/Experimental/Chaos/Private/PBDRigidsSolver.cpp
+++ b/Engine/Source/Runtime/Experimental/Chaos/Private/PBDRigidsSolver.cpp
@@ -1293,6 +1293,10 @@ namespace Chaos
 		{
 			DebugDraw::DrawConstraintGraph(FRigidTransform3(), GetEvolution()->GetCollisionConstraintsRule().GetGraphColor(), &ChaosSolverDebugDebugDrawSettings);
 		}
+		if (ChaosSolverDebugDrawCollidingShapes == 1)
+		{
+			DebugDraw::DrawCollidingShapes(FRigidTransform3(), GetEvolution()->GetCollisionConstraints(), 1.f, &ChaosSolverDebugDebugDrawSettings);
+		}
 		if (ChaosSolverDrawJoints == 1)
 		{
 			DebugDraw::DrawJointConstraints(FRigidTransform3(), JointConstraints, 1.0f, ChaosSolverDrawJointFeatures, &ChaosSolverDebugDebugDrawSettings);
