Chaos: No longer consider statics as awake when deciding whether or not to wake an island. This fixes bStartAwake not functioning on loading a level because adding the static world geometry to an island would wake it.

diff --git a/Engine/Source/Runtime/Experimental/Chaos/Private/Chaos/PBDConstraintGraph.cpp b/Engine/Source/Runtime/Experimental/Chaos/Private/Chaos/PBDConstraintGraph.cpp
--- a/Engine/Source/Runtime/Experimental/Chaos/Private/Chaos/PBDConstraintGraph.cpp
+++ b/Engine/Source/Runtime/Experimental/Chaos/Private/Chaos/PBDConstraintGraph.cpp
@@ -478,7 +478,7 @@ void FPBDConstraintGraph::ComputeIslands(const TParticleView<FPBDRigidParticles>
 
 				for (FGeometryParticleHandle* Particle : NewIslandParticles[Island])
 				{
-					if (!Particle->Sleeping())
+					if (Particle->ObjectState() != EObjectStateType::Static && !Particle->Sleeping())
 					{
 						bSleepState = false;
 						break;
@@ -487,15 +487,21 @@ void FPBDConstraintGraph::ComputeIslands(const TParticleView<FPBDRigidParticles>
 
 				for (FGeometryParticleHandle* Particle : NewIslandParticles[Island])
 				{
-					//@todo(DEMO_HACK) : Need to fix, remove the !InParticles.Disabled(Index)
-					if (Particle->Sleeping() && !bSleepState/* && !Particle->Disabled()*/)
+					if (Particle->Sleeping() && !bSleepState)
 					{
 						Particles.ActivateParticle(Particle); 	//todo: record state change for array reorder
 					}
 
 					FPBDRigidParticleHandle* PBDRigid = Particle->CastToRigidParticle();
-					if(PBDRigid && PBDRigid->ObjectState() != EObjectStateType::Kinematic)
+					if(PBDRigid)
 					{
+						const EObjectStateType CurrState = PBDRigid->ObjectState();
+						if(CurrState == EObjectStateType::Kinematic || CurrState == EObjectStateType::Static)
+						{
+							// Statics and kinematics can't have sleeping states so don't attempt to set one.
+							break;
+						}
+
 						if (!Particle->Sleeping() && bSleepState)
 						{
 							Particles.DeactivateParticle(Particle); 	//todo: record state change for array reorder
@@ -506,7 +512,7 @@ void FPBDConstraintGraph::ComputeIslands(const TParticleView<FPBDRigidParticles>
 						PBDRigid->SetSleeping(bSleepState);
 					}
 
-					if ((Particle->Sleeping() /*|| Particle->Disabled()*/))
+					if (Particle->Sleeping())
 					{
 						Particles.DeactivateParticle(Particle);	//todo: record state change for array reorder (function could return true/false)
 					}
