Chaos - stat counters for solver iterations

diff --git a/Engine/Source/Runtime/Experimental/Chaos/Private/PBDRigidsSolver.cpp b/Engine/Source/Runtime/Experimental/Chaos/Private/PBDRigidsSolver.cpp
--- a/Engine/Source/Runtime/Experimental/Chaos/Private/PBDRigidsSolver.cpp
+++ b/Engine/Source/Runtime/Experimental/Chaos/Private/PBDRigidsSolver.cpp
@@ -44,6 +44,14 @@ DECLARE_DWORD_ACCUMULATOR_STAT(TEXT("NumJoints"), STAT_ChaosCounter_NumJoints, S
 
 CSV_DEFINE_CATEGORY(ChaosCounters, true);
 
+// Stat Iteration counters
+DECLARE_DWORD_ACCUMULATOR_STAT(TEXT("NumIterations"), STAT_ChaosIterations_NumIterations, STATGROUP_ChaosIterations);
+DECLARE_DWORD_ACCUMULATOR_STAT(TEXT("NumCollisionIterations"), STAT_ChaosIterations_NumCollisionIterations, STATGROUP_ChaosIterations);
+DECLARE_DWORD_ACCUMULATOR_STAT(TEXT("NumJointIterations"), STAT_ChaosIterations_NumJointIterations, STATGROUP_ChaosIterations);
+DECLARE_DWORD_ACCUMULATOR_STAT(TEXT("NumPushOutIterations"), STAT_ChaosIterations_NumPushOutIterations, STATGROUP_ChaosIterations);
+DECLARE_DWORD_ACCUMULATOR_STAT(TEXT("NumPushOutCollisionIterations"), STAT_ChaosIterations_NumPushOutCollisionIterations, STATGROUP_ChaosIterations);
+DECLARE_DWORD_ACCUMULATOR_STAT(TEXT("NumPushOutJointIterations"), STAT_ChaosIterations_NumPushOutJointIterations, STATGROUP_ChaosIterations);
+
 
 // DebugDraw CVars
 #if CHAOS_DEBUG_DRAW
@@ -317,7 +325,7 @@ namespace Chaos
 #endif
 			}
 
-			{
+			MSolver->UpdateCounters();
 				SCOPE_CYCLE_COUNTER(STAT_EventDataGathering);
 				{
 					SCOPE_CYCLE_COUNTER(STAT_FillProducerData);
@@ -1308,6 +1316,46 @@ namespace Chaos
 #endif
 	}
 
+	void FPBDRigidsSolver::UpdateCounters() const
+	{
+		// Particle counts
+		SET_DWORD_STAT(STAT_ChaosCounter_NumDisabledParticles, GetEvolution()->GetParticles().GetAllParticlesView().Num() - GetEvolution()->GetParticles().GetNonDisabledView().Num());
+		SET_DWORD_STAT(STAT_ChaosCounter_NumParticles, GetEvolution()->GetParticles().GetNonDisabledView().Num());
+		SET_DWORD_STAT(STAT_ChaosCounter_NumDynamicParticles, GetEvolution()->GetParticles().GetNonDisabledDynamicView().Num());
+		SET_DWORD_STAT(STAT_ChaosCounter_NumActiveDynamicParticles, GetEvolution()->GetParticles().GetActiveParticlesView().Num());
+		SET_DWORD_STAT(STAT_ChaosCounter_NumKinematicParticles, GetEvolution()->GetParticles().GetActiveKinematicParticlesView().Num());
+		SET_DWORD_STAT(STAT_ChaosCounter_NumStaticParticles, GetEvolution()->GetParticles().GetActiveStaticParticlesView().Num());
+		SET_DWORD_STAT(STAT_ChaosCounter_NumGeometryCollectionParticles, GetEvolution()->GetParticles().GetGeometryCollectionParticles().Size());
+
+		// Constraint counts
+		SET_DWORD_STAT(STAT_ChaosCounter_NumIslands, GetEvolution()->GetConstraintGraph().NumIslands());
+		SET_DWORD_STAT(STAT_ChaosCounter_NumContacts, NumCollisionConstraints());
+		SET_DWORD_STAT(STAT_ChaosCounter_NumJoints, NumJointConstraints());
+
+		// Iterations
+		SET_DWORD_STAT(STAT_ChaosIterations_NumIterations, GetEvolution()->GetNumIterations());
+		SET_DWORD_STAT(STAT_ChaosIterations_NumCollisionIterations, GetEvolution()->GetCollisionConstraints().GetPairIterations());
+		SET_DWORD_STAT(STAT_ChaosIterations_NumJointIterations, GetJointConstraints().GetSettings().ApplyPairIterations);
+		SET_DWORD_STAT(STAT_ChaosIterations_NumPushOutIterations, GetEvolution()->GetNumPushOutIterations());
+		SET_DWORD_STAT(STAT_ChaosIterations_NumPushOutCollisionIterations, GetEvolution()->GetCollisionConstraints().GetPushOutPairIterations());
+		SET_DWORD_STAT(STAT_ChaosIterations_NumPushOutJointIterations, GetJointConstraints().GetSettings().ApplyPushOutPairIterations);
+
+#if CSV_PROFILER
+		// Particle counts
+		CSV_CUSTOM_STAT(ChaosCounters, NumDisabledParticles, GetEvolution()->GetParticles().GetAllParticlesView().Num() - GetEvolution()->GetParticles().GetNonDisabledView().Num(), ECsvCustomStatOp::Accumulate);
+		CSV_CUSTOM_STAT(ChaosCounters, NumParticles, GetEvolution()->GetParticles().GetNonDisabledView().Num(), ECsvCustomStatOp::Accumulate);
+		CSV_CUSTOM_STAT(ChaosCounters, NumDynamicParticles, GetEvolution()->GetParticles().GetNonDisabledDynamicView().Num(), ECsvCustomStatOp::Accumulate);
+		CSV_CUSTOM_STAT(ChaosCounters, NumKinematicParticles, GetEvolution()->GetParticles().GetActiveKinematicParticlesView().Num(), ECsvCustomStatOp::Accumulate);
+		CSV_CUSTOM_STAT(ChaosCounters, NumStaticParticles, GetEvolution()->GetParticles().GetActiveStaticParticlesView().Num(), ECsvCustomStatOp::Accumulate);
+		CSV_CUSTOM_STAT(ChaosCounters, NumGeometryCollectionParticles, (int)GetEvolution()->GetParticles().GetGeometryCollectionParticles().Size(), ECsvCustomStatOp::Accumulate);
+
+		// Constraint counts
+		CSV_CUSTOM_STAT(ChaosCounters, NumIslands, GetEvolution()->GetConstraintGraph().NumIslands(), ECsvCustomStatOp::Accumulate);
+		CSV_CUSTOM_STAT(ChaosCounters, NumContacts, NumCollisionConstraints(), ECsvCustomStatOp::Accumulate);
+		CSV_CUSTOM_STAT(ChaosCounters, NumJoints, NumJointConstraints(), ECsvCustomStatOp::Accumulate);
+#endif
+	}
+
 	void FPBDRigidsSolver::PostEvolutionVDBPush() const
 	{
 #if CHAOS_VISUAL_DEBUGGER_ENABLED
diff --git a/Engine/Source/Runtime/Experimental/Chaos/Public/Chaos/PBDCollisionConstraints.h b/Engine/Source/Runtime/Experimental/Chaos/Public/Chaos/PBDCollisionConstraints.h
--- a/Engine/Source/Runtime/Experimental/Chaos/Public/Chaos/PBDCollisionConstraints.h
+++ b/Engine/Source/Runtime/Experimental/Chaos/Public/Chaos/PBDCollisionConstraints.h
@@ -278,11 +278,21 @@ public:
 		MApplyPairIterations = InPairIterations;
 	}
 
+	int32 GetPairIterations() const
+	{
+		return MApplyPairIterations;
+	}
+
 	void SetPushOutPairIterations(int32 InPairIterations)
 	{
 		MApplyPushOutPairIterations = InPairIterations;
 	}
 
+	int32 GetPushOutPairIterations() const
+	{
+		return MApplyPushOutPairIterations;
+	}
+
 	void SetCollisionsEnabled(bool bInEnableCollisions)
 	{
 		bEnableCollisions = bInEnableCollisions;
diff --git a/Engine/Source/Runtime/Experimental/Chaos/Public/Chaos/PBDRigidsEvolution.h b/Engine/Source/Runtime/Experimental/Chaos/Public/Chaos/PBDRigidsEvolution.h
--- a/Engine/Source/Runtime/Experimental/Chaos/Public/Chaos/PBDRigidsEvolution.h
+++ b/Engine/Source/Runtime/Experimental/Chaos/Public/Chaos/PBDRigidsEvolution.h
@@ -312,11 +312,21 @@ public:
 		NumIterations = InNumIterations;
 	}
 
+	CHAOS_API int32 GetNumIterations() const
+	{
+		return NumIterations;
+	}
+
 	CHAOS_API void SetNumPushOutIterations(int32 InNumIterations)
 	{
 		NumPushOutIterations = InNumIterations;
 	}
 
+	CHAOS_API int32 GetNumPushOutIterations() const
+	{
+		return NumPushOutIterations;
+	}
+
 	CHAOS_API void EnableParticle(FGeometryParticleHandle* Particle, const FGeometryParticleHandle* ParentParticle)
 	{
 		Particles.EnableParticle(Particle);
diff --git a/Engine/Source/Runtime/Experimental/Chaos/Public/ChaosStats.h b/Engine/Source/Runtime/Experimental/Chaos/Public/ChaosStats.h
--- a/Engine/Source/Runtime/Experimental/Chaos/Public/ChaosStats.h
+++ b/Engine/Source/Runtime/Experimental/Chaos/Public/ChaosStats.h
@@ -14,6 +14,7 @@ DECLARE_STATS_GROUP(TEXT("ChaosCollision"), STATGROUP_ChaosCollision, STATCAT_Ad
 DECLARE_STATS_GROUP(TEXT("ChaosJoint"), STATGROUP_ChaosJoint, STATCAT_Advanced);
 DECLARE_STATS_GROUP(TEXT("ChaosMinEvolution"), STATGROUP_ChaosMinEvolution, STATCAT_Advanced);
 DECLARE_STATS_GROUP(TEXT("ChaosCounters"), STATGROUP_ChaosCounters, STATCAT_Advanced);
+DECLARE_STATS_GROUP(TEXT("ChaosIterations"), STATGROUP_ChaosIterations, STATCAT_Advanced);
 
 
 DECLARE_CYCLE_STAT_EXTERN(TEXT("Physics Tick"), STAT_ChaosTick, STATGROUP_Chaos, CHAOS_API);
diff --git a/Engine/Source/Runtime/Experimental/Chaos/Public/PBDRigidsSolver.h b/Engine/Source/Runtime/Experimental/Chaos/Public/PBDRigidsSolver.h
--- a/Engine/Source/Runtime/Experimental/Chaos/Public/PBDRigidsSolver.h
+++ b/Engine/Source/Runtime/Experimental/Chaos/Public/PBDRigidsSolver.h
@@ -230,6 +230,7 @@ namespace Chaos
 
 		/**/
 		void PostTickDebugDraw(FReal Dt) const;
+		void UpdateCounters() const;
 
 		// Visual debugger (VDB) push methods
 		void PostEvolutionVDBPush() const;
