CIS - Fix non unity build erroring on Chaos cloth constraints.

diff --git a/Engine/Source/Runtime/Experimental/Chaos/Public/Chaos/PerParticlePBDSpringConstraints.h b/Engine/Source/Runtime/Experimental/Chaos/Public/Chaos/PerParticlePBDSpringConstraints.h
--- a/Engine/Source/Runtime/Experimental/Chaos/Public/Chaos/PerParticlePBDSpringConstraints.h
+++ b/Engine/Source/Runtime/Experimental/Chaos/Public/Chaos/PerParticlePBDSpringConstraints.h
@@ -11,27 +11,28 @@ namespace Chaos
 class FPerParticlePBDSpringConstraints : public FPerParticleRule, public FPBDSpringConstraintsBase
 {
 	typedef FPBDSpringConstraintsBase Base;
-	using Base::MConstraints;
+	using Base::Constraints;
+	using Base::Stiffness;
 
   public:
-	FPerParticlePBDSpringConstraints(const FDynamicParticles& InParticles, TArray<TVec2<int32>>&& Constraints, const FReal Stiffness = (FReal)1.)
-	    : Base(InParticles, MoveTemp(Constraints), Stiffness)
+	FPerParticlePBDSpringConstraints(const FPBDParticles& InParticles, const TArray<TVec2<int32>>& InConstraints, const FReal InStiffness = (FReal)1.)
+	    : Base(InParticles, 0, 0, InConstraints, TConstArrayView<FRealSingle>(), FVec2(InStiffness))
 	{
-		for (int32 i = 0; i < MConstraints.Num(); ++i)
+		for (int32 i = 0; i < Constraints.Num(); ++i)
 		{
-			const auto& Constraint = MConstraints[i];
+			const auto& Constraint = Constraints[i];
 			int32 i1 = Constraint[0];
 			int32 i2 = Constraint[1];
-			if (i1 >= MParticleToConstraints.Num())
+			if (i1 >= ParticleToConstraints.Num())
 			{
-				MParticleToConstraints.SetNum(i1 + 1);
+				ParticleToConstraints.SetNum(i1 + 1);
 			}
-			if (i2 >= MParticleToConstraints.Num())
+			if (i2 >= ParticleToConstraints.Num())
 			{
-				MParticleToConstraints.SetNum(i2 + 1);
+				ParticleToConstraints.SetNum(i2 + 1);
 			}
-			MParticleToConstraints[i1].Add(i);
-			MParticleToConstraints[i2].Add(i);
+			ParticleToConstraints[i1].Add(i);
+			ParticleToConstraints[i2].Add(i);
 		}
 	}
 	virtual ~FPerParticlePBDSpringConstraints() {}
@@ -39,25 +40,25 @@ class FPerParticlePBDSpringConstraints : public FPerParticleRule, public FPBDSpr
 	// TODO(mlentine): We likely need to use time n positions here
 	void Apply(FPBDParticles& InParticles, const FReal Dt, const int32 Index) const override //-V762
 	{
-		for (int32 i = 0; i < MParticleToConstraints[Index].Num(); ++i)
+		for (int32 i = 0; i < ParticleToConstraints[Index].Num(); ++i)
 		{
-			int32 CIndex = MParticleToConstraints[Index][i];
-			const auto& Constraint = MConstraints[CIndex];
+			int32 CIndex = ParticleToConstraints[Index][i];
+			const auto& Constraint = Constraints[CIndex];
 			int32 i1 = Constraint[0];
 			int32 i2 = Constraint[1];
 			if (Index == i1 && InParticles.InvM(i1) > 0)
 			{
-				InParticles.P(i1) -= InParticles.InvM(i1) * Base::GetDelta(InParticles, CIndex);
+				InParticles.P(i1) -= InParticles.InvM(i1) * Base::GetDelta(InParticles, CIndex, (FReal)Stiffness);
 			}
 			else if (InParticles.InvM(i2) > 0)
 			{
 				check(Index == i2);
-				InParticles.P(i2) += InParticles.InvM(i2) * Base::GetDelta(InParticles, CIndex);
+				InParticles.P(i2) += InParticles.InvM(i2) * Base::GetDelta(InParticles, CIndex, (FReal)Stiffness);
 			}
 		}
 	}
 
   private:
-	TArray<TArray<int32>> MParticleToConstraints;
+	TArray<TArray<int32>> ParticleToConstraints;
 };
 }
