CIS - Fix merge errors

diff --git a/Engine/Source/Runtime/Experimental/Chaos/Private/PBDRigidsSolver.cpp b/Engine/Source/Runtime/Experimental/Chaos/Private/PBDRigidsSolver.cpp
--- a/Engine/Source/Runtime/Experimental/Chaos/Private/PBDRigidsSolver.cpp
+++ b/Engine/Source/Runtime/Experimental/Chaos/Private/PBDRigidsSolver.cpp
@@ -325,7 +325,7 @@ namespace Chaos
 #endif
 			}
 
-			MSolver->UpdateCounters();
+			{
 				SCOPE_CYCLE_COUNTER(STAT_EventDataGathering);
 				{
 					SCOPE_CYCLE_COUNTER(STAT_FillProducerData);
@@ -1250,6 +1250,14 @@ namespace Chaos
 		SET_DWORD_STAT(STAT_ChaosCounter_NumContacts, NumCollisionConstraints());
 		SET_DWORD_STAT(STAT_ChaosCounter_NumJoints, NumJointConstraints());
 
+		// Iterations
+		SET_DWORD_STAT(STAT_ChaosIterations_NumIterations, GetEvolution()->GetNumIterations());
+		SET_DWORD_STAT(STAT_ChaosIterations_NumCollisionIterations, GetEvolution()->GetCollisionConstraints().GetPairIterations());
+		SET_DWORD_STAT(STAT_ChaosIterations_NumJointIterations, GetJointConstraints().GetSettings().ApplyPairIterations);
+		SET_DWORD_STAT(STAT_ChaosIterations_NumPushOutIterations, GetEvolution()->GetNumPushOutIterations());
+		SET_DWORD_STAT(STAT_ChaosIterations_NumPushOutCollisionIterations, GetEvolution()->GetCollisionConstraints().GetPushOutPairIterations());
+		SET_DWORD_STAT(STAT_ChaosIterations_NumPushOutJointIterations, GetJointConstraints().GetSettings().ApplyPushOutPairIterations);
+
 #if CSV_PROFILER
 		// Particle counts
 		CSV_CUSTOM_STAT(ChaosCounters, NumDisabledParticles, GetEvolution()->GetParticles().GetAllParticlesView().Num() - GetEvolution()->GetParticles().GetNonDisabledView().Num(), ECsvCustomStatOp::Accumulate);
@@ -1316,46 +1324,6 @@ namespace Chaos
 #endif
 	}
 
-	void FPBDRigidsSolver::UpdateCounters() const
-	{
-		// Particle counts
-		SET_DWORD_STAT(STAT_ChaosCounter_NumDisabledParticles, GetEvolution()->GetParticles().GetAllParticlesView().Num() - GetEvolution()->GetParticles().GetNonDisabledView().Num());
-		SET_DWORD_STAT(STAT_ChaosCounter_NumParticles, GetEvolution()->GetParticles().GetNonDisabledView().Num());
-		SET_DWORD_STAT(STAT_ChaosCounter_NumDynamicParticles, GetEvolution()->GetParticles().GetNonDisabledDynamicView().Num());
-		SET_DWORD_STAT(STAT_ChaosCounter_NumActiveDynamicParticles, GetEvolution()->GetParticles().GetActiveParticlesView().Num());
-		SET_DWORD_STAT(STAT_ChaosCounter_NumKinematicParticles, GetEvolution()->GetParticles().GetActiveKinematicParticlesView().Num());
-		SET_DWORD_STAT(STAT_ChaosCounter_NumStaticParticles, GetEvolution()->GetParticles().GetActiveStaticParticlesView().Num());
-		SET_DWORD_STAT(STAT_ChaosCounter_NumGeometryCollectionParticles, GetEvolution()->GetParticles().GetGeometryCollectionParticles().Size());
-
-		// Constraint counts
-		SET_DWORD_STAT(STAT_ChaosCounter_NumIslands, GetEvolution()->GetConstraintGraph().NumIslands());
-		SET_DWORD_STAT(STAT_ChaosCounter_NumContacts, NumCollisionConstraints());
-		SET_DWORD_STAT(STAT_ChaosCounter_NumJoints, NumJointConstraints());
-
-		// Iterations
-		SET_DWORD_STAT(STAT_ChaosIterations_NumIterations, GetEvolution()->GetNumIterations());
-		SET_DWORD_STAT(STAT_ChaosIterations_NumCollisionIterations, GetEvolution()->GetCollisionConstraints().GetPairIterations());
-		SET_DWORD_STAT(STAT_ChaosIterations_NumJointIterations, GetJointConstraints().GetSettings().ApplyPairIterations);
-		SET_DWORD_STAT(STAT_ChaosIterations_NumPushOutIterations, GetEvolution()->GetNumPushOutIterations());
-		SET_DWORD_STAT(STAT_ChaosIterations_NumPushOutCollisionIterations, GetEvolution()->GetCollisionConstraints().GetPushOutPairIterations());
-		SET_DWORD_STAT(STAT_ChaosIterations_NumPushOutJointIterations, GetJointConstraints().GetSettings().ApplyPushOutPairIterations);
-
-#if CSV_PROFILER
-		// Particle counts
-		CSV_CUSTOM_STAT(ChaosCounters, NumDisabledParticles, GetEvolution()->GetParticles().GetAllParticlesView().Num() - GetEvolution()->GetParticles().GetNonDisabledView().Num(), ECsvCustomStatOp::Accumulate);
-		CSV_CUSTOM_STAT(ChaosCounters, NumParticles, GetEvolution()->GetParticles().GetNonDisabledView().Num(), ECsvCustomStatOp::Accumulate);
-		CSV_CUSTOM_STAT(ChaosCounters, NumDynamicParticles, GetEvolution()->GetParticles().GetNonDisabledDynamicView().Num(), ECsvCustomStatOp::Accumulate);
-		CSV_CUSTOM_STAT(ChaosCounters, NumKinematicParticles, GetEvolution()->GetParticles().GetActiveKinematicParticlesView().Num(), ECsvCustomStatOp::Accumulate);
-		CSV_CUSTOM_STAT(ChaosCounters, NumStaticParticles, GetEvolution()->GetParticles().GetActiveStaticParticlesView().Num(), ECsvCustomStatOp::Accumulate);
-		CSV_CUSTOM_STAT(ChaosCounters, NumGeometryCollectionParticles, (int)GetEvolution()->GetParticles().GetGeometryCollectionParticles().Size(), ECsvCustomStatOp::Accumulate);
-
-		// Constraint counts
-		CSV_CUSTOM_STAT(ChaosCounters, NumIslands, GetEvolution()->GetConstraintGraph().NumIslands(), ECsvCustomStatOp::Accumulate);
-		CSV_CUSTOM_STAT(ChaosCounters, NumContacts, NumCollisionConstraints(), ECsvCustomStatOp::Accumulate);
-		CSV_CUSTOM_STAT(ChaosCounters, NumJoints, NumJointConstraints(), ECsvCustomStatOp::Accumulate);
-#endif
-	}
-
 	void FPBDRigidsSolver::PostEvolutionVDBPush() const
 	{
 #if CHAOS_VISUAL_DEBUGGER_ENABLED
diff --git a/Engine/Source/Runtime/Experimental/Chaos/Public/PBDRigidsSolver.h b/Engine/Source/Runtime/Experimental/Chaos/Public/PBDRigidsSolver.h
--- a/Engine/Source/Runtime/Experimental/Chaos/Public/PBDRigidsSolver.h
+++ b/Engine/Source/Runtime/Experimental/Chaos/Public/PBDRigidsSolver.h
@@ -230,7 +230,6 @@ namespace Chaos
 
 		/**/
 		void PostTickDebugDraw(FReal Dt) const;
-		void UpdateCounters() const;
 
 		// Visual debugger (VDB) push methods
 		void PostEvolutionVDBPush() const;
