Chaos - box structuredata static initialization fixes

diff --git a/Engine/Source/Runtime/Experimental/Chaos/Private/Chaos/Box.cpp b/Engine/Source/Runtime/Experimental/Chaos/Private/Chaos/Box.cpp
--- a/Engine/Source/Runtime/Experimental/Chaos/Private/Chaos/Box.cpp
+++ b/Engine/Source/Runtime/Experimental/Chaos/Private/Chaos/Box.cpp
@@ -5,12 +5,7 @@
 
 namespace Chaos
 {
-	template<typename T, int D> TBox<T, D>::FStaticInitializer::FStaticInitializer()
-	{
-		TBox<FReal, 3>::InitializeStructureData();
-	}
-
-	template<typename T, int D> void TBox<T, D>::InitializeStructureData()
+	FConvexHalfEdgeStructureDataS16 MakeBoxStructureData()
 	{
 		TArray<TArray<int32>> PlaneVertices
 		{
@@ -22,7 +17,7 @@ namespace Chaos
 			{0, 1, 2, 3},	// -Z,
 		};
 
-		SStructureData.SetPlaneVertices(PlaneVertices, 8);
+		return FConvexHalfEdgeStructureDataS16::MakePlaneVertices(PlaneVertices, 8);
 	}
 
 	template<typename T, int D> TArray<FVec3> TBox<T, D>::SNormals =
@@ -47,9 +42,7 @@ namespace Chaos
 		FVec3(1,-1,1),
 	};
 
-	template<typename T, int D> FConvexHalfEdgeStructureDataS16 TBox<T, D>::SStructureData;
-
-	template<typename T, int D> typename TBox<T, D>::FStaticInitializer TBox<T, D>::SInitializer;
+	template<typename T, int D> FConvexHalfEdgeStructureDataS16 TBox<T, D>::SStructureData = MakeBoxStructureData();
 
 	template class TBox<FReal, 3>;
 }
\ No newline at end of file
diff --git a/Engine/Source/Runtime/Experimental/Chaos/Public/Chaos/Box.h b/Engine/Source/Runtime/Experimental/Chaos/Public/Chaos/Box.h
--- a/Engine/Source/Runtime/Experimental/Chaos/Public/Chaos/Box.h
+++ b/Engine/Source/Runtime/Experimental/Chaos/Public/Chaos/Box.h
@@ -481,8 +481,6 @@ namespace Chaos
 			return AABB.GetTypeHash();
 		}
 
-		static void InitializeStructureData();
-
 	private:
 		TAABB<T, d> AABB;
 
@@ -490,12 +488,6 @@ namespace Chaos
 		static TArray<FVec3> SNormals;
 		static TArray<FVec3> SVertices;
 		static FConvexHalfEdgeStructureDataS16 SStructureData;
-
-		struct FStaticInitializer
-		{
-			FStaticInitializer();
-		};
-		static FStaticInitializer SInitializer;
 	};
 
 
