Chaos Cloth - Add ISPC weight maps implementation for material constraints.

diff --git a/Engine/Source/Runtime/Experimental/Chaos/Private/Chaos/PBDAxialSpringConstraints.cpp b/Engine/Source/Runtime/Experimental/Chaos/Private/Chaos/PBDAxialSpringConstraints.cpp
--- a/Engine/Source/Runtime/Experimental/Chaos/Private/Chaos/PBDAxialSpringConstraints.cpp
+++ b/Engine/Source/Runtime/Experimental/Chaos/Private/Chaos/PBDAxialSpringConstraints.cpp
@@ -91,17 +91,37 @@ void FPBDAxialSpringConstraints::Apply(FPBDParticles& Particles, const FReal Dt)
 				}
 			}
 		}
-		else
+		else  // Has weight maps
 		{
-			// TODO: ISPC  weight map implementation
-			for (const TArray<int32>& ConstraintBatch : ConstraintsPerColor)
+#if INTEL_ISPC
+			if (bRealTypeCompatibleWithISPC && bChaos_AxialSpring_ISPC_Enabled)
+			{
+				for (const TArray<int32>& ConstraintBatch : ConstraintsPerColor)
+				{
+					ispc::ApplyAxialSpringConstraintsWithWeightMaps(
+						(ispc::FVector*)&Particles.GetP()[0],
+						(ispc::FIntVector*)&Constraints.GetData()[0],
+						&ConstraintBatch.GetData()[0],
+						&Particles.GetInvM().GetData()[0],
+						&Barys.GetData()[0],
+						&Dists.GetData()[0],
+						&Stiffness.GetIndices().GetData()[0],
+						&Stiffness.GetTable().GetData()[0],
+						ConstraintBatch.Num());
+				}
+			}
+			else
+#endif
 			{
-				PhysicsParallelFor(ConstraintBatch.Num(), [&](const int32 Index)
+				for (const TArray<int32>& ConstraintBatch : ConstraintsPerColor)
 				{
-					const int32 ConstraintIndex = ConstraintBatch[Index];
-					const FReal ExpStiffnessValue = Stiffness[ConstraintIndex];
-					ApplyHelper(Particles, Dt, ConstraintIndex, ExpStiffnessValue);
-				});
+					PhysicsParallelFor(ConstraintBatch.Num(), [&](const int32 Index)
+					{
+						const int32 ConstraintIndex = ConstraintBatch[Index];
+						const FReal ExpStiffnessValue = Stiffness[ConstraintIndex];
+						ApplyHelper(Particles, Dt, ConstraintIndex, ExpStiffnessValue);
+					});
+				}
 			}
 		}
 	}
diff --git a/Engine/Source/Runtime/Experimental/Chaos/Private/Chaos/PBDAxialSpringConstraints.ispc b/Engine/Source/Runtime/Experimental/Chaos/Private/Chaos/PBDAxialSpringConstraints.ispc
--- a/Engine/Source/Runtime/Experimental/Chaos/Private/Chaos/PBDAxialSpringConstraints.ispc
+++ b/Engine/Source/Runtime/Experimental/Chaos/Private/Chaos/PBDAxialSpringConstraints.ispc
@@ -92,3 +92,98 @@ export void ApplyAxialSpringConstraints(uniform FVector P[],
 		}
 	}
 }
+
+export void ApplyAxialSpringConstraintsWithWeightMaps(uniform FVector P[],
+										const uniform FIntVector AllConstraints[],
+										const uniform int32 ActiveConstraints[],
+										const uniform float InvM[],
+										const uniform float Barys[],
+										const uniform float Dists[],
+										const uniform uint8 StiffnessIndices[],
+										const uniform float StiffnessTable[],
+										const uniform int32 NumConstraints)
+{
+	varying int32 ActiveConstraint;
+	varying FIntVector Constraint;
+
+#if HW_GATHER_SUPPORTED == 1
+	if(programIndex < NumConstraints)
+	{
+		ActiveConstraint = ActiveConstraints[programIndex];
+		#pragma ignore warning(perf)
+		Constraint = AllConstraints[ActiveConstraint];
+	}
+#endif
+
+	foreach(i = 0 ... NumConstraints)
+	{
+#if HW_GATHER_SUPPORTED == 0
+		ActiveConstraint = ActiveConstraints[i];
+		#pragma ignore warning(perf)
+		Constraint = AllConstraints[ActiveConstraint];
+#endif
+
+		#pragma ignore warning(perf)
+		const varying float Bary = Barys[ActiveConstraint];
+		#pragma ignore warning(perf)
+		const varying float Dist = Dists[ActiveConstraint];
+		#pragma ignore warning(perf)
+		const varying uint8 StiffnessIndex = StiffnessIndices[ActiveConstraint];
+
+		const varying int32 i1 = Constraint.V[0];
+		const varying int32 i2 = Constraint.V[1];
+		const varying int32 i3 = Constraint.V[2];
+
+		const varying FVector P2 = VectorGather(&P[i2]);
+		const varying FVector P3 = VectorGather(&P[i3]);
+
+#if HW_GATHER_SUPPORTED == 1
+		if(i + programCount < NumConstraints)
+		{
+			ActiveConstraint = ActiveConstraints[i+programCount];
+			#pragma ignore warning(perf)
+			Constraint = AllConstraints[ActiveConstraint];
+		}
+#endif
+
+		const varying FVector P1 = VectorGather(&P[i1]);
+
+		const FVector P0 = (P2 - P3) * Bary + P3;
+		const FVector Difference = P1 - P0;
+
+		const float Distance = VectorSize(Difference);
+		const float Multiplier = 2 / (max(Bary, 1 - Bary) + 1);
+
+		#pragma ignore warning(perf)
+		const varying float M1 = InvM[i1];
+		#pragma ignore warning(perf)
+		const varying float M2 = InvM[i2];
+		#pragma ignore warning(perf)
+		const varying float M3 = InvM[i3];
+
+		#pragma ignore warning(perf)
+		const varying float Stiffness = StiffnessTable[StiffnessIndex];
+
+		const FVector Direction = Difference / Distance;
+		const FVector InnerDelta = (Distance - Dist) * Direction;
+
+		const float PInvMass = M3 * (1 - Bary) + M2 * Bary;
+		const float CombinedInvMass = PInvMass + M1;
+		const FVector ComputedDelta = Stiffness * InnerDelta / CombinedInvMass;
+
+		const FVector Delta = VectorSelect(((M1 == 0 && PInvMass == 0) || Distance <= SMALL_NUMBER), ZeroVector, ComputedDelta);
+
+		if (M1 > 0)
+		{
+			VectorScatter(&P[i1], P1 - (Multiplier * M1 * Delta));
+		}
+		if (M2)
+		{
+			VectorScatter(&P[i2], P2 + (Multiplier * M2 * Bary * Delta));
+		}
+		if (M3)
+		{
+			VectorScatter(&P[i3], P3 + (Multiplier * M3 * (1 - Bary) * Delta));
+		}
+	}
+}
diff --git a/Engine/Source/Runtime/Experimental/Chaos/Private/Chaos/PBDSpringConstraints.cpp b/Engine/Source/Runtime/Experimental/Chaos/Private/Chaos/PBDSpringConstraints.cpp
--- a/Engine/Source/Runtime/Experimental/Chaos/Private/Chaos/PBDSpringConstraints.cpp
+++ b/Engine/Source/Runtime/Experimental/Chaos/Private/Chaos/PBDSpringConstraints.cpp
@@ -91,17 +91,36 @@ void FPBDSpringConstraints::Apply(FPBDParticles& Particles, const FReal Dt) cons
 				}
 			}
 		}
-		else
+		else  // Has weight maps
 		{
-			// TODO: ISPC version
-			for (const TArray<int32>& ConstraintBatch : ConstraintsPerColor)
+#if INTEL_ISPC
+			if (bRealTypeCompatibleWithISPC && bChaos_Spring_ISPC_Enabled)
 			{
-				PhysicsParallelFor(ConstraintBatch.Num(), [&](const int32 Index)
+				for (const TArray<int32>& ConstraintBatch : ConstraintsPerColor)
 				{
-					const int32 ConstraintIndex = ConstraintBatch[Index];
-					const FReal ExpStiffnessValue = Stiffness[ConstraintIndex];
-					ApplyHelper(Particles, Dt, ConstraintIndex, ExpStiffnessValue);
-				});
+					ispc::ApplySpringConstraintsWithWeightMaps(
+						(ispc::FVector*) & Particles.GetP()[0],
+						(ispc::FIntVector2*) & Constraints.GetData()[0],
+						&ConstraintBatch.GetData()[0],
+						&Particles.GetInvM().GetData()[0],
+						&Dists.GetData()[0],
+						&Stiffness.GetIndices().GetData()[0],
+						&Stiffness.GetTable().GetData()[0],
+						ConstraintBatch.Num());
+				}
+			}
+			else
+#endif
+			{
+				for (const TArray<int32>& ConstraintBatch : ConstraintsPerColor)
+				{
+					PhysicsParallelFor(ConstraintBatch.Num(), [&](const int32 Index)
+					{
+						const int32 ConstraintIndex = ConstraintBatch[Index];
+						const FReal ExpStiffnessValue = Stiffness[ConstraintIndex];
+						ApplyHelper(Particles, Dt, ConstraintIndex, ExpStiffnessValue);
+					});
+				}
 			}
 		}
 	}
diff --git a/Engine/Source/Runtime/Experimental/Chaos/Private/Chaos/PBDSpringConstraints.ispc b/Engine/Source/Runtime/Experimental/Chaos/Private/Chaos/PBDSpringConstraints.ispc
--- a/Engine/Source/Runtime/Experimental/Chaos/Private/Chaos/PBDSpringConstraints.ispc
+++ b/Engine/Source/Runtime/Experimental/Chaos/Private/Chaos/PBDSpringConstraints.ispc
@@ -82,3 +82,82 @@ export void ApplySpringConstraints(uniform FVector P[],
 		}
 	}
 }
+
+export void ApplySpringConstraintsWithWeightMaps(uniform FVector P[],
+									const uniform FIntVector2 AllConstraints[],
+									const uniform int32 ActiveConstraints[],
+									const uniform float InvM[],
+									const uniform float Dists[],
+									const uniform uint8 StiffnessIndices[],
+									const uniform float StiffnessTable[],
+									const uniform int32 NumConstraints)
+{
+	varying int32 ActiveConstraint;
+	varying FIntVector2 Constraint;
+
+#if HW_GATHER_SUPPORTED == 1
+	if(programIndex < NumConstraints)
+	{
+		ActiveConstraint = ActiveConstraints[programIndex];
+		#pragma ignore warning(perf)
+		Constraint = AllConstraints[ActiveConstraint];
+	}
+#endif
+
+	foreach(i = 0 ... NumConstraints)
+	{
+#if HW_GATHER_SUPPORTED == 0
+		ActiveConstraint = ActiveConstraints[i];
+		#pragma ignore warning(perf)
+		Constraint = AllConstraints[ActiveConstraint];
+#endif
+
+		const varying int32 i1 = Constraint.V[0];
+		const varying int32 i2 = Constraint.V[1];
+
+		const varying FVector P1 = VectorGather(&P[i1]);
+		const varying FVector P2 = VectorGather(&P[i2]);
+
+		FVector Direction = P1 - P2;
+		const float Distance = SafeNormalize(Direction);
+
+		#pragma ignore warning(perf)
+		const varying float Dist = Dists[ActiveConstraint];
+
+		#pragma ignore warning(perf)
+		const varying uint8 StiffnessIndex = StiffnessIndices[ActiveConstraint];
+
+#if HW_GATHER_SUPPORTED == 1
+		if(i + programCount < NumConstraints)
+		{
+			ActiveConstraint = ActiveConstraints[i+programCount];
+			#pragma ignore warning(perf)
+			Constraint = AllConstraints[ActiveConstraint];
+		}
+#endif
+
+		const FVector InnerDelta = (Distance - Dist) * Direction;
+
+		#pragma ignore warning(perf)
+		const varying float InvM1 = InvM[i1];
+		#pragma ignore warning(perf)
+		const varying float InvM2 = InvM[i2];
+
+		#pragma ignore warning(perf)
+		const varying float Stiffness = StiffnessTable[StiffnessIndex];
+
+		const float CombinedInvMass = InvM2 + InvM1;
+		const FVector ComputedDelta = Stiffness * InnerDelta / CombinedInvMass;
+
+		const FVector Delta = VectorSelect((InvM2 == 0 && InvM1 == 0), ZeroVector, ComputedDelta);
+
+		if (InvM1 > 0)
+		{
+			VectorScatter(&P[i1], P1 - (InvM1 * Delta));
+		}
+		if (InvM2 > 0)
+		{
+			VectorScatter(&P[i2], P2 + (InvM2 * Delta));
+		}
+	}
+}
diff --git a/Engine/Source/Runtime/Experimental/Chaos/Public/Chaos/PBDAxialSpringConstraintsBase.h b/Engine/Source/Runtime/Experimental/Chaos/Public/Chaos/PBDAxialSpringConstraintsBase.h
--- a/Engine/Source/Runtime/Experimental/Chaos/Public/Chaos/PBDAxialSpringConstraintsBase.h
+++ b/Engine/Source/Runtime/Experimental/Chaos/Public/Chaos/PBDAxialSpringConstraintsBase.h
@@ -94,7 +94,7 @@ private:
 						Index1 <= Index2 ? TVec3<int32>(Index0, Index1, Index2) :
 						Index0 <= Index2 ? TVec3<int32>(Index0, Index2, Index1) :
 										   TVec3<int32>(Index2, Index0, Index1) :
-					// Index1 <= Index0
+					// Index1 < Index0
 						Index0 <= Index2 ? TVec3<int32>(Index1, Index0, Index2) :
 						Index1 <= Index2 ? TVec3<int32>(Index1, Index2, Index0) :
 										   TVec3<int32>(Index2, Index1, Index0));
diff --git a/Engine/Source/Runtime/Experimental/Chaos/Public/Chaos/PBDStiffness.h b/Engine/Source/Runtime/Experimental/Chaos/Public/Chaos/PBDStiffness.h
--- a/Engine/Source/Runtime/Experimental/Chaos/Public/Chaos/PBDStiffness.h
+++ b/Engine/Source/Runtime/Experimental/Chaos/Public/Chaos/PBDStiffness.h
@@ -81,6 +81,12 @@ public:
 	/** Return the exponential stiffness value when the weight map is not used. */
 	explicit operator FReal() const { return GetLow(); }
 
+	/** Return the particles/constraints indices to the stiffness table as a read only array. */
+	TConstArrayView<uint8> GetIndices() const { return TConstArrayView<uint8>(Indices); }
+
+	/** Return the table of stiffnesses as a read only array. */
+	TConstArrayView<FReal> GetTable() const { return TConstArrayView<FReal>(Table); }
+
 private:
 	static constexpr FReal ParameterFrequency = (FReal)120.;  // 60Hz @ 2 iterations as a root for all stiffness values TODO: Make this a global solver parameter
 
