Chaos - fix implicit casting from int to FIndex in structuredata

diff --git a/Engine/Source/Runtime/Experimental/Chaos/Public/Chaos/ConvexHalfEdgeStructureData.h b/Engine/Source/Runtime/Experimental/Chaos/Public/Chaos/ConvexHalfEdgeStructureData.h
--- a/Engine/Source/Runtime/Experimental/Chaos/Public/Chaos/ConvexHalfEdgeStructureData.h
+++ b/Engine/Source/Runtime/Experimental/Chaos/Public/Chaos/ConvexHalfEdgeStructureData.h
@@ -359,7 +359,7 @@ namespace Chaos
 					// If this is the first time Vertex0 has showed up, set its edge index
 					if (Vertices[VertexIndex0].FirstHalfEdgeIndex == InvalidIndex)
 					{
-						Vertices[VertexIndex0].FirstHalfEdgeIndex = NextHalfEdgeIndex;
+						Vertices[VertexIndex0].FirstHalfEdgeIndex = (FIndex)NextHalfEdgeIndex;
 					}
 
 					++NextHalfEdgeIndex;
@@ -471,7 +471,7 @@ namespace Chaos
 					{
 						if (Edge.VertexIndex < TwinEdge.VertexIndex)
 						{
-							Edges.Add(HalfEdgeIndex);
+							Edges.Add((FIndex)HalfEdgeIndex);
 						}
 					}
 				}
