Remove previous transform from kinematic target
- no more previous transform from kinematic targets
- ApplyKinematicTarget now get the current XR and interpolate using StepFraction ( which is now correct )
- Revert the proxy from setting the previous transform and just copy the GT kinematic target to the PT

diff --git a/Engine/Source/Runtime/Engine/Private/PhysicsEngine/ImmediatePhysics/ImmediatePhysicsChaos/ImmediatePhysicsActorHandle_Chaos.cpp b/Engine/Source/Runtime/Engine/Private/PhysicsEngine/ImmediatePhysics/ImmediatePhysicsChaos/ImmediatePhysicsActorHandle_Chaos.cpp
--- a/Engine/Source/Runtime/Engine/Private/PhysicsEngine/ImmediatePhysics/ImmediatePhysicsChaos/ImmediatePhysicsActorHandle_Chaos.cpp
+++ b/Engine/Source/Runtime/Engine/Private/PhysicsEngine/ImmediatePhysics/ImmediatePhysicsChaos/ImmediatePhysicsActorHandle_Chaos.cpp
@@ -387,10 +387,9 @@ namespace ImmediatePhysics_Chaos
 		if (ensure(GetIsKinematic()))
 		{
 			FGenericParticleHandle GenericHandle(ParticleHandle);
-			FTransform PreviousTransform(GenericHandle->R(), GenericHandle->X());
 			FTransform ParticleTransform = FParticleUtilities::ActorWorldToParticleWorld(GenericHandle, WorldTM);
 
-			GetKinematicTarget().SetTargetMode(ParticleTransform, PreviousTransform);
+			GetKinematicTarget().SetTargetMode(ParticleTransform);
 		}
 
 	}
diff --git a/Engine/Source/Runtime/Experimental/Chaos/Private/Chaos/Evolution/PBDMinEvolution.cpp b/Engine/Source/Runtime/Experimental/Chaos/Private/Chaos/Evolution/PBDMinEvolution.cpp
--- a/Engine/Source/Runtime/Experimental/Chaos/Private/Chaos/Evolution/PBDMinEvolution.cpp
+++ b/Engine/Source/Runtime/Experimental/Chaos/Private/Chaos/Evolution/PBDMinEvolution.cpp
@@ -564,11 +564,8 @@ namespace Chaos
 		for (auto& Particle : Particles.GetActiveKinematicParticlesView())
 		{
 			TKinematicTarget<FReal, 3>& KinematicTarget = Particle.KinematicTarget();
-
-			const TRigidTransform<FReal, 3>& Previous = KinematicTarget.GetPrevious();
-			const FVec3 PrevX = Previous.GetTranslation();
-			const FRotation3 PrevR = Previous.GetRotation();
-
+			const FVec3 CurrentX = Particle.X();
+			const FRotation3 CurrentR = Particle.R();
 
 			switch (KinematicTarget.GetMode())
 			{
@@ -589,26 +586,28 @@ namespace Chaos
 			{
 				// Move to kinematic target and update velocities to match
 				// Target positions only need to be processed once, and we reset the velocity next frame (if no new target is set)
-				FVec3 TargetPos;
-				FRotation3 TargetRot;
+				FVec3 NewX;
+				FRotation3 NewR;
 				if (FMath::IsNearlyEqual(StepFraction, (FReal)1, (FReal)KINDA_SMALL_NUMBER))
 				{
-					TargetPos = KinematicTarget.GetTarget().GetLocation();
-					TargetRot = KinematicTarget.GetTarget().GetRotation();
+					NewX = KinematicTarget.GetTarget().GetLocation();
+					NewR = KinematicTarget.GetTarget().GetRotation();
 					KinematicTarget.SetMode(EKinematicTargetMode::Reset);
 				}
 				else
 				{
-					TargetPos = FVec3::Lerp(PrevX, KinematicTarget.GetTarget().GetLocation(), StepFraction);
-					TargetRot = FRotation3::Slerp(PrevR, KinematicTarget.GetTarget().GetRotation(), StepFraction);
+					// as a reminder, stepfraction is the remaing fraction of the step from the remaining steps
+					// for total of 4 steps and current step of 2, this will be 1/3 ( 1 step passed, 3 steps remains )
+					NewX = FVec3::Lerp(CurrentX, KinematicTarget.GetTarget().GetLocation(), StepFraction);
+					NewR = FRotation3::Slerp(CurrentR, KinematicTarget.GetTarget().GetRotation(), StepFraction);
 				}
 				if (Dt > MinDt)
 				{
-					Particle.V() = FVec3::CalculateVelocity(PrevX, TargetPos, Dt);
-					Particle.W() = FRotation3::CalculateAngularVelocity(PrevR, TargetRot, Dt);
+					Particle.V() = FVec3::CalculateVelocity(CurrentX, NewX, Dt);
+					Particle.W() = FRotation3::CalculateAngularVelocity(CurrentR, NewR, Dt);
 				}
-				Particle.X() = TargetPos;
-				Particle.R() = TargetRot;
+				Particle.X() = NewX;
+				Particle.R() = NewR;
 				break;
 			}
 
diff --git a/Engine/Source/Runtime/Experimental/Chaos/Private/PhysicsProxy/SingleParticlePhysicsProxy.cpp b/Engine/Source/Runtime/Experimental/Chaos/Private/PhysicsProxy/SingleParticlePhysicsProxy.cpp
--- a/Engine/Source/Runtime/Experimental/Chaos/Private/PhysicsProxy/SingleParticlePhysicsProxy.cpp
+++ b/Engine/Source/Runtime/Experimental/Chaos/Private/PhysicsProxy/SingleParticlePhysicsProxy.cpp
@@ -73,11 +73,7 @@ void PushToPhysicsStateImp(const Chaos::FDirtyPropertiesManager& Manager, Chaos:
 		auto NewKinematicTargetGT = bHasKinematicData ? ParticleData.FindKinematicTarget(Manager, DataIdx) : nullptr;
 		if (NewKinematicTargetGT)
 		{
-			const TRigidTransform<float, 3>& TargetTransform = NewKinematicTargetGT->GetTarget();
-			const TRigidTransform<float, 3> PreviousTransform(KinematicHandle->X(), KinematicHandle->R());
-			TKinematicTarget<float, 3> NewKinematicTargetPT;
-			NewKinematicTargetPT.SetTargetMode(TargetTransform, PreviousTransform);
-			KinematicHandle->SetKinematicTarget(NewKinematicTargetPT);
+			KinematicHandle->SetKinematicTarget(*NewKinematicTargetGT);
 		}
 
 		if(NewXR || NewNonFrequentData || NewVelocities || NewKinematicTargetGT)
diff --git a/Engine/Source/Runtime/Experimental/Chaos/Public/Chaos/KinematicTargets.h b/Engine/Source/Runtime/Experimental/Chaos/Public/Chaos/KinematicTargets.h
--- a/Engine/Source/Runtime/Experimental/Chaos/Public/Chaos/KinematicTargets.h
+++ b/Engine/Source/Runtime/Experimental/Chaos/Public/Chaos/KinematicTargets.h
@@ -38,14 +38,11 @@ namespace Chaos
 		/** Get the target transform (asserts if not in Position mode) */
 		const TRigidTransform<T, d>& GetTarget() const { check(Mode == EKinematicTargetMode::Position); return Target; }
 
-		/** Get the particle's previous transform for velocity calculations) */
-		const TRigidTransform<T, d>& GetPrevious() const { return Previous; }
-
 		/** Clear the kinematic target */
 		void Clear() { Target = TRigidTransform<T, d>(); Mode = EKinematicTargetMode::None; }
 
 		/** Use transform target mode and set the transform target */
-		void SetTargetMode(const TRigidTransform<T, d>& InTarget, const TRigidTransform<T, d>& InPrevious) { Target = InTarget; Previous = InPrevious;  Mode = EKinematicTargetMode::Position; }
+		void SetTargetMode(const TRigidTransform<T, d>& InTarget) { Target = InTarget; Mode = EKinematicTargetMode::Position; }
 
 		/** Use velocity target mode */
 		void SetVelocityMode() { Mode = EKinematicTargetMode::Velocity; }
@@ -65,10 +62,7 @@ namespace Chaos
 				Mode == other.Mode &&
 				Target.GetTranslation() == other.Target.GetTranslation() &&
 				Target.GetRotation() == other.Target.GetRotation() &&
-				Target.GetScale3D() == other.Target.GetScale3D() &&
-				Previous.GetTranslation() == other.Previous.GetTranslation() &&
-				Previous.GetRotation() == other.Previous.GetRotation() &&
-				Previous.GetScale3D() == other.Previous.GetScale3D()
+				Target.GetScale3D() == other.Target.GetScale3D()
 				);
 		}
 
@@ -84,7 +78,6 @@ namespace Chaos
 		}
 
 	private:
-		TRigidTransform<T, d> Previous;
 		TRigidTransform<T, d> Target;
 		EKinematicTargetMode Mode;
 	};
diff --git a/Engine/Source/Runtime/Experimental/Chaos/Public/Chaos/PBDRigidsEvolution.h b/Engine/Source/Runtime/Experimental/Chaos/Public/Chaos/PBDRigidsEvolution.h
--- a/Engine/Source/Runtime/Experimental/Chaos/Public/Chaos/PBDRigidsEvolution.h
+++ b/Engine/Source/Runtime/Experimental/Chaos/Public/Chaos/PBDRigidsEvolution.h
@@ -652,9 +652,8 @@ public:
 		for (auto& Particle : Particles.GetActiveKinematicParticlesView())
 		{
 			TKinematicTarget<FReal, 3>& KinematicTarget = Particle.KinematicTarget();
-			const TRigidTransform<FReal, 3>& Previous = KinematicTarget.GetPrevious();
-			const FVec3 PrevX = Previous.GetTranslation();
-			const FRotation3 PrevR = Previous.GetRotation();
+			const FVec3 CurrentX = Particle.X();
+			const FRotation3 CurrentR = Particle.R();
 
 			switch (KinematicTarget.GetMode())
 			{
@@ -676,29 +675,31 @@ public:
 			{
 				// Move to kinematic target and update velocities to match
 				// Target positions only need to be processed once, and we reset the velocity next frame (if no new target is set)
-				FVec3 TargetPos;
-				FRotation3 TargetRot;
+				FVec3 NewX;
+				FRotation3 NewR;
 				if (FMath::IsNearlyEqual(StepFraction, (FReal)1, (FReal)KINDA_SMALL_NUMBER))
 				{
-					TargetPos = KinematicTarget.GetTarget().GetLocation();
-					TargetRot = KinematicTarget.GetTarget().GetRotation();
+					NewX = KinematicTarget.GetTarget().GetLocation();
+					NewR = KinematicTarget.GetTarget().GetRotation();
 					KinematicTarget.SetMode(EKinematicTargetMode::Reset);
 				}
 				else
 				{
-					TargetPos = FVec3::Lerp(PrevX, KinematicTarget.GetTarget().GetLocation(), StepFraction);
-					TargetRot = FRotation3::Slerp(PrevR, KinematicTarget.GetTarget().GetRotation(), StepFraction);
+					// as a reminder, stepfraction is the remaing fraction of the step from the remaining steps
+					// for total of 4 steps and current step of 2, this will be 1/3 ( 1 step passed, 3 steps remains )
+					NewX = FVec3::Lerp(CurrentX, KinematicTarget.GetTarget().GetLocation(), StepFraction);
+					NewR = FRotation3::Slerp(CurrentR, KinematicTarget.GetTarget().GetRotation(), StepFraction);
 				}
 				if (Dt > MinDt)
 				{
-					FVec3 V = FVec3::CalculateVelocity(PrevX, TargetPos, Dt);
+					FVec3 V = FVec3::CalculateVelocity(CurrentX, NewX, Dt);
 					Particle.V() = V;
 
-					FVec3 W = FRotation3::CalculateAngularVelocity(PrevR, TargetRot, Dt);
+					FVec3 W = FRotation3::CalculateAngularVelocity(CurrentR, NewR, Dt);
 					Particle.W() = W;
 				}
-				Particle.X() = TargetPos;
-				Particle.R() = TargetRot;
+				Particle.X() = NewX;
+				Particle.R() = NewR;
 				Particles.MarkTransientDirtyParticle(Particle.Handle());
 				break;
 			}
diff --git a/Engine/Source/Runtime/PhysicsCore/Private/ChaosEngineInterface.cpp b/Engine/Source/Runtime/PhysicsCore/Private/ChaosEngineInterface.cpp
--- a/Engine/Source/Runtime/PhysicsCore/Private/ChaosEngineInterface.cpp
+++ b/Engine/Source/Runtime/PhysicsCore/Private/ChaosEngineInterface.cpp
@@ -1661,22 +1661,18 @@ void FChaosEngineInterface::SetGlobalPose_AssumesLocked(const FPhysicsActorHandl
 
 void FChaosEngineInterface::SetKinematicTarget_AssumesLocked(const FPhysicsActorHandle& InActorReference,const FTransform& InNewTarget)
 {
-	{
-		Chaos::TKinematicTarget <Chaos::FReal, 3 > NewKinematicTarget;
-		// SetKinematicTarget_AssumesLocked could be called multiple times in one time step
-		// Don't update Previous here. Previous is updated in FSingleParticlePhysicsProxy::PushToPhysicsState. Previous is not used in game thread and set to identity.
-		// @todo(chaos): create a new KinematicTarget class for game thread that does not have Preivous 
-		Chaos::FRigidTransform3 IdentityTransform;
-		NewKinematicTarget.SetTargetMode(InNewTarget, IdentityTransform);
-		InActorReference->GetGameThreadAPI().SetKinematicTarget(NewKinematicTarget);
+	Chaos::TKinematicTarget<Chaos::FReal, 3> NewKinematicTarget;
+	// SetKinematicTarget_AssumesLocked could be called multiple times in one time step
+	NewKinematicTarget.SetTargetMode(InNewTarget);
+	InActorReference->GetGameThreadAPI().SetKinematicTarget(NewKinematicTarget);
 
-		InActorReference->GetGameThreadAPI().SetX(InNewTarget.GetLocation(), false);
-		InActorReference->GetGameThreadAPI().SetR(InNewTarget.GetRotation(), false);
-		InActorReference->GetGameThreadAPI().UpdateShapeBounds();
+	// IMPORTANT : we do not invalidate X and R as they will be properly computed using the kinematic target information 
+	InActorReference->GetGameThreadAPI().SetX(InNewTarget.GetLocation(), false); 
+	InActorReference->GetGameThreadAPI().SetR(InNewTarget.GetRotation(), false); 
+	InActorReference->GetGameThreadAPI().UpdateShapeBounds();
 
-		FChaosScene* Scene = GetCurrentScene(InActorReference);
-		Scene->UpdateActorInAccelerationStructure(InActorReference);
-	}
+	FChaosScene* Scene = GetCurrentScene(InActorReference);
+	Scene->UpdateActorInAccelerationStructure(InActorReference);
 }
 
 #elif WITH_ENGINE //temp physx code to make moving out of engine easier
