physics : Fix build after recent merge of kinematic target related changes
- removed reference to previous transform

diff --git a/Engine/Source/Runtime/Experimental/Chaos/Private/PhysicsProxy/GeometryCollectionPhysicsProxy.cpp b/Engine/Source/Runtime/Experimental/Chaos/Private/PhysicsProxy/GeometryCollectionPhysicsProxy.cpp
--- a/Engine/Source/Runtime/Experimental/Chaos/Private/PhysicsProxy/GeometryCollectionPhysicsProxy.cpp
+++ b/Engine/Source/Runtime/Experimental/Chaos/Private/PhysicsProxy/GeometryCollectionPhysicsProxy.cpp
@@ -1530,8 +1530,7 @@ void FGeometryCollectionPhysicsProxy::SetClusteredParticleKinematicTarget(Chaos:
 	if (ensure(ObjectState == Chaos::EObjectStateType::Kinematic))
 	{
 		Chaos::TKinematicTarget<Chaos::FReal, 3> NewKinematicTarget;
-		Chaos::FRigidTransform3 CurrentWorldTransform(Handle->X(), Handle->R());
-		NewKinematicTarget.SetTargetMode(NewWorldTransform, CurrentWorldTransform);
+		NewKinematicTarget.SetTargetMode(NewWorldTransform);
 		Handle->SetKinematicTarget(NewKinematicTarget);
 	}
 }
