Fix chaos constraint graph

diff --git a/Engine/Source/Runtime/Experimental/Chaos/Private/Chaos/PBDConstraintGraph.cpp b/Engine/Source/Runtime/Experimental/Chaos/Private/Chaos/PBDConstraintGraph.cpp
--- a/Engine/Source/Runtime/Experimental/Chaos/Private/Chaos/PBDConstraintGraph.cpp
+++ b/Engine/Source/Runtime/Experimental/Chaos/Private/Chaos/PBDConstraintGraph.cpp
@@ -131,6 +131,7 @@ void FPBDConstraintGraph::ParticleRemove(FGeometryParticleHandle* RemovedParticl
 
 		Visited[NodeIdx]=0;
 		ParticleToNodeIndex.Remove(RemovedParticle);
+		UpdatedNodes.RemoveSwap(NodeIdx, false);
 	}
 }
 
@@ -187,7 +188,7 @@ void FPBDConstraintGraph::InitializeGraph(const TParticleView<FGeometryParticles
 	}
 	else
 	{
-		if (!CHAOS_ENSURE(NumNonDisabledParticles <= Nodes.Num()))
+		if (!(NumNonDisabledParticles <= Nodes.Num())) 
 		{
 			for (auto& Particle : Particles)
 			{
diff --git a/Engine/Source/Runtime/Experimental/Chaos/Private/PhysicsProxy/GeometryCollectionPhysicsProxy.cpp b/Engine/Source/Runtime/Experimental/Chaos/Private/PhysicsProxy/GeometryCollectionPhysicsProxy.cpp
--- a/Engine/Source/Runtime/Experimental/Chaos/Private/PhysicsProxy/GeometryCollectionPhysicsProxy.cpp
+++ b/Engine/Source/Runtime/Experimental/Chaos/Private/PhysicsProxy/GeometryCollectionPhysicsProxy.cpp
@@ -1565,6 +1565,7 @@ void FGeometryCollectionPhysicsProxy::BufferPhysicsResults(Chaos::FPBDRigidsSolv
 	const TManagedArray<int32>& Parent = PhysicsThreadCollection.Parent;
 	const TManagedArray<TSet<int32>>& Children = PhysicsThreadCollection.Children;
 
+	if(NumTransformGroupElements > 0)
 	{ 
 		SCOPE_CYCLE_COUNTER(STAT_CalcParticleToWorld);
 
@@ -2604,7 +2605,7 @@ void FGeometryCollectionPhysicsProxy::FieldParameterUpdateCallback(Chaos::FPBDRi
 
 	// Process Particle-Collection commands
 	int32 NumCommands = Commands.Num();
-	if (NumCommands && RigidSolver && Collection.Transform.Num())
+	if (NumCommands && RigidSolver && !RigidSolver->IsShuttingDown() && Collection.Transform.Num())
 	{
 		TArray<int32> CommandsToRemove;
 		CommandsToRemove.Reserve(NumCommands);
@@ -2758,7 +2759,7 @@ void FGeometryCollectionPhysicsProxy::FieldForcesUpdateCallback(Chaos::FPBDRigid
 	SCOPE_CYCLE_COUNTER(STAT_ForceUpdateField_Object);
 
 	const int32 NumCommands = Commands.Num();
-	if (NumCommands && RigidSolver)
+	if (NumCommands && RigidSolver && !RigidSolver->IsShuttingDown())
 	{
 		TArray<int32> CommandsToRemove;
 		CommandsToRemove.Reserve(NumCommands);
diff --git a/Engine/Source/Runtime/Experimental/Chaos/Private/PhysicsProxy/PerSolverFieldSystem.cpp b/Engine/Source/Runtime/Experimental/Chaos/Private/PhysicsProxy/PerSolverFieldSystem.cpp
--- a/Engine/Source/Runtime/Experimental/Chaos/Private/PhysicsProxy/PerSolverFieldSystem.cpp
+++ b/Engine/Source/Runtime/Experimental/Chaos/Private/PhysicsProxy/PerSolverFieldSystem.cpp
@@ -134,8 +134,11 @@ void FPerSolverFieldSystem::FieldParameterUpdateCallback(
 	Chaos::FPBDPositionConstraints& PositionTarget,
 	TMap<int32, int32>& TargetedParticles)
 {
-	FieldParameterUpdateInternal(InSolver, PositionTarget, TargetedParticles, TransientCommands, true);
-	FieldParameterUpdateInternal(InSolver, PositionTarget, TargetedParticles, PersistentCommands, false);
+	if (InSolver && !InSolver->IsShuttingDown())
+	{
+		FieldParameterUpdateInternal(InSolver, PositionTarget, TargetedParticles, TransientCommands, true);
+		FieldParameterUpdateInternal(InSolver, PositionTarget, TargetedParticles, PersistentCommands, false);
+	}
 }
 
 void FPerSolverFieldSystem::FieldForcesUpdateInternal(
@@ -193,8 +196,11 @@ void FPerSolverFieldSystem::FieldForcesUpdateInternal(
 void FPerSolverFieldSystem::FieldForcesUpdateCallback(
 	Chaos::FPBDRigidsSolver* InSolver)
 {
-	FieldForcesUpdateInternal(InSolver, TransientCommands, true);
-	FieldForcesUpdateInternal(InSolver, PersistentCommands, false);
+	if (InSolver && !InSolver->IsShuttingDown())
+	{
+		FieldForcesUpdateInternal(InSolver, TransientCommands, true);
+		FieldForcesUpdateInternal(InSolver, PersistentCommands, false);
+	}
 }
 
 FORCEINLINE void EvaluateImpulseField(
diff --git a/Engine/Source/Runtime/Experimental/GeometryCollectionEngine/Private/GeometryCollection/GeometryCollectionComponent.cpp b/Engine/Source/Runtime/Experimental/GeometryCollectionEngine/Private/GeometryCollection/GeometryCollectionComponent.cpp
--- a/Engine/Source/Runtime/Experimental/GeometryCollectionEngine/Private/GeometryCollection/GeometryCollectionComponent.cpp
+++ b/Engine/Source/Runtime/Experimental/GeometryCollectionEngine/Private/GeometryCollection/GeometryCollectionComponent.cpp
@@ -1633,7 +1633,8 @@ void UGeometryCollectionComponent::TickComponent(float DeltaTime, enum ELevelTic
 	//if (bRenderStateDirty && DynamicCollection)	//todo: always send for now
 	if (RestCollection)
 	{
-		if (CHAOS_ENSURE(DynamicCollection)) //, TEXT("No dynamic collection available for component %s during tick."), *GetName()))
+		// In editor mode we have no DynamicCollection so this test is necessary
+		if(DynamicCollection) //, TEXT("No dynamic collection available for component %s during tick."), *GetName()))
 		{
 			if (RestCollection->HasVisibleGeometry() || DynamicCollection->IsDirty())
 			{
