Network Physics
-Remove unneeded logging

#rnx

diff --git a/Engine/Source/Runtime/Experimental/Chaos/Private/Chaos/Framework/ChaosResultsManager.cpp b/Engine/Source/Runtime/Experimental/Chaos/Private/Chaos/Framework/ChaosResultsManager.cpp
--- a/Engine/Source/Runtime/Experimental/Chaos/Private/Chaos/Framework/ChaosResultsManager.cpp
+++ b/Engine/Source/Runtime/Experimental/Chaos/Private/Chaos/Framework/ChaosResultsManager.cpp
@@ -392,17 +392,6 @@ namespace Chaos
 					}
 				}
 			}
-
-			for (const FDirtyRigidParticleData& ResimDirty : Results.Next->DirtyRigids)
-			{
-				if (FSingleParticlePhysicsProxy* ResimProxy = ResimDirty.GetProxy())
-				{
-					if (ParticleToResimInfo.FindChecked(ResimProxy).bDiverged)
-					{
-						UE_LOG(LogTemp, Warning, TEXT("Diverged"));
-					}
-				}
-			}
 		}
 	}
 
