Fix wrong field boxfalloff bbox + fix field particlehandles construction + fix field commands removal

diff --git a/Engine/Source/Runtime/Experimental/Chaos/Private/PhysicsProxy/GeometryCollectionPhysicsProxy.cpp b/Engine/Source/Runtime/Experimental/Chaos/Private/PhysicsProxy/GeometryCollectionPhysicsProxy.cpp
--- a/Engine/Source/Runtime/Experimental/Chaos/Private/PhysicsProxy/GeometryCollectionPhysicsProxy.cpp
+++ b/Engine/Source/Runtime/Experimental/Chaos/Private/PhysicsProxy/GeometryCollectionPhysicsProxy.cpp
@@ -2791,8 +2791,6 @@ void FGeometryCollectionPhysicsProxy::FieldForcesUpdateCallback(Chaos::FPBDRigid
 						TArray<FVector>& FinalResults = ExecutionDatas.VectorResults[(uint8)EFieldCommandResultType::FinalResult];
 						ResetResultsArray < FVector >(ExecutionDatas.SamplePositions.Num(), FinalResults, FVector::ZeroVector);
 
-						TArrayView<FVector> ResultsView(&(FinalResults[0]), FinalResults.Num());
-
 						Chaos::FieldVectorForceUpdate(RigidSolver, FieldCommand, ParticleHandles,
 							FieldContext, FinalResults);
 					}
