fix for CIS issue

diff --git a/Engine/Source/Runtime/Experimental/Chaos/Private/Chaos/ChaosDebugDraw.cpp b/Engine/Source/Runtime/Experimental/Chaos/Private/Chaos/ChaosDebugDraw.cpp
--- a/Engine/Source/Runtime/Experimental/Chaos/Private/Chaos/ChaosDebugDraw.cpp
+++ b/Engine/Source/Runtime/Experimental/Chaos/Private/Chaos/ChaosDebugDraw.cpp
@@ -75,7 +75,7 @@ namespace Chaos
 			/* AngVelScale =				*/ 0.0f,
 			/* ImpulseScale =				*/ 0.0f,
 			/* InertiaScale =				*/ 0.0f,
-			/* DrawPriority =				*/ 10.0f,
+			/* DrawPriority =				*/ 10,
 			/* bShowSimple =				*/ true,
 			/* bShowComplex =				*/ false,
 			/* bInShowLevelSetCollision =	*/ false
