Added async dt to log

diff --git a/Engine/Source/Runtime/Experimental/Chaos/Private/Chaos/Framework/PhysicsSolverBase.cpp b/Engine/Source/Runtime/Experimental/Chaos/Private/Chaos/Framework/PhysicsSolverBase.cpp
--- a/Engine/Source/Runtime/Experimental/Chaos/Private/Chaos/Framework/PhysicsSolverBase.cpp
+++ b/Engine/Source/Runtime/Experimental/Chaos/Private/Chaos/Framework/PhysicsSolverBase.cpp
@@ -146,8 +146,27 @@ namespace Chaos
 		, bStealAdvanceTasksForTesting(false)
 #endif
 	{
+		UE_LOG(LogChaos, Log, TEXT("FPhysicsSolverBase::AsyncDt:%f"), AsyncDt);
 	}
 
+	void FPhysicsSolverBase::EnableAsyncMode(FReal FixedDt)
+	{
+		if (AsyncDt != FixedDt)
+		{
+			AccumulatedTime = 0;
+		}
+		AsyncDt = FixedDt;
+		
+		UE_LOG(LogChaos, Log, TEXT("FPhysicsSolverBase::AsyncDt:%f"), AsyncDt);
+	}
+
+	void FPhysicsSolverBase::DisableAsyncMode()
+	{
+		AsyncDt = -1;
+		UE_LOG(LogChaos, Log, TEXT("FPhysicsSolverBase::AsyncDt:%f"), AsyncDt);
+	}
+
+
 	FPhysicsSolverBase::~FPhysicsSolverBase()
 	{
 		//free pull results before MarshallingManager is destroyed (technically should be safe because constructed after, but this makes it more explicit)
diff --git a/Engine/Source/Runtime/Experimental/Chaos/Public/Chaos/Framework/PhysicsSolverBase.h b/Engine/Source/Runtime/Experimental/Chaos/Public/Chaos/Framework/PhysicsSolverBase.h
--- a/Engine/Source/Runtime/Experimental/Chaos/Public/Chaos/Framework/PhysicsSolverBase.h
+++ b/Engine/Source/Runtime/Experimental/Chaos/Public/Chaos/Framework/PhysicsSolverBase.h
@@ -263,21 +263,10 @@ namespace Chaos
 
 		bool IsShuttingDown() const { return bIsShuttingDown; }
 
-		void EnableAsyncMode(FReal FixedDt)
-		{
-			if (AsyncDt != FixedDt)
-			{
-				AccumulatedTime = 0;
-			}
-
-			AsyncDt = FixedDt;
-		}
-
-		void DisableAsyncMode()
-		{
-			AsyncDt = -1;
-		}
-
+		void EnableAsyncMode(FReal FixedDt);
+		
+		void DisableAsyncMode();
+		
 		virtual void ConditionalApplyRewind_Internal(){}
 
 		FChaosMarshallingManager& GetMarshallingManager() { return MarshallingManager; }
