Hoist ~20-30 FP out of inner loops

Test with 100 ~static objects with different geom types shows this to give around 0.1ms saving on average and 0.5ms on the max case.

diff --git a/Engine/Source/Runtime/Experimental/Chaos/Public/Chaos/AABBTree.h b/Engine/Source/Runtime/Experimental/Chaos/Public/Chaos/AABBTree.h
--- a/Engine/Source/Runtime/Experimental/Chaos/Public/Chaos/AABBTree.h
+++ b/Engine/Source/Runtime/Experimental/Chaos/Public/Chaos/AABBTree.h
@@ -37,7 +37,7 @@ template <typename TQueryFastData, EAABBQueryType Query>
 struct TAABBTreeIntersectionHelper
 {
 	static bool Intersects(const FVec3& Start, TQueryFastData& QueryFastData, FReal& TOI, FVec3& OutPosition,
-		const FAABB3& Bounds, const FAABB3& QueryBounds, const FVec3& QueryHalfExtents)
+		const FAABB3& Bounds, const FAABB3& QueryBounds, const FVec3& QueryHalfExtents, const FVec3& Dir, const FVec3 InvDir, const bool bParallel[3])
 	{
 		check(false);
 		return true;
@@ -48,20 +48,20 @@ template<>
 struct TAABBTreeIntersectionHelper<FQueryFastData, EAABBQueryType::Raycast>
 {
 	FORCEINLINE_DEBUGGABLE static bool Intersects(const FVec3& Start, FQueryFastData& QueryFastData, FReal& TOI, FVec3& OutPosition,
-		const FAABB3& Bounds, const FAABB3& QueryBounds, const FVec3& QueryHalfExtents)
+		const FAABB3& Bounds, const FAABB3& QueryBounds, const FVec3& QueryHalfExtents, const FVec3& Dir, const FVec3 InvDir, const bool bParallel[3])
 	{
-		return Bounds.RaycastFast(Start, QueryFastData.Dir, QueryFastData.InvDir, QueryFastData.bParallel, QueryFastData.CurrentLength, QueryFastData.InvCurrentLength, TOI, OutPosition);
+		return Bounds.RaycastFast(Start, Dir, InvDir, bParallel, QueryFastData.CurrentLength, QueryFastData.InvCurrentLength, TOI, OutPosition);
 	}
 };
 
 template <>
 struct TAABBTreeIntersectionHelper<FQueryFastData, EAABBQueryType::Sweep>
 {
-	FORCEINLINE_DEBUGGABLE static bool Intersects(const FVec3& Start, FQueryFastData& QueryFastData, FReal& TOI, FVec3& OutPosition, 
-		const FAABB3& Bounds, const FAABB3& QueryBounds, const FVec3& QueryHalfExtents)
+	FORCEINLINE_DEBUGGABLE static bool Intersects(const FVec3& Start, FQueryFastData& QueryFastData, FReal& TOI, FVec3& OutPosition,
+		const FAABB3& Bounds, const FAABB3& QueryBounds, const FVec3& QueryHalfExtents, const FVec3& Dir, const FVec3 InvDir, const bool bParallel[3])
 	{
 		FAABB3 SweepBounds(Bounds.Min() - QueryHalfExtents, Bounds.Max() + QueryHalfExtents);
-		return SweepBounds.RaycastFast(Start, QueryFastData.Dir, QueryFastData.InvDir, QueryFastData.bParallel, QueryFastData.CurrentLength, QueryFastData.InvCurrentLength, TOI, OutPosition);
+		return SweepBounds.RaycastFast(Start, Dir, InvDir, bParallel, QueryFastData.CurrentLength, QueryFastData.InvCurrentLength, TOI, OutPosition);
 	}
 };
 
@@ -69,7 +69,7 @@ template <>
 struct TAABBTreeIntersectionHelper<FQueryFastDataVoid, EAABBQueryType::Overlap>
 {
 	FORCEINLINE_DEBUGGABLE static bool Intersects(const FVec3& Start, FQueryFastDataVoid& QueryFastData, FReal& TOI, FVec3& OutPosition,
-		const FAABB3& Bounds, const FAABB3& QueryBounds, const FVec3& QueryHalfExtents)
+		const FAABB3& Bounds, const FAABB3& QueryBounds, const FVec3& QueryHalfExtents, const FVec3& Dir, const FVec3 InvDir, const bool bParallel[3])
 	{
 		return QueryBounds.Intersects(Bounds);
 	}
@@ -136,15 +136,16 @@ struct TAABBTreeLeafArray : public TBoundsWrapperHelper<TPayloadType, bComputeBo
 	}
 
 	template <typename TSQVisitor, typename TQueryFastData>
-	bool RaycastFast(const FVec3& Start, TQueryFastData& QueryFastData, TSQVisitor& Visitor) const
+	FORCEINLINE_DEBUGGABLE bool RaycastFast(const FVec3& Start, TQueryFastData& QueryFastData, TSQVisitor& Visitor, const FVec3& Dir, const FVec3 InvDir, const bool bParallel[3]) const
 	{
-		return RaycastSweepImp</*bSweep=*/false>(Start, QueryFastData, FVec3(), Visitor);
+		return RaycastSweepImp</*bSweep=*/false>(Start, QueryFastData, FVec3(), Visitor, Dir, InvDir, bParallel);
 	}
 
 	template <typename TSQVisitor, typename TQueryFastData>
-	bool SweepFast(const FVec3& Start, TQueryFastData& QueryFastData, const FVec3& QueryHalfExtents, TSQVisitor& Visitor) const
+	FORCEINLINE_DEBUGGABLE bool SweepFast(const FVec3& Start, TQueryFastData& QueryFastData, const FVec3& QueryHalfExtents, TSQVisitor& Visitor,
+			const FVec3& Dir, const FVec3 InvDir, const bool bParallel[3]) const
 	{
-		return RaycastSweepImp</*bSweep=*/true>(Start, QueryFastData, QueryHalfExtents, Visitor);
+		return RaycastSweepImp</*bSweep=*/true>(Start, QueryFastData, QueryHalfExtents, Visitor, Dir, InvDir, bParallel);
 	}
 
 	template <typename TSQVisitor>
@@ -166,14 +167,15 @@ struct TAABBTreeLeafArray : public TBoundsWrapperHelper<TPayloadType, bComputeBo
 	}
 
 	template <bool bSweep, typename TQueryFastData, typename TSQVisitor>
-	bool RaycastSweepImp(const FVec3& Start, TQueryFastData& QueryFastData, const FVec3& QueryHalfExtents, TSQVisitor& Visitor) const
+	FORCEINLINE_DEBUGGABLE bool RaycastSweepImp(const FVec3& Start, TQueryFastData& QueryFastData, const FVec3& QueryHalfExtents, TSQVisitor& Visitor, const FVec3& Dir, const FVec3 InvDir, const bool bParallel[3]) const
 	{
 		FVec3 TmpPosition;
 		FReal TOI;
 		for (const auto& Elem : Elems)
 		{
 			const auto& InstanceBounds = Elem.Bounds;
-			if (TAABBTreeIntersectionHelper<TQueryFastData, bSweep ? EAABBQueryType::Sweep : EAABBQueryType::Raycast>::Intersects(Start, QueryFastData, TOI, TmpPosition, InstanceBounds, FAABB3(), QueryHalfExtents))
+			if (TAABBTreeIntersectionHelper<TQueryFastData, bSweep ? EAABBQueryType::Sweep :
+					EAABBQueryType::Raycast>::Intersects(Start, QueryFastData, TOI, TmpPosition, InstanceBounds, FAABB3(), QueryHalfExtents, Dir, InvDir, bParallel))
 			{
 				TSpatialVisitorData<TPayloadType> VisitData(Elem.Payload, true, InstanceBounds);
 				const bool bContinue = (bSweep && Visitor.VisitSweep(VisitData, QueryFastData)) || (!bSweep &&  Visitor.VisitRaycast(VisitData, QueryFastData));
@@ -427,13 +429,13 @@ public:
 	void Raycast(const FVec3& Start, const FVec3& Dir, const FReal Length, SQVisitor& Visitor) const
 	{
 		FQueryFastData QueryFastData(Dir, Length);
-		QueryImp<EAABBQueryType::Raycast>(Start, QueryFastData, FVec3(), FAABB3(), Visitor);
+		QueryImp<EAABBQueryType::Raycast>(Start, QueryFastData, FVec3(), FAABB3(), Visitor, QueryFastData.Dir, QueryFastData.InvDir, QueryFastData.bParallel);
 	}
 
 	template <typename SQVisitor>
-	bool RaycastFast(const FVec3& Start, FQueryFastData& CurData, SQVisitor& Visitor) const
+	bool RaycastFast(const FVec3& Start, FQueryFastData& CurData, SQVisitor& Visitor, const FVec3& Dir, const FVec3 InvDir, const bool bParallel[3]) const
 	{
-		return QueryImp<EAABBQueryType::Raycast>(Start, CurData, FVec3(), FAABB3(), Visitor);
+		return QueryImp<EAABBQueryType::Raycast>(Start, CurData, FVec3(), FAABB3(), Visitor, Dir, InvDir, bParallel);
 	}
 
 	void Sweep(const FVec3& Start, const FVec3& Dir, const FReal Length, const FVec3 QueryHalfExtents, ISpatialVisitor<TPayloadType, FReal>& Visitor) const override
@@ -446,13 +448,13 @@ public:
 	void Sweep(const FVec3& Start, const FVec3& Dir, const FReal Length, const FVec3 QueryHalfExtents, SQVisitor& Visitor) const
 	{
 		FQueryFastData QueryFastData(Dir, Length);
-		QueryImp<EAABBQueryType::Sweep>(Start, QueryFastData, QueryHalfExtents, FAABB3(), Visitor);
+		QueryImp<EAABBQueryType::Sweep>(Start, QueryFastData, QueryHalfExtents, FAABB3(), Visitor, QueryFastData.Dir, QueryFastData.InvDir, QueryFastData.bParallel);
 	}
 
 	template <typename SQVisitor>
-	bool SweepFast(const FVec3& Start, FQueryFastData& CurData, const FVec3 QueryHalfExtents, SQVisitor& Visitor) const
+	bool SweepFast(const FVec3& Start, FQueryFastData& CurData, const FVec3 QueryHalfExtents, SQVisitor& Visitor, const FVec3& Dir, const FVec3 InvDir, const bool bParallel[3]) const
 	{
-		return QueryImp<EAABBQueryType::Sweep>(Start,CurData, QueryHalfExtents, FAABB3(), Visitor);
+		return QueryImp<EAABBQueryType::Sweep>(Start,CurData, QueryHalfExtents, FAABB3(), Visitor, Dir, InvDir, bParallel);
 	}
 
 	void Overlap(const FAABB3& QueryBounds, ISpatialVisitor<TPayloadType, FReal>& Visitor) const override
@@ -472,7 +474,7 @@ public:
 	{
 		//dummy variables to reuse templated path
 		FQueryFastDataVoid VoidData;
-		return QueryImp<EAABBQueryType::Overlap>(FVec3(), VoidData, FVec3(), QueryBounds, Visitor);
+		return QueryImp<EAABBQueryType::Overlap>(FVec3(), VoidData, FVec3(), QueryBounds, Visitor, VoidData.Dir, VoidData.InvDir, VoidData.bParallel);
 	}
 
 	virtual void RemoveElement(const TPayloadType& Payload)
@@ -692,7 +694,7 @@ private:
 	}
 
 	template <EAABBQueryType Query, typename TQueryFastData, typename SQVisitor>
-	bool QueryImp(const FVec3& Start, TQueryFastData& CurData, const FVec3 QueryHalfExtents, const FAABB3& QueryBounds, SQVisitor& Visitor) const
+	bool QueryImp(const FVec3& RESTRICT Start, TQueryFastData& CurData, const FVec3 QueryHalfExtents, const FAABB3& QueryBounds, SQVisitor& Visitor, const FVec3& Dir, const FVec3 InvDir, const bool bParallel[3]) const
 	{
 		//QUICK_SCOPE_CYCLE_COUNTER(AABBTreeQueryImp);
 		FVec3 TmpPosition;
@@ -711,7 +713,7 @@ private:
 				}
 
 				const auto& InstanceBounds = Elem.Bounds;
-				if(TAABBTreeIntersectionHelper<TQueryFastData,Query>::Intersects(Start,CurData,TOI,TmpPosition,InstanceBounds,QueryBounds,QueryHalfExtents))
+				if(TAABBTreeIntersectionHelper<TQueryFastData,Query>::Intersects(Start, CurData, TOI, TmpPosition, InstanceBounds, QueryBounds, QueryHalfExtents, Dir, InvDir, bParallel))
 				{
 					TSpatialVisitorData<TPayloadType> VisitData(Elem.Payload,true);
 					bool bContinue;
@@ -743,7 +745,7 @@ private:
 					continue;
 				}
 
-				if (TAABBTreeIntersectionHelper<TQueryFastData, Query>::Intersects(Start, CurData, TOI, TmpPosition, InstanceBounds, QueryBounds, QueryHalfExtents))
+				if (TAABBTreeIntersectionHelper<TQueryFastData, Query>::Intersects(Start, CurData, TOI, TmpPosition, InstanceBounds, QueryBounds, QueryHalfExtents, Dir, InvDir, bParallel))
 				{
 					TSpatialVisitorData<TPayloadType> VisitData(Elem.Payload, true, InstanceBounds);
 					bool bContinue;
@@ -801,12 +803,12 @@ private:
 				}
 				else if (Query == EAABBQueryType::Sweep)
 				{
-					if (Leaf.SweepFast(Start, CurData, QueryHalfExtents, Visitor) == false)
+					if (Leaf.SweepFast(Start, CurData, QueryHalfExtents, Visitor, Dir, InvDir, bParallel) == false)
 					{
 						return false;
 					}
 				}
-				else if (Leaf.RaycastFast(Start, CurData, Visitor) == false)
+				else if (Leaf.RaycastFast(Start, CurData, Visitor, Dir, InvDir, bParallel) == false)
 				{
 					return false;
 				}
@@ -816,7 +818,7 @@ private:
 				int32 Idx = 0;
 				for (const FAABB3& AABB : Node.ChildrenBounds)
 				{
-					if(TAABBTreeIntersectionHelper<TQueryFastData, Query>::Intersects(Start, CurData, TOI, TmpPosition, AABB, QueryBounds, QueryHalfExtents))
+					if(TAABBTreeIntersectionHelper<TQueryFastData, Query>::Intersects(Start, CurData, TOI, TmpPosition, AABB, QueryBounds, QueryHalfExtents, Dir, InvDir, bParallel))
 					{
 						NodeStack.Add(FNodeQueueEntry{ Node.ChildrenNodes[Idx], TOI });
 					}
diff --git a/Engine/Source/Runtime/Experimental/Chaos/Public/Chaos/BoundingVolume.h b/Engine/Source/Runtime/Experimental/Chaos/Public/Chaos/BoundingVolume.h
--- a/Engine/Source/Runtime/Experimental/Chaos/Public/Chaos/BoundingVolume.h
+++ b/Engine/Source/Runtime/Experimental/Chaos/Public/Chaos/BoundingVolume.h
@@ -275,16 +275,16 @@ public:
 	}
 
 	template <typename SQVisitor>
-	bool RaycastFast(const TVector<T,d>& Start, FQueryFastData& CurData, SQVisitor& Visitor) const
+	bool RaycastFast(const TVector<T,d>& Start, FQueryFastData& CurData, SQVisitor& Visitor, const FVec3& Dir, const FVec3 InvDir, const bool bParallel[3]) const
 	{
-		return RaycastImp(Start, CurData, Visitor);
+		return RaycastImp(Start, CurData, Visitor, Dir, InvDir, bParallel);
 	}
 
 	template <typename SQVisitor, bool bPruneDuplicates = true>
 	void Raycast(const TVector<T, d>& Start, const TVector<T, d>& Dir, const T Length, SQVisitor& Visitor) const
 	{
 		FQueryFastData CurData(Dir, Length); 
-		RaycastImp<SQVisitor, bPruneDuplicates>(Start, CurData, Visitor);
+		RaycastImp<SQVisitor, bPruneDuplicates>(Start, CurData, Visitor, CurData.Dir, CurData.InvDir, CurData.bParallel);
 	}
 
 	void Sweep(const TVector<T, d>& Start, const TVector<T, d>& Dir, const T Length, const TVector<T, d> QueryHalfExtents, ISpatialVisitor<TPayloadType, T>& Visitor) const override
@@ -294,16 +294,16 @@ public:
 	}
 
 	template <typename SQVisitor>
-	bool SweepFast(const TVector<T,d>& Start, FQueryFastData& CurData, const TVector<T,d>& QueryHalfExtents, SQVisitor& Visitor) const
+	bool SweepFast(const TVector<T,d>& Start, FQueryFastData& CurData, const TVector<T,d>& QueryHalfExtents, SQVisitor& Visitor, const FVec3& Dir, const FVec3 InvDir, const bool bParallel[3]) const
 	{
-		return SweepImp(Start, CurData, QueryHalfExtents, Visitor);
+		return SweepImp(Start, CurData, QueryHalfExtents, Visitor, Dir, InvDir, bParallel);
 	}
 
 	template <typename SQVisitor, bool bPruneDuplicates = true>
 	void Sweep(const TVector<T, d>& Start, const TVector<T, d>& Dir, const T Length, const TVector<T, d> QueryHalfExtents, SQVisitor& Visitor) const
 	{
 		FQueryFastData CurData(Dir, Length);
-		SweepImp<SQVisitor, bPruneDuplicates>(Start, CurData, QueryHalfExtents, Visitor);
+		SweepImp<SQVisitor, bPruneDuplicates>(Start, CurData, QueryHalfExtents, Visitor, CurData.Dir, CurData.InvDir, CurData.bParallel);
 	}
 
 	virtual void Overlap(const TAABB<T, d>& QueryBounds, ISpatialVisitor<TPayloadType, T>& Visitor) const override
@@ -415,7 +415,7 @@ private:
 	}
 
 	template <typename SQVisitor, bool bPruneDuplicates = true>
-	bool RaycastImp(const TVector<T, d>& Start, FQueryFastData& CurData, SQVisitor& Visitor) const
+	bool RaycastImp(const TVector<T, d>& Start, FQueryFastData& CurData, SQVisitor& Visitor, const FVec3& Dir, const FVec3 InvDir, const bool bParallel[3]) const
 	{
 		TVector<T, d> TmpPosition;
 		T TOI;
@@ -430,7 +430,7 @@ private:
 
 			const auto& InstanceBounds = Elem.Bounds;
 			if (InstanceBounds.RaycastFast(Start,
-				CurData.Dir, CurData.InvDir, CurData.bParallel, CurData.CurrentLength, CurData.InvCurrentLength, TOI, TmpPosition))
+				Dir, InvDir, bParallel, CurData.CurrentLength, CurData.InvCurrentLength, TOI, TmpPosition))
 			{
 				TSpatialVisitorData<TPayloadType> VisitData(Elem.Payload, true, InstanceBounds);
 				const bool bContinue = Visitor.VisitRaycast(VisitData, CurData);
@@ -450,8 +450,8 @@ private:
 			}
 
 			const auto& InstanceBounds = Elem.Bounds;
-			if (InstanceBounds.RaycastFast(Start, CurData.Dir,
-				CurData.InvDir, CurData.bParallel, CurData.CurrentLength, CurData.InvCurrentLength, TOI, TmpPosition))
+			if (InstanceBounds.RaycastFast(Start, Dir,
+				InvDir, bParallel, CurData.CurrentLength, CurData.InvCurrentLength, TOI, TmpPosition))
 			{
 				TSpatialVisitorData<TPayloadType> VisitData(Elem.Payload, true, InstanceBounds);
 				const bool bContinue = Visitor.VisitRaycast(VisitData, CurData);
@@ -470,8 +470,8 @@ private:
 		TAABB<T, d> GlobalBounds(MGrid.MinCorner(), MGrid.MaxCorner());
 		TVector<T, d> NextStart;
 		TVector<int32, d> CellIdx;
-		bool bCellsLeft = MElements.Num() && GlobalBounds.RaycastFast(Start, CurData.Dir,
-			CurData.InvDir, CurData.bParallel, CurData.CurrentLength, CurData.InvCurrentLength, TOI, NextStart);
+		bool bCellsLeft = MElements.Num() && GlobalBounds.RaycastFast(Start, Dir,
+			InvDir, bParallel, CurData.CurrentLength, CurData.InvCurrentLength, TOI, NextStart);
 		if (bCellsLeft)
 		{
 			CellIdx = MGrid.Cell(NextStart);
@@ -514,7 +514,7 @@ private:
 					}
 					const auto& InstanceBounds = Elem.Bounds;
 					if (InstanceBounds.RaycastFast(Start,
-						CurData.Dir, CurData.InvDir, CurData.bParallel, CurData.CurrentLength, CurData.InvCurrentLength, TOI, TmpPosition))
+							Dir, InvDir, bParallel, CurData.CurrentLength, CurData.InvCurrentLength, TOI, TmpPosition))
 					{
 						TSpatialVisitorData<TPayloadType> VisitData(Elem.Payload, true, InstanceBounds);
 						const bool bContinue = Visitor.VisitRaycast(VisitData, CurData);
@@ -539,11 +539,11 @@ private:
 				bool bTerminate = true;
 				for (int Axis = 0; Axis < d; ++Axis)
 				{
-					if (!CurData.bParallel[Axis])
+					if (!bParallel[Axis])
 					{
-						const T CrossPoint = CurData.Dir[Axis] > 0 ? CellCenter[Axis] + Dx[Axis] / 2 : CellCenter[Axis] - Dx[Axis] / 2;
+						const T CrossPoint = Dir[Axis] > 0 ? CellCenter[Axis] + Dx[Axis] / 2 : CellCenter[Axis] - Dx[Axis] / 2;
 						const T Distance = CrossPoint - NextStart[Axis];	//note: CellCenter already has /2, we probably want to use the corner instead
-						const T Time = Distance * CurData.InvDir[Axis];
+						const T Time = Distance * InvDir[Axis];
 						Times[Axis] = Time;
 						if (Time < BestTime)
 						{
@@ -565,14 +565,14 @@ private:
 				for (int Axis = 0; Axis < d; ++Axis)
 				{
 					constexpr T Epsilon = 1e-2f;	//if raycast is slightly off we still count it as hitting the cell surface
-					CellIdx[Axis] += (Times[Axis] <= BestTime + Epsilon) ? (CurData.Dir[Axis] > 0 ? 1 : -1) : 0;
+					CellIdx[Axis] += (Times[Axis] <= BestTime + Epsilon) ? (Dir[Axis] > 0 ? 1 : -1) : 0;
 					if (CellIdx[Axis] < 0 || CellIdx[Axis] >= MGrid.Counts()[Axis])
 					{
 						return true;
 					}
 				}
 
-				NextStart = NextStart + CurData.Dir * BestTime;
+				NextStart = NextStart + Dir * BestTime;
 			} while (true);
 		}
 
@@ -580,7 +580,7 @@ private:
 	}
 
 	template <typename SQVisitor, bool bPruneDuplicates = true>
-	bool SweepImp(const TVector<T, d>& Start, FQueryFastData& CurData, const TVector<T, d> QueryHalfExtents, SQVisitor& Visitor) const
+	bool SweepImp(const TVector<T, d>& Start, FQueryFastData& CurData, const TVector<T, d> QueryHalfExtents, SQVisitor& Visitor, const FVec3& Dir, const FVec3 InvDir, const bool bParallel[3]) const
 	{
 		T TOI = 0;
 		const void* QueryData = Visitor.GetQueryData();
@@ -595,7 +595,7 @@ private:
 			const TVector<T, d> Min = InstanceBounds.Min() - QueryHalfExtents;
 			const TVector<T, d> Max = InstanceBounds.Max() + QueryHalfExtents;
 			TVector<T, d> TmpPosition;
-			if (TAABB<T, d>(Min,Max).RaycastFast(Start, CurData.Dir, CurData.InvDir, CurData.bParallel, CurData.CurrentLength, CurData.InvCurrentLength, TOI, TmpPosition))
+			if (TAABB<T, d>(Min,Max).RaycastFast(Start, Dir, InvDir, bParallel, CurData.CurrentLength, CurData.InvCurrentLength, TOI, TmpPosition))
 			{
 				TSpatialVisitorData<TPayloadType> VisitData(Elem.Payload, true, InstanceBounds);
 				const bool bContinue = Visitor.VisitSweep(VisitData, CurData);
@@ -617,7 +617,7 @@ private:
 			const TVector<T, d> Min = InstanceBounds.Min() - QueryHalfExtents;
 			const TVector<T, d> Max = InstanceBounds.Max() + QueryHalfExtents;
 			TVector<T, d> TmpPosition;
-			if (TAABB<T, d>(Min, Max).RaycastFast(Start, CurData.Dir, CurData.InvDir, CurData.bParallel, CurData.CurrentLength, CurData.InvCurrentLength, TOI, TmpPosition))
+			if (TAABB<T, d>(Min, Max).RaycastFast(Start, Dir, InvDir, bParallel, CurData.CurrentLength, CurData.InvCurrentLength, TOI, TmpPosition))
 			{
 				TSpatialVisitorData<TPayloadType> VisitData(Elem.Payload, true, InstanceBounds);
 				const bool bContinue = Visitor.VisitSweep(VisitData, CurData);
@@ -647,7 +647,7 @@ private:
 
 		if (StartMinIndex == StartMaxIndex)
 		{
-			const TVector<T, d> End = Start + CurData.CurrentLength * CurData.Dir;
+			const TVector<T, d> End = Start + CurData.CurrentLength * Dir;
 			const TVector<int32, d> EndMinIndex = MGrid.ClampIndex(MGrid.Cell(End  - QueryHalfExtents));
 			const TVector<int32, d> EndMaxIndex = MGrid.ClampIndex(MGrid.Cell(End + QueryHalfExtents));
 			if (StartMinIndex == EndMinIndex && StartMinIndex == EndMaxIndex)
@@ -660,7 +660,7 @@ private:
 					const TAABB<T, d>& InstanceBounds = Elem.Bounds;
 					const TVector<T, d> Min = InstanceBounds.Min() - QueryHalfExtents;
 					const TVector<T, d> Max = InstanceBounds.Max() + QueryHalfExtents;
-					if (TAABB<T, d>(Min, Max).RaycastFast(Start, CurData.Dir, CurData.InvDir, CurData.bParallel, CurData.CurrentLength, CurData.InvCurrentLength, TOI, TmpPoint))
+					if (TAABB<T, d>(Min, Max).RaycastFast(Start, Dir, InvDir, bParallel, CurData.CurrentLength, CurData.InvCurrentLength, TOI, TmpPoint))
 					{
 						TSpatialVisitorData<TPayloadType> VisitData(Elem.Payload, true, InstanceBounds);
 						const bool bContinue = Visitor.VisitSweep(VisitData, CurData);
@@ -693,7 +693,7 @@ private:
 		}
 
 		TVector<T, d> HitPoint;
-		const bool bInitialHit = GlobalBounds.RaycastFast(Start, CurData.Dir, CurData.InvDir, CurData.bParallel, CurData.CurrentLength, CurData.InvCurrentLength, TOI, HitPoint);
+		const bool bInitialHit = GlobalBounds.RaycastFast(Start, Dir, InvDir, bParallel, CurData.CurrentLength, CurData.InvCurrentLength, TOI, HitPoint);
 		if (bInitialHit)	//NOTE: it's possible to have a non empty IdxsQueue and bInitialHit be false. This is because the IdxsQueue works off clamped cells which we can skip
 		{
 			//Flood fill from inflated cell so that we get all cells along the ray
@@ -756,7 +756,7 @@ private:
 						const TVector<T, d> NeighborCenter = MGrid.Location(NeighborIdx);
 						const TVector<T, d> Min = NeighborCenter - QueryHalfExtents - HalfDx;
 						const TVector<T, d> Max = NeighborCenter + QueryHalfExtents + HalfDx;
-						if (TAABB<T, d>(Min, Max).RaycastFast(Start, CurData.Dir, CurData.InvDir, CurData.bParallel, CurData.CurrentLength, CurData.InvCurrentLength, TOI, HitPoint))
+						if (TAABB<T, d>(Min, Max).RaycastFast(Start, Dir, InvDir, bParallel, CurData.CurrentLength, CurData.InvCurrentLength, TOI, HitPoint))
 						{
 							IdxsQueue.Add({ NeighborIdx, TOI });	//should we sort by TOI?
 						}
@@ -798,7 +798,7 @@ private:
 					const TVector<T, d> Min = InstanceBounds.Min() - QueryHalfExtents;
 					const TVector<T, d> Max = InstanceBounds.Max() + QueryHalfExtents;
 					if (TAABB<T, d>(Min,Max).RaycastFast(Start,
-						CurData.Dir, CurData.InvDir, CurData.bParallel, CurData.CurrentLength, CurData.InvCurrentLength, TOI, HitPoint))
+							Dir, InvDir, bParallel, CurData.CurrentLength, CurData.InvCurrentLength, TOI, HitPoint))
 					{
 						TSpatialVisitorData<TPayloadType> VisitData(Elem.Payload, true, InstanceBounds);
 						const bool bContinue = Visitor.VisitSweep(VisitData, CurData);
diff --git a/Engine/Source/Runtime/Experimental/Chaos/Public/Chaos/SpatialAccelerationCollection.h b/Engine/Source/Runtime/Experimental/Chaos/Public/Chaos/SpatialAccelerationCollection.h
--- a/Engine/Source/Runtime/Experimental/Chaos/Public/Chaos/SpatialAccelerationCollection.h
+++ b/Engine/Source/Runtime/Experimental/Chaos/Public/Chaos/SpatialAccelerationCollection.h
@@ -182,12 +182,12 @@ template <int TypeIdx, int NumTypes, typename Tuple, typename TPayloadType, type
 struct TSpatialAccelerationCollectionHelper
 {
 	template <typename SQVisitor>
-	static bool RaycastFast(const Tuple& Types, const TVector<T, d>& Start, FQueryFastData& CurData, SQVisitor& Visitor)
+	static bool RaycastFast(const Tuple& Types, const TVector<T, d>& Start, FQueryFastData& CurData, SQVisitor& Visitor, const FVec3& Dir, const FVec3 InvDir, const bool bParallel[3])
 	{
 		const auto& Accelerations = GetAccelerationsPerType<TypeIdx>(Types).Objects;
 		for (const auto& Accelerator : Accelerations)
 		{
-			if (Accelerator && !Accelerator->RaycastFast(Start, CurData, Visitor))
+			if (Accelerator && !Accelerator->RaycastFast(Start, CurData, Visitor, Dir, InvDir, bParallel))
 			{
 				return false;
 			}
@@ -196,19 +196,19 @@ struct TSpatialAccelerationCollectionHelper
 		constexpr int NextType = TypeIdx + 1;
 		if (NextType < NumTypes)
 		{
-			return TSpatialAccelerationCollectionHelper<NextType < NumTypes ? NextType : 0, NumTypes, Tuple, TPayloadType, T, d>::RaycastFast(Types, Start, CurData, Visitor);
+			return TSpatialAccelerationCollectionHelper<NextType < NumTypes ? NextType : 0, NumTypes, Tuple, TPayloadType, T, d>::RaycastFast(Types, Start, CurData, Visitor, Dir, InvDir, bParallel);
 		}
 
 		return true;
 	}
 
 	template <typename SQVisitor>
-	static bool SweepFast(const Tuple& Types, const TVector<T, d>& Start, FQueryFastData& CurData, const TVector<T, d> QueryHalfExtents, SQVisitor& Visitor)
+	static bool SweepFast(const Tuple& Types, const TVector<T, d>& Start, FQueryFastData& CurData, const TVector<T, d> QueryHalfExtents, SQVisitor& Visitor, const FVec3& Dir, const FVec3 InvDir, const bool bParallel[3])
 	{
 		const auto& Accelerations = GetAccelerationsPerType<TypeIdx>(Types).Objects;
 		for (const auto& Accelerator : Accelerations)
 		{
-			if (Accelerator && !Accelerator->SweepFast(Start, CurData, QueryHalfExtents, Visitor))
+			if (Accelerator && !Accelerator->SweepFast(Start, CurData, QueryHalfExtents, Visitor, Dir, InvDir, bParallel))
 			{
 				return false;
 			}
@@ -217,7 +217,7 @@ struct TSpatialAccelerationCollectionHelper
 		constexpr int NextType = TypeIdx + 1;
 		if (NextType < NumTypes)
 		{
-			return TSpatialAccelerationCollectionHelper < NextType < NumTypes ? NextType : 0, NumTypes, Tuple, TPayloadType, T, d>::SweepFast(Types, Start, CurData, QueryHalfExtents, Visitor);
+			return TSpatialAccelerationCollectionHelper < NextType < NumTypes ? NextType : 0, NumTypes, Tuple, TPayloadType, T, d>::SweepFast(Types, Start, CurData, QueryHalfExtents, Visitor, Dir, InvDir, bParallel);
 		}
 
 		return true;
@@ -426,7 +426,7 @@ public:
 	void Raycast(const TVector<T, d>& Start, const TVector<T, d>& Dir, const T Length, SQVisitor& Visitor) const
 	{
 		FQueryFastData QueryFastData(Dir, Length);
-		TSpatialAccelerationCollectionHelper<0, NumTypes, decltype(Types), TPayloadType, T, d>::RaycastFast(Types, Start, QueryFastData, Visitor);
+		TSpatialAccelerationCollectionHelper<0, NumTypes, decltype(Types), TPayloadType, T, d>::RaycastFast(Types, Start, QueryFastData, Visitor, QueryFastData.Dir, QueryFastData.InvDir, QueryFastData.bParallel);
 	}
 
 	void Sweep(const TVector<T, d>& Start, const TVector<T, d>& Dir, const T Length, const TVector<T, d> QueryHalfExtents, ISpatialVisitor<TPayloadType, T>& Visitor) const override
@@ -439,7 +439,7 @@ public:
 	void Sweep(const TVector<T, d>& Start, const TVector<T, d>& Dir, const T Length, const TVector<T, d> QueryHalfExtents, SQVisitor& Visitor) const
 	{
 		FQueryFastData QueryFastData(Dir, Length);
-		TSpatialAccelerationCollectionHelper<0, NumTypes, decltype(Types), TPayloadType, T, d>::SweepFast(Types, Start, QueryFastData, QueryHalfExtents, Visitor);
+		TSpatialAccelerationCollectionHelper<0, NumTypes, decltype(Types), TPayloadType, T, d>::SweepFast(Types, Start, QueryFastData, QueryHalfExtents, Visitor, QueryFastData.Dir, QueryFastData.InvDir, QueryFastData.bParallel);
 	}
 
 	virtual void Overlap(const TAABB<T, d>& QueryBounds, ISpatialVisitor<TPayloadType, T>& Visitor) const override
