Chaos : add support for geometry collection scaling ( uniform scaling )

diff --git a/Engine/Source/Runtime/Experimental/Chaos/Private/Chaos/ChaosDebugDraw.cpp b/Engine/Source/Runtime/Experimental/Chaos/Private/Chaos/ChaosDebugDraw.cpp
--- a/Engine/Source/Runtime/Experimental/Chaos/Private/Chaos/ChaosDebugDraw.cpp
+++ b/Engine/Source/Runtime/Experimental/Chaos/Private/Chaos/ChaosDebugDraw.cpp
@@ -350,7 +350,13 @@ namespace Chaos
 			case ImplicitObjectType::Union:
 				break;
 			case ImplicitObjectType::LevelSet:
+			{
+				const TImplicitObjectScaled<FLevelSet, bInstanced>* Scaled = Shape->template GetObject<TImplicitObjectScaled<FLevelSet, bInstanced>>();
+				// even though thhe levelset is scaled, the debugdraw uses the collisionParticles  that are pre-scaled
+				// so no need to pass the scaled transform and just extract the wrapped LevelSet
+				DrawShapesImpl(Particle, ShapeTransform, Scaled->GetUnscaledObject(), Scaled->GetMargin(), Color, Settings);
 				break;
+			}
 			case ImplicitObjectType::Unknown:
 				break;
 			case ImplicitObjectType::Convex:
diff --git a/Engine/Source/Runtime/Experimental/Chaos/Private/Chaos/Convex.cpp b/Engine/Source/Runtime/Experimental/Chaos/Private/Chaos/Convex.cpp
--- a/Engine/Source/Runtime/Experimental/Chaos/Private/Chaos/Convex.cpp
+++ b/Engine/Source/Runtime/Experimental/Chaos/Private/Chaos/Convex.cpp
@@ -2,11 +2,18 @@
 #include "Chaos/Convex.h"
 #include "Chaos/GJK.h"
 #include "Chaos/Sphere.h"
+#include "Chaos/ImplicitObjectScaled.h"
 
 //PRAGMA_DISABLE_OPTIMIZATION
 
 namespace Chaos
 {
+	TUniquePtr<FImplicitObject> FConvex::CopyWithScale(const FVec3& Scale) const
+	{
+		TSharedPtr<FConvex, ESPMode::ThreadSafe> ConvexCopy = MakeShared<FConvex, ESPMode::ThreadSafe>(GetVertices(), GetMargin());
+		return  TUniquePtr<FImplicitObject>(new TImplicitObjectScaled<FConvex>(ConvexCopy, Scale));
+	}
+
 	bool FConvex::Raycast(const FVec3& StartPoint, const FVec3& Dir, const FReal Length, const FReal Thickness, FReal& OutTime, FVec3& OutPosition, FVec3& OutNormal, int32& OutFaceIndex) const
 	{
 		OutFaceIndex = INDEX_NONE;	//finding face is expensive, should be called directly by user
diff --git a/Engine/Source/Runtime/Experimental/Chaos/Private/Chaos/ImplicitObject.cpp b/Engine/Source/Runtime/Experimental/Chaos/Private/Chaos/ImplicitObject.cpp
--- a/Engine/Source/Runtime/Experimental/Chaos/Private/Chaos/ImplicitObject.cpp
+++ b/Engine/Source/Runtime/Experimental/Chaos/Private/Chaos/ImplicitObject.cpp
@@ -81,6 +81,12 @@ TUniquePtr<FImplicitObject> FImplicitObject::Copy() const
 	return nullptr;
 }
 
+TUniquePtr<FImplicitObject> FImplicitObject::CopyWithScale(const FVec3& Scale) const
+{
+	check(false);
+	return nullptr;
+}
+
 bool FImplicitObject::IsUnderlyingUnion() const
 {
 	return Type == ImplicitObjectType::Union;
diff --git a/Engine/Source/Runtime/Experimental/Chaos/Private/Chaos/Levelset.cpp b/Engine/Source/Runtime/Experimental/Chaos/Private/Chaos/Levelset.cpp
--- a/Engine/Source/Runtime/Experimental/Chaos/Private/Chaos/Levelset.cpp
+++ b/Engine/Source/Runtime/Experimental/Chaos/Private/Chaos/Levelset.cpp
@@ -170,7 +170,9 @@ FLevelSet::FLevelSet(FLevelSet&& Other)
     : FImplicitObject(EImplicitObject::HasBoundingBox, ImplicitObjectType::LevelSet)
     , MGrid(MoveTemp(Other.MGrid))
     , MPhi(MoveTemp(Other.MPhi))
+	, MNormals(MoveTemp(Other.MNormals))
     , MLocalBoundingBox(MoveTemp(Other.MLocalBoundingBox))
+	, MOriginalLocalBoundingBox(MoveTemp(Other.MOriginalLocalBoundingBox))
     , MBandWidth(Other.MBandWidth)
 {
 }
@@ -191,6 +193,19 @@ TUniquePtr<FImplicitObject> FLevelSet::DeepCopy() const
 	return TUniquePtr<FImplicitObject>(Copy);
 }
 
+TUniquePtr<FImplicitObject> FLevelSet::DeepCopyWithScale(const FVec3& Scale) const
+{
+	FLevelSet Copy;
+	Copy.MGrid = MGrid;
+	Copy.MPhi.Copy(MPhi);
+	Copy.MNormals.Copy(MNormals);
+	Copy.MLocalBoundingBox = MLocalBoundingBox;
+	Copy.MOriginalLocalBoundingBox = MOriginalLocalBoundingBox;
+	Copy.MBandWidth = MBandWidth;
+	TSharedPtr<FLevelSet, ESPMode::ThreadSafe> LevelSetCopy = MakeShared<FLevelSet, ESPMode::ThreadSafe>(MoveTemp(Copy));
+	return TUniquePtr<FImplicitObject>(new TImplicitObjectScaled<FLevelSet>(LevelSetCopy, Scale));
+}
+
 bool FLevelSet::ComputeMassProperties(FReal& OutVolume, FVec3& OutCOM, FMatrix33& OutInertia, FRotation3& OutRotationOfMass) const
 {
 	FVec3 COM(0);
diff --git a/Engine/Source/Runtime/Experimental/Chaos/Private/Chaos/PBDRigidClustering.cpp b/Engine/Source/Runtime/Experimental/Chaos/Private/Chaos/PBDRigidClustering.cpp
--- a/Engine/Source/Runtime/Experimental/Chaos/Private/Chaos/PBDRigidClustering.cpp
+++ b/Engine/Source/Runtime/Experimental/Chaos/Private/Chaos/PBDRigidClustering.cpp
@@ -1772,9 +1772,21 @@ namespace Chaos
 
 		if (ProxyGeometry)
 		{
+			const FVector Scale = Parameters.Scale;
+			auto DeepCopyImplicit = [&Scale](const TSharedPtr<Chaos::FImplicitObject, ESPMode::ThreadSafe>& ImplicitToCopy) -> TUniquePtr<Chaos::FImplicitObject>
+			{
+				if (Scale.Equals(FVector::OneVector))
+				{
+					return ImplicitToCopy->DeepCopy();
+				}
+				else
+				{
+					return ImplicitToCopy->DeepCopyWithScale(Scale);
+				}
+			};
 			//ensureMsgf(false, TEXT("Checking usage with proxy"));
 			//@coverage {production}
-			Parent->SetSharedGeometry(TSharedPtr<FImplicitObject, ESPMode::ThreadSafe>(ProxyGeometry->DeepCopy().Release()));
+			Parent->SetSharedGeometry(TSharedPtr<FImplicitObject, ESPMode::ThreadSafe>(DeepCopyImplicit(ProxyGeometry).Release()));
 		}
 		else if (Objects.Num() == 0)
 		{
diff --git a/Engine/Source/Runtime/Experimental/Chaos/Private/PhysicsProxy/GeometryCollectionPhysicsProxy.cpp b/Engine/Source/Runtime/Experimental/Chaos/Private/PhysicsProxy/GeometryCollectionPhysicsProxy.cpp
--- a/Engine/Source/Runtime/Experimental/Chaos/Private/PhysicsProxy/GeometryCollectionPhysicsProxy.cpp
+++ b/Engine/Source/Runtime/Experimental/Chaos/Private/PhysicsProxy/GeometryCollectionPhysicsProxy.cpp
@@ -17,6 +17,8 @@
 #include "Chaos/PBDCollisionConstraintsUtil.h"
 #include "Chaos/PerParticleGravity.h"
 #include "Chaos/ImplicitObject.h"
+#include "Chaos/ImplicitObjectScaled.h"
+#include "Chaos/Convex.h"
 #include "Chaos/Serializable.h"
 #include "Chaos/ErrorReporter.h"
 #include "Chaos/PBDRigidClustering.h"
@@ -147,7 +149,8 @@ void PopulateSimulatedParticle(
 	FVector3f InertiaTensorVec,
 	const FTransform& WorldTransform, 
 	const uint8 DynamicState, 
-	const int16 CollisionGroup)
+	const int16 CollisionGroup,
+	float CollisionParticlesPerObjectFraction)
 {
 	SCOPE_CYCLE_COUNTER(STAT_PopulateSimulatedParticle);
 	Handle->SetDisabledLowLevel(false);
@@ -202,10 +205,22 @@ void PopulateSimulatedParticle(
 	{
 		SingleSupportedCollisionTypeData = SharedParams.SizeSpecificData[0].CollisionShapesData[0];
 	}
-
+	const FVector Scale = WorldTransform.GetScale3D();
 	if (Implicit)	//todo(ocohen): this is only needed for cases where clusters have no proxy. Kind of gross though, should refactor
 	{
-		TSharedPtr<Chaos::FImplicitObject, ESPMode::ThreadSafe> SharedImplicitTS(Implicit->DeepCopy().Release());
+		auto DeepCopyImplicit = [&Scale](FGeometryDynamicCollection::FSharedImplicit ImplicitToCopy) -> TUniquePtr<Chaos::FImplicitObject>
+		{
+			if (Scale.Equals(FVector::OneVector))
+			{
+				return ImplicitToCopy->DeepCopy();
+			}
+			else
+			{
+				return ImplicitToCopy->DeepCopyWithScale(Scale);
+			}
+		};
+
+		TSharedPtr<Chaos::FImplicitObject, ESPMode::ThreadSafe> SharedImplicitTS(DeepCopyImplicit(Implicit).Release());
 		FCollisionStructureManager::UpdateImplicitFlags(SharedImplicitTS.Get(), SingleSupportedCollisionTypeData.CollisionType);
 		Handle->SetSharedGeometry(SharedImplicitTS);
 		Handle->SetHasBounds(true);
@@ -221,22 +236,27 @@ void PopulateSimulatedParticle(
 		Handle->CollisionParticlesInitIfNeeded();
 
 		TUniquePtr<Chaos::FBVHParticles>& CollisionParticles = Handle->CollisionParticles();
-		if (Simplicial->Size())
-		{
-			const Chaos::FAABB3 ImplicitShapeDomain = 
-				Implicit && Implicit->GetType() == Chaos::ImplicitObjectType::LevelSet && Implicit->HasBoundingBox() ? 
-				Implicit->BoundingBox() : Chaos::FAABB3::FullAABB();
+		CollisionParticles.Reset(Simplicial->NewCopy()); // @chaos(optimize) : maybe just move this memory instead. 
 
-			CollisionParticles->Resize(0);
-			CollisionParticles->AddParticles(Simplicial->Size());
-			for (int32 VertexIndex = 0; VertexIndex < (int32)Simplicial->Size(); ++VertexIndex)
-			{
-				CollisionParticles->X(VertexIndex) = Simplicial->X(VertexIndex);
+		int32 NumCollisionParticles = CollisionParticles->Size();
+		int32 CollisionParticlesSize = FMath::Max(0, FMath::Min(int(NumCollisionParticles * CollisionParticlesPerObjectFraction), NumCollisionParticles));
+		CollisionParticles->Resize(CollisionParticlesSize); // Truncates! ( particles are already sorted by importance )
 
-				// Make sure the collision particles are at least in the domain 
-				// of the implicit shape.
-				ensure(ImplicitShapeDomain.Contains(CollisionParticles->X(VertexIndex)));
-			}
+		Chaos::FAABB3 ImplicitShapeDomain = Chaos::FAABB3::FullAABB();
+		if (Implicit && Implicit->GetType() == Chaos::ImplicitObjectType::LevelSet && Implicit->HasBoundingBox())
+		{
+			ImplicitShapeDomain = Implicit->BoundingBox();
+			ImplicitShapeDomain.Scale(Scale);
+		}
+
+		// we need to account for scale and check if the particle is still within its domain
+		for (int32 ParticleIndex = 0; ParticleIndex < (int32)CollisionParticles->Size(); ++ParticleIndex)
+		{
+			CollisionParticles->X(ParticleIndex) *= Scale;
+			
+			// Make sure the collision particles are at least in the domain 
+			// of the implicit shape.
+			ensure(ImplicitShapeDomain.Contains(CollisionParticles->X(ParticleIndex)));
 		}
 
 		// @todo(remove): IF there is no simplicial we should not be forcing one. 
@@ -645,7 +665,8 @@ void FGeometryCollectionPhysicsProxy::InitializeBodiesPT(Chaos::FPBDRigidsSolver
 					InertiaTensor[TransformGroupIndex],
 					WorldTransform,
 					static_cast<uint8>(DynamicState[TransformGroupIndex]),
-					static_cast<int16>(CollisionGroup[TransformGroupIndex]));
+					static_cast<int16>(CollisionGroup[TransformGroupIndex]),
+					CollisionParticlesPerObjectFraction);
 
 				if (Parameters.EnableClustering)
 				{
@@ -653,15 +674,6 @@ void FGeometryCollectionPhysicsProxy::InitializeBodiesPT(Chaos::FPBDRigidsSolver
 					Handle->SetStrain(StrainDefault);
 				}
 
-				TUniquePtr<Chaos::FBVHParticles>& CollisionParticles = Handle->CollisionParticles();
-				CollisionParticles.Reset(Simplicials[TransformGroupIndex]?Simplicials[TransformGroupIndex]->NewCopy():nullptr); // @chaos(optimize) : maybe just move this memory instead. 
-				if (CollisionParticles)
-				{
-					Chaos::FReal NumCollisionParticles = static_cast<Chaos::FReal>(CollisionParticles->Size());
-					int32 CollisionParticlesSize = FMath::TruncToInt(Chaos::FReal(FMath::Max((Chaos::FReal)0, FMath::Min(Chaos::FReal(NumCollisionParticles * CollisionParticlesPerObjectFraction), NumCollisionParticles))));
-					CollisionParticles->Resize(CollisionParticlesSize); // Truncates!
-				}
-
 				// #BGTODO - non-updating parameters - remove lin/ang drag arrays and always query material if this stays a material parameter
 				Chaos::FChaosPhysicsMaterial* SolverMaterial = RigidsSolver->GetSimMaterials().Get(Parameters.PhysicalMaterialHandle.InnerHandle);
 				if(SolverMaterial)
@@ -794,6 +806,7 @@ void FGeometryCollectionPhysicsProxy::InitializeBodiesPT(Chaos::FPBDRigidsSolver
 
 					Chaos::FClusterCreationParameters CreationParameters;
 					CreationParameters.ClusterParticleHandle = ClusterHandles.Num() ? ClusterHandles[ClusterHandlesIndex++] : nullptr;
+					CreationParameters.Scale = Parameters.WorldTransform.GetScale3D();
 
 					// Hook the handle up with the GT particle
 					Chaos::FGeometryParticle* GTParticle = GTParticles[TransformGroupIndex].Get();
@@ -1017,7 +1030,8 @@ FGeometryCollectionPhysicsProxy::BuildClusters(
 		Parent->I().GetDiagonal() != Chaos::FVec3(0.0) ? Parent->I().GetDiagonal() : Chaos::FVec3(InertiaTensor[CollectionClusterIndex]),
 		ParticleTM, 
 		(uint8)DynamicState[CollectionClusterIndex], 
-		0); // CollisionGroup
+		0,
+		CollisionParticlesPerObjectFraction); // CollisionGroup
 
 	// two-way mapping
 	SolverClusterHandles[CollectionClusterIndex] = Parent;
@@ -1539,6 +1553,8 @@ void FGeometryCollectionPhysicsProxy::BufferPhysicsResults(Chaos::FPBDRigidsSolv
 	const FTransform& ActorToWorld = Parameters.WorldTransform;
 	const TManagedArray<int32>& Parent = PhysicsThreadCollection.Parent;
 	const TManagedArray<TSet<int32>>& Children = PhysicsThreadCollection.Children;
+	const bool IsActorScaled = !ActorToWorld.GetScale3D().Equals(FVector::OneVector);
+	const FTransform ActorScaleTransform(FQuat::Identity,  FVector::ZeroVector, ActorToWorld.GetScale3D());
 
 	if(NumTransformGroupElements > 0)
 	{ 
@@ -1605,6 +1621,10 @@ void FGeometryCollectionPhysicsProxy::BufferPhysicsResults(Chaos::FPBDRigidsSolv
 
 				TargetResults.Transforms[TransformGroupIndex] = MassToLocal.GetRelativeTransformReverse(ParticleToWorld).GetRelativeTransform(ActorToWorld);
 				TargetResults.Transforms[TransformGroupIndex].NormalizeRotation();
+				if (IsActorScaled)
+				{
+					TargetResults.Transforms[TransformGroupIndex] = MassToLocal.Inverse() * ActorScaleTransform * MassToLocal * TargetResults.Transforms[TransformGroupIndex];
+				}
 
 				PhysicsThreadCollection.Transform[TransformGroupIndex] = TargetResults.Transforms[TransformGroupIndex];
 
@@ -1672,6 +1692,10 @@ void FGeometryCollectionPhysicsProxy::BufferPhysicsResults(Chaos::FPBDRigidsSolv
 							const FTransform MassToLocal                  = PhysicsThreadCollection.MassToLocal[TransformGroupIndex];
 							TargetResults.Transforms[TransformGroupIndex] = MassToLocal.GetRelativeTransformReverse(ParticleToWorld).GetRelativeTransform(ActorToWorld);
 							TargetResults.Transforms[TransformGroupIndex].NormalizeRotation();
+							if (IsActorScaled)
+							{
+								TargetResults.Transforms[TransformGroupIndex] = MassToLocal.Inverse() * ActorScaleTransform * MassToLocal * TargetResults.Transforms[TransformGroupIndex];
+							}
 
 							PhysicsThreadCollection.Transform[TransformGroupIndex] = TargetResults.Transforms[TransformGroupIndex];
 
@@ -2301,7 +2325,9 @@ void FGeometryCollectionPhysicsProxy::InitializeSharedCollisionStructures(
  					CollectionInertiaTensor[TransformGroupIndex], 
 					GeometryWorldTransform,
  					(uint8)EObjectStateTypeEnum::Chaos_Object_Dynamic, 
- 					INDEX_NONE); // CollisionGroup
+ 					INDEX_NONE,  // CollisionGroup
+					1.0f // todo(chaos) CollisionParticlesPerObjectFraction is not accessible right there for now but we can pass 1.0 for the time being
+				);
  			}
  		}
 
diff --git a/Engine/Source/Runtime/Experimental/Chaos/Public/Chaos/BVHParticles.h b/Engine/Source/Runtime/Experimental/Chaos/Public/Chaos/BVHParticles.h
--- a/Engine/Source/Runtime/Experimental/Chaos/Public/Chaos/BVHParticles.h
+++ b/Engine/Source/Runtime/Experimental/Chaos/Public/Chaos/BVHParticles.h
@@ -35,7 +35,7 @@ namespace Chaos
 	    CHAOS_API FBVHParticles& operator=(const FBVHParticles& Other);
 	    CHAOS_API FBVHParticles& operator=(FBVHParticles&& Other);
 
-		CHAOS_API FBVHParticles* NewCopy()
+		CHAOS_API FBVHParticles* NewCopy() const
 		{
 			return new FBVHParticles(*this);
 		}
diff --git a/Engine/Source/Runtime/Experimental/Chaos/Public/Chaos/Box.h b/Engine/Source/Runtime/Experimental/Chaos/Public/Chaos/Box.h
--- a/Engine/Source/Runtime/Experimental/Chaos/Public/Chaos/Box.h
+++ b/Engine/Source/Runtime/Experimental/Chaos/Public/Chaos/Box.h
@@ -100,6 +100,11 @@ namespace Chaos
 			return TUniquePtr<FImplicitObject>(new TBox<T,d>(*this));
 		}
 
+		virtual TUniquePtr<FImplicitObject> CopyWithScale(const FVec3& Scale) const override
+		{
+			return TUniquePtr<FImplicitObject>(new TBox<T, d>(AABB.Min() * Scale, AABB.Max() * Scale, Margin * Scale.Min()));
+		}
+
 		FReal GetRadius() const
 		{
 			return 0.0f;
diff --git a/Engine/Source/Runtime/Experimental/Chaos/Public/Chaos/Capsule.h b/Engine/Source/Runtime/Experimental/Chaos/Public/Chaos/Capsule.h
--- a/Engine/Source/Runtime/Experimental/Chaos/Public/Chaos/Capsule.h
+++ b/Engine/Source/Runtime/Experimental/Chaos/Public/Chaos/Capsule.h
@@ -392,6 +392,11 @@ namespace Chaos
 			return TUniquePtr<FImplicitObject>(new FCapsule(*this));
 		}
 
+		virtual TUniquePtr<FImplicitObject> CopyWithScale(const FVec3& Scale) const override
+		{
+			return  TUniquePtr<FImplicitObject>(new FCapsule(GetX1() * Scale, GetX2() * Scale, GetRadius() * Scale.Min()));
+		}
+
 		FReal GetHeight() const { return MSegment.GetLength(); }
 		/** Returns the bottommost point on the capsule. */
 		const FVec3 GetOrigin() const { return GetX1() + GetAxis() * -GetRadius(); }
diff --git a/Engine/Source/Runtime/Experimental/Chaos/Public/Chaos/ClusterCreationParameters.h b/Engine/Source/Runtime/Experimental/Chaos/Public/Chaos/ClusterCreationParameters.h
--- a/Engine/Source/Runtime/Experimental/Chaos/Public/Chaos/ClusterCreationParameters.h
+++ b/Engine/Source/Runtime/Experimental/Chaos/Public/Chaos/ClusterCreationParameters.h
@@ -30,6 +30,7 @@ namespace Chaos
 			, bool bGenerateConnectionGraphIn = true, EConnectionMethod ConnectionMethodIn = EConnectionMethod::MinimalSpanningSubsetDelaunayTriangulation
 			, FBVHParticles* CollisionParticlesIn = nullptr
 			, Chaos::TPBDRigidClusteredParticleHandle<Chaos::FReal,3>* ClusterParticleHandleIn = nullptr
+			, const FVec3& ScaleIn = FVec3::OneVector
 		)
 			: CoillisionThicknessPercent(CoillisionThicknessPercentIn)
 			, MaxNumConnections(MaxNumConnectionsIn)
@@ -39,6 +40,7 @@ namespace Chaos
 			, ConnectionMethod(ConnectionMethodIn)
 			, CollisionParticles(CollisionParticlesIn)
 			, ClusterParticleHandle(ClusterParticleHandleIn)
+			, Scale(ScaleIn)
 		{}
 
 		FReal CoillisionThicknessPercent;
@@ -49,5 +51,6 @@ namespace Chaos
 		EConnectionMethod ConnectionMethod;
 		FBVHParticles* CollisionParticles;
 		Chaos::FPBDRigidClusteredParticleHandle* ClusterParticleHandle;
+		FVec3 Scale;
 	};
 }
diff --git a/Engine/Source/Runtime/Experimental/Chaos/Public/Chaos/Convex.h b/Engine/Source/Runtime/Experimental/Chaos/Public/Chaos/Convex.h
--- a/Engine/Source/Runtime/Experimental/Chaos/Public/Chaos/Convex.h
+++ b/Engine/Source/Runtime/Experimental/Chaos/Public/Chaos/Convex.h
@@ -151,6 +151,7 @@ namespace Chaos
 			return TUniquePtr<FImplicitObject>(new FConvex(GetVertices(), GetMargin()));
 		}
 
+		virtual TUniquePtr<FImplicitObject> CopyWithScale(const FVec3& Scale) const override;
 
 		void MovePlanesAndRebuild(FReal InDelta);
 
diff --git a/Engine/Source/Runtime/Experimental/Chaos/Public/Chaos/ImplicitObject.h b/Engine/Source/Runtime/Experimental/Chaos/Public/Chaos/ImplicitObject.h
--- a/Engine/Source/Runtime/Experimental/Chaos/Public/Chaos/ImplicitObject.h
+++ b/Engine/Source/Runtime/Experimental/Chaos/Public/Chaos/ImplicitObject.h
@@ -192,7 +192,9 @@ public:
 	virtual bool IsValidGeometry() const;
 
 	virtual TUniquePtr<FImplicitObject> Copy() const;
+	virtual TUniquePtr<FImplicitObject> CopyWithScale(const FVec3& Scale) const;
 	virtual TUniquePtr<FImplicitObject> DeepCopy() const { return Copy(); }
+	virtual TUniquePtr<FImplicitObject> DeepCopyWithScale(const FVec3& Scale) const { return CopyWithScale(Scale); }
 
 	//This is strictly used for optimization purposes
 	bool IsUnderlyingUnion() const;
diff --git a/Engine/Source/Runtime/Experimental/Chaos/Public/Chaos/Levelset.h b/Engine/Source/Runtime/Experimental/Chaos/Public/Chaos/Levelset.h
--- a/Engine/Source/Runtime/Experimental/Chaos/Public/Chaos/Levelset.h
+++ b/Engine/Source/Runtime/Experimental/Chaos/Public/Chaos/Levelset.h
@@ -37,6 +37,7 @@ class CHAOS_API FLevelSet final : public FImplicitObject
 	virtual ~FLevelSet();
 
 	virtual TUniquePtr<FImplicitObject> DeepCopy() const;
+	virtual TUniquePtr<FImplicitObject> DeepCopyWithScale(const FVec3& Scale) const;
 
 	void Write(std::ostream& Stream) const;
 	virtual FReal PhiWithNormal(const FVec3& x, FVec3& Normal) const override;
diff --git a/Engine/Source/Runtime/Experimental/Chaos/Public/Chaos/Sphere.h b/Engine/Source/Runtime/Experimental/Chaos/Public/Chaos/Sphere.h
--- a/Engine/Source/Runtime/Experimental/Chaos/Public/Chaos/Sphere.h
+++ b/Engine/Source/Runtime/Experimental/Chaos/Public/Chaos/Sphere.h
@@ -351,6 +351,11 @@ namespace Chaos
 			return TUniquePtr<FImplicitObject>(new TSphere<T,d>(Center, GetRadius()));
 		}
 
+		virtual TUniquePtr<FImplicitObject> CopyWithScale(const FVec3& Scale) const override
+		{
+			return  TUniquePtr<FImplicitObject>(new TSphere<T, d>(Center * Scale, GetRadius() * Scale.Min()));
+		}
+
 	private:
 		void SetRadius(FReal InRadius) { SetMargin(InRadius); }
 
