Chaos : Fix compilation error related to int to float truncation

diff --git a/Engine/Source/Runtime/Experimental/Chaos/Private/PhysicsProxy/GeometryCollectionPhysicsProxy.cpp b/Engine/Source/Runtime/Experimental/Chaos/Private/PhysicsProxy/GeometryCollectionPhysicsProxy.cpp
--- a/Engine/Source/Runtime/Experimental/Chaos/Private/PhysicsProxy/GeometryCollectionPhysicsProxy.cpp
+++ b/Engine/Source/Runtime/Experimental/Chaos/Private/PhysicsProxy/GeometryCollectionPhysicsProxy.cpp
@@ -238,8 +238,9 @@ void PopulateSimulatedParticle(
 		TUniquePtr<Chaos::FBVHParticles>& CollisionParticles = Handle->CollisionParticles();
 		CollisionParticles.Reset(Simplicial->NewCopy()); // @chaos(optimize) : maybe just move this memory instead. 
 
-		int32 NumCollisionParticles = CollisionParticles->Size();
-		int32 CollisionParticlesSize = FMath::Max(0, FMath::Min(int(NumCollisionParticles * CollisionParticlesPerObjectFraction), NumCollisionParticles));
+		const int32 NumCollisionParticles = CollisionParticles->Size();
+		const int32 AdjustedNumCollisionParticles = FMath::TruncToInt(CollisionParticlesPerObjectFraction * NumCollisionParticles);
+		int32 CollisionParticlesSize = FMath::Max<int32>(0, FMath::Min<int32>(AdjustedNumCollisionParticles, NumCollisionParticles));
 		CollisionParticles->Resize(CollisionParticlesSize); // Truncates! ( particles are already sorted by importance )
 
 		Chaos::FAABB3 ImplicitShapeDomain = Chaos::FAABB3::FullAABB();
