Chaos : fix compilation issue

diff --git a/Engine/Source/Runtime/Experimental/Chaos/Private/PhysicsProxy/GeometryCollectionPhysicsProxy.cpp b/Engine/Source/Runtime/Experimental/Chaos/Private/PhysicsProxy/GeometryCollectionPhysicsProxy.cpp
--- a/Engine/Source/Runtime/Experimental/Chaos/Private/PhysicsProxy/GeometryCollectionPhysicsProxy.cpp
+++ b/Engine/Source/Runtime/Experimental/Chaos/Private/PhysicsProxy/GeometryCollectionPhysicsProxy.cpp
@@ -239,7 +239,7 @@ void PopulateSimulatedParticle(
 		CollisionParticles.Reset(Simplicial->NewCopy()); // @chaos(optimize) : maybe just move this memory instead. 
 
 		const int32 NumCollisionParticles = CollisionParticles->Size();
-		const int32 AdjustedNumCollisionParticles = FMath::TruncToInt(CollisionParticlesPerObjectFraction * NumCollisionParticles);
+		const int32 AdjustedNumCollisionParticles = FMath::TruncToInt(CollisionParticlesPerObjectFraction * (float)NumCollisionParticles);
 		int32 CollisionParticlesSize = FMath::Max<int32>(0, FMath::Min<int32>(AdjustedNumCollisionParticles, NumCollisionParticles));
 		CollisionParticles->Resize(CollisionParticlesSize); // Truncates! ( particles are already sorted by importance )
 
