Chaos

diff --git a/Engine/Source/Runtime/Engine/Private/PhysicsEngine/ImmediatePhysics/ImmediatePhysicsChaos/ImmediatePhysicsSimulation_Chaos.cpp b/Engine/Source/Runtime/Engine/Private/PhysicsEngine/ImmediatePhysics/ImmediatePhysicsChaos/ImmediatePhysicsSimulation_Chaos.cpp
--- a/Engine/Source/Runtime/Engine/Private/PhysicsEngine/ImmediatePhysics/ImmediatePhysicsChaos/ImmediatePhysicsSimulation_Chaos.cpp
+++ b/Engine/Source/Runtime/Engine/Private/PhysicsEngine/ImmediatePhysics/ImmediatePhysicsChaos/ImmediatePhysicsSimulation_Chaos.cpp
@@ -187,6 +187,7 @@ Chaos::DebugDraw::FChaosDebugDrawSettings ChaosImmPhysDebugDebugDrawSettings(
 	/* VelScale =					*/ 0.0f,
 	/* AngVelScale =				*/ 0.0f,
 	/* ImpulseScale =				*/ 0.0f,
+	/* PushOutScale =				*/ 0.0f,
 	/* InertiaScale =				*/ 0.0f,
 	/* DrawPriority =				*/ 10.0f,
 	/* bShowSimple =				*/ true,
@@ -210,6 +211,7 @@ FAutoConsoleVariableRef CVarChaosImmPhysLineShapeThickness(TEXT("p.Chaos.ImmPhys
 FAutoConsoleVariableRef CVarChaosImmPhysVelScale(TEXT("p.Chaos.ImmPhys.DebugDraw.VelScale"), ChaosImmPhysDebugDebugDrawSettings.VelScale, TEXT("If >0 show velocity when drawing particle transforms."));
 FAutoConsoleVariableRef CVarChaosImmPhysAngVelScale(TEXT("p.Chaos.ImmPhys.DebugDraw.AngVelScale"), ChaosImmPhysDebugDebugDrawSettings.AngVelScale, TEXT("If >0 show angular velocity when drawing particle transforms."));
 FAutoConsoleVariableRef CVarChaosImmPhysImpulseScale(TEXT("p.Chaos.ImmPhys.DebugDraw.ImpulseScale"), ChaosImmPhysDebugDebugDrawSettings.ImpulseScale, TEXT("If >0 show impulses when drawing collisions."));
+FAutoConsoleVariableRef CVarChaosImmPhysPushOutScale(TEXT("p.Chaos.ImmPhys.DebugDraw.PushOutScale"), ChaosImmPhysDebugDebugDrawSettings.PushOutScale, TEXT("If >0 show pushouts when drawing collisions."));
 FAutoConsoleVariableRef CVarChaosImmPhysScale(TEXT("p.Chaos.ImmPhys.DebugDraw.Scale"), ChaosImmPhysDebugDebugDrawSettings.DrawScale, TEXT("Scale applied to all Chaos Debug Draw line lengths etc."));
 #endif
 
@@ -347,6 +349,7 @@ namespace ImmediatePhysics_Chaos
 		Implementation->NarrowPhase.GetContext().bFilteringEnabled = false;
 		Implementation->NarrowPhase.GetContext().bDeferUpdate = true;
 		Implementation->NarrowPhase.GetContext().bAllowManifolds = false;
+		Implementation->NarrowPhase.GetContext().CollisionAllocator = &Implementation->Collisions.GetConstraintAllocator();
 
 
 #if CHAOS_DEBUG_DRAW
@@ -368,48 +371,6 @@ namespace ImmediatePhysics_Chaos
 				DebugDrawDynamicParticles(3, 3, 1.0f);
 				DebugDrawConstraints(3, 3, 1.0f);
 			});
-		Implementation->Collisions.SetPostApplyCallback(
-			[this](const FReal Dt, const TArray<FPBDCollisionConstraintHandle*>& InConstraintHandles)
-			{
-				if (ChaosImmediate_DebugDrawCollisions == 4)
-				{
-					DebugDraw::DrawCollisions(Implementation->SimulationSpace.Transform, Implementation->Collisions, 0.3f, &ChaosImmPhysDebugDebugDrawSettings);
-				}
-				DebugDrawDynamicParticles(4, 4, 0.5f);
-			});
-		Implementation->Collisions.SetPostApplyPushOutCallback(
-			[this](const FReal Dt, const TArray<FPBDCollisionConstraintHandle*>& InConstraintHandles, bool bRequiresAnotherIteration)
-			{
-				if (ChaosImmediate_DebugDrawCollisions == 4)
-				{
-					DebugDraw::DrawCollisions(Implementation->SimulationSpace.Transform, Implementation->Collisions, 0.6f, &ChaosImmPhysDebugDebugDrawSettings);
-				}
-			});
-		Implementation->Joints.SetPreApplyCallback(
-			[this](const FReal Dt, const TArray<FPBDJointConstraintHandle*>& InConstraintHandles)
-			{
-				if (ChaosImmediate_DebugDrawJoints == 4)
-				{
-					DebugDraw::DrawJointConstraints(Implementation->SimulationSpace.Transform, InConstraintHandles, 0.3f, DebugDraw::FChaosDebugDrawJointFeatures::MakeDefault(), &ChaosImmPhysDebugDebugDrawSettings);
-				}
-			});
-		Implementation->Joints.SetPostApplyCallback(
-			[this](const FReal Dt, const TArray<FPBDJointConstraintHandle*>& InConstraintHandles)
-			{
-				if (ChaosImmediate_DebugDrawJoints == 4)
-				{
-					DebugDraw::DrawJointConstraints(Implementation->SimulationSpace.Transform, InConstraintHandles, 0.6f, DebugDraw::FChaosDebugDrawJointFeatures::MakeDefault(), &ChaosImmPhysDebugDebugDrawSettings);
-				}
-				DebugDrawDynamicParticles(4, 4, 0.5f);
-			});
-		Implementation->Joints.SetPostProjectCallback(
-			[this](const FReal Dt, const TArray<FPBDJointConstraintHandle*>& InConstraintHandles)
-			{
-				if (ChaosImmediate_DebugDrawJoints == 4)
-				{
-					DebugDraw::DrawJointConstraints(Implementation->SimulationSpace.Transform, InConstraintHandles, 0.6f, DebugDraw::FChaosDebugDrawJointFeatures::MakeDefault(), &ChaosImmPhysDebugDebugDrawSettings);
-				}
-			});
 #endif
 	}
 
diff --git a/Engine/Source/Runtime/Experimental/Chaos/Private/Chaos/AABBTree.cpp b/Engine/Source/Runtime/Experimental/Chaos/Private/Chaos/AABBTree.cpp
--- a/Engine/Source/Runtime/Experimental/Chaos/Private/Chaos/AABBTree.cpp
+++ b/Engine/Source/Runtime/Experimental/Chaos/Private/Chaos/AABBTree.cpp
@@ -9,7 +9,7 @@ float FAABBTreeCVars::MaxNonGlobalElementBoundsExtrema = 1e15f;
 FAutoConsoleVariableRef FAABBTreeCVars::CVarMaxNonGlobalElementBoundsExtrema(TEXT("p.aabbtree.MaxNonGlobalElementBoundsExtrema"), FAABBTreeCVars::MaxNonGlobalElementBoundsExtrema, TEXT("If Min/Max of particle world bounds exceeds this threshold, element is put in global array instead of being used in leaf splitting."));
 
 
-int32 FAABBTreeDirtyGridCVars::DirtyElementGridCellSize = 0; //1000; 0 means disabled
+int32 FAABBTreeDirtyGridCVars::DirtyElementGridCellSize = 1000; // 0 means disabled
 FAutoConsoleVariableRef FAABBTreeDirtyGridCVars::CVarDirtyElementGridCellSize(TEXT("p.aabbtree.DirtyElementGridCellSize"), FAABBTreeDirtyGridCVars::DirtyElementGridCellSize, TEXT("DirtyElement Grid acceleration structure cell size in cm. 0 or less will disable the feature"));
 
 int32 FAABBTreeDirtyGridCVars::DirtyElementMaxGridCellQueryCount = 340;
@@ -23,3 +23,8 @@ FAutoConsoleVariableRef FAABBTreeDirtyGridCVars::CVarDirtyElementMaxCellCapacity
 
 CSV_DEFINE_CATEGORY(ChaosPhysicsTimers, true);
 
+int32 FAABBTreeCVars::SplitAtAverageCenter = 1;
+FAutoConsoleVariableRef FAABBTreeCVars::CVarSplitAtAverageCenter(TEXT("p.aabbtree.splitataveragecenter"), FAABBTreeCVars::SplitAtAverageCenter, TEXT("Split AABB tree nodes at the average of the element centers"));
+
+int32 FAABBTreeCVars::SplitOnVarianceAxis = 1;
+FAutoConsoleVariableRef FAABBTreeCVars::CVarSplitOnVarianceAxis(TEXT("p.aabbtree.splitonvarianceaxis"), FAABBTreeCVars::SplitOnVarianceAxis, TEXT("Split AABB tree nodes along the axis with the largest element center variance"));
diff --git a/Engine/Source/Runtime/Experimental/Chaos/Private/Chaos/ChaosDebugDraw.cpp b/Engine/Source/Runtime/Experimental/Chaos/Private/Chaos/ChaosDebugDraw.cpp
--- a/Engine/Source/Runtime/Experimental/Chaos/Private/Chaos/ChaosDebugDraw.cpp
+++ b/Engine/Source/Runtime/Experimental/Chaos/Private/Chaos/ChaosDebugDraw.cpp
@@ -26,6 +26,9 @@
 
 namespace Chaos
 {
+
+	extern bool bChaos_Collision_Manifold_FixNormalsInWorldSpace;
+
 	namespace DebugDraw
 	{
 #if CHAOS_DEBUG_DRAW
@@ -42,8 +45,11 @@ namespace Chaos
 		bool bChaosDebugDebugDrawColorShapesByShapeType = false;
 		FAutoConsoleVariableRef CVarChaosDebugDebugDrawColorShapesByShapeType(TEXT("p.Chaos.DebugDraw.ColorShapesByShapeType"), bChaosDebugDebugDrawColorShapesByShapeType, TEXT("Whether to use shape type to define the color of the shapes instead of using the particle state "));
 
+		bool bChaosDebugDebugDrawColorShapesByIsland = false;
+		FAutoConsoleVariableRef CVarChaosDebugDebugDrawColorShapesByIsland(TEXT("p.Chaos.DebugDraw.ColorShapesByIsland"), bChaosDebugDebugDrawColorShapesByIsland, TEXT("Whether to use particle island to define the color of the shapes instead of using the particle state "));
+
 		bool bChaosDebugDebugDrawColorBoundsByShapeType = false;
-		FAutoConsoleVariableRef CVarbChaosDebugDebugDrawColorBoundsByShapeType(TEXT("p.Chaos.DebugDraw.ColorBoundsByShapeType"), bChaosDebugDebugDrawColorBoundsByShapeType, TEXT("Whether to use shape type to define the color of the bounds instead of using the particle state (if multiple shapes , will use the first one)"));
+		FAutoConsoleVariableRef CVarChaosDebugDebugDrawColorBoundsByShapeType(TEXT("p.Chaos.DebugDraw.ColorBoundsByShapeType"), bChaosDebugDebugDrawColorBoundsByShapeType, TEXT("Whether to use shape type to define the color of the bounds instead of using the particle state (if multiple shapes , will use the first one)"));
 
 		bool bChaosDebugDebugDrawConvexVertices = false;
 		bool bChaosDebugDebugDrawCoreShapes = false;
@@ -80,6 +86,7 @@ namespace Chaos
 			/* VelScale =					*/ 0.0f,
 			/* AngVelScale =				*/ 0.0f,
 			/* ImpulseScale =				*/ 0.0f,
+			/* PushOutScale =				*/ 0.0f,
 			/* InertiaScale =				*/ 0.0f,
 			/* DrawPriority =				*/ 10,
 			/* bShowSimple =				*/ true,
@@ -127,32 +134,50 @@ namespace Chaos
 			}
 		}
 
-		// default colors by state for shapes
-		FChaosDebugDrawColorsByState ChaosDefaultShapesColorsByState(
-			/* InDynamicColor =	  */ FColor(255, 255, 0),
-			/* InSleepingColor =  */ FColor(128, 128, 128),
-			/* InKinematicColor = */ FColor(0, 128, 255),
-			/* InStaticColor =	  */ FColor(255, 0, 0)
-		);
-
 		const FChaosDebugDrawColorsByState& GetDefaultShapesColorsByState()
 		{
+			// default colors by state for shapes
+			static FChaosDebugDrawColorsByState ChaosDefaultShapesColorsByState(
+				/* InDynamicColor =	  */ FColor(255, 255, 0),
+				/* InSleepingColor =  */ FColor(128, 128, 128),
+				/* InKinematicColor = */ FColor(0, 128, 255),
+				/* InStaticColor =	  */ FColor(255, 0, 0)
+			);
+
 			return ChaosDefaultShapesColorsByState;
 		}
 
-		// default colors by state for bounds ( darker version of the shapes colors - see above )
-		FChaosDebugDrawColorsByState ChaosDefaultBoundsColorsByState(
-			/* InDynamicColor =	  */ FColor(128, 128, 0),
-			/* InSleepingColor =  */ FColor(64, 64, 64),
-			/* InKinematicColor = */ FColor(0, 64, 128),
-			/* InStaticColor =	  */ FColor(128, 0, 0)
-		);
-
 		const FChaosDebugDrawColorsByState& GetDefaultBoundsColorsByState()
 		{
+			// default colors by state for bounds ( darker version of the shapes colors - see above )
+			static FChaosDebugDrawColorsByState ChaosDefaultBoundsColorsByState(
+				/* InDynamicColor =	  */ FColor(128, 128, 0),
+				/* InSleepingColor =  */ FColor(64, 64, 64),
+				/* InKinematicColor = */ FColor(0, 64, 128),
+				/* InStaticColor =	  */ FColor(128, 0, 0)
+			);
+
 			return ChaosDefaultBoundsColorsByState;
 		}
 
+		FColor GetIslandColor(const int32 IslandIndex, const bool bIsAwake)
+		{
+			static FColor AwakeColors[] =
+			{
+				FColor::Red,
+				FColor::Orange,
+				FColor::Yellow,
+				FColor::Green,
+				FColor::Blue,
+				FColor::Magenta,
+			};
+			const int32 NumAwakeColors = UE_ARRAY_COUNT(AwakeColors);
+			static FColor SleepingColor = FColor::Black;
+
+			return bIsAwake ? AwakeColors[IslandIndex % NumAwakeColors] : SleepingColor;
+		};
+
+
 		//-------------------------------------------------------------------------------------------------
 
 		FChaosDebugDrawColorsByShapeType::FChaosDebugDrawColorsByShapeType(
@@ -526,6 +551,10 @@ namespace Chaos
 			{
 				ShapeColor = Settings.ShapesColorsPerShapeType.GetColorFromShapeType(InnerType);
 			}
+			if (bChaosDebugDebugDrawColorShapesByIsland && (FConstGenericParticleHandle(Particle)->Island() != INDEX_NONE))
+			{
+				ShapeColor = GetIslandColor(FConstGenericParticleHandle(Particle)->Island(), true);
+			}
 
 			// If we get here, we have an actual shape to render
 			switch (InnerType)
@@ -743,7 +772,7 @@ namespace Chaos
 			}
 		}
 
-		void DrawCollisionImpl(const FVec3& Location, const FVec3& Normal, FReal Phi, const FVec3& Impulse, const FColor& DiscColor, const FColor& NormalColor, const FColor& ImpulseColor, FRealSingle ColorScale, const FChaosDebugDrawSettings& Settings)
+		void DrawCollisionImpl(const FVec3& Location, const FVec3& Normal, const FVec3& Tangent0, const FVec3& Tangent1, FReal Phi, const FVec3& Impulse, const FColor& DiscColor, const FColor& NormalColor, const FColor& ImpulseColor, FRealSingle ColorScale, const FChaosDebugDrawSettings& Settings)
 		{
 			FMatrix Axes = FRotationMatrix::MakeFromX(Normal);
 			if (Settings.ContactWidth > 0)
@@ -755,6 +784,8 @@ namespace Chaos
 			{
 				FColor C1 = (ColorScale * NormalColor).ToFColor(false);
 				FDebugDrawQueue::GetInstance().DrawDebugLine(Location, Location + Settings.DrawScale * Settings.ContactLen * Normal, C1, false, KINDA_SMALL_NUMBER, Settings.DrawPriority, Settings.LineThickness);
+				FDebugDrawQueue::GetInstance().DrawDebugLine(Location, Location + 0.5f * Settings.DrawScale * Settings.ContactLen * Tangent0, C1, false, KINDA_SMALL_NUMBER, Settings.DrawPriority, Settings.LineThickness);
+				FDebugDrawQueue::GetInstance().DrawDebugLine(Location, Location + 0.5f * Settings.DrawScale * Settings.ContactLen * Tangent1, C1, false, KINDA_SMALL_NUMBER, Settings.DrawPriority, Settings.LineThickness);
 			}
 			if (Settings.ContactPhiWidth > 0 && (Phi < FLT_MAX))
 			{
@@ -768,15 +799,14 @@ namespace Chaos
 			}
 		}
 
-		void DrawCollisionImpl(const FRigidTransform3& SpaceTransform, const FCollisionConstraintBase& Contact, FRealSingle ColorScale, const FChaosDebugDrawSettings& Settings)
+		void DrawCollisionImpl(const FRigidTransform3& SpaceTransform, const FPBDCollisionConstraint& Contact, FRealSingle ColorScale, const FChaosDebugDrawSettings& Settings)
 		{
 			if ((Settings.ContactWidth > 0) || (Settings.ContactLen > 0) || (Settings.ImpulseScale > 0.0f))
 			{
-				const FRigidBodyPointContactConstraint* PointConstraint = Contact.template As<FRigidBodyPointContactConstraint>();
-				if (PointConstraint->GetManifoldPoints().Num() > 0)
+				if (Contact.GetUseManifold())
 				{
-					FConstGenericParticleHandle Particle0 = Contact.Particle[0];
-					FConstGenericParticleHandle Particle1 = Contact.Particle[1];
+					const FConstGenericParticleHandle Particle0 = Contact.Particle[0];
+					const FConstGenericParticleHandle Particle1 = Contact.Particle[1];
 					const FRigidTransform3 WorldCoMTransform0 = FParticleUtilities::GetCoMWorldTransform(Particle0);
 					const FRigidTransform3 WorldCoMTransform1 = FParticleUtilities::GetCoMWorldTransform(Particle1);
 
@@ -792,20 +822,34 @@ namespace Chaos
 						return;
 					}
 
-					for (const FManifoldPoint& ManifoldPoint : PointConstraint->GetManifoldPoints())
+					for (const FManifoldPoint& ManifoldPoint : Contact.GetManifoldPoints())
 					{
-						const bool bIsActive = ManifoldPoint.bActive || !ManifoldPoint.NetPushOut.IsNearlyZero();
+						const bool bIsActive = ManifoldPoint.bActive || !ManifoldPoint.NetPushOut.IsNearlyZero() || !ManifoldPoint.NetImpulse.IsNearlyZero();
 						if (!bIsActive && !bChaosDebugDebugDrawInactiveContacts)
 						{
 							continue;
 						}
 
-						const int32 ContactPlaneOwner = ManifoldPoint.ContactPoint.ContactNormalOwnerIndex;
+						const int32 ContactPlaneOwner = bChaos_Collision_Manifold_FixNormalsInWorldSpace ? 1 : ManifoldPoint.ContactPoint.ContactNormalOwnerIndex;
 						const int32 ContactPointOwner = 1 - ContactPlaneOwner;
 						const FRigidTransform3& PlaneTransform = (ContactPlaneOwner == 0) ? WorldCoMTransform0 : WorldCoMTransform1;
 						const FRigidTransform3& PointTransform = (ContactPlaneOwner == 0) ? WorldCoMTransform1 : WorldCoMTransform0;
-						FConstGenericParticleHandle PlaneParticle = (ContactPlaneOwner == 0) ? Particle0 : Particle1;
-						const FVec3 PlaneNormal = PlaneTransform.TransformVector((ContactPlaneOwner == 1) ? ManifoldPoint.CoMContactNormal : -ManifoldPoint.CoMContactNormal);	// Normal is always points body 2->1 internally, but shows better if it points from plane owner to point owner
+						const FConstGenericParticleHandle PlaneParticle = (ContactPlaneOwner == 0) ? Particle0 : Particle1;
+						FVec3 PlaneNormal;
+						FVec3 PlaneTangent0;
+						FVec3 PlaneTangent1;
+						if (bChaos_Collision_Manifold_FixNormalsInWorldSpace)
+						{
+							PlaneNormal = ManifoldPoint.ManifoldContactNormal;
+							PlaneTangent0 = ManifoldPoint.ManifoldContactTangents[0];
+							PlaneTangent1 = ManifoldPoint.ManifoldContactTangents[1];
+						}
+						else
+						{
+							PlaneNormal = PlaneTransform.TransformVector((ContactPlaneOwner == 1) ? ManifoldPoint.ManifoldContactNormal : -ManifoldPoint.ManifoldContactNormal);	// Normal is always points body 2->1 internally, but shows better if it points from plane owner to point owner
+							PlaneTangent0 = PlaneTransform.TransformVector((ContactPlaneOwner == 1) ? ManifoldPoint.ManifoldContactTangents[0] : -ManifoldPoint.ManifoldContactTangents[0]);
+							PlaneTangent1 = PlaneTransform.TransformVector((ContactPlaneOwner == 1) ? ManifoldPoint.ManifoldContactTangents[1] : -ManifoldPoint.ManifoldContactTangents[1]);
+						}
 						const FVec3 PointLocation = PointTransform.TransformPosition(ManifoldPoint.CoMContactPoints[ContactPointOwner]) - ManifoldPoint.ContactPoint.ShapeMargins[ContactPointOwner] * PlaneNormal;
 						const FVec3 PlaneLocation = PlaneTransform.TransformPosition(ManifoldPoint.CoMContactPoints[ContactPlaneOwner]) + ManifoldPoint.ContactPoint.ShapeMargins[ContactPlaneOwner] * PlaneNormal;
 						const FVec3 PointPlaneLocation = PointLocation - FVec3::DotProduct(PointLocation - PlaneLocation, PlaneNormal) * PlaneNormal;
@@ -836,12 +880,19 @@ namespace Chaos
 						const FVec3 WorldPlaneLocation = SpaceTransform.TransformPosition(PlaneLocation);
 						const FVec3 WorldPointPlaneLocation = SpaceTransform.TransformPosition(PointPlaneLocation);
 						const FVec3 WorldPlaneNormal = SpaceTransform.TransformVectorNoScale(PlaneNormal);
+						const FVec3 WorldPlaneTangent0 = SpaceTransform.TransformVectorNoScale(PlaneTangent0);
+						const FVec3 WorldPlaneTangent1 = SpaceTransform.TransformVectorNoScale(PlaneTangent1);
 
 						// Pushout
-						if ((Settings.ImpulseScale > 0) && !ManifoldPoint.NetPushOut.IsNearlyZero())
+						if ((Settings.PushOutScale > 0) && !ManifoldPoint.NetPushOut.IsNearlyZero())
+						{
+							FColor Color = (ColorScale * PushOutImpusleColor).ToFColor(false);
+							FDebugDrawQueue::GetInstance().DrawDebugLine(WorldPointPlaneLocation, WorldPointPlaneLocation + Settings.DrawScale * Settings.PushOutScale * ManifoldPoint.NetPushOut, Color, false, KINDA_SMALL_NUMBER, Settings.DrawPriority, Settings.LineThickness);
+						}
+						if ((Settings.ImpulseScale > 0) && !ManifoldPoint.NetImpulse.IsNearlyZero())
 						{
 							FColor Color = (ColorScale * PushOutImpusleColor).ToFColor(false);
-							FDebugDrawQueue::GetInstance().DrawDebugLine(WorldPointPlaneLocation, WorldPointPlaneLocation + Settings.DrawScale * ManifoldPoint.NetPushOut, Color, false, KINDA_SMALL_NUMBER, Settings.DrawPriority, Settings.LineThickness);
+							FDebugDrawQueue::GetInstance().DrawDebugLine(WorldPointPlaneLocation, WorldPointPlaneLocation + Settings.DrawScale * Settings.ImpulseScale * ManifoldPoint.NetImpulse, Color, false, KINDA_SMALL_NUMBER, Settings.DrawPriority, Settings.LineThickness);
 						}
 						if ((Settings.ImpulseScale > 0) && !FMath::IsNearlyZero(ManifoldPoint.NetPushOutImpulseNormal))
 						{
@@ -857,18 +908,34 @@ namespace Chaos
 						}
 
 						// Manifold plane and normal
-						DrawCollisionImpl(WorldPlaneLocation, WorldPlaneNormal, ManifoldPoint.ContactPoint.Phi, ManifoldPoint.NetImpulse, DiscColor, NormalColor, ImpulseColor, ColorScale, Settings);
+						DrawCollisionImpl(WorldPlaneLocation, WorldPlaneNormal, WorldPlaneTangent0, WorldPlaneTangent1, ManifoldPoint.ContactPoint.Phi, ManifoldPoint.NetImpulse, DiscColor, NormalColor, ImpulseColor, ColorScale, Settings);
 
 						// Manifold point
 						FMatrix Axes = FRotationMatrix::MakeFromX(WorldPlaneNormal);
 						FDebugDrawQueue::GetInstance().DrawDebugCircle(WorldPointLocation, 0.5f * Settings.DrawScale * Settings.ContactWidth, 12, DiscColor, false, KINDA_SMALL_NUMBER, Settings.DrawPriority, Settings.LineThickness, Axes.GetUnitAxis(EAxis::Y), Axes.GetUnitAxis(EAxis::Z), false);
+
+						// Previous points
+						const FVec3 PrevPointLocation = PointTransform.TransformPosition(ManifoldPoint.CoMAnchorPoints[ContactPointOwner]) - ManifoldPoint.ContactPoint.ShapeMargins[ContactPointOwner] * PlaneNormal;
+						const FVec3 PrevPlaneLocation = PlaneTransform.TransformPosition(ManifoldPoint.CoMAnchorPoints[ContactPlaneOwner]) + ManifoldPoint.ContactPoint.ShapeMargins[ContactPlaneOwner] * PlaneNormal;
+						const FVec3 WorldPrevPointLocation = SpaceTransform.TransformPosition(PrevPointLocation);
+						const FVec3 WorldPrevPlaneLocation = SpaceTransform.TransformPosition(PrevPlaneLocation);
+						FDebugDrawQueue::GetInstance().DrawDebugLine(WorldPrevPointLocation, WorldPointLocation, FColor::White, false, KINDA_SMALL_NUMBER, Settings.DrawPriority, Settings.LineThickness);
+						FDebugDrawQueue::GetInstance().DrawDebugLine(WorldPrevPlaneLocation, WorldPlaneLocation, FColor::White, false, KINDA_SMALL_NUMBER, Settings.DrawPriority, Settings.LineThickness);
+
+						// Whether restored
+						if (Contact.WasManifoldRestored())
+						{
+							const FReal BoxScale = Settings.DrawScale * Settings.ContactWidth;
+							FDebugDrawQueue::GetInstance().DrawDebugBox(WorldPlaneLocation, FVec3(BoxScale, BoxScale, FReal(0.01)), FRotation3(FRotationMatrix::MakeFromZ(WorldPlaneNormal)), FColor::Blue, false, KINDA_SMALL_NUMBER, Settings.DrawPriority, 0.5f * Settings.LineThickness);
+						}
 					}
 				}
 				else
 				{
 					const FVec3 Location = SpaceTransform.TransformPosition(Contact.GetLocation());
 					const FVec3 Normal = SpaceTransform.TransformVector(Contact.GetNormal());
-					DrawCollisionImpl(Location, Normal, Contact.GetPhi(), FVec3(0), FColor(200, 0, 0), FColor(200, 0, 0), FColor(200, 0, 0), ColorScale, Settings);
+					FMatrix Axes = FRotationMatrix::MakeFromZ(Normal);
+					DrawCollisionImpl(Location, Axes.GetUnitAxis(EAxis::Z), Axes.GetUnitAxis(EAxis::X), Axes.GetUnitAxis(EAxis::Y), Contact.GetPhi(), FVec3(0), FColor(200, 0, 0), FColor(200, 0, 0), FColor(200, 0, 0), ColorScale, Settings);
 				}
 			}
 			if (Settings.ContactOwnerWidth > 0)
@@ -898,7 +965,7 @@ namespace Chaos
 
 			for (int32 ConstraintIndex = 0; ConstraintIndex < Collisions.NumConstraints(); ++ConstraintIndex)
 			{
-				const FRigidBodyPointContactConstraint* PointConstraint = Collisions.GetConstraint(ConstraintIndex).template As<FRigidBodyPointContactConstraint>();
+				const FPBDCollisionConstraint* PointConstraint = &Collisions.GetConstraint(ConstraintIndex);
 				if (!PointConstraint->GetDisabled() && (PointConstraint->Manifold.Phi < TNumericLimits<FReal>::Max()))
 				{
 					const FImplicitObject* Implicit0 = PointConstraint->Manifold.Implicit[0];
@@ -1071,93 +1138,67 @@ namespace Chaos
 			FDebugDrawQueue::GetInstance().DrawDebugLine(Pos, Pos + 0.01f * Settings.AngVelScale * SimSpace.AngularAcceleration, FColor::Orange, false, KINDA_SMALL_NUMBER, Settings.DrawPriority, Settings.LineThickness);
 		}
 
-		void DrawConstraintGraphImpl(const FRigidTransform3& SpaceTransform, const FPBDConstraintColor& GraphColor, const FChaosDebugDrawSettings& Settings)
+		void DrawConstraintGraphImpl(const FRigidTransform3& SpaceTransform, const FPBDConstraintGraph& Graph, const FChaosDebugDrawSettings& Settings)
 		{
-			static FColor IslandColors[] = 
-			{
-				FColor::Red,
-				FColor::Orange,
-				FColor::Yellow,
-				FColor::Green,
-				FColor::Blue,
-				FColor::Magenta,
-				FColor::White,
-				FColor::Black
-			};
-			const int32 NumIslandColors = UE_ARRAY_COUNT(IslandColors);
-
-			auto DrawGraphCollision = [&](const FRigidTransform3& SpaceTransform, const FPBDCollisionConstraintHandle& CollisionConstraint,  int32 IslandIndex, int32 LevelIndex, int32 ColorIndex, int32 OrderIndex, const FChaosDebugDrawSettings& Settings)
+			auto DrawGraphCollision = [&](const FRigidTransform3& SpaceTransform, const FPBDCollisionConstraint* Constraint,  int32 IslandIndex, int32 LevelIndex, int32 ColorIndex, int32 OrderIndex, const FChaosDebugDrawSettings& Settings)
 			{
-				FColor IslandColor = IslandColors[IslandIndex % NumIslandColors];
+				const FRigidTransform3 Transform0 = FParticleUtilities::GetCoMWorldTransform(FConstGenericParticleHandle(Constraint->GetConstrainedParticles()[0])) * SpaceTransform;
+				const FRigidTransform3 Transform1 = FParticleUtilities::GetCoMWorldTransform(FConstGenericParticleHandle(Constraint->GetConstrainedParticles()[1])) * SpaceTransform;
 
-				const FRigidTransform3 Transform0 = FParticleUtilities::GetCoMWorldTransform(FConstGenericParticleHandle(CollisionConstraint.GetConstrainedParticles()[0])) * SpaceTransform;
-				const FRigidTransform3 Transform1 = FParticleUtilities::GetCoMWorldTransform(FConstGenericParticleHandle(CollisionConstraint.GetConstrainedParticles()[1])) * SpaceTransform;
-
-				FVec3 ContactPos = CollisionConstraint.GetContactLocation();
-				if (CollisionConstraint.GetContact().GetType() == FCollisionConstraintBase::FType::SinglePoint)
+				FVec3 ContactPos = Constraint->GetContactLocation();
+				if ((Constraint != nullptr) && (Constraint->GetManifoldPoints().Num() > 0))
 				{
-					const FRigidBodyPointContactConstraint* Constraint = CollisionConstraint.GetContact().template As<const FRigidBodyPointContactConstraint>();
-					if ((Constraint != nullptr) && (Constraint->GetManifoldPoints().Num() > 0))
+					ContactPos = FVec3(0);
+					for (const FManifoldPoint& ManifoldPoint : Constraint->GetManifoldPoints())
 					{
-						ContactPos = FVec3(0);
-						for (const FManifoldPoint& ManifoldPoint : Constraint->GetManifoldPoints())
-						{
-							ContactPos += Transform0.TransformPositionNoScale(ManifoldPoint.CoMContactPoints[0]);
-							ContactPos += Transform1.TransformPositionNoScale(ManifoldPoint.CoMContactPoints[1]);
-						}
-						ContactPos /= (FReal)(2 * Constraint->GetManifoldPoints().Num());
+						ContactPos += Transform0.TransformPositionNoScale(ManifoldPoint.CoMContactPoints[0]);
+						ContactPos += Transform1.TransformPositionNoScale(ManifoldPoint.CoMContactPoints[1]);
 					}
+					ContactPos /= (FReal)(2 * Constraint->GetManifoldPoints().Num());
 				}
 
 
-				FDebugDrawQueue::GetInstance().DrawDebugLine(Transform0.GetLocation(), ContactPos, FColor::Black, false, KINDA_SMALL_NUMBER, Settings.DrawPriority, Settings.LineThickness);
+				FDebugDrawQueue::GetInstance().DrawDebugLine(Transform0.GetLocation(), ContactPos, FColor::Red, false, KINDA_SMALL_NUMBER, Settings.DrawPriority, Settings.LineThickness);
 				FDebugDrawQueue::GetInstance().DrawDebugLine(Transform1.GetLocation(), ContactPos, FColor::White, false, KINDA_SMALL_NUMBER, Settings.DrawPriority, Settings.LineThickness);
 				FDebugDrawQueue::GetInstance().DrawDebugString(ContactPos, FString::Format(TEXT("{0}-{1}-{2}"), { LevelIndex, ColorIndex, OrderIndex }), nullptr, FColor::Yellow, KINDA_SMALL_NUMBER, false, Settings.FontScale);
 
 			};
 
-			for (int32 IslandIndex = 0; IslandIndex < GraphColor.NumIslands(); ++IslandIndex)
+			for (int32 IslandIndex = 0; IslandIndex < Graph.NumIslands(); ++IslandIndex)
 			{
+				const FPBDIslandSolver* Island = Graph.GetSolverIsland(IslandIndex);
+
 				FAABB3 IslandAABB = FAABB3::EmptyAABB();
+				const TArray<FGeometryParticleHandle*> Particles = Island->GetParticles();
+				for (const FGeometryParticleHandle* GeoParticle : Particles)
+				{
+					FConstGenericParticleHandle Particle = GeoParticle;
+					if (Particle->IsDynamic() && Particle->HasBounds())
+					{
+						IslandAABB.GrowToInclude(Particle->BoundingBox());
+					}
+				}
 
-				const typename FPBDConstraintColor::FLevelToColorToConstraintListMap& LevelColorConstraintMap = GraphColor.GetIslandLevelToColorToConstraintListMap(IslandIndex);
-				int32 MaxColor = GraphColor.GetIslandMaxColor(IslandIndex);
-				int32 MaxLevel = GraphColor.GetIslandMaxLevel(IslandIndex);
-				for (int32 LevelIndex = 0; LevelIndex <= MaxLevel; ++LevelIndex)
+				const FColor IslandColor = GetIslandColor(IslandIndex, !Island->GetIsSleeping());
+				const FAABB3 Bounds = IslandAABB.TransformedAABB(SpaceTransform);
+				FDebugDrawQueue::GetInstance().DrawDebugBox(Bounds.Center(), 0.5f * Bounds.Extents(), SpaceTransform.GetRotation(), IslandColor, false, KINDA_SMALL_NUMBER, Settings.DrawPriority, 3.0f * Settings.LineThickness);
+
+				if (bChaosDebugDebugDrawContactGraph)
 				{
-					for (int32 ColorIndex = 0; ColorIndex <= MaxColor; ++ColorIndex)
+					const FConstraintHandleArray& Constraints = Island->GetConstraints();
+					for (int32 ConstraintIndex = 0; ConstraintIndex < Constraints.Num(); ++ConstraintIndex)
 					{
-						if (LevelColorConstraintMap[LevelIndex].Contains(ColorIndex) && LevelColorConstraintMap[LevelIndex][ColorIndex].Num())
-						{
-							const TArray<FConstraintHandle*>& ConstraintHandles = LevelColorConstraintMap[LevelIndex][ColorIndex];
-							for (int32 HandleIndex = 0; HandleIndex < ConstraintHandles.Num(); ++HandleIndex)
-							{
-								const FConstraintHandle* ConstraintHandle = ConstraintHandles[HandleIndex];
-
-								if (const FPBDCollisionConstraintHandle* CollisionHandle = ConstraintHandle->As<FPBDCollisionConstraintHandle>())
-								{
-									if (CollisionHandle->GetConstrainedParticles()[0]->HasBounds() && (CollisionHandle->GetConstrainedParticles()[0]->ObjectState() == EObjectStateType::Dynamic))
-									{
-										IslandAABB.GrowToInclude(CollisionHandle->GetConstrainedParticles()[0]->WorldSpaceInflatedBounds());
-									}
-									if (CollisionHandle->GetConstrainedParticles()[1]->HasBounds() && (CollisionHandle->GetConstrainedParticles()[1]->ObjectState() == EObjectStateType::Dynamic))
-									{
-										IslandAABB.GrowToInclude(CollisionHandle->GetConstrainedParticles()[1]->WorldSpaceInflatedBounds());
-									}
-
-									if (bChaosDebugDebugDrawContactGraph)
-									{
-										DrawGraphCollision(SpaceTransform, *CollisionHandle, IslandIndex, LevelIndex, ColorIndex, HandleIndex, Settings);
-									}
-								}
-							}
+						const FConstraintHandle* Constraint = Constraints[ConstraintIndex];
 
+						if (const FPBDCollisionConstraintHandle* Collision = Constraint->As<FPBDCollisionConstraintHandle>())
+						{
+							// @chaos(todo): store level and color in the constraint? Only for debug draw...
+							const int32 LevelIndex = INDEX_NONE;
+							const int32 ColorIndex = INDEX_NONE;
+							DrawGraphCollision(SpaceTransform, Collision->GetConstraint(), IslandIndex, LevelIndex, ColorIndex, ConstraintIndex, Settings);
 						}
 					}
 				}
-
-				FAABB3 Bounds = IslandAABB.TransformedAABB(SpaceTransform);
-				FDebugDrawQueue::GetInstance().DrawDebugBox(Bounds.Center(), 0.5f * Bounds.Extents(), SpaceTransform.GetRotation(), IslandColors[IslandIndex % NumIslandColors], false, KINDA_SMALL_NUMBER, Settings.DrawPriority, Settings.LineThickness);
 			}
 		}
 
@@ -1354,7 +1395,7 @@ namespace Chaos
 			}
 		}
 
-		void DrawConstraintGraph(const FRigidTransform3& SpaceTransform, const FPBDConstraintColor& Graph, const FChaosDebugDrawSettings* Settings)
+		void DrawConstraintGraph(const FRigidTransform3& SpaceTransform, const FPBDConstraintGraph& Graph, const FChaosDebugDrawSettings* Settings)
 		{
 			if (FDebugDrawQueue::IsDebugDrawingEnabled())
 			{
@@ -1370,6 +1411,33 @@ namespace Chaos
 			}
 		}
 
+		class FSpatialDebugDrawInterface: public ISpatialDebugDrawInterface
+		{
+		public:
+			FSpatialDebugDrawInterface(const FChaosDebugDrawSettings& InSettings)
+				: Settings(InSettings)
+			{}
+			
+			virtual ~FSpatialDebugDrawInterface() override = default;
+
+			virtual void Box(const FAABB3& InBox, const TVec3<FReal>& InLinearColor, float InThickness) override
+			{
+				FDebugDrawQueue::GetInstance().DrawDebugBox(InBox.Center(), InBox.Extents(), FQuat::Identity, FLinearColor(InLinearColor).ToFColor(false), false, -1.f, Settings.DrawPriority, InThickness);
+			}
+
+			virtual void Line(const TVec3<FReal>& InBegin, const TVec3<FReal>& InEnd, const TVec3<FReal>& InLinearColor, float InThickness) override
+			{
+				FDebugDrawQueue::GetInstance().DrawDebugLine(InBegin, InEnd, FLinearColor(InLinearColor).ToFColor(false), false, -1.f, Settings.DrawPriority, InThickness);
+			}
+		private:
+			FChaosDebugDrawSettings Settings;
+		};
+
+		void DrawSpatialAccelerationStructure(const ISpatialAcceleration<FAccelerationStructureHandle, FReal, 3>& InSpatialAccelerationStructure, const FChaosDebugDrawSettings* InSettings)
+		{
+			FSpatialDebugDrawInterface DebugDrawInterface(GetChaosDebugDrawSettings(InSettings));
+			InSpatialAccelerationStructure.DebugDraw(&DebugDrawInterface);
+		}
 #endif
 	}
 }
diff --git a/Engine/Source/Runtime/Experimental/Chaos/Private/Chaos/Collision/CapsuleConvexContactPoint.cpp b/Engine/Source/Runtime/Experimental/Chaos/Private/Chaos/Collision/CapsuleConvexContactPoint.cpp
--- a/Engine/Source/Runtime/Experimental/Chaos/Private/Chaos/Collision/CapsuleConvexContactPoint.cpp
+++ b/Engine/Source/Runtime/Experimental/Chaos/Private/Chaos/Collision/CapsuleConvexContactPoint.cpp
@@ -38,7 +38,7 @@ namespace Chaos
 				ContactPoint.ShapeContactNormal = -NormalConvex;
 				ContactPoint.ShapeMargins[0] = 0.0f;
 				ContactPoint.ShapeMargins[1] = 0.0f;
-				ContactPoint.ContactNormalOwnerIndex = 0;
+				ContactPoint.ContactNormalOwnerIndex = 0; // ContactNormalOwnerIndex is ignored when  bChaos_Collision_Manifold_FixNormalsInWorldSpace is set
 
 				ContactPoint.Location = ConvexTransform.TransformPosition(ContactPoint.ShapeContactPoints[0]);
 				ContactPoint.Normal = ConvexTransform.TransformVectorNoScale(-NormalConvex);
diff --git a/Engine/Source/Runtime/Experimental/Chaos/Private/Chaos/Collision/CollisionDetector.cpp b/Engine/Source/Runtime/Experimental/Chaos/Private/Chaos/Collision/CollisionDetector.cpp
--- a/Engine/Source/Runtime/Experimental/Chaos/Private/Chaos/Collision/CollisionDetector.cpp
+++ b/Engine/Source/Runtime/Experimental/Chaos/Private/Chaos/Collision/CollisionDetector.cpp
@@ -1,5 +1,6 @@
 // Copyright Epic Games, Inc. All Rights Reserved.
 #include "Chaos/Collision/CollisionDetector.h"
+#include "Chaos/Collision/CollisionConstraintAllocator.h"
 #include "Chaos/Collision/ParticlePairBroadPhase.h"
 #include "Chaos/Collision/SpatialAccelerationBroadPhase.h"
 #include "ChaosLog.h"
@@ -10,9 +11,10 @@ namespace Chaos
 	DEFINE_STAT(STAT_Collisions_Detect);
 	DEFINE_STAT(STAT_Collisions_BroadPhase);
 	DEFINE_STAT(STAT_Collisions_SpatialBroadPhase);
+	DEFINE_STAT(STAT_Collisions_AABBTree);
+	DEFINE_STAT(STAT_Collisions_Restore);
 	DEFINE_STAT(STAT_Collisions_Filtering);
 	DEFINE_STAT(STAT_Collisions_GenerateCollisions);
-	DEFINE_STAT(STAT_Collisions_ReceiveCollisions);
 #if CHAOS_ENABLE_STAT_NARROWPHASE
 	DEFINE_STAT(STAT_Collisions_NarrowPhase);
 #endif
diff --git a/Engine/Source/Runtime/Experimental/Chaos/Private/Chaos/Collision/CollisionSolver.cpp b/Engine/Source/Runtime/Experimental/Chaos/Private/Chaos/Collision/CollisionSolver.cpp
--- a/Engine/Source/Runtime/Experimental/Chaos/Private/Chaos/Collision/CollisionSolver.cpp
+++ b/Engine/Source/Runtime/Experimental/Chaos/Private/Chaos/Collision/CollisionSolver.cpp
@@ -19,7 +19,7 @@ DEFINE_LOG_CATEGORY(LogChaosCollisionSolver);
 
 namespace Chaos
 {
-	namespace Collisions
+	namespace CVars
 	{
 		extern int32 Chaos_Collision_UseShockPropagation;
 		extern FRealSingle Chaos_Collision_CollisionClipTolerance;
@@ -32,13 +32,11 @@ namespace Chaos
 		bool Chaos_Manifold_PushOut_Restitution = false;
 		bool Chaos_Manifold_PushOut_PositionCorrection = true;
 		int32 Chaos_Manifold_PushOut_VelocityCorrection = 2;
-		float Chaos_Manifold_PushOut_DistanceFactor = 0.0f;
 		FAutoConsoleVariableRef CVarChaos_Manifold_PushOut_NegativePushOut(TEXT("p.Chaos.Collision.Manifold.PushOut.NegativePushOut"), Chaos_Manifold_PushOut_NegativePushOut, TEXT(""));
 		FAutoConsoleVariableRef CVarChaos_Manifold_PushOut_StaticFriction(TEXT("p.Chaos.Collision.Manifold.PushOut.StaticFriction"), Chaos_Manifold_PushOut_StaticFriction, TEXT(""));
 		FAutoConsoleVariableRef CVarChaos_Manifold_PushOut_Restitution(TEXT("p.Chaos.Collision.Manifold.PushOut.Restitution"), Chaos_Manifold_PushOut_Restitution, TEXT(""));
 		FAutoConsoleVariableRef CVarChaos_Manifold_PushOut_PositionCorrection(TEXT("p.Chaos.Collision.Manifold.PushOut.PositionCorrection"), Chaos_Manifold_PushOut_PositionCorrection, TEXT(""));
 		FAutoConsoleVariableRef CVarChaos_Manifold_PushOut_VelocityCorrection(TEXT("p.Chaos.Collision.Manifold.PushOut.VelocityCorrectionMode"), Chaos_Manifold_PushOut_VelocityCorrection, TEXT("0 = No Velocity Correction; 1 = Normal Velocity Correction; 2 = Normal + Tangential Velocity Correction"));
-		FAutoConsoleVariableRef CVarChaos_Manifold_PushOut_DistanceFactor(TEXT("p.Chaos.Collision.Manifold.PushOut.DistanceFactor"), Chaos_Manifold_PushOut_DistanceFactor, TEXT(""));
 
 		// Note: Velocity solve requires full stiffness for restitution to work correctly so the max stiffness should be 1
 		FRealSingle Chaos_Manifold_Apply_MinStiffness = 0.25f;
@@ -59,7 +57,13 @@ namespace Chaos
 
 		float Chaos_Manifold_PushOut_PositionTolerance = 1.e-4f;
 		FAutoConsoleVariableRef CVarChaos_Manifold_PushOut_PositionTolerance(TEXT("p.Chaos.Collision.Manifold.PushOut.PositionTolerance"), Chaos_Manifold_PushOut_PositionTolerance, TEXT(""));
+	}
+	using namespace CVars;
+
+	extern bool bChaos_Collision_Manifold_FixNormalsInWorldSpace;
 
+	namespace Collisions
+	{
 		inline FReal GetConstraintStiffness(const FReal InStiffness, const FContactIterationParameters& IterationParameters)
 		{
 			FReal Stiffness = FMath::Clamp(InStiffness, (FReal)0.0, (FReal)1.0);
@@ -71,7 +75,7 @@ namespace Chaos
 					return Stiffness / Factor;
 				}
 			}
-			return (FReal)1;			
+			return (FReal)1;
 		}
 
 		inline FReal GetApplyStiffness(int32 It, int32 NumIts)
@@ -94,26 +98,13 @@ namespace Chaos
 			const FContactIterationParameters& IterationParameters,
 			const FContactParticleParameters& ParticleParameters,
 			const FReal Restitution,
-			const FReal InvM0,
-			const FMatrix33& InvI0,
-			const FReal InvM1,
-			const FMatrix33& InvI1,
-			const FVec3& P0, // Centre of Mass Positions and Rotations
-			const FRotation3& Q0,
-			const FVec3& P1,
-			const FRotation3& Q1,
-			FVec3& V0,
-			FVec3& W0,
-			FVec3& V1,
-			FVec3& W1,
+			FSolverBody& Body0,
+			FSolverBody& Body1,
 			FManifoldPoint& ManifoldPoint)
 		{
-			const bool bIsRigidDynamic0 = (InvM0 > 0.0f);
-			const bool bIsRigidDynamic1 = (InvM1 > 0.0f);
-
 			// Velocity correction uses the average contact point, and not the deepest point on each body
-			const FVec3 RelativeContactPoint0 = ManifoldPoint.ContactPoint.Location - P0;
-			const FVec3 RelativeContactPoint1 = ManifoldPoint.ContactPoint.Location - P1;
+			const FVec3 RelativeContactPoint0 = ManifoldPoint.ContactPoint.Location - Body0.P();
+			const FVec3 RelativeContactPoint1 = ManifoldPoint.ContactPoint.Location - Body1.P();
 			const FVec3 ContactNormal = ManifoldPoint.ContactPoint.Normal;
 			const FReal ContactPhi = ManifoldPoint.ContactPoint.Phi;
 
@@ -125,8 +116,8 @@ namespace Chaos
 				return;
 			}
 
-			const FVec3 ContactVelocity0 = V0 + FVec3::CrossProduct(W0, RelativeContactPoint0);
-			const FVec3 ContactVelocity1 = V1 + FVec3::CrossProduct(W1, RelativeContactPoint1);
+			const FVec3 ContactVelocity0 = Body0.V() + FVec3::CrossProduct(Body0.W(), RelativeContactPoint0);
+			const FVec3 ContactVelocity1 = Body1.V() + FVec3::CrossProduct(Body1.W(), RelativeContactPoint1);
 			const FVec3 ContactVelocity = ContactVelocity0 - ContactVelocity1;
 			const FReal ContactVelocityNormalLen = FVec3::DotProduct(ContactVelocity, ContactNormal);
 
@@ -147,8 +138,8 @@ namespace Chaos
 
 			// Calculate constraint-space mass
 			const FMatrix33 ConstraintMassInv =
-				(bIsRigidDynamic0 ? ComputeFactorMatrix3(RelativeContactPoint0, InvI0, InvM0) : FMatrix33(0)) +
-				(bIsRigidDynamic1 ? ComputeFactorMatrix3(RelativeContactPoint1, InvI1, InvM1) : FMatrix33(0));
+				(Body0.IsDynamic() ? ComputeFactorMatrix3(RelativeContactPoint0, Body0.InvI(), Body0.InvM()) : FMatrix33(0)) +
+				(Body1.IsDynamic() ? ComputeFactorMatrix3(RelativeContactPoint1, Body1.InvI(), Body1.InvM()) : FMatrix33(0));
 
 			// Calculate the impulse required to drive normal contact velocity to zero
 			const FReal ContactVelocityNormalChange = (ContactVelocityTargetNormal - ContactVelocityNormalLen);
@@ -172,17 +163,19 @@ namespace Chaos
 			}
 
 			// Calculate the velocity deltas from the impulse
-			if (bIsRigidDynamic0)
+			if (Body0.IsDynamic())
 			{
 				const FVec3 AngularImpulse = FVec3::CrossProduct(RelativeContactPoint0, Impulse);
-				V0 += InvM0 * Impulse;
-				W0 += InvI0 * AngularImpulse;
+				const FVec3 DV0 = Body0.InvM() * Impulse;
+				const FVec3 DW0 = Body0.InvI() * AngularImpulse;
+				Body0.ApplyVelocityDelta(DV0, DW0);
 			}
-			if (bIsRigidDynamic1)
+			if (Body1.IsDynamic())
 			{
 				const FVec3 AngularImpulse = FVec3::CrossProduct(RelativeContactPoint1, Impulse);
-				V1 -= InvM1 * Impulse;
-				W1 -= InvI1 * AngularImpulse;
+				const FVec3 DV1 = Body1.InvM() * -Impulse;
+				const FVec3 DW1 = Body1.InvI() * -AngularImpulse;
+				Body1.ApplyVelocityDelta(DV1, DW1);
 			}
 
 			// PushOut needs to know if we applied restitution and static friction
@@ -202,18 +195,8 @@ namespace Chaos
 			const FContactParticleParameters& ParticleParameters,
 			const FReal StaticFriction,
 			const FReal DynamicFriction,
-			const FReal InvM0,
-			const FMatrix33& InvI0,
-			const FReal InvM1,
-			const FMatrix33& InvI1,
-			const FVec3& P0, // Centre of Mass Positions and Rotations
-			const FRotation3& Q0,
-			const FVec3& P1,
-			const FRotation3& Q1,
-			FVec3& V0,
-			FVec3& W0,
-			FVec3& V1,
-			FVec3& W1,
+			FSolverBody& Body0,
+			FSolverBody& Body1,
 			FManifoldPoint& ManifoldPoint)
 		{
 			// Reject inactive contacts (no normal impulse)
@@ -222,16 +205,13 @@ namespace Chaos
 				return;
 			}
 
-			const bool bIsRigidDynamic0 = (InvM0 > 0.0f);
-			const bool bIsRigidDynamic1 = (InvM1 > 0.0f);
-
 			// Velocity correction uses the average contact point, and not the deepest point on each body
-			const FVec3 RelativeContactPoint0 = ManifoldPoint.ContactPoint.Location - P0;
-			const FVec3 RelativeContactPoint1 = ManifoldPoint.ContactPoint.Location - P1;
+			const FVec3 RelativeContactPoint0 = ManifoldPoint.ContactPoint.Location - Body0.P();
+			const FVec3 RelativeContactPoint1 = ManifoldPoint.ContactPoint.Location - Body1.P();
 			const FVec3 ContactNormal = ManifoldPoint.ContactPoint.Normal;
 
-			const FVec3 ContactVelocity0 = V0 + FVec3::CrossProduct(W0, RelativeContactPoint0);
-			const FVec3 ContactVelocity1 = V1 + FVec3::CrossProduct(W1, RelativeContactPoint1);
+			const FVec3 ContactVelocity0 = Body0.V() + FVec3::CrossProduct(Body0.W(), RelativeContactPoint0);
+			const FVec3 ContactVelocity1 = Body1.V() + FVec3::CrossProduct(Body1.W(), RelativeContactPoint1);
 			const FVec3 ContactVelocity = ContactVelocity0 - ContactVelocity1;
 			const FReal ContactVelocityNormalLen = FVec3::DotProduct(ContactVelocity, ContactNormal);
 			const FVec3 ContactVelocityTangent = ContactVelocity - ContactVelocityNormalLen * ContactNormal;
@@ -249,8 +229,8 @@ namespace Chaos
 
 				// Calculate constraint-space mass
 				const FMatrix33 ConstraintMassInv =
-					(bIsRigidDynamic0 ? ComputeFactorMatrix3(RelativeContactPoint0, InvI0, InvM0) : FMatrix33(0)) +
-					(bIsRigidDynamic1 ? ComputeFactorMatrix3(RelativeContactPoint1, InvI1, InvM1) : FMatrix33(0));
+					(Body0.IsDynamic() ? ComputeFactorMatrix3(RelativeContactPoint0, Body0.InvI(), Body0.InvM()) : FMatrix33(0)) +
+					(Body1.IsDynamic() ? ComputeFactorMatrix3(RelativeContactPoint1, Body1.InvI(), Body1.InvM()) : FMatrix33(0));
 
 				// Calculate the impulse required to drive contact velocity to zero along the tangent
 				const FReal ContactVelocityTangentChange = -ContactVelocityTangentLen;
@@ -277,17 +257,19 @@ namespace Chaos
 			const FVec3 NetImpulse = ManifoldPoint.NetImpulse + Impulse;
 
 			// Calculate the velocity deltas from the impulse
-			if (bIsRigidDynamic0)
+			if (Body0.IsDynamic())
 			{
 				const FVec3 AngularImpulse = FVec3::CrossProduct(RelativeContactPoint0, Impulse);
-				V0 += InvM0 * Impulse;
-				W0 += InvI0 * AngularImpulse;
+				const FVec3 DV0 = Body0.InvM() * Impulse;
+				const FVec3 DW0 = Body0.InvI() * AngularImpulse;
+				Body0.ApplyVelocityDelta(DV0, DW0);
 			}
-			if (bIsRigidDynamic1)
+			if (Body1.IsDynamic())
 			{
 				const FVec3 AngularImpulse = FVec3::CrossProduct(RelativeContactPoint1, Impulse);
-				V1 -= InvM1 * Impulse;
-				W1 -= InvI1 * AngularImpulse;
+				const FVec3 DV1 = Body1.InvM() * -Impulse;
+				const FVec3 DW1 = Body1.InvI() * -AngularImpulse;
+				Body1.ApplyVelocityDelta(DV1, DW1);
 			}
 
 			// PushOut needs to know if we applied restitution and static friction
@@ -304,33 +286,29 @@ namespace Chaos
 		// This modifies the in/out CoM position and rotations.
 		void ApplyManifoldPushOutCorrection(
 			const FReal Stiffness,
-			const FRigidBodyPointContactConstraint& Constraint,
-			const FGenericParticleHandle Particle0,
-			const FGenericParticleHandle Particle1,
+			const FPBDCollisionConstraint& Constraint,
 			const FContactIterationParameters& IterationParameters,
 			const FContactParticleParameters& ParticleParameters,
-			const bool bIsRigidDynamic0,
-			const bool bIsRigidDynamic1,
-			FVec3& P0, // Centre of Mass Positions and Rotations
-			FRotation3& Q0,
-			FVec3& P1,
-			FRotation3& Q1,
-			FVec3& V0,
-			FVec3& W0,
-			FVec3& V1,
-			FVec3& W1,
+			FConstraintSolverBody& Body0,
+			FConstraintSolverBody& Body1,
 			FManifoldPoint& ManifoldPoint)
 		{
-			const TPBDRigidParticleHandle<FReal, 3>* PBDRigid0 = Particle0->CastToRigidParticle();
-			const TPBDRigidParticleHandle<FReal, 3>* PBDRigid1 = Particle1->CastToRigidParticle();
-
 			const FReal Margin0 = ManifoldPoint.ContactPoint.ShapeMargins[0];
 			const FReal Margin1 = ManifoldPoint.ContactPoint.ShapeMargins[1];
 
 			// Calculate the position error we need to correct, including static friction and restitution
 			// Position correction uses the deepest point on each body (see velocity correction which uses average contact)
-			const FRotation3& PlaneQ = (ManifoldPoint.ContactPoint.ContactNormalOwnerIndex == 0) ? Q0 : Q1;
-			const FVec3 ContactNormal = PlaneQ * ManifoldPoint.CoMContactNormal;
+
+			FVec3 ContactNormal;
+			if (bChaos_Collision_Manifold_FixNormalsInWorldSpace)
+			{
+				ContactNormal = ManifoldPoint.ManifoldContactNormal;
+			}
+			else
+			{
+				const FRotation3& PlaneQ = (ManifoldPoint.ContactPoint.ContactNormalOwnerIndex == 0) ? Body0.Q() : Body1.Q();
+				ContactNormal = PlaneQ * ManifoldPoint.ManifoldContactNormal;
+			}
 
 			const bool bApplyStaticFriction = (ManifoldPoint.bActive && ManifoldPoint.bInsideStaticFrictionCone && ManifoldPoint.bPotentialRestingContact && Chaos_Manifold_PushOut_StaticFriction);
 			FVec3 LocalContactPoint0 = ManifoldPoint.CoMContactPoints[0];
@@ -339,7 +317,7 @@ namespace Chaos
 			// If we are applying static friction, we are attempting to move the contact points at their intitial positions back together
 			if (bApplyStaticFriction)
 			{
-				Constraint.CalculatePrevCoMContactPoints(Particle0, Particle1, ManifoldPoint, IterationParameters.Dt, LocalContactPoint0, LocalContactPoint1);
+				Constraint.CalculatePrevCoMContactPoints(Body0.SolverBody(), Body1.SolverBody(), ManifoldPoint, IterationParameters.Dt, LocalContactPoint0, LocalContactPoint1);
 			}
 
 			// We could push out to the PBD distance that would give an implicit velocity equal to -(1+e).Vin
@@ -348,20 +326,11 @@ namespace Chaos
 			const FReal PhiPadding = 0;
 
 			// Contact points on each body adjusted so that the points end up separated by TargetPhi
-			const FVec3 RelativeContactPoint0 = Q0 * LocalContactPoint0 - (PhiPadding + Margin0) * ContactNormal;
-			const FVec3 RelativeContactPoint1 = Q1 * LocalContactPoint1 + (PhiPadding + Margin1) * ContactNormal;
+			const FVec3 RelativeContactPoint0 = Body0.Q() * LocalContactPoint0 - (PhiPadding + Margin0) * ContactNormal;
+			const FVec3 RelativeContactPoint1 = Body1.Q() * LocalContactPoint1 + (PhiPadding + Margin1) * ContactNormal;
 
 			// Net error we need to correct, including lateral movement to correct for friction
-			FVec3 ContactError = (P1 + RelativeContactPoint1) - (P0 + RelativeContactPoint0);
-
-			// The max rest distance for PBD object is the distance it can move in one frame from rest along the contact normal (i.e., position delta from gravity)
-			if ((Chaos_Manifold_PushOut_DistanceFactor > 0.0f) && (ContactNormal.Z > 0.0f))
-			{
-				// @todo(chaos): pass gravity into the solver if we use this feature. Otherwise remove it.
-				const FReal GravityNegZ = 980.0f;
-				const FReal MaxRestDistance = GravityNegZ * IterationParameters.Dt * IterationParameters.Dt;
-				ContactError += (Chaos_Manifold_PushOut_DistanceFactor * MaxRestDistance * ContactNormal.Z) * ContactNormal;
-			}
+			FVec3 ContactError = (Body1.P() + RelativeContactPoint1) - (Body0.P() + RelativeContactPoint0);
 
 			FReal ContactErrorNormal = FVec3::DotProduct(ContactError, ContactNormal);
 
@@ -385,11 +354,9 @@ namespace Chaos
 			}
 
 			// Calculate joint-space mass matrix (J.M.Jt)
-			const FMatrix33 InvI0 = bIsRigidDynamic0 ? Utilities::ComputeWorldSpaceInertia(Q0, PBDRigid0->InvI()) * Constraint.Manifold.InvInertiaScale0 : FMatrix33(0);
-			const FMatrix33 InvI1 = bIsRigidDynamic1 ? Utilities::ComputeWorldSpaceInertia(Q1, PBDRigid1->InvI()) * Constraint.Manifold.InvInertiaScale1 : FMatrix33(0);
 			const FMatrix33 ContactMassInv =
-				(bIsRigidDynamic0 ? ComputeFactorMatrix3(RelativeContactPoint0, InvI0, PBDRigid0->InvM()) : FMatrix33(0)) +
-				(bIsRigidDynamic1 ? ComputeFactorMatrix3(RelativeContactPoint1, InvI1, PBDRigid1->InvM()) : FMatrix33(0));
+				(Body0.IsDynamic() ? ComputeFactorMatrix3(RelativeContactPoint0, Body0.InvI(), Body0.InvM()) : FMatrix33(0)) +
+				(Body1.IsDynamic() ? ComputeFactorMatrix3(RelativeContactPoint1, Body1.InvI(), Body1.InvM()) : FMatrix33(0));
 
 			// Calculate pushout
 			// - If we were inside the static friction cone during the apply step, correct positions so that
@@ -413,28 +380,22 @@ namespace Chaos
 
 			FVec3 NetPushOut = ManifoldPoint.NetPushOut + PushOut;
 
-			if (bIsRigidDynamic0)
+			if (Body0.IsDynamic())
 			{
 				const FVec3 AngularPushOut = FVec3::CrossProduct(RelativeContactPoint0, PushOut);
-				const FVec3 DX0 = PBDRigid0->InvM() * PushOut;
-				const FVec3 DR0 = InvI0 * AngularPushOut;
-
-				P0 += DX0;
-				Q0 += FRotation3::FromElements(DR0, 0.f) * Q0 * FReal(0.5);
-				Q0.Normalize();
-				FParticleUtilities::SetCoMWorldTransform(Particle0, P0, Q0);
+				const FVec3 DX0 = Body0.InvM() * PushOut;
+				const FVec3 DR0 = Body0.InvI() * AngularPushOut;
+				Body0.ApplyTransformDelta(DX0, DR0);
+				Body0.UpdateRotationDependentState();
 			}
 
-			if (bIsRigidDynamic1)
+			if (Body1.IsDynamic())
 			{
 				const FVec3 AngularPushOut = FVec3::CrossProduct(RelativeContactPoint1, PushOut);
-				const FVec3 DX1 = -(PBDRigid1->InvM() * PushOut);
-				const FVec3 DR1 = -(InvI1 * AngularPushOut);
-
-				P1 += DX1;
-				Q1 += FRotation3::FromElements(DR1, 0.f) * Q1 * FReal(0.5);
-				Q1.Normalize();
-				FParticleUtilities::SetCoMWorldTransform(Particle1, P1, Q1);
+				const FVec3 DX1 = Body1.InvM() * -PushOut;
+				const FVec3 DR1 = Body1.InvI() * -AngularPushOut;
+				Body1.ApplyTransformDelta(DX1, DR1);
+				Body1.UpdateRotationDependentState();
 			}
 
 			ManifoldPoint.NetPushOut = NetPushOut;
@@ -450,23 +411,13 @@ namespace Chaos
 			const FReal Restitution,
 			const FReal StaticFriction,
 			const FReal DynamicFriction,
-			const FReal InvM0,
-			const FMatrix33& InvI0,
-			const FReal InvM1,
-			const FMatrix33& InvI1,
-			const FVec3& P0, // Centre of Mass Positions and Rotations
-			const FRotation3& Q0,
-			const FVec3& P1,
-			const FRotation3& Q1,
-			FVec3& V0,
-			FVec3& W0,
-			FVec3& V1,
-			FVec3& W1,
+			FConstraintSolverBody& Body0,
+			FConstraintSolverBody& Body1,
 			FManifoldPoint& ManifoldPoint)
 		{
 			// Velocity correction uses the average contact point, and not the deepest point on each body
-			const FVec3 RelativeContactPoint0 = ManifoldPoint.ContactPoint.Location - P0;
-			const FVec3 RelativeContactPoint1 = ManifoldPoint.ContactPoint.Location - P1;
+			const FVec3 RelativeContactPoint0 = ManifoldPoint.ContactPoint.Location - Body0.P();
+			const FVec3 RelativeContactPoint1 = ManifoldPoint.ContactPoint.Location - Body1.P();
 			const FVec3 ContactNormal = ManifoldPoint.ContactPoint.Normal;
 			FVec3 ContactTangent = FVec3(0);
 
@@ -475,8 +426,8 @@ namespace Chaos
 				ContactTangent = (ManifoldPoint.NetPushOut - FVec3::DotProduct(ManifoldPoint.NetPushOut, ContactNormal) * ContactNormal).GetSafeNormal();
 			}
 
-			const FVec3 ContactVelocity0 = V0 + FVec3::CrossProduct(W0, RelativeContactPoint0);
-			const FVec3 ContactVelocity1 = V1 + FVec3::CrossProduct(W1, RelativeContactPoint1);
+			const FVec3 ContactVelocity0 = Body0.V() + FVec3::CrossProduct(Body0.W(), RelativeContactPoint0);
+			const FVec3 ContactVelocity1 = Body1.V() + FVec3::CrossProduct(Body1.W(), RelativeContactPoint1);
 			const FVec3 ContactVelocity = ContactVelocity0 - ContactVelocity1;
 			const FReal ContactVelocityNormal = FVec3::DotProduct(ContactVelocity, ContactNormal);
 			const FReal ContactVelocityTangent = FVec3::DotProduct(ContactVelocity, ContactTangent);
@@ -490,8 +441,8 @@ namespace Chaos
 
 			// Calculate constraint-space inverse mass
 			const FMatrix33 ContactMassInv =
-				((InvM0 > 0.0f) ? ComputeFactorMatrix3(RelativeContactPoint0, InvI0, InvM0) : FMatrix33(0)) +
-				((InvM1 > 0.0f) ? ComputeFactorMatrix3(RelativeContactPoint1, InvI1, InvM1) : FMatrix33(0));
+				(Body0.IsDynamic() ? ComputeFactorMatrix3(RelativeContactPoint0, Body0.InvI(), Body0.InvM()) : FMatrix33(0)) +
+				(Body1.IsDynamic() ? ComputeFactorMatrix3(RelativeContactPoint1, Body1.InvI(), Body1.InvM()) : FMatrix33(0));
 
 			// Calculate the impulse to get the desired target normal velocity
 			// We are ignoring both static and dynamic friction here
@@ -546,24 +497,20 @@ namespace Chaos
 
 
 			const FVec3 Impulse = ImpulseNormal * ContactNormal + ImpulseTangent * ContactTangent;
-			if (InvM0 > 0.0f)
+			if (Body0.IsDynamic())
 			{
 				const FVec3 AngularImpulse = FVec3::CrossProduct(RelativeContactPoint0, Impulse);
-				const FVec3 DV0 = InvM0 * Impulse;
-				const FVec3 DW0 = InvI0 * AngularImpulse;
-
-				V0 += DV0;
-				W0 += DW0;
+				const FVec3 DV0 = Body0.InvM() * Impulse;
+				const FVec3 DW0 = Body0.InvI() * AngularImpulse;
+				Body0.ApplyVelocityDelta(DV0, DW0);
 			}
 
-			if (InvM1 > 0.0f)
+			if (Body1.IsDynamic())
 			{
-				const FVec3 AngularImpulse = FVec3::CrossProduct(RelativeContactPoint1, -Impulse);
-				const FVec3 DV1 = -InvM1 * Impulse;
-				const FVec3 DW1 = InvI1 * AngularImpulse;
-
-				V1 += DV1;
-				W1 += DW1;
+				const FVec3 AngularImpulse = FVec3::CrossProduct(RelativeContactPoint1, Impulse);
+				const FVec3 DV1 = Body1.InvM() * -Impulse;
+				const FVec3 DW1 = Body1.InvI() * -AngularImpulse;
+				Body1.ApplyVelocityDelta(DV1, DW1);
 			}
 
 			ManifoldPoint.NetPushOutImpulseNormal = NetImpulseNormal;
@@ -573,50 +520,36 @@ namespace Chaos
 
 		// Velocity solver loop for a single contact manifold.
 		void ApplyContactManifold(
-			FRigidBodyPointContactConstraint& Constraint,
-			FGenericParticleHandle Particle0,
-			FGenericParticleHandle Particle1,
+			FPBDCollisionConstraint& Constraint,
 			const FContactIterationParameters& IterationParameters,
 			const FContactParticleParameters& ParticleParameters)
 		{
-			TPBDRigidParticleHandle<FReal, 3>* PBDRigid0 = Particle0->CastToRigidParticle();
-			TPBDRigidParticleHandle<FReal, 3>* PBDRigid1 = Particle1->CastToRigidParticle();
-			const bool bIsRigidDynamic0 = PBDRigid0 && PBDRigid0->ObjectState() == EObjectStateType::Dynamic;
-			const bool bIsRigidDynamic1 = PBDRigid1 && PBDRigid1->ObjectState() == EObjectStateType::Dynamic;
-
-			FVec3 P0 = FParticleUtilities::GetCoMWorldPosition(Particle0);
-			FRotation3 Q0 = FParticleUtilities::GetCoMWorldRotation(Particle0);
-			FVec3 P1 = FParticleUtilities::GetCoMWorldPosition(Particle1);
-			FRotation3 Q1 = FParticleUtilities::GetCoMWorldRotation(Particle1);
-			FVec3 V0 = Particle0->V();
-			FVec3 W0 = Particle0->W();
-			FVec3 V1 = Particle1->V();
-			FVec3 W1 = Particle1->W();
-
-			const FReal InvM0 = bIsRigidDynamic0 ? PBDRigid0->InvM() : 0.0f;
-			const FReal InvM1 = bIsRigidDynamic1 ? PBDRigid1->InvM() : 0.0f;
-			const FMatrix33 InvI0 = bIsRigidDynamic0 ? Utilities::ComputeWorldSpaceInertia(Q0, PBDRigid0->InvI()) * Constraint.Manifold.InvInertiaScale0 : FMatrix33(0);
-			const FMatrix33 InvI1 = bIsRigidDynamic1 ? Utilities::ComputeWorldSpaceInertia(Q1, PBDRigid1->InvI()) * Constraint.Manifold.InvInertiaScale1 : FMatrix33(0);
+			FSolverBody& Body0 = *Constraint.GetSolverBody0();
+			FSolverBody& Body1 = *Constraint.GetSolverBody1();
 
 			Constraint.AccumulatedImpulse = FVec3(0);
 
 			const FReal ConstraintStiffness = GetConstraintStiffness(Constraint.GetStiffness(), IterationParameters);
-			const FReal Stiffness = ConstraintStiffness*GetApplyStiffness(IterationParameters.Iteration, IterationParameters.NumIterations);
+			const FReal Stiffness = ConstraintStiffness * GetApplyStiffness(IterationParameters.Iteration, IterationParameters.NumIterations);
+
+			const FVec3 InitialV0 = Body0.V();
+			const FVec3 InitialW0 = Body0.W();
+			const FVec3 InitialV1 = Body1.V();
+			const FVec3 InitialW1 = Body1.W();
 
-			// Iterate over the manifold and accumulate velocity corrections - we will apply them after the loop
+			// Iterate over the manifold and accumulate velocity corrections
 			TArrayView<FManifoldPoint> ManifoldPoints = Constraint.GetManifoldPoints();
 			for (int32 PointIndex = 0; PointIndex < ManifoldPoints.Num(); ++PointIndex)
 			{
-				FManifoldPoint& ManifoldPoint = Constraint.SetActiveManifoldPoint(PointIndex, P0, Q0, P1, Q1);
+				FManifoldPoint& ManifoldPoint = Constraint.SetActiveManifoldPoint(PointIndex, Body0.P(), Body0.Q(), Body1.P(), Body1.Q());
 
 				CalculateManifoldNormalVelocityCorrection(
 					Stiffness,
 					IterationParameters,
 					ParticleParameters,
 					Constraint.Manifold.Restitution,
-					InvM0, InvI0, InvM1, InvI1,
-					P0, Q0, P1, Q1,
-					V0, W0, V1, W1,
+					Body0, 
+					Body1,
 					ManifoldPoint);
 
 				Constraint.AccumulatedImpulse += ManifoldPoint.NetImpulse;
@@ -628,7 +561,7 @@ namespace Chaos
 			{
 				for (int32 PointIndex = 0; PointIndex < ManifoldPoints.Num(); ++PointIndex)
 				{
-					FManifoldPoint& ManifoldPoint = Constraint.SetActiveManifoldPoint(PointIndex, P0, Q0, P1, Q1);
+					FManifoldPoint& ManifoldPoint = Constraint.SetActiveManifoldPoint(PointIndex, Body0.P(), Body0.Q(), Body1.P(), Body1.Q());
 
 					CalculateManifoldFrictionVelocityCorrection(
 						Stiffness,
@@ -636,142 +569,97 @@ namespace Chaos
 						ParticleParameters,
 						StaticFriction,
 						DynamicFriction,
-						InvM0, InvI0, InvM1, InvI1,
-						P0, Q0, P1, Q1,
-						V0, W0, V1, W1,
+						Body0, 
+						Body1,
 						ManifoldPoint);
 
 					Constraint.AccumulatedImpulse += ManifoldPoint.NetImpulse;
 				}
 			}
 
-			if (bIsRigidDynamic0)
+			// Apply a position correction so that we get the desired velocity in the implicit velocity step
+			if (Body0.IsDynamic())
 			{
-				const FVec3 DV0 = V0 - PBDRigid0->V();
-				const FVec3 DW0 = W0 - PBDRigid0->W();
+				const FVec3 DV0 = Body0.V() - InitialV0;
+				const FVec3 DW0 = Body0.W() - InitialW0;
 				const FVec3 DX0 = DV0 * IterationParameters.Dt;
 				const FVec3 DR0 = DW0 * IterationParameters.Dt;
-
-				PBDRigid0->V() = V0;
-				PBDRigid0->W() = W0;
-				P0 += DX0;
-				Q0 += FRotation3::FromElements(DR0, 0.f) * Q0 * FReal(0.5);
-				Q0.Normalize();
-				FParticleUtilities::SetCoMWorldTransform(PBDRigid0, P0, Q0);
+				Body0.ApplyTransformDelta(DX0, DR0);
+				Body0.UpdateRotationDependentState();
 			}
-
-			if (bIsRigidDynamic1)
+			if (Body1.IsDynamic())
 			{
-				const FVec3 DV1 = V1 - PBDRigid1->V();
-				const FVec3 DW1 = W1 - PBDRigid1->W();
+				const FVec3 DV1 = Body1.V() - InitialV1;
+				const FVec3 DW1 = Body1.W() - InitialW1;
 				const FVec3 DX1 = DV1 * IterationParameters.Dt;
 				const FVec3 DR1 = DW1 * IterationParameters.Dt;
-
-				PBDRigid1->V() = V1;
-				PBDRigid1->W() = W1;
-				P1 += DX1;
-				Q1 += FRotation3::FromElements(DR1, 0.f) * Q1 * FReal(0.5);
-				Q1.Normalize();
-				FParticleUtilities::SetCoMWorldTransform(PBDRigid1, P1, Q1);
+				Body1.ApplyTransformDelta(DX1, DR1);
+				Body1.UpdateRotationDependentState();
 			}
 		}
 
 
 		// Pushout solver loop for a single contafct manifold.
 		void ApplyPushOutManifold(
-			FRigidBodyPointContactConstraint& Constraint,
-			const TSet<const TGeometryParticleHandle<FReal, 3>*>& IsTemporarilyStatic,
+			FPBDCollisionConstraint& Constraint,
 			const FContactIterationParameters& IterationParameters,
-			const FContactParticleParameters& ParticleParameters,
-			const FVec3& GravityDir)
+			const FContactParticleParameters& ParticleParameters)
 		{
-			FGenericParticleHandle Particle0 = FGenericParticleHandle(Constraint.Particle[0]);
-			FGenericParticleHandle Particle1 = FGenericParticleHandle(Constraint.Particle[1]);
 			TArrayView<FManifoldPoint> ManifoldPoints = Constraint.GetManifoldPoints();
 
-			// Lock bodies for shock propagation?
-			bool IsTemporarilyStatic0 = false;
-			bool IsTemporarilyStatic1 = false;
-			if (Chaos_Collision_UseShockPropagation)
-			{
-				IsTemporarilyStatic0 = IsTemporarilyStatic.Contains(Particle0->GeometryParticleHandle());
-				IsTemporarilyStatic1 = IsTemporarilyStatic.Contains(Particle1->GeometryParticleHandle());
+			FConstraintSolverBody Body0 = *Constraint.GetSolverBody0();
+			FConstraintSolverBody Body1 = *Constraint.GetSolverBody1();
 
-				// In the case where two objects at the same level in shock propagation end
-				// up in contact with each other, try to decide which one has higher priority.
-				// For now, we use the one which is "lower" (in direction of gravity) as the frozen one.
-				if (IsTemporarilyStatic0 && IsTemporarilyStatic1)
+			// How much shock propagation should we apply - currently 100% on final iteration, otherwise none
+			FReal ShockPropagationInvMassScale = (IterationParameters.Iteration < (IterationParameters.NumIterations - 1)) ? 1.0f : 0.0f;
+			if ((ShockPropagationInvMassScale < 1.0f) && Chaos_Collision_UseShockPropagation)
+			{
+				// Lock one of the bodies if want shock propagation and they are at different shock levels
+				if (Body0.Level() < Body1.Level())
 				{
-					if (ManifoldPoints.Num() > 0)
-					{
-						const FReal NormalThreshold = 0.2f;	// @chaos(todo): what value here?
-						const FReal NormalGrav = FVec3::DotProduct(ManifoldPoints[0].ContactPoint.Normal, GravityDir);
-						const bool bNormalIsUp = NormalGrav < -NormalThreshold;
-						const bool bNormalIsDown = NormalGrav > NormalThreshold;
-						IsTemporarilyStatic0 = bNormalIsDown;
-						IsTemporarilyStatic1 = bNormalIsUp;
-					}
-					else
-					{
-						IsTemporarilyStatic0 = false;
-						IsTemporarilyStatic1 = false;
-					}
+					// Body0 is lower level and should be frozen
+					Body0.SetInvMassScale(ShockPropagationInvMassScale);
+				}
+				else if (Body0.Level() > Body1.Level())
+				{
+					// Body1 is lower level and should be frozen
+					Body1.SetInvMassScale(ShockPropagationInvMassScale);
+				}
+				else
+				{
+					// Both bodies are at the same level so neither is frozen
 				}
 			}
 
-			TPBDRigidParticleHandle<FReal, 3>* PBDRigid0 = Particle0->CastToRigidParticle();
-			TPBDRigidParticleHandle<FReal, 3>* PBDRigid1 = Particle1->CastToRigidParticle();
-			const bool bIsRigidDynamic0 = PBDRigid0 && (PBDRigid0->ObjectState() == EObjectStateType::Dynamic) && !IsTemporarilyStatic0;
-			const bool bIsRigidDynamic1 = PBDRigid1 && (PBDRigid1->ObjectState() == EObjectStateType::Dynamic) && !IsTemporarilyStatic1;
-
-			if (!bIsRigidDynamic0 && !bIsRigidDynamic1)
+			// If both bodies are frozen or kinematic just skip
+			// @todo(chaos): this should not happen - if one body is kinematic, the other should have no mass scaling. Check and remove
+			if (!Body0.IsDynamic() && !Body1.IsDynamic())
 			{
 				return;
 			}
 
 			// Gradually increase position correction through iterations (optional based on cvars)
-			const FReal ConstraintStiffness = GetConstraintStiffness(Constraint.GetStiffness(), IterationParameters);
-			const FReal Stiffness = ConstraintStiffness*GetPushOutStiffness(IterationParameters.Iteration, IterationParameters.NumIterations);
-			
-			FVec3 P0 = FParticleUtilities::GetCoMWorldPosition(Particle0);
-			FRotation3 Q0 = FParticleUtilities::GetCoMWorldRotation(Particle0);
-			FVec3 P1 = FParticleUtilities::GetCoMWorldPosition(Particle1);
-			FRotation3 Q1 = FParticleUtilities::GetCoMWorldRotation(Particle1);
-			FVec3 V0 = Particle0->V();
-			FVec3 W0 = Particle0->W();
-			FVec3 V1 = Particle1->V();
-			FVec3 W1 = Particle1->W();
+			const FReal Stiffness = GetPushOutStiffness(IterationParameters.Iteration, IterationParameters.NumIterations);
 
 			// Apply the position correction so that all contacts have zero separation
 			if (Chaos_Manifold_PushOut_PositionCorrection)
 			{
 				for (int32 PointIndex = 0; PointIndex < ManifoldPoints.Num(); ++PointIndex)
 				{
-					FManifoldPoint& ManifoldPoint = Constraint.SetActiveManifoldPoint(PointIndex, P0, Q0, P1, Q1);
+					FManifoldPoint& ManifoldPoint = Constraint.SetActiveManifoldPoint(PointIndex, Body0.P(), Body0.Q(), Body1.P(), Body1.Q());
 
 					// The distance we should pushout to
 					FReal PushOutSeparation = 0.0f;
-					if ((Chaos_Manifold_PushOut_DistanceFactor > 0.0f) && (ManifoldPoint.ContactPoint.Normal.Z > 0.0f))
-					{
-						// @todo(chaos): pass gravity into the solver if we use this feature. Otherwise remove it.
-						const FReal GravityNegZ = 980.0f;
-						const FReal MaxRestDistance = GravityNegZ * IterationParameters.Dt * IterationParameters.Dt;
-						PushOutSeparation = (Chaos_Manifold_PushOut_DistanceFactor * MaxRestDistance * ManifoldPoint.ContactPoint.Normal.Z);
-					}
 
 					if (ManifoldPoint.bActive || (ManifoldPoint.ContactPoint.Phi < PushOutSeparation))
 					{
 						ApplyManifoldPushOutCorrection(
 							Stiffness,
 							Constraint,
-							Particle0,
-							Particle1,
 							IterationParameters,
 							ParticleParameters,
-							bIsRigidDynamic0,
-							bIsRigidDynamic1,
-							P0, Q0, P1, Q1,
-							V0, V1, W0, W1,
+							Body0,
+							Body1,
 							ManifoldPoint);
 					}
 				}
@@ -780,16 +668,12 @@ namespace Chaos
 			// Apply a velocity correction so that all contacts have non-negative contact velocity
 			if (Chaos_Manifold_PushOut_VelocityCorrection > 0)
 			{
-				const FReal InvM0 = bIsRigidDynamic0 ? PBDRigid0->InvM() : 0.0f;
-				const FReal InvM1 = bIsRigidDynamic1 ? PBDRigid1->InvM() : 0.0f;
-				const FMatrix33 InvI0 = bIsRigidDynamic0 ? Utilities::ComputeWorldSpaceInertia(Q0, PBDRigid0->InvI()) * Constraint.Manifold.InvInertiaScale0 : FMatrix33(0);
-				const FMatrix33 InvI1 = bIsRigidDynamic1 ? Utilities::ComputeWorldSpaceInertia(Q1, PBDRigid1->InvI()) * Constraint.Manifold.InvInertiaScale1 : FMatrix33(0);
 				const FReal StaticFriction = FMath::Max(Constraint.Manifold.Friction, Constraint.Manifold.AngularFriction);
 				const FReal DynamicFriction = Constraint.Manifold.Friction;
 
 				for (int32 PointIndex = 0; PointIndex < ManifoldPoints.Num(); ++PointIndex)
 				{
-					FManifoldPoint& ManifoldPoint = Constraint.SetActiveManifoldPoint(PointIndex, P0, Q0, P1, Q1);
+					FManifoldPoint& ManifoldPoint = Constraint.SetActiveManifoldPoint(PointIndex, Body0.P(), Body0.Q(), Body1.P(), Body1.Q());
 
 					if (ManifoldPoint.bActive || !ManifoldPoint.NetPushOut.IsNearlyZero(KINDA_SMALL_NUMBER))
 					{
@@ -800,29 +684,13 @@ namespace Chaos
 							Constraint.Manifold.Restitution,
 							StaticFriction,
 							DynamicFriction,
-							InvM0, InvI0,
-							InvM1, InvI1,
-							P0, Q0, 
-							P1, Q1,
-							V0, W0,
-							V1, W1,
+							Body0,
+							Body1,
 							ManifoldPoint);
 					}
 
 					//Constraint.AccumulatedImpulse += ManifoldPoint.NetPushOutImpulseNormal * ManifoldPoint.ContactPoint.Normal;
 				}
-
-				if (bIsRigidDynamic0)
-				{
-					PBDRigid0->V() = V0;
-					PBDRigid0->W() = W0;
-				}
-
-				if (bIsRigidDynamic1)
-				{
-					PBDRigid1->V() = V1;
-					PBDRigid1->W() = W1;
-				}
 			}
 		}
 
diff --git a/Engine/Source/Runtime/Experimental/Chaos/Private/Chaos/Collision/CollisionSolver.h b/Engine/Source/Runtime/Experimental/Chaos/Private/Chaos/Collision/CollisionSolver.h
--- a/Engine/Source/Runtime/Experimental/Chaos/Private/Chaos/Collision/CollisionSolver.h
+++ b/Engine/Source/Runtime/Experimental/Chaos/Private/Chaos/Collision/CollisionSolver.h
@@ -11,18 +11,14 @@ namespace Chaos
 	namespace Collisions
 	{
 		void ApplyContactManifold(
-			FRigidBodyPointContactConstraint& Constraint,
-			FGenericParticleHandle Particle0,
-			FGenericParticleHandle Particle1,
+			FPBDCollisionConstraint& Constraint,
 			const FContactIterationParameters& IterationParameters,
 			const FContactParticleParameters& ParticleParameters);
 
 		void ApplyPushOutManifold(
-			FRigidBodyPointContactConstraint& Constraint,
-			const TSet<const FGeometryParticleHandle*>& IsTemporarilyStatic,
+			FPBDCollisionConstraint& Constraint,
 			const FContactIterationParameters& IterationParameters,
-			const FContactParticleParameters& ParticleParameters,
-			const FVec3& GravityDir);
+			const FContactParticleParameters& ParticleParameters);
 
 	}
 }
