Chaos

diff --git a/Engine/Source/Runtime/Experimental/Chaos/Private/Chaos/Collision/PBDCollisionConstraint.cpp b/Engine/Source/Runtime/Experimental/Chaos/Private/Chaos/Collision/PBDCollisionConstraint.cpp
--- a/Engine/Source/Runtime/Experimental/Chaos/Private/Chaos/Collision/PBDCollisionConstraint.cpp
+++ b/Engine/Source/Runtime/Experimental/Chaos/Private/Chaos/Collision/PBDCollisionConstraint.cpp
@@ -1,7 +1,18 @@
 // Copyright Epic Games, Inc. All Rights Reserved.
 #include "Chaos/Collision/PBDCollisionConstraint.h"
+#include "Chaos/Collision/CollisionConstraintAllocator.h"
+#include "Chaos/Collision/PBDCollisionConstraintHandle.h"
+#include "Chaos/Collision/SolverCollisionContainer.h"
+#include "Chaos/PBDCollisionConstraints.h"
+#include "Chaos/Evolution/SolverBody.h"
+#include "Chaos/Evolution/SolverBodyContainer.h"
 #include "Chaos/ParticleHandle.h"
 #include "Chaos/Particle/ParticleUtilities.h"
+#include "Chaos/PBDCollisionConstraints.h"
+
+// Private includes
+#include "PBDCollisionSolver.h"
+
 
 //PRAGMA_DISABLE_OPTIMIZATION
 
@@ -12,22 +23,40 @@ namespace Chaos
 	FAutoConsoleVariableRef CVarChaos_Manifold_MatchPositionTolerance(TEXT("p.Chaos.Collision.Manifold.MatchPositionTolerance"), Chaos_Manifold_MatchPositionTolerance, TEXT("A tolerance as a fraction of object size used to determine if two contact points are the same"));
 	FAutoConsoleVariableRef CVarChaos_Manifold_MatchNormalTolerance(TEXT("p.Chaos.Collision.Manifold.MatchNormalTolerance"), Chaos_Manifold_MatchNormalTolerance, TEXT("A tolerance on the normal dot product used to determine if two contact points are the same"));
 
-	bool bChaos_Manifold_EnabledWithJoints = false;
+	FRealSingle Chaos_Manifold_FrictionPositionTolerance = 1.0f;	// Distance a shape-relative contact point can move and still be considered the same point
+	FAutoConsoleVariableRef CVarChaos_Manifold_FrictionPositionTolerance(TEXT("p.Chaos.Collision.Manifold.FrictionPositionTolerance"), Chaos_Manifold_FrictionPositionTolerance, TEXT(""));
+
+	FRealSingle Chaos_GBFCharacteristicTimeRatio = 1.0f;
+	FAutoConsoleVariableRef CVarChaos_GBFCharacteristicTimeRatio(TEXT("p.Chaos.Collision.GBFCharacteristicTimeRatio"), Chaos_GBFCharacteristicTimeRatio, TEXT("The ratio between characteristic time and Dt"));
+
+	bool bChaos_Manifold_EnabledWithJoints = true;
 	FAutoConsoleVariableRef CVarChaos_Manifold_EnabledWithJoints(TEXT("p.Chaos.Collision.Manifold.EnabledWithJoints"), bChaos_Manifold_EnabledWithJoints, TEXT(""));
 
-	FString FCollisionConstraintBase::ToString() const
+	bool bChaos_Manifold_EnableGjkWarmStart = true;
+	FAutoConsoleVariableRef CVarChaos_Manifold_EnableGjkWarmStart(TEXT("p.Chaos.Collision.Manifold.EnableGjkWarmStart"), bChaos_Manifold_EnableGjkWarmStart, TEXT(""));
+
+	bool bChaos_Manifold_EnableFrictionRestore = true;
+	FAutoConsoleVariableRef CVarChaos_Manifold_EnableFrictionRestore(TEXT("p.Chaos.Collision.Manifold.EnableFrictionRestore"), bChaos_Manifold_EnableFrictionRestore, TEXT(""));
+
+	extern bool bChaos_Collision_Manifold_FixNormalsInWorldSpace;
+
+	FString FPBDCollisionConstraint::ToString() const
 	{
 		return FString::Printf(TEXT("Particle:%s, Levelset:%s, AccumulatedImpulse:%s"), *Particle[0]->ToString(), *Particle[1]->ToString(), *AccumulatedImpulse.ToString());
 	}
 
-	// @todo(chaos): this is not an obvious implementation of operator<. We should make it a named predicate
-	bool FCollisionConstraintBase::operator<(const FCollisionConstraintBase& Other) const
+	bool ContactConstraintSortPredicate(const FPBDCollisionConstraint& L, const FPBDCollisionConstraint& R)
 	{
 		//sort constraints by the smallest particle idx in them first
 		//if the smallest particle idx is the same for both, use the other idx
 
-		const FParticleID ParticleIdxs[] ={Particle[0]->ParticleID(),Particle[1]->ParticleID()};
-		const FParticleID OtherParticleIdxs[] ={Other.Particle[0]->ParticleID(),Other.Particle[1]->ParticleID()};
+		if (L.GetType() != R.GetType())
+		{
+			return L.GetType() < R.GetType();
+		}
+
+		const FParticleID ParticleIdxs[] = { L.Particle[0]->ParticleID(), L.Particle[1]->ParticleID() };
+		const FParticleID OtherParticleIdxs[] = { R.Particle[0]->ParticleID(), R.Particle[1]->ParticleID() };
 
 		const int32 MinIdx = ParticleIdxs[0] < ParticleIdxs[1] ? 0 : 1;
 		const int32 OtherMinIdx = OtherParticleIdxs[0] < OtherParticleIdxs[1] ? 0 : 1;
@@ -44,6 +73,154 @@ namespace Chaos
 		return false;
 	}
 
+	FPBDCollisionConstraint* FPBDCollisionConstraint::Make(
+		FGeometryParticleHandle* Particle0,
+		const FImplicitObject* Implicit0,
+		const FBVHParticles* Simplicial0,
+		const FRigidTransform3& ParticleWorldTransform0,
+		const FRigidTransform3& ImplicitLocalTransform0,
+		FGeometryParticleHandle* Particle1,
+		const FImplicitObject* Implicit1,
+		const FBVHParticles* Simplicial1,
+		const FRigidTransform3& ParticleWorldTransform1,
+		const FRigidTransform3& ImplicitLocalTransform1,
+		const FReal InCullDistance,
+		const EContactShapesType ShapesType,
+		const bool bInUseManifold,
+		FCollisionConstraintAllocator& Allocator)
+	{
+		// Create a constraint, or return an existing one for the same shape pair
+		FPBDCollisionConstraint* Constraint = Allocator.FindOrCreateConstraint(Particle0, Implicit0, Simplicial0, Particle1, Implicit1, Simplicial1);
+		if (Constraint != nullptr)
+		{
+			// Store the existing manifold as the previous manifold for friction etc
+			Constraint->SaveManifold();
+
+			// Initialize the constraint if we did not fully restore it - the manifold will get built later
+			Constraint->Setup(ECollisionConstraintType::Standard, ShapesType, ImplicitLocalTransform0, ImplicitLocalTransform1, InCullDistance, bInUseManifold);
+		}
+		return Constraint;
+	}
+
+	FPBDCollisionConstraint* FPBDCollisionConstraint::MakeSwept(
+		FGeometryParticleHandle* Particle0,
+		const FImplicitObject* Implicit0,
+		const FBVHParticles* Simplicial0,
+		const FRigidTransform3& ParticleWorldTransform0,
+		const FRigidTransform3& ImplicitLocalTransform0,
+		FGeometryParticleHandle* Particle1,
+		const FImplicitObject* Implicit1,
+		const FBVHParticles* Simplicial1,
+		const FRigidTransform3& ParticleWorldTransform1,
+		const FRigidTransform3& ImplicitLocalTransform1,
+		const FReal InCullDistance,
+		EContactShapesType ShapesType,
+		FCollisionConstraintAllocator& Allocator)
+	{
+		const bool bUseManifold = true;
+		
+		// Create a constraint, or return an existing one for the same shape pair
+		FPBDCollisionConstraint* Constraint = Allocator.FindOrCreateConstraint(Particle0, Implicit0, Simplicial0, Particle1, Implicit1, Simplicial1);
+		if (Constraint != nullptr)
+		{
+			// We do not attempt to restore swept constraints since we only get here when an object is moving fast
+			Constraint->ResetManifold();
+
+			// Initialize the constraint - the manifold will get built later
+			Constraint->Setup(ECollisionConstraintType::Swept, ShapesType, ImplicitLocalTransform0, ImplicitLocalTransform1, InCullDistance, bUseManifold);
+		}
+		return Constraint;
+	}
+
+	FPBDCollisionConstraint FPBDCollisionConstraint::MakeResimCache(
+		const FPBDCollisionConstraint& Source)
+	{
+		// @todo(chaos): The resim cache version probably doesn't need all the data, so maybe try to cur this down?
+		FPBDCollisionConstraint Constraint = Source;
+
+		// Invalidate the data that maps the constraint to its container (we are no longer in the container)
+		// @todo(chaos): this should probably be handled by the copy constructor
+		Constraint.GetContainerCookie().ClearContainerData();
+
+		return Constraint;
+	}
+
+
+	FPBDCollisionConstraint::FPBDCollisionConstraint()
+		: ImplicitTransform{ FRigidTransform3(), FRigidTransform3() }
+		, Particle{ nullptr, nullptr }
+		, AccumulatedImpulse(0)
+		, Manifold()
+		, TimeOfImpact(0)
+		, Type(ECollisionConstraintType::None)
+		, Stiffness(FReal(1))
+		, CullDistance(TNumericLimits<FReal>::Max())
+		, bUseManifold(false)
+		, bUseIncrementalManifold(false)
+		, SolverBodies{ nullptr, nullptr }
+		, GJKWarmStartData()
+		, PrevManifoldPoints()
+		, bWasManifoldRestored(false)
+	{
+		Manifold.Implicit[0] = nullptr;
+		Manifold.Implicit[1] = nullptr;
+		Manifold.Simplicial[0] = nullptr;
+		Manifold.Simplicial[1] = nullptr;
+		Manifold.ShapesType = EContactShapesType::Unknown;
+	}
+
+	FPBDCollisionConstraint::FPBDCollisionConstraint(
+		FGeometryParticleHandle* Particle0,
+		const FImplicitObject* Implicit0,
+		const FBVHParticles* Simplicial0,
+		FGeometryParticleHandle* Particle1,
+		const FImplicitObject* Implicit1,
+		const FBVHParticles* Simplicial1)
+		: ImplicitTransform{ FRigidTransform3(), FRigidTransform3() }
+		, Particle{ Particle0, Particle1 }
+		, AccumulatedImpulse(0)
+		, Manifold()
+		, TimeOfImpact(0)
+		, Type(ECollisionConstraintType::None)
+		, Stiffness(FReal(1))
+		, CullDistance(TNumericLimits<FReal>::Max())
+		, bUseManifold(false)
+		, bUseIncrementalManifold(false)
+		, SolverBodies{ nullptr, nullptr }
+		, GJKWarmStartData()
+		, PrevManifoldPoints()
+		, bWasManifoldRestored(false)
+	{
+		Manifold.Implicit[0] = Implicit0;
+		Manifold.Implicit[1] = Implicit1;
+		Manifold.Simplicial[0] = Simplicial0;
+		Manifold.Simplicial[1] = Simplicial1;
+		Manifold.ShapesType = EContactShapesType::Unknown;
+	}
+
+	void FPBDCollisionConstraint::Setup(
+		const ECollisionConstraintType InType,
+		const EContactShapesType InShapesType,
+		const FRigidTransform3& InImplicitLocalTransform0,
+		const FRigidTransform3& InImplicitLocalTransform1,
+		const FReal InCullDistance,
+		const bool bInUseManifold)
+	{
+		Type = InType;
+
+		Manifold.ShapesType = InShapesType;
+
+		ImplicitTransform[0] = InImplicitLocalTransform0;
+		ImplicitTransform[1] = InImplicitLocalTransform1;
+
+		CullDistance = InCullDistance;
+
+		bUseManifold = bInUseManifold && CanUseManifold(Particle[0], Particle[1]);
+		bUseIncrementalManifold = true;	// This will get changed later if we call AddOneShotManifoldContact
+
+		SetIsSleeping(false);
+	}
+
 	// Are the two manifold points the same point?
 	// Ideally a contact is considered the same as one from the previous iteration if
 	//		The contact is Vertex - Face and there was a prior iteration collision on the same Vertex
@@ -55,7 +232,7 @@ namespace Chaos
 	//		Edge - Edge - ?? hard...
 	//		Face - Face - ?? hard...
 	//
-	bool FRigidBodyPointContactConstraint::AreMatchingContactPoints(const FContactPoint& A, const FContactPoint& B, FReal& OutScore) const
+	bool FPBDCollisionConstraint::AreMatchingContactPoints(const FContactPoint& A, const FContactPoint& B, FReal& OutScore) const
 	{
 		OutScore = 0.0f;
 
@@ -110,7 +287,7 @@ namespace Chaos
 		return false;
 	}
 
-	int32 FRigidBodyPointContactConstraint::FindManifoldPoint(const FContactPoint& ContactPoint) const
+	int32 FPBDCollisionConstraint::FindManifoldPoint(const FContactPoint& ContactPoint) const
 	{
 		const int32 NumManifoldPoints = ManifoldPoints.Num();
 		int32 BestMatchIndex = INDEX_NONE;
@@ -133,33 +310,59 @@ namespace Chaos
 		return BestMatchIndex;
 	}
 
-	void FRigidBodyPointContactConstraint::UpdateManifoldContacts(FReal Dt)
+	// @todo(chaos): this exists primarily for the contact modification callbacks. See if we can eliminate it when not needed
+	void FPBDCollisionConstraint::UpdateManifoldContacts()
 	{
+		FVec3 P0, P1;
+		FRotation3 Q0, Q1;
+
+		// @todo(chaos): Remove this when we don't need to support incremental manifolds (this will only be called on creation/restore)
+		if ((GetSolverBody0() != nullptr) && (GetSolverBody1() != nullptr))
+		{
+			P0 = GetSolverBody0()->P();
+			Q0 = GetSolverBody0()->Q();
+			P1 = GetSolverBody1()->P();
+			Q1 = GetSolverBody1()->Q();
+		}
+		else
+		{
+			// @todo(chaos): we should not need to regenerate the CoM transform
+			P0 = FParticleUtilities::GetCoMWorldPosition(FConstGenericParticleHandle(Particle[0]));
+			Q0 = FParticleUtilities::GetCoMWorldRotation(FConstGenericParticleHandle(Particle[0]));
+			P1 = FParticleUtilities::GetCoMWorldPosition(FConstGenericParticleHandle(Particle[1]));
+			Q1 = FParticleUtilities::GetCoMWorldRotation(FConstGenericParticleHandle(Particle[1]));
+		}
+
 		for (int32 Index = 0; Index < ManifoldPoints.Num() ; Index++)
 		{
-			const FContactPoint& ContactPoint = ManifoldPoints[Index].ContactPoint;
-			UpdateManifoldPoint(Index, ContactPoint, Dt);
+			FManifoldPoint& ManifoldPoint = ManifoldPoints[Index];
+			GetWorldSpaceManifoldPoint(ManifoldPoint, P0, Q0, P1, Q1, ManifoldPoint.ContactPoint.Location, ManifoldPoint.ContactPoint.Normal, ManifoldPoint.ContactPoint.Phi);
+			
+			ManifoldPoint.bPotentialRestingContact = bUseManifold;
+			ManifoldPoint.bInsideStaticFrictionCone = bUseManifold;
 
 			// Copy currently active point
-			if (ContactPoint.Phi < Manifold.Phi)
+			if (ManifoldPoint.ContactPoint.Phi < Manifold.Phi)
 			{
-				SetActiveContactPoint(ContactPoint);
+				SetActiveContactPoint(ManifoldPoint.ContactPoint);
 			}
 		}
 	}
 
-	void FRigidBodyPointContactConstraint::AddOneshotManifoldContact(const FContactPoint& ContactPoint, const FReal Dt)
+	void FPBDCollisionConstraint::AddOneshotManifoldContact(const FContactPoint& ContactPoint, const FReal Dt)
 	{
-		AddManifoldPoint(ContactPoint, Dt);
+		int32 ManifoldPointIndex = AddManifoldPoint(ContactPoint, Dt);
 
 		// Copy currently active point
-		if (ContactPoint.Phi < Manifold.Phi)
+		if (ManifoldPoints[ManifoldPointIndex].ContactPoint.Phi < Manifold.Phi)
 		{
-			SetActiveContactPoint(ContactPoint);
+			SetActiveContactPoint(ManifoldPoints[ManifoldPointIndex].ContactPoint);
 		}
+
+		bUseIncrementalManifold = false;
 	}
 
-	void FRigidBodyPointContactConstraint::AddIncrementalManifoldContact(const FContactPoint& ContactPoint, const FReal Dt)
+	void FPBDCollisionConstraint::AddIncrementalManifoldContact(const FContactPoint& ContactPoint, const FReal Dt)
 	{
 		if (bUseManifold)
 		{
@@ -173,13 +376,13 @@ namespace Chaos
 			else
 			{
 				// This is a new manifold point - capture the state and generate initial properties
-				AddManifoldPoint(ContactPoint, Dt);
+				ManifoldPointIndex = AddManifoldPoint(ContactPoint, Dt);
 			}
 
 			// Copy currently active point
-			if (ContactPoint.Phi < Manifold.Phi)
+			if (ManifoldPoints[ManifoldPointIndex].ContactPoint.Phi < Manifold.Phi)
 			{
-				SetActiveContactPoint(ContactPoint);
+				SetActiveContactPoint(ManifoldPoints[ManifoldPointIndex].ContactPoint);
 			}
 		}
 		else 
@@ -191,107 +394,154 @@ namespace Chaos
 			}
 			else
 			{
-				ManifoldPoints[0] = ContactPoint;
+				ManifoldPoints[0].ContactPoint = ContactPoint;
 			}
 
 			InitManifoldPoint(ManifoldPoints[0], Dt);
 
-			SetActiveContactPoint(ContactPoint);
+			SetActiveContactPoint(ManifoldPoints[0].ContactPoint);
 		}
+
+		bUseIncrementalManifold = true;
 	}
 
-	void FRigidBodyPointContactConstraint::ClearManifold()
+	void FPBDCollisionConstraint::ClearManifold()
 	{
 		ManifoldPoints.Reset();
 	}
 
-	void FRigidBodyPointContactConstraint::InitManifoldPoint(FManifoldPoint& ManifoldPoint, FReal Dt)
+	void FPBDCollisionConstraint::InitManifoldPoint(FManifoldPoint& ManifoldPoint, FReal Dt)
 	{
-		// @todo(chaos): This is only needed for unit tests - make it more easily unit testable
-		if ((Particle[0] == nullptr) || (Particle[1] == nullptr))
+		FConstGenericParticleHandle Particle0 = Particle[0];
+		FConstGenericParticleHandle Particle1 = Particle[1];
+		if (!Particle0.IsValid() || !Particle1.IsValid())
 		{
+			// @todo(chaos): This is just for unit tests testing one-shot manifolds - remove it somehow... 
+			// maybe ConstructConvexConvexOneShotManifold should not take a Constraint
 			return;
 		}
 
-		FConstGenericParticleHandle Particle0 = Particle[0];
-		FConstGenericParticleHandle Particle1 = Particle[1];
+		// Note we start with contact inactive - they get activated if Phi goes negative
+		ManifoldPoint.bActive = false;
 
 		// @todo(chaos): determine potentially resting contact case based on contact velocity
 		// @todo(chaos): support static friction position correction for non-manifold contacts (spheres, point clouds, etc)
 		ManifoldPoint.bPotentialRestingContact = bUseManifold;
+		ManifoldPoint.bInsideStaticFrictionCone = bUseManifold;
 
 		// Update the derived contact state (CoM relative data)
 		UpdateManifoldPointFromContact(ManifoldPoint);
 
+		// Initialize the previous contact transforms if the data is available, otherwise reset them to current
+		TryRestoreFrictionData(ManifoldPoint);
+
+		// @todo(chaos): Remove SolverBody use from InitManifoldPoint when we don't need Incremental Manifolds any more
+		// This is a bit ugly. While we still support incremental manifolds, this function can be called
+		// in CreateConstraint (OK) and UpdateConstraint (not OK). In UpdateConstraint we should be using
+		// the SolverBodies for the state, but in InitConstraint we access the Particles directly.
+		FVec3 PCoM0, PCoM1;
+		FRotation3 QCoM0, QCoM1;
+		if ((GetSolverBody0() != nullptr) && (GetSolverBody1() != nullptr))
+		{
+			// This will be called from UpdateConstraint
+			PCoM0 = GetSolverBody0()->P();
+			QCoM0 = GetSolverBody0()->Q();
+			PCoM1 = GetSolverBody1()->P();
+			QCoM1 = GetSolverBody1()->Q();
+		}
+		else
+		{
+			// This will be called from CreateConstraint
+			PCoM0 = FParticleUtilities::GetCoMWorldPosition(Particle0);
+			QCoM0 = FParticleUtilities::GetCoMWorldRotation(Particle0);
+			PCoM1 = FParticleUtilities::GetCoMWorldPosition(Particle1);
+			QCoM1 = FParticleUtilities::GetCoMWorldRotation(Particle1);
+		}
+
+		// We can used the particle here even in UpdateConstraints because we are accessing tick constants 
+		// and nothing that is changed by the constraint solvers
+		const FVec3& PreV0 = Particle0->PreV();
+		const FVec3& PreW0 = Particle0->PreW();
+		const FVec3& PreV1 = Particle1->PreV();
+		const FVec3& PreW1 = Particle1->PreW();
+
 		// World-space contact state used below
-		// @todo(chaos): some of this should already be known - optimize
-		const FRigidTransform3 WorldCoMTransform0 = FParticleUtilities::GetCoMWorldTransform(Particle0);
-		const FRigidTransform3 WorldCoMTransform1 = FParticleUtilities::GetCoMWorldTransform(Particle1);
-		FVec3 WorldContactLocation, WorldContactNormal;
-		FReal WorldContactPhi;
-		GetWorldSpaceManifoldPoint(ManifoldPoint, WorldCoMTransform0.GetLocation(), WorldCoMTransform0.GetRotation(), WorldCoMTransform1.GetLocation(), WorldCoMTransform1.GetRotation(), WorldContactLocation, WorldContactNormal, WorldContactPhi);
+		GetWorldSpaceManifoldPoint(ManifoldPoint, PCoM0, QCoM0, PCoM1, QCoM1, ManifoldPoint.ContactPoint.Location, ManifoldPoint.ContactPoint.Normal, ManifoldPoint.ContactPoint.Phi);
 
 		// Calculate and store the data required for static friction and restitution in PushOut
 		// We use PreV and PreW to support incremental manifold generation. In this case, manifold points
 		// can be added after we have already run some solver iterations, which gives us an incorrect initial
 		// velocity if we just use V and W (one-shots will work the same either way since V=PreV on first pass)
-		const FVec3 WorldContactVel0 = Particle0->PreV() + FVec3::CrossProduct(Particle0->PreW(), WorldCoMTransform0.GetRotation() * ManifoldPoint.CoMContactPoints[0]);
-		const FVec3 WorldContactVel1 = Particle1->PreV() + FVec3::CrossProduct(Particle1->PreW(), WorldCoMTransform1.GetRotation() * ManifoldPoint.CoMContactPoints[1]);
-		const FReal WorldContactVelNorm = FVec3::DotProduct(WorldContactVel0 - WorldContactVel1, WorldContactNormal);
+		const FVec3 WorldContactVel0 = PreV0 + FVec3::CrossProduct(PreW0, QCoM0 * ManifoldPoint.CoMContactPoints[0]);
+		const FVec3 WorldContactVel1 = PreV1 + FVec3::CrossProduct(PreW1, QCoM1 * ManifoldPoint.CoMContactPoints[1]);
+		const FReal WorldContactVelNorm = FVec3::DotProduct(WorldContactVel0 - WorldContactVel1, ManifoldPoint.ContactPoint.Normal);
 		ManifoldPoint.InitialContactVelocity = WorldContactVelNorm;
 
 		// Store the initial penetration depth for use with restitution with PBD
 		// NOTE: This is incorrect if we are updating the point each iteration (which we are for incremental manifolds but not one-shots - see UpdateManifoldPoint)
-		ManifoldPoint.InitialPhi = WorldContactPhi;
+		ManifoldPoint.InitialPhi = ManifoldPoint.ContactPoint.Phi;
 	}
 
-	void FRigidBodyPointContactConstraint::CalculatePrevCoMContactPoints(
-		const FConstGenericParticleHandle Particle0,
-		const FConstGenericParticleHandle Particle1,
+	void FPBDCollisionConstraint::CalculatePrevCoMContactPoints(
+		const FSolverBody& Body0,
+		const FSolverBody& Body1,
 		FManifoldPoint& ManifoldPoint,
 		FReal Dt,
 		FVec3& OutPrevCoMContactPoint0,
 		FVec3& OutPrevCoMContactPoint1) const
 	{
+		// @todo(chaos): remove this function when persistent collisions are fully enabled
+		if (!GetUseIncrementalCollisionDetection())
+		{
+			OutPrevCoMContactPoint0 = ManifoldPoint.CoMAnchorPoints[0];
+			OutPrevCoMContactPoint1 = ManifoldPoint.CoMAnchorPoints[1];
+			return;
+		}
+
+
 		// Recalculate the previous local-space contact position on the plane owner. This is used by static friction where
 		// we try to move the contact points back to their previous relative positions.
-		auto CalculatePrevCoMTransform = [Dt](const FConstGenericParticleHandle& ParticleHandle) -> FRigidTransform3
+		auto CalculatePrevCoMTransform = [Dt](const FSolverBody& Body) -> FRigidTransform3
 		{
-			FRigidTransform3 PrevCoMTransform = FParticleUtilitiesXR::GetCoMWorldTransform(ParticleHandle);
-			if (ParticleHandle->ObjectState() == EObjectStateType::Kinematic && (ParticleHandle->V().SizeSquared() > 0.0f || ParticleHandle->W().SizeSquared() > 0.0f))
+			FRigidTransform3 PrevCoMTransform = FRigidTransform3(Body.X(), Body.R());
+			if (!Body.IsDynamic() && (Body.V().SizeSquared() > 0.0f || Body.W().SizeSquared() > 0.0f))
 			{
 				// Undo velocity integration (see KinematicTargets) for kinematic bodies
-				PrevCoMTransform.AddToTranslation(-Dt * ParticleHandle->V());
-				PrevCoMTransform.SetRotation(FRotation3::IntegrateRotationWithAngularVelocity(PrevCoMTransform.GetRotation(), -ParticleHandle->W(), Dt));
+				PrevCoMTransform.AddToTranslation(-Dt * Body.V());
+				PrevCoMTransform.SetRotation(FRotation3::IntegrateRotationWithAngularVelocity(PrevCoMTransform.GetRotation(), -Body.W(), Dt));
 			}
 			return PrevCoMTransform; 
 		};
 
-		const FRigidTransform3 PrevCoMTransform0 = CalculatePrevCoMTransform(Particle0);
-		const FRigidTransform3 PrevCoMTransform1 = CalculatePrevCoMTransform(Particle1);
+		const FRigidTransform3 PrevCoMTransform0 = CalculatePrevCoMTransform(Body0);
+		const FRigidTransform3 PrevCoMTransform1 = CalculatePrevCoMTransform(Body1);
 
 		const FVec3 PrevWorldContactLocation0 = PrevCoMTransform0.GetTranslation() + PrevCoMTransform0.GetRotation() * ManifoldPoint.CoMContactPoints[0];
 		const FVec3 PrevWorldContactLocation1 = PrevCoMTransform1.GetTranslation() + PrevCoMTransform1.GetRotation() * ManifoldPoint.CoMContactPoints[1];
 		FVec3 PrevCoMContactPoint0 = ManifoldPoint.CoMContactPoints[0];
 		FVec3 PrevCoMContactPoint1 = ManifoldPoint.CoMContactPoints[1];
-		if (ManifoldPoint.ContactPoint.ContactNormalOwnerIndex == 0)
+		if (false == bChaos_Collision_Manifold_FixNormalsInWorldSpace)
 		{
-			const FVec3 PrevWorldContactNormal = PrevCoMTransform0.GetRotation() * ManifoldPoint.CoMContactNormal;
-			const FReal PrevPhi = FVec3::DotProduct(PrevWorldContactLocation0 - PrevWorldContactLocation1, PrevWorldContactNormal);
-			PrevCoMContactPoint0 = PrevCoMTransform0.GetRotation().Inverse() * (PrevWorldContactLocation1 + PrevPhi * PrevWorldContactNormal - PrevCoMTransform0.GetTranslation());
-		}
-		else
-		{
-			const FVec3 PrevWorldContactNormal = PrevCoMTransform1.GetRotation() * ManifoldPoint.CoMContactNormal;
-			const FReal PrevPhi = FVec3::DotProduct(PrevWorldContactLocation0 - PrevWorldContactLocation1, PrevWorldContactNormal);
-			PrevCoMContactPoint1 = PrevCoMTransform1.GetRotation().Inverse() * (PrevWorldContactLocation0 - PrevPhi * PrevWorldContactNormal - PrevCoMTransform1.GetTranslation());
+			if (ManifoldPoint.ContactPoint.ContactNormalOwnerIndex == 0)
+			{
+				const FVec3 PrevWorldContactNormal = PrevCoMTransform0.GetRotation() * ManifoldPoint.ManifoldContactNormal;
+				const FReal PrevPhi = FVec3::DotProduct(PrevWorldContactLocation0 - PrevWorldContactLocation1, PrevWorldContactNormal);
+				PrevCoMContactPoint0 = PrevCoMTransform0.GetRotation().Inverse() * (PrevWorldContactLocation1 + PrevPhi * PrevWorldContactNormal - PrevCoMTransform0.GetTranslation());
+			}
+			else
+			{
+				const FVec3 PrevWorldContactNormal = PrevCoMTransform1.GetRotation() * ManifoldPoint.ManifoldContactNormal;
+				const FReal PrevPhi = FVec3::DotProduct(PrevWorldContactLocation0 - PrevWorldContactLocation1, PrevWorldContactNormal);
+				PrevCoMContactPoint1 = PrevCoMTransform1.GetRotation().Inverse() * (PrevWorldContactLocation0 - PrevPhi * PrevWorldContactNormal - PrevCoMTransform1.GetTranslation());
+			}
 		}
+		
 
 		OutPrevCoMContactPoint0 = PrevCoMContactPoint0;
 		OutPrevCoMContactPoint1 = PrevCoMContactPoint1;
 	}
 
-	int32 FRigidBodyPointContactConstraint::AddManifoldPoint(const FContactPoint& ContactPoint, const FReal Dt)
+	int32 FPBDCollisionConstraint::AddManifoldPoint(const FContactPoint& ContactPoint, const FReal Dt)
 	{
 		int32 ManifoldPointIndex = ManifoldPoints.Add(ContactPoint);
 
@@ -300,7 +550,7 @@ namespace Chaos
 		return ManifoldPointIndex;
 	}
 
-	void FRigidBodyPointContactConstraint::UpdateManifoldPoint(int32 ManifoldPointIndex, const FContactPoint& ContactPoint, const FReal Dt)
+	void FPBDCollisionConstraint::UpdateManifoldPoint(int32 ManifoldPointIndex, const FContactPoint& ContactPoint, const FReal Dt)
 	{
 		// We really need to know that it's exactly the same contact and not just a close one to update it here
 		// otherwise the PrevLocalContactPoint1 we calculated is not longer for the correct point.
@@ -311,28 +561,70 @@ namespace Chaos
 
 	// Update the derived contact state (CoM relative data)
 	// @todo(chaos): we shouldn't have to recalculate local space positions - they were available in collision update
-	void FRigidBodyPointContactConstraint::UpdateManifoldPointFromContact(FManifoldPoint& ManifoldPoint)
+	void FPBDCollisionConstraint::UpdateManifoldPointFromContact(FManifoldPoint& ManifoldPoint)
 	{
 		FConstGenericParticleHandle Particle0 = Particle[0];
 		FConstGenericParticleHandle Particle1 = Particle[1];
 
 		check(ManifoldPoint.ContactPoint.ContactNormalOwnerIndex >= 0);
 		check(ManifoldPoint.ContactPoint.ContactNormalOwnerIndex < 2);
-		const FRigidTransform3& PlaneTransform = ImplicitTransform[ManifoldPoint.ContactPoint.ContactNormalOwnerIndex];
+
+		const int32 PlaneOwner = bChaos_Collision_Manifold_FixNormalsInWorldSpace ? 1 : ManifoldPoint.ContactPoint.ContactNormalOwnerIndex;
+		const FRigidTransform3& PlaneTransform = ImplicitTransform[PlaneOwner];
 		const FVec3 LocalContactPoint0 = ImplicitTransform[0].TransformPositionNoScale(ManifoldPoint.ContactPoint.ShapeContactPoints[0]);	// Particle Space on body 0
 		const FVec3 LocalContactPoint1 = ImplicitTransform[1].TransformPositionNoScale(ManifoldPoint.ContactPoint.ShapeContactPoints[1]);	// Particle Space on body 1
-		const FVec3 LocalContactNormal = PlaneTransform.TransformNormalNoScale(ManifoldPoint.ContactPoint.ShapeContactNormal);				// Particle Space on Plane owner
 
-		const FRotation3& PlaneRotationOfMass = (ManifoldPoint.ContactPoint.ContactNormalOwnerIndex == 0) ? Particle0->RotationOfMass() : Particle1->RotationOfMass();
+		// Build the constraint space axes relative to the plane owner. 
+		const FVec3 ContactNormal = PlaneTransform.TransformNormalNoScale(ManifoldPoint.ContactPoint.ShapeContactNormal);				// Particle Space on Plane owner
+		int32 GoodComponentIndex = 0;
+		for (; GoodComponentIndex < 2; GoodComponentIndex++) // Only test first 2 components
+		{
+			if (ContactNormal[GoodComponentIndex] > 0.5) // At least of of the indices must be greater than 1/sqrt(3)
+			{
+				break;
+			}
+		}
+		GoodComponentIndex = (GoodComponentIndex + 1) % 3; // Use any other component
+
+		FVec3 ContactTangents[2];
+
+		ContactTangents[0] = FVec3(0);
+		ContactTangents[0][GoodComponentIndex] = 1.0f;
+		ContactTangents[0] = FVec3::CrossProduct(ContactNormal, ContactTangents[0]).GetSafeNormal();
+		ContactTangents[1] = FVec3::CrossProduct(ContactNormal, ContactTangents[0]);
+
+		
+		
 		const FVec3 CoMContactPoint0 = Particle0->RotationOfMass().Inverse() * (LocalContactPoint0 - Particle0->CenterOfMass());	// CoM Space on Body 0
 		const FVec3 CoMContactPoint1 = Particle1->RotationOfMass().Inverse() * (LocalContactPoint1 - Particle1->CenterOfMass());	// CoM Space on Body 1
-		const FVec3 CoMContactNormal = PlaneRotationOfMass.Inverse() * LocalContactNormal;											// CoM Space on Planer owner
+
+		FVec3 ManifoldContactNormal;
+		FVec3 CoMContactTangent0;
+		FVec3 CoMContactTangent1;
+
+		if (bChaos_Collision_Manifold_FixNormalsInWorldSpace)
+		{
+			const FRotation3&  Particle1Rotation = Particle1->R();// Plane is attached to particle1
+			ManifoldContactNormal = Particle1Rotation * ContactNormal;											// CoM Space on Planer owner
+			CoMContactTangent0 =  Particle1Rotation * ContactTangents[0];
+			CoMContactTangent1 =  Particle1Rotation * ContactTangents[1];
+		}
+		else
+		{
+			const FRotation3&  PlaneRotationOfMass = (ManifoldPoint.ContactPoint.ContactNormalOwnerIndex == 0) ? Particle0->RotationOfMass() : Particle1->RotationOfMass();
+			ManifoldContactNormal = PlaneRotationOfMass.Inverse() * ContactNormal;											// CoM Space on Planer owner
+			CoMContactTangent0 = PlaneRotationOfMass.Inverse() * ContactTangents[0];
+			CoMContactTangent1 = PlaneRotationOfMass.Inverse() * ContactTangents[1];
+		}
+		
 		ManifoldPoint.CoMContactPoints[0] = CoMContactPoint0;
 		ManifoldPoint.CoMContactPoints[1] = CoMContactPoint1;
-		ManifoldPoint.CoMContactNormal = CoMContactNormal;
+		ManifoldPoint.ManifoldContactTangents[0] = CoMContactTangent0;
+		ManifoldPoint.ManifoldContactTangents[1] = CoMContactTangent1;
+		ManifoldPoint.ManifoldContactNormal = ManifoldContactNormal;
 	}
 
-	void FRigidBodyPointContactConstraint::SetActiveContactPoint(const FContactPoint& ContactPoint)
+	void FPBDCollisionConstraint::SetActiveContactPoint(const FContactPoint& ContactPoint)
 	{
 		// @todo(chaos): once we settle on manifolds we should just store the index
 		Manifold.Location = ContactPoint.Location;
@@ -340,7 +632,7 @@ namespace Chaos
 		Manifold.Phi = ContactPoint.Phi;
 	}
 
-	void FRigidBodyPointContactConstraint::GetWorldSpaceManifoldPoint(
+	void FPBDCollisionConstraint::GetWorldSpaceManifoldPoint(
 		const FManifoldPoint& ManifoldPoint,
 		const FVec3& P0,			// World-Space CoM
 		const FRotation3& Q0,		// World-Space CoM
@@ -352,9 +644,17 @@ namespace Chaos
 	{
 		const FReal Margin0 = ManifoldPoint.ContactPoint.ShapeMargins[0];
 		const FReal Margin1 = ManifoldPoint.ContactPoint.ShapeMargins[1];
-		const FRotation3& PlaneQ = (ManifoldPoint.ContactPoint.ContactNormalOwnerIndex == 0) ? Q0 : Q1;
-
-		const FVec3 ContactNormal = PlaneQ.RotateVector(ManifoldPoint.CoMContactNormal);
+		FVec3 ContactNormal;
+		if (bChaos_Collision_Manifold_FixNormalsInWorldSpace)
+		{
+			ContactNormal = ManifoldPoint.ManifoldContactNormal;
+		}
+		else
+		{
+			const FRotation3& PlaneQ = (ManifoldPoint.ContactPoint.ContactNormalOwnerIndex == 0) ? Q0 : Q1;
+			ContactNormal = PlaneQ.RotateVector(ManifoldPoint.ManifoldContactNormal);
+		}
+		
 		const FVec3 ContactPos0 = P0 + Q0.RotateVector(ManifoldPoint.CoMContactPoints[0]) - Margin0 * ContactNormal;
 		const FVec3 ContactPos1 = P1 + Q1.RotateVector(ManifoldPoint.CoMContactPoints[1]) + Margin1 * ContactNormal;
 
@@ -364,7 +664,7 @@ namespace Chaos
 	}
 
 
-	FManifoldPoint& FRigidBodyPointContactConstraint::SetActiveManifoldPoint(
+	FManifoldPoint& FPBDCollisionConstraint::SetActiveManifoldPoint(
 		int32 ManifoldPointIndex,
 		const FVec3& P0,
 		const FRotation3& Q0,
@@ -382,10 +682,146 @@ namespace Chaos
 		return ManifoldPoint;
 	}
 
-	bool FRigidBodyPointContactConstraint::CanUseManifold(FGeometryParticleHandle* Particle0, FGeometryParticleHandle* Particle1) const
+	bool FPBDCollisionConstraint::CanUseManifold(FGeometryParticleHandle* Particle0, FGeometryParticleHandle* Particle1) const
 	{
 		// Do not use manifolds when a body is connected by a joint to another. Manifolds do not work when the bodies may be moved
 		// and rotated by significant amounts and joints can do that.
 		return bChaos_Manifold_EnabledWithJoints || ((Particle0->ParticleConstraints().Num() == 0) && (Particle1->ParticleConstraints().Num() == 0));
 	}
+
+	void FPBDCollisionConstraint::ResetManifold()
+	{
+		ManifoldPoints.Reset();
+		PrevManifoldPoints.Reset();
+		Manifold.Reset();
+	}
+
+	void FPBDCollisionConstraint::SaveManifold()
+	{
+		// Save off the previous data for use by static friction
+		// We do this by swapping the arrays and resetting to avoid allocations
+		Swap(ManifoldPoints, PrevManifoldPoints);
+		ManifoldPoints.Reset();
+		Manifold.Reset();
+
+		bWasManifoldRestored = false;
+	}
+
+	void FPBDCollisionConstraint::RestoreManifold()
+	{
+		// We want to restore the manifold as-is and will skip the narrow phase, which means we just
+		// leave the manifold in place, but we still have some cleanup to do to account for slight movement of the bodies
+		Manifold.Phi = TNumericLimits<FReal>::Max();
+
+		UpdateManifoldContacts();
+
+		bWasManifoldRestored = true;
+	}
+
+	bool FPBDCollisionConstraint::TryRestoreFrictionData(FManifoldPoint& ManifoldPoint)
+	{
+		// Find the previous manifold point that matches
+		const FManifoldPoint* MatchedManifoldPoint = nullptr;
+		if (!IsNew() && bChaos_Manifold_EnableFrictionRestore)
+		{
+			// @todo(chaos): ManifoldPoints and PrevManifoldPoints are usually in the same order, so this loop could normally terminate in 1 iteration
+			for (int32 PrevManifoldPointIndex = 0; PrevManifoldPointIndex < PrevManifoldPoints.Num(); ++PrevManifoldPointIndex)
+			{
+				const FManifoldPoint& PrevManifoldPoint = PrevManifoldPoints[PrevManifoldPointIndex];
+				const FVec3 DP0 = ManifoldPoint.CoMContactPoints[0] - PrevManifoldPoint.CoMAnchorPoints[0];
+				const FVec3 DP1 = ManifoldPoint.CoMContactPoints[1] - PrevManifoldPoint.CoMAnchorPoints[1];
+
+				// If the contact point is in the same spot on one of the bodies, assume it is the same contact
+				// @todo(chaos): more robust same-point test. E.g., this won't work for very small or very large objects,
+				// so at least make the tolerance size-dependent
+				const FReal DistanceToleranceSq = FMath::Square(Chaos_Manifold_FrictionPositionTolerance);
+				if ((DP0.SizeSquared() < DistanceToleranceSq) || (DP1.SizeSquared() < DistanceToleranceSq))
+				{
+					MatchedManifoldPoint = &PrevManifoldPoint;
+					break;
+				}
+			}
+		}
+
+		// If we have a previous point, use it to set the previous-state data required for 
+		// static friction otherwise disable static friction
+		if (MatchedManifoldPoint != nullptr)
+		{
+			ManifoldPoint.CoMAnchorPoints[0] = MatchedManifoldPoint->CoMAnchorPoints[0];
+			ManifoldPoint.CoMAnchorPoints[1] = MatchedManifoldPoint->CoMAnchorPoints[1];
+			ManifoldPoint.StaticFrictionMax = MatchedManifoldPoint->StaticFrictionMax;
+			ManifoldPoint.bPotentialRestingContact = true;
+			return true;
+		}
+		else
+		{
+			// Note: by setting bPotentialRestingContact to false, we are saying that initial contacts
+			// cannot have static friction enabled...
+			ManifoldPoint.CoMAnchorPoints[0] = ManifoldPoint.CoMContactPoints[0];
+			ManifoldPoint.CoMAnchorPoints[1] = ManifoldPoint.CoMContactPoints[1];
+			ManifoldPoint.StaticFrictionMax = FReal(0);
+			ManifoldPoint.bPotentialRestingContact = false;
+			return false;
+		}
+	}
+
+	ECollisionConstraintDirection FPBDCollisionConstraint::GetConstraintDirection(const FReal Dt) const
+	{
+		if (GetDisabled())
+		{
+			return NoRestingDependency;
+		}
+		// D\tau is the chacteristic time (as in GBF paper Sec 8.1)
+		const FReal Dtau = Dt * Chaos_GBFCharacteristicTimeRatio; 
+
+		const FVec3 Normal = GetNormal();
+		const FReal Phi = GetPhi();
+		if (GetPhi() >= GetCullDistance())
+		{
+			return NoRestingDependency;
+		}
+
+		FVec3 GravityDirection = ConcreteContainer()->GetGravityDirection();
+		FReal GravitySize = ConcreteContainer()->GetGravitySize();
+		// When gravity is zero, we still want to sort the constraints instead of having a random order. In this case, set gravity to default gravity.
+		if (GravitySize < SMALL_NUMBER)
+		{
+			GravityDirection = FVec3(0, 0, -1);
+			GravitySize = 980.f;
+		}
+
+		// How far an object travels in gravity direction within time Dtau starting with zero velocity (as in GBF paper Sec 8.1). 
+		// Theoretically this should be 0.5 * GravityMagnitude * Dtau * Dtau.
+		// Omitting 0.5 to be more consistent with our integration scheme.
+		// Multiplying 0.5 can alternatively be achieved by setting Chaos_GBFCharacteristicTimeRatio=sqrt(0.5)
+		const FReal StepSize = GravitySize * Dtau * Dtau; 
+		const FReal NormalDotG = FVec3::DotProduct(Normal, GravityDirection);
+		const FReal NormalDirectionThreshold = 0.1f; // Hack
+		if (NormalDotG < -NormalDirectionThreshold) // Object 0 rests on object 1
+		{
+			if (Phi + NormalDotG * StepSize < 0) // Hack to simulate object 0 falling (as in GBF paper Sec 8.1)
+			{
+				return Particle1ToParticle0;
+			}
+			else
+			{
+				return NoRestingDependency;
+			}
+		}
+		else if (NormalDotG > NormalDirectionThreshold) // Object 1 rests on object 0
+		{
+			if (Phi - NormalDotG * StepSize < 0) // Hack to simulate object 1 falling (as in GBF paper Sec 8.1)
+			{
+				return Particle0ToParticle1;
+			}
+			else
+			{
+				return NoRestingDependency;
+			}
+		}
+		else // Horizontal contact
+		{
+			return NoRestingDependency;
+		}
+	}
 }
\ No newline at end of file
diff --git a/Engine/Source/Runtime/Experimental/Chaos/Private/Chaos/Collision/PBDCollisionConstraintHandle.cpp b/Engine/Source/Runtime/Experimental/Chaos/Private/Chaos/Collision/PBDCollisionConstraintHandle.cpp
--- a/Engine/Source/Runtime/Experimental/Chaos/Private/Chaos/Collision/PBDCollisionConstraintHandle.cpp
+++ b/Engine/Source/Runtime/Experimental/Chaos/Private/Chaos/Collision/PBDCollisionConstraintHandle.cpp
@@ -1,65 +1,79 @@
 // Copyright Epic Games, Inc. All Rights Reserved.
 #include "Chaos/Collision/PBDCollisionConstraintHandle.h"
+#include "Chaos/Collision/PBDCollisionConstraint.h"
+#include "Chaos/Evolution/SolverBody.h"
+#include "Chaos/Evolution/SolverBodyContainer.h"
 #include "Chaos/PBDCollisionConstraints.h"
 
 namespace Chaos
 {
-	typename FPBDCollisionConstraintHandle::FHandleKey FPBDCollisionConstraintHandle::GetKey()
+	const FPBDCollisionConstraints* FPBDCollisionConstraintHandle::ConcreteContainer() const
 	{
-		const FCollisionConstraintBase& Contact = GetContact();
-		return FHandleKey(
-			FImplicitPair(Contact.Manifold.Implicit[0], Contact.Manifold.Implicit[1]),
-			FGeometryPair(Contact.Particle[0], Contact.Particle[1]));
+		return static_cast<FPBDCollisionConstraints*>(ConstraintContainer);
 	}
 
-	const FCollisionConstraintBase& FPBDCollisionConstraintHandle::GetContact() const
+	FPBDCollisionConstraints* FPBDCollisionConstraintHandle::ConcreteContainer()
 	{
-		if (GetType() == FCollisionConstraintBase::FType::SinglePoint)
-		{
-			return ConcreteContainer()->Constraints.SinglePointConstraints[ConstraintIndex];
-		}
-		else
-		{
-			check(GetType() == FCollisionConstraintBase::FType::SinglePointSwept);
-			return ConcreteContainer()->Constraints.SinglePointSweptConstraints[ConstraintIndex];
-		}
+		return static_cast<FPBDCollisionConstraints*>(ConstraintContainer);
 	}
 
-	FCollisionConstraintBase& FPBDCollisionConstraintHandle::GetContact()
+	const FPBDCollisionConstraint& FPBDCollisionConstraintHandle::GetContact() const
 	{
-		if (GetType() == FCollisionConstraintBase::FType::SinglePoint)
-		{
-			return ConcreteContainer()->Constraints.SinglePointConstraints[ConstraintIndex];
-		}
-		else
-		{
-			check(GetType() == FCollisionConstraintBase::FType::SinglePointSwept);
-			return ConcreteContainer()->Constraints.SinglePointSweptConstraints[ConstraintIndex];
-		}
+		return *GetConstraint();
 	}
 
-	const FRigidBodyPointContactConstraint& FPBDCollisionConstraintHandle::GetPointContact() const
-	{ 
-		check(GetType() == FCollisionConstraintBase::FType::SinglePoint);
-		return ConcreteContainer()->Constraints.SinglePointConstraints[ConstraintIndex];
+	FPBDCollisionConstraint& FPBDCollisionConstraintHandle::GetContact()
+	{
+		return *GetConstraint();
+	}
+
+	ECollisionConstraintType FPBDCollisionConstraintHandle::GetType() const
+	{
+		return GetContact().GetType();
+	}
+
+	void FPBDCollisionConstraintHandle::SetEnabled(bool InEnabled)
+	{
+		GetContact().SetDisabled(!InEnabled);
+	}
+
+	bool FPBDCollisionConstraintHandle::IsEnabled() const
+	{
+		return !GetContact().GetDisabled();
+	}
+
+	FVec3 FPBDCollisionConstraintHandle::GetContactLocation() const
+	{
+		return GetContact().GetLocation();
+	}
+
+	FVec3 FPBDCollisionConstraintHandle::GetAccumulatedImpulse() const
+	{
+		return GetContact().AccumulatedImpulse;
+	}
+
+	TVector<const TGeometryParticleHandle<FReal, 3>*, 2> FPBDCollisionConstraintHandle::GetConstrainedParticles() const
+	{
+		return { GetContact().Particle[0], GetContact().Particle[1] };
 	}
-	
-	FRigidBodyPointContactConstraint& FPBDCollisionConstraintHandle::GetPointContact() 
-	{ 
-		check(GetType() == FCollisionConstraintBase::FType::SinglePoint);
-		return ConcreteContainer()->Constraints.SinglePointConstraints[ConstraintIndex];
+
+	TVector<TGeometryParticleHandle<FReal, 3>*, 2> FPBDCollisionConstraintHandle::GetConstrainedParticles()
+	{
+		return { GetContact().Particle[0], GetContact().Particle[1] };
 	}
 
-	const FRigidBodySweptPointContactConstraint& FPBDCollisionConstraintHandle::GetSweptPointContact() const
+	void FPBDCollisionConstraintHandle::GatherInput(FReal Dt, const int32 Particle0Level, const int32 Particle1Level, FPBDIslandSolverData& SolverData)
 	{
-		check(GetType() == FCollisionConstraintBase::FType::SinglePointSwept);
-		return ConcreteContainer()->Constraints.SinglePointSweptConstraints[ConstraintIndex];
+		ConcreteContainer()->GatherInput(Dt, GetContact(), Particle0Level, Particle1Level, SolverData);
 	}
-	
-	FRigidBodySweptPointContactConstraint& FPBDCollisionConstraintHandle::GetSweptPointContact()
+
+	FSolverBody* FPBDCollisionConstraintHandle::GetSolverBody0()
 	{
-		check(GetType() == FCollisionConstraintBase::FType::SinglePointSwept);
-		return ConcreteContainer()->Constraints.SinglePointSweptConstraints[ConstraintIndex];
+		return GetContact().GetSolverBody0();
 	}
 
+	FSolverBody* FPBDCollisionConstraintHandle::GetSolverBody1()
+	{
+		return GetContact().GetSolverBody1();
+	}
 }
\ No newline at end of file
diff --git a/Engine/Source/Runtime/Experimental/Chaos/Private/Chaos/Collision/PBDCollisionSolver.cpp b/Engine/Source/Runtime/Experimental/Chaos/Private/Chaos/Collision/PBDCollisionSolver.cpp
--- a/Engine/Source/Runtime/Experimental/Chaos/Private/Chaos/Collision/PBDCollisionSolver.cpp
+++ b/Engine/Source/Runtime/Experimental/Chaos/Private/Chaos/Collision/PBDCollisionSolver.cpp
@@ -3,9 +3,11 @@
 
 #include "Chaos/Collision/CollisionApplyType.h"
 #include "Chaos/Collision/PBDCollisionConstraint.h"
+#include "Chaos/Collision/SolverCollisionContainer.h"
 #include "Chaos/CollisionResolution.h"
 #include "Chaos/CollisionResolutionTypes.h"
 #include "Chaos/CollisionResolutionUtil.h"
+#include "Chaos/Evolution/SolverBodyContainer.h"
 #include "Chaos/Particle/ParticleUtilities.h"
 #include "Chaos/PBDCollisionConstraints.h"
 #include "Chaos/PBDCollisionConstraintsContact.h"
@@ -13,547 +15,630 @@
 
 //PRAGMA_DISABLE_OPTIMIZATION
 
-bool bChaos_PBD_Collision_UseCorrectApplyPosition = true;	// Temp for testing
-FAutoConsoleVariableRef CVarChaos_PBD_Collision_UseCorrectApplyPosition(TEXT("p.Chaos.PBD.Collision.UseCorrectApplyPosition"), bChaos_PBD_Collision_UseCorrectApplyPosition, TEXT(""));
-
-bool bChaos_PBD_Collision_Position_SolveEnabled = true;
-float Chaos_PBD_Collision_Position_MinStiffness = 1.0f;
-float Chaos_PBD_Collision_Position_MaxStiffness = 1.0f;
-int32 Chaos_PBD_Collision_Position_ShockPropagationIterations = 0;
-float Chaos_PBD_Collision_Position_MinInvMassScale = 0.0f;
-float Chaos_PBD_Collision_Position_StaticFrictionIterations = 0.5f;
-float Chaos_PBD_Collision_Position_Tolerance = 0.1f;
-float Chaos_PBD_Collision_Position_ShockPropagationNormalZ = 0.8f;
-
-FAutoConsoleVariableRef CVarChaos_PBD_Collision_Position_SolveEnabled(TEXT("p.Chaos.PBD.Collision.Position.SolveEnabled"), bChaos_PBD_Collision_Position_SolveEnabled, TEXT(""));
-FAutoConsoleVariableRef CVarChaos_PBD_Collision_Position_MinStiffness(TEXT("p.Chaos.PBD.Collision.Position.MinStiffness"), Chaos_PBD_Collision_Position_MinStiffness, TEXT(""));
-FAutoConsoleVariableRef CVarChaos_PBD_Collision_Position_MaxStiffness(TEXT("p.Chaos.PBD.Collision.Position.MaxStiffness"), Chaos_PBD_Collision_Position_MaxStiffness, TEXT(""));
-FAutoConsoleVariableRef CVarChaos_PBD_Collision_Position_UseShockPropagation(TEXT("p.Chaos.PBD.Collision.Position.ShockPropagationIterations"), Chaos_PBD_Collision_Position_ShockPropagationIterations, TEXT(""));
-FAutoConsoleVariableRef CVarChaos_PBD_Collision_Position_MinInvMassScale(TEXT("p.Chaos.PBD.Collision.Position.MinInvMassScale"), Chaos_PBD_Collision_Position_MinInvMassScale, TEXT(""));
-FAutoConsoleVariableRef CVarChaos_PBD_Collision_Position_StaticFrictionIterations(TEXT("p.Chaos.PBD.Collision.Position.StaticFrictionIterations"), Chaos_PBD_Collision_Position_StaticFrictionIterations, TEXT(""));
-FAutoConsoleVariableRef CVarChaos_PBD_Collision_Position_Tolerance(TEXT("p.Chaos.PBD.Collision.Position.Tolerance"), Chaos_PBD_Collision_Position_Tolerance, TEXT(""));
-
-bool bChaos_PBD_Collision_Velocity_SolveEnabled = true;
-float Chaos_PBD_Collision_Velocity_MinStiffness = 1.0f;
-float Chaos_PBD_Collision_Velocity_MaxStiffness = 1.0f;
-int32 Chaos_PBD_Collision_Velocity_ShockPropagationIterations = 0;
-float Chaos_PBD_Collision_Velocity_MinInvMassScale = 0.0f;
-bool bChaos_PBD_Collision_Velocity_DynamicFrictionEnabled = true;
-bool bChaos_PBD_Collision_Velocity_ImpulseClampEnabled = false;
-
-
-FAutoConsoleVariableRef CVarChaos_PBD_Collision_Velocity_SolveEnabled(TEXT("p.Chaos.PBD.Collision.Velocity.SolveEnabled"), bChaos_PBD_Collision_Velocity_SolveEnabled, TEXT(""));
-FAutoConsoleVariableRef CVarChaos_PBD_Collision_Velocity_MinStiffness(TEXT("p.Chaos.PBD.Collision.Velocity.MinStiffness"), Chaos_PBD_Collision_Velocity_MinStiffness, TEXT(""));
-FAutoConsoleVariableRef CVarChaos_PBD_Collision_Velocity_MaxStiffness(TEXT("p.Chaos.PBD.Collision.Velocity.MaxStiffness"), Chaos_PBD_Collision_Velocity_MaxStiffness, TEXT(""));
-FAutoConsoleVariableRef CVarChaos_PBD_Collision_Velocity_UseShockPropagation(TEXT("p.Chaos.PBD.Collision.Velocity.ShockPropagationIterations"), Chaos_PBD_Collision_Velocity_ShockPropagationIterations, TEXT(""));
-FAutoConsoleVariableRef CVarChaos_PBD_Collision_Velocity_MinInvMassScale(TEXT("p.Chaos.PBD.Collision.Velocity.MinInvMassScale"), Chaos_PBD_Collision_Velocity_MinInvMassScale, TEXT(""));
-FAutoConsoleVariableRef CVarChaos_PBD_Collision_Velocity_DynamicFrictionEnabled(TEXT("p.Chaos.PBD.Collision.Velocity.DynamicFrictionEnabled"), bChaos_PBD_Collision_Velocity_DynamicFrictionEnabled, TEXT(""));
-FAutoConsoleVariableRef CVarChaos_PBD_Collision_Velocity_ImpulseClampEnabled(TEXT("p.Chaos.PBD.Collision.Velocity.ImpulseClampEnabled"), bChaos_PBD_Collision_Velocity_ImpulseClampEnabled, TEXT(""));
-
+DEFINE_LOG_CATEGORY(LogChaosCollision);
 
 namespace Chaos
 {
-	inline FReal GetPositionSolveStiffness(int32 It, int32 NumIts)
+	namespace CVars
 	{
-		if (NumIts > 1)
-		{
-			const FReal Interpolant = (FReal)It / (FReal)(NumIts - 1);
-			return FMath::Lerp((FReal)Chaos_PBD_Collision_Position_MinStiffness, (FReal)Chaos_PBD_Collision_Position_MaxStiffness, Interpolant);
-		}
-		return Chaos_PBD_Collision_Position_MaxStiffness;
+		extern int32 Chaos_Collision_UseShockPropagation;
+
+		bool bChaos_PBDCollisionSolver_Position_SolveEnabled = true;
+		int32 Chaos_PBDCollisionSolver_Position_ShockPropagationIterations = 3;
+		float Chaos_PBDCollisionSolver_Position_MinInvMassScale = 0.5f;
+		int32 Chaos_PBDCollisionSolver_Position_ZeroFrictionIterations = 4;
+		float Chaos_PBDCollisionSolver_Position_NormalTolerance = 0.1f;
+		bool bChaos_PBDCollisionSolver_Position_NegativePushOutEnabled = true;
+		float Chaos_PBDCollisionSolver_Position_StaticFrictionStiffness = 0.5f;
+		float Chaos_PBDCollisionSolver_Position_StaticFrictionLerpRate = 0.1f;
+		//float Chaos_PBDCollisionSolver_Position_PositionSolverTolerance = 0.01f;
+		//float Chaos_PBDCollisionSolver_Position_RotationSolverTolerance = 0.01f;
+
+		FAutoConsoleVariableRef CVarChaos_PBDCollisionSolver_Position_SolveEnabled(TEXT("p.Chaos.PBDCollisionSolver.Position.SolveEnabled"), bChaos_PBDCollisionSolver_Position_SolveEnabled, TEXT(""));
+		FAutoConsoleVariableRef CVarChaos_PBDCollisionSolver_Position_UseShockPropagation(TEXT("p.Chaos.PBDCollisionSolver.Position.ShockPropagationIterations"), Chaos_PBDCollisionSolver_Position_ShockPropagationIterations, TEXT(""));
+		FAutoConsoleVariableRef CVarChaos_PBDCollisionSolver_Position_MinInvMassScale(TEXT("p.Chaos.PBDCollisionSolver.Position.MinInvMassScale"), Chaos_PBDCollisionSolver_Position_MinInvMassScale, TEXT(""));
+		FAutoConsoleVariableRef CVarChaos_PBDCollisionSolver_Position_ZeroFrictionIterations(TEXT("p.Chaos.PBDCollisionSolver.Position.ZeroFrictionIterations"), Chaos_PBDCollisionSolver_Position_ZeroFrictionIterations, TEXT(""));
+		FAutoConsoleVariableRef CVarChaos_PBDCollisionSolver_Position_NormalTolerance(TEXT("p.Chaos.PBDCollisionSolver.Position.NormalTolerance"), Chaos_PBDCollisionSolver_Position_NormalTolerance, TEXT(""));
+		FAutoConsoleVariableRef CVarChaos_PBDCollisionSolver_Position_StaticFrictionStiffness(TEXT("p.Chaos.PBDCollisionSolver.Position.StaticFriction.Stiffness"), Chaos_PBDCollisionSolver_Position_StaticFrictionStiffness, TEXT(""));
+		FAutoConsoleVariableRef CVarChaos_PBDCollisionSolver_Position_StaticFrictionLerpRate(TEXT("p.Chaos.PBDCollisionSolver.Position.StaticFriction.LerpRate"), Chaos_PBDCollisionSolver_Position_StaticFrictionLerpRate, TEXT(""));
+		//FAutoConsoleVariableRef CVarChaos_PBDCollisionSolver_Position_PositionSolverTolerance(TEXT("p.Chaos.PBDCollisionSolver.Position.PositionTolerance"), Chaos_PBDCollisionSolver_Position_PositionSolverTolerance, TEXT(""));
+		//FAutoConsoleVariableRef CVarChaos_PBDCollisionSolver_Position_RotationSolverTolerance(TEXT("p.Chaos.PBDCollisionSolver.Position.RotationTolerance"), Chaos_PBDCollisionSolver_Position_RotationSolverTolerance, TEXT(""));
+
+		bool bChaos_PBDCollisionSolver_Velocity_SolveEnabled = true;
+		int32 Chaos_PBDCollisionSolver_Velocity_ShockPropagationIterations = 1;
+		float Chaos_PBDCollisionSolver_Velocity_MinInvMassScale = 0.1f;
+		bool bChaos_PBDCollisionSolver_Velocity_DynamicFrictionEnabled = true;
+		bool bChaos_PBDCollisionSolver_Velocity_NegativeImpulseEnabled = true;
+		bool bChaos_PBDCollisionSolver_Velocity_ImpulseClampEnabled = false;
+
+		FAutoConsoleVariableRef CVarChaos_PBDCollisionSolver_Velocity_SolveEnabled(TEXT("p.Chaos.PBDCollisionSolver.Velocity.SolveEnabled"), bChaos_PBDCollisionSolver_Velocity_SolveEnabled, TEXT(""));
+		FAutoConsoleVariableRef CVarChaos_PBDCollisionSolver_Velocity_UseShockPropagation(TEXT("p.Chaos.PBDCollisionSolver.Velocity.ShockPropagationIterations"), Chaos_PBDCollisionSolver_Velocity_ShockPropagationIterations, TEXT(""));
+		FAutoConsoleVariableRef CVarChaos_PBDCollisionSolver_Velocity_MinInvMassScale(TEXT("p.Chaos.PBDCollisionSolver.Velocity.MinInvMassScale"), Chaos_PBDCollisionSolver_Velocity_MinInvMassScale, TEXT(""));
+		FAutoConsoleVariableRef CVarChaos_PBDCollisionSolver_Velocity_DynamicFrictionEnabled(TEXT("p.Chaos.PBDCollisionSolver.Velocity.DynamicFrictionEnabled"), bChaos_PBDCollisionSolver_Velocity_DynamicFrictionEnabled, TEXT(""));
+		FAutoConsoleVariableRef CVarChaos_PBDCollisionSolver_Velocity_NegativeImpulseEnabled(TEXT("p.Chaos.PBDCollisionSolver.Velocity.NegativeImpulseEnabled"), bChaos_PBDCollisionSolver_Velocity_NegativeImpulseEnabled, TEXT(""));
+		FAutoConsoleVariableRef CVarChaos_PBDCollisionSolver_Velocity_ImpulseClampEnabled(TEXT("p.Chaos.PBDCollisionSolver.Velocity.ImpulseClampEnabled"), bChaos_PBDCollisionSolver_Velocity_ImpulseClampEnabled, TEXT(""));
 	}
+	using namespace CVars;
 
 	inline FReal GetPositionInvMassScale(int32 It, int32 NumIts)
 	{
-		if ((NumIts > 1) && (Chaos_PBD_Collision_Position_ShockPropagationIterations > 0))
+		if (Chaos_PBDCollisionSolver_Position_ShockPropagationIterations > 0)
 		{
-			const int32 FirstShockIt = FMath::Max(NumIts - Chaos_PBD_Collision_Position_ShockPropagationIterations, 0);
-			const FReal Interpolant = FMath::Clamp(FReal(It - FirstShockIt) / FReal(NumIts - FirstShockIt - 1), FReal(0), FReal(1));
-			return FMath::Lerp((FReal)1.0f, (FReal)Chaos_PBD_Collision_Position_MinInvMassScale, Interpolant);
+			const int32 FirstShockIt = FMath::Max(NumIts - Chaos_PBDCollisionSolver_Position_ShockPropagationIterations, 0);
+			const FReal Interpolant = FMath::Clamp(FReal(It - FirstShockIt + 1) / FReal(NumIts - FirstShockIt), FReal(0), FReal(1));
+			return FMath::Lerp((FReal)1.0f, (FReal)Chaos_PBDCollisionSolver_Position_MinInvMassScale, Interpolant);
 		}
 		return 1.0f;
 	}
 
-	inline FReal GetVelocitySolveStiffness(int32 It, int32 NumIts)
+	inline FReal GetVelocityInvMassScale(int32 It, int32 NumIts)
 	{
-		if (NumIts > 1)
+		if (Chaos_PBDCollisionSolver_Velocity_ShockPropagationIterations > 0)
 		{
-			const FReal Interpolant = (FReal)It / (FReal)(NumIts - 1);
-			return FMath::Lerp((FReal)Chaos_PBD_Collision_Velocity_MinStiffness, (FReal)Chaos_PBD_Collision_Velocity_MaxStiffness, Interpolant);
+			const int32 FirstShockIt = FMath::Max(NumIts - Chaos_PBDCollisionSolver_Velocity_ShockPropagationIterations, 0);
+			const FReal Interpolant = FMath::Clamp(FReal(It - FirstShockIt + 1) / FReal(NumIts - FirstShockIt), FReal(0), FReal(1));
+			return FMath::Lerp((FReal)1.0f, (FReal)Chaos_PBDCollisionSolver_Velocity_MinInvMassScale, Interpolant);
 		}
-		return Chaos_PBD_Collision_Velocity_MaxStiffness;
+		return 1.0f;
 	}
 
-	inline FReal GetVelocityInvMassScale(int32 It, int32 NumIts)
+	inline void CalculatePositionCorrectionWithoutFriction(
+		const FReal Stiffness,
+		const FReal ContactDeltaNormal,
+		const FVec3& ContactNormal,
+		const FReal ContactMassNormal,
+		FVec3& InOutNetPushOut,
+		FVec3& OutPushOut)
 	{
-		if ((NumIts > 1) && (Chaos_PBD_Collision_Velocity_ShockPropagationIterations > 0))
+		FVec3 PushOut = -(Stiffness * ContactDeltaNormal * ContactMassNormal) * ContactNormal;
+
+		// The total pushout so far this sub-step
+		// We allow negative incremental impulses, but not net negative impulses
+		const FVec3 NetPushOut = InOutNetPushOut + PushOut;
+		const FReal NetPushOutNormal = FVec3::DotProduct(NetPushOut, ContactNormal);
+		if (NetPushOutNormal < 0)
 		{
-			const int32 FirstShockIt = FMath::Max(NumIts - Chaos_PBD_Collision_Velocity_ShockPropagationIterations, 0);
-			const FReal Interpolant = FMath::Clamp(FReal(It - FirstShockIt) / FReal(NumIts - FirstShockIt - 1), FReal(0), FReal(1));
-			return FMath::Lerp((FReal)1.0f, (FReal)Chaos_PBD_Collision_Velocity_MinInvMassScale, Interpolant);
+			PushOut = -InOutNetPushOut;
 		}
-		return 1.0f;
+
+		InOutNetPushOut += PushOut;
+		OutPushOut = PushOut;
 	}
 
 
-	void ApplyPositionCorrection(
+	inline void CalculatePositionCorrectionWithFriction(
 		const FReal Stiffness,
-		const Collisions::FContactIterationParameters& IterationParameters,
-		const Collisions::FContactParticleParameters& ParticleParameters,
-		const FVec3& ApplyContactPoint0,
-		const FVec3& ApplyContactPoint1,
-		const FVec3& ContactError,
+		const FVec3& ContactDelta,
+		const FReal ContactDeltaNormal,
 		const FVec3& ContactNormal,
 		const FMatrix33& ContactMass,
-		const FReal ContactMassNormal,
 		const FReal StaticFriction,
-		const FReal InvM0,
-		const FMatrix33& InvI0,
-		const FReal InvM1,
-		const FMatrix33& InvI1,
-		FVec3& P0, // Centre of Mass Positions and Rotations
-		FRotation3& Q0,
-		FVec3& P1,
-		FRotation3& Q1,
-		FManifoldPoint& ManifoldPoint)
+		const FReal DynamicFriction,
+		FVec3& InOutNetPushOut,
+		FVec3& OutPushOut,
+		FReal& InOutStaticFrictionMax,
+		bool& bOutInsideStaticFrictionCone)
 	{
-		// Calculate the pushout required to correct the error for this iteration
-		FVec3 PushOut = FVec3(0);
-		if ((StaticFriction > 0) && ManifoldPoint.bInsideStaticFrictionCone)
+		// If static friction is enabled, calculate the correction to move the contact point back to its
+		// original relative location on all axes.
+		// @todo(chaos): this should be moved to the ManifoldPoint error calculation?
+		FVec3 ModifiedContactError = -ContactDelta;
+		const FReal FrictionStiffness = Chaos_PBDCollisionSolver_Position_StaticFrictionStiffness;
+		if (FrictionStiffness < 1.0f)
 		{
-			PushOut = Stiffness * ContactMass * ContactError;
-		}
-		else
-		{
-			const FReal ContactErrorNormal = FVec3::DotProduct(ContactError, ContactNormal);
-			PushOut = (Stiffness * ContactErrorNormal * ContactMassNormal) * ContactNormal;
+			const FVec3 ContactDeltaTangent = ContactDelta - ContactDeltaNormal * ContactNormal;
+			ModifiedContactError = -ContactDeltaNormal * ContactNormal - FrictionStiffness * ContactDeltaTangent;
 		}
 
-		// The total pushout so far this sub-step
-		FVec3 NetPushOut = ManifoldPoint.NetPushOut + PushOut;
-		FReal NetPushOutNormal = FVec3::DotProduct(NetPushOut, ContactNormal);
+		FVec3 PushOut = Stiffness * ContactMass * ModifiedContactError;
 
-		// If we ended up with a negative normal pushout, stop
-		if (NetPushOutNormal < 0)
+		// If we ended up with a negative normal pushout, disable friction
+		const FVec3 NetPushOut = InOutNetPushOut + PushOut;
+		const FReal NetPushOutNormal = FVec3::DotProduct(NetPushOut, ContactNormal);
+		bool bInsideStaticFrictionCone = true;
+		if (NetPushOutNormal < FReal(SMALL_NUMBER))
 		{
-			PushOut = -ManifoldPoint.NetPushOut;
-			NetPushOut = FVec3(0);
-			NetPushOutNormal = 0.0f;
+			bInsideStaticFrictionCone = false;
 		}
 
-		// If we exceed the friction cone, remove all tangential correction and rely on the velocity 
-		// solve to add dynamic friction. Also disable static friction from now on.
-		// Note: this will not result in the correct normal impulse, but it should be corrected
-		// on a subsequent iteration (we could probably do better here)
-		if ((StaticFriction > 0.0f) && (NetPushOutNormal > 0.0f) && ManifoldPoint.bInsideStaticFrictionCone)
+		// Static friction limit: immediately increase maximum lateral correction, but smoothly decay maximum static friction limit. 
+		// This is so that small variations in position (jitter) and therefore NetPushOutNormal don't cause static friction to slip
+		// @todo(chaos): StaticFriction smoothing is iteration count depenendent - try to make it not so
+		const FReal StaticFrictionLerpRate = Chaos_PBDCollisionSolver_Position_StaticFrictionLerpRate;
+		const FReal StaticFrictionDest = FMath::Max(NetPushOutNormal, FReal(0));
+		FReal StaticFrictionMax = FMath::Lerp(FMath::Max(InOutStaticFrictionMax, StaticFrictionDest), StaticFrictionDest, StaticFrictionLerpRate);
+
+		// If we exceed the friction cone, stop adding frictional corrections (although
+		// any already added lateral corrections will not be undone)
+		// @todo(chaos): clamp to dynamic friction
+		if (bInsideStaticFrictionCone)
 		{
-			const FReal MaxPushOutTangential = StaticFriction * NetPushOutNormal;
-
-			// Current tangential correction
-			const FVec3 NetPushOutTangential = NetPushOut - NetPushOutNormal * ContactNormal;
-			const FReal NetPushOutTangentialLen = NetPushOutTangential.Size();
-
-			// Check friction cone
-			if (NetPushOutTangentialLen > MaxPushOutTangential)
+			const FReal MaxPushOutTangentSq = FMath::Square(StaticFriction * StaticFrictionMax);
+			const FReal NetPushOutTangentSq = (NetPushOut - NetPushOutNormal * ContactNormal).SizeSquared();
+			if (NetPushOutTangentSq > MaxPushOutTangentSq)
 			{
-				ManifoldPoint.bInsideStaticFrictionCone = false;
-				PushOut = PushOut - NetPushOutTangential;
-				NetPushOut = ManifoldPoint.NetPushOut + PushOut;
+				bInsideStaticFrictionCone = false;
 			}
 		}
 
-		// Update the particle state based on the pushout
-		if (InvM0 > SMALL_NUMBER)
+		// If we leave the friction cone, we will flll through into the non-friction impulse calculation
+		// so do not export the results
+		if (bInsideStaticFrictionCone)
 		{
-			const FVec3 AngularPushOut = FVec3::CrossProduct(ApplyContactPoint0, PushOut);
-			const FVec3 DX0 = InvM0 * PushOut;
-			const FVec3 DR0 = InvI0 * AngularPushOut;
-
-			P0 += DX0;
-			Q0 += FRotation3::FromElements(DR0, 0.f) * Q0 * FReal(0.5);
-			Q0.Normalize();
+			InOutNetPushOut = NetPushOut;
+			OutPushOut = PushOut;
 		}
-		if (InvM1 > SMALL_NUMBER)
+		else
 		{
-			const FVec3 AngularPushOut = FVec3::CrossProduct(ApplyContactPoint1, PushOut);
-			const FVec3 DX1 = -(InvM1 * PushOut);
-			const FVec3 DR1 = -(InvI1 * AngularPushOut);
+			OutPushOut = FVec3(0);
+		}
+
+		InOutStaticFrictionMax = StaticFrictionMax;
+		bOutInsideStaticFrictionCone = bInsideStaticFrictionCone;
+	}
 
-			P1 += DX1;
-			Q1 += FRotation3::FromElements(DR1, 0.f) * Q1 * FReal(0.5);
-			Q1.Normalize();
+	inline void ApplyPositionCorrection(
+		const FReal Stiffness,
+		const FReal StaticFriction,
+		const FReal DynamicFriction,
+		const FVec3& ContactDelta,
+		const FReal ContactDeltaNormal,
+		FPBDCollisionSolver::FSolverManifoldPoint& ManifoldPoint,
+		FConstraintSolverBody& Body0,
+		FConstraintSolverBody& Body1)
+	{
+		FVec3 PushOut = FVec3(0);
+
+		// If we want to add static friction...
+		// (note: we run a few iterations without friction by temporarily setting StaticFriction to 0)
+		if ((StaticFriction > 0) && ManifoldPoint.bInsideStaticFrictionCone)
+		{
+			CalculatePositionCorrectionWithFriction(
+				Stiffness,
+				ContactDelta,
+				ContactDeltaNormal,
+				ManifoldPoint.WorldContactNormal,
+				ManifoldPoint.WorldContactMass,
+				StaticFriction,
+				DynamicFriction,
+				ManifoldPoint.NetPushOut,					// Out
+				PushOut,									// Out
+				ManifoldPoint.StaticFrictionMax,			// Out
+				ManifoldPoint.bInsideStaticFrictionCone);	// Out
 		}
+		
+		// If friction is disabled or we tried to add friction but left the friction cone...
+		if ((StaticFriction == 0) || !ManifoldPoint.bInsideStaticFrictionCone)
+		{
+			CalculatePositionCorrectionWithoutFriction(
+				Stiffness,
+				ContactDeltaNormal,
+				ManifoldPoint.WorldContactNormal,
+				ManifoldPoint.WorldContactMassNormal,
+				ManifoldPoint.NetPushOut,					// Out
+				PushOut);									// Out
+		}
+
+		//UE_LOG(LogChaosCollision, VeryVerbose, TEXT("    PushOut %f"), PushOut.Size());
 
-		ManifoldPoint.NetPushOut = NetPushOut;
-		ManifoldPoint.bActive = (NetPushOutNormal > 0);
+		// Update the particle state based on the pushout
+		if (Body0.IsDynamic())
+		{
+			const FVec3 AngularPushOut = FVec3::CrossProduct(ManifoldPoint.WorldRelativeImpulsePoint0, PushOut);
+			const FVec3 DX0 = Body0.InvM() * PushOut;
+			const FVec3 DR0 = Body0.InvI() * AngularPushOut;
+			//if (!DX0.IsNearlyZero(Chaos_PBDCollisionSolver_Position_PositionSolverTolerance) || !DR0.IsNearlyZero(Chaos_PBDCollisionSolver_Position_RotationSolverTolerance))
+			{
+				Body0.ApplyTransformDelta(DX0, DR0);
+				//Body0.UpdateRotationDependentState();
+			}
+		}
+		if (Body1.IsDynamic())
+		{
+			const FVec3 AngularPushOut = FVec3::CrossProduct(ManifoldPoint.WorldRelativeImpulsePoint1, PushOut);
+			const FVec3 DX1 = -(Body1.InvM() * PushOut);
+			const FVec3 DR1 = -(Body1.InvI() * AngularPushOut);
+			//if (!DX1.IsNearlyZero(Chaos_PBDCollisionSolver_Position_PositionSolverTolerance) || !DR1.IsNearlyZero(Chaos_PBDCollisionSolver_Position_RotationSolverTolerance))
+			{
+				Body1.ApplyTransformDelta(DX1, DR1);
+				//Body1.UpdateRotationDependentState();
+			}
+		}
 	}
 
-	void ApplyVelocityCorrection(
+	inline void CalculateVelocityCorrectionImpulse(
 		const FReal Stiffness,
-		const Collisions::FContactIterationParameters& IterationParameters,
-		const Collisions::FContactParticleParameters& ParticleParameters,
-		const FVec3& ApplyContactPoint0,
-		const FVec3& ApplyContactPoint1,
+		const FReal Dt,
+		const FReal DynamicFriction,
 		const FVec3& ContactNormal,
 		const FMatrix33& ContactMass,
 		const FReal ContactMassNormal,
-		const FReal DynamicFriction,
-		const FReal Restitution,
-		const FReal InvM0,
-		const FMatrix33& InvI0,
-		const FReal InvM1,
-		const FMatrix33& InvI1,
-		FVec3& V0,
-		FVec3& W0,
-		FVec3& V1,
-		FVec3& W1,
-		FManifoldPoint& ManifoldPoint)
+		const FVec3& ContactVelocityDelta,
+		const FReal ContactVelocityDeltaNormal,
+		const FVec3& NetPushOut,
+		FVec3& InOutNetImpulse,
+		FVec3& OutImpulse)
 	{
-		const FVec3 ContactVelocity0 = V0 + FVec3::CrossProduct(W0, ApplyContactPoint0);
-		const FVec3 ContactVelocity1 = V1 + FVec3::CrossProduct(W1, ApplyContactPoint1);
-		const FVec3 ContactVelocity = ContactVelocity0 - ContactVelocity1;
-		const FReal ContactVelocityNormal = FVec3::DotProduct(ContactVelocity, ContactNormal);
-
-		// Add up the errors in the velocity (current velocity - desired velocity)
-		FVec3 ContactVelocityError = FVec3(0);
 		FVec3 Impulse = FVec3(0);
 
-		// Target normal velocity, including restitution
-		FReal ContactVelocityTargetNormal = 0.0f;
-		const bool bApplyRestitution = (Restitution > 0.0f) && (ManifoldPoint.InitialContactVelocity < -ParticleParameters.RestitutionVelocityThreshold);
-		if (bApplyRestitution)
+		if ((ContactVelocityDeltaNormal > FReal(0)) && !bChaos_PBDCollisionSolver_Velocity_NegativeImpulseEnabled)
 		{
-			ContactVelocityTargetNormal = FMath::Max(0.0f, -Restitution * ManifoldPoint.InitialContactVelocity);
+			return;
 		}
-		ContactVelocityError += (ContactVelocityNormal - ContactVelocityTargetNormal) * ContactNormal;
 
 		// Tangential velocity (dynamic friction)
-		const bool bApplyFriction = (DynamicFriction > 0) && (IterationParameters.Dt > 0);
+		const bool bApplyFriction = (DynamicFriction > 0) && (Dt > 0);
 		if (bApplyFriction)
 		{
-			const FVec3 ContactVelocityTangential = ContactVelocity - ContactVelocityNormal * ContactNormal;
-			const FReal ContactVelocityTangentialLen = ContactVelocityTangential.Size();
-			if (ContactVelocityTangentialLen > SMALL_NUMBER)
-			{
-				// PushOut = ContactMass * DP, where DP is the contact positional correction
-				// Friction force is proportional to the normal force, so friction velocity correction
-				// is proprtional to normal velocity correction, or DVn = DPn/dt = PushOut.N / (ContactMass * dt);
-				const FReal PushOutNormal = FVec3::DotProduct(ManifoldPoint.NetPushOut, ContactNormal);
-				const FReal DynamicFrictionVelocityError = PushOutNormal / (ContactMassNormal * IterationParameters.Dt);
-				if (DynamicFrictionVelocityError > SMALL_NUMBER)
-				{
-					const FReal ContactVelocityErrorTangential = FMath::Min(DynamicFriction * DynamicFrictionVelocityError, ContactVelocityTangentialLen);
-					ContactVelocityError += ContactVelocityTangential * (ContactVelocityErrorTangential / ContactVelocityTangentialLen);
-				}
-			}
-
-			Impulse = -Stiffness * (ContactMass * ContactVelocityError);
+			Impulse = -Stiffness * (ContactMass * ContactVelocityDelta);
 		}
 		else
 		{
-			Impulse = -(Stiffness * ContactMassNormal) * ContactVelocityError;
+			Impulse = -(Stiffness * ContactMassNormal) * ContactVelocityDelta;
 		}
 
-		FVec3 NetImpulse = ManifoldPoint.NetImpulse + Impulse;
-
-		// Clamp the total impulse to be positive along the normal
-		// We can apply negative velocity correction, but only to correct the velocity that
-		// was added by pushout, or in this velocity solve step.
-		if (bChaos_PBD_Collision_Velocity_ImpulseClampEnabled)
+		// Clamp the total impulse to be positive along the normal. We can apply negative velocity correction, 
+		// but only to correct the velocity that was added by pushout, or in this velocity solve step.
+		if (bChaos_PBDCollisionSolver_Velocity_ImpulseClampEnabled && (Dt > 0))
 		{
-			const FReal PushOutVelocityNormal = FMath::Max(0.0f, FVec3::DotProduct(ManifoldPoint.NetPushOut, ContactNormal) / IterationParameters.Dt);
+			const FVec3 NetImpulse = InOutNetImpulse + Impulse;
+			const FReal PushOutVelocityNormal = FMath::Max(0.0f, FVec3::DotProduct(NetPushOut, ContactNormal) / Dt);
 			if (FVec3::DotProduct(NetImpulse, ContactNormal) < -PushOutVelocityNormal)
 			{
-				Impulse = -ManifoldPoint.NetImpulse;
-				NetImpulse = FVec3(0);
+				Impulse = -InOutNetImpulse;
 			}
 		}
 
+		OutImpulse = Impulse;
+		InOutNetImpulse += Impulse;
+	}
+
+	inline void ApplyVelocityCorrection(
+		const FReal Stiffness,
+		const FReal Dt,
+		const FReal DynamicFriction,
+		const FVec3& ContactVelocityDelta,
+		const FReal ContactVelocityDeltaNormal,
+		FPBDCollisionSolver::FSolverManifoldPoint& ManifoldPoint,
+		FConstraintSolverBody& Body0,
+		FConstraintSolverBody& Body1)
+	{
+		FVec3 Impulse = FVec3(0);
+
+		CalculateVelocityCorrectionImpulse(
+			Stiffness,
+			Dt,
+			DynamicFriction,
+			ManifoldPoint.WorldContactNormal,
+			ManifoldPoint.WorldContactMass,
+			ManifoldPoint.WorldContactMassNormal,
+			ContactVelocityDelta,
+			ContactVelocityDeltaNormal,
+			ManifoldPoint.NetPushOut,	// Out
+			ManifoldPoint.NetImpulse,	// Out
+			Impulse);					// Out
+
+		//UE_LOG(LogChaosCollision, VeryVerbose, TEXT("    Impulse %f"), Impulse.Size());
+
 		// Calculate the velocity deltas from the impulse
-		if (InvM0 > SMALL_NUMBER)
+		if (Body0.IsDynamic())
 		{
-			const FVec3 AngularImpulse = FVec3::CrossProduct(ApplyContactPoint0, Impulse);
-			V0 += InvM0 * Impulse;
-			W0 += InvI0 * AngularImpulse;
+			const FVec3 AngularImpulse = FVec3::CrossProduct(ManifoldPoint.WorldRelativeImpulsePoint0, Impulse);
+			Body0.ApplyVelocityDelta(Body0.InvM() * Impulse, Body0.InvI() * AngularImpulse);
 		}
-		if (InvM1 > SMALL_NUMBER)
+		if (Body1.IsDynamic())
 		{
-			const FVec3 AngularImpulse = FVec3::CrossProduct(ApplyContactPoint1, Impulse);
-			V1 -= InvM1 * Impulse;
-			W1 -= InvI1 * AngularImpulse;
+			const FVec3 AngularImpulse = FVec3::CrossProduct(ManifoldPoint.WorldRelativeImpulsePoint1, Impulse);
+			Body1.ApplyVelocityDelta(Body1.InvM() * -Impulse, Body1.InvI() * -AngularImpulse);
 		}
+	}
 
-		ManifoldPoint.NetImpulse = NetImpulse;
+	FPBDCollisionSolver::FState::FState()
+		: SolverBodies()
+		, ManifoldPoints()
+		, NumManifoldPoints(0)
+		, StaticFriction(0)
+		, DynamicFriction(0)
+		, Stiffness(1)
+		, BodyEpochs{ INDEX_NONE, INDEX_NONE }
+		, bIsSolved(false)
+	{
 	}
 
+	FPBDCollisionSolver::FPBDCollisionSolver()
+		: State()
+	{
+	}
 
-	void FPBDCollisionSolver::SolvePosition(
-		FRigidBodyPointContactConstraint& Constraint,
-		const Collisions::FContactIterationParameters& IterationParameters,
-		const Collisions::FContactParticleParameters& ParticleParameters)
+	void FPBDCollisionSolver::FSolverManifoldPoint::InitContact(
+		const FConstraintSolverBody& Body0,
+		const FConstraintSolverBody& Body1,
+		const FVec3& InCoMAnchorPoint0,
+		const FVec3& InCoMAnchorPoint1,
+		const FVec3& InWorldContactNormal)
 	{
-		if (!bChaos_PBD_Collision_Position_SolveEnabled)
-		{
-			return;
-		}
+		UpdateContact(Body0, Body1, InCoMAnchorPoint0, InCoMAnchorPoint1, InWorldContactNormal);
+		UpdateMass(Body0, Body1);
+	}
 
-		FGenericParticleHandle Particle0 = FGenericParticleHandle(Constraint.Particle[0]);
-		FGenericParticleHandle Particle1 = FGenericParticleHandle(Constraint.Particle[1]);
-		TArrayView<FManifoldPoint> ManifoldPoints = Constraint.GetManifoldPoints();
-
-		// Lock bodies for shock propagation?
-		bool IsTemporarilyStatic0 = false;
-		bool IsTemporarilyStatic1 = false;
-		//if (Chaos_Collision_UseShockPropagation)
-		//{
-		//	IsTemporarilyStatic0 = IsTemporarilyStatic.Contains(Particle0->GeometryParticleHandle());
-		//	IsTemporarilyStatic1 = IsTemporarilyStatic.Contains(Particle1->GeometryParticleHandle());
-
-		//	// In the case where two objects at the same level in shock propagation end
-		//	// up in contact with each other, try to decide which one has higher priority.
-		//	// For now, we use the one which is "lower" (in direction of gravity) as the frozen one.
-		//	if (IsTemporarilyStatic0 && IsTemporarilyStatic1)
-		//	{
-		//		if (ManifoldPoints.Num() > 0)
-		//		{
-		//			const FReal NormalThreshold = 0.2f;	// @chaos(todo): what value here?
-		//			const FReal NormalGrav = FVec3::DotProduct(ManifoldPoints[0].ContactPoint.Normal, GravityDir);
-		//			const bool bNormalIsUp = NormalGrav < -NormalThreshold;
-		//			const bool bNormalIsDown = NormalGrav > NormalThreshold;
-		//			IsTemporarilyStatic0 = bNormalIsDown;
-		//			IsTemporarilyStatic1 = bNormalIsUp;
-		//		}
-		//		else
-		//		{
-		//			IsTemporarilyStatic0 = false;
-		//			IsTemporarilyStatic1 = false;
-		//		}
-		//	}
-		//}
-
-		TPBDRigidParticleHandle<FReal, 3>* PBDRigid0 = Particle0->CastToRigidParticle();
-		TPBDRigidParticleHandle<FReal, 3>* PBDRigid1 = Particle1->CastToRigidParticle();
-		const bool bIsRigidDynamic0 = PBDRigid0 && (PBDRigid0->ObjectState() == EObjectStateType::Dynamic) && (PBDRigid0->InvM() > SMALL_NUMBER) && !IsTemporarilyStatic0;
-		const bool bIsRigidDynamic1 = PBDRigid1 && (PBDRigid1->ObjectState() == EObjectStateType::Dynamic) && (PBDRigid1->InvM() > SMALL_NUMBER) && !IsTemporarilyStatic1;
-
-		if (!bIsRigidDynamic0 && !bIsRigidDynamic1)
-		{
-			return;
-		}
+	void FPBDCollisionSolver::FSolverManifoldPoint::InitMaterial(
+		const FReal InWorldContactVelocityTargetNormal,
+		const bool bInEnableStaticFriction,
+		const FReal InStaticFrictionMax)
+	{
+		NetPushOut = FVec3(0);
+		NetImpulse = FVec3(0);
+		StaticFrictionMax = InStaticFrictionMax;
+		bInsideStaticFrictionCone = bInEnableStaticFriction;
+		WorldContactVelocityTargetNormal = InWorldContactVelocityTargetNormal;
+	}
 
-		// Gradually increase position correction through iterations (optional based on cvars)
-		const FReal Stiffness = GetPositionSolveStiffness(IterationParameters.Iteration, IterationParameters.NumIterations);
+	inline void FPBDCollisionSolver::FSolverManifoldPoint::UpdateContact(
+		const FConstraintSolverBody& Body0, 
+		const FConstraintSolverBody& Body1,
+		const FVec3& InCoMAnchorPoint0,
+		const FVec3& InCoMAnchorPoint1, 
+		const FVec3& InWorldContactNormal)
+	{
+		// The contact points on each body
+		LocalRelativeAnchorPoint0 = InCoMAnchorPoint0;
+		LocalRelativeAnchorPoint1 = InCoMAnchorPoint1;
 
-		FVec3 P0 = FParticleUtilities::GetCoMWorldPosition(Particle0);
-		FRotation3 Q0 = FParticleUtilities::GetCoMWorldRotation(Particle0);
-		FVec3 P1 = FParticleUtilities::GetCoMWorldPosition(Particle1);
-		FRotation3 Q1 = FParticleUtilities::GetCoMWorldRotation(Particle1);
+		// The world-space point where we apply impulses/corrections (same world-space point for momentum conservation)
+		const FVec3 WorldContactPoint = FReal(0.5) * (Body0.P() + Body0.Q() * LocalRelativeAnchorPoint0 + Body1.P() + Body1.Q() * LocalRelativeAnchorPoint1);
+		WorldRelativeImpulsePoint0 = WorldContactPoint - Body0.P();
+		WorldRelativeImpulsePoint1 = WorldContactPoint - Body1.P();
 
-		// Apply the position correction so that all contacts have zero separation
-		for (int32 PointIndex = 0; PointIndex < ManifoldPoints.Num(); ++PointIndex)
-		{
-			FManifoldPoint& ManifoldPoint = Constraint.SetActiveManifoldPoint(PointIndex, P0, Q0, P1, Q1);
+		// The world-space contact normal
+		WorldContactNormal = InWorldContactNormal;
+	}
 
-			// @todo(chaos): this should be set at initialization, but the old path requires it to be false by default and we want it to be true
-			const int32 NumStaticFrictionIterations = FMath::CeilToInt((FReal)Chaos_PBD_Collision_Position_StaticFrictionIterations * (FReal)IterationParameters.NumIterations);
-			if (IterationParameters.Iteration == 0)
-			{
-				ManifoldPoint.bInsideStaticFrictionCone = (NumStaticFrictionIterations > 0);
-			}
+	inline void FPBDCollisionSolver::FSolverManifoldPoint::UpdateMass(const FConstraintSolverBody& Body0, const FConstraintSolverBody& Body1)
+	{
+		const FMatrix33 ContactMassInv =
+			(Body0.IsDynamic() ? Collisions::ComputeFactorMatrix3(WorldRelativeImpulsePoint0, Body0.InvI(), Body0.InvM()) : FMatrix33(0)) +
+			(Body1.IsDynamic() ? Collisions::ComputeFactorMatrix3(WorldRelativeImpulsePoint1, Body1.InvI(), Body1.InvM()) : FMatrix33(0));
+		const FMatrix33 ContactMass = ContactMassInv.Inverse();
+		const FReal ContactMassInvNormal = FVec3::DotProduct(WorldContactNormal, Utilities::Multiply(ContactMassInv, WorldContactNormal));
+		const FReal ContactMassNormal = (ContactMassInvNormal > FReal(SMALL_NUMBER)) ? FReal(1) / ContactMassInvNormal : FReal(0);
+		WorldContactMass = ContactMass;
+		WorldContactMassNormal = ContactMassNormal;
+	}
 
-			const bool bProcessManifoldPoint = (ManifoldPoint.ContactPoint.Phi < Chaos_PBD_Collision_Position_Tolerance) || !ManifoldPoint.NetPushOut.IsNearlyZero();
-			if (bProcessManifoldPoint)
+	FVec3 FPBDCollisionSolver::FSolverManifoldPoint::CalculateContactVelocity(const FConstraintSolverBody& Body0, const FConstraintSolverBody& Body1) const
+	{
+		const FVec3 ContactVelocity0 = Body0.V() + FVec3::CrossProduct(Body0.W(), WorldRelativeImpulsePoint0);
+		const FVec3 ContactVelocity1 = Body1.V() + FVec3::CrossProduct(Body1.W(), WorldRelativeImpulsePoint1);
+		return ContactVelocity0 - ContactVelocity1;
+	}
+
+	inline void FPBDCollisionSolver::FSolverManifoldPoint::CalculateContactPositionError(const FConstraintSolverBody& Body0, const FConstraintSolverBody& Body1, FVec3& OutContactDelta, FReal& OutContactDeltaNormal) const
+	{
+		OutContactDelta = (Body0.P() + Body0.Q() * LocalRelativeAnchorPoint0) - (Body1.P() + Body1.Q() * LocalRelativeAnchorPoint1);
+		OutContactDeltaNormal = FVec3::DotProduct(OutContactDelta, WorldContactNormal);
+	}
+
+	inline void FPBDCollisionSolver::FSolverManifoldPoint::CalculateContactVelocityError(const FConstraintSolverBody& Body0, const FConstraintSolverBody& Body1, const FReal DynamicFriction, const FReal Dt, FVec3& OutContactVelocityDelta, FReal& OutContactVelocityDeltaNormal) const
+	{
+		const FVec3 ContactVelocity0 = Body0.V() + FVec3::CrossProduct(Body0.W(), WorldRelativeImpulsePoint0);
+		const FVec3 ContactVelocity1 = Body1.V() + FVec3::CrossProduct(Body1.W(), WorldRelativeImpulsePoint1);
+		const FVec3 ContactVelocity = ContactVelocity0 - ContactVelocity1;
+		const FReal ContactVelocityNormal = FVec3::DotProduct(ContactVelocity, WorldContactNormal);
+
+		// Target normal velocity, including restitution
+		const FReal ContactVelocityTargetNormal = WorldContactVelocityTargetNormal;
+
+		// Add up the errors in the velocity (current velocity - desired velocity)
+		OutContactVelocityDeltaNormal = (ContactVelocityNormal - WorldContactVelocityTargetNormal);
+		OutContactVelocityDelta = (ContactVelocityNormal - ContactVelocityTargetNormal) * WorldContactNormal;
+
+		const bool bApplyFriction = (DynamicFriction > 0) && (Dt > 0);
+		if (bApplyFriction)
+		{
+			const FVec3 ContactVelocityTangential = ContactVelocity - ContactVelocityNormal * WorldContactNormal;
+			const FReal ContactVelocityTangentialLen = ContactVelocityTangential.Size();
+			if (ContactVelocityTangentialLen > SMALL_NUMBER)
 			{
-				// Calculate the position error we need to correct, including static friction and restitution
-				// Position correction uses the deepest point on each body (see velocity correction which uses average contact)
-				const FRotation3& PlaneQ = (ManifoldPoint.ContactPoint.ContactNormalOwnerIndex == 0) ? Q0 : Q1;
-				const FVec3 ContactNormal = PlaneQ * ManifoldPoint.CoMContactNormal;
-
-				const int32 IterationsRemaining = IterationParameters.NumIterations - IterationParameters.Iteration;	// Including this one
-				const bool bApplyStaticFriction = ManifoldPoint.bPotentialRestingContact && (IterationsRemaining <= NumStaticFrictionIterations);
-				const FReal StaticFriction = bApplyStaticFriction ? FMath::Max(Constraint.Manifold.Friction, Constraint.Manifold.AngularFriction) : 0.0f;
-
-				// Relative contact points on each body in world-space
-				const FVec3 RelativeContactPoint0 = Q0 * ManifoldPoint.CoMContactPoints[0];
-				const FVec3 RelativeContactPoint1 = Q1 * ManifoldPoint.CoMContactPoints[1];
-				const FVec3 ContactPoint = FReal(0.5) * ((P1 + RelativeContactPoint1) + (P0 + RelativeContactPoint0));
-				const FVec3 ApplyContactPoint0 = (bChaos_PBD_Collision_UseCorrectApplyPosition) ? ContactPoint - P0 : RelativeContactPoint0;
-				const FVec3 ApplyContactPoint1 = (bChaos_PBD_Collision_UseCorrectApplyPosition) ? ContactPoint - P1 : RelativeContactPoint1;
-				FVec3 ContactError = (P1 + RelativeContactPoint1) - (P0 + RelativeContactPoint0);
-				
-				// If we have static friction, we attempt to move the bodies so that the old contact positions are aligned
-				if (bApplyStaticFriction)
+				// PushOut = ContactMass * DP, where DP is the contact positional correction
+				// Friction force is proportional to the normal force, so friction velocity correction
+				// is proprtional to normal velocity correction, or DVn = DPn/dt = PushOut.N / (ContactMass * dt);
+				const FReal PushOutNormal = FVec3::DotProduct(NetPushOut, WorldContactNormal);
+				const FReal DynamicFrictionVelocityError = PushOutNormal / (WorldContactMassNormal * Dt);
+				if (DynamicFrictionVelocityError > SMALL_NUMBER)
 				{
-					FVec3 StaticFrictionLocalContactPoint0 = ManifoldPoint.CoMContactPoints[0];
-					FVec3 StaticFrictionLocalContactPoint1 = ManifoldPoint.CoMContactPoints[1];
-					Constraint.CalculatePrevCoMContactPoints(Particle0, Particle1, ManifoldPoint, IterationParameters.Dt, StaticFrictionLocalContactPoint0, StaticFrictionLocalContactPoint1);
-					const FVec3 StaticFrictionRelativeContactPoint0 = Q0 * StaticFrictionLocalContactPoint0;
-					const FVec3 StaticFrictionRelativeContactPoint1 = Q1 * StaticFrictionLocalContactPoint1;
-					ContactError = (P1 + StaticFrictionRelativeContactPoint1) - (P0 + StaticFrictionRelativeContactPoint0);
+					const FReal ContactVelocityErrorTangential = FMath::Min(DynamicFriction * DynamicFrictionVelocityError, ContactVelocityTangentialLen);
+					OutContactVelocityDelta += ContactVelocityTangential * (ContactVelocityErrorTangential / ContactVelocityTangentialLen);
 				}
+			}
+		}
+	}
 
-				// Temporary shock propagation implementation...
-				// Contact normal always points from body 1 to body 0 (regardless of which body owns the contact manifold plane)
-				// so if Contact Normal points up body 0 is on top
-				const bool bUseMassScale = bIsRigidDynamic0 && bIsRigidDynamic1 && (FMath::Abs(ContactNormal.Z) > Chaos_PBD_Collision_Position_ShockPropagationNormalZ);
-				const bool bOnTop0 = ContactNormal.Z > 0.0f;
-				const FReal InvMassScale0 = (!bUseMassScale || bOnTop0) ? 1.0f : GetPositionInvMassScale(IterationParameters.Iteration, IterationParameters.NumIterations);
-				const FReal InvMassScale1 = (!bUseMassScale || !bOnTop0) ? 1.0f : GetPositionInvMassScale(IterationParameters.Iteration, IterationParameters.NumIterations);
-
-				const FReal InvM0 = bIsRigidDynamic0 ? InvMassScale0 * PBDRigid0->InvM() : 0.0f;
-				const FReal InvM1 = bIsRigidDynamic1 ? InvMassScale1 * PBDRigid1->InvM() : 0.0f;
-				const FMatrix33 InvI0 = bIsRigidDynamic0 ? InvMassScale0 * Utilities::ComputeWorldSpaceInertia(Q0, PBDRigid0->InvI()) * Constraint.Manifold.InvInertiaScale0 : FMatrix33(0);
-				const FMatrix33 InvI1 = bIsRigidDynamic1 ? InvMassScale1 * Utilities::ComputeWorldSpaceInertia(Q1, PBDRigid1->InvI()) * Constraint.Manifold.InvInertiaScale1 : FMatrix33(0);
-
-				const FMatrix33 ContactMassInv =
-					((InvM0 > SMALL_NUMBER) ? Collisions::ComputeFactorMatrix3(ApplyContactPoint0, InvI0, InvM0) : FMatrix33(0)) +
-					((InvM1 > SMALL_NUMBER) ? Collisions::ComputeFactorMatrix3(ApplyContactPoint1, InvI1, InvM1) : FMatrix33(0));
-				const FMatrix33 ContactMass = ContactMassInv.Inverse();
-				const FReal ContactMassInvNormal = FVec3::DotProduct(ContactNormal, Utilities::Multiply(ContactMassInv, ContactNormal));
-				const FReal ContactMassNormal = (ContactMassInvNormal > (FReal)SMALL_NUMBER) ? (FReal)1.0f / ContactMassInvNormal : (FReal)0.0f;
+	void FPBDCollisionSolver::EnablePositionShockPropagation()
+	{
+		SetShockPropagationInvMassScale(Chaos_PBDCollisionSolver_Position_MinInvMassScale);
+	}
 
-				ApplyPositionCorrection(
-					Stiffness,
-					IterationParameters,
-					ParticleParameters,
-					ApplyContactPoint0,
-					ApplyContactPoint1,
-					ContactError,
-					ContactNormal,
-					ContactMass,
-					ContactMassNormal,
-					StaticFriction,
-					InvM0, InvI0, InvM1, InvI1,
-					P0, Q0, P1, Q1,
-					ManifoldPoint);
-			}
+	void FPBDCollisionSolver::EnableVelocityShockPropagation()
+	{
+		SetShockPropagationInvMassScale(Chaos_PBDCollisionSolver_Velocity_MinInvMassScale);
+	}
 
-			if (bIsRigidDynamic0)
+	void FPBDCollisionSolver::DisableShockPropagation()
+	{
+		SetShockPropagationInvMassScale(FReal(1));
+	}
+
+	void FPBDCollisionSolver::SetShockPropagationInvMassScale(const FReal InvMassScale)
+	{
+		FConstraintSolverBody& Body0 = SolverBody0();
+		FConstraintSolverBody& Body1 = SolverBody1();
+
+		// No need to set an inverse mass scale if ther other body is kinematic (with inv mass of 0)
+		// Also, bodies at the same level do not take part in shock propagation
+		if (Body0.IsDynamic() && Body1.IsDynamic() && (Body0.Level() != Body1.Level()))
+		{
+			// Set the inv mass scale of the "lower" body to make it heavier
+			bool bInvMassUpdated = false;
+			if (Body0.Level() < Body1.Level())
 			{
-				FParticleUtilities::SetCoMWorldTransform(Particle0, P0, Q0);
+				if (Body0.InvMassScale() != InvMassScale)
+				{
+					Body0.SetInvMassScale(InvMassScale);
+					bInvMassUpdated = true;
+				}
+			}
+			else
+			{
+				if (Body1.InvMassScale() != InvMassScale)
+				{
+					Body1.SetInvMassScale(InvMassScale);
+					bInvMassUpdated = true;
+				}
 			}
 
-			if (bIsRigidDynamic1)
+			// If the masses changed, we need to rebuild the contact mass for each manifold point
+			if (bInvMassUpdated)
 			{
-				FParticleUtilities::SetCoMWorldTransform(Particle1, P1, Q1);
+				for (int32 PointIndex = 0; PointIndex < State.NumManifoldPoints; ++PointIndex)
+				{
+					State.ManifoldPoints[PointIndex].UpdateMass(Body0, Body1);
+				}
 			}
 		}
+	}
 
-		// @todo(chaos): support early-out
-		*IterationParameters.NeedsAnotherIteration = true;
+	void FPBDCollisionSolver::InitContact(
+		const int32 ManifoldPoiontIndex,
+		const FVec3& InCoMAnchorPoint0,
+		const FVec3& InCoMAnchorPoint1,
+		const FVec3& InWorldContactNormal)
+	{
+		State.ManifoldPoints[ManifoldPoiontIndex].InitContact(State.SolverBodies[0], State.SolverBodies[1], InCoMAnchorPoint0, InCoMAnchorPoint1, InWorldContactNormal);
 	}
 
-	void FPBDCollisionSolver::SolveVelocity(
-		FRigidBodyPointContactConstraint& Constraint,
-		const Collisions::FContactIterationParameters& IterationParameters,
-		const Collisions::FContactParticleParameters& ParticleParameters)
+	void FPBDCollisionSolver::InitMaterial(
+		const int32 ManifoldPoiontIndex,
+		const FReal InWorldContactVelocityTargetNormal,
+		const bool bInEnableStaticFriction,
+		const FReal InStaticFrictionMax)
 	{
-		FGenericParticleHandle Particle0 = FGenericParticleHandle(Constraint.Particle[0]);
-		FGenericParticleHandle Particle1 = FGenericParticleHandle(Constraint.Particle[1]);
-
-		TPBDRigidParticleHandle<FReal, 3>* PBDRigid0 = Particle0->CastToRigidParticle();
-		TPBDRigidParticleHandle<FReal, 3>* PBDRigid1 = Particle1->CastToRigidParticle();
-		const bool bIsRigidDynamic0 = PBDRigid0 && PBDRigid0->ObjectState() == EObjectStateType::Dynamic;
-		const bool bIsRigidDynamic1 = PBDRigid1 && PBDRigid1->ObjectState() == EObjectStateType::Dynamic;
-
-		const FVec3 P0 = FParticleUtilities::GetCoMWorldPosition(Particle0);
-		const FRotation3 Q0 = FParticleUtilities::GetCoMWorldRotation(Particle0);
-		const FVec3 P1 = FParticleUtilities::GetCoMWorldPosition(Particle1);
-		const FRotation3 Q1 = FParticleUtilities::GetCoMWorldRotation(Particle1);
-		FVec3 V0 = Particle0->V();
-		FVec3 W0 = Particle0->W();
-		FVec3 V1 = Particle1->V();
-		FVec3 W1 = Particle1->W();
-
-		//Constraint.AccumulatedImpulse = FVec3(0);
-
-		if (bChaos_PBD_Collision_Velocity_SolveEnabled)
+		State.ManifoldPoints[ManifoldPoiontIndex].InitMaterial(InWorldContactVelocityTargetNormal, bInEnableStaticFriction, InStaticFrictionMax);
+	}
+
+	void FPBDCollisionSolver::UpdateContact(
+		const int32 ManifoldPoiontIndex,
+		const FVec3& InCoMAnchorPoint0,
+		const FVec3& InCoMAnchorPoint1,
+		const FVec3& InWorldContactNormal)
+	{
+		State.ManifoldPoints[ManifoldPoiontIndex].UpdateContact(State.SolverBodies[0], State.SolverBodies[1], InCoMAnchorPoint0, InCoMAnchorPoint1, InWorldContactNormal);
+	}
+
+	bool FPBDCollisionSolver::SolvePosition(const FReal Dt, const bool bApplyStaticFriction)
+	{
+		if (!bChaos_PBDCollisionSolver_Position_SolveEnabled)
 		{
-			TArrayView<FManifoldPoint> ManifoldPoints = Constraint.GetManifoldPoints();
-			const FReal Stiffness = GetVelocitySolveStiffness(IterationParameters.Iteration, IterationParameters.NumIterations);
+			return false;
+		}
 
-			// Only apply friction for 1 iteration - the last one
-			const bool bApplyDynamicFriction = (IterationParameters.Iteration == IterationParameters.NumIterations - 1) && bChaos_PBD_Collision_Velocity_DynamicFrictionEnabled;
+		// SolverBody decorator used to add mass scaling
+		FConstraintSolverBody& Body0 = SolverBody0();
+		FConstraintSolverBody& Body1 = SolverBody1();
 
-			for (int32 PointIndex = 0; PointIndex < ManifoldPoints.Num(); ++PointIndex)
+		// Check for solved
+		if (State.bIsSolved)
+		{
+			if ((Body0.LastChangeEpoch() == State.BodyEpochs[0]) && (Body1.LastChangeEpoch() == State.BodyEpochs[1]))
 			{
-				FManifoldPoint& ManifoldPoint = Constraint.SetActiveManifoldPoint(PointIndex, P0, Q0, P1, Q1);
-
-				if (ManifoldPoint.bActive)
-				{
-					const FVec3 RelativeContactPoint0 = Q0 * ManifoldPoint.CoMContactPoints[0];
-					const FVec3 RelativeContactPoint1 = Q1 * ManifoldPoint.CoMContactPoints[1];
-					const FVec3 ContactPoint = FReal(0.5) * ((P1 + RelativeContactPoint1) + (P0 + RelativeContactPoint0));
-					const FVec3 ApplyContactPoint0 = (bChaos_PBD_Collision_UseCorrectApplyPosition) ? ContactPoint - P0 : RelativeContactPoint0;
-					const FVec3 ApplyContactPoint1 = (bChaos_PBD_Collision_UseCorrectApplyPosition) ? ContactPoint - P1 : RelativeContactPoint1;
-
-					const FRotation3& PlaneQ = (ManifoldPoint.ContactPoint.ContactNormalOwnerIndex == 0) ? Q0 : Q1;
-					const FVec3 ContactNormal = PlaneQ * ManifoldPoint.CoMContactNormal;
-
-					// Temporary shock propagation implementation...
-					// Contact normal always points from body 1 to body 0 (regardless of which body owns the contact manifold plane)
-					// so if Contact Normal points up body 0 is on top
-					const bool bUseMassScale = bIsRigidDynamic0 && bIsRigidDynamic1 && (FMath::Abs(ContactNormal.Z) > Chaos_PBD_Collision_Position_ShockPropagationNormalZ);
-					const bool bOnTop0 = ContactNormal.Z > 0.0f;
-					const FReal InvMassScale0 = (!bUseMassScale || bOnTop0) ? 1.0f : GetVelocityInvMassScale(IterationParameters.Iteration, IterationParameters.NumIterations);
-					const FReal InvMassScale1 = (!bUseMassScale || !bOnTop0) ? 1.0f : GetVelocityInvMassScale(IterationParameters.Iteration, IterationParameters.NumIterations);
-
-					const FReal InvM0 = bIsRigidDynamic0 ? InvMassScale0 * PBDRigid0->InvM() : 0.0f;
-					const FReal InvM1 = bIsRigidDynamic1 ? InvMassScale1 * PBDRigid1->InvM() : 0.0f;
-					const FMatrix33 InvI0 = bIsRigidDynamic0 ? InvMassScale0 * Utilities::ComputeWorldSpaceInertia(Q0, PBDRigid0->InvI()) * Constraint.Manifold.InvInertiaScale0 : FMatrix33(0);
-					const FMatrix33 InvI1 = bIsRigidDynamic1 ? InvMassScale1 * Utilities::ComputeWorldSpaceInertia(Q1, PBDRigid1->InvI()) * Constraint.Manifold.InvInertiaScale1 : FMatrix33(0);
-
-					const FMatrix33 ContactMassInv =
-						((InvM0 > SMALL_NUMBER) ? Collisions::ComputeFactorMatrix3(ApplyContactPoint0, InvI0, InvM0) : FMatrix33(0)) +
-						((InvM1 > SMALL_NUMBER) ? Collisions::ComputeFactorMatrix3(ApplyContactPoint1, InvI1, InvM1) : FMatrix33(0));
-					const FMatrix33 ContactMass = ContactMassInv.Inverse();
-					const FReal ContactMassInvNormal = FVec3::DotProduct(ContactNormal, Utilities::Multiply(ContactMassInv, ContactNormal));
-					const FReal ContactMassNormal = (ContactMassInvNormal > (FReal)SMALL_NUMBER) ? (FReal)1.0f / ContactMassInvNormal : (FReal)0.0f;
-
-					const FReal DynamicFriction = (bApplyDynamicFriction) ? Constraint.Manifold.Friction : (FReal)0.0f;
-
-					const FReal Resitution = Constraint.Manifold.Restitution;
-
-					ApplyVelocityCorrection(
-						Stiffness,
-						IterationParameters,
-						ParticleParameters,
-						ApplyContactPoint0,
-						ApplyContactPoint1,
-						ContactNormal,
-						ContactMass,
-						ContactMassNormal,
-						DynamicFriction,
-						Resitution,
-						InvM0, InvI0, InvM1, InvI1,
-						V0, W0, V1, W1,
-						ManifoldPoint);
-
-					//Constraint.AccumulatedImpulse += ManifoldPoint.NetImpulse;
-				}
+				// We did not apply a correction last iteration (within tolerance) and nothing else has moved the bodies
+				return false;
 			}
+		}
+
+		// The first few iterations have friction disabled. This allows us to solve the normal penetration at each contact point
+		// without adding spurious lateral impulses which get cancelled out by other contacts. The spurious lateral impulses
+		// cause problems when clipping to the friction cone on a per-contact basis.
+		FReal StaticFriction = (bApplyStaticFriction) ? State.StaticFriction : FReal(0);
+		FReal DynamicFriction = (bApplyStaticFriction) ? State.DynamicFriction : FReal(0);
+		const FReal Stiffness = State.Stiffness;
 
-			// @hack - disable dynamic friction after the first iteration so that pair iteration count doesn't impact the result - the formula is iteration-count sensitive (swith to XPBD)
-			if (bApplyDynamicFriction)
+		// Apply the position correction so that all contacts have zero separation
+		for (int32 PointIndex = 0; PointIndex < State.NumManifoldPoints; ++PointIndex)
+		{
+			FSolverManifoldPoint& SolverManifoldPoint = State.ManifoldPoints[PointIndex];
+
+			//SolverManifoldPoint.UpdateContactPoint(Body0, Body1, State.Constraint->GetManifoldPoints()[PointIndex]);
+			//SolverManifoldPoint.UpdateContactMass(Body0, Body1);
+
+			FVec3 ContactDelta;
+			FReal ContactDeltaNormal;
+			SolverManifoldPoint.CalculateContactPositionError(Body0.SolverBody(), Body1.SolverBody(), ContactDelta, ContactDeltaNormal);
+
+			const bool bProcessManifoldPoint = (ContactDeltaNormal < Chaos_PBDCollisionSolver_Position_NormalTolerance) || !SolverManifoldPoint.NetPushOut.IsNearlyZero();
+			if (bProcessManifoldPoint)
 			{
-				Constraint.Manifold.Friction = 0;
+				ApplyPositionCorrection(
+					Stiffness,
+					StaticFriction,
+					DynamicFriction,
+					ContactDelta,
+					ContactDeltaNormal,
+					SolverManifoldPoint,
+					Body0,
+					Body1);
 			}
-
 		}
 
-		if (bIsRigidDynamic0)
+		// We are solved if we did not move the bodies within some tolerance
+		State.bIsSolved = (Body0.LastChangeEpoch() == State.BodyEpochs[0]) && (Body1.LastChangeEpoch() == State.BodyEpochs[1]);
+		State.BodyEpochs[0] = Body0.LastChangeEpoch();
+		State.BodyEpochs[1] = Body1.LastChangeEpoch();
+
+		// @todo(chaos): support early-out
+		return !State.bIsSolved;
+	}
+
+	bool FPBDCollisionSolver::SolveVelocity(const FReal Dt, const bool bApplyDynamicFriction)
+	{
+		if (!bChaos_PBDCollisionSolver_Velocity_SolveEnabled)
 		{
-			PBDRigid0->V() = V0;
-			PBDRigid0->W() = W0;
+			return false;
 		}
 
-		if (bIsRigidDynamic1)
+		FConstraintSolverBody& Body0 = SolverBody0();
+		FConstraintSolverBody& Body1 = SolverBody1();
+
+		// Only apply friction for 1 iteration (the last one) because the solver is iteration-count sensitive
+		// @todo(chaos): use xpbd dynamic friction in QuasiPBD Solver
+		const FReal DynamicFriction = (bApplyDynamicFriction && bChaos_PBDCollisionSolver_Velocity_DynamicFrictionEnabled) ? State.DynamicFriction : (FReal)0.0f;
+		const FReal Stiffness = State.Stiffness;
+
+		for (int32 PointIndex = 0; PointIndex < State.NumManifoldPoints; ++PointIndex)
 		{
-			PBDRigid1->V() = V1;
-			PBDRigid1->W() = W1;
+			FSolverManifoldPoint& SolverManifoldPoint = State.ManifoldPoints[PointIndex];
+
+			//SolverManifoldPoint.UpdateContactPoint(Body0, Body1, State.Constraint->GetManifoldPoints()[PointIndex]);
+			//SolverManifoldPoint.UpdateContactMass(Body0, Body1);
+
+			FVec3 ContactVelocityDelta;
+			FReal ContactVelocityDeltaNormal;
+			SolverManifoldPoint.CalculateContactVelocityError(Body0, Body1, DynamicFriction, Dt, ContactVelocityDelta, ContactVelocityDeltaNormal);
+
+			if (!SolverManifoldPoint.NetPushOut.IsNearlyZero())
+			{
+				ApplyVelocityCorrection(
+					Stiffness,
+					Dt,
+					DynamicFriction,
+					ContactVelocityDelta,
+					ContactVelocityDeltaNormal,
+					SolverManifoldPoint,
+					Body0,
+					Body1);
+			}
 		}
 
 		// @todo(chaos): support early-out
-		*IterationParameters.NeedsAnotherIteration = true;
+		return true;
 	}
 }
\ No newline at end of file
