Chaos

diff --git a/Engine/Source/Runtime/Experimental/Chaos/Private/Chaos/Collision/PBDCollisionSolver.h b/Engine/Source/Runtime/Experimental/Chaos/Private/Chaos/Collision/PBDCollisionSolver.h
--- a/Engine/Source/Runtime/Experimental/Chaos/Private/Chaos/Collision/PBDCollisionSolver.h
+++ b/Engine/Source/Runtime/Experimental/Chaos/Private/Chaos/Collision/PBDCollisionSolver.h
@@ -1,8 +1,12 @@
 // Copyright Epic Games, Inc. All Rights Reserved.
 #pragma once
 
+#include "CoreMinimal.h"
+#include "Logging/LogMacros.h"
+
 #include "Chaos/Core.h"
 #include "Chaos/Defines.h"
+#include "Chaos/Evolution/SolverBody.h"
 
 namespace Chaos
 {
@@ -11,19 +15,231 @@ namespace Chaos
 		struct FContactIterationParameters;
 		struct FContactParticleParameters;
 	}
-	class FRigidBodyPointContactConstraint;
+	class FManifoldPoint;
+	class FPBDCollisionSolver;
+	class FPBDCollisionConstraint;
 
+	/**
+	 * @brief 
+	 * @todo(chaos): Make this solver operate on a single contact point rather than all points in a manifold.
+	 * This would be beneficial if we have many contacts with less than 4 points in the manifold. However this
+	 * is dificult to do while we are still supporting non-manifold collisions.
+	*/
 	class FPBDCollisionSolver
 	{
 	public:
-		void SolvePosition(
-			FRigidBodyPointContactConstraint& Constraint,
-			const Collisions::FContactIterationParameters& IterationParameters,
-			const Collisions::FContactParticleParameters& ParticleParameters);
-
-		void SolveVelocity(
-			FRigidBodyPointContactConstraint& Constraint,
-			const Collisions::FContactIterationParameters& IterationParameters,
-			const Collisions::FContactParticleParameters& ParticleParameters);
+		static const int32 MaxConstrainedBodies = 2;
+		static const int32 MaxPointsPerConstraint = 4;
+
+		class FSolverManifoldPoint
+		{
+		public:
+			/**
+			 * @brief Initialize the geometric data for the contact
+			*/
+			void InitContact(
+				const FConstraintSolverBody& Body0,
+				const FConstraintSolverBody& Body1,
+				const FVec3& InCoMAnchorPoint0,
+				const FVec3& InCoMAnchorPoint1,
+				const FVec3& InWorldContactNormal);
+
+			/**
+			 * @brief Initialize the material related properties of the contact
+			*/
+			void InitMaterial(
+				const FReal InWorldContactVelocityTargetNormal,
+				const bool bInEnableStaticFriction,
+				const FReal InStaticFrictionMax);
+
+			/**
+			 * @brief Update the world-space relative contact points based on current body transforms and body-space contact positions
+			*/
+			void UpdateContact(
+				const FConstraintSolverBody& Body0,
+				const FConstraintSolverBody& Body1,
+				const FVec3& CoMAnchorPoint0,
+				const FVec3& CoMAnchorPoint1,
+				const FVec3& WorldContactNormal);
+
+			/**
+			 * @brief Update the cached mass properties based on the current body transforms
+			*/
+			void UpdateMass(const FConstraintSolverBody& Body0, const FConstraintSolverBody& Body1);
+
+
+			/**
+			 * @brief Calculate the relative velocity at the contact point
+			 * @note InitContact must be called before calling this function
+			*/
+			FVec3 CalculateContactVelocity(const FConstraintSolverBody& Body0, const FConstraintSolverBody& Body1) const;
+
+			/**
+			 * @brief Calculate the position error at the current transforms
+			*/
+			void CalculateContactPositionError(const FConstraintSolverBody& Body0, const FConstraintSolverBody& Body1, FVec3& OutContactDelta, FReal& OutContactDeltaNormal) const;
+
+			/**
+			 * @brief Calculate the velocity error at the current transforms
+			*/
+			void CalculateContactVelocityError(const FConstraintSolverBody& Body0, const FConstraintSolverBody& Body1, const FReal DynamicFriction, const FReal Dt, FVec3& OutContactVelocityDelta, FReal& OutContactVelocityDeltaNormal) const;
+
+		// @todo(chaos): make private
+		public:
+			friend class FPBDCollisionSolver;
+
+			// Contact points on the surface of each shape
+			FVec3 LocalRelativeAnchorPoint0;
+			FVec3 LocalRelativeAnchorPoint1;
+
+			// The body-relative point at which impulses/position corrections will be applied
+			FVec3 WorldRelativeImpulsePoint0;
+			FVec3 WorldRelativeImpulsePoint1;
+
+			// Contact normal
+			FVec3 WorldContactNormal;
+
+			// Contact mass
+			FMatrix33 WorldContactMass;
+			FReal WorldContactMassNormal;
+
+			// Desired final normal velocity, taking Restitution into account
+			FReal WorldContactVelocityTargetNormal;
+
+			// Solver outputs
+			FVec3 NetPushOut;
+			FVec3 NetImpulse;
+
+			// A smoothed NetImpulse along the normal, used for clipping to the static friction cone
+			FReal StaticFrictionMax;
+
+			// Whether we are still in the static friction cone
+			bool bInsideStaticFrictionCone;
+		};
+
+		FPBDCollisionSolver();
+
+		void SetFriction(const FReal InStaticFriction, const FReal InDynamicFriction)
+		{
+			State.StaticFriction = InStaticFriction;
+			State.DynamicFriction = InDynamicFriction;
+		}
+
+		void SetStiffness(const FReal InStiffness)
+		{
+			State.Stiffness = InStiffness;
+		}
+
+		void SetSolverBodies(FSolverBody* SolverBody0, FSolverBody* SolverBody1)
+		{
+			State.SolverBodies[0] = *SolverBody0;
+			State.SolverBodies[1] = *SolverBody1;
+		}
+
+		void ResetSolverBodies()
+		{
+			State.SolverBodies[0].Reset();
+			State.SolverBodies[1].Reset();
+		}
+
+		int32 NumManifoldPoints() const
+		{
+			return State.NumManifoldPoints;
+		}
+
+		void SetNumManifoldPoints(const int32 InNumManifoldPoints)
+		{
+			State.NumManifoldPoints = InNumManifoldPoints;
+		}
+
+		const FSolverManifoldPoint& GetManifoldPoint(const int32 ManifoldPointIndex) const
+		{
+			check(ManifoldPointIndex < NumManifoldPoints());
+			return State.ManifoldPoints[ManifoldPointIndex];
+		}
+
+		void InitContact(
+			const int32 ManifoldPoiontIndex,
+			const FVec3& InCoMAnchorPoint0,
+			const FVec3& InCoMAnchorPoint1,
+			const FVec3& InWorldContactNormal);
+
+		void InitMaterial(
+			const int32 ManifoldPoiontIndex,
+			const FReal InWorldContactVelocityTargetNormal,
+			const bool bInEnableStaticFriction,
+			const FReal InStaticFrictionMax);
+
+		void UpdateContact(
+			const int32 ManifoldPoiontIndex,
+			const FVec3& InCoMAnchorPoint0,
+			const FVec3& InCoMAnchorPoint1,
+			const FVec3& InWorldContactNormal);
+
+		/**
+		 * @brief Get the first (decaorated) solver body
+		 * The decorator add a possible mass scale
+		*/
+		FConstraintSolverBody& SolverBody0() { return State.SolverBodies[0]; }
+		const FConstraintSolverBody& SolverBody0() const { return State.SolverBodies[0]; }
+
+		/**
+		 * @brief Get the second (decaorated) solver body
+		 * The decorator add a possible mass scale
+		*/
+		FConstraintSolverBody& SolverBody1() { return State.SolverBodies[1]; }
+		const FConstraintSolverBody& SolverBody1() const { return State.SolverBodies[1]; }
+
+		/**
+		 * @brief Set up the mass scaling for shock propagation, using the position-phase mass scale
+		*/
+		void EnablePositionShockPropagation();
+
+		/**
+		 * @brief Set up the mass scaling for shock propagation, using the velocity-phase mass scale
+		*/
+		void EnableVelocityShockPropagation();
+
+		/**
+		 * @brief Disable mass scaling
+		*/
+		void DisableShockPropagation();
+
+		/**
+		 * @brief Calculate and apply the position correction for this iteration
+		 * @return true if we need to run more iterations, false if we did not apply any correction
+		*/
+		bool SolvePosition(const FReal Dt, const bool bApplyStaticFriction);
+
+		/**
+		 * @brief Calculate and apply the velocity correction for this iteration
+		 * @return true if we need to run more iterations, false if we did not apply any correction
+		*/
+		bool SolveVelocity(const FReal Dt, const bool bApplyDynamicFriction);
+
+	private:
+		/**
+		 * @brief Apply the inverse mass scale the body with the lower level
+		 * @param InvMassScale 
+		*/
+		void SetShockPropagationInvMassScale(const FReal InvMassScale);
+
+		struct FState
+		{
+			FState();
+
+			FConstraintSolverBody SolverBodies[MaxConstrainedBodies];
+			FSolverManifoldPoint ManifoldPoints[MaxPointsPerConstraint];
+			int32 NumManifoldPoints;
+			FReal StaticFriction;
+			FReal DynamicFriction;
+			FReal Stiffness;
+			int32 BodyEpochs[MaxConstrainedBodies];
+			bool bIsSolved;
+		};
+
+		FState State;
 	};
-}
\ No newline at end of file
+}
+
+CHAOS_API DECLARE_LOG_CATEGORY_EXTERN(LogChaosCollision, Log, All);
diff --git a/Engine/Source/Runtime/Experimental/Chaos/Private/Chaos/Collision/SolverCollisionContainer.cpp b/Engine/Source/Runtime/Experimental/Chaos/Private/Chaos/Collision/SolverCollisionContainer.cpp
new file mode 100644
--- /dev/null
+++ b/Engine/Source/Runtime/Experimental/Chaos/Private/Chaos/Collision/SolverCollisionContainer.cpp
@@ -0,0 +1,362 @@
+// Copyright Epic Games, Inc. All Rights Reserved.
+#include "Chaos/Collision/SolverCollisionContainer.h"
+
+#include "Chaos/Collision/CollisionApplyType.h"
+#include "Chaos/Collision/PBDCollisionConstraint.h"
+#include "Chaos/Evolution/SolverBodyContainer.h"
+#include "Chaos/Particle/ParticleUtilities.h"
+#include "Chaos/PBDCollisionConstraints.h"
+#include "Chaos/PBDCollisionConstraintsContact.h"
+#include "Chaos/Utilities.h"
+
+// Private includes
+#include "PBDCollisionSolver.h"
+
+#include "ChaosLog.h"
+
+
+namespace Chaos
+{
+	namespace CVars
+	{
+		extern int32 Chaos_PBDCollisionSolver_Position_ZeroFrictionIterations;
+		extern int32 Chaos_PBDCollisionSolver_Position_ShockPropagationIterations;
+
+		extern int32 Chaos_PBDCollisionSolver_Velocity_ShockPropagationIterations;
+	}
+	using namespace CVars;
+
+	extern bool bChaos_Collision_Manifold_FixNormalsInWorldSpace;
+
+
+	/**
+	 * @brief A wrapper for FPBDCollisionSolver which binds to a Collision Constraint and adds Gather/Scatter from/to the constraint
+	*/
+	class FPBDCollisionSolverAdapter
+	{
+	public:
+		FPBDCollisionSolverAdapter()
+			: Constraint(nullptr)
+		{
+		}
+
+		FPBDCollisionSolver& GetSolver() { return Solver; }
+		FPBDCollisionConstraint* GetConstraint() { return Constraint; }
+
+		bool IsIncrementalManifold() const { return Constraint->GetUseIncrementalCollisionDetection(); }
+
+		/**
+		 * @brief Initialize the data required for the solver, and bind to the bodies
+		*/
+		void GatherInput(
+			const FReal Dt,
+			FPBDCollisionConstraint& InConstraint,
+			const int32 Particle0Level,
+			const int32 Particle1Level,
+			FSolverBodyContainer& SolverBodyContainer)
+		{
+			Constraint = &InConstraint;
+
+			// Find the solver bodies for the particles we constrain. This will add them to the container
+			// if they aren't there already, and ensure that they are populated with the latest data.
+			FSolverBody* Body0 = SolverBodyContainer.FindOrAdd(Constraint->Particle[0]);
+			FSolverBody* Body1 = SolverBodyContainer.FindOrAdd(Constraint->Particle[1]);
+
+			Body0->SetLevel(Particle0Level);
+			Body1->SetLevel(Particle1Level);
+
+			Solver.SetFriction(
+				FMath::Max(Constraint->Manifold.AngularFriction, Constraint->Manifold.Friction),
+				Constraint->Manifold.Friction);
+
+			Solver.SetStiffness(Constraint->GetStiffness());
+
+			Solver.SetSolverBodies(Body0, Body1);
+
+			GatherManifoldPoints(Dt);
+
+			// We should try to remove this
+			Constraint->SetSolverBodies(Body0, Body1);
+		}
+
+		void GatherManifoldPoints(const FReal Dt)
+		{
+			// We handle incremental manifolds by just colecting any new contacts
+			const int32 BeginPointIndex = Solver.NumManifoldPoints();
+			const int32 EndPointIndex = FMath::Min(Constraint->GetManifoldPoints().Num(), FPBDCollisionSolver::MaxPointsPerConstraint);
+
+			Solver.SetNumManifoldPoints(EndPointIndex);
+
+			const FReal RestitutionVelocityThreshold = Constraint->Manifold.RestitutionThreshold * Dt;
+			const FReal Restitution = Constraint->Manifold.Restitution;
+
+			for (int32 ManifoldPointIndex = BeginPointIndex; ManifoldPointIndex < EndPointIndex; ++ManifoldPointIndex)
+			{
+				const FManifoldPoint& ManifoldPoint = Constraint->GetManifoldPoints()[ManifoldPointIndex];
+
+				// The world-space contact normal
+				FVec3 WorldContactNormal;
+				if (bChaos_Collision_Manifold_FixNormalsInWorldSpace)
+				{
+					WorldContactNormal = ManifoldPoint.ManifoldContactNormal;
+				}
+				else
+				{
+					const FRotation3& PlaneQ = (ManifoldPoint.ContactPoint.ContactNormalOwnerIndex == 0) ? Solver.SolverBody0().Q() : Solver.SolverBody1().Q();
+					WorldContactNormal = PlaneQ * ManifoldPoint.ManifoldContactNormal;
+				}
+
+				// Initialize the structural data in the contact (relative contact points, contact mass etc)
+				Solver.InitContact(
+					ManifoldPointIndex,
+					ManifoldPoint.CoMAnchorPoints[0],
+					ManifoldPoint.CoMAnchorPoints[1],
+					WorldContactNormal);
+
+				// Calculate the target normal velocity based on restitution
+				FReal WorldContactVelocityTargetNormal = FReal(0);
+				if (Restitution > FReal(0))
+				{
+					const FVec3 ContactVelocity = Solver.GetManifoldPoint(ManifoldPointIndex).CalculateContactVelocity(Solver.SolverBody0(), Solver.SolverBody1());
+					const FReal ContactVelocityNormal = FVec3::DotProduct(ContactVelocity, WorldContactNormal);
+					if (ContactVelocityNormal < -RestitutionVelocityThreshold)
+					{
+						WorldContactVelocityTargetNormal = -Restitution * ContactVelocityNormal;
+					}
+				}
+
+				// Initialize the material properties (restitution and friction related)
+				Solver.InitMaterial(
+					ManifoldPointIndex,
+					WorldContactVelocityTargetNormal,
+					ManifoldPoint.bPotentialRestingContact,
+					ManifoldPoint.StaticFrictionMax);
+			}
+		}
+
+		/**
+		 * @brief Send all solver results to the constraint
+		*/
+		void ScatterOutput(
+			const FReal Dt)
+		{
+			FVec3 AccumulatedImpulse = FVec3(0);
+
+			for (int32 PointIndex = 0; PointIndex < Solver.NumManifoldPoints(); ++PointIndex)
+			{
+				const FPBDCollisionSolver::FSolverManifoldPoint& SolverManifoldPoint = Solver.GetManifoldPoint(PointIndex);
+				FManifoldPoint& ManifoldPoint = Constraint->GetManifoldPoints()[PointIndex];
+
+				ManifoldPoint.NetPushOut = SolverManifoldPoint.NetPushOut;
+				ManifoldPoint.NetPushOutNormal = FVec3::DotProduct(SolverManifoldPoint.NetPushOut, SolverManifoldPoint.WorldContactNormal);
+				ManifoldPoint.NetImpulse = SolverManifoldPoint.NetImpulse;
+				ManifoldPoint.NetImpulseNormal = FVec3::DotProduct(SolverManifoldPoint.NetImpulse, SolverManifoldPoint.WorldContactNormal);
+				ManifoldPoint.bActive = !SolverManifoldPoint.NetPushOut.IsNearlyZero();
+				ManifoldPoint.bInsideStaticFrictionCone = SolverManifoldPoint.bInsideStaticFrictionCone;
+
+				// If the contact has no static friction, reset the static friction point to the current contact position
+				// @todo(chaos): this should keep the prev contact at the friction cone edge, rather than just resetting 
+				// (this currently causes stick-slide movement just above the friction threshold)
+				if (!ManifoldPoint.bActive || !ManifoldPoint.bInsideStaticFrictionCone)
+				{
+					ManifoldPoint.CoMAnchorPoints[0] = ManifoldPoint.CoMContactPoints[0];
+					ManifoldPoint.CoMAnchorPoints[1] = ManifoldPoint.CoMContactPoints[1];
+				}
+
+				AccumulatedImpulse += SolverManifoldPoint.NetImpulse + (SolverManifoldPoint.NetPushOut / Dt);
+			}
+
+			Constraint->AccumulatedImpulse = AccumulatedImpulse;
+
+			Constraint->SetSolverBodies(nullptr, nullptr);
+			Constraint = nullptr;
+			Solver.ResetSolverBodies();
+		}
+
+
+	private:
+		FPBDCollisionSolver Solver;
+		FPBDCollisionConstraint* Constraint;
+	};
+
+
+	FPBDCollisionSolverContainer::FPBDCollisionSolverContainer()
+		: FConstraintSolverContainer()
+	{
+	}
+
+	FPBDCollisionSolverContainer::~FPBDCollisionSolverContainer()
+	{
+	}
+
+	void FPBDCollisionSolverContainer::Reset(const int32 MaxCollisions)
+	{
+		CollisionSolvers.Reset(MaxCollisions);
+		SweptCollisionSolvers.Reset();
+	}
+
+	FPBDCollisionSolver* FPBDCollisionSolverContainer::AddConstraintSolver(FReal Dt, FPBDCollisionConstraint& Constraint, const int32 Particle0Level, const int32 Particle1Level, FSolverBodyContainer& SolverBodyContainer)
+	{
+		// This container is required to allocate pointers that are valid for the whole tick,
+		// so we cannot allow the container to resize during the tick. See Reset()
+		check(CollisionSolvers.Num() < CollisionSolvers.Max());
+
+		const int32 SolverCollisionIndex = CollisionSolvers.AddDefaulted();
+
+		CollisionSolvers[SolverCollisionIndex].GatherInput(Dt, Constraint, Particle0Level, Particle1Level, SolverBodyContainer);
+
+		// Keep a list of CCD contacts - they will get resolved before we move on to the main solver loop
+		if (Constraint.GetType() == ECollisionConstraintType::Swept)
+		{
+			SweptCollisionSolvers.Add(SolverCollisionIndex);
+		}
+
+		return &CollisionSolvers[SolverCollisionIndex].GetSolver();
+	}
+
+	void FPBDCollisionSolverContainer::UpdatePositionShockPropagation(const FReal Dt, const int32 It, const int32 NumIts, const int32 BeginIndex, const int32 EndIndex)
+	{
+		// If this is the first shock propagation iteration, enable it on each solver
+		const int32 NumShockPropIterations = Chaos_PBDCollisionSolver_Position_ShockPropagationIterations;
+		const bool bEnableShockPropagation = (It == NumIts - NumShockPropIterations);
+		if (bEnableShockPropagation)
+		{
+			for (int32 SolverIndex = BeginIndex; SolverIndex < EndIndex; ++SolverIndex)
+			{
+				CollisionSolvers[SolverIndex].GetSolver().EnablePositionShockPropagation();
+			}
+		}
+	}
+
+	void FPBDCollisionSolverContainer::UpdateVelocityShockPropagation(const FReal Dt, const int32 It, const int32 NumIts, const int32 BeginIndex, const int32 EndIndex)
+	{
+		// Set/reset the shock propagation based on current iteration. The position solve may
+		// have left the bodies with a mass scale and we want to change or reset it.
+		const int32 NumShockPropIterations = Chaos_PBDCollisionSolver_Velocity_ShockPropagationIterations;
+		const bool bEnableShockPropagation = (It == NumIts - NumShockPropIterations);
+		if (bEnableShockPropagation)
+		{
+			for (int32 SolverIndex = BeginIndex; SolverIndex < EndIndex; ++SolverIndex)
+			{
+				CollisionSolvers[SolverIndex].GetSolver().EnableVelocityShockPropagation();
+			}
+		}
+		else if (It == 0)
+		{
+			for (int32 SolverIndex = BeginIndex; SolverIndex < EndIndex; ++SolverIndex)
+			{
+				CollisionSolvers[SolverIndex].GetSolver().DisableShockPropagation();
+			}
+		}
+	}
+
+	void FPBDCollisionSolverContainer::SolveSwept(const FReal Dt)
+	{
+		for (int32 SweptSolverCollisionIndex : SweptCollisionSolvers)
+		{
+			SolveSwept(Dt, *CollisionSolvers[SweptSolverCollisionIndex].GetConstraint(), CollisionSolvers[SweptSolverCollisionIndex].GetSolver());
+		}
+	}
+
+	bool FPBDCollisionSolverContainer::SolvePositionSerial(const FReal Dt, const int32 It, const int32 NumIts, const int32 BeginIndex, const int32 EndIndex)
+	{
+		return SolvePositionImpl(Dt, It, NumIts, BeginIndex, EndIndex, false);
+	}
+
+	bool FPBDCollisionSolverContainer::SolveVelocitySerial(const FReal Dt, const int32 It, const int32 NumIts, const int32 BeginIndex, const int32 EndIndex)
+	{
+		return SolveVelocityImpl(Dt, It, NumIts, BeginIndex, EndIndex, false);
+	}
+
+	bool FPBDCollisionSolverContainer::SolvePositionParallel(const FReal Dt, const int32 It, const int32 NumIts, const int32 BeginIndex, const int32 EndIndex)
+	{
+		return SolvePositionImpl(Dt, It, NumIts, BeginIndex, EndIndex, true);
+	}
+
+	bool FPBDCollisionSolverContainer::SolveVelocityParallel(const FReal Dt, const int32 It, const int32 NumIts, const int32 BeginIndex, const int32 EndIndex)
+	{
+		return SolveVelocityImpl(Dt, It, NumIts, BeginIndex, EndIndex, true);
+	}
+
+	void FPBDCollisionSolverContainer::SolveSwept(const FReal Dt, FPBDCollisionConstraint& Constraint, FPBDCollisionSolver& CollisionSolver)
+	{
+		// @todo(chaos): Fix CCD. This "solution" just rewinds the CCD object to the first contact and will result in lost momentum
+		if (Constraint.TimeOfImpact < FReal(1))
+		{
+			Collisions::UpdateSwept(Constraint, Dt);
+
+			// Just rewind the body to the impact. This is not good at all, but it will prevent penetration for now
+			FSolverBody& Body0 = CollisionSolver.SolverBody0().SolverBody();
+			Body0.SetP(FMath::Lerp(Body0.X(), Body0.P(), Constraint.TimeOfImpact));
+		}
+	}
+
+	bool FPBDCollisionSolverContainer::SolvePositionImpl(const FReal Dt, const int32 It, const int32 NumIts, const int32 BeginIndex, const int32 EndIndex, const bool bParallel)
+	{
+		SCOPE_CYCLE_COUNTER(STAT_Collisions_Apply);
+
+		UpdatePositionShockPropagation(Dt, It, NumIts, BeginIndex, EndIndex);
+
+		const int32 NumZeroFrictionIterations = Chaos_PBDCollisionSolver_Position_ZeroFrictionIterations;
+		const bool bApplyStaticFriction = (It >= NumZeroFrictionIterations);
+
+		// Apply the position correction
+		// @todo(chaos): parallel version of SolvePosition
+		bool bNeedsAnotherIteration = false;
+		for (int32 SolverIndex = BeginIndex; SolverIndex < EndIndex; ++SolverIndex)
+		{
+			FPBDCollisionSolverAdapter& CollisionSolver = CollisionSolvers[SolverIndex];
+
+			// @todo(chaos): remove this when we no longer need to support incremental manifolds - this prevents us from fully optimizing the solver loop
+			if (CollisionSolver.IsIncrementalManifold())
+			{
+				Collisions::Update(*CollisionSolver.GetConstraint(), Dt);
+				CollisionSolver.GatherManifoldPoints(Dt);
+			}
+
+			bNeedsAnotherIteration |= CollisionSolvers[SolverIndex].GetSolver().SolvePosition(Dt, bApplyStaticFriction);
+		}
+
+		return bNeedsAnotherIteration;
+	}
+
+	bool FPBDCollisionSolverContainer::SolveVelocityImpl(const FReal Dt, const int32 It, const int32 NumIts, const int32 BeginIndex, const int32 EndIndex, const bool bParallel)
+	{
+		SCOPE_CYCLE_COUNTER(STAT_Collisions_ApplyPushOut);
+
+		UpdateVelocityShockPropagation(Dt, It, NumIts, BeginIndex, EndIndex);
+
+		const int32 NumDynamicFrictionIterations = 1;
+		const bool bApplyDynamicFriction = (It >= NumIts - NumDynamicFrictionIterations);
+
+		// Apply the velocity correction
+		// @todo(chaos): parallel version of SolveVelocity
+		bool bNeedsAnotherIteration = false;
+		for (int32 SolverIndex = BeginIndex; SolverIndex < EndIndex; ++SolverIndex)
+		{
+			FPBDCollisionSolverAdapter& CollisionSolver = CollisionSolvers[SolverIndex];
+
+			bNeedsAnotherIteration |= CollisionSolver.GetSolver().SolveVelocity(Dt, bApplyDynamicFriction);
+		}
+
+		return bNeedsAnotherIteration;
+	}
+
+	void FPBDCollisionSolverContainer::ScatterOutput(const FReal Dt, const int32 BeginIndex, const int32 EndIndex)
+	{
+		return ScatterOutputImpl(Dt, BeginIndex, EndIndex, false);
+	}
+
+	void FPBDCollisionSolverContainer::ScatterOutputImpl(const FReal Dt, const int32 BeginIndex, const int32 EndIndex, const bool bParallel)
+	{
+		SCOPE_CYCLE_COUNTER(STAT_Collisions_Scatter);
+		check(BeginIndex >= 0);
+		check(EndIndex <= CollisionSolvers.Num());
+
+		for (int32 SolverIndex = BeginIndex; SolverIndex < EndIndex; ++SolverIndex)
+		{
+			CollisionSolvers[SolverIndex].ScatterOutput(Dt);
+		}
+	}
+
+}
diff --git a/Engine/Source/Runtime/Experimental/Chaos/Private/Chaos/Collision/SphereConvexContactPoint.cpp b/Engine/Source/Runtime/Experimental/Chaos/Private/Chaos/Collision/SphereConvexContactPoint.cpp
--- a/Engine/Source/Runtime/Experimental/Chaos/Private/Chaos/Collision/SphereConvexContactPoint.cpp
+++ b/Engine/Source/Runtime/Experimental/Chaos/Private/Chaos/Collision/SphereConvexContactPoint.cpp
@@ -82,7 +82,7 @@ namespace Chaos
 				ContactPoint.ShapeContactNormal = -NormalConvex;
 				ContactPoint.ShapeMargins[0] = 0.0f;
 				ContactPoint.ShapeMargins[1] = 0.0f;
-				ContactPoint.ContactNormalOwnerIndex = 0;
+				ContactPoint.ContactNormalOwnerIndex = 0; // ContactNormalOwnerIndex is ignored when  bChaos_Collision_Manifold_FixNormalsInWorldSpace is set
 
 				ContactPoint.Location = ConvexTransform.TransformPosition(ContactPoint.ShapeContactPoints[0]);
 				ContactPoint.Normal = ConvexTransform.TransformVectorNoScale(-NormalConvex);
diff --git a/Engine/Source/Runtime/Experimental/Chaos/Private/Chaos/CollisionOneShotManifolds.cpp b/Engine/Source/Runtime/Experimental/Chaos/Private/Chaos/CollisionOneShotManifolds.cpp
--- a/Engine/Source/Runtime/Experimental/Chaos/Private/Chaos/CollisionOneShotManifolds.cpp
+++ b/Engine/Source/Runtime/Experimental/Chaos/Private/Chaos/CollisionOneShotManifolds.cpp
@@ -10,11 +10,39 @@
 #include "Chaos/ImplicitObjectScaled.h"
 #include "Chaos/Transform.h"
 #include "Chaos/Utilities.h"
+#include "ChaosStats.h"
 
 #include "HAL/IConsoleManager.h"
 
 //PRAGMA_DISABLE_OPTIMIZATION
 
+#if 0
+DECLARE_STATS_GROUP(TEXT("ChaosManifold"), STATGROUP_ChaosManifold, STATCAT_Advanced);
+
+DECLARE_CYCLE_STAT(TEXT("Manifold::Manifold"), STAT_Collisions_Manifold, STATGROUP_ChaosManifold);
+DECLARE_CYCLE_STAT(TEXT("Manifold::ManifoldGJK"), STAT_Collisions_ManifoldGJK, STATGROUP_ChaosManifold);
+DECLARE_CYCLE_STAT(TEXT("Manifold::ManifoldEdgeEdge"), STAT_Collisions_ManifoldEdgeEdge, STATGROUP_ChaosManifold);
+DECLARE_CYCLE_STAT(TEXT("Manifold::ManifoldClip"), STAT_Collisions_ManifoldClip, STATGROUP_ChaosManifold);
+DECLARE_CYCLE_STAT(TEXT("Manifold::ManifoldReduce"), STAT_Collisions_ManifoldReduce, STATGROUP_ChaosManifold);
+DECLARE_CYCLE_STAT(TEXT("Manifold::ManifoldFaceVertex"), STAT_Collisions_ManifoldFaceVertex, STATGROUP_ChaosManifold);
+#define SCOPE_CYCLE_COUNTER_MANIFOLD() SCOPE_CYCLE_COUNTER(STAT_Collisions_Manifold)
+#define SCOPE_CYCLE_COUNTER_MANIFOLD_GJK() SCOPE_CYCLE_COUNTER(STAT_Collisions_ManifoldGJK)
+#define SCOPE_CYCLE_COUNTER_MANIFOLD_ADDEDGEEDGE() SCOPE_CYCLE_COUNTER(STAT_Collisions_ManifoldEdgeEdge)
+#define SCOPE_CYCLE_COUNTER_MANIFOLD_CLIP() SCOPE_CYCLE_COUNTER(STAT_Collisions_ManifoldClip)
+#define SCOPE_CYCLE_COUNTER_MANIFOLD_REDUCE() SCOPE_CYCLE_COUNTER(STAT_Collisions_ManifoldReduce)
+#define SCOPE_CYCLE_COUNTER_MANIFOLD_ADDFACEVERTEX() SCOPE_CYCLE_COUNTER(STAT_Collisions_ManifoldFaceVertex)
+#else
+#define SCOPE_CYCLE_COUNTER_MANIFOLD()
+#define SCOPE_CYCLE_COUNTER_MANIFOLDGJK()
+#define SCOPE_CYCLE_COUNTER_MANIFOLD()
+#define SCOPE_CYCLE_COUNTER_MANIFOLD_GJK()
+#define SCOPE_CYCLE_COUNTER_MANIFOLD_ADDEDGEEDGE()
+#define SCOPE_CYCLE_COUNTER_MANIFOLD_CLIP()
+#define SCOPE_CYCLE_COUNTER_MANIFOLD_REDUCE()
+#define SCOPE_CYCLE_COUNTER_MANIFOLD_ADDFACEVERTEX()
+#endif
+
+
 namespace Chaos
 {
 	FRealSingle Chaos_Collision_Manifold_PlaneContactNormalEpsilon = 0.001f;
@@ -36,6 +64,9 @@ namespace Chaos
 	FRealSingle Chaos_Collision_Manifold_MinFaceSearchDistance = 1.0f;
 	FAutoConsoleVariableRef CVarChaosCollisioConvexManifoldMinFaceSearchDistance(TEXT("p.Chaos.Collision.Manifold.MinFaceSearchDistance"), Chaos_Collision_Manifold_MinFaceSearchDistance, TEXT(""));
 
+	bool bChaos_Collision_Manifold_FixNormalsInWorldSpace = true;
+	FAutoConsoleVariableRef CVarChaosCollisioFixNormalsInWorldSpace(TEXT("p.Chaos.Collision.Manifold.FixNormalsInWorldSpace"), bChaos_Collision_Manifold_FixNormalsInWorldSpace, TEXT("Fix normals in world space at beginning of frame"));	
+
 	bool ForceOneShotManifoldEdgeEdgeCaseZeroCullDistance = false;
 	FAutoConsoleVariableRef CVarForceOneShotManifoldEdgeEdgeCaseZeroCullDistance(TEXT("p.Chaos.Collision.Manifold.ForceOneShotManifoldEdgeEdgeCaseZeroCullDistance"), ForceOneShotManifoldEdgeEdgeCaseZeroCullDistance,
 	TEXT("If enabled, if one shot manifold hits edge/edge case, we will force a cull distance of zero. That means edge/edge contacts will be thrown out if separated at all. Only applies to Convex/Convex oneshot impl."));
@@ -127,7 +158,7 @@ namespace Chaos
 			const FImplicitBox3& Box2,
 			const FRigidTransform3& Box2Transform, //world
 			const FReal Dt,
-			FRigidBodyPointContactConstraint& Constraint)
+			FPBDCollisionConstraint& Constraint)
 		{
 			if (bChaos_Collision_Manifold_BoxAsConvex)
 			{
@@ -322,8 +353,16 @@ namespace Chaos
 				ContactPoint.ShapeMargins[1] = 0.0f;
 				ContactPoint.ShapeContactPoints[0] = ReferenceFaceBox1 ? PointProjectedOntoReferenceFace + RefBox->GetCenter() : ClippedPointInOtherCubeCoordinates + OtherBox->GetCenter();
 				ContactPoint.ShapeContactPoints[1] = ReferenceFaceBox1 ? ClippedPointInOtherCubeCoordinates + OtherBox->GetCenter() : PointProjectedOntoReferenceFace + RefBox->GetCenter();
-				ContactPoint.ShapeContactNormal = ReferenceFaceBox1 ? SeparationDirectionLocalBox1 : SeparationDirectionLocalBox2;
-				ContactPoint.ContactNormalOwnerIndex = ReferenceFaceBox1 ? 0 : 1;
+				if (bChaos_Collision_Manifold_FixNormalsInWorldSpace)
+				{
+					ContactPoint.ShapeContactNormal = SeparationDirectionLocalBox2;
+					ContactPoint.ContactNormalOwnerIndex = 1; // Ignored
+				}
+				else
+				{
+					ContactPoint.ShapeContactNormal = ReferenceFaceBox1 ? SeparationDirectionLocalBox1 : SeparationDirectionLocalBox2;
+					ContactPoint.ContactNormalOwnerIndex = ReferenceFaceBox1 ? 0 : 1;
+				}
 				ContactPoint.Location = RefBoxTM->TransformPositionNoScale(PointProjectedOntoReferenceFace);
 				ContactPoint.Normal = GJKContactPoint.Normal;
 				ContactPoint.Phi = FVec3::DotProduct(PointProjectedOntoReferenceFace - VertexInReferenceCubeCoordinates, ReferenceFaceBox1 ? SeparationDirectionLocalBox1 : -SeparationDirectionLocalBox2);
@@ -540,13 +579,14 @@ namespace Chaos
 		// Use GJK to find the closest points (or shallowest penetrating points) on two convex shapes usingthe specified margin
 		// @todo(chaos): dedupe from GJKContactPoint in CollisionResolution.cpp
 		template <typename GeometryA, typename GeometryB>
-		FContactPoint GJKContactPointMargin(const GeometryA& A, const GeometryB& B, const FRigidTransform3& ATM, const FRigidTransform3& BToATM, FReal MarginA, FReal MarginB)
+		FContactPoint GJKContactPointMargin(const GeometryA& A, const GeometryB& B, const FRigidTransform3& ATM, const FRigidTransform3& BToATM, FReal MarginA, FReal MarginB, FGJKSimplexData& InOutGjkWarmStartData)
 		{
+			SCOPE_CYCLE_COUNTER_MANIFOLD_GJK();
+
 			FContactPoint Contact;
 
 			FReal Penetration;
-			FVec3 ClosestA, ClosestBInA, Normal;
-			int32 ClosestVertexIndexA, ClosestVertexIndexB;
+			FVec3 ClosestA, ClosestB, NormalA, NormalB;
 
 			// Slightly increased epsilon to reduce error in normal for almost touching objects.
 			const FReal Epsilon = 3.e-3f;
@@ -554,16 +594,15 @@ namespace Chaos
 			const TGJKCoreShape<GeometryA> AWithMargin(A, MarginA);
 			const TGJKCoreShape<GeometryB> BWithMargin(B, MarginB);
 
-			if (GJKPenetration<true>(AWithMargin, BWithMargin, BToATM, Penetration, ClosestA, ClosestBInA, Normal, ClosestVertexIndexA, ClosestVertexIndexB, (FReal)0, (FReal)0, FVec3(1,0,0), Epsilon))
+			if (GJKPenetrationWarmStartable<true>(AWithMargin, BWithMargin, BToATM, FReal(0), FReal(0), Penetration, ClosestA, ClosestB, NormalA, NormalB, InOutGjkWarmStartData, Epsilon))
 			{
-				// GJK output is all in the local space of A. We need to transform the B-relative position and the normal in to B-space
 				Contact.ShapeMargins[0] = 0.0f;
 				Contact.ShapeMargins[1] = 0.0f;
 				Contact.ShapeContactPoints[0] = ClosestA;
-				Contact.ShapeContactPoints[1] = BToATM.InverseTransformPosition(ClosestBInA);
-				Contact.ShapeContactNormal = -BToATM.InverseTransformVector(Normal);
-				Contact.Location = ATM.TransformPosition(ClosestA);
-				Contact.Normal = -ATM.TransformVectorNoScale(Normal);
+				Contact.ShapeContactPoints[1] = ClosestB;
+				Contact.ShapeContactNormal = -NormalB;	// We want normal pointing from B to A
+				Contact.Location = ATM.TransformPositionNoScale(ClosestA);
+				Contact.Normal = -ATM.TransformVectorNoScale(NormalA);
 				Contact.Phi = -Penetration;
 			}
 
@@ -636,8 +675,10 @@ namespace Chaos
 			const ConvexImplicitType2& Convex2,
 			const FRigidTransform3& Convex2Transform, //world
 			const FReal Dt,
-			FRigidBodyPointContactConstraint& Constraint)
+			FPBDCollisionConstraint& Constraint)
 		{
+			SCOPE_CYCLE_COUNTER_MANIFOLD();
+
 			const uint32 SpaceDimension = 3;
 
 			// We only build one shot manifolds once
@@ -665,7 +706,7 @@ namespace Chaos
 			}
 
 			// Find the deepest penetration. This is used to determine the planes and points to use for the manifold
-			FContactPoint GJKContactPoint = GJKContactPointMargin(Convex1, Convex2, Convex1Transform, Convex2ToConvex1Transform, Margin1, Margin2);
+			FContactPoint GJKContactPoint = GJKContactPointMargin(Convex1, Convex2, Convex1Transform, Convex2ToConvex1Transform, Margin1, Margin2, Constraint.GetGJKWarmStartData());
 
 			// GJK is using margins and rounded corner, so if we have a corner-to-corner contact it will under-report the actual distance by an amount that depends on how
 			// "pointy" the edge/corner is - it can be arbitrarily large. For now just expand the cull distance by some multiple of the margin (it would need to be
@@ -704,6 +745,8 @@ namespace Chaos
 			// For edge-edge contacts, we find the edges involved and project the contact onto the edges
 			if (!bIsPlaneContact)
 			{
+				SCOPE_CYCLE_COUNTER_MANIFOLD_ADDEDGEEDGE();
+
 				if (ForceOneShotManifoldEdgeEdgeCaseZeroCullDistance && GJKContactPoint.Phi > 0)
 				{
 					return;
@@ -739,41 +782,44 @@ namespace Chaos
 			const FRigidTransform3* RefConvexTM;
 			FRigidTransform3 ConvexOtherToRef;
 
-			if (ReferenceFaceConvex1)
 			{
-				RefConvexTM = &Convex1Transform;
-				ConvexOtherToRef = Convex2Transform.GetRelativeTransform(Convex1Transform);
-
-				ClippedVertices = GenerateConvexManifoldClippedVertices(
-					Convex1,
-					Convex2,
-					ConvexOtherToRef, 
-					MostOpposingPlaneIndexConvex1,
-					MostOpposingPlaneIndexConvex2,
-					RefPlaneNormal,
-					RefPlanePosition,
-					ClippedVertices1,
-					ClippedVertices2,
-					ContactPointCount,
-					MaxContactPointCount);
-			}
-			else
-			{
-				RefConvexTM = &Convex2Transform;
-				ConvexOtherToRef = Convex1Transform.GetRelativeTransform(Convex2Transform);
-
-				ClippedVertices = GenerateConvexManifoldClippedVertices(
-					Convex2,
-					Convex1,
-					ConvexOtherToRef,
-					MostOpposingPlaneIndexConvex2,
-					MostOpposingPlaneIndexConvex1,
-					RefPlaneNormal,
-					RefPlanePosition,
-					ClippedVertices1,
-					ClippedVertices2,
-					ContactPointCount,
-					MaxContactPointCount);
+				SCOPE_CYCLE_COUNTER_MANIFOLD_CLIP();
+				if (ReferenceFaceConvex1)
+				{
+					RefConvexTM = &Convex1Transform;
+					ConvexOtherToRef = Convex2Transform.GetRelativeTransform(Convex1Transform);
+
+					ClippedVertices = GenerateConvexManifoldClippedVertices(
+						Convex1,
+						Convex2,
+						ConvexOtherToRef, 
+						MostOpposingPlaneIndexConvex1,
+						MostOpposingPlaneIndexConvex2,
+						RefPlaneNormal,
+						RefPlanePosition,
+						ClippedVertices1,
+						ClippedVertices2,
+						ContactPointCount,
+						MaxContactPointCount);
+				}
+				else
+				{
+					RefConvexTM = &Convex2Transform;
+					ConvexOtherToRef = Convex1Transform.GetRelativeTransform(Convex2Transform);
+
+					ClippedVertices = GenerateConvexManifoldClippedVertices(
+						Convex2,
+						Convex1,
+						ConvexOtherToRef,
+						MostOpposingPlaneIndexConvex2,
+						MostOpposingPlaneIndexConvex1,
+						RefPlaneNormal,
+						RefPlanePosition,
+						ClippedVertices1,
+						ClippedVertices2,
+						ContactPointCount,
+						MaxContactPointCount);
+				}
 			}
 
 			// If we have the max number of contact points already, they will be in cyclic order. Stability is better
@@ -787,6 +833,8 @@ namespace Chaos
 			// Reduce number of contacts to the maximum allowed
 			if (ContactPointCount > 4)
 			{
+				SCOPE_CYCLE_COUNTER_MANIFOLD_REDUCE();
+
 				FRotation3 RotateSeperationToZ = FRotation3::FromRotatedVector(RefPlaneNormal, FVec3(0.0f, 0.0f, 1.0f));
 				for (uint32 ContactPointIndex = 0; ContactPointIndex < ContactPointCount; ++ContactPointIndex)
 				{
@@ -802,25 +850,38 @@ namespace Chaos
 			}
 
 			// Generate the contact points from the clipped vertices
-			for (uint32 ContactPointIndex = 0; ContactPointIndex < ContactPointCount; ++ContactPointIndex)
 			{
-				FContactPoint ContactPoint;
-				const FVec3 VertexInReferenceCoordinates = ClippedVertices[ContactPointIndex];
-				FVec3 PointProjectedOntoReferenceFace = VertexInReferenceCoordinates - FVec3::DotProduct(VertexInReferenceCoordinates - RefPlanePosition, RefPlaneNormal) * RefPlaneNormal;
-				FVec3 ClippedPointInOtherCoordinates = ConvexOtherToRef.InverseTransformPositionNoScale(VertexInReferenceCoordinates);
-
-				ContactPoint.ShapeMargins[0] = 0.0f;
-				ContactPoint.ShapeMargins[1] = 0.0f;
-				ContactPoint.ShapeContactPoints[0] = ReferenceFaceConvex1 ? PointProjectedOntoReferenceFace : ClippedPointInOtherCoordinates;
-				ContactPoint.ShapeContactPoints[1] = ReferenceFaceConvex1 ? ClippedPointInOtherCoordinates : PointProjectedOntoReferenceFace;
-				ContactPoint.ShapeContactNormal = RefSeparationDirection;
-				ContactPoint.ContactNormalOwnerIndex = ReferenceFaceConvex1 ? 0 : 1;
+				SCOPE_CYCLE_COUNTER_MANIFOLD_ADDFACEVERTEX();
+				for (uint32 ContactPointIndex = 0; ContactPointIndex < ContactPointCount; ++ContactPointIndex)
+				{
+					FContactPoint ContactPoint;
+					const FVec3 VertexInReferenceCoordinates = ClippedVertices[ContactPointIndex];
+					FVec3 PointProjectedOntoReferenceFace = VertexInReferenceCoordinates - FVec3::DotProduct(VertexInReferenceCoordinates - RefPlanePosition, RefPlaneNormal) * RefPlaneNormal;
+					FVec3 ClippedPointInOtherCoordinates = ConvexOtherToRef.InverseTransformPositionNoScale(VertexInReferenceCoordinates);
+
+					ContactPoint.ShapeMargins[0] = 0.0f;
+					ContactPoint.ShapeMargins[1] = 0.0f;
+					ContactPoint.ShapeContactPoints[0] = ReferenceFaceConvex1 ? PointProjectedOntoReferenceFace : ClippedPointInOtherCoordinates;
+					ContactPoint.ShapeContactPoints[1] = ReferenceFaceConvex1 ? ClippedPointInOtherCoordinates : PointProjectedOntoReferenceFace;
+					if (bChaos_Collision_Manifold_FixNormalsInWorldSpace)
+					{
+						ContactPoint.ShapeContactNormal = SeparationDirectionLocalConvex2;
+						ContactPoint.ContactNormalOwnerIndex = 1;// Owner is ignored
+						// These values will be overwritten later, but set it here for testing
+						ContactPoint.Location = ReferenceFaceConvex1 ? RefConvexTM->TransformPositionNoScale(VertexInReferenceCoordinates) : RefConvexTM->TransformPositionNoScale(PointProjectedOntoReferenceFace);
+					}
+					else
+					{
+						ContactPoint.ShapeContactNormal = RefSeparationDirection;
+						ContactPoint.ContactNormalOwnerIndex = ReferenceFaceConvex1 ? 0 : 1;
+						ContactPoint.Location = RefConvexTM->TransformPositionNoScale(PointProjectedOntoReferenceFace);
+					}
 
-				ContactPoint.Location = RefConvexTM->TransformPositionNoScale(PointProjectedOntoReferenceFace);
-				ContactPoint.Normal = GJKContactPoint.Normal;
-				ContactPoint.Phi = FVec3::DotProduct(PointProjectedOntoReferenceFace - VertexInReferenceCoordinates, ReferenceFaceConvex1 ? SeparationDirectionLocalConvex1 : -SeparationDirectionLocalConvex2);
+					ContactPoint.Normal = GJKContactPoint.Normal;
+					ContactPoint.Phi = FVec3::DotProduct(PointProjectedOntoReferenceFace - VertexInReferenceCoordinates, ReferenceFaceConvex1 ? SeparationDirectionLocalConvex1 : -SeparationDirectionLocalConvex2);
 
-				Constraint.AddOneshotManifoldContact(ContactPoint, Dt);
+					Constraint.AddOneshotManifoldContact(ContactPoint, Dt);
+				}
 			}
 		}
 
@@ -838,7 +899,7 @@ namespace Chaos
 			const FImplicitBox3& Implicit2,
 			const FRigidTransform3& Convex2Transform, //world
 			const FReal Dt,
-			FRigidBodyPointContactConstraint& Constraint);
+			FPBDCollisionConstraint& Constraint);
 
 		template
 		void ConstructConvexConvexOneShotManifold<FImplicitBox3, FImplicitConvex3>(
@@ -847,7 +908,7 @@ namespace Chaos
 			const FImplicitConvex3& Implicit2,
 			const FRigidTransform3& Convex2Transform, //world
 			const FReal Dt,
-			FRigidBodyPointContactConstraint& Constraint);
+			FPBDCollisionConstraint& Constraint);
 
 		template
 		void ConstructConvexConvexOneShotManifold<FImplicitConvex3, FImplicitBox3>(
@@ -856,7 +917,7 @@ namespace Chaos
 			const FImplicitBox3& Implicit2,
 			const FRigidTransform3& Convex2Transform, //world
 			const FReal Dt,
-			FRigidBodyPointContactConstraint& Constraint);
+			FPBDCollisionConstraint& Constraint);
 
 		template
 		void ConstructConvexConvexOneShotManifold<FImplicitBox3, TImplicitObjectInstanced<FImplicitConvex3>>(
@@ -865,7 +926,7 @@ namespace Chaos
 			const TImplicitObjectInstanced<FImplicitConvex3>& Implicit2,
 			const FRigidTransform3& Convex2Transform, //world
 			const FReal Dt,
-			FRigidBodyPointContactConstraint& Constraint);
+			FPBDCollisionConstraint& Constraint);
 
 		template
 		void ConstructConvexConvexOneShotManifold<TImplicitObjectInstanced<FImplicitConvex3>, FImplicitBox3>(
@@ -874,7 +935,7 @@ namespace Chaos
 			const FImplicitBox3& Implicit2,
 			const FRigidTransform3& Convex2Transform, //world
 			const FReal Dt,
-			FRigidBodyPointContactConstraint& Constraint);
+			FPBDCollisionConstraint& Constraint);
 
 		template
 		void ConstructConvexConvexOneShotManifold<FImplicitBox3, TImplicitObjectScaled<FImplicitConvex3>>(
@@ -883,7 +944,7 @@ namespace Chaos
 			const TImplicitObjectScaled<FImplicitConvex3>& Implicit2,
 			const FRigidTransform3& Convex2Transform, //world
 			const FReal Dt,
-			FRigidBodyPointContactConstraint& Constraint);
+			FPBDCollisionConstraint& Constraint);
 
 		template
 		void ConstructConvexConvexOneShotManifold<TImplicitObjectScaled<FImplicitConvex3>, FImplicitBox3>(
@@ -892,7 +953,7 @@ namespace Chaos
 			const FImplicitBox3& Implicit2,
 			const FRigidTransform3& Convex2Transform, //world
 			const FReal Dt,
-			FRigidBodyPointContactConstraint& Constraint);
+			FPBDCollisionConstraint& Constraint);
 
 		template
 		void ConstructConvexConvexOneShotManifold<FImplicitConvex3, FImplicitConvex3>(
@@ -901,7 +962,7 @@ namespace Chaos
 			const FImplicitConvex3& Implicit2,
 			const FRigidTransform3& Convex2Transform, //world
 			const FReal Dt,
-			FRigidBodyPointContactConstraint& Constraint);
+			FPBDCollisionConstraint& Constraint);
 
 		template
 		void ConstructConvexConvexOneShotManifold<TImplicitObjectInstanced<FImplicitConvex3>, FImplicitConvex3>(
@@ -910,7 +971,7 @@ namespace Chaos
 			const FImplicitConvex3& Implicit2,
 			const FRigidTransform3& Convex2Transform, //world
 			const FReal Dt,
-			FRigidBodyPointContactConstraint& Constraint);
+			FPBDCollisionConstraint& Constraint);
 
 		template
 		void ConstructConvexConvexOneShotManifold<TImplicitObjectScaled<FImplicitConvex3>, FImplicitConvex3>(
@@ -919,7 +980,7 @@ namespace Chaos
 			const FImplicitConvex3& Implicit2,
 			const FRigidTransform3& Convex2Transform, //world
 			const FReal Dt,
-			FRigidBodyPointContactConstraint& Constraint);
+			FPBDCollisionConstraint& Constraint);
 
 		template
 		void ConstructConvexConvexOneShotManifold<FImplicitConvex3, TImplicitObjectInstanced<FImplicitConvex3>>(
@@ -928,7 +989,7 @@ namespace Chaos
 			const TImplicitObjectInstanced<FImplicitConvex3>& Implicit2,
 			const FRigidTransform3& Convex2Transform, //world
 			const FReal Dt,
-			FRigidBodyPointContactConstraint& Constraint);
+			FPBDCollisionConstraint& Constraint);
 
 		template
 		void ConstructConvexConvexOneShotManifold<FImplicitConvex3, TImplicitObjectScaled<FImplicitConvex3>>(
@@ -937,7 +998,7 @@ namespace Chaos
 			const TImplicitObjectScaled<FImplicitConvex3>& Implicit2,
 			const FRigidTransform3& Convex2Transform, //world
 			const FReal Dt,
-			FRigidBodyPointContactConstraint& Constraint);
+			FPBDCollisionConstraint& Constraint);
 
 		template
 		void ConstructConvexConvexOneShotManifold<TImplicitObjectInstanced<FImplicitConvex3>, TImplicitObjectInstanced<FImplicitConvex3>>(
@@ -946,7 +1007,7 @@ namespace Chaos
 			const TImplicitObjectInstanced<FImplicitConvex3>& Implicit2,
 			const FRigidTransform3& Convex2Transform, //world
 			const FReal Dt,
-			FRigidBodyPointContactConstraint& Constraint);
+			FPBDCollisionConstraint& Constraint);
 
 		template
 		void ConstructConvexConvexOneShotManifold<TImplicitObjectScaled<FImplicitConvex3>, TImplicitObjectInstanced<FImplicitConvex3>>(
@@ -955,7 +1016,7 @@ namespace Chaos
 			const TImplicitObjectInstanced<FImplicitConvex3>& Implicit2,
 			const FRigidTransform3& Convex2Transform, //world
 			const FReal Dt,
-			FRigidBodyPointContactConstraint& Constraint);
+			FPBDCollisionConstraint& Constraint);
 
 		template
 		void ConstructConvexConvexOneShotManifold<TImplicitObjectInstanced<FImplicitConvex3>, TImplicitObjectScaled<FImplicitConvex3>>(
@@ -964,7 +1025,7 @@ namespace Chaos
 			const TImplicitObjectScaled<FImplicitConvex3>& Implicit2,
 			const FRigidTransform3& Convex2Transform, //world
 			const FReal Dt,
-			FRigidBodyPointContactConstraint& Constraint);
+			FPBDCollisionConstraint& Constraint);
 
 		template
 		void ConstructConvexConvexOneShotManifold<TImplicitObjectScaled<FImplicitConvex3>, TImplicitObjectScaled<FImplicitConvex3>>(
@@ -973,7 +1034,7 @@ namespace Chaos
 			const TImplicitObjectScaled<FImplicitConvex3>& Implicit2,
 			const FRigidTransform3& Convex2Transform, //world
 			const FReal Dt,
-			FRigidBodyPointContactConstraint& Constraint);
+			FPBDCollisionConstraint& Constraint);
 	}
 }
 
diff --git a/Engine/Source/Runtime/Experimental/Chaos/Private/Chaos/CollisionOneShotManifolds.h b/Engine/Source/Runtime/Experimental/Chaos/Private/Chaos/CollisionOneShotManifolds.h
--- a/Engine/Source/Runtime/Experimental/Chaos/Private/Chaos/CollisionOneShotManifolds.h
+++ b/Engine/Source/Runtime/Experimental/Chaos/Private/Chaos/CollisionOneShotManifolds.h
@@ -9,6 +9,7 @@
 
 namespace Chaos
 {
+	extern bool bChaos_Collision_Manifold_FixNormalsInWorldSpace;
 	namespace Collisions
 	{
 		uint32 BoxBoxClipVerticesAgainstPlane(const FVec3* InputVertexBuffer, FVec3* outputVertexBuffer, uint32 ClipPointCount, int32 ClippingAxis, FReal Distance);
@@ -20,7 +21,7 @@ namespace Chaos
 			const FImplicitBox3& Box2,
 			const FRigidTransform3& Box2Transform, //world
 			const FReal Dt,
-			FRigidBodyPointContactConstraint& Constraint);
+			FPBDCollisionConstraint& Constraint);
 
 		template <typename ConvexImplicitType1, typename ConvexImplicitType2>
 		void ConstructConvexConvexOneShotManifold(
@@ -29,6 +30,6 @@ namespace Chaos
 			const ConvexImplicitType2& Implicit2,
 			const FRigidTransform3& Convex2Transform, //world
 			const FReal Dt,
-			FRigidBodyPointContactConstraint& Constraint);
+			FPBDCollisionConstraint& Constraint);
 	}
 }
\ No newline at end of file
