Chaos

diff --git a/Engine/Source/Runtime/Experimental/Chaos/Private/Chaos/CollisionResolution.cpp b/Engine/Source/Runtime/Experimental/Chaos/Private/Chaos/CollisionResolution.cpp
--- a/Engine/Source/Runtime/Experimental/Chaos/Private/Chaos/CollisionResolution.cpp
+++ b/Engine/Source/Runtime/Experimental/Chaos/Private/Chaos/CollisionResolution.cpp
@@ -9,6 +9,7 @@
 #include "Chaos/CollisionResolutionUtil.h"
 #include "Chaos/Collision/CapsuleConvexContactPoint.h"
 #include "Chaos/Collision/CollisionContext.h"
+#include "Chaos/Collision/CollisionConstraintAllocator.h"
 #include "Chaos/Collision/PBDCollisionConstraint.h"
 #include "Chaos/Collision/SphereConvexContactPoint.h"
 #include "Chaos/Convex.h"
@@ -86,7 +87,7 @@ FAutoConsoleVariableRef  CVarCCDEnableThresholdBoundsScale(TEXT("p.Chaos.CCD.Ena
 Chaos::FRealSingle CCDAllowedDepthBoundsScale = 0.1f;
 FAutoConsoleVariableRef CVarCCDAllowedDepthBoundsScale(TEXT("p.Chaos.CCD.AllowedDepthBoundsScale"), CCDAllowedDepthBoundsScale, TEXT("When rolling back to TOI, allow (smallest bound's extent) * AllowedDepthBoundsScale, instead of rolling back to exact TOI w/ penetration = 0."));
 
-int32 CCDUseGenericSweptConvexConstraints = 0;
+int32 CCDUseGenericSweptConvexConstraints = 1;
 FAutoConsoleVariableRef CVarUseGenericSweptConvexConstraints(TEXT("p.Chaos.CCD.UseGenericSweptConvexConstraints"), CCDUseGenericSweptConvexConstraints, TEXT("Use generic convex convex swept constraint generation for convex shape pairs which don't have specialized implementations."));
 
 int32 CCDOnlyConsiderDynamicStatic = 0;
@@ -104,7 +105,7 @@ FAutoConsoleVariableRef CVarCCDAllowForceDisable(TEXT("p.Chaos.CCD.AllowForceDis
 int32 CCDManualForceDisable = 0;
 FAutoConsoleVariableRef CVarCCDManualForceDisable(TEXT("p.Chaos.CCD.ManualForceDisable"), CCDManualForceDisable, TEXT("For debugging, manually disable CCD always."));
 
-bool bChaos_Collision_AllowLevelsetManifolds = false;
+bool bChaos_Collision_AllowLevelsetManifolds = true;
 FAutoConsoleVariableRef CVarChaosCollisionAllowLevelsetManifolds(TEXT("p.Chaos.Collision.AllowLevelsetManifolds"), bChaos_Collision_AllowLevelsetManifolds, TEXT("Use incremental manifolds for levelset-levelset collision. This does not work well atm - too much rotation in the small pieces"));
 
 //Chaos::FRealSingl Chaos_Collision_ManifoldFaceAngle = 5.0f;
@@ -124,6 +125,12 @@ bool Chaos_Collision_NarrowPhase_AABBBoundsCheck = true;
 FAutoConsoleVariableRef CVarChaosCollisionSphereBoundsCheck(TEXT("p.Chaos.Collision.SphereBoundsCheck"), Chaos_Collision_NarrowPhase_SphereBoundsCheck, TEXT(""));
 FAutoConsoleVariableRef CVarChaosCollisionAABBBoundsCheck(TEXT("p.Chaos.Collision.AABBBoundsCheck"), Chaos_Collision_NarrowPhase_AABBBoundsCheck, TEXT(""));
 
+bool bChaos_Collision_EnableManifoldRestore = true;
+Chaos::FRealSingle Chaos_Collision_ParticlePositionTolerance = 0.3f;	// Particles that move less than this can reuse their manifolds
+Chaos::FRealSingle Chaos_Collision_ParticleRotationTolerance = 0.3f;	// Particles that move less than this can reuse their manifolds
+FAutoConsoleVariableRef CVarChaos_Collision_EnableManifoldRestore(TEXT("p.Chaos.Collision.EnableManifoldRestore"), bChaos_Collision_EnableManifoldRestore, TEXT(""));
+FAutoConsoleVariableRef CVarChaos_Collision_ParticlePositionTolerance(TEXT("p.Chaos.Collision.ParticlePositionTolerance"), Chaos_Collision_ParticlePositionTolerance, TEXT("Cms. Particle pairs that move less than this may have their contacts reinstated"));
+FAutoConsoleVariableRef CVarChaos_Collision_ParticleRotationTolerance(TEXT("p.Chaos.Collision.ParticleRotationTolerance"), Chaos_Collision_ParticleRotationTolerance, TEXT("Degrees. Particle pairs that move less than this may have their contacts reinstated"));
 
 namespace Chaos
 {
@@ -165,7 +172,7 @@ namespace Chaos
 		}
 
 		// Determines if body should use CCD. If using CCD, computes Dir and Length of sweep.
-		bool UseCCD(const TGeometryParticleHandle<FReal, 3>* SweptParticle, const TGeometryParticleHandle<FReal, 3>* OtherParticle, const FImplicitObject* Implicit, FVec3& Dir, FReal& Length, const bool bForceDisableCCD)
+		inline bool UseCCDImpl(const TGeometryParticleHandle<FReal, 3>* SweptParticle, const TGeometryParticleHandle<FReal, 3>* OtherParticle, const FImplicitObject* SweptImplicit, const FVec3& SweptX, const FVec3& SweptP, FVec3& Dir, FReal& Length, const bool bForceDisableCCD)
 		{
 			if (CCDManualForceDisable)
 			{
@@ -183,17 +190,17 @@ namespace Chaos
 			}
 
 			auto* RigidParticle = SweptParticle->CastToRigidParticle();
-			if (RigidParticle  && RigidParticle->CCDEnabled())
+			if (RigidParticle && RigidParticle->CCDEnabled())
 			{
-				Dir = RigidParticle->P() - RigidParticle->X();
+				Dir = SweptP - SweptX;
 				Length = Dir.Size();
 
-				if (!Implicit->HasBoundingBox())
+				if (!SweptImplicit->HasBoundingBox())
 				{
 					return false;
 				}
 
-				FReal MinBoundsAxis = Implicit->BoundingBox().Extents().Min();
+				FReal MinBoundsAxis = SweptImplicit->BoundingBox().Extents().Min();
 				FReal LengthCCDThreshold = MinBoundsAxis * CCDEnableThresholdBoundsScale;
 
 				if (Length > 0.0f && CCDEnableThresholdBoundsScale >= 0.0f && Length >= LengthCCDThreshold)
@@ -211,6 +218,25 @@ namespace Chaos
 			return false;
 		}
 
+		// UeeCCD function for use in collision detection loop (is uses Particles directly for particle state)
+		// NOTE: should only be called in ConstructConstraint (also see UseCCD_Update)
+		inline bool UseCCD_Init(const TGeometryParticleHandle<FReal, 3>* SweptParticle, const TGeometryParticleHandle<FReal, 3>* OtherParticle, const FImplicitObject* Implicit, FVec3& Dir, FReal& Length, bool bForceDisableCCD)
+		{
+			return UseCCDImpl(SweptParticle, OtherParticle, Implicit, FConstGenericParticleHandle(SweptParticle)->X(), FConstGenericParticleHandle(SweptParticle)->P(), Dir, Length, bForceDisableCCD);
+		}
+
+		// UeeCCD function for use in solver loop (is uses SolverBody for particle state)
+		// NOTE: should only be called in UpdateConstraint (also see UseCCD_Init)
+		inline bool UseCCD_Update(FPBDCollisionConstraint& Constraint, FVec3& Dir, FReal& Length, bool bForceDisableCCD)
+		{
+			// This function assumes constraint Particle[0] is always the moving one - which is also assumed at the calling site
+			if (Constraint.GetType() == ECollisionConstraintType::Swept)
+			{
+				return UseCCDImpl(Constraint.Particle[0], Constraint.Particle[1], Constraint.GetManifold().Implicit[0], Constraint.GetSolverBody0()->X(), Constraint.GetSolverBody0()->P(), Dir, Length, bForceDisableCCD);
+			}
+			return false;
+		}
+
 		void ComputeSweptContactPhiAndTOIHelper(const FVec3& ContactNormal, const FVec3& Dir, const FReal& Length, const FReal& HitTime, FReal& OutTOI, FReal& OutPhi)
 		{
 			if (HitTime >= 0.0f)
@@ -234,7 +260,7 @@ namespace Chaos
 		}
 
 		// Initializes TimeOfImpact and adjusts Phi for TOI.
-		void SetSweptConstraintTOI(TGeometryParticleHandle<FReal, 3>* Particle, const FReal TOI, const FReal Length, const FVec3& Dir, FRigidBodySweptPointContactConstraint& SweptConstraint)
+		void SetSweptConstraintTOI(TGeometryParticleHandle<FReal, 3>* Particle, const FReal TOI, const FReal Length, const FVec3& Dir, FPBDCollisionConstraint& SweptConstraint)
 		{
 			CONDITIONAL_SCOPE_CYCLE_COUNTER(STAT_Collisions_SetSweptConstraintTOI, ConstraintsDetailedStats);
 
@@ -279,11 +305,14 @@ namespace Chaos
 		}
 
 		// Add the contact to the manifold (or update the existing point). Disable the contact if the contact distance is greater than the Cull Distance
-		void UpdateContactPoint(FRigidBodyPointContactConstraint& Constraint, const FContactPoint& ContactPoint, const FReal Dt)
+		void UpdateContactPoint(FPBDCollisionConstraint& Constraint, const FContactPoint& ContactPoint, const FReal Dt)
 		{
 			// Permanently disable contacts beyond the CullDistance
+			// NOTE: We cannot do this for incremental manifolds because of a hack in the heightfield collision detection.
+			// See FHeightField::GJKContactPointImp where it uses a sweep.
 #if CHAOS_COLLISION_CREATE_BOUNDSCHECK
-			if (ContactPoint.Phi > Constraint.GetCullDistance())
+			const bool bIsIncrementalManifold = Constraint.GetUseManifold() && (Constraint.GetManifoldPoints().Num() > 0);
+			if ((ContactPoint.Phi > Constraint.GetCullDistance()) && !bIsIncrementalManifold)
 			{
 				Constraint.SetDisabled(true);
 				return;
@@ -296,10 +325,20 @@ namespace Chaos
 			{
 				Constraint.AddIncrementalManifoldContact(ContactPoint, Dt);
 			}
+			else
+			{
+				// If we did not add a new contact point, select the best existing one
+				// We only hit this because of the FHeightField::GJKContactPointImp hack mentioned above
+				// @todo(chaos): fix heightfield collision and get rid of this
+				if (Constraint.GetUseManifold())
+				{
+					Constraint.UpdateManifoldContacts();
+				}
+			}
 		}
 
 		// Same as UpdateContact Point but without checking CullDistance. Used by CCD because sweeps do not set the separation unless the sweep actually hits
-		void UpdateContactPointNoCull(FRigidBodyPointContactConstraint& Constraint, const FContactPoint& ContactPoint, const FReal Dt)
+		void UpdateContactPointNoCull(FPBDCollisionConstraint& Constraint, const FContactPoint& ContactPoint, const FReal Dt)
 		{
 			if (ContactPoint.IsSet())
 			{
@@ -359,13 +398,16 @@ namespace Chaos
 			if (GJKRaycast2(B, A, AToBTM, LocalDir, Length, OutTime, Location, Normal, (FReal)0, true))
 			{
 				// GJK output is all in the local space of B. We need to transform the B-relative position and the normal in to B-space
+				const FVec3 WorldLocation = BTM.TransformPosition(Location);
+				const FVec3 WorldNormal = BTM.TransformVectorNoScale(Normal);
+				const FRigidTransform3 ATMAtHit = FRigidTransform3(ATM.GetTranslation() + Dir * OutTime, ATM.GetRotation());	// OutTime is a distance :|
 				Contact.ShapeMargins[0] = 0.0f;
 				Contact.ShapeMargins[1] = 0.0f;
-				Contact.ShapeContactPoints[0] = AToBTM.InverseTransformPosition(Location);
+				Contact.ShapeContactPoints[0] = ATMAtHit.InverseTransformPosition(WorldLocation);
 				Contact.ShapeContactPoints[1] = Location;
-				Contact.ShapeContactNormal = -Normal;
-				Contact.Location = BTM.TransformPosition(Location);
-				Contact.Normal = BTM.TransformVectorNoScale(Normal);
+				Contact.ShapeContactNormal = Normal;
+				Contact.Location = WorldLocation;
+				Contact.Normal = WorldNormal;
 				ComputeSweptContactPhiAndTOIHelper(Contact.Normal, Dir, Length, OutTime, TOI, Contact.Phi);
 			}
 
@@ -573,7 +615,7 @@ namespace Chaos
 		struct TConvexConvexUpdater
 		{
 			template<typename T_ImplicitA, typename T_ImplicitB>
-			static void UpdateConvexConvexConstraint(const T_ImplicitA& A, const FRigidTransform3& ATM, const T_ImplicitB& B, const FRigidTransform3& BTM, const FReal Dt, const FReal ShapePadding, FRigidBodyPointContactConstraint& Constraint)
+			static void UpdateConvexConvexConstraint(const T_ImplicitA& A, const FRigidTransform3& ATM, const T_ImplicitB& B, const FRigidTransform3& BTM, const FReal Dt, const FReal ShapePadding, FPBDCollisionConstraint& Constraint)
 			{
 				UpdateContactPoint(Constraint, GJKContactPoint(A, ATM, B, BTM, FVec3(1, 0, 0), ShapePadding), Dt);
 			}
@@ -584,7 +626,7 @@ namespace Chaos
 		struct TConvexConvexUpdater<true>
 		{
 			template<typename T_ImplicitA, typename T_ImplicitB>
-			static void UpdateConvexConvexConstraint(const T_ImplicitA& A, const FRigidTransform3& ATM, const T_ImplicitB& B, const FRigidTransform3& BTM, const FReal Dt, const FReal ShapePadding, FRigidBodyPointContactConstraint& Constraint)
+			static void UpdateConvexConvexConstraint(const T_ImplicitA& A, const FRigidTransform3& ATM, const T_ImplicitB& B, const FRigidTransform3& BTM, const FReal Dt, const FReal ShapePadding, FPBDCollisionConstraint& Constraint)
 			{
 				if (Constraint.GetUseManifold())
 				{
@@ -595,7 +637,8 @@ namespace Chaos
 					}
 					else
 					{
-						Constraint.UpdateManifoldContacts(Dt);
+						// @todo(chaos): this is only used for the MinEvolution path - see if it can be removed/moved
+						Constraint.UpdateManifoldContacts();
 					}
 				}
 				else
@@ -610,7 +653,7 @@ namespace Chaos
 		struct FConvexConvexUpdaterCaller
 		{
 			template<typename T_ImplicitA, typename T_ImplicitB>
-			static void Update(const T_ImplicitA& A, const FRigidTransform3& ATM, const T_ImplicitB& B, const FRigidTransform3& BTM, const FReal Dt, const FReal ShapePadding, FRigidBodyPointContactConstraint& Constraint)
+			static void Update(const T_ImplicitA& A, const FRigidTransform3& ATM, const T_ImplicitB& B, const FRigidTransform3& BTM, const FReal Dt, const FReal ShapePadding, FPBDCollisionConstraint& Constraint)
 			{
 				using FConvexConvexUpdater = TConvexConvexUpdater<TConvexImplicitPairTraits<T_ImplicitA, T_ImplicitB>::bSupportsOneShotManifold>;
 				FConvexConvexUpdater::UpdateConvexConvexConstraint(A, ATM, B, BTM, Dt, ShapePadding, Constraint);
@@ -619,7 +662,7 @@ namespace Chaos
 
 
 		// Unwrap the many convex types, including scaled, and call the appropriate update which depends on the concrete types
-		void UpdateGenericConvexConvexConstraintHelper(const FImplicitObject& A, const FRigidTransform3& ATM, const FImplicitObject& B, const FRigidTransform3& BTM, const FReal Dt, const FReal ShapePadding, FRigidBodyPointContactConstraint& Constraint)
+		void UpdateGenericConvexConvexConstraintHelper(const FImplicitObject& A, const FRigidTransform3& ATM, const FImplicitObject& B, const FRigidTransform3& BTM, const FReal Dt, const FReal ShapePadding, FPBDCollisionConstraint& Constraint)
 		{
 			// This expands to a switch of switches that calls the inner function with the appropriate concrete implicit types
 			Utilities::CastHelperNoUnwrap(A, ATM, [&](const auto& ADowncast, const FRigidTransform3& AFullTM)
@@ -678,7 +721,7 @@ namespace Chaos
 			return GJKContactPoint(Box1, Box1TM, Box2, Box2TM, FVec3(1, 0, 0), ShapePadding);
 		}
 
-		void UpdateBoxBoxConstraint(const FImplicitBox3& Box1, const FRigidTransform3& Box1Transform, const FImplicitBox3& Box2, const FRigidTransform3& Box2Transform, const FReal Dt, FRigidBodyPointContactConstraint& Constraint)
+		void UpdateBoxBoxConstraint(const FImplicitBox3& Box1, const FRigidTransform3& Box1Transform, const FImplicitBox3& Box2, const FRigidTransform3& Box2Transform, const FReal Dt, FPBDCollisionConstraint& Constraint)
 		{
 			if (Constraint.GetUseManifold())
 			{
@@ -689,7 +732,7 @@ namespace Chaos
 				}
 				else
 				{
-					Constraint.UpdateManifoldContacts(Dt);
+					Constraint.UpdateManifoldContacts();
 				}
 			}
 			else
@@ -699,20 +742,19 @@ namespace Chaos
 		}
 
 		template<typename T_TRAITS>
-		void ConstructBoxBoxConstraints(TGeometryParticleHandle<FReal, 3>* Particle0, TGeometryParticleHandle<FReal, 3>* Particle1, const FImplicitObject* Implicit0, const FImplicitObject* Implicit1, const FRigidTransform3& LocalTransform0, const FRigidTransform3& LocalTransform1, const FReal CullDistance, const FReal Dt, const FCollisionContext& Context, FCollisionConstraintsArray& NewConstraints)
+		void ConstructBoxBoxConstraints(TGeometryParticleHandle<FReal, 3>* Particle0, TGeometryParticleHandle<FReal, 3>* Particle1, const FImplicitObject* Implicit0, const FImplicitObject* Implicit1, const FRigidTransform3& ParticleWorldTransform0, const FRigidTransform3& LocalTransform0, const FRigidTransform3& ParticleWorldTransform1, const FRigidTransform3& LocalTransform1, const FReal CullDistance, const FReal Dt, const FCollisionContext& Context)
 		{
 			const TBox<FReal, 3>* Object0 = Implicit0->template GetObject<const TBox<FReal, 3> >();
 			const TBox<FReal, 3>* Object1 = Implicit1->template GetObject<const TBox<FReal, 3> >();
 			if (ensure(Object0 && Object1))
 			{
-				FRigidBodyPointContactConstraint Constraint = FRigidBodyPointContactConstraint(Particle0, Implicit0, nullptr, LocalTransform0, Particle1, Implicit1, nullptr, LocalTransform1, CullDistance, EContactShapesType::BoxBox, Context.bAllowManifolds);
-				if (T_TRAITS::bImmediateUpdate)
+				FPBDCollisionConstraint* Constraint = FPBDCollisionConstraint::Make(Particle0, Implicit0, nullptr, ParticleWorldTransform0, LocalTransform0, Particle1, Implicit1, nullptr, ParticleWorldTransform1, LocalTransform1, CullDistance, EContactShapesType::BoxBox, Context.bAllowManifolds, *Context.CollisionAllocator);
+				if (T_TRAITS::bImmediateUpdate && (Constraint != nullptr) && !Constraint->WasManifoldRestored())
 				{
-					FRigidTransform3 WorldTransform0 = LocalTransform0 * Collisions::GetTransform(Particle0);
-					FRigidTransform3 WorldTransform1 = LocalTransform1 * Collisions::GetTransform(Particle1);
-					UpdateBoxBoxConstraint(*Object0, WorldTransform0, *Object1, WorldTransform1, Dt, Constraint);
+					const FRigidTransform3 WorldTransform0 = LocalTransform0 * ParticleWorldTransform0;
+					const FRigidTransform3 WorldTransform1 = LocalTransform1 * ParticleWorldTransform1;
+					UpdateBoxBoxConstraint(*Object0, WorldTransform0, *Object1, WorldTransform1, Dt, *Constraint);
 				}
-				NewConstraints.Add(Constraint);
 			}
 		}
 
@@ -727,14 +769,14 @@ namespace Chaos
 		}
 
 
-		void UpdateBoxHeightFieldConstraint(const FImplicitBox3& A, const FRigidTransform3& ATransform, const FHeightField& B, const FRigidTransform3& BTransform, const FReal Dt, FRigidBodyPointContactConstraint& Constraint)
+		void UpdateBoxHeightFieldConstraint(const FImplicitBox3& A, const FRigidTransform3& ATransform, const FHeightField& B, const FRigidTransform3& BTransform, const FReal Dt, FPBDCollisionConstraint& Constraint)
 		{
 			// @todo(chaos): restitutionpadding
 			UpdateContactPoint(Constraint, BoxHeightFieldContactPoint(A, ATransform, B, BTransform, Constraint.GetCullDistance(), 0.0f), Dt);
 		}
 
 		template<typename T_TRAITS>
-		void ConstructBoxHeightFieldConstraints(TGeometryParticleHandle<FReal, 3>* Particle0, TGeometryParticleHandle<FReal, 3>* Particle1, const FImplicitObject* Implicit0, const FImplicitObject* Implicit1, const FRigidTransform3& LocalTransform0, const FRigidTransform3& LocalTransform1, const FReal CullDistance, const FReal Dt, const FCollisionContext& Context, FCollisionConstraintsArray& NewConstraints)
+		void ConstructBoxHeightFieldConstraints(TGeometryParticleHandle<FReal, 3>* Particle0, TGeometryParticleHandle<FReal, 3>* Particle1, const FImplicitObject* Implicit0, const FImplicitObject* Implicit1, const FRigidTransform3& ParticleWorldTransform0, const FRigidTransform3& LocalTransform0, const FRigidTransform3& ParticleWorldTransform1, const FRigidTransform3& LocalTransform1, const FReal CullDistance, const FReal Dt, const FCollisionContext& Context)
 		{
 
 			const TBox<FReal, 3>* Object0 = Implicit0->template GetObject<const TBox<FReal, 3> >();
@@ -742,14 +784,13 @@ namespace Chaos
 			if (ensure(Object0 && Object1))
 			{
 				// @todo(chaos): one-shot manifold
-				FRigidBodyPointContactConstraint Constraint = FRigidBodyPointContactConstraint(Particle0, Implicit0, nullptr, LocalTransform0, Particle1, Implicit1, nullptr, LocalTransform1, CullDistance, EContactShapesType::BoxHeightField, Context.bAllowManifolds);
-				if (T_TRAITS::bImmediateUpdate)
+				FPBDCollisionConstraint* Constraint = FPBDCollisionConstraint::Make(Particle0, Implicit0, nullptr, ParticleWorldTransform0, LocalTransform0, Particle1, Implicit1, nullptr, ParticleWorldTransform1, LocalTransform1, CullDistance, EContactShapesType::BoxHeightField, Context.bAllowManifolds, *Context.CollisionAllocator);
+				if (T_TRAITS::bImmediateUpdate && (Constraint != nullptr) && !Constraint->WasManifoldRestored())
 				{
-					FRigidTransform3 WorldTransform0 = LocalTransform0 * Collisions::GetTransform(Particle0);
-					FRigidTransform3 WorldTransform1 = LocalTransform1 * Collisions::GetTransform(Particle1);
-					UpdateBoxHeightFieldConstraint(*Object0, WorldTransform0, *Object1, WorldTransform1, Dt, Constraint);
+					FRigidTransform3 WorldTransform0 = LocalTransform0 * ParticleWorldTransform0;
+					FRigidTransform3 WorldTransform1 = LocalTransform1 * ParticleWorldTransform1;
+					UpdateBoxHeightFieldConstraint(*Object0, WorldTransform0, *Object1, WorldTransform1, Dt, *Constraint);
 				}
-				NewConstraints.Add(Constraint);
 			}
 		}
 
@@ -824,7 +865,7 @@ namespace Chaos
 			return Contact;
 		}
 
-		void UpdateBoxPlaneConstraint(const FImplicitBox3& A, const FRigidTransform3& ATransform, const TPlane<FReal, 3>& B, const FRigidTransform3& BTransform, const FReal Dt, FRigidBodyPointContactConstraint& Constraint)
+		void UpdateBoxPlaneConstraint(const FImplicitBox3& A, const FRigidTransform3& ATransform, const TPlane<FReal, 3>& B, const FRigidTransform3& BTransform, const FReal Dt, FPBDCollisionConstraint& Constraint)
 		{
 			// @todo(chaos): one-shot manifold
 			// @todo(chaos): restitutionpadding
@@ -833,20 +874,19 @@ namespace Chaos
 
 
 		template<typename T_TRAITS>
-		void ConstructBoxPlaneConstraints(TGeometryParticleHandle<FReal, 3>* Particle0, TGeometryParticleHandle<FReal, 3>* Particle1, const FImplicitObject* Implicit0, const FImplicitObject* Implicit1, const FRigidTransform3& LocalTransform0, const FRigidTransform3& LocalTransform1, const FReal CullDistance, const FReal Dt, const FCollisionContext& Context, FCollisionConstraintsArray& NewConstraints)
+		void ConstructBoxPlaneConstraints(TGeometryParticleHandle<FReal, 3>* Particle0, TGeometryParticleHandle<FReal, 3>* Particle1, const FImplicitObject* Implicit0, const FImplicitObject* Implicit1, const FRigidTransform3& ParticleWorldTransform0, const FRigidTransform3& LocalTransform0, const FRigidTransform3& ParticleWorldTransform1, const FRigidTransform3& LocalTransform1, const FReal CullDistance, const FReal Dt, const FCollisionContext& Context)
 		{
 			const TBox<FReal, 3>* Object0 = Implicit0->template GetObject<const TBox<FReal, 3> >();
 			const TPlane<FReal, 3>* Object1 = Implicit1->template GetObject<const TPlane<FReal, 3> >();
 			if (ensure(Object0 && Object1))
 			{
-				FRigidBodyPointContactConstraint Constraint = FRigidBodyPointContactConstraint(Particle0, Implicit0, nullptr, LocalTransform0, Particle1, Implicit1, nullptr, LocalTransform1, CullDistance, EContactShapesType::BoxPlane, Context.bAllowManifolds);
-				if (T_TRAITS::bImmediateUpdate)
+				FPBDCollisionConstraint* Constraint = FPBDCollisionConstraint::Make(Particle0, Implicit0, nullptr, ParticleWorldTransform0, LocalTransform0, Particle1, Implicit1, nullptr, ParticleWorldTransform1, LocalTransform1, CullDistance, EContactShapesType::BoxPlane, Context.bAllowManifolds, *Context.CollisionAllocator);
+				if (T_TRAITS::bImmediateUpdate && (Constraint != nullptr) && !Constraint->WasManifoldRestored())
 				{
-					FRigidTransform3 WorldTransform0 = LocalTransform0 * Collisions::GetTransform(Particle0);
-					FRigidTransform3 WorldTransform1 = LocalTransform1 * Collisions::GetTransform(Particle1);
-					UpdateBoxPlaneConstraint(*Object0, WorldTransform0, *Object1, WorldTransform1, Dt, Constraint);
+					FRigidTransform3 WorldTransform0 = LocalTransform0 * ParticleWorldTransform0;
+					FRigidTransform3 WorldTransform1 = LocalTransform1 * ParticleWorldTransform1;
+					UpdateBoxPlaneConstraint(*Object0, WorldTransform0, *Object1, WorldTransform1, Dt, *Constraint);
 				}
-				NewConstraints.Add(Constraint);
 			}
 		}
 
@@ -861,14 +901,14 @@ namespace Chaos
 		}
 
 		template <typename TriMeshType>
-		void UpdateBoxTriangleMeshConstraint(const FImplicitBox3& Box0, const FRigidTransform3& WorldTransform0, const TriMeshType& TriangleMesh1, const FRigidTransform3& WorldTransform1, const FReal Dt, FRigidBodyPointContactConstraint& Constraint)
+		void UpdateBoxTriangleMeshConstraint(const FImplicitBox3& Box0, const FRigidTransform3& WorldTransform0, const TriMeshType& TriangleMesh1, const FRigidTransform3& WorldTransform1, const FReal Dt, FPBDCollisionConstraint& Constraint)
 		{
 			// @toto(chaos): restitutionpadding
 			UpdateContactPoint(Constraint, BoxTriangleMeshContactPoint(Box0, WorldTransform0, TriangleMesh1, WorldTransform1, Constraint.GetCullDistance(), 0.0f), Dt);
 		}
 
 		template<typename T_TRAITS>
-		void ConstructBoxTriangleMeshConstraints(TGeometryParticleHandle<FReal, 3>* Particle0, TGeometryParticleHandle<FReal, 3>* Particle1, const FImplicitObject* Implicit0, const FImplicitObject* Implicit1, const FRigidTransform3& LocalTransform0, const FRigidTransform3& LocalTransform1, const FReal CullDistance, const FReal Dt, const FCollisionContext& Context, FCollisionConstraintsArray& NewConstraints)
+		void ConstructBoxTriangleMeshConstraints(TGeometryParticleHandle<FReal, 3>* Particle0, TGeometryParticleHandle<FReal, 3>* Particle1, const FImplicitObject* Implicit0, const FImplicitObject* Implicit1, const FRigidTransform3& ParticleWorldTransform0, const FRigidTransform3& LocalTransform0, const FRigidTransform3& ParticleWorldTransform1, const FRigidTransform3& LocalTransform1, const FReal CullDistance, const FReal Dt, const FCollisionContext& Context)
 		{
 			const TBox<FReal, 3>* Object0 = Implicit0->template GetObject<const TBox<FReal, 3> >();
 			if (ensure(Object0))
@@ -876,26 +916,24 @@ namespace Chaos
 				if (const TImplicitObjectScaled<FTriangleMeshImplicitObject>* ScaledTriangleMesh = Implicit1->template GetObject<const TImplicitObjectScaled<FTriangleMeshImplicitObject>>())
 				{
 					// @todo(chaos): one-shot manifold
-					FRigidBodyPointContactConstraint Constraint = FRigidBodyPointContactConstraint(Particle0, Implicit0, nullptr, LocalTransform0, Particle1, Implicit1, nullptr, LocalTransform1, CullDistance, EContactShapesType::BoxTriMesh, Context.bAllowManifolds);
-					if (T_TRAITS::bImmediateUpdate)
+					FPBDCollisionConstraint* Constraint = FPBDCollisionConstraint::Make(Particle0, Implicit0, nullptr, ParticleWorldTransform0, LocalTransform0, Particle1, Implicit1, nullptr, ParticleWorldTransform1, LocalTransform1, CullDistance, EContactShapesType::BoxTriMesh, Context.bAllowManifolds, *Context.CollisionAllocator);
+					if (T_TRAITS::bImmediateUpdate && (Constraint != nullptr) && !Constraint->WasManifoldRestored())
 					{
-						FRigidTransform3 WorldTransform0 = LocalTransform0 * Collisions::GetTransform(Particle0);
-						FRigidTransform3 WorldTransform1 = LocalTransform1 * Collisions::GetTransform(Particle1);
-						UpdateBoxTriangleMeshConstraint(*Object0, WorldTransform0, *ScaledTriangleMesh, WorldTransform1, Dt, Constraint);
+						FRigidTransform3 WorldTransform0 = LocalTransform0 * ParticleWorldTransform0;
+						FRigidTransform3 WorldTransform1 = LocalTransform1 * ParticleWorldTransform1;
+						UpdateBoxTriangleMeshConstraint(*Object0, WorldTransform0, *ScaledTriangleMesh, WorldTransform1, Dt, *Constraint);
 					}
-					NewConstraints.Add(Constraint);
 				}
 				else if (const FTriangleMeshImplicitObject* TriangleMesh = Implicit1->template GetObject<const FTriangleMeshImplicitObject>())
 				{
 					// @todo(chaos): one-shot manifold
-					FRigidBodyPointContactConstraint Constraint = FRigidBodyPointContactConstraint(Particle0, Implicit0, nullptr, LocalTransform0, Particle1, Implicit1, nullptr, LocalTransform1, CullDistance, EContactShapesType::BoxTriMesh, Context.bAllowManifolds);
-					if (T_TRAITS::bImmediateUpdate)
+					FPBDCollisionConstraint* Constraint = FPBDCollisionConstraint::Make(Particle0, Implicit0, nullptr, ParticleWorldTransform0, LocalTransform0, Particle1, Implicit1, nullptr, ParticleWorldTransform1, LocalTransform1, CullDistance, EContactShapesType::BoxTriMesh, Context.bAllowManifolds, *Context.CollisionAllocator);
+					if (T_TRAITS::bImmediateUpdate && (Constraint != nullptr) && !Constraint->WasManifoldRestored())
 					{
-						FRigidTransform3 WorldTransform0 = LocalTransform0 * Collisions::GetTransform(Particle0);
-						FRigidTransform3 WorldTransform1 = LocalTransform1 * Collisions::GetTransform(Particle1);
-						UpdateBoxTriangleMeshConstraint(*Object0, WorldTransform0, *TriangleMesh, WorldTransform1, Dt, Constraint);
+						FRigidTransform3 WorldTransform0 = LocalTransform0 * ParticleWorldTransform0;
+						FRigidTransform3 WorldTransform1 = LocalTransform1 * ParticleWorldTransform1;
+						UpdateBoxTriangleMeshConstraint(*Object0, WorldTransform0, *TriangleMesh, WorldTransform1, Dt, *Constraint);
 					}
-					NewConstraints.Add(Constraint);
 				}
 				else
 				{
@@ -938,26 +976,25 @@ namespace Chaos
 			return Result;
 		}
 
-		void UpdateSphereSphereConstraint(const TSphere<FReal, 3>& Sphere1, const FRigidTransform3& Sphere1Transform, const TSphere<FReal, 3>& Sphere2, const FRigidTransform3& Sphere2Transform, const FReal Dt, FRigidBodyPointContactConstraint& Constraint)
+		void UpdateSphereSphereConstraint(const TSphere<FReal, 3>& Sphere1, const FRigidTransform3& Sphere1Transform, const TSphere<FReal, 3>& Sphere2, const FRigidTransform3& Sphere2Transform, const FReal Dt, FPBDCollisionConstraint& Constraint)
 		{
 			UpdateContactPoint(Constraint, SphereSphereContactPoint(Sphere1, Sphere1Transform, Sphere2, Sphere2Transform, Constraint.Manifold.RestitutionPadding), Dt);
 		}
 
 		template<typename T_TRAITS>
-		void ConstructSphereSphereConstraints(TGeometryParticleHandle<FReal, 3>* Particle0, TGeometryParticleHandle<FReal, 3>* Particle1, const FImplicitObject* Implicit0, const FImplicitObject* Implicit1, const FRigidTransform3& LocalTransform0, const FRigidTransform3& LocalTransform1, const FReal CullDistance, const FReal Dt, const FCollisionContext& Context, FCollisionConstraintsArray& NewConstraints)
+		void ConstructSphereSphereConstraints(TGeometryParticleHandle<FReal, 3>* Particle0, TGeometryParticleHandle<FReal, 3>* Particle1, const FImplicitObject* Implicit0, const FImplicitObject* Implicit1, const FRigidTransform3& ParticleWorldTransform0, const FRigidTransform3& LocalTransform0, const FRigidTransform3& ParticleWorldTransform1, const FRigidTransform3& LocalTransform1, const FReal CullDistance, const FReal Dt, const FCollisionContext& Context)
 		{
 			const TSphere<FReal, 3>* Object0 = Implicit0->template GetObject<const TSphere<FReal, 3> >();
 			const TSphere<FReal, 3>* Object1 = Implicit1->template GetObject<const TSphere<FReal, 3> >();
 			if (ensure(Object0 && Object1))
 			{
-				FRigidBodyPointContactConstraint Constraint = FRigidBodyPointContactConstraint(Particle0, Implicit0, nullptr, LocalTransform0, Particle1, Implicit1, nullptr, LocalTransform1, CullDistance, EContactShapesType::SphereSphere, Context.bAllowManifolds);
-				if (T_TRAITS::bImmediateUpdate)
+				FPBDCollisionConstraint* Constraint = FPBDCollisionConstraint::Make(Particle0, Implicit0, nullptr, ParticleWorldTransform0, LocalTransform0, Particle1, Implicit1, nullptr, ParticleWorldTransform1, LocalTransform1, CullDistance, EContactShapesType::SphereSphere, Context.bAllowManifolds, *Context.CollisionAllocator);
+				if (T_TRAITS::bImmediateUpdate && (Constraint != nullptr) && !Constraint->WasManifoldRestored())
 				{
-					FRigidTransform3 WorldTransform0 = LocalTransform0 * Collisions::GetTransform(Particle0);
-					FRigidTransform3 WorldTransform1 = LocalTransform1 * Collisions::GetTransform(Particle1);
-					UpdateSphereSphereConstraint(*Object0, WorldTransform0, *Object1, WorldTransform1, Dt, Constraint);
+					FRigidTransform3 WorldTransform0 = LocalTransform0 * ParticleWorldTransform0;
+					FRigidTransform3 WorldTransform1 = LocalTransform1 * ParticleWorldTransform1;
+					UpdateSphereSphereConstraint(*Object0, WorldTransform0, *Object1, WorldTransform1, Dt, *Constraint);
 				}
-				NewConstraints.Add(Constraint);
 			}
 		}
 
@@ -971,13 +1008,13 @@ namespace Chaos
 		}
 
 
-		void UpdateSphereHeightFieldConstraint(const TSphere<FReal, 3>& A, const FRigidTransform3& ATransform, const FHeightField& B, const FRigidTransform3& BTransform, const FReal Dt, FRigidBodyPointContactConstraint& Constraint)
+		void UpdateSphereHeightFieldConstraint(const TSphere<FReal, 3>& A, const FRigidTransform3& ATransform, const FHeightField& B, const FRigidTransform3& BTransform, const FReal Dt, FPBDCollisionConstraint& Constraint)
 		{
 			// @todo(chaos): restitutionpadding
 			UpdateContactPoint(Constraint, SphereHeightFieldContactPoint(A, ATransform, B, BTransform, Constraint.GetCullDistance(), 0.0f), Dt);
 		}
 
-		void UpdateSphereHeightFieldConstraintSwept(TGeometryParticleHandle<FReal, 3>* Particle0, const TSphere<FReal, 3>& A, const FRigidTransform3& ATransform, const FHeightField& B, const FRigidTransform3& BTransform, const FVec3& Dir, const FReal Length, const FReal Dt, FRigidBodySweptPointContactConstraint& Constraint)
+		void UpdateSphereHeightFieldConstraintSwept(TGeometryParticleHandle<FReal, 3>* Particle0, const TSphere<FReal, 3>& A, const FRigidTransform3& ATransform, const FHeightField& B, const FRigidTransform3& BTransform, const FVec3& Dir, const FReal Length, const FReal Dt, FPBDCollisionConstraint& Constraint)
 		{
 			FReal TOI = 1.0f;
 			UpdateContactPoint(Constraint, GJKImplicitSweptContactPoint<TSphere<FReal, 3>>(A, ATransform, B, BTransform, Dir, Length, TOI), Dt);
@@ -985,35 +1022,33 @@ namespace Chaos
 		}
 
 		template<typename T_TRAITS>
-		void ConstructSphereHeightFieldConstraints(TGeometryParticleHandle<FReal, 3>* Particle0, TGeometryParticleHandle<FReal, 3>* Particle1, const FImplicitObject* Implicit0, const FImplicitObject* Implicit1, const FRigidTransform3& LocalTransform0, const FRigidTransform3& LocalTransform1, const FReal CullDistance, const FReal Dt, const FCollisionContext& Context, FCollisionConstraintsArray& NewConstraints)
+		void ConstructSphereHeightFieldConstraints(TGeometryParticleHandle<FReal, 3>* Particle0, TGeometryParticleHandle<FReal, 3>* Particle1, const FImplicitObject* Implicit0, const FImplicitObject* Implicit1, const FRigidTransform3& ParticleWorldTransform0, const FRigidTransform3& LocalTransform0, const FRigidTransform3& ParticleWorldTransform1, const FRigidTransform3& LocalTransform1, const FReal CullDistance, const FReal Dt, const FCollisionContext& Context)
 		{
 
 			const TSphere<FReal, 3>* Object0 = Implicit0->template GetObject<const TSphere<FReal, 3> >();
 			const FHeightField* Object1 = Implicit1->template GetObject<const FHeightField >();
 			if (ensure(Object0 && Object1))
 			{
-				FRigidBodyPointContactConstraint Constraint = FRigidBodyPointContactConstraint(Particle0, Implicit0, nullptr, LocalTransform0, Particle1, Implicit1, nullptr, LocalTransform1, CullDistance, EContactShapesType::SphereHeightField, Context.bAllowManifolds);
-				if (T_TRAITS::bImmediateUpdate)
+				FPBDCollisionConstraint* Constraint = FPBDCollisionConstraint::Make(Particle0, Implicit0, nullptr, ParticleWorldTransform0, LocalTransform0, Particle1, Implicit1, nullptr, ParticleWorldTransform1, LocalTransform1, CullDistance, EContactShapesType::SphereHeightField, Context.bAllowManifolds, *Context.CollisionAllocator);
+				if (T_TRAITS::bImmediateUpdate && (Constraint != nullptr) && !Constraint->WasManifoldRestored())
 				{
-					FRigidTransform3 WorldTransform0 = LocalTransform0 * Collisions::GetTransform(Particle0);
-					FRigidTransform3 WorldTransform1 = LocalTransform1 * Collisions::GetTransform(Particle1);
-					UpdateSphereHeightFieldConstraint(*Object0, WorldTransform0, *Object1, WorldTransform1, Dt, Constraint);
+					FRigidTransform3 WorldTransform0 = LocalTransform0 * ParticleWorldTransform0;
+					FRigidTransform3 WorldTransform1 = LocalTransform1 * ParticleWorldTransform1;
+					UpdateSphereHeightFieldConstraint(*Object0, WorldTransform0, *Object1, WorldTransform1, Dt, *Constraint);
 				}
-				NewConstraints.Add(Constraint);
 			}
 		}
 
-		void ConstructSphereHeightFieldConstraintsSwept(TGeometryParticleHandle<FReal, 3>* Particle0, TGeometryParticleHandle<FReal, 3>* Particle1, const FImplicitObject* Implicit0, const FImplicitObject* Implicit1, const FRigidTransform3& LocalTransform0, const FRigidTransform3& LocalTransform1, const FReal CullDistance, const FReal Dt, const FVec3& Dir, FReal Length, FCollisionConstraintsArray& NewConstraints)
+		void ConstructSphereHeightFieldConstraintsSwept(TGeometryParticleHandle<FReal, 3>* Particle0, TGeometryParticleHandle<FReal, 3>* Particle1, const FImplicitObject* Implicit0, const FImplicitObject* Implicit1, const FRigidTransform3& ParticleWorldTransform0, const FRigidTransform3& LocalTransform0, const FRigidTransform3& ParticleWorldTransform1, const FRigidTransform3& LocalTransform1, const FReal CullDistance, const FReal Dt, const FVec3& Dir, FReal Length, const FCollisionContext& Context)
 		{
 			const TSphere<FReal, 3>* Object0 = Implicit0->template GetObject<const TSphere<FReal, 3>>();
 			const FHeightField* Object1 = Implicit1->template GetObject<const FHeightField>();
 			if (ensure(Object0 && Object1))
 			{
-				FRigidBodySweptPointContactConstraint Constraint = FRigidBodySweptPointContactConstraint(Particle0, Implicit0, nullptr, LocalTransform0, Particle1, Implicit1, nullptr, LocalTransform1, CullDistance, EContactShapesType::SphereHeightField);
+				FPBDCollisionConstraint* Constraint = FPBDCollisionConstraint::MakeSwept(Particle0, Implicit0, nullptr, ParticleWorldTransform0, LocalTransform0, Particle1, Implicit1, nullptr, ParticleWorldTransform1, LocalTransform1, CullDistance, EContactShapesType::SphereHeightField, *Context.CollisionAllocator);
 				FRigidTransform3 WorldTransformX0 = LocalTransform0 * FRigidTransform3(Particle0->X(), Particle0->R());
-				FRigidTransform3 WorldTransform1 = LocalTransform1 * Collisions::GetTransform(Particle1);
-				UpdateSphereHeightFieldConstraintSwept(Particle0, *Object0, WorldTransformX0, *Object1, WorldTransform1, Dir, Length, Dt, Constraint);
-				NewConstraints.Add(Constraint);
+				FRigidTransform3 WorldTransform1 = LocalTransform1 * ParticleWorldTransform1;
+				UpdateSphereHeightFieldConstraintSwept(Particle0, *Object0, WorldTransformX0, *Object1, WorldTransform1, Dir, Length, Dt, *Constraint);
 			}
 		}
 
@@ -1048,27 +1083,26 @@ namespace Chaos
 			return Result;
 		}
 
-		void UpdateSpherePlaneConstraint(const TSphere<FReal, 3>& Sphere, const FRigidTransform3& SphereTransform, const TPlane<FReal, 3>& Plane, const FRigidTransform3& PlaneTransform, const FReal Dt, FRigidBodyPointContactConstraint& Constraint)
+		void UpdateSpherePlaneConstraint(const TSphere<FReal, 3>& Sphere, const FRigidTransform3& SphereTransform, const TPlane<FReal, 3>& Plane, const FRigidTransform3& PlaneTransform, const FReal Dt, FPBDCollisionConstraint& Constraint)
 		{
 			UpdateContactPoint(Constraint, SpherePlaneContactPoint(Sphere, SphereTransform, Plane, PlaneTransform, Constraint.Manifold.RestitutionPadding), Dt);
 		}
 
 		template<typename T_TRAITS>
-		void ConstructSpherePlaneConstraints(TGeometryParticleHandle<FReal, 3>* Particle0, TGeometryParticleHandle<FReal, 3>* Particle1, const FImplicitObject* Implicit0, const FImplicitObject* Implicit1, const FRigidTransform3& LocalTransform0, const FRigidTransform3& LocalTransform1, const FReal CullDistance, const FReal Dt, const FCollisionContext& Context, FCollisionConstraintsArray& NewConstraints)
+		void ConstructSpherePlaneConstraints(TGeometryParticleHandle<FReal, 3>* Particle0, TGeometryParticleHandle<FReal, 3>* Particle1, const FImplicitObject* Implicit0, const FImplicitObject* Implicit1, const FRigidTransform3& ParticleWorldTransform0, const FRigidTransform3& LocalTransform0, const FRigidTransform3& ParticleWorldTransform1, const FRigidTransform3& LocalTransform1, const FReal CullDistance, const FReal Dt, const FCollisionContext& Context)
 		{
 
 			const TSphere<FReal, 3>* Object0 = Implicit0->template GetObject<const TSphere<FReal, 3> >();
 			const TPlane<FReal, 3>* Object1 = Implicit1->template GetObject<const TPlane<FReal, 3> >();
 			if (ensure(Object0 && Object1))
 			{
-				FRigidBodyPointContactConstraint Constraint = FRigidBodyPointContactConstraint(Particle0, Implicit0, nullptr, LocalTransform0, Particle1, Implicit1, nullptr, LocalTransform1, CullDistance, EContactShapesType::SpherePlane, Context.bAllowManifolds);
-				if (T_TRAITS::bImmediateUpdate)
+				FPBDCollisionConstraint* Constraint = FPBDCollisionConstraint::Make(Particle0, Implicit0, nullptr, ParticleWorldTransform0, LocalTransform0, Particle1, Implicit1, nullptr, ParticleWorldTransform1, LocalTransform1, CullDistance, EContactShapesType::SpherePlane, Context.bAllowManifolds, *Context.CollisionAllocator);
+				if (T_TRAITS::bImmediateUpdate && (Constraint != nullptr) && !Constraint->WasManifoldRestored())
 				{
-					FRigidTransform3 WorldTransform0 = LocalTransform0 * Collisions::GetTransform(Particle0);
-					FRigidTransform3 WorldTransform1 = LocalTransform1 * Collisions::GetTransform(Particle1);
-					UpdateSpherePlaneConstraint(*Object0, WorldTransform0, *Object1, WorldTransform1, Dt, Constraint);
+					FRigidTransform3 WorldTransform0 = LocalTransform0 * ParticleWorldTransform0;
+					FRigidTransform3 WorldTransform1 = LocalTransform1 * ParticleWorldTransform1;
+					UpdateSpherePlaneConstraint(*Object0, WorldTransform0, *Object1, WorldTransform1, Dt, *Constraint);
 				}
-				NewConstraints.Add(Constraint);
 			}
 		}
 
@@ -1104,27 +1138,26 @@ namespace Chaos
 		}
 
 
-		void UpdateSphereBoxConstraint(const TSphere<FReal, 3>& Sphere, const FRigidTransform3& SphereTransform, const FImplicitBox3& Box, const FRigidTransform3& BoxTransform, const FReal Dt, FRigidBodyPointContactConstraint& Constraint)
+		void UpdateSphereBoxConstraint(const TSphere<FReal, 3>& Sphere, const FRigidTransform3& SphereTransform, const FImplicitBox3& Box, const FRigidTransform3& BoxTransform, const FReal Dt, FPBDCollisionConstraint& Constraint)
 		{
 			UpdateContactPoint(Constraint, SphereBoxContactPoint(Sphere, SphereTransform, Box, BoxTransform, Constraint.Manifold.RestitutionPadding), Dt);
 		}
 
 		template<typename T_TRAITS>
-		void ConstructSphereBoxConstraints(TGeometryParticleHandle<FReal, 3>* Particle0, TGeometryParticleHandle<FReal, 3>* Particle1, const FImplicitObject* Implicit0, const FImplicitObject* Implicit1, const FRigidTransform3& LocalTransform0, const FRigidTransform3& LocalTransform1, const FReal CullDistance, const FReal Dt, const FCollisionContext& Context, FCollisionConstraintsArray& NewConstraints)
+		void ConstructSphereBoxConstraints(TGeometryParticleHandle<FReal, 3>* Particle0, TGeometryParticleHandle<FReal, 3>* Particle1, const FImplicitObject* Implicit0, const FImplicitObject* Implicit1, const FRigidTransform3& ParticleWorldTransform0, const FRigidTransform3& LocalTransform0, const FRigidTransform3& ParticleWorldTransform1, const FRigidTransform3& LocalTransform1, const FReal CullDistance, const FReal Dt, const FCollisionContext& Context)
 		{
 
 			const TSphere<FReal, 3>* Object0 = Implicit0->template GetObject<const TSphere<FReal, 3> >();
 			const TBox<FReal, 3>* Object1 = Implicit1->template GetObject<const TBox<FReal, 3> >();
 			if (ensure(Object0 && Object1))
 			{
-				FRigidBodyPointContactConstraint Constraint = FRigidBodyPointContactConstraint(Particle0, Implicit0, nullptr, LocalTransform0, Particle1, Implicit1, nullptr, LocalTransform1, CullDistance, EContactShapesType::SphereBox, Context.bAllowManifolds);
-				if (T_TRAITS::bImmediateUpdate)
+				FPBDCollisionConstraint* Constraint = FPBDCollisionConstraint::Make(Particle0, Implicit0, nullptr, ParticleWorldTransform0, LocalTransform0, Particle1, Implicit1, nullptr, ParticleWorldTransform1, LocalTransform1, CullDistance, EContactShapesType::SphereBox, Context.bAllowManifolds, *Context.CollisionAllocator);
+				if (T_TRAITS::bImmediateUpdate && (Constraint != nullptr) && !Constraint->WasManifoldRestored())
 				{
-					FRigidTransform3 WorldTransform0 = LocalTransform0 * Collisions::GetTransform(Particle0);
-					FRigidTransform3 WorldTransform1 = LocalTransform1 * Collisions::GetTransform(Particle1);
-					UpdateSphereBoxConstraint(*Object0, WorldTransform0, *Object1, WorldTransform1, Dt, Constraint);
+					FRigidTransform3 WorldTransform0 = LocalTransform0 * ParticleWorldTransform0;
+					FRigidTransform3 WorldTransform1 = LocalTransform1 * ParticleWorldTransform1;
+					UpdateSphereBoxConstraint(*Object0, WorldTransform0, *Object1, WorldTransform1, Dt, *Constraint);
 				}
-				NewConstraints.Add(Constraint);
 			}
 		}
 
@@ -1166,26 +1199,25 @@ namespace Chaos
 		}
 
 
-		void UpdateSphereCapsuleConstraint(const TSphere<FReal, 3>& A, const FRigidTransform3& ATransform, const FCapsule& B, const FRigidTransform3& BTransform, const FReal Dt, FRigidBodyPointContactConstraint& Constraint)
+		void UpdateSphereCapsuleConstraint(const TSphere<FReal, 3>& A, const FRigidTransform3& ATransform, const FCapsule& B, const FRigidTransform3& BTransform, const FReal Dt, FPBDCollisionConstraint& Constraint)
 		{
 			UpdateContactPoint(Constraint, SphereCapsuleContactPoint(A, ATransform, B, BTransform, Constraint.Manifold.RestitutionPadding), Dt);
 		}
 
 		template<typename T_TRAITS>
-		void ConstructSphereCapsuleConstraints(TGeometryParticleHandle<FReal, 3>* Particle0, TGeometryParticleHandle<FReal, 3>* Particle1, const FImplicitObject* Implicit0, const FImplicitObject* Implicit1, const FRigidTransform3& LocalTransform0, const FRigidTransform3& LocalTransform1, const FReal CullDistance, const FReal Dt, const FCollisionContext& Context, FCollisionConstraintsArray& NewConstraints)
+		void ConstructSphereCapsuleConstraints(TGeometryParticleHandle<FReal, 3>* Particle0, TGeometryParticleHandle<FReal, 3>* Particle1, const FImplicitObject* Implicit0, const FImplicitObject* Implicit1, const FRigidTransform3& ParticleWorldTransform0, const FRigidTransform3& LocalTransform0, const FRigidTransform3& ParticleWorldTransform1, const FRigidTransform3& LocalTransform1, const FReal CullDistance, const FReal Dt, const FCollisionContext& Context)
 		{
 			const TSphere<FReal, 3>* Object0 = Implicit0->template GetObject<const TSphere<FReal, 3> >();
 			const FCapsule* Object1 = Implicit1->template GetObject<const FCapsule >();
 			if (ensure(Object0 && Object1))
 			{
-				FRigidBodyPointContactConstraint Constraint = FRigidBodyPointContactConstraint(Particle0, Implicit0, nullptr, LocalTransform0, Particle1, Implicit1, nullptr, LocalTransform1, CullDistance, EContactShapesType::SphereCapsule, Context.bAllowManifolds);
-				if (T_TRAITS::bImmediateUpdate)
+				FPBDCollisionConstraint* Constraint = FPBDCollisionConstraint::Make(Particle0, Implicit0, nullptr, ParticleWorldTransform0, LocalTransform0, Particle1, Implicit1, nullptr, ParticleWorldTransform1, LocalTransform1, CullDistance, EContactShapesType::SphereCapsule, Context.bAllowManifolds, *Context.CollisionAllocator);
+				if (T_TRAITS::bImmediateUpdate && (Constraint != nullptr) && !Constraint->WasManifoldRestored())
 				{
-					FRigidTransform3 WorldTransform0 = LocalTransform0 * Collisions::GetTransform(Particle0);
-					FRigidTransform3 WorldTransform1 = LocalTransform1 * Collisions::GetTransform(Particle1);
-					UpdateSphereCapsuleConstraint(*Object0, WorldTransform0, *Object1, WorldTransform1, Dt, Constraint);
+					FRigidTransform3 WorldTransform0 = LocalTransform0 * ParticleWorldTransform0;
+					FRigidTransform3 WorldTransform1 = LocalTransform1 * ParticleWorldTransform1;
+					UpdateSphereCapsuleConstraint(*Object0, WorldTransform0, *Object1, WorldTransform1, Dt, *Constraint);
 				}
-				NewConstraints.Add(Constraint);
 			}
 		}
 
@@ -1193,26 +1225,25 @@ namespace Chaos
 		// Sphere-Convex
 		//
 
-		void UpdateSphereConvexConstraint(const TSphere<FReal, 3>& A, const FRigidTransform3& ATransform, const FImplicitObject3& B, const FRigidTransform3& BTransform, const FReal Dt, FRigidBodyPointContactConstraint& Constraint)
+		void UpdateSphereConvexConstraint(const TSphere<FReal, 3>& A, const FRigidTransform3& ATransform, const FImplicitObject3& B, const FRigidTransform3& BTransform, const FReal Dt, FPBDCollisionConstraint& Constraint)
 		{
 			UpdateContactPoint(Constraint, SphereConvexContactPoint(A, ATransform, B, BTransform), Dt);
 		}
 
 		template<typename T_TRAITS>
-		void ConstructSphereConvexConstraints(TGeometryParticleHandle<FReal, 3>* Particle0, TGeometryParticleHandle<FReal, 3>* Particle1, const FImplicitObject* Implicit0, const FImplicitObject* Implicit1, const FRigidTransform3& LocalTransform0, const FRigidTransform3& LocalTransform1, const FReal CullDistance, const FReal Dt, const FCollisionContext& Context, FCollisionConstraintsArray& NewConstraints)
+		void ConstructSphereConvexConstraints(TGeometryParticleHandle<FReal, 3>* Particle0, TGeometryParticleHandle<FReal, 3>* Particle1, const FImplicitObject* Implicit0, const FImplicitObject* Implicit1, const FRigidTransform3& ParticleWorldTransform0, const FRigidTransform3& LocalTransform0, const FRigidTransform3& ParticleWorldTransform1, const FRigidTransform3& LocalTransform1, const FReal CullDistance, const FReal Dt, const FCollisionContext& Context)
 		{
 			const TSphere<FReal, 3>* Object0 = Implicit0->template GetObject<const TSphere<FReal, 3> >();
 			const FImplicitObject* Object1 = Implicit1;
 			if (ensure(Object0 && Object1))
 			{
-				FRigidBodyPointContactConstraint Constraint = FRigidBodyPointContactConstraint(Particle0, Implicit0, nullptr, LocalTransform0, Particle1, Implicit1, nullptr, LocalTransform1, CullDistance, EContactShapesType::SphereConvex, Context.bAllowManifolds);
-				if (T_TRAITS::bImmediateUpdate)
+				FPBDCollisionConstraint* Constraint = FPBDCollisionConstraint::Make(Particle0, Implicit0, nullptr, ParticleWorldTransform0, LocalTransform0, Particle1, Implicit1, nullptr, ParticleWorldTransform1, LocalTransform1, CullDistance, EContactShapesType::SphereConvex, Context.bAllowManifolds, *Context.CollisionAllocator);
+				if (T_TRAITS::bImmediateUpdate && (Constraint != nullptr) && !Constraint->WasManifoldRestored())
 				{
-					FRigidTransform3 WorldTransform0 = LocalTransform0 * Collisions::GetTransform(Particle0);
-					FRigidTransform3 WorldTransform1 = LocalTransform1 * Collisions::GetTransform(Particle1);
-					UpdateSphereConvexConstraint(*Object0, WorldTransform0, *Object1, WorldTransform1, Dt, Constraint);
+					FRigidTransform3 WorldTransform0 = LocalTransform0 * ParticleWorldTransform0;
+					FRigidTransform3 WorldTransform1 = LocalTransform1 * ParticleWorldTransform1;
+					UpdateSphereConvexConstraint(*Object0, WorldTransform0, *Object1, WorldTransform1, Dt, *Constraint);
 				}
-				NewConstraints.Add(Constraint);
 			}
 		}
 
@@ -1243,14 +1274,14 @@ namespace Chaos
 		}
 
 		template <typename TriMeshType>
-		void UpdateSphereTriangleMeshConstraint(const TSphere<FReal, 3>& Sphere0, const FRigidTransform3& WorldTransform0, const TriMeshType& TriangleMesh1, const FRigidTransform3& WorldTransform1, const FReal Dt, FRigidBodyPointContactConstraint& Constraint)
+		void UpdateSphereTriangleMeshConstraint(const TSphere<FReal, 3>& Sphere0, const FRigidTransform3& WorldTransform0, const TriMeshType& TriangleMesh1, const FRigidTransform3& WorldTransform1, const FReal Dt, FPBDCollisionConstraint& Constraint)
 		{
 			// @todo(chaos): restitutionpadding
 			UpdateContactPoint(Constraint, SphereTriangleMeshContactPoint(Sphere0, WorldTransform0, TriangleMesh1, WorldTransform1, Constraint.GetCullDistance(), 0.0f), Dt);
 		}
 
 		template<typename TriMeshType>
-		void UpdateSphereTriangleMeshConstraintSwept(TGeometryParticleHandle<FReal, 3>* Particle0, const TSphere<FReal, 3>& Sphere0, const FRigidTransform3& WorldTransform0, const TriMeshType& TriangleMesh1, const FRigidTransform3& WorldTransform1, const FVec3& Dir, const FReal Length, const FReal Dt, FRigidBodySweptPointContactConstraint& Constraint)
+		void UpdateSphereTriangleMeshConstraintSwept(TGeometryParticleHandle<FReal, 3>* Particle0, const TSphere<FReal, 3>& Sphere0, const FRigidTransform3& WorldTransform0, const TriMeshType& TriangleMesh1, const FRigidTransform3& WorldTransform1, const FVec3& Dir, const FReal Length, const FReal Dt, FPBDCollisionConstraint& Constraint)
 		{
 			FReal TOI = 1.0f;
 			UpdateContactPoint(Constraint, SphereTriangleMeshSweptContactPoint(Sphere0, WorldTransform0, TriangleMesh1, WorldTransform1, Dir, Length, TOI), Dt);
@@ -1258,32 +1289,30 @@ namespace Chaos
 		}
 
 		template<typename T_TRAITS>
-		void ConstructSphereTriangleMeshConstraints(TGeometryParticleHandle<FReal, 3>* Particle0, TGeometryParticleHandle<FReal, 3>* Particle1, const FImplicitObject* Implicit0, const FImplicitObject* Implicit1, const FRigidTransform3& LocalTransform0, const FRigidTransform3& LocalTransform1, const FReal CullDistance, const FReal Dt, const FCollisionContext& Context, FCollisionConstraintsArray& NewConstraints)
+		void ConstructSphereTriangleMeshConstraints(TGeometryParticleHandle<FReal, 3>* Particle0, TGeometryParticleHandle<FReal, 3>* Particle1, const FImplicitObject* Implicit0, const FImplicitObject* Implicit1, const FRigidTransform3& ParticleWorldTransform0, const FRigidTransform3& LocalTransform0, const FRigidTransform3& ParticleWorldTransform1, const FRigidTransform3& LocalTransform1, const FReal CullDistance, const FReal Dt, const FCollisionContext& Context)
 		{
 			const TSphere<FReal, 3>* Object0 = Implicit0->template GetObject<const TSphere<FReal, 3> >();
 			if (ensure(Object0))
 			{
 				if (const TImplicitObjectScaled<FTriangleMeshImplicitObject>* ScaledTriangleMesh = Implicit1->template GetObject<const TImplicitObjectScaled<FTriangleMeshImplicitObject>>())
 				{
-					FRigidBodyPointContactConstraint Constraint = FRigidBodyPointContactConstraint(Particle0, Implicit0, nullptr, LocalTransform0, Particle1, Implicit1, nullptr, LocalTransform1, CullDistance, EContactShapesType::SphereTriMesh, Context.bAllowManifolds);
-					if (T_TRAITS::bImmediateUpdate)
+					FPBDCollisionConstraint* Constraint = FPBDCollisionConstraint::Make(Particle0, Implicit0, nullptr, ParticleWorldTransform0, LocalTransform0, Particle1, Implicit1, nullptr, ParticleWorldTransform1, LocalTransform1, CullDistance, EContactShapesType::SphereTriMesh, Context.bAllowManifolds, *Context.CollisionAllocator);
+					if (T_TRAITS::bImmediateUpdate && (Constraint != nullptr) && !Constraint->WasManifoldRestored())
 					{
-						FRigidTransform3 WorldTransform0 = LocalTransform0 * Collisions::GetTransform(Particle0);
-						FRigidTransform3 WorldTransform1 = LocalTransform1 * Collisions::GetTransform(Particle1);
-						UpdateSphereTriangleMeshConstraint(*Object0, WorldTransform0, *ScaledTriangleMesh, WorldTransform1, Dt, Constraint);
+						FRigidTransform3 WorldTransform0 = LocalTransform0 * ParticleWorldTransform0;
+						FRigidTransform3 WorldTransform1 = LocalTransform1 * ParticleWorldTransform1;
+						UpdateSphereTriangleMeshConstraint(*Object0, WorldTransform0, *ScaledTriangleMesh, WorldTransform1, Dt, *Constraint);
 					}
-					NewConstraints.Add(Constraint);
 				}
 				else if (const FTriangleMeshImplicitObject* TriangleMesh = Implicit1->template GetObject<const FTriangleMeshImplicitObject>())
 				{
-					FRigidBodyPointContactConstraint Constraint = FRigidBodyPointContactConstraint(Particle0, Implicit0, nullptr, LocalTransform0, Particle1, Implicit1, nullptr, LocalTransform1, CullDistance, EContactShapesType::SphereTriMesh, Context.bAllowManifolds);
-					if (T_TRAITS::bImmediateUpdate)
+					FPBDCollisionConstraint* Constraint = FPBDCollisionConstraint::Make(Particle0, Implicit0, nullptr, ParticleWorldTransform0, LocalTransform0, Particle1, Implicit1, nullptr, ParticleWorldTransform1, LocalTransform1, CullDistance, EContactShapesType::SphereTriMesh, Context.bAllowManifolds, *Context.CollisionAllocator);
+					if (T_TRAITS::bImmediateUpdate && (Constraint != nullptr) && !Constraint->WasManifoldRestored())
 					{
-						FRigidTransform3 WorldTransform0 = LocalTransform0 * Collisions::GetTransform(Particle0);
-						FRigidTransform3 WorldTransform1 = LocalTransform1 * Collisions::GetTransform(Particle1);
-						UpdateSphereTriangleMeshConstraint(*Object0, WorldTransform0, *TriangleMesh, WorldTransform1, Dt, Constraint);
+						FRigidTransform3 WorldTransform0 = LocalTransform0 * ParticleWorldTransform0;
+						FRigidTransform3 WorldTransform1 = LocalTransform1 * ParticleWorldTransform1;
+						UpdateSphereTriangleMeshConstraint(*Object0, WorldTransform0, *TriangleMesh, WorldTransform1, Dt, *Constraint);
 					}
-					NewConstraints.Add(Constraint);
 				}
 				else
 				{
@@ -1293,26 +1322,24 @@ namespace Chaos
 		}
 
 
-		void ConstructSphereTriangleMeshConstraintsSwept(TGeometryParticleHandle<FReal, 3>* Particle0, TGeometryParticleHandle<FReal, 3>* Particle1, const FImplicitObject* Implicit0, const FImplicitObject* Implicit1, const FRigidTransform3& LocalTransform0, const FRigidTransform3& LocalTransform1, const FReal CullDistance, const FReal Dt, const FVec3& Dir, FReal Length, FCollisionConstraintsArray& NewConstraints)
+		void ConstructSphereTriangleMeshConstraintsSwept(TGeometryParticleHandle<FReal, 3>* Particle0, TGeometryParticleHandle<FReal, 3>* Particle1, const FImplicitObject* Implicit0, const FImplicitObject* Implicit1, const FRigidTransform3& ParticleWorldTransform0, const FRigidTransform3& LocalTransform0, const FRigidTransform3& ParticleWorldTransform1, const FRigidTransform3& LocalTransform1, const FReal CullDistance, const FReal Dt, const FVec3& Dir, FReal Length, const FCollisionContext& Context)
 		{
 			const TSphere<FReal, 3>* Object0 = Implicit0->template GetObject<const TSphere<FReal, 3>>();
 			if (ensure(Object0))
 			{
 				if (const TImplicitObjectScaled<FTriangleMeshImplicitObject>* ScaledTriangleMesh = Implicit1->template GetObject<const TImplicitObjectScaled<FTriangleMeshImplicitObject>>())
 				{
-					FRigidBodySweptPointContactConstraint Constraint = FRigidBodySweptPointContactConstraint(Particle0, Implicit0, nullptr, LocalTransform0, Particle1, Implicit1, nullptr, LocalTransform1, CullDistance, EContactShapesType::SphereTriMesh);
+					FPBDCollisionConstraint* Constraint = FPBDCollisionConstraint::MakeSwept(Particle0, Implicit0, nullptr, ParticleWorldTransform0, LocalTransform0, Particle1, Implicit1, nullptr, ParticleWorldTransform1, LocalTransform1, CullDistance, EContactShapesType::SphereTriMesh, *Context.CollisionAllocator);
 					FRigidTransform3 WorldTransformX0 = LocalTransform0 * FRigidTransform3(Particle0->X(), Particle0->R());
-					FRigidTransform3 WorldTransform1 = LocalTransform1 * Collisions::GetTransform(Particle1);
-					UpdateSphereTriangleMeshConstraintSwept(Particle0, *Object0, WorldTransformX0, *ScaledTriangleMesh, WorldTransform1, Dir, Length, Dt, Constraint);
-					NewConstraints.Add(Constraint);
+					FRigidTransform3 WorldTransform1 = LocalTransform1 * ParticleWorldTransform1;
+					UpdateSphereTriangleMeshConstraintSwept(Particle0, *Object0, WorldTransformX0, *ScaledTriangleMesh, WorldTransform1, Dir, Length, Dt, *Constraint);
 				}
 				else if (const FTriangleMeshImplicitObject* TriangleMesh = Implicit1->template GetObject<const FTriangleMeshImplicitObject>())
 				{
-					FRigidBodySweptPointContactConstraint Constraint = FRigidBodySweptPointContactConstraint(Particle0, Implicit0, nullptr, LocalTransform0, Particle1, Implicit1, nullptr, LocalTransform1, CullDistance, EContactShapesType::SphereTriMesh);
+					FPBDCollisionConstraint* Constraint = FPBDCollisionConstraint::MakeSwept(Particle0, Implicit0, nullptr, ParticleWorldTransform0, LocalTransform0, Particle1, Implicit1, nullptr, ParticleWorldTransform1, LocalTransform1, CullDistance, EContactShapesType::SphereTriMesh, *Context.CollisionAllocator);
 					FRigidTransform3 WorldTransformX0 = LocalTransform0 * FRigidTransform3(Particle0->X(), Particle0->R());
-					FRigidTransform3 WorldTransform1 = LocalTransform1 * Collisions::GetTransform(Particle1);
-					UpdateSphereTriangleMeshConstraintSwept(Particle0, *Object0, WorldTransformX0, *TriangleMesh, WorldTransform1, Dir, Length, Dt, Constraint);
-					NewConstraints.Add(Constraint);
+					FRigidTransform3 WorldTransform1 = LocalTransform1 * ParticleWorldTransform1;
+					UpdateSphereTriangleMeshConstraintSwept(Particle0, *Object0, WorldTransformX0, *TriangleMesh, WorldTransform1, Dir, Length, Dt, *Constraint);
 				}
 				else
 				{
@@ -1363,26 +1390,25 @@ namespace Chaos
 		}
 
 
-		void UpdateCapsuleCapsuleConstraint(const FCapsule& A, const FRigidTransform3& ATransform, const FCapsule& B, const FRigidTransform3& BTransform, const FReal Dt, FRigidBodyPointContactConstraint& Constraint)
+		void UpdateCapsuleCapsuleConstraint(const FCapsule& A, const FRigidTransform3& ATransform, const FCapsule& B, const FRigidTransform3& BTransform, const FReal Dt, FPBDCollisionConstraint& Constraint)
 		{
 			UpdateContactPoint(Constraint, CapsuleCapsuleContactPoint(A, ATransform, B, BTransform, Constraint.Manifold.RestitutionPadding), Dt);
 		}
 
 		template<typename T_TRAITS>
-		void ConstructCapsuleCapsuleConstraints(TGeometryParticleHandle<FReal, 3>* Particle0, TGeometryParticleHandle<FReal, 3>* Particle1, const FImplicitObject* Implicit0, const FImplicitObject* Implicit1, const FRigidTransform3& LocalTransform0, const FRigidTransform3& LocalTransform1, const FReal CullDistance, const FReal Dt, const FCollisionContext& Context, FCollisionConstraintsArray& NewConstraints)
+		void ConstructCapsuleCapsuleConstraints(TGeometryParticleHandle<FReal, 3>* Particle0, TGeometryParticleHandle<FReal, 3>* Particle1, const FImplicitObject* Implicit0, const FImplicitObject* Implicit1, const FRigidTransform3& ParticleWorldTransform0, const FRigidTransform3& LocalTransform0, const FRigidTransform3& ParticleWorldTransform1, const FRigidTransform3& LocalTransform1, const FReal CullDistance, const FReal Dt, const FCollisionContext& Context)
 		{
 			const FCapsule* Object0 = Implicit0->template GetObject<const FCapsule >();
 			const FCapsule* Object1 = Implicit1->template GetObject<const FCapsule >();
 			if (ensure(Object0 && Object1))
 			{
-				FRigidBodyPointContactConstraint Constraint = FRigidBodyPointContactConstraint(Particle0, Implicit0, nullptr, LocalTransform0, Particle1, Implicit1, nullptr, LocalTransform1, CullDistance, EContactShapesType::CapsuleCapsule, Context.bAllowManifolds);
-				if (T_TRAITS::bImmediateUpdate)
+				FPBDCollisionConstraint* Constraint = FPBDCollisionConstraint::Make(Particle0, Implicit0, nullptr, ParticleWorldTransform0, LocalTransform0, Particle1, Implicit1, nullptr, ParticleWorldTransform1, LocalTransform1, CullDistance, EContactShapesType::CapsuleCapsule, Context.bAllowManifolds, *Context.CollisionAllocator);
+				if (T_TRAITS::bImmediateUpdate && (Constraint != nullptr) && !Constraint->WasManifoldRestored())
 				{
-					FRigidTransform3 WorldTransform0 = LocalTransform0 * Collisions::GetTransform(Particle0);
-					FRigidTransform3 WorldTransform1 = LocalTransform1 * Collisions::GetTransform(Particle1);
-					UpdateCapsuleCapsuleConstraint(*Object0, WorldTransform0, *Object1, WorldTransform1, Dt, Constraint);
+					FRigidTransform3 WorldTransform0 = LocalTransform0 * ParticleWorldTransform0;
+					FRigidTransform3 WorldTransform1 = LocalTransform1 * ParticleWorldTransform1;
+					UpdateCapsuleCapsuleConstraint(*Object0, WorldTransform0, *Object1, WorldTransform1, Dt, *Constraint);
 				}
-				NewConstraints.Add(Constraint);
 			}
 		}
 
@@ -1397,233 +1423,26 @@ namespace Chaos
 		}
 
 
-		void UpdateCapsuleBoxConstraint(const FCapsule& A, const FRigidTransform3& ATransform, const FImplicitBox3& B, const FRigidTransform3& BTransform, const FReal Dt, FRigidBodyPointContactConstraint& Constraint)
+		void UpdateCapsuleBoxConstraint(const FCapsule& A, const FRigidTransform3& ATransform, const FImplicitBox3& B, const FRigidTransform3& BTransform, const FReal Dt, FPBDCollisionConstraint& Constraint)
 		{
 			const FVec3 InitialDir = ATransform.GetRotation().Inverse() * -Constraint.GetNormal();
 			UpdateContactPoint(Constraint, CapsuleBoxContactPoint(A, ATransform, B, BTransform, InitialDir, Constraint.Manifold.RestitutionPadding), Dt);
 		}
 
-		// @todo(chaos): make this work with the new one-shot manifold
-		/**
-		 * Generate a Capsule-Box Manifold.
-		 *
-		 * This starts by finding the closest feature using GJK/EPA. Then it adds points to the manifold that depend on
-		 * the closest feature types on the box (face, edge, vertex) and capsule (edge, vertex).
-		 *
-		 * For box vertex collisions, only the single near point/plane is used in the manifold.
-		 * For box face collisions, the box owns the manifold plane, and the capsule edge nearest the box face is used to generate manifold points.
-		 * For box edge collision with the capsule edge, the capsule owns the manifold plane, and the capsule edge is projected onto the box plane to generate the manifold points.
-		 * For box edge with capsule vertex collisions, only the single near point/plane is used in the manifold.
-		 *
-		 */
-		//void UpdateCapsuleBoxManifold(const FCapsule& Capsule, const FRigidTransform3& CapsuleTM, const FImplicitBox3& Box, const FRigidTransform3& BoxTM, const FReal CullDistance, FRigidBodyPointContactConstraint& Constraint)
-		//{
-		//	// Initialize with an empty manifold. If we early-out this leaves the manifold as created but not usable to
-		//	// support the fallback to geometry-based collision detection
-		//	Constraint.InitManifold();
-
-		//	// Capture the state required to invalidate the manifold if things move too much
-		//	const FReal ManifoldPositionTolerance = Chaos_Collision_ManifoldPositionTolerance;
-		//	const FReal ManifoldRotationTolerance = Chaos_Collision_ManifoldRotationTolerance;
-		//	Constraint.InitManifoldTolerance(CapsuleTM, BoxTM, ManifoldPositionTolerance, ManifoldRotationTolerance);
-
-		//	// Find the nearest points on the capsule and box
-		//	// Note: We flip the order for GJK so we get the normal in box space. This makes it easier to build the face-capsule manifold.
-		//	const FRigidTransform3 CapsuleToBoxTM = CapsuleTM.GetRelativeTransform(BoxTM);
-
-		//	// NOTE: All GJK results in box-space
-		//	// @todo(ccaulfield): use center-to-center direction for InitialDir
-		//	FVec3 InitialDir = FVec3(1, 0, 0);
-		//	FReal Penetration;
-		//	FVec3 CapsuleClosestBoxSpace, BoxClosestBoxSpace, NormalBoxSpace;
-		//	int32 ClosestVertexIndexBox, ClosestVertexIndexCapsule;
-		//	{
-		//		SCOPE_CYCLE_COUNTER_GJK();
-		//		if (!ensure(GJKPenetration<true>(Box, Capsule, CapsuleToBoxTM, Penetration, BoxClosestBoxSpace, CapsuleClosestBoxSpace, NormalBoxSpace, ClosestVertexIndexBox, ClosestVertexIndexCapsule, (FReal)0, (FReal)0, InitialDir)))
-		//		{
-		//			return;
-		//		}
-		//	}
-
-		//	// Cache the closest point so we don't need to re-iterate over the manifold on the first iteration.
-		//	Constraint.Manifold.Location = BoxTM.TransformPosition(BoxClosestBoxSpace);
-		//	Constraint.Manifold.Normal = BoxTM.TransformVectorNoScale(NormalBoxSpace);
-		//	Constraint.Manifold.Phi = -Penetration;
-
-		//	// Find the box feature that the near point is on
-		//	// Face, Edge, or Vertex can be determined from number of non-zero elements in the box-space normal.
-		//	const FReal ComponentEpsilon = Chaos_Collision_ManifoldFaceEpsilon;
-		//	int32 NumNonZeroNormalComponents = 0;
-		//	int32 MaxComponentIndex = INDEX_NONE;
-		//	FReal MaxComponentValue = 0;
-		//	for (int32 ComponentIndex = 0; ComponentIndex < 3; ++ComponentIndex)
-		//	{
-		//		FReal AbsComponentValue = FMath::Abs(NormalBoxSpace[ComponentIndex]);
-		//		if (AbsComponentValue > ComponentEpsilon)
-		//		{
-		//			++NumNonZeroNormalComponents;
-		//		}
-		//		if (AbsComponentValue > MaxComponentValue)
-		//		{
-		//			MaxComponentValue = AbsComponentValue;
-		//			MaxComponentIndex = ComponentIndex;
-		//		}
-		//	}
-
-		//	// Make sure we actually have a feature to use
-		//	if (!ensure(MaxComponentIndex != INDEX_NONE))
-		//	{
-		//		return;
-		//	}
-
-		//	FVec3 CapsuleAxis = CapsuleToBoxTM.TransformVectorNoScale(Capsule.GetAxis()); // Box space capsule axis
-		//	FReal CapsuleAxisNormal = FVec3::DotProduct(CapsuleAxis, NormalBoxSpace);
-		//	const bool bIsBoxFaceContact = (NumNonZeroNormalComponents == 1);
-		//	const bool bIsBoxEdgeContact = (NumNonZeroNormalComponents == 2);
-		//	const bool bIsBoxVertexContact = (NumNonZeroNormalComponents == 3);
-		//	const bool bIsCapsuleEdgeContact = (CapsuleAxisNormal < ComponentEpsilon);
-		//	const bool bIsCapsuleVertexContact = !bIsCapsuleEdgeContact;
-
-		//	// We just use the nearest point for these combinations, with the box as the plane owner
-		//	if (bIsBoxVertexContact || (bIsBoxEdgeContact && bIsCapsuleVertexContact))
-		//	{
-		//		Constraint.SetManifoldPlane(1, INDEX_NONE, NormalBoxSpace, BoxClosestBoxSpace);
-		//		Constraint.AddManifoldPoint(CapsuleToBoxTM.InverseTransformPosition(CapsuleClosestBoxSpace));
-		//		return;
-		//	}
-
-		//	// Get the radial direction from the capsule axis to the box face
-		//	FVec3 CapsuleRadiusDir = FVec3::CrossProduct(FVec3::CrossProduct(NormalBoxSpace, CapsuleAxis), CapsuleAxis);
-		//	if (!Utilities::NormalizeSafe(CapsuleRadiusDir))
-		//	{
-		//		// Capsule is perpendicular to the face, just leave the one-point manifold
-		//		Constraint.SetManifoldPlane(1, INDEX_NONE, NormalBoxSpace, BoxClosestBoxSpace);
-		//		Constraint.AddManifoldPoint(CapsuleToBoxTM.InverseTransformPosition(CapsuleClosestBoxSpace));
-		//		return;
-		//	}
-
-		//	// This is the axis of the box face we will be using below
-		//	int32 FaceAxisIndex = MaxComponentIndex;
-
-		//	// If this is a box face contact, the box is the manifold plane owner, regardless of the capsule feature involved. 
-		//	// We take the edge of the capsule nearest the face, clip it to the face, and use the clipped points as the manifold points.
-		//	// NOTE: The plane normal is not exactly the face normal, so this is not strictly the correct manifold. In particular
-		//	// it will over-limit the rotation for this frame. We could fix this by reverse-projecting the edge points based on the angle...
-		//	if (bIsBoxFaceContact)
-		//	{
-		//		// Initialize the manifold with the nearest plane and nearest point
-		//		Constraint.SetManifoldPlane(1, INDEX_NONE, NormalBoxSpace, BoxClosestBoxSpace);
-		//		Constraint.AddManifoldPoint(CapsuleToBoxTM.InverseTransformPosition(CapsuleClosestBoxSpace));
-
-		//		// Calculate the capsule edge nearest the face
-		//		FVec3 CapsuleVert0 = CapsuleToBoxTM.TransformPosition(Capsule.GetX1()) + Capsule.GetMargin() * CapsuleRadiusDir;
-		//		FVec3 CapsuleVert1 = CapsuleToBoxTM.TransformPosition(Capsule.GetX2()) + Capsule.GetMargin() * CapsuleRadiusDir;
-
-		//		// Clip the capsule edge to the axis-aligned box face
-		//		for (int32 ClipAxisIndex = 0; ClipAxisIndex < 3; ++ClipAxisIndex)
-		//		{
-		//			if (ClipAxisIndex != FaceAxisIndex)
-		//			{
-		//				bool bAcceptedPos = Utilities::ClipLineSegmentToAxisAlignedPlane(CapsuleVert0, CapsuleVert1, ClipAxisIndex, 1.0f, Box.Max()[ClipAxisIndex]);
-		//				bool bAcceptedNeg = Utilities::ClipLineSegmentToAxisAlignedPlane(CapsuleVert0, CapsuleVert1, ClipAxisIndex, -1.0f, Box.Min()[ClipAxisIndex]);
-		//				if (!bAcceptedPos || !bAcceptedNeg)
-		//				{
-		//					// Capsule edge is outside a face - stick with the single point we have
-		//					return;
-		//				}
-		//			}
-		//		}
-
-		//		// Add the clipped points to the manifold if they are not too close to each other, or the near point calculated above.
-		//		// Note: verts are in box space - need to be in capsule space
-		//		// @todo(chaos): manifold point distance tolerance should be a per-solver or per object setting
-		//		const FReal DistanceToleranceSq = (0.1f * Capsule.GetHeight()) * (0.1f * Capsule.GetHeight());
-		//		bool bUseVert0 = ((CapsuleVert1 - CapsuleVert0).SizeSquared() > DistanceToleranceSq) && ((CapsuleVert0 - CapsuleClosestBoxSpace).SizeSquared() > DistanceToleranceSq);
-		//		bool bUseVert1 = ((CapsuleVert1 - CapsuleClosestBoxSpace).SizeSquared() > DistanceToleranceSq);
-		//		if (bUseVert0)
-		//		{
-		//			Constraint.AddManifoldPoint(CapsuleToBoxTM.InverseTransformPosition(CapsuleVert0));
-		//		}
-		//		if (bUseVert1)
-		//		{
-		//			Constraint.AddManifoldPoint(CapsuleToBoxTM.InverseTransformPosition(CapsuleVert1));
-		//		}
-
-		//		return;
-		//	}
-
-		//	// If the box edge and the capsule edge are the closest features, we treat the capsule as the
-		//	// plane owner. Then select the most-parallel box face as the source of manifold points. The
-		//	// manifold points are found by projecting the capsule verts onto the face and clipping to it.
-		//	// We have a capsule edge if the normal is perpendicular to the capsule axis
-		//	if (bIsBoxEdgeContact && bIsCapsuleEdgeContact)
-		//	{
-		//		// Initialize the manifold with the plane and point on the capsule (move to capsule space)
-		//		FVec3 NormalCapsuleSpace = CapsuleToBoxTM.InverseTransformVectorNoScale(NormalBoxSpace);
-		//		FVec3 CapsuleClosestCapsuleSpace = CapsuleToBoxTM.InverseTransformPosition(CapsuleClosestBoxSpace);
-		//		Constraint.SetManifoldPlane(0, INDEX_NONE, -NormalCapsuleSpace, CapsuleClosestCapsuleSpace);
-		//		Constraint.AddManifoldPoint(BoxClosestBoxSpace);
-
-		//		// Project the capsule edge onto the box face
-		//		FReal FaceAxisSign = FMath::Sign(NormalBoxSpace[FaceAxisIndex]);
-		//		FReal FacePos = (FaceAxisSign >= 0.0f) ? Box.Max()[FaceAxisIndex] : Box.Min()[FaceAxisIndex];
-		//		FVec3 CapsuleVert0 = CapsuleToBoxTM.TransformPosition(Capsule.GetX1());
-		//		FVec3 CapsuleVert1 = CapsuleToBoxTM.TransformPosition(Capsule.GetX2());
-		//		Utilities::ProjectPointOntoAxisAlignedPlane(CapsuleVert0, CapsuleRadiusDir, FaceAxisIndex, FaceAxisSign, FacePos);
-		//		Utilities::ProjectPointOntoAxisAlignedPlane(CapsuleVert1, CapsuleRadiusDir, FaceAxisIndex, FaceAxisSign, FacePos);
-
-		//		// Clip the capsule edge to the axis-aligned box face
-		//		for (int32 ClipAxisIndex = 0; ClipAxisIndex < 3; ++ClipAxisIndex)
-		//		{
-		//			if (ClipAxisIndex != FaceAxisIndex)
-		//			{
-		//				bool bAcceptedPos = Utilities::ClipLineSegmentToAxisAlignedPlane(CapsuleVert0, CapsuleVert1, ClipAxisIndex, 1.0f, Box.Max()[ClipAxisIndex]);
-		//				bool bAcceptedNeg = Utilities::ClipLineSegmentToAxisAlignedPlane(CapsuleVert0, CapsuleVert1, ClipAxisIndex, -1.0f, Box.Min()[ClipAxisIndex]);
-		//				if (!bAcceptedPos || !bAcceptedNeg)
-		//				{
-		//					// Capsule edge is outside a face - stick with the single point we have
-		//					return;
-		//				}
-		//			}
-		//		}
-
-		//		// Add the clipped points to the manifold if they are not too close to each other, or the near point calculated above.
-		//		// Note: verts are in box space - need to be in capsule space
-		//		// @todo(chaos): manifold point distance tolerance should be a per-solver or per object setting
-		//		const FReal DistanceToleranceSq = (0.1f * Capsule.GetHeight()) * (0.1f * Capsule.GetHeight());
-		//		bool bUseVert0 = ((CapsuleVert1 - CapsuleVert0).SizeSquared() > DistanceToleranceSq) && ((CapsuleVert0 - BoxClosestBoxSpace).SizeSquared() > DistanceToleranceSq);
-		//		bool bUseVert1 = ((CapsuleVert1 - BoxClosestBoxSpace).SizeSquared() > DistanceToleranceSq);
-		//		if (bUseVert0)
-		//		{
-		//			Constraint.AddManifoldPoint(CapsuleVert0);
-		//		}
-		//		if (bUseVert1)
-		//		{
-		//			Constraint.AddManifoldPoint(CapsuleVert1);
-		//		}
-
-		//		return;
-		//	}
-
-		//	// All feature combinations should be covered above
-		//	ensure(false);
-		//}
-
 		template<typename T_TRAITS>
-		void ConstructCapsuleBoxConstraints(TGeometryParticleHandle<FReal, 3>* Particle0, TGeometryParticleHandle<FReal, 3>* Particle1, const FImplicitObject* Implicit0, const FImplicitObject* Implicit1, const FRigidTransform3& LocalTransform0, const FRigidTransform3& LocalTransform1, const FReal CullDistance, const FReal Dt, const FCollisionContext& Context, FCollisionConstraintsArray& NewConstraints)
+		void ConstructCapsuleBoxConstraints(TGeometryParticleHandle<FReal, 3>* Particle0, TGeometryParticleHandle<FReal, 3>* Particle1, const FImplicitObject* Implicit0, const FImplicitObject* Implicit1, const FRigidTransform3& ParticleWorldTransform0, const FRigidTransform3& LocalTransform0, const FRigidTransform3& ParticleWorldTransform1, const FRigidTransform3& LocalTransform1, const FReal CullDistance, const FReal Dt, const FCollisionContext& Context)
 		{
 			const FCapsule* Object0 = Implicit0->template GetObject<const FCapsule>();
 			const TBox<FReal, 3>* Object1 = Implicit1->template GetObject<const TBox<FReal, 3>>();
 			if (ensure(Object0 && Object1))
 			{
-				FRigidBodyPointContactConstraint Constraint = FRigidBodyPointContactConstraint(Particle0, Implicit0, nullptr, LocalTransform0, Particle1, Implicit1, nullptr, LocalTransform1, CullDistance, EContactShapesType::CapsuleBox, Context.bAllowManifolds);
-				if (T_TRAITS::bImmediateUpdate)
+				FPBDCollisionConstraint* Constraint = FPBDCollisionConstraint::Make(Particle0, Implicit0, nullptr, ParticleWorldTransform0, LocalTransform0, Particle1, Implicit1, nullptr, ParticleWorldTransform1, LocalTransform1, CullDistance, EContactShapesType::CapsuleBox, Context.bAllowManifolds, *Context.CollisionAllocator);
+				if (T_TRAITS::bImmediateUpdate && (Constraint != nullptr) && !Constraint->WasManifoldRestored())
 				{
-					FRigidTransform3 WorldTransform0 = LocalTransform0 * Collisions::GetTransform(Particle0);
-					FRigidTransform3 WorldTransform1 = LocalTransform1 * Collisions::GetTransform(Particle1);
-					UpdateCapsuleBoxConstraint(*Object0, WorldTransform0, *Object1, WorldTransform1, Dt, Constraint);
+					FRigidTransform3 WorldTransform0 = LocalTransform0 * ParticleWorldTransform0;
+					FRigidTransform3 WorldTransform1 = LocalTransform1 * ParticleWorldTransform1;
+					UpdateCapsuleBoxConstraint(*Object0, WorldTransform0, *Object1, WorldTransform1, Dt, *Constraint);
 				}
-				NewConstraints.Add(Constraint);
 			}
 		}
 
@@ -1632,26 +1451,25 @@ namespace Chaos
 		// Capsule-Convex
 		//
 
-		void UpdateCapsuleConvexConstraint(const FCapsule& A, const FRigidTransform3& ATransform, const FImplicitObject& B, const FRigidTransform3& BTransform, const FReal Dt, FRigidBodyPointContactConstraint& Constraint)
+		void UpdateCapsuleConvexConstraint(const FCapsule& A, const FRigidTransform3& ATransform, const FImplicitObject& B, const FRigidTransform3& BTransform, const FReal Dt, FPBDCollisionConstraint& Constraint)
 		{
 			UpdateContactPoint(Constraint, CapsuleConvexContactPoint(A, ATransform, B, BTransform), Dt);
 		}
 
 		template<typename T_TRAITS>
-		void ConstructCapsuleConvexConstraints(TGeometryParticleHandle<FReal, 3>* Particle0, TGeometryParticleHandle<FReal, 3>* Particle1, const FImplicitObject* Implicit0, const FImplicitObject* Implicit1, const FRigidTransform3& LocalTransform0, const FRigidTransform3& LocalTransform1, const FReal CullDistance, const FReal Dt, const FCollisionContext& Context, FCollisionConstraintsArray& NewConstraints)
+		void ConstructCapsuleConvexConstraints(TGeometryParticleHandle<FReal, 3>* Particle0, TGeometryParticleHandle<FReal, 3>* Particle1, const FImplicitObject* Implicit0, const FImplicitObject* Implicit1, const FRigidTransform3& ParticleWorldTransform0, const FRigidTransform3& LocalTransform0, const FRigidTransform3& ParticleWorldTransform1, const FRigidTransform3& LocalTransform1, const FReal CullDistance, const FReal Dt, const FCollisionContext& Context)
 		{
 			const FCapsule* Object0 = Implicit0->template GetObject<const FCapsule>();
 			const FImplicitObject* Object1 = Implicit1;
 			if (ensure(Object0 && Object1))
 			{
-				FRigidBodyPointContactConstraint Constraint = FRigidBodyPointContactConstraint(Particle0, Implicit0, nullptr, LocalTransform0, Particle1, Implicit1, nullptr, LocalTransform1, CullDistance, EContactShapesType::CapsuleConvex, Context.bAllowManifolds);
-				if (T_TRAITS::bImmediateUpdate)
+				FPBDCollisionConstraint* Constraint = FPBDCollisionConstraint::Make(Particle0, Implicit0, nullptr, ParticleWorldTransform0, LocalTransform0, Particle1, Implicit1, nullptr, ParticleWorldTransform1, LocalTransform1, CullDistance, EContactShapesType::CapsuleConvex, Context.bAllowManifolds, *Context.CollisionAllocator);
+				if (T_TRAITS::bImmediateUpdate && (Constraint != nullptr) && !Constraint->WasManifoldRestored())
 				{
-					FRigidTransform3 WorldTransform0 = LocalTransform0 * Collisions::GetTransform(Particle0);
-					FRigidTransform3 WorldTransform1 = LocalTransform1 * Collisions::GetTransform(Particle1);
-					UpdateCapsuleConvexConstraint(*Object0, WorldTransform0, *Object1, WorldTransform1, Dt, Constraint);
+					FRigidTransform3 WorldTransform0 = LocalTransform0 * ParticleWorldTransform0;
+					FRigidTransform3 WorldTransform1 = LocalTransform1 * ParticleWorldTransform1;
+					UpdateCapsuleConvexConstraint(*Object0, WorldTransform0, *Object1, WorldTransform1, Dt, *Constraint);
 				}
-				NewConstraints.Add(Constraint);
 			}
 		}
 
@@ -1668,7 +1486,7 @@ namespace Chaos
 		}
 
 
-		void UpdateCapsuleHeightFieldConstraint(const FCapsule& A, const FRigidTransform3& ATransform, const FHeightField& B, const FRigidTransform3& BTransform, const FReal Dt, FRigidBodyPointContactConstraint& Constraint)
+		void UpdateCapsuleHeightFieldConstraint(const FCapsule& A, const FRigidTransform3& ATransform, const FHeightField& B, const FRigidTransform3& BTransform, const FReal Dt, FPBDCollisionConstraint& Constraint)
 		{
 			// @todo(chaos): restitutionpadding
 			CONDITIONAL_SCOPE_CYCLE_COUNTER(STAT_Collisions_UpdateCapsuleHeightFieldConstraint, ConstraintsDetailedStats);
@@ -1676,7 +1494,7 @@ namespace Chaos
 		}
 
 
-		void UpdateCapsuleHeightFieldConstraintSwept(TGeometryParticleHandle<FReal, 3>* Particle0, const FCapsule& A, const FRigidTransform3& ATransform, const FHeightField& B, const FRigidTransform3& BTransform, const FVec3& Dir, const FReal Length, const FReal Dt, FRigidBodySweptPointContactConstraint& Constraint)
+		void UpdateCapsuleHeightFieldConstraintSwept(TGeometryParticleHandle<FReal, 3>* Particle0, const FCapsule& A, const FRigidTransform3& ATransform, const FHeightField& B, const FRigidTransform3& BTransform, const FVec3& Dir, const FReal Length, const FReal Dt, FPBDCollisionConstraint& Constraint)
 		{
 			CONDITIONAL_SCOPE_CYCLE_COUNTER(STAT_Collisions_UpdateCapsuleHeightFieldConstraintSwept, ConstraintsDetailedStats);
 			FReal TOI = 1.0f;
@@ -1685,7 +1503,7 @@ namespace Chaos
 		}
 
 		template<typename T_TRAITS>
-		void ConstructCapsuleHeightFieldConstraints(TGeometryParticleHandle<FReal, 3>* Particle0, TGeometryParticleHandle<FReal, 3>* Particle1, const FImplicitObject* Implicit0, const FImplicitObject* Implicit1, const FRigidTransform3& LocalTransform0, const FRigidTransform3& LocalTransform1, const FReal CullDistance, const FReal Dt, const FCollisionContext& Context, FCollisionConstraintsArray& NewConstraints)
+		void ConstructCapsuleHeightFieldConstraints(TGeometryParticleHandle<FReal, 3>* Particle0, TGeometryParticleHandle<FReal, 3>* Particle1, const FImplicitObject* Implicit0, const FImplicitObject* Implicit1, const FRigidTransform3& ParticleWorldTransform0, const FRigidTransform3& LocalTransform0, const FRigidTransform3& ParticleWorldTransform1, const FRigidTransform3& LocalTransform1, const FReal CullDistance, const FReal Dt, const FCollisionContext& Context)
 		{
 			CONDITIONAL_SCOPE_CYCLE_COUNTER(STAT_Collisions_ConstructCapsuleHeightFieldConstraints, ConstraintsDetailedStats);
 
@@ -1693,29 +1511,27 @@ namespace Chaos
 			const FHeightField* Object1 = Implicit1->template GetObject<const FHeightField >();
 			if (ensure(Object0 && Object1))
 			{
-				FRigidBodyPointContactConstraint Constraint = FRigidBodyPointContactConstraint(Particle0, Implicit0, nullptr, LocalTransform0, Particle1, Implicit1, nullptr, LocalTransform1, CullDistance, EContactShapesType::CapsuleHeightField, Context.bAllowManifolds);
-				if (T_TRAITS::bImmediateUpdate)
+				FPBDCollisionConstraint* Constraint = FPBDCollisionConstraint::Make(Particle0, Implicit0, nullptr, ParticleWorldTransform0, LocalTransform0, Particle1, Implicit1, nullptr, ParticleWorldTransform1, LocalTransform1, CullDistance, EContactShapesType::CapsuleHeightField, Context.bAllowManifolds, *Context.CollisionAllocator);
+				if (T_TRAITS::bImmediateUpdate && (Constraint != nullptr) && !Constraint->WasManifoldRestored())
 				{
-					FRigidTransform3 WorldTransform0 = LocalTransform0 * Collisions::GetTransform(Particle0);
-					FRigidTransform3 WorldTransform1 = LocalTransform1 * Collisions::GetTransform(Particle1);
-					UpdateCapsuleHeightFieldConstraint(*Object0, WorldTransform0, *Object1, WorldTransform1, Dt, Constraint);
+					FRigidTransform3 WorldTransform0 = LocalTransform0 * ParticleWorldTransform0;
+					FRigidTransform3 WorldTransform1 = LocalTransform1 * ParticleWorldTransform1;
+					UpdateCapsuleHeightFieldConstraint(*Object0, WorldTransform0, *Object1, WorldTransform1, Dt, *Constraint);
 				}
-				NewConstraints.Add(Constraint);
 			}
 		}
 
-		void ConstructCapsuleHeightFieldConstraintsSwept(TGeometryParticleHandle<FReal, 3>* Particle0, TGeometryParticleHandle<FReal, 3>* Particle1, const FImplicitObject* Implicit0, const FImplicitObject* Implicit1, const FRigidTransform3& LocalTransform0, const FRigidTransform3& LocalTransform1, const FReal CullDistance, const FReal Dt, const FVec3& Dir, FReal Length, FCollisionConstraintsArray& NewConstraints)
+		void ConstructCapsuleHeightFieldConstraintsSwept(TGeometryParticleHandle<FReal, 3>* Particle0, TGeometryParticleHandle<FReal, 3>* Particle1, const FImplicitObject* Implicit0, const FImplicitObject* Implicit1, const FRigidTransform3& ParticleWorldTransform0, const FRigidTransform3& LocalTransform0, const FRigidTransform3& ParticleWorldTransform1, const FRigidTransform3& LocalTransform1, const FReal CullDistance, const FReal Dt, const FVec3& Dir, FReal Length, const FCollisionContext& Context)
 		{
 			CONDITIONAL_SCOPE_CYCLE_COUNTER(STAT_Collisions_ConstructCapsuleHeightFieldConstraintsSwept, ConstraintsDetailedStats);
 			const FCapsule* Object0 = Implicit0->template GetObject<const FCapsule >();
 			const FHeightField* Object1 = Implicit1->template GetObject<const FHeightField >();
 			if (ensure(Object0 && Object1))
 			{
-				FRigidBodySweptPointContactConstraint Constraint = FRigidBodySweptPointContactConstraint(Particle0, Implicit0, nullptr, LocalTransform0, Particle1, Implicit1, nullptr, LocalTransform1, CullDistance, EContactShapesType::CapsuleHeightField);
+				FPBDCollisionConstraint* Constraint = FPBDCollisionConstraint::MakeSwept(Particle0, Implicit0, nullptr, ParticleWorldTransform0, LocalTransform0, Particle1, Implicit1, nullptr, ParticleWorldTransform1, LocalTransform1, CullDistance, EContactShapesType::CapsuleHeightField, *Context.CollisionAllocator);
 				FRigidTransform3 WorldTransformX0 = LocalTransform0 * FRigidTransform3(Particle0->X(), Particle0->R());
-				FRigidTransform3 WorldTransform1 = LocalTransform1 * Collisions::GetTransform(Particle1);
-				UpdateCapsuleHeightFieldConstraintSwept(Particle0, *Object0, WorldTransformX0, *Object1, WorldTransform1, Dir, Length, Dt, Constraint);
-				NewConstraints.Add(Constraint);
+				FRigidTransform3 WorldTransform1 = LocalTransform1 * ParticleWorldTransform1;
+				UpdateCapsuleHeightFieldConstraintSwept(Particle0, *Object0, WorldTransformX0, *Object1, WorldTransform1, Dir, Length, Dt, *Constraint);
 			}
 		}
 
@@ -1750,7 +1566,7 @@ namespace Chaos
 
 
 		template <typename TriMeshType>
-		void UpdateCapsuleTriangleMeshConstraint(const FCapsule& Capsule0, const FRigidTransform3& WorldTransform0, const TriMeshType& TriangleMesh1, const FRigidTransform3& WorldTransform1, const FReal Dt, FRigidBodyPointContactConstraint& Constraint)
+		void UpdateCapsuleTriangleMeshConstraint(const FCapsule& Capsule0, const FRigidTransform3& WorldTransform0, const TriMeshType& TriangleMesh1, const FRigidTransform3& WorldTransform1, const FReal Dt, FPBDCollisionConstraint& Constraint)
 		{
 			// @todo(chaos): restitutionpadding
 			CONDITIONAL_SCOPE_CYCLE_COUNTER(STAT_Collisions_UpdateCapsuleTriangleMeshConstraint, ConstraintsDetailedStats);
@@ -1758,7 +1574,7 @@ namespace Chaos
 		}
 
 		template <typename TriMeshType>
-		void UpdateCapsuleTriangleMeshConstraintSwept(TGeometryParticleHandle<FReal, 3>* Particle0, const FCapsule& Capsule0, const FRigidTransform3& WorldTransform0, const TriMeshType& TriangleMesh1, const FRigidTransform3& WorldTransform1, const FVec3& Dir, const FReal Length, const FReal Dt, FRigidBodySweptPointContactConstraint& Constraint)
+		void UpdateCapsuleTriangleMeshConstraintSwept(TGeometryParticleHandle<FReal, 3>* Particle0, const FCapsule& Capsule0, const FRigidTransform3& WorldTransform0, const TriMeshType& TriangleMesh1, const FRigidTransform3& WorldTransform1, const FVec3& Dir, const FReal Length, const FReal Dt, FPBDCollisionConstraint& Constraint)
 		{
 			CONDITIONAL_SCOPE_CYCLE_COUNTER(STAT_Collisions_UpdateCapsuleTriangleMeshConstraint, ConstraintsDetailedStats);
 			FReal TOI = 1.0f;
@@ -1767,7 +1583,7 @@ namespace Chaos
 		}
 
 		template<typename T_TRAITS>
-		void ConstructCapsuleTriangleMeshConstraints(TGeometryParticleHandle<FReal, 3>* Particle0, TGeometryParticleHandle<FReal, 3>* Particle1, const FImplicitObject* Implicit0, const FImplicitObject* Implicit1, const FRigidTransform3& LocalTransform0, const FRigidTransform3& LocalTransform1, const FReal CullDistance, const FReal Dt, const FCollisionContext& Context, FCollisionConstraintsArray& NewConstraints)
+		void ConstructCapsuleTriangleMeshConstraints(TGeometryParticleHandle<FReal, 3>* Particle0, TGeometryParticleHandle<FReal, 3>* Particle1, const FImplicitObject* Implicit0, const FImplicitObject* Implicit1, const FRigidTransform3& ParticleWorldTransform0, const FRigidTransform3& LocalTransform0, const FRigidTransform3& ParticleWorldTransform1, const FRigidTransform3& LocalTransform1, const FReal CullDistance, const FReal Dt, const FCollisionContext& Context)
 		{
 			CONDITIONAL_SCOPE_CYCLE_COUNTER(STAT_Collisions_ConstructCapsuleTriangleMeshConstraints, ConstraintsDetailedStats);
 
@@ -1776,25 +1592,23 @@ namespace Chaos
 			{
 				if (const TImplicitObjectScaled<FTriangleMeshImplicitObject>* ScaledTriangleMesh = Implicit1->template GetObject<const TImplicitObjectScaled<FTriangleMeshImplicitObject>>())
 				{
-					FRigidBodyPointContactConstraint Constraint = FRigidBodyPointContactConstraint(Particle0, Implicit0, nullptr, LocalTransform0, Particle1, Implicit1, nullptr, LocalTransform1, CullDistance, EContactShapesType::CapsuleTriMesh, Context.bAllowManifolds);
-					if (T_TRAITS::bImmediateUpdate)
+					FPBDCollisionConstraint* Constraint = FPBDCollisionConstraint::Make(Particle0, Implicit0, nullptr, ParticleWorldTransform0, LocalTransform0, Particle1, Implicit1, nullptr, ParticleWorldTransform1, LocalTransform1, CullDistance, EContactShapesType::CapsuleTriMesh, Context.bAllowManifolds, *Context.CollisionAllocator);
+					if (T_TRAITS::bImmediateUpdate && (Constraint != nullptr) && !Constraint->WasManifoldRestored())
 					{
-						FRigidTransform3 WorldTransform0 = LocalTransform0 * Collisions::GetTransform(Particle0);
-						FRigidTransform3 WorldTransform1 = LocalTransform1 * Collisions::GetTransform(Particle1);
-						UpdateCapsuleTriangleMeshConstraint(*Object0, WorldTransform0, *ScaledTriangleMesh, WorldTransform1, Dt, Constraint);
+						FRigidTransform3 WorldTransform0 = LocalTransform0 * ParticleWorldTransform0;
+						FRigidTransform3 WorldTransform1 = LocalTransform1 * ParticleWorldTransform1;
+						UpdateCapsuleTriangleMeshConstraint(*Object0, WorldTransform0, *ScaledTriangleMesh, WorldTransform1, Dt, *Constraint);
 					}
-					NewConstraints.Add(Constraint);
 				}
 				else if (const FTriangleMeshImplicitObject* TriangleMesh = Implicit1->template GetObject<const FTriangleMeshImplicitObject>())
 				{
-					FRigidBodyPointContactConstraint Constraint = FRigidBodyPointContactConstraint(Particle0, Implicit0, nullptr, LocalTransform0, Particle1, Implicit1, nullptr, LocalTransform1, CullDistance, EContactShapesType::CapsuleTriMesh, Context.bAllowManifolds);
-					if (T_TRAITS::bImmediateUpdate)
+					FPBDCollisionConstraint* Constraint = FPBDCollisionConstraint::Make(Particle0, Implicit0, nullptr, ParticleWorldTransform0, LocalTransform0, Particle1, Implicit1, nullptr, ParticleWorldTransform1, LocalTransform1, CullDistance, EContactShapesType::CapsuleTriMesh, Context.bAllowManifolds, *Context.CollisionAllocator);
+					if (T_TRAITS::bImmediateUpdate && (Constraint != nullptr) && !Constraint->WasManifoldRestored())
 					{
-						FRigidTransform3 WorldTransform0 = LocalTransform0 * Collisions::GetTransform(Particle0);
-						FRigidTransform3 WorldTransform1 = LocalTransform1 * Collisions::GetTransform(Particle1);
-						UpdateCapsuleTriangleMeshConstraint(*Object0, WorldTransform0, *TriangleMesh, WorldTransform1, Dt, Constraint);
+						FRigidTransform3 WorldTransform0 = LocalTransform0 * ParticleWorldTransform0;
+						FRigidTransform3 WorldTransform1 = LocalTransform1 * ParticleWorldTransform1;
+						UpdateCapsuleTriangleMeshConstraint(*Object0, WorldTransform0, *TriangleMesh, WorldTransform1, Dt, *Constraint);
 					}
-					NewConstraints.Add(Constraint);
 				}
 				else
 				{
@@ -1803,7 +1617,7 @@ namespace Chaos
 			}
 		}
 
-		void ConstructCapsuleTriangleMeshConstraintsSwept(TGeometryParticleHandle<FReal, 3>* Particle0, TGeometryParticleHandle<FReal, 3>* Particle1, const FImplicitObject* Implicit0, const FImplicitObject* Implicit1, const FRigidTransform3& LocalTransform0, const FRigidTransform3& LocalTransform1, const FReal CullDistance, const FReal Dt, const FVec3& Dir, FReal Length, FCollisionConstraintsArray& NewConstraints)
+		void ConstructCapsuleTriangleMeshConstraintsSwept(TGeometryParticleHandle<FReal, 3>* Particle0, TGeometryParticleHandle<FReal, 3>* Particle1, const FImplicitObject* Implicit0, const FImplicitObject* Implicit1, const FRigidTransform3& ParticleWorldTransform0, const FRigidTransform3& LocalTransform0, const FRigidTransform3& ParticleWorldTransform1, const FRigidTransform3& LocalTransform1, const FReal CullDistance, const FReal Dt, const FVec3& Dir, FReal Length, const FCollisionContext& Context)
 		{
 			CONDITIONAL_SCOPE_CYCLE_COUNTER(STAT_Collisions_ConstructCapsuleTriangleMeshConstraintsSwept, ConstraintsDetailedStats);
 
@@ -1812,19 +1626,17 @@ namespace Chaos
 			{
 				if (const TImplicitObjectScaled<FTriangleMeshImplicitObject>* ScaledTriangleMesh = Implicit1->template GetObject<const TImplicitObjectScaled<FTriangleMeshImplicitObject>>())
 				{
-					FRigidBodySweptPointContactConstraint Constraint = FRigidBodySweptPointContactConstraint(Particle0, Implicit0, nullptr, LocalTransform0, Particle1, Implicit1, nullptr, LocalTransform1, CullDistance, EContactShapesType::CapsuleTriMesh);
+					FPBDCollisionConstraint* Constraint = FPBDCollisionConstraint::MakeSwept(Particle0, Implicit0, nullptr, ParticleWorldTransform0, LocalTransform0, Particle1, Implicit1, nullptr, ParticleWorldTransform1, LocalTransform1, CullDistance, EContactShapesType::CapsuleTriMesh, *Context.CollisionAllocator);
 					FRigidTransform3 WorldTransformX0 = LocalTransform0 * FRigidTransform3(Particle0->X(), Particle0->R());
-					FRigidTransform3 WorldTransform1 = LocalTransform1 * Collisions::GetTransform(Particle1);
-					UpdateCapsuleTriangleMeshConstraintSwept(Particle0, *Object0, WorldTransformX0, *ScaledTriangleMesh, WorldTransform1, Dir, Length, Dt, Constraint);
-					NewConstraints.Add(Constraint);
+					FRigidTransform3 WorldTransform1 = LocalTransform1 * ParticleWorldTransform1;
+					UpdateCapsuleTriangleMeshConstraintSwept(Particle0, *Object0, WorldTransformX0, *ScaledTriangleMesh, WorldTransform1, Dir, Length, Dt, *Constraint);
 				}
 				else if (const FTriangleMeshImplicitObject* TriangleMesh = Implicit1->template GetObject<const FTriangleMeshImplicitObject>())
 				{
-					FRigidBodySweptPointContactConstraint Constraint = FRigidBodySweptPointContactConstraint(Particle0, Implicit0, nullptr, LocalTransform0, Particle1, Implicit1, nullptr, LocalTransform1, CullDistance, EContactShapesType::CapsuleTriMesh);
+					FPBDCollisionConstraint* Constraint = FPBDCollisionConstraint::MakeSwept(Particle0, Implicit0, nullptr, ParticleWorldTransform0, LocalTransform0, Particle1, Implicit1, nullptr, ParticleWorldTransform1, LocalTransform1, CullDistance, EContactShapesType::CapsuleTriMesh, *Context.CollisionAllocator);
 					FRigidTransform3 WorldTransformX0 = LocalTransform0 * FRigidTransform3(Particle0->X(), Particle0->R());
-					FRigidTransform3 WorldTransform1 = LocalTransform1 * Collisions::GetTransform(Particle1);
-					UpdateCapsuleTriangleMeshConstraintSwept(Particle0, *Object0, WorldTransformX0, *TriangleMesh, WorldTransform1, Dir, Length, Dt, Constraint);
-					NewConstraints.Add(Constraint);
+					FRigidTransform3 WorldTransform1 = LocalTransform1 * ParticleWorldTransform1;
+					UpdateCapsuleTriangleMeshConstraintSwept(Particle0, *Object0, WorldTransformX0, *TriangleMesh, WorldTransform1, Dir, Length, Dt, *Constraint);
 				}
 				else
 				{
@@ -1837,14 +1649,14 @@ namespace Chaos
 		// Generic Convex - Convex (actualy concrete type could be anything)
 		//
 
-		void UpdateGenericConvexConvexConstraint(const FImplicitObject& Implicit0, const FRigidTransform3& WorldTransform0, const FImplicitObject& Implicit1, const FRigidTransform3& WorldTransform1, const FReal Dt, FRigidBodyPointContactConstraint& Constraint)
+		void UpdateGenericConvexConvexConstraint(const FImplicitObject& Implicit0, const FRigidTransform3& WorldTransform0, const FImplicitObject& Implicit1, const FRigidTransform3& WorldTransform1, const FReal Dt, FPBDCollisionConstraint& Constraint)
 		{
 			CONDITIONAL_SCOPE_CYCLE_COUNTER(STAT_Collisions_UpdateGenericConvexConvexConstraint, ConstraintsDetailedStats);
 
 			UpdateGenericConvexConvexConstraintHelper(Implicit0, WorldTransform0, Implicit1, WorldTransform1, Dt, Constraint.Manifold.RestitutionPadding, Constraint);
 		}
 
-		void UpdateGenericConvexConvexConstraintSwept(TGeometryParticleHandle<FReal, 3>* Particle0, const FImplicitObject& Implicit0, const FRigidTransform3& WorldTransform0, const FImplicitObject& Implicit1, const FRigidTransform3& WorldTransform1, const FVec3& Dir, const FReal Length, const FReal Dt, FRigidBodySweptPointContactConstraint& Constraint)
+		void UpdateGenericConvexConvexConstraintSwept(TGeometryParticleHandle<FReal, 3>* Particle0, const FImplicitObject& Implicit0, const FRigidTransform3& WorldTransform0, const FImplicitObject& Implicit1, const FRigidTransform3& WorldTransform1, const FVec3& Dir, const FReal Length, const FReal Dt, FPBDCollisionConstraint& Constraint)
 		{
 			CONDITIONAL_SCOPE_CYCLE_COUNTER(STAT_Collisions_UpdateGenericConvexConvexConstraintSwept, ConstraintsDetailedStats);
 			FReal TOI = 1.0f;
@@ -1854,33 +1666,31 @@ namespace Chaos
 
 
 		template<typename T_TRAITS>
-		void ConstructGenericConvexConvexConstraints(TGeometryParticleHandle<FReal, 3>* Particle0, TGeometryParticleHandle<FReal, 3>* Particle1, const FImplicitObject* Implicit0, const FImplicitObject* Implicit1, const FRigidTransform3& LocalTransform0, const FRigidTransform3& LocalTransform1, const FReal CullDistance, const FReal Dt, const FCollisionContext& Context, FCollisionConstraintsArray& NewConstraints)
+		void ConstructGenericConvexConvexConstraints(TGeometryParticleHandle<FReal, 3>* Particle0, TGeometryParticleHandle<FReal, 3>* Particle1, const FImplicitObject* Implicit0, const FImplicitObject* Implicit1, const FRigidTransform3& ParticleWorldTransform0, const FRigidTransform3& LocalTransform0, const FRigidTransform3& ParticleWorldTransform1, const FRigidTransform3& LocalTransform1, const FReal CullDistance, const FReal Dt, const FCollisionContext& Context)
 		{
 			CONDITIONAL_SCOPE_CYCLE_COUNTER(STAT_Collisions_ConstructGenericConvexConvexConstraints, ConstraintsDetailedStats);
 			EImplicitObjectType Implicit0Type = Particle0->Geometry()->GetType();
 			EImplicitObjectType Implicit1Type = Particle1->Geometry()->GetType();
 
-			FRigidBodyPointContactConstraint Constraint = FRigidBodyPointContactConstraint(Particle0, Implicit0, nullptr, LocalTransform0, Particle1, Implicit1, nullptr, LocalTransform1, CullDistance, EContactShapesType::GenericConvexConvex, Context.bAllowManifolds);
-			if (T_TRAITS::bImmediateUpdate)
+			FPBDCollisionConstraint* Constraint = FPBDCollisionConstraint::Make(Particle0, Implicit0, nullptr, ParticleWorldTransform0, LocalTransform0, Particle1, Implicit1, nullptr, ParticleWorldTransform1, LocalTransform1, CullDistance, EContactShapesType::GenericConvexConvex, Context.bAllowManifolds, *Context.CollisionAllocator);
+			if (T_TRAITS::bImmediateUpdate && (Constraint != nullptr) && !Constraint->WasManifoldRestored())
 			{
-				FRigidTransform3 WorldTransform0 = LocalTransform0 * Collisions::GetTransform(Particle0);
-				FRigidTransform3 WorldTransform1 = LocalTransform1 * Collisions::GetTransform(Particle1);
-				UpdateGenericConvexConvexConstraint(*Implicit0, WorldTransform0, *Implicit1, WorldTransform1, Dt, Constraint);
+				FRigidTransform3 WorldTransform0 = LocalTransform0 * ParticleWorldTransform0;
+				FRigidTransform3 WorldTransform1 = LocalTransform1 * ParticleWorldTransform1;
+				UpdateGenericConvexConvexConstraint(*Implicit0, WorldTransform0, *Implicit1, WorldTransform1, Dt, *Constraint);
 			}
-			NewConstraints.Add(Constraint);
 		}
 
-		void ConstructGenericConvexConvexConstraintsSwept(TGeometryParticleHandle<FReal, 3>* Particle0, TGeometryParticleHandle<FReal, 3>* Particle1, const FImplicitObject* Implicit0, const FImplicitObject* Implicit1, const FRigidTransform3& LocalTransform0, const FRigidTransform3& LocalTransform1, const FReal CullDistance, const FReal Dt, const FVec3& Dir, FReal Length, FCollisionConstraintsArray& NewConstraints)
+		void ConstructGenericConvexConvexConstraintsSwept(TGeometryParticleHandle<FReal, 3>* Particle0, TGeometryParticleHandle<FReal, 3>* Particle1, const FImplicitObject* Implicit0, const FImplicitObject* Implicit1, const FRigidTransform3& ParticleWorldTransform0, const FRigidTransform3& LocalTransform0, const FRigidTransform3& ParticleWorldTransform1, const FRigidTransform3& LocalTransform1, const FReal CullDistance, const FReal Dt, const FVec3& Dir, FReal Length, const FCollisionContext& Context)
 		{
 			CONDITIONAL_SCOPE_CYCLE_COUNTER(STAT_Collisions_ConstructGenericConvexConvexConstraintsSwept, ConstraintsDetailedStats);
 			FGenericParticleHandle RigidParticle = FGenericParticleHandle(Particle0);
-			FRigidBodySweptPointContactConstraint Constraint = FRigidBodySweptPointContactConstraint(Particle0, Implicit0, nullptr, LocalTransform0, Particle1, Implicit1, nullptr, LocalTransform1, CullDistance, EContactShapesType::GenericConvexConvex);
+			FPBDCollisionConstraint* Constraint = FPBDCollisionConstraint::MakeSwept(Particle0, Implicit0, nullptr, ParticleWorldTransform0, LocalTransform0, Particle1, Implicit1, nullptr, ParticleWorldTransform1, LocalTransform1, CullDistance, EContactShapesType::GenericConvexConvex, *Context.CollisionAllocator);
 			// Note: This is unusual but we are using a mix of the previous and current transform
 			// This is due to how CCD only rewinds the position
 			FRigidTransform3 WorldTransformXQ0 = LocalTransform0 * FRigidTransform3(RigidParticle->X(), RigidParticle->Q());
-			FRigidTransform3 WorldTransform1 = LocalTransform1 * Collisions::GetTransform(Particle1);
-			UpdateGenericConvexConvexConstraintSwept(Particle0, *Implicit0, WorldTransformXQ0, *Implicit1, WorldTransform1, Dir, Length, Dt, Constraint);
-			NewConstraints.Add(Constraint);
+			FRigidTransform3 WorldTransform1 = LocalTransform1 * ParticleWorldTransform1;
+			UpdateGenericConvexConvexConstraintSwept(Particle0, *Implicit0, WorldTransformXQ0, *Implicit1, WorldTransform1, Dir, Length, Dt, *Constraint);
 		}
 
 		//
@@ -1895,7 +1705,7 @@ namespace Chaos
 		}
 
 
-		void UpdateConvexHeightFieldConstraint(const FImplicitObject& A, const FRigidTransform3& ATransform, const FHeightField& B, const FRigidTransform3& BTransform, const FReal Dt, FRigidBodyPointContactConstraint& Constraint)
+		void UpdateConvexHeightFieldConstraint(const FImplicitObject& A, const FRigidTransform3& ATransform, const FHeightField& B, const FRigidTransform3& BTransform, const FReal Dt, FPBDCollisionConstraint& Constraint)
 		{
 			// @todo(chaos): restitutionpadding
 			CONDITIONAL_SCOPE_CYCLE_COUNTER(STAT_Collisions_UpdateConvexHeightFieldConstraint, ConstraintsDetailedStats);
@@ -1903,7 +1713,7 @@ namespace Chaos
 		}
 
 
-		void UpdateConvexHeightFieldConstraintSwept(TGeometryParticleHandle<FReal, 3>* Particle0, const FImplicitObject& A, const FRigidTransform3& ATransform, const FHeightField& B, const FRigidTransform3& BTransform, const FVec3& Dir, const FReal Length, const FReal Dt, FRigidBodySweptPointContactConstraint& Constraint)
+		void UpdateConvexHeightFieldConstraintSwept(TGeometryParticleHandle<FReal, 3>* Particle0, const FImplicitObject& A, const FRigidTransform3& ATransform, const FHeightField& B, const FRigidTransform3& BTransform, const FVec3& Dir, const FReal Length, const FReal Dt, FPBDCollisionConstraint& Constraint)
 		{
 			CONDITIONAL_SCOPE_CYCLE_COUNTER(STAT_Collisions_UpdateConvexHeightFieldConstraintSwept, ConstraintsDetailedStats);
 			FReal TOI = 1.0f;
@@ -1912,34 +1722,32 @@ namespace Chaos
 		}
 
 		template<typename T_TRAITS>
-		void ConstructConvexHeightFieldConstraints(TGeometryParticleHandle<FReal, 3>* Particle0, TGeometryParticleHandle<FReal, 3>* Particle1, const FImplicitObject* Implicit0, const FImplicitObject* Implicit1, const FRigidTransform3& LocalTransform0, const FRigidTransform3& LocalTransform1, const FReal CullDistance, const FReal Dt, const FCollisionContext& Context, FCollisionConstraintsArray& NewConstraints)
+		void ConstructConvexHeightFieldConstraints(TGeometryParticleHandle<FReal, 3>* Particle0, TGeometryParticleHandle<FReal, 3>* Particle1, const FImplicitObject* Implicit0, const FImplicitObject* Implicit1, const FRigidTransform3& ParticleWorldTransform0, const FRigidTransform3& LocalTransform0, const FRigidTransform3& ParticleWorldTransform1, const FRigidTransform3& LocalTransform1, const FReal CullDistance, const FReal Dt, const FCollisionContext& Context)
 		{
 			CONDITIONAL_SCOPE_CYCLE_COUNTER(STAT_Collisions_ConstructConvexHeightFieldConstraints, ConstraintsDetailedStats);
 			const FHeightField* Object1 = Implicit1->template GetObject<const FHeightField >();
 			if (ensure(Implicit0->IsConvex() && Object1))
 			{
-				FRigidBodyPointContactConstraint Constraint = FRigidBodyPointContactConstraint(Particle0, Implicit0, nullptr, LocalTransform0, Particle1, Implicit1, nullptr, LocalTransform1, CullDistance, EContactShapesType::ConvexHeightField, Context.bAllowManifolds);
-				if (T_TRAITS::bImmediateUpdate)
+				FPBDCollisionConstraint* Constraint = FPBDCollisionConstraint::Make(Particle0, Implicit0, nullptr, ParticleWorldTransform0, LocalTransform0, Particle1, Implicit1, nullptr, ParticleWorldTransform1, LocalTransform1, CullDistance, EContactShapesType::ConvexHeightField, Context.bAllowManifolds, *Context.CollisionAllocator);
+				if (T_TRAITS::bImmediateUpdate && (Constraint != nullptr) && !Constraint->WasManifoldRestored())
 				{
-					FRigidTransform3 WorldTransform0 = LocalTransform0 * Collisions::GetTransform(Particle0);
-					FRigidTransform3 WorldTransform1 = LocalTransform1 * Collisions::GetTransform(Particle1);
-					UpdateConvexHeightFieldConstraint(*Implicit0, WorldTransform0, *Object1, WorldTransform1, Dt, Constraint);
+					FRigidTransform3 WorldTransform0 = LocalTransform0 * ParticleWorldTransform0;
+					FRigidTransform3 WorldTransform1 = LocalTransform1 * ParticleWorldTransform1;
+					UpdateConvexHeightFieldConstraint(*Implicit0, WorldTransform0, *Object1, WorldTransform1, Dt, *Constraint);
 				}
-				NewConstraints.Add(Constraint);
 			}
 		}
 
-		void ConstructConvexHeightFieldConstraintsSwept(TGeometryParticleHandle<FReal, 3>* Particle0, TGeometryParticleHandle<FReal, 3>* Particle1, const FImplicitObject* Implicit0, const FImplicitObject* Implicit1, const FRigidTransform3& LocalTransform0, const FRigidTransform3& LocalTransform1, const FReal CullDistance, const FReal Dt, const FVec3& Dir, FReal Length, FCollisionConstraintsArray& NewConstraints)
+		void ConstructConvexHeightFieldConstraintsSwept(TGeometryParticleHandle<FReal, 3>* Particle0, TGeometryParticleHandle<FReal, 3>* Particle1, const FImplicitObject* Implicit0, const FImplicitObject* Implicit1, const FRigidTransform3& ParticleWorldTransform0, const FRigidTransform3& LocalTransform0, const FRigidTransform3& ParticleWorldTransform1, const FRigidTransform3& LocalTransform1, const FReal CullDistance, const FReal Dt, const FVec3& Dir, FReal Length, const FCollisionContext& Context)
 		{
 			CONDITIONAL_SCOPE_CYCLE_COUNTER(STAT_Collisions_ConstructConvexHeightFieldConstraintsSwept, ConstraintsDetailedStats);
 			const FHeightField* Object1 = Implicit1->template GetObject<const FHeightField >();
 			if (ensure(Implicit0->IsConvex() && Object1))
 			{
-				FRigidBodySweptPointContactConstraint Constraint = FRigidBodySweptPointContactConstraint(Particle0, Implicit0, nullptr, LocalTransform0, Particle1, Implicit1, nullptr, LocalTransform1, CullDistance, EContactShapesType::ConvexHeightField);
+				FPBDCollisionConstraint* Constraint = FPBDCollisionConstraint::MakeSwept(Particle0, Implicit0, nullptr, ParticleWorldTransform0, LocalTransform0, Particle1, Implicit1, nullptr, ParticleWorldTransform1, LocalTransform1, CullDistance, EContactShapesType::ConvexHeightField, *Context.CollisionAllocator);
 				FRigidTransform3 WorldTransformX0 = LocalTransform0 * FRigidTransform3(Particle0->X(), Particle0->R());
-				FRigidTransform3 WorldTransform1 = LocalTransform1 * Collisions::GetTransform(Particle1);
-				UpdateConvexHeightFieldConstraintSwept(Particle0, *Implicit0, WorldTransformX0, *Object1, WorldTransform1, Dir, Length, Dt, Constraint);
-				NewConstraints.Add(Constraint);
+				FRigidTransform3 WorldTransform1 = LocalTransform1 * ParticleWorldTransform1;
+				UpdateConvexHeightFieldConstraintSwept(Particle0, *Implicit0, WorldTransformX0, *Object1, WorldTransform1, Dir, Length, Dt, *Constraint);
 			}
 		}
 
@@ -1972,7 +1780,7 @@ namespace Chaos
 		}
 
 		template <typename TriMeshType>
-		void UpdateConvexTriangleMeshConstraint(const FImplicitObject& Convex0, const FRigidTransform3& WorldTransform0, const TriMeshType& TriangleMesh1, const FRigidTransform3& WorldTransform1, const FReal Dt, FRigidBodyPointContactConstraint& Constraint)
+		void UpdateConvexTriangleMeshConstraint(const FImplicitObject& Convex0, const FRigidTransform3& WorldTransform0, const TriMeshType& TriangleMesh1, const FRigidTransform3& WorldTransform1, const FReal Dt, FPBDCollisionConstraint& Constraint)
 		{
 			// @todo(chaos): restitutionpadding
 			CONDITIONAL_SCOPE_CYCLE_COUNTER(STAT_Collisions_UpdateConvexTriangleMeshConstraint, ConstraintsDetailedStats);
@@ -1981,7 +1789,7 @@ namespace Chaos
 
 		// Sweeps convex against trimesh
 		template <typename TriMeshType>
-		void UpdateConvexTriangleMeshConstraintSwept(TGeometryParticleHandle<FReal, 3>* Particle0, const FImplicitObject& Convex0, const FRigidTransform3& WorldTransform0, const TriMeshType& TriangleMesh1, const FRigidTransform3& WorldTransform1, const FVec3& Dir, const FReal Length, const FReal Dt, FRigidBodySweptPointContactConstraint& Constraint)
+		void UpdateConvexTriangleMeshConstraintSwept(TGeometryParticleHandle<FReal, 3>* Particle0, const FImplicitObject& Convex0, const FRigidTransform3& WorldTransform0, const TriMeshType& TriangleMesh1, const FRigidTransform3& WorldTransform1, const FVec3& Dir, const FReal Length, const FReal Dt, FPBDCollisionConstraint& Constraint)
 		{
 			CONDITIONAL_SCOPE_CYCLE_COUNTER(STAT_Collisions_UpdateConvexTriangleMeshConstraintSwept, ConstraintsDetailedStats);
 			FReal TOI = 1.0f;
@@ -1990,32 +1798,30 @@ namespace Chaos
 		}
 
 		template<typename T_TRAITS>
-		void ConstructConvexTriangleMeshConstraints(TGeometryParticleHandle<FReal, 3>* Particle0, TGeometryParticleHandle<FReal, 3>* Particle1, const FImplicitObject* Implicit0, const FImplicitObject* Implicit1, const FRigidTransform3& LocalTransform0, const FRigidTransform3& LocalTransform1, const FReal CullDistance, const FReal Dt, const FCollisionContext& Context, FCollisionConstraintsArray& NewConstraints)
+		void ConstructConvexTriangleMeshConstraints(TGeometryParticleHandle<FReal, 3>* Particle0, TGeometryParticleHandle<FReal, 3>* Particle1, const FImplicitObject* Implicit0, const FImplicitObject* Implicit1, const FRigidTransform3& ParticleWorldTransform0, const FRigidTransform3& LocalTransform0, const FRigidTransform3& ParticleWorldTransform1, const FRigidTransform3& LocalTransform1, const FReal CullDistance, const FReal Dt, const FCollisionContext& Context)
 		{
 			CONDITIONAL_SCOPE_CYCLE_COUNTER(STAT_Collisions_ConstructConvexTriangleMeshConstraints, ConstraintsDetailedStats);
 			if (ensure(Implicit0->IsConvex()))
 			{
 				if (const TImplicitObjectScaled<FTriangleMeshImplicitObject>* ScaledTriangleMesh = Implicit1->template GetObject<const TImplicitObjectScaled<FTriangleMeshImplicitObject>>())
 				{
-					FRigidBodyPointContactConstraint Constraint = FRigidBodyPointContactConstraint(Particle0, Implicit0, nullptr, LocalTransform0, Particle1, Implicit1, nullptr, LocalTransform1, CullDistance, EContactShapesType::ConvexTriMesh, Context.bAllowManifolds);
-					if (T_TRAITS::bImmediateUpdate)
+					FPBDCollisionConstraint* Constraint = FPBDCollisionConstraint::Make(Particle0, Implicit0, nullptr, ParticleWorldTransform0, LocalTransform0, Particle1, Implicit1, nullptr, ParticleWorldTransform1, LocalTransform1, CullDistance, EContactShapesType::ConvexTriMesh, Context.bAllowManifolds, *Context.CollisionAllocator);
+					if (T_TRAITS::bImmediateUpdate && (Constraint != nullptr) && !Constraint->WasManifoldRestored())
 					{
-						FRigidTransform3 WorldTransform0 = LocalTransform0 * Collisions::GetTransform(Particle0);
-						FRigidTransform3 WorldTransform1 = LocalTransform1 * Collisions::GetTransform(Particle1);
-						UpdateConvexTriangleMeshConstraint(*Implicit0, WorldTransform0, *ScaledTriangleMesh, WorldTransform1, Dt, Constraint);
+						FRigidTransform3 WorldTransform0 = LocalTransform0 * ParticleWorldTransform0;
+						FRigidTransform3 WorldTransform1 = LocalTransform1 * ParticleWorldTransform1;
+						UpdateConvexTriangleMeshConstraint(*Implicit0, WorldTransform0, *ScaledTriangleMesh, WorldTransform1, Dt, *Constraint);
 					}
-					NewConstraints.Add(Constraint);
 				}
 				else if (const FTriangleMeshImplicitObject* TriangleMesh = Implicit1->template GetObject<const FTriangleMeshImplicitObject>())
 				{
-					FRigidBodyPointContactConstraint Constraint = FRigidBodyPointContactConstraint(Particle0, Implicit0, nullptr, LocalTransform0, Particle1, Implicit1, nullptr, LocalTransform1, CullDistance, EContactShapesType::ConvexTriMesh, Context.bAllowManifolds);
-					if (T_TRAITS::bImmediateUpdate)
+					FPBDCollisionConstraint* Constraint = FPBDCollisionConstraint::Make(Particle0, Implicit0, nullptr, ParticleWorldTransform0, LocalTransform0, Particle1, Implicit1, nullptr, ParticleWorldTransform1, LocalTransform1, CullDistance, EContactShapesType::ConvexTriMesh, Context.bAllowManifolds, *Context.CollisionAllocator);
+					if (T_TRAITS::bImmediateUpdate && (Constraint != nullptr) && !Constraint->WasManifoldRestored())
 					{
-						FRigidTransform3 WorldTransform0 = LocalTransform0 * Collisions::GetTransform(Particle0);
-						FRigidTransform3 WorldTransform1 = LocalTransform1 * Collisions::GetTransform(Particle1);
-						UpdateConvexTriangleMeshConstraint(*Implicit0, WorldTransform0, *TriangleMesh, WorldTransform1, Dt, Constraint);
+						FRigidTransform3 WorldTransform0 = LocalTransform0 * ParticleWorldTransform0;
+						FRigidTransform3 WorldTransform1 = LocalTransform1 * ParticleWorldTransform1;
+						UpdateConvexTriangleMeshConstraint(*Implicit0, WorldTransform0, *TriangleMesh, WorldTransform1, Dt, *Constraint);
 					}
-					NewConstraints.Add(Constraint);
 				}
 				else
 				{
@@ -2024,26 +1830,24 @@ namespace Chaos
 			}
 		}
 
-		void ConstructConvexTriangleMeshConstraintsSwept(TGeometryParticleHandle<FReal, 3>* Particle0, TGeometryParticleHandle<FReal, 3>* Particle1, const FImplicitObject* Implicit0, const FImplicitObject* Implicit1, const FRigidTransform3& LocalTransform0, const FRigidTransform3& LocalTransform1, const FReal CullDistance, const FReal Dt, const FVec3& Dir, FReal Length, FCollisionConstraintsArray& NewConstraints)
+		void ConstructConvexTriangleMeshConstraintsSwept(TGeometryParticleHandle<FReal, 3>* Particle0, TGeometryParticleHandle<FReal, 3>* Particle1, const FImplicitObject* Implicit0, const FImplicitObject* Implicit1, const FRigidTransform3& ParticleWorldTransform0, const FRigidTransform3& LocalTransform0, const FRigidTransform3& ParticleWorldTransform1, const FRigidTransform3& LocalTransform1, const FReal CullDistance, const FReal Dt, const FVec3& Dir, FReal Length, const FCollisionContext& Context)
 		{
 			CONDITIONAL_SCOPE_CYCLE_COUNTER(STAT_Collisions_ConstructConvexTriangleMeshConstraintsSwept, ConstraintsDetailedStats);
 			if (ensure(Implicit0->IsConvex()))
 			{
 				if (const TImplicitObjectScaled<FTriangleMeshImplicitObject>* ScaledTriangleMesh = Implicit1->template GetObject<const TImplicitObjectScaled<FTriangleMeshImplicitObject>>())
 				{
-					FRigidBodySweptPointContactConstraint Constraint = FRigidBodySweptPointContactConstraint(Particle0, Implicit0, nullptr, LocalTransform0, Particle1, Implicit1, nullptr, LocalTransform1, CullDistance, EContactShapesType::ConvexTriMesh);
+					FPBDCollisionConstraint* Constraint = FPBDCollisionConstraint::MakeSwept(Particle0, Implicit0, nullptr, ParticleWorldTransform0, LocalTransform0, Particle1, Implicit1, nullptr, ParticleWorldTransform1, LocalTransform1, CullDistance, EContactShapesType::ConvexTriMesh, *Context.CollisionAllocator);
 					FRigidTransform3 WorldTransformX0 = LocalTransform0 * FRigidTransform3(Particle0->X(), Particle0->R());
-					FRigidTransform3 WorldTransform1 = LocalTransform1 * Collisions::GetTransform(Particle1);
-					UpdateConvexTriangleMeshConstraintSwept(Particle0, *Implicit0, WorldTransformX0, *ScaledTriangleMesh, WorldTransform1, Dir, Length, Dt, Constraint);
-					NewConstraints.Add(Constraint);
+					FRigidTransform3 WorldTransform1 = LocalTransform1 * ParticleWorldTransform1;
+					UpdateConvexTriangleMeshConstraintSwept(Particle0, *Implicit0, WorldTransformX0, *ScaledTriangleMesh, WorldTransform1, Dir, Length, Dt, *Constraint);
 				}
 				else if (const FTriangleMeshImplicitObject* TriangleMesh = Implicit1->template GetObject<const FTriangleMeshImplicitObject>())
 				{
-					FRigidBodySweptPointContactConstraint Constraint = FRigidBodySweptPointContactConstraint(Particle0, Implicit0, nullptr, LocalTransform0, Particle1, Implicit1, nullptr, LocalTransform1, CullDistance, EContactShapesType::ConvexTriMesh);
+					FPBDCollisionConstraint* Constraint = FPBDCollisionConstraint::MakeSwept(Particle0, Implicit0, nullptr, ParticleWorldTransform0, LocalTransform0, Particle1, Implicit1, nullptr, ParticleWorldTransform1, LocalTransform1, CullDistance, EContactShapesType::ConvexTriMesh, *Context.CollisionAllocator);
 					FRigidTransform3 WorldTransformX0 = LocalTransform0 * FRigidTransform3(Particle0->X(), Particle0->R());
-					FRigidTransform3 WorldTransform1 = LocalTransform1 * Collisions::GetTransform(Particle1);
-					UpdateConvexTriangleMeshConstraintSwept(Particle0, *Implicit0, WorldTransformX0, *TriangleMesh, WorldTransform1, Dir, Length, Dt, Constraint);
-					NewConstraints.Add(Constraint);
+					FRigidTransform3 WorldTransform1 = LocalTransform1 * ParticleWorldTransform1;
+					UpdateConvexTriangleMeshConstraintSwept(Particle0, *Implicit0, WorldTransformX0, *TriangleMesh, WorldTransform1, Dir, Length, Dt, *Constraint);
 				}
 				else
 				{
@@ -2058,7 +1862,7 @@ namespace Chaos
 		//
 
 		template<ECollisionUpdateType UpdateType>
-		void UpdateLevelsetLevelsetConstraint(const FRigidTransform3& WorldTransform0, const FRigidTransform3& WorldTransform1, const FReal Dt, FRigidBodyPointContactConstraint& Constraint)
+		void UpdateLevelsetLevelsetConstraint(const FRigidTransform3& WorldTransform0, const FRigidTransform3& WorldTransform1, const FReal Dt, FPBDCollisionConstraint& Constraint)
 		{
 			SCOPE_CYCLE_COUNTER(STAT_UpdateLevelsetLevelsetConstraint);
 
@@ -2074,11 +1878,25 @@ namespace Chaos
 				return;
 			}
 
-			FGenericParticleHandle Particle0 = Constraint.Particle[0];
 			const FBVHParticles* SampleParticles = Constraint.Manifold.Simplicial[0];
 			if (!SampleParticles)
 			{
-				ParticlesTM = FRigidTransform3(Particle0->P(), Particle0->Q());
+				// Sample particles are in Particle(Actor) space, but SolverBody state in in CoM space
+				// The transforms passed in are the Particle transforms but concatenatated with the shape transform(s) so we need
+				// to regenerate the Particle transform without those shape transforms. We should be able to eliminate this...
+				FGenericParticleHandle Particle0 = Constraint.Particle[0];
+				// @todo(chaos): we should not be accessing the particle here (or anywhere in collision detection)
+				// The two paths here are required for unit testing. Units tests don't always set up the solver bodies, but this function is
+				// usually called in the solver loop where solver bodies must be used.
+				if (Constraint.GetSolverBody0() != nullptr)
+				{
+					const FRigidTransform3 CoMToParticleTransform = FParticleUtilities::ParticleLocalToCoMLocal(Particle0, FRigidTransform3::Identity);
+					ParticlesTM = CoMToParticleTransform * FRigidTransform3(Constraint.GetSolverBody0()->P(), Constraint.GetSolverBody0()->Q());
+				}
+				else
+				{
+					ParticlesTM = FParticleUtilities::GetActorWorldTransform(Particle0);
+				}
 				SampleParticles = Particle0->CollisionParticles().Get();
 			}
 
@@ -2101,7 +1919,7 @@ namespace Chaos
 			}
 		}
 
-		void UpdateLevelsetLevelsetManifold(FCollisionConstraintBase& Constraint, const FRigidTransform3& ATM, const FRigidTransform3& BTM)
+		void UpdateLevelsetLevelsetManifold(FPBDCollisionConstraint& Constraint, const FRigidTransform3& ATM, const FRigidTransform3& BTM)
 		{
 			// @todo(chaos) : Stub Update Manifold
 			//   Stub function for updating the manifold prior to the Apply and ApplyPushOut
@@ -2109,37 +1927,40 @@ namespace Chaos
 
 
 		template<typename T_TRAITS>
-		void ConstructLevelsetLevelsetConstraints(TGeometryParticleHandle<FReal, 3>* Particle0, TGeometryParticleHandle<FReal, 3>* Particle1, const FImplicitObject* Implicit0, const FBVHParticles* Simplicial0, const FImplicitObject* Implicit1, const FBVHParticles* Simplicial1, const FRigidTransform3& LocalTransform0, const FRigidTransform3& LocalTransform1, const FReal CullDistance, const FReal Dt, const FCollisionContext& Context, FCollisionConstraintsArray& NewConstraints)
+		void ConstructLevelsetLevelsetConstraints(TGeometryParticleHandle<FReal, 3>* Particle0, TGeometryParticleHandle<FReal, 3>* Particle1, const FImplicitObject* Implicit0, const FBVHParticles* Simplicial0, const FImplicitObject* Implicit1, const FBVHParticles* Simplicial1, const FRigidTransform3& ParticleWorldTransform0, const FRigidTransform3& LocalTransform0, const FRigidTransform3& ParticleWorldTransform1, const FRigidTransform3& LocalTransform1, const FReal CullDistance, const FReal Dt, const FCollisionContext& Context)
 		{
 			const bool bAllowManifolds = Context.bAllowManifolds && bChaos_Collision_AllowLevelsetManifolds;
-			FRigidBodyPointContactConstraint Constraint = FRigidBodyPointContactConstraint(Particle0, Implicit0, Simplicial0, LocalTransform0, Particle1, Implicit1, Simplicial1, LocalTransform1, CullDistance, EContactShapesType::LevelSetLevelSet, bAllowManifolds);
+			FPBDCollisionConstraint* Constraint = FPBDCollisionConstraint::Make(Particle0, Implicit0, Simplicial0, ParticleWorldTransform0, LocalTransform0, Particle1, Implicit1, Simplicial1, ParticleWorldTransform1, LocalTransform1, CullDistance, EContactShapesType::LevelSetLevelSet, bAllowManifolds, *Context.CollisionAllocator);
 
+			// @todo(chaos): support manifold restoration for Levelsets and remove this
+			Constraint->ResetManifold();
+
+			// @todo(chaos): the order test should be done befor we create the constraint so we can do it right first time, rather than patching it after
 			bool bIsParticleDynamic0 = Particle0->CastToRigidParticle() && Particle0->ObjectState() == EObjectStateType::Dynamic;
 			int32 P0NumCollisionParticles = Simplicial0 ? Simplicial0->Size() : Particle0->CastToRigidParticle()->CollisionParticlesSize();
 			if (!Particle1->Geometry() || (bIsParticleDynamic0 && !P0NumCollisionParticles && Particle0->Geometry() && !Particle0->Geometry()->IsUnderlyingUnion()))
 			{
-				Constraint.Particle[0] = Particle1;
-				Constraint.Particle[1] = Particle0;
-				Constraint.ImplicitTransform[0] = LocalTransform1;
-				Constraint.ImplicitTransform[1] = LocalTransform0;
-				Constraint.SetManifold(Implicit1, Simplicial1, Implicit0, Simplicial0);
+				Constraint->Particle[0] = Particle1;
+				Constraint->Particle[1] = Particle0;
+				Constraint->ImplicitTransform[0] = LocalTransform1;
+				Constraint->ImplicitTransform[1] = LocalTransform0;
+				Constraint->SetManifold(Implicit1, Simplicial1, Implicit0, Simplicial0);
 			}
 			else
 			{
-				Constraint.Particle[0] = Particle0;
-				Constraint.Particle[1] = Particle1;
-				Constraint.ImplicitTransform[0] = LocalTransform0;
-				Constraint.ImplicitTransform[1] = LocalTransform1;
-				Constraint.SetManifold(Implicit0, Simplicial0, Implicit1, Simplicial1);
+				Constraint->Particle[0] = Particle0;
+				Constraint->Particle[1] = Particle1;
+				Constraint->ImplicitTransform[0] = LocalTransform0;
+				Constraint->ImplicitTransform[1] = LocalTransform1;
+				Constraint->SetManifold(Implicit0, Simplicial0, Implicit1, Simplicial1);
 			}
 
-			if (T_TRAITS::bImmediateUpdate)
+			if (T_TRAITS::bImmediateUpdate && (Constraint != nullptr) && !Constraint->WasManifoldRestored())
 			{
-				FRigidTransform3 WorldTransform0 = Constraint.ImplicitTransform[0] * Collisions::GetTransform(Constraint.Particle[0]);
-				FRigidTransform3 WorldTransform1 = Constraint.ImplicitTransform[1] * Collisions::GetTransform(Constraint.Particle[1]);
-				UpdateLevelsetLevelsetConstraint<ECollisionUpdateType::Deepest>(WorldTransform0, WorldTransform1, Dt, Constraint);
+				FRigidTransform3 WorldTransform0 = Constraint->ImplicitTransform[0] * ParticleWorldTransform0;
+				FRigidTransform3 WorldTransform1 = Constraint->ImplicitTransform[1] * ParticleWorldTransform1;
+				UpdateLevelsetLevelsetConstraint<ECollisionUpdateType::Deepest>(WorldTransform0, WorldTransform1, Dt, *Constraint);
 			}
-			NewConstraints.Add(Constraint);
 		}
 
 
@@ -2154,7 +1975,7 @@ namespace Chaos
 		// @todo(chaos): add the missing cases below
 		// @todo(chaos): see use GetInnerObject below - we should try to use the leaf types for all (currently only TriMesh needs this)
 		template<ECollisionUpdateType UpdateType>
-		inline void UpdateConstraintFromGeometryImpl(FRigidBodyPointContactConstraint& Constraint, const FRigidTransform3& WorldTransform0, const FRigidTransform3& WorldTransform1, const FReal Dt)
+		inline void UpdateConstraintFromGeometryImpl(FPBDCollisionConstraint& Constraint, const FRigidTransform3& WorldTransform0, const FRigidTransform3& WorldTransform1, const FReal Dt)
 		{
 			CONDITIONAL_SCOPE_CYCLE_COUNTER(STAT_Collisions_UpdateConstraintFromGeometryInternal, ConstraintsDetailedStats);
 
@@ -2317,11 +2138,11 @@ namespace Chaos
 		// NOTE: Transforms are world space shape transforms
 		
 		template<ECollisionUpdateType UpdateType>
-		inline void UpdateConstraintFromGeometryImpl(FRigidBodySweptPointContactConstraint& Constraint, const FRigidTransform3& WorldTransform0, const FRigidTransform3& WorldTransform1, const FReal Dt)
+		inline void UpdateConstraintFromGeometrySweptImpl(FPBDCollisionConstraint& Constraint, const FRigidTransform3& WorldTransform0, const FRigidTransform3& WorldTransform1, const FReal Dt)
 		{
 			FReal LengthCCD = 0.0f;
 			FVec3 DirCCD(0.0f);
-			bool bUseCCD = UseCCD(Constraint.Particle[0], Constraint.Particle[1], Constraint.GetManifold().Implicit[0], DirCCD, LengthCCD, false);
+			bool bUseCCD = UseCCD_Update(Constraint, DirCCD, LengthCCD, false);
 
 			const FImplicitObject& Implicit0 = *Constraint.Manifold.Implicit[0];
 			const FImplicitObject& Implicit1 = *Constraint.Manifold.Implicit[1];
@@ -2410,11 +2231,11 @@ namespace Chaos
 
 			Constraint.TimeOfImpact = 1.0f; // CCD will not be used
 			// Do a normal non-swept update if we reach this point - Not required for now, since this will be done in solver
-			// UpdateConstraintFromGeometryImpl<UpdateType>(*(Constraint.As<FRigidBodyPointContactConstraint>()), WorldTransform0, WorldTransform1, Dt);
+			// UpdateConstraintFromGeometryImpl<UpdateType>(*(Constraint.As<FPBDCollisionConstraint>()), WorldTransform0, WorldTransform1, Dt);
 		}
 
 		template<typename T_TRAITS>
-		void ConstructConstraintsImpl(TGeometryParticleHandle<FReal, 3>* Particle0, TGeometryParticleHandle<FReal, 3>* Particle1, const FImplicitObject* Implicit0, const FBVHParticles* Simplicial0, const FImplicitObject* Implicit1, const FBVHParticles* Simplicial1, const FRigidTransform3& LocalTransform0, const FRigidTransform3& LocalTransform1, const FReal CullDistance, const FReal Dt, const FCollisionContext& Context, FCollisionConstraintsArray& NewConstraints)
+		void ConstructConstraintsImpl(TGeometryParticleHandle<FReal, 3>* Particle0, TGeometryParticleHandle<FReal, 3>* Particle1, const FImplicitObject* Implicit0, const FBVHParticles* Simplicial0, const FImplicitObject* Implicit1, const FBVHParticles* Simplicial1, const FRigidTransform3& ParticleWorldTransform0, const FRigidTransform3& LocalTransform0, const FRigidTransform3& ParticleWorldTransform1, const FRigidTransform3& LocalTransform1, const FReal CullDistance, const FReal Dt, const FCollisionContext& Context)
 		{
 			CONDITIONAL_SCOPE_CYCLE_COUNTER(STAT_Collisions_ConstructConstraintsInternal, ConstraintsDetailedStats);
 			INC_DWORD_STAT(STAT_ChaosCollisionCounter_NumShapePairs);
@@ -2427,17 +2248,15 @@ namespace Chaos
 
 			FReal LengthCCD = 0.0f;
 			FVec3 DirCCD(0.0f);
-			bool bUseCCD = UseCCD(Particle0, Particle1, Implicit0, DirCCD, LengthCCD, Context.bForceDisableCCD);
+			bool bUseCCD = UseCCD_Init(Particle0, Particle1, Implicit0, DirCCD, LengthCCD, Context.bForceDisableCCD);
 			const bool bUseGenericSweptConstraints = bUseCCD && (CCDUseGenericSweptConvexConstraints > 0) && bIsConvex0 && bIsConvex1;
 #if CHAOS_COLLISION_CREATE_BOUNDSCHECK
 			if ((Implicit0 != nullptr) && (Implicit1 != nullptr))
 			{
 				if (Implicit0->HasBoundingBox() && Implicit1->HasBoundingBox())
 				{
-					const FRigidTransform3 ParticleTransform0 = Collisions::GetTransform(Particle0);
-					const FRigidTransform3 ParticleTransform1 = Collisions::GetTransform(Particle1);
-					const FRigidTransform3 WorldTransform0 = LocalTransform0 * ParticleTransform0;
-					const FRigidTransform3 WorldTransform1 = LocalTransform1 * ParticleTransform1;
+					const FRigidTransform3 WorldTransform0 = LocalTransform0 * ParticleWorldTransform0;
+					const FRigidTransform3 WorldTransform1 = LocalTransform1 * ParticleWorldTransform1;
 					if (Chaos_Collision_NarrowPhase_SphereBoundsCheck)
 					{
 						const FReal R1 = Implicit0->BoundingBox().OriginRadius();
@@ -2467,231 +2286,239 @@ namespace Chaos
 
 			if (Implicit0Type == TBox<FReal, 3>::StaticType() && Implicit1Type == TBox<FReal, 3>::StaticType() && !bUseGenericSweptConstraints)
 			{
-				ConstructBoxBoxConstraints<T_TRAITS>(Particle0, Particle1, Implicit0, Implicit1, LocalTransform0, LocalTransform1, CullDistance, Dt, Context, NewConstraints);
+				ConstructBoxBoxConstraints<T_TRAITS>(Particle0, Particle1, Implicit0, Implicit1, ParticleWorldTransform0, LocalTransform0, ParticleWorldTransform1, LocalTransform1, CullDistance, Dt, Context);
 			}
 			else if (Implicit0Type == TBox<FReal, 3>::StaticType() && Implicit1Type == FHeightField::StaticType())
 			{
-				ConstructBoxHeightFieldConstraints<T_TRAITS>(Particle0, Particle1, Implicit0, Implicit1, LocalTransform0, LocalTransform1, CullDistance, Dt, Context, NewConstraints);
+				ConstructBoxHeightFieldConstraints<T_TRAITS>(Particle0, Particle1, Implicit0, Implicit1, ParticleWorldTransform0, LocalTransform0, ParticleWorldTransform1, LocalTransform1, CullDistance, Dt, Context);
 			}
 			else if (Implicit0Type == FHeightField::StaticType() && Implicit1Type == TBox<FReal, 3>::StaticType())
 			{
-				ConstructBoxHeightFieldConstraints<T_TRAITS>(Particle1, Particle0, Implicit1, Implicit0, LocalTransform1, LocalTransform0, CullDistance, Dt, Context, NewConstraints);
+				ConstructBoxHeightFieldConstraints<T_TRAITS>(Particle1, Particle0, Implicit1, Implicit0, ParticleWorldTransform1, LocalTransform1, ParticleWorldTransform0, LocalTransform0, CullDistance, Dt, Context);
 			}
 			else if (Implicit0Type == TBox<FReal, 3>::StaticType() && Implicit1Type == TPlane<FReal, 3>::StaticType() && !bUseGenericSweptConstraints)
 			{
-				ConstructBoxPlaneConstraints<T_TRAITS>(Particle0, Particle1, Implicit0, Implicit1, LocalTransform0, LocalTransform1, CullDistance, Dt, Context, NewConstraints);
+				ConstructBoxPlaneConstraints<T_TRAITS>(Particle0, Particle1, Implicit0, Implicit1, ParticleWorldTransform0, LocalTransform0, ParticleWorldTransform1, LocalTransform1, CullDistance, Dt, Context);
 			}
 			else if (Implicit0Type == TPlane<FReal, 3>::StaticType() && Implicit1Type == TBox<FReal, 3>::StaticType() && !bUseGenericSweptConstraints)
 			{
-				ConstructBoxPlaneConstraints<T_TRAITS>(Particle1, Particle0, Implicit1, Implicit0, LocalTransform1, LocalTransform0, CullDistance, Dt, Context, NewConstraints);
+				ConstructBoxPlaneConstraints<T_TRAITS>(Particle1, Particle0, Implicit1, Implicit0, ParticleWorldTransform1, LocalTransform1, ParticleWorldTransform0, LocalTransform0, CullDistance, Dt, Context);
 			}
 			else if (Implicit0Type == TSphere<FReal, 3>::StaticType() && Implicit1Type == TSphere<FReal, 3>::StaticType() && !bUseGenericSweptConstraints)
 			{
-				ConstructSphereSphereConstraints<T_TRAITS>(Particle0, Particle1, Implicit0, Implicit1, LocalTransform0, LocalTransform1, CullDistance, Dt, Context, NewConstraints);
+				ConstructSphereSphereConstraints<T_TRAITS>(Particle0, Particle1, Implicit0, Implicit1, ParticleWorldTransform0, LocalTransform0, ParticleWorldTransform1, LocalTransform1, CullDistance, Dt, Context);
 			}
 			else if (Implicit0Type == TSphere<FReal, 3>::StaticType() && Implicit1Type == FHeightField::StaticType())
 			{
 				if (bUseCCD)
 				{
-					ConstructSphereHeightFieldConstraintsSwept(Particle0, Particle1, Implicit0, Implicit1, LocalTransform0, LocalTransform1, CullDistance, Dt, DirCCD, LengthCCD, NewConstraints);
+					ConstructSphereHeightFieldConstraintsSwept(Particle0, Particle1, Implicit0, Implicit1, ParticleWorldTransform0, LocalTransform0, ParticleWorldTransform1, LocalTransform1, CullDistance, Dt, DirCCD, LengthCCD, Context);
 					return;
 				}
-				ConstructSphereHeightFieldConstraints<T_TRAITS>(Particle0, Particle1, Implicit0, Implicit1, LocalTransform0, LocalTransform1, CullDistance, Dt, Context, NewConstraints);
+				ConstructSphereHeightFieldConstraints<T_TRAITS>(Particle0, Particle1, Implicit0, Implicit1, ParticleWorldTransform0, LocalTransform0, ParticleWorldTransform1, LocalTransform1, CullDistance, Dt, Context);
 			}
 			else if (Implicit0Type == FHeightField::StaticType() && Implicit1Type == TSphere<FReal, 3>::StaticType())
 			{
 				if (bUseCCD)
 				{
-					ConstructSphereHeightFieldConstraintsSwept(Particle1, Particle0, Implicit1, Implicit0, LocalTransform1, LocalTransform0, CullDistance, Dt, DirCCD, LengthCCD, NewConstraints);
+					ConstructSphereHeightFieldConstraintsSwept(Particle1, Particle0, Implicit1, Implicit0, ParticleWorldTransform1, LocalTransform1, ParticleWorldTransform0, LocalTransform0, CullDistance, Dt, DirCCD, LengthCCD, Context);
 					return;
 				}
-				ConstructSphereHeightFieldConstraints<T_TRAITS>(Particle1, Particle0, Implicit1, Implicit0, LocalTransform1, LocalTransform0, CullDistance, Dt, Context, NewConstraints);
+				ConstructSphereHeightFieldConstraints<T_TRAITS>(Particle1, Particle0, Implicit1, Implicit0, ParticleWorldTransform1, LocalTransform1, ParticleWorldTransform0, LocalTransform0, CullDistance, Dt, Context);
 			}
 			else if (Implicit0Type == TSphere<FReal, 3>::StaticType() && Implicit1Type == TPlane<FReal, 3>::StaticType() && !bUseGenericSweptConstraints)
 			{
-				ConstructSpherePlaneConstraints<T_TRAITS>(Particle0, Particle1, Implicit0, Implicit1, LocalTransform0, LocalTransform1, CullDistance, Dt, Context, NewConstraints);
+				ConstructSpherePlaneConstraints<T_TRAITS>(Particle0, Particle1, Implicit0, Implicit1, ParticleWorldTransform0, LocalTransform0, ParticleWorldTransform1, LocalTransform1, CullDistance, Dt, Context);
 			}
 			else if (Implicit0Type == TPlane<FReal, 3>::StaticType() && Implicit1Type == TSphere<FReal, 3>::StaticType() && !bUseGenericSweptConstraints)
 			{
-				ConstructSpherePlaneConstraints<T_TRAITS>(Particle1, Particle0, Implicit1, Implicit0, LocalTransform1, LocalTransform0, CullDistance, Dt, Context, NewConstraints);
+				ConstructSpherePlaneConstraints<T_TRAITS>(Particle1, Particle0, Implicit1, Implicit0, ParticleWorldTransform1, LocalTransform1, ParticleWorldTransform0, LocalTransform0, CullDistance, Dt, Context);
 			}
 			else if (Implicit0Type == TSphere<FReal, 3>::StaticType() && Implicit1Type == TBox<FReal, 3>::StaticType() && !bUseGenericSweptConstraints)
 			{
-				ConstructSphereBoxConstraints<T_TRAITS>(Particle0, Particle1, Implicit0, Implicit1, LocalTransform0, LocalTransform1, CullDistance, Dt, Context, NewConstraints);
+				ConstructSphereBoxConstraints<T_TRAITS>(Particle0, Particle1, Implicit0, Implicit1, ParticleWorldTransform0, LocalTransform0, ParticleWorldTransform1, LocalTransform1, CullDistance, Dt, Context);
 			}
 			else if (Implicit0Type == TBox<FReal, 3>::StaticType() && Implicit1Type == TSphere<FReal, 3>::StaticType() && !bUseGenericSweptConstraints)
 			{
-				ConstructSphereBoxConstraints<T_TRAITS>(Particle1, Particle0, Implicit1, Implicit0, LocalTransform1, LocalTransform0, CullDistance, Dt, Context, NewConstraints);
+				ConstructSphereBoxConstraints<T_TRAITS>(Particle1, Particle0, Implicit1, Implicit0, ParticleWorldTransform1, LocalTransform1, ParticleWorldTransform0, LocalTransform0, CullDistance, Dt, Context);
 			}
 			else if (Implicit0Type == TSphere<FReal, 3>::StaticType() && Implicit1Type == FCapsule::StaticType() && !bUseGenericSweptConstraints)
 			{
-				ConstructSphereCapsuleConstraints<T_TRAITS>(Particle0, Particle1, Implicit0, Implicit1, LocalTransform0, LocalTransform1, CullDistance, Dt, Context, NewConstraints);
+				ConstructSphereCapsuleConstraints<T_TRAITS>(Particle0, Particle1, Implicit0, Implicit1, ParticleWorldTransform0, LocalTransform0, ParticleWorldTransform1, LocalTransform1, CullDistance, Dt, Context);
 			}
 			else if (Implicit0Type == FCapsule::StaticType() && Implicit1Type == TSphere<FReal, 3>::StaticType() && !bUseGenericSweptConstraints)
 			{
-				ConstructSphereCapsuleConstraints<T_TRAITS>(Particle1, Particle0, Implicit1, Implicit0, LocalTransform1, LocalTransform0, CullDistance, Dt, Context, NewConstraints);
+				ConstructSphereCapsuleConstraints<T_TRAITS>(Particle1, Particle0, Implicit1, Implicit0, ParticleWorldTransform1, LocalTransform1, ParticleWorldTransform0, LocalTransform0, CullDistance, Dt, Context);
 			}
 			else if (Implicit0Type == TSphere<FReal, 3>::StaticType() && Implicit1Type == FImplicitConvex3::StaticType() && !bUseGenericSweptConstraints)
 			{
-				ConstructSphereConvexConstraints<T_TRAITS>(Particle0, Particle1, Implicit0, Implicit1, LocalTransform0, LocalTransform1, CullDistance, Dt, Context, NewConstraints);
+				ConstructSphereConvexConstraints<T_TRAITS>(Particle0, Particle1, Implicit0, Implicit1, ParticleWorldTransform0, LocalTransform0, ParticleWorldTransform1, LocalTransform1, CullDistance, Dt, Context);
 			}
 			else if (Implicit0Type == FImplicitConvex3::StaticType() && Implicit1Type == TSphere<FReal, 3>::StaticType() && !bUseGenericSweptConstraints)
 			{
-				ConstructSphereConvexConstraints<T_TRAITS>(Particle1, Particle0, Implicit1, Implicit0, LocalTransform1, LocalTransform0, CullDistance, Dt, Context, NewConstraints);
+				ConstructSphereConvexConstraints<T_TRAITS>(Particle1, Particle0, Implicit1, Implicit0, ParticleWorldTransform1, LocalTransform1, ParticleWorldTransform0, LocalTransform0, CullDistance, Dt, Context);
 			}
 			else if (Implicit0Type == FCapsule::StaticType() && Implicit1Type == FCapsule::StaticType() && !bUseGenericSweptConstraints)
 			{
-				ConstructCapsuleCapsuleConstraints<T_TRAITS>(Particle0, Particle1, Implicit0, Implicit1, LocalTransform0, LocalTransform1, CullDistance, Dt, Context, NewConstraints);
+				ConstructCapsuleCapsuleConstraints<T_TRAITS>(Particle0, Particle1, Implicit0, Implicit1, ParticleWorldTransform0, LocalTransform0, ParticleWorldTransform1, LocalTransform1, CullDistance, Dt, Context);
 			}
 			else if (Implicit0Type == FCapsule::StaticType() && Implicit1Type == TBox<FReal, 3>::StaticType() && !bUseGenericSweptConstraints)
 			{
-				ConstructCapsuleBoxConstraints<T_TRAITS>(Particle0, Particle1, Implicit0, Implicit1, LocalTransform0, LocalTransform1, CullDistance, Dt, Context, NewConstraints);
+				ConstructCapsuleBoxConstraints<T_TRAITS>(Particle0, Particle1, Implicit0, Implicit1, ParticleWorldTransform0, LocalTransform0, ParticleWorldTransform1, LocalTransform1, CullDistance, Dt, Context);
 			}
 			else if (Implicit0Type == TBox<FReal, 3>::StaticType() && Implicit1Type == FCapsule::StaticType() && !bUseGenericSweptConstraints)
 			{
-				ConstructCapsuleBoxConstraints<T_TRAITS>(Particle1, Particle0, Implicit1, Implicit0, LocalTransform1, LocalTransform0, CullDistance, Dt, Context, NewConstraints);
+				ConstructCapsuleBoxConstraints<T_TRAITS>(Particle1, Particle0, Implicit1, Implicit0, ParticleWorldTransform1, LocalTransform1, ParticleWorldTransform0, LocalTransform0, CullDistance, Dt, Context);
 			}
 			else if (Implicit0Type == FCapsule::StaticType() && Implicit1Type == FImplicitConvex3::StaticType() && !bUseGenericSweptConstraints)
 			{
-				ConstructCapsuleConvexConstraints<T_TRAITS>(Particle0, Particle1, Implicit0, Implicit1, LocalTransform0, LocalTransform1, CullDistance, Dt, Context, NewConstraints);
+				ConstructCapsuleConvexConstraints<T_TRAITS>(Particle0, Particle1, Implicit0, Implicit1, ParticleWorldTransform0, LocalTransform0, ParticleWorldTransform1, LocalTransform1, CullDistance, Dt, Context);
 			}
 			else if (Implicit0Type == FImplicitConvex3::StaticType() && Implicit1Type == FCapsule::StaticType() && !bUseGenericSweptConstraints)
 			{
-				ConstructCapsuleConvexConstraints<T_TRAITS>(Particle1, Particle0, Implicit1, Implicit0, LocalTransform1, LocalTransform0, CullDistance, Dt, Context, NewConstraints);
+				ConstructCapsuleConvexConstraints<T_TRAITS>(Particle1, Particle0, Implicit1, Implicit0, ParticleWorldTransform1, LocalTransform1, ParticleWorldTransform0, LocalTransform0, CullDistance, Dt, Context);
 			}
 			else if (Implicit0Type == FCapsule::StaticType() && Implicit1Type == FHeightField::StaticType())
 			{
 				if (bUseCCD)
 				{
-					ConstructCapsuleHeightFieldConstraintsSwept(Particle0, Particle1, Implicit0, Implicit1, LocalTransform0, LocalTransform1, CullDistance, Dt, DirCCD, LengthCCD, NewConstraints);
+					ConstructCapsuleHeightFieldConstraintsSwept(Particle0, Particle1, Implicit0, Implicit1, ParticleWorldTransform0, LocalTransform0, ParticleWorldTransform1, LocalTransform1, CullDistance, Dt, DirCCD, LengthCCD, Context);
 					return;
 				}
 
-				ConstructCapsuleHeightFieldConstraints<T_TRAITS>(Particle0, Particle1, Implicit0, Implicit1, LocalTransform0, LocalTransform1, CullDistance, Dt, Context, NewConstraints);
+				ConstructCapsuleHeightFieldConstraints<T_TRAITS>(Particle0, Particle1, Implicit0, Implicit1, ParticleWorldTransform0, LocalTransform0, ParticleWorldTransform1, LocalTransform1, CullDistance, Dt, Context);
 			}
 			else if (Implicit0Type == FHeightField::StaticType() && Implicit1Type == FCapsule::StaticType())
 			{
 				if (bUseCCD)
 				{
-					ConstructCapsuleHeightFieldConstraintsSwept(Particle1, Particle0, Implicit1, Implicit0, LocalTransform1, LocalTransform0, CullDistance, Dt, DirCCD, LengthCCD, NewConstraints);
+					ConstructCapsuleHeightFieldConstraintsSwept(Particle1, Particle0, Implicit1, Implicit0, ParticleWorldTransform1, LocalTransform1, ParticleWorldTransform0, LocalTransform0, CullDistance, Dt, DirCCD, LengthCCD, Context);
 					return;
 				}
 
-				ConstructCapsuleHeightFieldConstraints<T_TRAITS>(Particle1, Particle0, Implicit1, Implicit0, LocalTransform1, LocalTransform0, CullDistance, Dt, Context, NewConstraints);
+				ConstructCapsuleHeightFieldConstraints<T_TRAITS>(Particle1, Particle0, Implicit1, Implicit0, ParticleWorldTransform1, LocalTransform1, ParticleWorldTransform0, LocalTransform0, CullDistance, Dt, Context);
 			}
 			else if (Implicit0Type == TBox<FReal, 3>::StaticType() && Implicit1Type == FTriangleMeshImplicitObject::StaticType())
 			{
-				ConstructBoxTriangleMeshConstraints<T_TRAITS>(Particle0, Particle1, Implicit0, Implicit1, LocalTransform0, LocalTransform1, CullDistance, Dt, Context, NewConstraints);
+				ConstructBoxTriangleMeshConstraints<T_TRAITS>(Particle0, Particle1, Implicit0, Implicit1, ParticleWorldTransform0, LocalTransform0, ParticleWorldTransform1, LocalTransform1, CullDistance, Dt, Context);
 			}
 			else if (Implicit0Type == FTriangleMeshImplicitObject::StaticType() && Implicit1Type == TBox<FReal, 3>::StaticType())
 			{
-				ConstructBoxTriangleMeshConstraints<T_TRAITS>(Particle1, Particle0, Implicit1, Implicit0, LocalTransform1, LocalTransform0, CullDistance, Dt, Context, NewConstraints);
+				ConstructBoxTriangleMeshConstraints<T_TRAITS>(Particle1, Particle0, Implicit1, Implicit0, ParticleWorldTransform1, LocalTransform1, ParticleWorldTransform0, LocalTransform0, CullDistance, Dt, Context);
 			}
 			else if (Implicit0Type == TSphere<FReal, 3>::StaticType() && Implicit1Type == FTriangleMeshImplicitObject::StaticType())
 			{
 				if (bUseCCD)
 				{
-					ConstructSphereTriangleMeshConstraintsSwept(Particle0, Particle1, Implicit0, Implicit1, LocalTransform0, LocalTransform1, CullDistance, Dt, DirCCD, LengthCCD, NewConstraints);
+					ConstructSphereTriangleMeshConstraintsSwept(Particle0, Particle1, Implicit0, Implicit1, ParticleWorldTransform0, LocalTransform0, ParticleWorldTransform1, LocalTransform1, CullDistance, Dt, DirCCD, LengthCCD, Context);
 					return;
 				}
-				ConstructSphereTriangleMeshConstraints<T_TRAITS>(Particle0, Particle1, Implicit0, Implicit1, LocalTransform0, LocalTransform1, CullDistance, Dt, Context, NewConstraints);
+				ConstructSphereTriangleMeshConstraints<T_TRAITS>(Particle0, Particle1, Implicit0, Implicit1, ParticleWorldTransform0, LocalTransform0, ParticleWorldTransform1, LocalTransform1, CullDistance, Dt, Context);
 			}
 			else if (Implicit0Type == FTriangleMeshImplicitObject::StaticType() && Implicit1Type == TSphere<FReal, 3>::StaticType())
 			{
 				if (bUseCCD)
 				{
-					ConstructSphereTriangleMeshConstraintsSwept(Particle1, Particle0, Implicit1, Implicit0, LocalTransform1, LocalTransform0, CullDistance, Dt, DirCCD, LengthCCD, NewConstraints);
+					ConstructSphereTriangleMeshConstraintsSwept(Particle1, Particle0, Implicit1, Implicit0, ParticleWorldTransform1, LocalTransform1, ParticleWorldTransform0, LocalTransform0, CullDistance, Dt, DirCCD, LengthCCD, Context);
 					return;
 				}
-				ConstructSphereTriangleMeshConstraints<T_TRAITS>(Particle1, Particle0, Implicit1, Implicit0, LocalTransform1, LocalTransform0, CullDistance, Dt, Context, NewConstraints);
+				ConstructSphereTriangleMeshConstraints<T_TRAITS>(Particle1, Particle0, Implicit1, Implicit0, ParticleWorldTransform1, LocalTransform1, ParticleWorldTransform0, LocalTransform0, CullDistance, Dt, Context);
 			}
 			else if (Implicit0Type == FCapsule::StaticType() && Implicit1Type == FTriangleMeshImplicitObject::StaticType())
 			{
 				if (bUseCCD)
 				{
-					ConstructCapsuleTriangleMeshConstraintsSwept(Particle0, Particle1, Implicit0, Implicit1, LocalTransform0, LocalTransform1, CullDistance, Dt, DirCCD, LengthCCD, NewConstraints);
+					ConstructCapsuleTriangleMeshConstraintsSwept(Particle0, Particle1, Implicit0, Implicit1, ParticleWorldTransform0, LocalTransform0, ParticleWorldTransform1, LocalTransform1, CullDistance, Dt, DirCCD, LengthCCD, Context);
 					return;
 				}
 
-				ConstructCapsuleTriangleMeshConstraints<T_TRAITS>(Particle0, Particle1, Implicit0, Implicit1, LocalTransform0, LocalTransform1, CullDistance, Dt, Context, NewConstraints);
+				ConstructCapsuleTriangleMeshConstraints<T_TRAITS>(Particle0, Particle1, Implicit0, Implicit1, ParticleWorldTransform0, LocalTransform0, ParticleWorldTransform1, LocalTransform1, CullDistance, Dt, Context);
 			}
 			else if (Implicit0Type == FTriangleMeshImplicitObject::StaticType() && Implicit1Type == FCapsule::StaticType())
 			{
 				if (bUseCCD)
 				{
-					ConstructCapsuleTriangleMeshConstraintsSwept(Particle1, Particle0, Implicit1, Implicit0, LocalTransform1, LocalTransform0, CullDistance, Dt, DirCCD, LengthCCD, NewConstraints);
+					ConstructCapsuleTriangleMeshConstraintsSwept(Particle1, Particle0, Implicit1, Implicit0, ParticleWorldTransform1, LocalTransform1, ParticleWorldTransform0, LocalTransform0, CullDistance, Dt, DirCCD, LengthCCD, Context);
 					return;
 				}
 
-				ConstructCapsuleTriangleMeshConstraints<T_TRAITS>(Particle1, Particle0, Implicit1, Implicit0, LocalTransform1, LocalTransform0, CullDistance, Dt, Context, NewConstraints);
+				ConstructCapsuleTriangleMeshConstraints<T_TRAITS>(Particle1, Particle0, Implicit1, Implicit0, ParticleWorldTransform1, LocalTransform1, ParticleWorldTransform0, LocalTransform0, CullDistance, Dt, Context);
 			}
 			else if (bIsConvex0 && Implicit1Type == FHeightField::StaticType())
 			{
 				if (bUseCCD)
 				{
-					ConstructConvexHeightFieldConstraintsSwept(Particle0, Particle1, Implicit0, Implicit1, LocalTransform0, LocalTransform1, CullDistance, Dt, DirCCD, LengthCCD, NewConstraints);
+					ConstructConvexHeightFieldConstraintsSwept(Particle0, Particle1, Implicit0, Implicit1, ParticleWorldTransform0, LocalTransform0, ParticleWorldTransform1, LocalTransform1, CullDistance, Dt, DirCCD, LengthCCD, Context);
 					return;
 				}
 
-				ConstructConvexHeightFieldConstraints<T_TRAITS>(Particle0, Particle1, Implicit0, Implicit1, LocalTransform0, LocalTransform1, CullDistance, Dt, Context, NewConstraints);
+				ConstructConvexHeightFieldConstraints<T_TRAITS>(Particle0, Particle1, Implicit0, Implicit1, ParticleWorldTransform0, LocalTransform0, ParticleWorldTransform1, LocalTransform1, CullDistance, Dt, Context);
 			}
 			else if (Implicit0Type == FHeightField::StaticType() && bIsConvex1)
 			{
 				if (bUseCCD)
 				{
-					ConstructConvexHeightFieldConstraintsSwept(Particle1, Particle0, Implicit1, Implicit0, LocalTransform1, LocalTransform0, CullDistance, Dt, DirCCD, LengthCCD, NewConstraints);
+					ConstructConvexHeightFieldConstraintsSwept(Particle1, Particle0, Implicit1, Implicit0, ParticleWorldTransform1, LocalTransform1, ParticleWorldTransform0, LocalTransform0, CullDistance, Dt, DirCCD, LengthCCD, Context);
 					return;
 				}
 
-				ConstructConvexHeightFieldConstraints<T_TRAITS>(Particle1, Particle0, Implicit1, Implicit0, LocalTransform1, LocalTransform0, CullDistance, Dt, Context, NewConstraints);
+				ConstructConvexHeightFieldConstraints<T_TRAITS>(Particle1, Particle0, Implicit1, Implicit0, ParticleWorldTransform1, LocalTransform1, ParticleWorldTransform0, LocalTransform0, CullDistance, Dt, Context);
 			}
 			else if (bIsConvex0 && Implicit1Type == FTriangleMeshImplicitObject::StaticType())
 			{
 				if (bUseCCD)
 				{
-					ConstructConvexTriangleMeshConstraintsSwept(Particle0, Particle1, Implicit0, Implicit1, LocalTransform0, LocalTransform1, CullDistance, Dt, DirCCD, LengthCCD, NewConstraints);
+					ConstructConvexTriangleMeshConstraintsSwept(Particle0, Particle1, Implicit0, Implicit1, ParticleWorldTransform0, LocalTransform0, ParticleWorldTransform1, LocalTransform1, CullDistance, Dt, DirCCD, LengthCCD, Context);
 					return;
 				}
 
-				ConstructConvexTriangleMeshConstraints<T_TRAITS>(Particle0, Particle1, Implicit0, Implicit1, LocalTransform0, LocalTransform1, CullDistance, Dt, Context, NewConstraints);
+				ConstructConvexTriangleMeshConstraints<T_TRAITS>(Particle0, Particle1, Implicit0, Implicit1, ParticleWorldTransform0, LocalTransform0, ParticleWorldTransform1, LocalTransform1, CullDistance, Dt, Context);
 			}
 			else if (Implicit0Type == FTriangleMeshImplicitObject::StaticType() && bIsConvex1)
 			{
 				if (bUseCCD)
 				{
-					ConstructConvexTriangleMeshConstraintsSwept(Particle1, Particle0, Implicit1, Implicit0, LocalTransform1, LocalTransform0, CullDistance, Dt, DirCCD, LengthCCD, NewConstraints);
+					ConstructConvexTriangleMeshConstraintsSwept(Particle1, Particle0, Implicit1, Implicit0, ParticleWorldTransform1, LocalTransform1, ParticleWorldTransform0, LocalTransform0, CullDistance, Dt, DirCCD, LengthCCD, Context);
 					return;
 				}
 
-				ConstructConvexTriangleMeshConstraints<T_TRAITS>(Particle1, Particle0, Implicit1, Implicit0, LocalTransform1, LocalTransform0, CullDistance, Dt, Context, NewConstraints);
+				ConstructConvexTriangleMeshConstraints<T_TRAITS>(Particle1, Particle0, Implicit1, Implicit0, ParticleWorldTransform1, LocalTransform1, ParticleWorldTransform0, LocalTransform0, CullDistance, Dt, Context);
 			}
 			else if (bIsConvex0 && bIsConvex1)
 			{
 				if (bUseCCD)
 				{
-					ConstructGenericConvexConvexConstraintsSwept(Particle0, Particle1, Implicit0, Implicit1, LocalTransform0, LocalTransform1, CullDistance, Dt, DirCCD, LengthCCD, NewConstraints);
+					ConstructGenericConvexConvexConstraintsSwept(Particle0, Particle1, Implicit0, Implicit1, ParticleWorldTransform0, LocalTransform0, ParticleWorldTransform1, LocalTransform1, CullDistance, Dt, DirCCD, LengthCCD, Context);
 					return;
 				}
 
-				ConstructGenericConvexConvexConstraints<T_TRAITS>(Particle0, Particle1, Implicit0, Implicit1, LocalTransform0, LocalTransform1, CullDistance, Dt, Context, NewConstraints);
+				ConstructGenericConvexConvexConstraints<T_TRAITS>(Particle0, Particle1, Implicit0, Implicit1, ParticleWorldTransform0, LocalTransform0, ParticleWorldTransform1, LocalTransform1, CullDistance, Dt, Context);
 			}
 			else
 			{
-				ConstructLevelsetLevelsetConstraints<T_TRAITS>(Particle0, Particle1, Implicit0, Simplicial0, Implicit1, Simplicial1, LocalTransform0, LocalTransform1, CullDistance, Dt, Context, NewConstraints);
+				ConstructLevelsetLevelsetConstraints<T_TRAITS>(Particle0, Particle1, Implicit0, Simplicial0, Implicit1, Simplicial1, ParticleWorldTransform0, LocalTransform0, ParticleWorldTransform1, LocalTransform1, CullDistance, Dt, Context);
 			}
 		}
 
 
 		// Run collision detection for the specified constraint to update the nearest contact point.
 		// NOTE: Transforms are world space particle transforms
-		template<ECollisionUpdateType UpdateType, typename RigidBodyContactConstraint>
-		void UpdateConstraintFromGeometry(RigidBodyContactConstraint& Constraint, const FRigidTransform3& ParticleTransform0, const FRigidTransform3& ParticleTransform1, const FReal Dt)
+		template<ECollisionUpdateType UpdateType>
+		void UpdateConstraintFromGeometry(FPBDCollisionConstraint& Constraint, const FRigidTransform3& ParticleWorldTransform0, const FRigidTransform3& ParticleWorldTransform1, const FReal Dt)
 		{
-			const FRigidTransform3 WorldTransform0 = Constraint.ImplicitTransform[0] * ParticleTransform0;
-			const FRigidTransform3 WorldTransform1 = Constraint.ImplicitTransform[1] * ParticleTransform1;
+			const FRigidTransform3 WorldTransform0 = Constraint.ImplicitTransform[0] * ParticleWorldTransform0;
+			const FRigidTransform3 WorldTransform1 = Constraint.ImplicitTransform[1] * ParticleWorldTransform1;
 			UpdateConstraintFromGeometryImpl<UpdateType>(Constraint, WorldTransform0, WorldTransform1, Dt);
 		}
 
+		template<ECollisionUpdateType UpdateType>
+		void UpdateConstraintFromGeometrySwept(FPBDCollisionConstraint& Constraint, const FRigidTransform3& ParticleWorldTransform0, const FRigidTransform3& ParticleWorldTransform1, const FReal Dt)
+		{
+			const FRigidTransform3 WorldTransform0 = Constraint.ImplicitTransform[0] * ParticleWorldTransform0;
+			const FRigidTransform3 WorldTransform1 = Constraint.ImplicitTransform[1] * ParticleWorldTransform1;
+			UpdateConstraintFromGeometrySweptImpl<UpdateType>(Constraint, WorldTransform0, WorldTransform1, Dt);
+		}
+
 
 		typedef uint8 FMaskFilter;
 		enum { NumExtraFilterBits = 6 };
@@ -2792,7 +2619,7 @@ namespace Chaos
 
 
 		template<typename T_TRAITS>
-		void ConstructConstraints(TGeometryParticleHandle<FReal, 3>* Particle0, TGeometryParticleHandle<FReal, 3>* Particle1, const FImplicitObject* Implicit0, const FBVHParticles* Simplicial0, const FImplicitObject* Implicit1, const FBVHParticles* Simplicial1, const FRigidTransform3& LocalTransform0, const FRigidTransform3& LocalTransform1, const FReal CullDistance, FReal Dt, const FCollisionContext& Context, FCollisionConstraintsArray& NewConstraints)
+		void ConstructConstraints(TGeometryParticleHandle<FReal, 3>* Particle0, TGeometryParticleHandle<FReal, 3>* Particle1, const FImplicitObject* Implicit0, const FBVHParticles* Simplicial0, const FImplicitObject* Implicit1, const FBVHParticles* Simplicial1, const FRigidTransform3& ParticleWorldTransform0, const FRigidTransform3& LocalTransform0, const FRigidTransform3& ParticleWorldTransform1, const FRigidTransform3& LocalTransform1, const FReal CullDistance, FReal Dt, const FCollisionContext& Context)
 		{
 			CONDITIONAL_SCOPE_CYCLE_COUNTER(STAT_Collisions_ConstructConstraints, ConstraintsDetailedStats);
 
@@ -2802,7 +2629,7 @@ namespace Chaos
 
 			if (!Implicit0 || !Implicit1)
 			{
-				ConstructLevelsetLevelsetConstraints<T_TRAITS>(Particle0, Particle1, Implicit0, Simplicial0, Implicit1, Simplicial1, LocalTransform0, LocalTransform1, CullDistance, Dt, Context, NewConstraints);
+				ConstructLevelsetLevelsetConstraints<T_TRAITS>(Particle0, Particle1, Implicit0, Simplicial0, Implicit1, Simplicial1, ParticleWorldTransform0, LocalTransform0, ParticleWorldTransform1, LocalTransform1, CullDistance, Dt, Context);
 				return;
 			}
 
@@ -2816,21 +2643,21 @@ namespace Chaos
 				const TImplicitObjectTransformed<FReal, 3>* TransformedImplicit1 = Implicit1->template GetObject<const TImplicitObjectTransformed<FReal, 3>>();
 				FRigidTransform3 TransformedTransform0 = TransformedImplicit0->GetTransform() * LocalTransform0;
 				FRigidTransform3 TransformedTransform1 = TransformedImplicit1->GetTransform() * LocalTransform1;
-				ConstructConstraints<T_TRAITS>(Particle0, Particle1, TransformedImplicit0->GetTransformedObject(), Simplicial0, TransformedImplicit1->GetTransformedObject(), Simplicial1, TransformedTransform0, TransformedTransform1, CullDistance, Dt, Context, NewConstraints);
+				ConstructConstraints<T_TRAITS>(Particle0, Particle1, TransformedImplicit0->GetTransformedObject(), Simplicial0, TransformedImplicit1->GetTransformedObject(), Simplicial1, ParticleWorldTransform0, TransformedTransform0, ParticleWorldTransform1, TransformedTransform1, CullDistance, Dt, Context);
 				return;
 			}
 			else if (Implicit0OuterType == TImplicitObjectTransformed<FReal, 3>::StaticType())
 			{
 				const TImplicitObjectTransformed<FReal, 3>* TransformedImplicit0 = Implicit0->template GetObject<const TImplicitObjectTransformed<FReal, 3>>();
 				FRigidTransform3 TransformedTransform0 = TransformedImplicit0->GetTransform() * LocalTransform0;
-				ConstructConstraints<T_TRAITS>(Particle0, Particle1, TransformedImplicit0->GetTransformedObject(), Simplicial0, Implicit1, Simplicial1, TransformedTransform0, LocalTransform1, CullDistance, Dt, Context, NewConstraints);
+				ConstructConstraints<T_TRAITS>(Particle0, Particle1, TransformedImplicit0->GetTransformedObject(), Simplicial0, Implicit1, Simplicial1, ParticleWorldTransform0, TransformedTransform0, ParticleWorldTransform1, LocalTransform1, CullDistance, Dt, Context);
 				return;
 			}
 			else if (Implicit1OuterType == TImplicitObjectTransformed<FReal, 3>::StaticType())
 			{
 				const TImplicitObjectTransformed<FReal, 3>* TransformedImplicit1 = Implicit1->template GetObject<const TImplicitObjectTransformed<FReal, 3>>();
 				FRigidTransform3 TransformedTransform1 = TransformedImplicit1->GetTransform() * LocalTransform1;
-				ConstructConstraints<T_TRAITS>(Particle0, Particle1, Implicit0, Simplicial0, TransformedImplicit1->GetTransformedObject(), Simplicial1, LocalTransform0, TransformedTransform1, CullDistance, Dt, Context, NewConstraints);
+				ConstructConstraints<T_TRAITS>(Particle0, Particle1, Implicit0, Simplicial0, TransformedImplicit1->GetTransformedObject(), Simplicial1, ParticleWorldTransform0, LocalTransform0, ParticleWorldTransform1, TransformedTransform1, CullDistance, Dt, Context);
 				return;
 			}
 
@@ -2851,17 +2678,17 @@ namespace Chaos
 				}
 				if (InnerImplicit0 && InnerImplicit1)
 				{
-					ConstructConstraints<T_TRAITS>(Particle0, Particle1, InnerImplicit0, Simplicial0, InnerImplicit1, Simplicial1, LocalTransform0, LocalTransform1, CullDistance, Dt, Context, NewConstraints);
+					ConstructConstraints<T_TRAITS>(Particle0, Particle1, InnerImplicit0, Simplicial0, InnerImplicit1, Simplicial1, ParticleWorldTransform0, LocalTransform0, ParticleWorldTransform1, LocalTransform1, CullDistance, Dt, Context);
 					return;
 				}
 				else if (InnerImplicit0 && !InnerImplicit1)
 				{
-					ConstructConstraints<T_TRAITS>(Particle0, Particle1, InnerImplicit0, Simplicial0, Implicit1, Simplicial1, LocalTransform0, LocalTransform1, CullDistance, Dt, Context, NewConstraints);
+					ConstructConstraints<T_TRAITS>(Particle0, Particle1, InnerImplicit0, Simplicial0, Implicit1, Simplicial1, ParticleWorldTransform0, LocalTransform0, ParticleWorldTransform1, LocalTransform1, CullDistance, Dt, Context);
 					return;
 				}
 				else if (!InnerImplicit0 && InnerImplicit1)
 				{
-					ConstructConstraints<T_TRAITS>(Particle0, Particle1, Implicit0, Simplicial0, InnerImplicit1, Simplicial1, LocalTransform0, LocalTransform1, CullDistance, Dt, Context, NewConstraints);
+					ConstructConstraints<T_TRAITS>(Particle0, Particle1, Implicit0, Simplicial0, InnerImplicit1, Simplicial1, ParticleWorldTransform0, LocalTransform0, ParticleWorldTransform1, LocalTransform1, CullDistance, Dt, Context);
 					return;
 				}
 
@@ -2876,7 +2703,7 @@ namespace Chaos
 					// If shape is not sim'd, we may end up iterating over a lot of shapes on particle1's union and wasting time filtering.
 					if (Context.bFilteringEnabled == false || HasSimEnabled(Particle0->GetImplicitShape(Child0.Get())))
 					{
-						ConstructConstraints<T_TRAITS>(Particle0, Particle1, Child0.Get(), Simplicial0, Implicit1, Simplicial1, LocalTransform0, LocalTransform1, CullDistance, Dt, Context, NewConstraints);
+						ConstructConstraints<T_TRAITS>(Particle0, Particle1, Child0.Get(), Simplicial0, Implicit1, Simplicial1, ParticleWorldTransform0, LocalTransform0, ParticleWorldTransform1, LocalTransform1, CullDistance, Dt, Context);
 					}
 				}
 				return;
@@ -2890,8 +2717,8 @@ namespace Chaos
 					TArray<Pair<Pair<const FImplicitObject*, const FBVHParticles*>, FRigidTransform3>> Children;
 
 					// Need to get transformed bounds of 1 in the space of 0
-					FRigidTransform3 TM0 = LocalTransform0 * Collisions::GetTransform(Particle0);
-					FRigidTransform3 TM1 = LocalTransform1 * Collisions::GetTransform(Particle1);
+					FRigidTransform3 TM0 = LocalTransform0 * ParticleWorldTransform0;
+					FRigidTransform3 TM1 = LocalTransform1 * ParticleWorldTransform1;
 					FRigidTransform3 TM1ToTM0 = TM1.GetRelativeTransform(TM0);
 					FAABB3 QueryBounds = Implicit1->BoundingBox().TransformedAABB(TM1ToTM0);
 
@@ -2899,7 +2726,7 @@ namespace Chaos
 
 					for (const Pair<Pair<const FImplicitObject*, const FBVHParticles*>, FRigidTransform3>& Child0 : Children)
 					{
-						ConstructConstraints<T_TRAITS>(Particle0, Particle1, Child0.First.First, Child0.First.Second, Implicit1, Simplicial1, Child0.Second * LocalTransform0, LocalTransform1, CullDistance, Dt, Context, NewConstraints);
+						ConstructConstraints<T_TRAITS>(Particle0, Particle1, Child0.First.First, Child0.First.Second, Implicit1, Simplicial1, ParticleWorldTransform0, Child0.Second * LocalTransform0, ParticleWorldTransform0, LocalTransform1, CullDistance, Dt, Context);
 					}
 				}
 				else
@@ -2907,7 +2734,7 @@ namespace Chaos
 					for (const TUniquePtr<FImplicitObject>& Child0 : Union0->GetObjects())
 					{
 						const TPBDRigidParticleHandle<FReal, 3>* OriginalParticle = Union0->FindParticleForImplicitObject(Child0.Get());
-						ConstructConstraints<T_TRAITS>(Particle0, Particle1, Child0.Get(), OriginalParticle ? OriginalParticle->CollisionParticles().Get() : Simplicial0, Implicit1, Simplicial1, LocalTransform0, LocalTransform1, CullDistance, Dt, Context, NewConstraints);
+						ConstructConstraints<T_TRAITS>(Particle0, Particle1, Child0.Get(), OriginalParticle ? OriginalParticle->CollisionParticles().Get() : Simplicial0, Implicit1, Simplicial1, ParticleWorldTransform0, LocalTransform0, ParticleWorldTransform1, LocalTransform1, CullDistance, Dt, Context);
 					}
 				}
 				return;
@@ -2921,7 +2748,7 @@ namespace Chaos
 					// If shape is not sim'd, we may end up iterating over a lot of shapes on particle1's union and wasting time filtering.
 					if (Context.bFilteringEnabled == false || HasSimEnabled(Particle1->GetImplicitShape(Child1.Get())))
 					{
-						ConstructConstraints<T_TRAITS>(Particle0, Particle1, Implicit0, Simplicial0, Child1.Get(), Simplicial1, LocalTransform0, LocalTransform1, CullDistance, Dt, Context, NewConstraints);
+						ConstructConstraints<T_TRAITS>(Particle0, Particle1, Implicit0, Simplicial0, Child1.Get(), Simplicial1, ParticleWorldTransform0, LocalTransform0, ParticleWorldTransform1, LocalTransform1, CullDistance, Dt, Context);
 					}
 				}
 				return;
@@ -2935,8 +2762,8 @@ namespace Chaos
 					TArray<Pair<Pair<const FImplicitObject*, const FBVHParticles*>, FRigidTransform3>> Children;
 
 					// Need to get transformed bounds of 0 in the space of 1
-					FRigidTransform3 TM0 = LocalTransform0 * Collisions::GetTransform(Particle0);
-					FRigidTransform3 TM1 = LocalTransform1 * Collisions::GetTransform(Particle1);
+					FRigidTransform3 TM0 = LocalTransform0 * ParticleWorldTransform0;
+					FRigidTransform3 TM1 = LocalTransform1 * ParticleWorldTransform1;
 					FRigidTransform3 TM0ToTM1 = TM0.GetRelativeTransform(TM1);
 					FAABB3 QueryBounds = Implicit0->BoundingBox().TransformedAABB(TM0ToTM1);
 
@@ -2947,7 +2774,7 @@ namespace Chaos
 
 					for (const Pair<Pair<const FImplicitObject*, const FBVHParticles*>, FRigidTransform3>& Child1 : Children)
 					{
-						ConstructConstraints<T_TRAITS>(Particle0, Particle1, Implicit0, Simplicial0, Child1.First.First, Child1.First.Second, LocalTransform0, Child1.Second * LocalTransform1, CullDistance, Dt, Context, NewConstraints);
+						ConstructConstraints<T_TRAITS>(Particle0, Particle1, Implicit0, Simplicial0, Child1.First.First, Child1.First.Second, ParticleWorldTransform0, LocalTransform0, ParticleWorldTransform1, Child1.Second * LocalTransform1, CullDistance, Dt, Context);
 					}
 				}
 				else
@@ -2955,7 +2782,7 @@ namespace Chaos
 					for (const TUniquePtr<FImplicitObject>& Child1 : Union1->GetObjects())
 					{
 						const TPBDRigidParticleHandle<FReal, 3>* OriginalParticle = Union1->FindParticleForImplicitObject(Child1.Get());
-						ConstructConstraints<T_TRAITS>(Particle0, Particle1, Implicit0, Simplicial0, Child1.Get(), OriginalParticle ? OriginalParticle->CollisionParticles().Get() : Simplicial1, LocalTransform0, LocalTransform1, CullDistance, Dt, Context, NewConstraints);
+						ConstructConstraints<T_TRAITS>(Particle0, Particle1, Implicit0, Simplicial0, Child1.Get(), OriginalParticle ? OriginalParticle->CollisionParticles().Get() : Simplicial1, ParticleWorldTransform0, LocalTransform0, ParticleWorldTransform1, LocalTransform1, CullDistance, Dt, Context);
 					}
 				}
 				return;
@@ -2969,10 +2796,10 @@ namespace Chaos
 
 			// If we get here, we have a pair of concrete shapes (i.e., no wrappers or containers)
 			// Create a constraint for the shape pair
-			ConstructConstraintsImpl<T_TRAITS>(Particle0, Particle1, Implicit0, Simplicial0, Implicit1, Simplicial1, LocalTransform0, LocalTransform1, CullDistance, Dt, Context, NewConstraints);
+			ConstructConstraintsImpl<T_TRAITS>(Particle0, Particle1, Implicit0, Simplicial0, Implicit1, Simplicial1, ParticleWorldTransform0, LocalTransform0, ParticleWorldTransform1, LocalTransform1, CullDistance, Dt, Context);
 		}
 
-		void ConstructConstraints(TGeometryParticleHandle<FReal, 3>* Particle0, TGeometryParticleHandle<FReal, 3>* Particle1, const FImplicitObject* Implicit0, const FBVHParticles* Simplicial0, const FImplicitObject* Implicit1, const FBVHParticles* Simplicial1, const FRigidTransform3& LocalTransform0, const FRigidTransform3& LocalTransform1, const FReal CullDistance, const FReal dT, const FCollisionContext& Context, FCollisionConstraintsArray& NewConstraints)
+		void ConstructConstraints(TGeometryParticleHandle<FReal, 3>* Particle0, TGeometryParticleHandle<FReal, 3>* Particle1, const FImplicitObject* Implicit0, const FBVHParticles* Simplicial0, const FImplicitObject* Implicit1, const FBVHParticles* Simplicial1, const FRigidTransform3& ParticleWorldTransform0, const FRigidTransform3& LocalTransform0, const FRigidTransform3& ParticleWorldTransform1, const FRigidTransform3& LocalTransform1, const FReal CullDistance, const FReal dT, const FCollisionContext& Context)
 		{
 			INC_DWORD_STAT(STAT_ChaosCollisionCounter_NumParticlePairs);
 
@@ -2985,24 +2812,92 @@ namespace Chaos
 			if (bDeferUpdate)
 			{
 				using TTraits = TConstructCollisionTraits<false>;
-				ConstructConstraints<TTraits>(Particle0, Particle1, Implicit0, Simplicial0, Implicit1, Simplicial1, LocalTransform0, LocalTransform1, CullDistance, dT, Context, NewConstraints);
+				ConstructConstraints<TTraits>(Particle0, Particle1, Implicit0, Simplicial0, Implicit1, Simplicial1, ParticleWorldTransform0, LocalTransform0, ParticleWorldTransform1, LocalTransform1, CullDistance, dT, Context);
 			}
 			else
 			{
 				using TTraits = TConstructCollisionTraits<true>;
-				ConstructConstraints<TTraits>(Particle0, Particle1, Implicit0, Simplicial0, Implicit1, Simplicial1, LocalTransform0, LocalTransform1, CullDistance, dT, Context, NewConstraints);
+				ConstructConstraints<TTraits>(Particle0, Particle1, Implicit0, Simplicial0, Implicit1, Simplicial1, ParticleWorldTransform0, LocalTransform0, ParticleWorldTransform1, LocalTransform1, CullDistance, dT, Context);
+			}
+		}
+
+
+		bool IsWithinManifoldRestorationThreshold(const FRigidTransform3& PrevParticleTransform0, const FRigidTransform3& ParticleTransform0, const FRigidTransform3& PrevParticleTransform1, const FRigidTransform3& ParticleTransform1)
+		{
+			// @todo(chaos): faster manifold reuse test. Use velocities first? Use approximation for small angles to avoid ACos
+
+			const FRigidTransform3 Transform1to0 = ParticleTransform1.GetRelativeTransformNoScale(ParticleTransform0);
+			const FRigidTransform3 PrevTransform1to0 = PrevParticleTransform1.GetRelativeTransformNoScale(PrevParticleTransform0);
+
+			const FReal PositionToleranceCm = Chaos_Collision_ParticlePositionTolerance;
+			const FReal RotationToleranceDeg = Chaos_Collision_ParticleRotationTolerance;
+
+			const FReal DPSq = (Transform1to0.GetTranslation() - PrevTransform1to0.GetTranslation()).SizeSquared();
+			if (DPSq < FMath::Square(PositionToleranceCm))
+			{
+				const FReal DR = FMath::RadiansToDegrees(Transform1to0.GetRotation().AngularDistance(PrevTransform1to0.GetRotation()));
+				if (DR < RotationToleranceDeg)
+				{
+					return true;
+				}
 			}
+
+			return false;
 		}
 
+		/**
+		 * @brief Attempt to restore all the contacts between a pair of particles
+		 * 
+		 * This will only succeed if the particles have not moved much since the contacts were created.
+		 * Note that a particle pair may each have multiple collisions shapes, and therefore multiple 
+		 * contact constraints between them (contact constraints are between shape pairs).
+		*/
+		bool TryRestoreConstraints(FConstGenericParticleHandle Particle0, FConstGenericParticleHandle Particle1, const FRigidTransform3& ParticleWorldTransform0, const FRigidTransform3& ParticleWorldTransform1, const FReal Dt, const FCollisionContext& Context)
+		{
+			// @todo(chaos): optimize the restorations test:
+			// - if the particles are in diffrent islands there will definitely be no previous contact to restore
+			// - check relative velocity and relative angular velocity
+
+			// See if we have any historical data for this pair
+			FParticlePairCollisionConstraints* PairContacts = Context.CollisionAllocator->FindParticlePairConstraints(Particle0->Handle(), Particle1->Handle());
+			if (PairContacts != nullptr)
+			{
+				// @todo(chaos): eventually we shouldn't even be calling this function for sleeping pairs
+				const bool bIsAsleep0 = (Particle0->IsDynamic() && Particle0->Sleeping()) || (Particle0->IsKinematic() && !Particle0->KinematicTarget().IsSet()) || Particle0->IsStatic();
+				const bool bIsAsleep1 = (Particle1->IsDynamic() && Particle1->Sleeping()) || (Particle1->IsKinematic() && !Particle1->KinematicTarget().IsSet()) || Particle1->IsStatic();
+
+				// See if the particles have the same relative transform and can have their constraint fully restored
+				const bool bCanRestore = (bIsAsleep0 && bIsAsleep1) || IsWithinManifoldRestorationThreshold(PairContacts->GetParticle0Transform(), ParticleWorldTransform0, PairContacts->GetParticle1Transform(), ParticleWorldTransform1);
+
+				if (bCanRestore)
+				{
+					// Particle have not moved so restore the contacts (if present) and skip the narrow phase
+					// Note: even if there are no contacts because actual separation is greater than cull distance even though bounds overlap, we still skip the narrow phase
+					// and avoid the additional shape-pair bounds checks etc.
+					Context.CollisionAllocator->RestoreParticlePairConstraints(PairContacts);
+					return true;
+				}
+				else
+				{
+					// Particles have moved and we must re-run the narrow phase, so reset the active contact list (it will be rebuilt if the narrow phase adds contacts)
+					// This also moves the "current" contact set to the "previous" for use in static friction
+					Context.CollisionAllocator->ResetParticlePairConstraints(PairContacts, ParticleWorldTransform0, ParticleWorldTransform1);
+					return false;
+				}
+			}
+			return false;
+		}
+
+
 
-		template void UpdateLevelsetLevelsetConstraint<ECollisionUpdateType::Any>(const FRigidTransform3& WorldTransform0, const FRigidTransform3& WorldTransform1, const FReal Dt, FRigidBodyPointContactConstraint& Constraint);
-		template void UpdateLevelsetLevelsetConstraint<ECollisionUpdateType::Deepest>(const FRigidTransform3& WorldTransform0, const FRigidTransform3& WorldTransform1, const FReal Dt, FRigidBodyPointContactConstraint& Constraint);
+		template void UpdateLevelsetLevelsetConstraint<ECollisionUpdateType::Any>(const FRigidTransform3& WorldTransform0, const FRigidTransform3& WorldTransform1, const FReal Dt, FPBDCollisionConstraint& Constraint);
+		template void UpdateLevelsetLevelsetConstraint<ECollisionUpdateType::Deepest>(const FRigidTransform3& WorldTransform0, const FRigidTransform3& WorldTransform1, const FReal Dt, FPBDCollisionConstraint& Constraint);
 
-		template void UpdateConstraintFromGeometry<ECollisionUpdateType::Any, FRigidBodyPointContactConstraint>(FRigidBodyPointContactConstraint& ConstraintBase, const FRigidTransform3& Transform0, const FRigidTransform3& Transform1, const FReal Dt);
-		template void UpdateConstraintFromGeometry<ECollisionUpdateType::Deepest, FRigidBodyPointContactConstraint>(FRigidBodyPointContactConstraint& ConstraintBase, const FRigidTransform3& Transform0, const FRigidTransform3& Transform1, const FReal Dt);
+		template void UpdateConstraintFromGeometry<ECollisionUpdateType::Any>(FPBDCollisionConstraint& ConstraintBase, const FRigidTransform3& Transform0, const FRigidTransform3& Transform1, const FReal Dt);
+		template void UpdateConstraintFromGeometry<ECollisionUpdateType::Deepest>(FPBDCollisionConstraint& ConstraintBase, const FRigidTransform3& Transform0, const FRigidTransform3& Transform1, const FReal Dt);
 
-		template void UpdateConstraintFromGeometry<ECollisionUpdateType::Any, FRigidBodySweptPointContactConstraint>(FRigidBodySweptPointContactConstraint& ConstraintBase, const FRigidTransform3& Transform0, const FRigidTransform3& Transform1, const FReal Dt);
-		template void UpdateConstraintFromGeometry<ECollisionUpdateType::Deepest, FRigidBodySweptPointContactConstraint>(FRigidBodySweptPointContactConstraint& ConstraintBase, const FRigidTransform3& Transform0, const FRigidTransform3& Transform1, const FReal Dt);
+		template void UpdateConstraintFromGeometrySwept<ECollisionUpdateType::Any>(FPBDCollisionConstraint& ConstraintBase, const FRigidTransform3& Transform0, const FRigidTransform3& Transform1, const FReal Dt);
+		template void UpdateConstraintFromGeometrySwept<ECollisionUpdateType::Deepest>(FPBDCollisionConstraint& ConstraintBase, const FRigidTransform3& Transform0, const FRigidTransform3& Transform1, const FReal Dt);
 
 	} // Collisions
 
