Chaos

diff --git a/Engine/Source/Runtime/Experimental/Chaos/Private/Chaos/Evolution/PBDMinEvolution.cpp b/Engine/Source/Runtime/Experimental/Chaos/Private/Chaos/Evolution/PBDMinEvolution.cpp
--- a/Engine/Source/Runtime/Experimental/Chaos/Private/Chaos/Evolution/PBDMinEvolution.cpp
+++ b/Engine/Source/Runtime/Experimental/Chaos/Private/Chaos/Evolution/PBDMinEvolution.cpp
@@ -39,11 +39,11 @@ namespace Chaos
 	DECLARE_CYCLE_STAT(TEXT("MinEvolution::AdvanceOneTimeStep"), STAT_MinEvolution_AdvanceOneTimeStep, STATGROUP_ChaosMinEvolution);
 	DECLARE_CYCLE_STAT(TEXT("MinEvolution::Integrate"), STAT_MinEvolution_Integrate, STATGROUP_ChaosMinEvolution);
 	DECLARE_CYCLE_STAT(TEXT("MinEvolution::KinematicTargets"), STAT_MinEvolution_KinematicTargets, STATGROUP_ChaosMinEvolution);
-	DECLARE_CYCLE_STAT(TEXT("MinEvolution::PrepareIteration"), STAT_MinEvolution_PrepareIteration, STATGROUP_ChaosMinEvolution);
-	DECLARE_CYCLE_STAT(TEXT("MinEvolution::UnprepareIteration"), STAT_MinEvolution_UnprepareIteration, STATGROUP_ChaosMinEvolution);
-	DECLARE_CYCLE_STAT(TEXT("MinEvolution::ApplyConstraints"), STAT_MinEvolution_ApplyConstraints, STATGROUP_ChaosMinEvolution);
+	DECLARE_CYCLE_STAT(TEXT("MinEvolution::Gather"), STAT_MinEvolution_Gather, STATGROUP_ChaosMinEvolution);
+	DECLARE_CYCLE_STAT(TEXT("MinEvolution::Scatter"), STAT_MinEvolution_Scatter, STATGROUP_ChaosMinEvolution);
+	DECLARE_CYCLE_STAT(TEXT("MinEvolution::ApplyConstraintsPhase1"), STAT_MinEvolution_ApplyConstraintsPhase1, STATGROUP_ChaosMinEvolution);
 	DECLARE_CYCLE_STAT(TEXT("MinEvolution::UpdateVelocities"), STAT_MinEvolution_UpdateVelocites, STATGROUP_ChaosMinEvolution);
-	DECLARE_CYCLE_STAT(TEXT("MinEvolution::ApplyPushOut"), STAT_MinEvolution_ApplyPushOut, STATGROUP_ChaosMinEvolution);
+	DECLARE_CYCLE_STAT(TEXT("MinEvolution::ApplyConstraintsPhase2"), STAT_MinEvolution_ApplyConstraintsPhase2, STATGROUP_ChaosMinEvolution);
 	DECLARE_CYCLE_STAT(TEXT("MinEvolution::DetectCollisions"), STAT_MinEvolution_DetectCollisions, STATGROUP_ChaosMinEvolution);
 	DECLARE_CYCLE_STAT(TEXT("MinEvolution::UpdatePositions"), STAT_MinEvolution_UpdatePositions, STATGROUP_ChaosMinEvolution);
 
@@ -163,7 +163,12 @@ namespace Chaos
 
 	void FPBDMinEvolution::AddConstraintRule(FSimpleConstraintRule* Rule)
 	{
-		ConstraintRules.Add(Rule);
+		check(Rule != nullptr);
+		if (Rule != nullptr)
+		{
+			const uint32 ContainerId = (uint32)ConstraintRules.Add(Rule);
+			Rule->BindToDatas(SolverData, ContainerId);
+		}
 	}
 
 	void FPBDMinEvolution::Advance(const FReal StepDt, const int32 NumSteps, const FReal RewindDt)
@@ -222,9 +227,9 @@ namespace Chaos
 
 		if (Dt > 0)
 		{
-			PrepareIteration(Dt);
+			GatherInput(Dt);
 
-			ApplyConstraints(Dt);
+			ApplyConstraintsPhase1(Dt);
 
 			if (PostApplyCallback != nullptr)
 			{
@@ -233,14 +238,14 @@ namespace Chaos
 
 			UpdateVelocities(Dt);
 
-			ApplyPushOutConstraints(Dt);
+			ApplyConstraintsPhase2(Dt);
 
 			if (PostApplyPushOutCallback != nullptr)
 			{
 				PostApplyPushOutCallback();
 			}
 
-			UnprepareIteration(Dt);
+			ScatterOutput(Dt);
 
 			UpdatePositions(Dt);
 		}
@@ -649,32 +654,36 @@ namespace Chaos
 			ConstraintRule->UpdatePositionBasedState(Dt);
 		}
 
-		CollisionDetector.DetectCollisions(Dt);
+		CollisionDetector.DetectCollisions(Dt, nullptr);
 	}
 
-	void FPBDMinEvolution::PrepareIteration(FReal Dt)
+	void FPBDMinEvolution::GatherInput(FReal Dt)
 	{
-		SCOPE_CYCLE_COUNTER(STAT_MinEvolution_PrepareIteration);
+		SCOPE_CYCLE_COUNTER(STAT_MinEvolution_Gather);
+
+		SolverData.GetBodyContainer().Reset(Particles.GetAllParticlesView().Num());
 
 		for (FSimpleConstraintRule* ConstraintRule : ConstraintRules)
 		{
-			ConstraintRule->PrepareIteration(Dt);
+			ConstraintRule->GatherSolverInput(Dt);
 		}
 	}
 
-	void FPBDMinEvolution::UnprepareIteration(FReal Dt)
+	void FPBDMinEvolution::ScatterOutput(FReal Dt)
 	{
-		SCOPE_CYCLE_COUNTER(STAT_MinEvolution_UnprepareIteration);
+		SCOPE_CYCLE_COUNTER(STAT_MinEvolution_Scatter);
 
 		for (FSimpleConstraintRule* ConstraintRule : ConstraintRules)
 		{
-			ConstraintRule->UnprepareIteration(Dt);
+			ConstraintRule->ScatterSolverOutput(Dt);
 		}
+
+		SolverData.GetBodyContainer().ScatterOutput();
 	}
 
-	void FPBDMinEvolution::ApplyConstraints(FReal Dt)
+	void FPBDMinEvolution::ApplyConstraintsPhase1(FReal Dt)
 	{
-		SCOPE_CYCLE_COUNTER(STAT_MinEvolution_ApplyConstraints);
+		SCOPE_CYCLE_COUNTER(STAT_MinEvolution_ApplyConstraintsPhase1);
 
 		for (int32 i = 0; i < NumApplyIterations; ++i)
 		{
@@ -695,16 +704,12 @@ namespace Chaos
 	{
 		SCOPE_CYCLE_COUNTER(STAT_MinEvolution_UpdateVelocites);
 
-		FPerParticlePBDUpdateFromDeltaPosition UpdateVelocityRule;
-		for (auto& Particle : Particles.GetActiveParticlesView())
-		{
-			UpdateVelocityRule.Apply(Particle, Dt);
-		}
+		SolverData.GetBodyContainer().SetImplicitVelocities(Dt);
 	}
 
-	void FPBDMinEvolution::ApplyPushOutConstraints(FReal Dt)
+	void FPBDMinEvolution::ApplyConstraintsPhase2(FReal Dt)
 	{
-		SCOPE_CYCLE_COUNTER(STAT_MinEvolution_ApplyPushOut);
+		SCOPE_CYCLE_COUNTER(STAT_MinEvolution_ApplyConstraintsPhase2);
 
 		for (int32 It = 0; It < NumApplyPushOutIterations; ++It)
 		{
diff --git a/Engine/Source/Runtime/Experimental/Chaos/Private/Chaos/Evolution/SolverBody.cpp b/Engine/Source/Runtime/Experimental/Chaos/Private/Chaos/Evolution/SolverBody.cpp
new file mode 100644
--- /dev/null
+++ b/Engine/Source/Runtime/Experimental/Chaos/Private/Chaos/Evolution/SolverBody.cpp
@@ -0,0 +1,83 @@
+// Copyright Epic Games, Inc. All Rights Reserved.
+#include "Chaos/Evolution/SolverBody.h"
+#include "Chaos/Particle/ParticleUtilities.h"
+
+namespace Chaos
+{
+
+	FSolverBody::FSolverBody()
+		//: Particle(nullptr)
+	{
+	}
+
+	//FSolverBody::FSolverBody(FGenericParticleHandle InParticle)
+	//	//: Particle(InParticle)
+	//{
+	//	//GatherInput();
+
+	//	// This will get set to something lower if there is a contact, based on the contact graph
+	//	State.Level = TNumericLimits<int32>::Max();
+	//}
+
+	//void FSolverBody::GatherInput()
+	//{
+	//	if (Particle.IsValid())
+	//	{
+	//		FRigidTransform3 CoMTransform = FParticleUtilitiesPQ::GetCoMWorldTransform(Particle);
+	//		State.P = CoMTransform.GetLocation();
+	//		State.Q = CoMTransform.GetRotation();
+	//		State.V = Particle->V();
+	//		State.W = Particle->W();
+
+	//		if (Particle->IsDynamic())
+	//		{
+	//			FRigidTransform3 PrevCoMTransform = FParticleUtilitiesXR::GetCoMWorldTransform(Particle);
+	//			State.X = PrevCoMTransform.GetLocation();
+	//			State.R = PrevCoMTransform.GetRotation();
+
+	//			State.InvM = Particle->InvM();
+	//			State.InvILocal = Particle->InvI().GetDiagonal();
+
+	//			// This will get set to something lower if there is a contact, based on the contact graph
+	//			State.Level = TNumericLimits<int32>::Max();
+	//		}
+	//		else
+	//		{
+	//			State.X = State.P;
+	//			State.R = State.Q;
+	//		}
+
+	//		UpdateRotationDependentState();
+	//	}
+	//}
+
+	//void FSolverBody::ScatterOutput()
+	//{
+	//	if (Particle.IsValid())
+	//	{
+	//		if (IsDynamic())
+	//		{
+	//			// Set the particle state
+	//			FParticleUtilities::SetCoMWorldTransform(Particle, State.P, State.Q);
+	//			Particle->SetV(State.V);
+	//			Particle->SetW(State.W);
+
+	//			if (State.bHasActiveCollision)
+	//			{
+	//				//Particle->AuxilaryValue(*ParticleParameters.Collided) = true;
+	//			}
+	//		}
+
+	//		// Reset SolverBodyIndex cookie every step - it will be reassigned next step
+	//		Particle->SetSolverBodyIndex(INDEX_NONE);
+	//	}
+	//}
+
+	void FSolverBody::UpdateRotationDependentState()
+	{
+		if (IsDynamic())
+		{
+			State.InvI = Utilities::ComputeWorldSpaceInertia(State.Q, State.InvILocal);
+		}
+	}
+}
\ No newline at end of file
diff --git a/Engine/Source/Runtime/Experimental/Chaos/Private/Chaos/Evolution/SolverBodyContainer.cpp b/Engine/Source/Runtime/Experimental/Chaos/Private/Chaos/Evolution/SolverBodyContainer.cpp
new file mode 100644
--- /dev/null
+++ b/Engine/Source/Runtime/Experimental/Chaos/Private/Chaos/Evolution/SolverBodyContainer.cpp
@@ -0,0 +1,118 @@
+// Copyright Epic Games, Inc. All Rights Reserved.
+#include "Chaos/Evolution/SolverBodyContainer.h"
+#include "Chaos/Particle/ParticleUtilities.h"
+
+namespace Chaos
+{
+	void FSolverBodyAdapter::GatherInput()
+	{
+		if (Particle.IsValid())
+		{
+			FRigidTransform3 CoMTransform = FParticleUtilitiesPQ::GetCoMWorldTransform(Particle);
+			SolverBody.SetP(CoMTransform.GetLocation());
+			SolverBody.SetQ(CoMTransform.GetRotation());
+			SolverBody.SetV(Particle->V());
+			SolverBody.SetW(Particle->W());
+			SolverBody.SetCoM(Particle->CenterOfMass());
+			SolverBody.SetRoM(Particle->RotationOfMass());
+
+			if (Particle->IsDynamic())
+			{
+				FRigidTransform3 PrevCoMTransform = FParticleUtilitiesXR::GetCoMWorldTransform(Particle);
+				SolverBody.SetX(PrevCoMTransform.GetLocation());
+				SolverBody.SetR(PrevCoMTransform.GetRotation());
+
+				SolverBody.SetInvM(Particle->InvM());
+				SolverBody.SetInvILocal(Particle->InvI().GetDiagonal());
+			}
+			else
+			{
+				SolverBody.SetX(SolverBody.P());
+				SolverBody.SetR(SolverBody.Q());
+			}
+
+			SolverBody.UpdateRotationDependentState();
+		}
+	}
+
+	void FSolverBodyAdapter::ScatterOutput()
+	{
+		if (Particle.IsValid())
+		{
+			if (SolverBody.IsDynamic())
+			{
+				// Set the particle state
+				FParticleUtilities::SetCoMWorldTransform(Particle, SolverBody.P(), SolverBody.Q());
+				Particle->SetV(SolverBody.V());
+				Particle->SetW(SolverBody.W());
+
+				if (SolverBody.HasActiveCollision())
+				{
+					//Particle->AuxilaryValue(*ParticleParameters.Collided) = true;
+				}
+			}
+
+			// Reset SolverBodyIndex cookie every step - it will be reassigned next step
+			Particle->SetSolverBodyIndex(INDEX_NONE);
+		}
+	}
+
+	int32 FSolverBodyContainer::AddParticle(FGenericParticleHandle InParticle)
+	{
+		// No array resizing allowed (we want fixed pointers)
+		check(NumItems() < MaxItems());
+		return SolverBodies.Emplace(InParticle);
+	}
+
+	FSolverBody* FSolverBodyContainer::FindOrAdd(FGenericParticleHandle InParticle)
+	{
+		// For dynamic bodies, we store a cookie on the Particle that holds the solver body index
+		// For kinematics we cannot do this because the kinematic may be in multiple islands and 
+		// would require a different index for each island, so we use a local map instead.
+		int32 ItemIndex = InParticle->SolverBodyIndex();
+
+		if (ItemIndex == INDEX_NONE)
+		{
+			if (InParticle->IsDynamic())
+			{
+				// First time we have seen this particle, so add it
+				ItemIndex = AddParticle(InParticle);
+				InParticle->SetSolverBodyIndex(ItemIndex);
+			}
+			else // Not Dynamic
+			{
+				int32* ItemIndexPtr = ParticleToIndexMap.Find(InParticle);
+				if (ItemIndexPtr != nullptr)
+				{
+					ItemIndex = *ItemIndexPtr;
+				}
+				else
+				{
+					// First time we have seen this particle, so add it
+					ItemIndex = AddParticle(InParticle);
+					ParticleToIndexMap.Add(InParticle, ItemIndex);
+				}
+			}			
+		}
+
+		check(ItemIndex != INDEX_NONE);
+		return &SolverBodies[ItemIndex].GetSolverBody();
+	}
+
+	void FSolverBodyContainer::ScatterOutput()
+	{
+		for (FSolverBodyAdapter& SolverBody : SolverBodies)
+		{
+			SolverBody.ScatterOutput();
+		}
+	}
+
+	void FSolverBodyContainer::SetImplicitVelocities(FReal Dt)
+	{
+		for (FSolverBodyAdapter& SolverBody : SolverBodies)
+		{
+			SolverBody.GetSolverBody().SetImplicitVelocity(Dt);
+		}
+	}
+
+}
\ No newline at end of file
diff --git a/Engine/Source/Runtime/Experimental/Chaos/Private/Chaos/HeightField.cpp b/Engine/Source/Runtime/Experimental/Chaos/Private/Chaos/HeightField.cpp
--- a/Engine/Source/Runtime/Experimental/Chaos/Private/Chaos/HeightField.cpp
+++ b/Engine/Source/Runtime/Experimental/Chaos/Private/Chaos/HeightField.cpp
@@ -1265,6 +1265,11 @@ namespace Chaos
 				// The regular penetration calculation vs a triangle may result in inward facing normals.
 				// To protect against this, we sweep against the triangle from a distance to ensure an outward
 				// facing normal and MTD.
+				//
+				// BUG: This does not detect collisions when we specify a cull distance. It is as if
+				// Thickness is always zero...
+				//
+
 				const FAABB3 Bounds = QueryGeom.BoundingBox();
 				const FReal ApproximateSizeOfObject = Bounds.Extents()[Bounds.LargestAxis()];
 				const FReal ApproximateDistToObject = FVec3::DistSquared(QueryTM.GetLocation(), A);
diff --git a/Engine/Source/Runtime/Experimental/Chaos/Private/Chaos/Island/IslandGraph.cpp b/Engine/Source/Runtime/Experimental/Chaos/Private/Chaos/Island/IslandGraph.cpp
new file mode 100644
--- /dev/null
+++ b/Engine/Source/Runtime/Experimental/Chaos/Private/Chaos/Island/IslandGraph.cpp
@@ -0,0 +1,518 @@
+// Copyright Epic Games, Inc. All Rights Reserved.
+
+#include "Chaos/Island/IslandGraph.h"
+#include "Chaos/Island/SolverIsland.h"
+#include "Chaos/ParticleHandle.h"
+#include "Chaos/ConstraintHandle.h"
+
+#include "ChaosStats.h"
+#include "ChaosLog.h"
+
+namespace Chaos
+{
+
+DECLARE_CYCLE_STAT(TEXT("MergeIslandGraph"), STAT_MergeIslandGraph, STATGROUP_Chaos);
+DECLARE_CYCLE_STAT(TEXT("SplitIslandGraph"), STAT_SplitIslandGraph, STATGROUP_Chaos);
+
+template<typename NodeType, typename EdgeType, typename IslandType>
+void FIslandGraph<NodeType, EdgeType, IslandType>::ResetNodes()
+{
+	GraphNodes.Reset();
+	ItemNodes.Reset();
+	GraphIslands.Reset();
+}
+
+template<typename NodeType, typename EdgeType, typename IslandType>
+void FIslandGraph<NodeType, EdgeType, IslandType>::ReserveNodes(const int32 NumNodes)
+{
+	GraphNodes.Reserve(NumNodes);
+	ItemNodes.Reserve(NumNodes);
+	GraphIslands.Reserve(NumNodes);
+}
+
+template<typename NodeType, typename EdgeType, typename IslandType>
+void FIslandGraph<NodeType, EdgeType, IslandType>::ParentIslands(const int32 FirstIsland, const int32 SecondIsland)
+{
+	if (GraphIslands.IsValidIndex(FirstIsland) && GraphIslands.IsValidIndex(SecondIsland) && (FirstIsland != SecondIsland))
+	{
+		GraphIslands[SecondIsland].ChildrenIslands.Add(FirstIsland);
+		GraphIslands[FirstIsland].ChildrenIslands.Add(SecondIsland);
+		
+		// If we are adding a constraint the island is no longer persistent
+		GraphIslands[FirstIsland].bIsPersistent = false;
+		GraphIslands[SecondIsland].bIsPersistent = false;
+	}
+}
+
+template<typename NodeType, typename EdgeType, typename IslandType>
+int32 FIslandGraph<NodeType, EdgeType, IslandType>::AddNode(const NodeType& NodeItem, const bool bValidNode, const int32 IslandIndex, const bool bDiscardNode)
+{
+	int32 NodeIndex = INDEX_NONE;
+	int32* ItemIndex = ItemNodes.Find(NodeItem);
+	if (ItemIndex && GraphNodes.IsValidIndex(*ItemIndex))
+	{
+		// In case the item is changing its state to be valid
+		// we merge all the connected islands 
+		if (bValidNode && !GraphNodes[*ItemIndex].bValidNode)
+		{
+			int32 MasterIsland = INDEX_NONE;
+			for (int32& EdgeIndex : GraphNodes[*ItemIndex].NodeEdges)
+			{
+				FGraphEdge& GraphEdge = GraphEdges[EdgeIndex];
+				ParentIslands(MasterIsland, GraphEdge.IslandIndex);
+				MasterIsland = GraphEdge.IslandIndex;
+				GraphIslands[GraphEdge.IslandIndex].bIsPersistent = false;
+			}
+			GraphNodes[*ItemIndex].bValidNode = bValidNode;
+		}
+		NodeIndex =  *ItemIndex;
+	}
+	else
+	{
+		FGraphNode GraphNode;
+		GraphNode.NodeItem = NodeItem;
+		GraphNode.IslandIndex = IslandIndex;
+		GraphNode.bValidNode = bValidNode;
+
+		NodeIndex =  ItemNodes.Add(NodeItem, GraphNodes.Emplace(GraphNode));
+	}
+	GraphNodes[NodeIndex].bDiscardNode = bDiscardNode;
+	return NodeIndex;
+}
+
+template<typename NodeType, typename EdgeType, typename IslandType>
+void FIslandGraph<NodeType, EdgeType, IslandType>::RemoveNode(const NodeType& NodeItem)
+{
+	if (const int32* IndexPtr = ItemNodes.Find(NodeItem))
+	{
+		const int32 NodeIndex = *IndexPtr;
+		if (GraphNodes.IsValidIndex(NodeIndex))
+		{
+			// The island is no longer persistent since the number of nodes is changing
+			if(GraphIslands.IsValidIndex(GraphNodes[NodeIndex].IslandIndex))
+			{
+				GraphIslands[GraphNodes[NodeIndex].IslandIndex].bIsPersistent = false;
+			}
+			// We remove all the connected edges as well
+			for (int32 NodeEdgeIndex = GraphNodes[NodeIndex].NodeEdges.GetMaxIndex() - 1; NodeEdgeIndex >= 0; --NodeEdgeIndex)
+			{
+				if (GraphNodes[NodeIndex].NodeEdges.IsValidIndex(NodeEdgeIndex))
+				{
+					const int32 GraphEdgeIndex = GraphNodes[NodeIndex].NodeEdges[NodeEdgeIndex];
+					RemoveEdge(GraphEdges[GraphEdgeIndex].EdgeItem);
+				}
+			}
+			ItemNodes.Remove(NodeItem);
+			GraphNodes.RemoveAt(NodeIndex);
+		}
+		else
+		{
+			UE_LOG(LogChaos, Error, TEXT("Island Graph : Trying to remove a node at index %d in a list of size %d"), NodeIndex, GraphNodes.Num());
+		}
+	}
+}
+
+template<typename NodeType, typename EdgeType, typename IslandType>
+void FIslandGraph<NodeType, EdgeType, IslandType>::ResetEdges()
+{
+	GraphEdges.Reset();
+	ItemEdges.Reset();
+}
+
+template<typename NodeType, typename EdgeType, typename IslandType>
+void FIslandGraph<NodeType, EdgeType, IslandType>::ReserveEdges(const int32 NumEdges)
+{
+	GraphEdges.Reserve(NumEdges);
+	ItemEdges.Reserve(NumEdges);
+}
+
+template<typename NodeType, typename EdgeType, typename IslandType>
+void FIslandGraph<NodeType, EdgeType, IslandType>::AttachIslands(const int32 EdgeIndex)
+{	
+	FGraphEdge& GraphEdge = GraphEdges[EdgeIndex];
+
+	if (GraphNodes.IsValidIndex(GraphEdge.FirstNode) && GraphNodes.IsValidIndex(GraphEdge.SecondNode))
+	{
+		FGraphNode& FirstNode = GraphNodes[GraphEdge.FirstNode];
+		FGraphNode& SecondNode = GraphNodes[GraphEdge.SecondNode];
+
+		const bool bFirstValidIsland = GraphIslands.IsValidIndex(FirstNode.IslandIndex) && FirstNode.bValidNode;
+		const bool bSecondValidIsland = GraphIslands.IsValidIndex(SecondNode.IslandIndex) && SecondNode.bValidNode;
+
+		// We check if one of the 2 nodes have invalid island
+		// if yes we set the invalid node island index 
+		// and the edge one to be the valid one 
+		// If none are valid we create a new island 
+		if (bFirstValidIsland && !bSecondValidIsland)
+		{
+			GraphEdge.IslandIndex = FirstNode.IslandIndex;
+			++GraphIslands[GraphEdge.IslandIndex].NumEdges;
+			if(SecondNode.bValidNode)
+			{
+				GraphIslands[GraphEdge.IslandIndex].bIsPersistent = false;
+				++GraphIslands[GraphEdge.IslandIndex ].NumNodes;
+			}
+		}
+		else if (!bFirstValidIsland && bSecondValidIsland)
+		{
+			GraphEdge.IslandIndex = SecondNode.IslandIndex;
+			++GraphIslands[GraphEdge.IslandIndex].NumEdges;
+			if(FirstNode.bValidNode)
+			{
+				GraphIslands[GraphEdge.IslandIndex].bIsPersistent = false;
+				++GraphIslands[GraphEdge.IslandIndex ].NumNodes;
+			}
+		}
+		else if (!bFirstValidIsland && !bSecondValidIsland)
+		{
+			FGraphIsland GraphIsland = { EdgeIndex, 1, 2 };
+			GraphEdge.IslandIndex = GraphIslands.Emplace(GraphIsland);
+		}
+		else
+		{
+			// If the 2 nodes are coming from 2 different islands, we need to merge these islands
+			// In order to do that we build an island graph and we will 
+			// merge recursively the children islands onto the parent one
+			GraphEdge.IslandIndex = FMath::Min(FirstNode.IslandIndex, SecondNode.IslandIndex);
+			ParentIslands(FirstNode.IslandIndex, SecondNode.IslandIndex);
+
+			++GraphIslands[GraphEdge.IslandIndex ].NumEdges;
+		}
+		// We set both island indices to be the equal to the edge one
+		FirstNode.IslandIndex = GraphEdge.IslandIndex;
+		SecondNode.IslandIndex = GraphEdge.IslandIndex;
+	}
+	else if (GraphNodes.IsValidIndex(GraphEdge.FirstNode) && !GraphNodes.IsValidIndex(GraphEdge.SecondNode) && GraphNodes[GraphEdge.FirstNode].bValidNode)
+	{
+		// If only the first node is valid, check if its island index is valid and if yes use it for the edge.
+		// otherwise we create a new island
+		if(!GraphIslands.IsValidIndex(GraphNodes[GraphEdge.FirstNode].IslandIndex))
+		{
+			FGraphIsland GraphIsland = { EdgeIndex, 1, 1 };
+			GraphNodes[GraphEdge.FirstNode].IslandIndex =  GraphIslands.Emplace(GraphIsland);
+		}
+		GraphEdge.IslandIndex =  GraphNodes[GraphEdge.FirstNode].IslandIndex;
+	}
+	else if (!GraphNodes.IsValidIndex(GraphEdge.FirstNode) && GraphNodes.IsValidIndex(GraphEdge.SecondNode) && GraphNodes[GraphEdge.SecondNode].bValidNode)
+	{
+		// If only the second node is valid, check if its island index is valid and if yes use it for the edge.
+		// otherwise we create a new island
+		if(!GraphIslands.IsValidIndex(GraphNodes[GraphEdge.SecondNode].IslandIndex))
+		{
+			FGraphIsland GraphIsland = { EdgeIndex, 1, 1 };
+			GraphNodes[GraphEdge.SecondNode].IslandIndex =  GraphIslands.Emplace(GraphIsland);
+		}
+		GraphEdge.IslandIndex =  GraphNodes[GraphEdge.SecondNode].IslandIndex;
+	}
+}
+
+template<typename NodeType, typename EdgeType, typename IslandType>
+int32 FIslandGraph<NodeType, EdgeType, IslandType>::AddEdge(const EdgeType& EdgeItem, const int32 ItermContainer, const int32 FirstNode, const int32 SecondNode)
+{
+	int32 EdgeIndex = INDEX_NONE;
+	// We only add an edge if one of the 2 nodes are valid
+	if (GraphNodes.IsValidIndex(FirstNode) || GraphNodes.IsValidIndex(SecondNode))
+	{
+		// This test could be quite slow. will probably be better to use the graph index that is stored on the constraint handle
+		int32* ItemIndex = ItemEdges.Find(EdgeItem);
+		if (ItemIndex)
+		{
+			// If the edge is already there, no need to attach islands since it must have been done before
+			EdgeIndex = *ItemIndex;
+		}
+		else
+		{
+			// Create a  new edge and enqueue the linked islands to be merged if necessary
+			FGraphEdge GraphEdge;
+			GraphEdge.EdgeItem = EdgeItem;
+			GraphEdge.FirstNode = FirstNode;
+			GraphEdge.SecondNode = SecondNode;
+			GraphEdge.IslandIndex = INDEX_NONE;
+			GraphEdge.ItemContainer = ItermContainer;
+
+			EdgeIndex = GraphEdges.Emplace(GraphEdge);
+			ItemEdges.Add(EdgeItem, EdgeIndex);
+
+			GraphEdges[EdgeIndex].FirstEdge = GraphNodes.IsValidIndex(FirstNode) ? GraphNodes[FirstNode].NodeEdges.Add(EdgeIndex) : -1;
+			GraphEdges[EdgeIndex].SecondEdge = GraphNodes.IsValidIndex(SecondNode) ? GraphNodes[SecondNode].NodeEdges.Add(EdgeIndex) : -1;
+
+			AttachIslands(EdgeIndex);
+		}
+	}	
+	else
+	{
+		UE_LOG(LogChaos, Error, TEXT("Island Graph : Trying to add an edge with invalid node indices %d  %d in a list of nodes of size %d"), FirstNode, SecondNode, GraphNodes.Num());
+	}
+	return EdgeIndex;
+}
+
+template<typename NodeType, typename EdgeType, typename IslandType>
+void FIslandGraph<NodeType, EdgeType, IslandType>::RemoveEdge(const EdgeType& EdgeItem)
+{
+	if (const int32* EdgePtr = ItemEdges.Find(EdgeItem))
+	{
+		int32 EdgeIndex = *EdgePtr;
+		if (GraphEdges.IsValidIndex(EdgeIndex))
+		{
+			FGraphEdge& GraphEdge = GraphEdges[EdgeIndex];
+			// We first remove the edge from the nodes edges list
+			if(GraphNodes.IsValidIndex(GraphEdge.FirstNode) && GraphNodes[GraphEdge.FirstNode].NodeEdges.IsValidIndex(GraphEdge.FirstEdge))
+			{
+				GraphNodes[GraphEdge.FirstNode].NodeEdges.RemoveAt(GraphEdge.FirstEdge);
+			}
+			if (GraphNodes.IsValidIndex(GraphEdge.SecondNode) && GraphNodes[GraphEdge.SecondNode].NodeEdges.IsValidIndex(GraphEdge.SecondEdge))
+			{
+				GraphNodes[GraphEdge.SecondNode].NodeEdges.RemoveAt(GraphEdge.SecondEdge);
+			}
+			// We then remove the edge from the item and graph edges
+			ItemEdges.Remove(EdgeItem);
+			GraphEdges.RemoveAt(EdgeIndex);
+		}
+		else
+		{
+			UE_LOG(LogChaos, Error, TEXT("Island Graph : Trying to remove an edge at index %d in a list of size %d"), EdgeIndex, GraphEdges.Num());
+		}
+	}
+}
+	
+template<typename NodeType, typename EdgeType, typename IslandType>
+void FIslandGraph<NodeType, EdgeType, IslandType>::ForEachNodes(const int32 IslandIndex, TFunctionRef<void(FGraphNode&)> ApplyFunction, const bool bUpdateCounter)
+{
+	if (bUpdateCounter)
+	{
+		GraphCounter = (GraphCounter + 1) % MaxCount;
+	}
+	if (GraphIslands.IsValidIndex(IslandIndex))
+	{
+		// We first loop over all the nodes within the island to apply the function
+		for (FGraphNode& GraphNode : GraphNodes)
+		{
+			if ((GraphNode.IslandIndex == IslandIndex) && (GraphNode.NodeCounter != GraphCounter))
+			{
+				ApplyFunction(GraphNode);
+				GraphNode.NodeCounter = GraphCounter;
+			}
+		}
+		
+		// We then loop over all the edges that are unique per island to process all the connected nodes
+		// Most of the nodes will already be processed but we need to apply the function to the non valid ones
+		// that could belong to several islands
+		for (FGraphEdge& GraphEdge : GraphEdges)
+		{
+			if (GraphEdge.IslandIndex == IslandIndex)
+			{
+				if (GraphNodes.IsValidIndex(GraphEdge.FirstNode) && GraphNodes[GraphEdge.FirstNode].NodeCounter != GraphCounter)
+				{
+					ApplyFunction(GraphNodes[GraphEdge.FirstNode]);
+					GraphNodes[GraphEdge.FirstNode].NodeCounter = GraphCounter;
+				}
+				if (GraphNodes.IsValidIndex(GraphEdge.SecondNode) && GraphNodes[GraphEdge.SecondNode].NodeCounter != GraphCounter)
+				{
+					ApplyFunction(GraphNodes[GraphEdge.SecondNode]);
+					GraphNodes[GraphEdge.SecondNode].NodeCounter = GraphCounter;
+				}
+			}
+		}
+	}
+}
+
+template<typename NodeType, typename EdgeType, typename IslandType>
+void FIslandGraph<NodeType, EdgeType, IslandType>::ForEachEdges(const int32 IslandIndex, TFunctionRef<void(FGraphEdge&)> ApplyFunction, const bool bUpdateCounter)
+{
+	if (bUpdateCounter)
+	{
+		GraphCounter = (GraphCounter + 1) % MaxCount;
+	}
+	// Loop over all the edges within the island to apply the function
+	for (FGraphEdge& GraphEdge : GraphEdges)
+	{
+		if ((GraphEdge.IslandIndex == IslandIndex) && (GraphEdge.EdgeCounter != GraphCounter))
+		{
+			ApplyFunction(GraphEdge);
+			GraphEdge.EdgeCounter = GraphCounter;
+		}
+	}
+}
+
+template<typename NodeType, typename EdgeType, typename IslandType>
+void FIslandGraph<NodeType, EdgeType, IslandType>::MergeIslands(const int32 ParentIndex, const int32 ChildIndex)
+{
+	// We loop recursively over the island tree to merge all 
+	// the children islands onto the parent one
+	if (GraphIslands[ChildIndex].IslandCounter != GraphCounter && ParentIndex != ChildIndex)
+	{
+		ForEachEdges(ChildIndex,
+			[ParentIndex](FGraphEdge& GraphEdge) {
+				GraphEdge.IslandIndex = ParentIndex;}, false);
+
+		ForEachNodes(ChildIndex,
+			[ParentIndex](FGraphNode& GraphNode){
+				GraphNode.IslandIndex = ParentIndex; }, false);
+
+		GraphIslands[ChildIndex].IslandCounter = GraphCounter;
+
+		// Add the number of nodes/edges to the parent one and set to 0 the children ones
+		GraphIslands[ParentIndex].NumEdges += GraphIslands[ChildIndex].NumEdges;
+		GraphIslands[ChildIndex].NumEdges = 0;
+
+		GraphIslands[ParentIndex].NumNodes += GraphIslands[ChildIndex].NumNodes;
+		GraphIslands[ChildIndex].NumNodes = 0;
+
+		// Recursively iterate over all the children ones
+		for (auto& MergedIsland : GraphIslands[ChildIndex].ChildrenIslands)
+		{
+			MergeIslands(ParentIndex, MergedIsland);
+		}
+		GraphIslands[ChildIndex].ChildrenIslands.Reset();
+	}
+}
+
+template<typename NodeType, typename EdgeType, typename IslandType>
+void FIslandGraph<NodeType, EdgeType, IslandType>::SplitIsland(const int32 IslandIndex)
+{
+	SCOPE_CYCLE_COUNTER(STAT_SplitIslandGraph);
+
+	GraphCounter = (GraphCounter + 1) % MaxCount;
+	TQueue<int32> NodeQueue;
+	bool bMainIsland = true;
+	int32 CurrentIsland = IslandIndex;
+	
+	for (int32 RootIndex = 0, NumNodes = GraphNodes.GetMaxIndex(); RootIndex < NumNodes; ++RootIndex)
+	{
+		// We pick all the nodes that are inside an island
+		if(GraphNodes.IsValidIndex(RootIndex))
+		{
+			FGraphNode& RootNode = GraphNodes[RootIndex];
+			if(RootNode.IslandIndex == IslandIndex && RootNode.NodeCounter != GraphCounter && RootNode.bValidNode)
+			{
+				// We don't want to rebuild a new island if this one can't be splitted
+				// It is why by default the first one is the main one
+				if(!bMainIsland)
+				{
+					FGraphIsland GraphIsland = { RootIndex, 0, 1 };
+					GraphIsland.bIsSleeping = false;
+					GraphIsland.bIsPersistent = false;
+					CurrentIsland = GraphIslands.Emplace(GraphIsland);
+				}
+				NodeQueue.Enqueue(RootIndex);
+				int32 NodeIndex = RootIndex;
+				
+				while (!NodeQueue.IsEmpty())
+				{
+					NodeQueue.Dequeue(NodeIndex);
+					FGraphNode& GraphNode = GraphNodes[NodeIndex];
+					
+					// Graph counter is there to avoid processing multiple times the same node/edge
+					if (GraphNode.NodeCounter != GraphCounter)
+					{
+						GraphNode.NodeCounter = GraphCounter;
+						GraphNode.IslandIndex = CurrentIsland;
+
+						// Loop over the node edges to continue the island discovery
+						for (int32& EdgeIndex : GraphNode.NodeEdges)
+						{
+							FGraphEdge& GraphEdge = GraphEdges[EdgeIndex];
+							if(GraphEdge.EdgeCounter != GraphCounter)
+							{
+								GraphEdge.EdgeCounter = GraphCounter;
+								GraphEdge.IslandIndex = CurrentIsland;
+							}
+							const int32 OtherIndex = (NodeIndex == GraphEdge.FirstNode) ?
+								GraphEdge.SecondNode : GraphEdge.FirstNode;
+
+							// Only the valid nodes are allowed to continue the graph traversal
+							if (GraphNodes.IsValidIndex(OtherIndex) && GraphNodes[OtherIndex].NodeCounter != GraphCounter && GraphNodes[OtherIndex].bValidNode)
+							{
+								NodeQueue.Enqueue(OtherIndex);
+							}
+						}
+					}
+				}
+				bMainIsland = false;
+			}
+		}
+	}
+}
+
+template<typename NodeType, typename EdgeType, typename IslandType>
+void FIslandGraph<NodeType, EdgeType, IslandType>::UpdateGraph()
+{
+	SCOPE_CYCLE_COUNTER(STAT_MergeIslandGraph);
+
+	GraphCounter = (GraphCounter + 1) % MaxCount;
+
+	// We loop over all the islands and if they have children
+	// we merge them onto the parent one
+	for (int32 IslandIndex = 0, NumIslands = GraphIslands.GetMaxIndex(); IslandIndex < NumIslands; ++IslandIndex)
+	{
+		if (GraphIslands.IsValidIndex(IslandIndex))
+		{
+			FGraphIsland& GraphIsland = GraphIslands[IslandIndex];
+			// If the graph island is persistent we don't need to merge its children
+			if (!GraphIsland.bIsPersistent)
+			{
+				for (const int32& MergedIsland : GraphIsland.ChildrenIslands)
+				{
+					MergeIslands(IslandIndex, MergedIsland);
+				}
+				GraphIsland.ChildrenIslands.Reset();
+			}
+		}
+	}
+	// Once the merging process is done we remove all the children islands
+	// since they have been merged onto the parent one
+	for (int32 IslandIndex = (GraphIslands.GetMaxIndex() - 1); IslandIndex >= 0; --IslandIndex)
+	{
+		// Only the island counter of the children have been updated
+		if (GraphIslands.IsValidIndex(IslandIndex) && GraphIslands[IslandIndex].IslandCounter == GraphCounter)
+		{
+			GraphIslands.RemoveAt(IslandIndex);
+		}
+	}
+	
+	for (int32 NodeIndex = 0, NumNodes = GraphNodes.GetMaxIndex(); NodeIndex < NumNodes; ++NodeIndex)
+	{
+		if (GraphNodes.IsValidIndex(NodeIndex))
+		{
+			// Add new islands for the all the particles that are not connected into the graph 
+			if(GraphNodes[NodeIndex].IslandIndex == INDEX_NONE && GraphNodes[NodeIndex].bValidNode)
+			{
+				FGraphIsland GraphIsland = { NodeIndex, 0, 1 };
+				GraphNodes[NodeIndex].IslandIndex = GraphIslands.Emplace(GraphIsland);
+			}
+			// Update of the sleeping flag on the island
+			if(GraphNodes[NodeIndex].IslandIndex != INDEX_NONE && !GraphNodes[NodeIndex].bDiscardNode)
+			{
+				GraphIslands[GraphNodes[NodeIndex].IslandIndex].bIsSleeping = false;
+			}
+		}
+	}
+	// Split the islands if this one is persistent and not sleeping
+	for (int32 IslandIndex = 0, NumIslands = GraphIslands.GetMaxIndex(); IslandIndex < NumIslands; ++IslandIndex)
+	{
+		if (GraphIslands.IsValidIndex(IslandIndex) && GraphIslands[IslandIndex].bIsPersistent && !GraphIslands[IslandIndex].bIsSleeping)
+		{
+			SplitIsland(IslandIndex);
+		}
+	}
+}
+
+template<typename NodeType, typename EdgeType, typename IslandType>
+void FIslandGraph<NodeType, EdgeType, IslandType>::ResetIslands()
+{
+	GraphEdges.Reset();
+	ItemEdges.Reset();
+	// Reset of all the edges + the nodes ones
+	// will probably need to change that with persistent contacts
+	for (int32 NodeIndex = 0, NumNodes = GraphNodes.GetMaxIndex(); NodeIndex < NumNodes; ++NodeIndex)
+	{
+		if (GraphNodes.IsValidIndex(NodeIndex))
+		{
+			GraphNodes[NodeIndex].NodeEdges.Reset();
+		}
+	}
+}
+
+template class FIslandGraph<FGeometryParticleHandle*, FConstraintHandle*, FPBDIslandSolver*>;
+template class FIslandGraph<int32, int32, int32>;
+}
diff --git a/Engine/Source/Runtime/Experimental/Chaos/Private/Chaos/Island/IslandManager.cpp b/Engine/Source/Runtime/Experimental/Chaos/Private/Chaos/Island/IslandManager.cpp
new file mode 100644
--- /dev/null
+++ b/Engine/Source/Runtime/Experimental/Chaos/Private/Chaos/Island/IslandManager.cpp
@@ -0,0 +1,483 @@
+// Copyright Epic Games, Inc. All Rights Reserved.
+#include "Chaos/Island/IslandManager.h"
+#include "Chaos/Island/SolverIsland.h"
+#include "Chaos/Island/IslandGraph.h"
+#include "Chaos/ParticleHandle.h"
+#include "Chaos/ConstraintHandle.h"
+#include "Chaos/PBDRigidParticles.h"
+#include "Chaos/PBDRigidsSOAs.h"
+#include "ChaosLog.h"
+
+/** Cvar to enable/disable the island sleeping */
+bool ChaosSolverSleepEnabled = true;
+FAutoConsoleVariableRef CVarChaosSolverSleepEnabled(TEXT("p.Chaos.Solver.SleepEnabled"), ChaosSolverSleepEnabled, TEXT(""));
+
+/** Cvar to override the sleep counter threshold if necessary */
+int32 ChaosSolverCollisionDefaultSleepCounterThresholdCVar = 20;
+FAutoConsoleVariableRef CVarChaosSolverCollisionDefaultSleepCounterThreshold(TEXT("p.ChaosSolverCollisionDefaultSleepCounterThreshold"), ChaosSolverCollisionDefaultSleepCounterThresholdCVar, TEXT("Default counter threshold for sleeping.[def:20]"));
+
+/** Cvar to override the sleep linear threshold if necessary */
+Chaos::FRealSingle ChaosSolverCollisionDefaultLinearSleepThresholdCVar = 0.001f; // .001 unit mass cm
+FAutoConsoleVariableRef CVarChaosSolverCollisionDefaultLinearSleepThreshold(TEXT("p.ChaosSolverCollisionDefaultLinearSleepThreshold"), ChaosSolverCollisionDefaultLinearSleepThresholdCVar, TEXT("Default linear threshold for sleeping.[def:0.001]"));
+
+/** Cvar to override the sleep angular threshold if necessary */
+Chaos::FRealSingle ChaosSolverCollisionDefaultAngularSleepThresholdCVar = 0.0087f;  //~1/2 unit mass degree
+FAutoConsoleVariableRef CVarChaosSolverCollisionDefaultAngularSleepThreshold(TEXT("p.ChaosSolverCollisionDefaultAngularSleepThreshold"), ChaosSolverCollisionDefaultAngularSleepThresholdCVar, TEXT("Default angular threshold for sleeping.[def:0.0087]"));
+
+namespace Chaos
+{
+	
+/** Check if a particle is dynamic or sleeping */
+FORCEINLINE bool IsDynamicParticle(const FGeometryParticleHandle* ParticleHandle)
+{
+	return (ParticleHandle->ObjectState() == EObjectStateType::Dynamic) ||
+		   (ParticleHandle->ObjectState() == EObjectStateType::Sleeping);
+}
+
+/** Check if a particle is dynamic or sleeping */
+inline const FChaosPhysicsMaterial* GetPhysicsMaterial(const TArrayCollectionArray<TSerializablePtr<FChaosPhysicsMaterial>>& ParticleMaterialAttributes,
+														    const THandleArray<FChaosPhysicsMaterial>& SolverPhysicsMaterials, FPBDRigidParticleHandle* RigidParticleHandle)
+{
+	const FChaosPhysicsMaterial* PhysicsMaterial = RigidParticleHandle->AuxilaryValue(ParticleMaterialAttributes).Get();
+	if (!PhysicsMaterial && RigidParticleHandle->ShapesArray().Num())
+	{
+		if (FPerShapeData* PerShapeData = RigidParticleHandle->ShapesArray()[0].Get())
+		{
+			if (PerShapeData->GetMaterials().Num())
+			{
+				PhysicsMaterial = SolverPhysicsMaterials.Get(PerShapeData->GetMaterials()[0].InnerHandle);
+			}
+		}
+	}
+	return PhysicsMaterial;
+}
+
+/** Check if an island is sleeping or not given a linear/angular velocities and sleeping thresholds */
+inline bool IsIslandSleeping(const FReal MaxLinearSpeed2, const FReal MaxAngularSpeed2,
+								  const FReal LinearSleepingThreshold, const FReal AngularSleepingThreshold,
+								  const int32 CounterThreshold, int32& SleepCounter)
+{
+	const FReal MaxLinearSpeed = FMath::Sqrt(MaxLinearSpeed2);
+	const FReal MaxAngularSpeed = FMath::Sqrt(MaxAngularSpeed2);
+
+	if (MaxLinearSpeed < LinearSleepingThreshold && MaxAngularSpeed < AngularSleepingThreshold)
+	{
+		if (SleepCounter >= CounterThreshold)
+		{
+			return true;
+		}
+		else
+		{
+			SleepCounter++;
+		}
+	}
+	else
+	{
+		SleepCounter = 0;
+	}
+	return false;
+}
+
+/** Check is an island is sleeping or not */
+inline void MakeAsleep(FPBDIslandSolver* IslandSolver, FPBDRigidsSOAs& Particles)
+{
+	if(IslandSolver)
+	{
+		bool bNeedsResim = false;
+		for (auto& IslandParticle : IslandSolver->GetParticles())
+		{
+			//if even one particle is soft/hard desync we must resim the entire island (when resim is used)
+			//seems cheap enough so just always do it, if slow pass resim template in here
+			if (IslandParticle->SyncState() != ESyncState::InSync)
+			{
+				bNeedsResim = true;
+				break;
+			}
+		}
+		IslandSolver->SetNeedsResim(bNeedsResim);
+		// Sleeping flag has already been computed by the island graph 
+		const bool bIsSleeping = IslandSolver->GetIsSleeping();
+		
+		for (auto& IslandParticle : IslandSolver->GetParticles())
+		{
+			// If not sleeping we activate the sleeping particles
+			if (!bIsSleeping && IslandParticle->Sleeping())
+			{
+				Particles.ActivateParticle(IslandParticle); 	
+			}
+			// If sleeping we deactivate the dynamic particles
+			else if (bIsSleeping && !IslandParticle->Sleeping())
+			{
+				Particles.DeactivateParticle(IslandParticle); 	
+			}
+		}
+	}
+}
+	
+/** Compute sleeping thresholds given a solver island  */
+inline bool ComputeSleepingThresholds(FPBDIslandSolver* IslandSolver,
+	const TArrayCollectionArray<TSerializablePtr<FChaosPhysicsMaterial>>& PerParticleMaterialAttributes,
+	const THandleArray<FChaosPhysicsMaterial>& SolverPhysicsMaterials, FReal& LinearSleepingThreshold,
+	FReal& AngularSleepingThreshold, FReal& MaxLinearSpeed2, FReal& MaxAngularSpeed2, int32& SleepCounterThreshold)
+{
+	LinearSleepingThreshold = FLT_MAX;  
+	AngularSleepingThreshold = FLT_MAX; 
+	MaxLinearSpeed2 = 0.f; 
+	MaxAngularSpeed2 = 0.f;
+	SleepCounterThreshold = 0;
+
+	bool bHaveSleepThreshold = false;
+	for (FGeometryParticleHandle* ParticleHandle : IslandSolver->GetParticles())
+	{
+		if (ParticleHandle)
+		{
+			if (FPBDRigidParticleHandle* PBDRigid = ParticleHandle->CastToRigidParticle())
+			{
+				if (IsDynamicParticle(ParticleHandle) && !PBDRigid->Sleeping())
+				{
+					// If any body in the island is not allowed to sleep, the whole island cannot sleep
+					// @todo(chaos): if this is a common thing, we should set a flag on the island when it has a particle
+					// with this property enabled and skip the sleep check altogether
+					if (PBDRigid->SleepType() == ESleepType::NeverSleep)
+					{
+						return false;
+					}
+
+					bHaveSleepThreshold = true;
+
+					MaxLinearSpeed2 = FMath::Max(PBDRigid->VSmooth().SizeSquared(), MaxLinearSpeed2);
+					MaxAngularSpeed2 = FMath::Max(PBDRigid->WSmooth().SizeSquared(), MaxAngularSpeed2);
+
+					const FChaosPhysicsMaterial* PhysicsMaterial = GetPhysicsMaterial(PerParticleMaterialAttributes, SolverPhysicsMaterials, PBDRigid);
+
+					const FReal LocalSleepingLinearThreshold = PhysicsMaterial ? PhysicsMaterial->SleepingLinearThreshold : ChaosSolverCollisionDefaultLinearSleepThresholdCVar;
+					const FReal LocalAngularSleepingThreshold = PhysicsMaterial ? PhysicsMaterial->SleepingAngularThreshold : ChaosSolverCollisionDefaultAngularSleepThresholdCVar;
+					const int32 LocalSleepCounterThreshold = PhysicsMaterial ? PhysicsMaterial->SleepCounterThreshold : ChaosSolverCollisionDefaultSleepCounterThresholdCVar;
+
+					LinearSleepingThreshold = FMath::Min(LinearSleepingThreshold, LocalSleepingLinearThreshold);
+					AngularSleepingThreshold = FMath::Min(AngularSleepingThreshold, LocalAngularSleepingThreshold);
+					SleepCounterThreshold = FMath::Max(SleepCounterThreshold, LocalSleepCounterThreshold);
+				}
+			}
+		}
+	}
+	return bHaveSleepThreshold;
+}
+
+FPBDIslandManager::FPBDIslandManager() : IslandSolvers(), IslandGraph(MakeUnique<GraphType>()), MaxParticleIndex(INDEX_NONE)
+{}
+
+FPBDIslandManager::FPBDIslandManager(const TParticleView<FGeometryParticles>& GeometryParticles) : IslandSolvers(), IslandGraph(MakeUnique<GraphType>())
+{
+	InitializeGraph(GeometryParticles);
+}
+
+void FPBDIslandManager::InitializeGraph(const TParticleView<FGeometryParticles>& GeometryParticles)
+{
+	MaxParticleIndex = 0;
+	ReserveParticles(GeometryParticles.Num());
+	// We are adding all the particles from the solver in
+	// case some were just created/activated
+	for (auto& GeometryParticle : GeometryParticles)
+	{
+		AddParticle(GeometryParticle.Handle());
+	}
+	// For now we are resetting all the constraints but we should keep the
+	// persistent collisions, joints...over time
+	IslandGraph->ResetIslands();
+	for(auto& IslandSolver : IslandSolvers)
+	{
+		IslandSolver->ClearConstraints();
+	}
+}
+
+int32 FPBDIslandManager::ReserveParticles(const int32 NumParticles)
+{
+	const int32 MaxIndex = IslandGraph->NumNodes();
+	IslandGraph->ReserveNodes(NumParticles);
+	
+	IslandSolvers.Reserve(NumParticles);
+	IslandIndexing.Reserve(NumParticles);
+	
+	return FMath::Max(0,NumParticles - MaxIndex);
+}
+
+void FPBDIslandManager::ReserveConstraints(const int32 NumConstraints)
+{
+	IslandGraph->ReserveEdges(NumConstraints);
+}
+
+int32 FPBDIslandManager::AddParticle(FGeometryParticleHandle* ParticleHandle, const int32 IslandIndex)
+{
+	if (ParticleHandle)
+	{
+		// It could be nice to have a graph index on the particle handle the same way
+		// we have one on the constraint handle. It will allow us to skip the query on the TSet to
+		// check if the particle is already there, which could be quite slow
+		MaxParticleIndex = FMath::Max(MaxParticleIndex, ParticleHandle->UniqueIdx().Idx);
+		return IslandGraph->AddNode(ParticleHandle, IsDynamicParticle(ParticleHandle), IslandIndex,
+			(ParticleHandle->ObjectState() == EObjectStateType::Static) || ParticleHandle->Sleeping());
+	}
+	return INDEX_NONE;
+}
+
+int32 FPBDIslandManager::AddConstraint(const uint32 ContainerId, FConstraintHandle* ConstraintHandle, const TVec2<FGeometryParticleHandle*>& ConstrainedParticles)
+{
+	if (ConstraintHandle)
+	{
+		const bool bValidParticle0 = ConstrainedParticles[0] && (ConstrainedParticles[0]->ObjectState() == EObjectStateType::Dynamic);
+		const bool bValidParticle1 = ConstrainedParticles[1] && (ConstrainedParticles[1]->ObjectState() == EObjectStateType::Dynamic);
+		
+		// We are checking if one of the 2 particle is dynamic to add the constraint to the graph
+		// Will discard constraints in between 2 sleeping particles. Huge gain but potential side effect?
+		if(bValidParticle0 || bValidParticle1)
+		{
+			const int32 NodeIndex0 = AddParticle(ConstrainedParticles[0]);
+			const int32 NodeIndex1 = AddParticle(ConstrainedParticles[1]);
+
+			const int32 EdgeIndex = IslandGraph->AddEdge(ConstraintHandle, ContainerId, NodeIndex0, NodeIndex1);
+			ConstraintHandle->SetConstraintGraphIndex(EdgeIndex);
+
+			return EdgeIndex;
+		}
+		return INDEX_NONE;
+	}
+	return INDEX_NONE;
+}
+
+void FPBDIslandManager::RemoveParticle(FGeometryParticleHandle* ParticleHandle)
+{
+	if (ParticleHandle)
+	{
+		IslandGraph->RemoveNode(ParticleHandle);
+	}
+}
+
+void FPBDIslandManager::RemoveConstraint(const uint32 ContainerId, FConstraintHandle* ConstraintHandle, const TVec2<FGeometryParticleHandle*>& ConstrainedParticles)
+{
+	if (ConstraintHandle)
+	{
+		IslandGraph->RemoveEdge(ConstraintHandle);
+		ConstraintHandle->SetConstraintGraphIndex(INDEX_NONE);
+	}
+}
+	
+void FPBDIslandManager::EnableParticle(FGeometryParticleHandle* ChildParticle, const FGeometryParticleHandle* ParentParticle)
+{
+	if (ParentParticle && ChildParticle)
+	{
+		// We are only adding the child particle to the graph if the parent is dynamic
+		const FPBDRigidParticleHandle* ParentPBDRigid = ParentParticle->CastToRigidParticle();
+		if (ParentPBDRigid && ParentPBDRigid->ObjectState() == EObjectStateType::Dynamic)
+		{
+			int32 IslandIndex = INDEX_NONE;
+			FPBDRigidParticleHandle* ChildPBDRigid = ChildParticle->CastToRigidParticle();
+			if (ChildPBDRigid && IsDynamicParticle(ChildPBDRigid))
+			{
+				// If the child particle is dynamic or sleeping we are transferring
+				// the sleeping flag from the parent to the child
+				// and using the parent island index. if the island index we directly
+				// update the solver island without waiting the next sync
+				IslandIndex = ParentPBDRigid->Island();
+				ChildPBDRigid->SetSleeping(ParentPBDRigid->Sleeping());	
+				if (IslandSolvers.IsValidIndex(IslandIndex))
+				{
+					IslandSolvers[IslandIndex]->AddParticle(ChildParticle);
+				}
+			}
+			// We add the child particle to the graph 
+			AddParticle(ChildParticle, IslandIndex);
+		}
+	}
+}
+
+void FPBDIslandManager::DisableParticle(FGeometryParticleHandle* ParticleHandle)
+{
+	if (ParticleHandle)
+	{
+		if (const int32* NodeIndex = IslandGraph->ItemNodes.Find(ParticleHandle))
+		{
+			if (IslandGraph->GraphNodes.IsValidIndex(*NodeIndex))
+			{
+				// We loop over all the connected edges to find all the islands in which the particle
+				// is (static/kinematic particles could belong to several islands)
+				// And remove the particle from the solver island. It will allow the solver
+				// islands to be updated directly and not at the next sync.
+				for (const int32& EdgeIndex : IslandGraph->GraphNodes[*NodeIndex].NodeEdges)
+				{
+					const int32 IslandIndex = IslandGraph->GraphEdges[EdgeIndex].IslandIndex;
+					if (IslandSolvers.IsValidIndex(IslandIndex))
+					{
+						IslandSolvers[IslandIndex]->RemoveParticle(ParticleHandle);
+					}
+				}
+			}
+		}
+		// We remove the particle handle from the graph
+		RemoveParticle(ParticleHandle);
+	}
+}
+
+void FPBDIslandManager::ResetIslands(const TParticleView<FPBDRigidParticles>& PBDRigids)
+{
+	// @todo 
+}
+	
+void FPBDIslandManager::SyncIslands(FPBDRigidsSOAs& Particles)
+{
+	IslandSolvers.Reserve(IslandGraph->NumIslands());
+	IslandIndexing.SetNum(IslandGraph->NumIslands(),false);
+	int32 LocalIsland = 0;
+	for (int32 IslandIndex = 0, NumIslands = IslandGraph->NumIslands(); IslandIndex < NumIslands; ++IslandIndex)
+	{
+		if (IslandGraph->GraphIslands.IsValidIndex(IslandIndex))
+		{
+			// If the island item is not set on the graph we check if the solver island
+			// at the right index is already there. If not we create a new solver island.
+			if(!IslandSolvers.IsValidIndex(IslandIndex))
+			{
+				IslandSolvers.EmplaceAt(IslandIndex, MakeUnique<FPBDIslandSolver>(this, LocalIsland));
+			}
+			FPBDIslandSolver*& IslandSolver = IslandGraph->GraphIslands[IslandIndex].IslandItem;
+			IslandSolver = IslandSolvers[IslandIndex].Get();
+		
+			// We then transfer the persistent flag and the graph dense index to the solver island
+			IslandSolver->SetIsPersistent(IslandGraph->GraphIslands[IslandIndex].bIsPersistent);
+			IslandSolver->SetIsSleeping(IslandGraph->GraphIslands[IslandIndex].bIsSleeping);
+			IslandSolver->GetIslandIndex() = LocalIsland;
+
+			// We update the IslandIndexing to retrieve the graph sparse and persistent index from the dense one.
+			IslandIndexing[LocalIsland++] = IslandIndex;
+
+			// We finally update the solver islands based on the new particles and constraints if
+			// the island is not persistent or not sleeping. 
+			if(!IslandSolver->GetIsPersistent() || !IslandSolver->GetIsSleeping())
+			{
+				IslandSolver->ReserveParticles(IslandGraph->GraphIslands[IslandIndex].NumNodes);
+				IslandSolver->ReserveConstraints(IslandGraph->GraphIslands[IslandIndex].NumEdges);
+
+				IslandGraph->ForEachNodes(IslandIndex, [IslandSolver](GraphType::FGraphNode& GraphNode) { IslandSolver->AddParticle(GraphNode.NodeItem); });
+				IslandGraph->ForEachEdges(IslandIndex, [IslandSolver](GraphType::FGraphEdge& GraphEdge) { IslandSolver->AddConstraint(GraphEdge.EdgeItem); });
+			}
+			if (!IslandSolver->GetIsPersistent())
+			{
+				IslandSolver->SetSleepCounter(0);
+			}
+			// We reset the persistent flag to be true on the island graph. 
+			IslandGraph->GraphIslands[IslandIndex].bIsPersistent = true;
+			IslandGraph->GraphIslands[IslandIndex].bIsSleeping = true;
+
+			// Make all the particles asleep if the solver island is sleeping
+			MakeAsleep(IslandSolver, Particles);
+		}
+		else if (IslandSolvers.IsValidIndex(IslandIndex))
+		{
+			IslandSolvers.RemoveAt(IslandIndex);
+		}
+	}
+	IslandIndexing.SetNum(LocalIsland,false);
+}
+	
+void FPBDIslandManager::UpdateIslands(const TParticleView<FPBDRigidParticles>& PBDRigids, FPBDRigidsSOAs& Particles)
+{
+	for (auto& PBDRigid : PBDRigids)
+	{
+		AddParticle(PBDRigid.Handle());
+	}
+	// Merge the graph islands if required
+	IslandGraph->UpdateGraph();
+	
+	// Sync the graph islands with the solver islands objects
+	SyncIslands(Particles);
+}
+
+bool FPBDIslandManager::SleepInactive(const int32 IslandIndex, const TArrayCollectionArray<TSerializablePtr<FChaosPhysicsMaterial>>& PerParticleMaterialAttributes, const THandleArray<FChaosPhysicsMaterial>& SolverPhysicsMaterials)
+{
+	// Only the persistent islands could start sleeping
+	const int32 GraphIndex = GetGraphIndex(IslandIndex);
+	if (!ChaosSolverSleepEnabled || !IslandSolvers.IsValidIndex(GraphIndex) ||
+		(IslandSolvers.IsValidIndex(GraphIndex) && !IslandSolvers[GraphIndex]->GetIsPersistent()) )
+	{
+		return false;
+	}
+
+	FReal LinearSleepingThreshold = FLT_MAX,  AngularSleepingThreshold = FLT_MAX, MaxLinearSpeed2 = 0.f, MaxAngularSpeed2 = 0.f;
+	int32 SleepCounterThreshold = 0, NumDynamicParticles = 0;
+
+	// Compute of the linear/angular velocities + thresholds to make islands sleeping 
+	if (ComputeSleepingThresholds(IslandSolvers[GraphIndex].Get(), PerParticleMaterialAttributes,
+		SolverPhysicsMaterials, LinearSleepingThreshold, AngularSleepingThreshold,
+		MaxLinearSpeed2, MaxAngularSpeed2, SleepCounterThreshold))
+	{
+		int32 SleepCounter = IslandSolvers[GraphIndex]->GetSleepCounter();
+		const bool bSleepingIsland = IsIslandSleeping(MaxLinearSpeed2, MaxAngularSpeed2, LinearSleepingThreshold,
+			AngularSleepingThreshold, SleepCounterThreshold, SleepCounter);
+		
+		IslandSolvers[GraphIndex]->SetSleepCounter(SleepCounter);
+
+		return bSleepingIsland;
+	}
+	return false;
+}
+
+void FPBDIslandManager::SleepIsland(FPBDRigidsSOAs& Particles, const int32 IslandIndex)
+{
+	const int32 GraphIndex = GetGraphIndex(IslandIndex);
+	if (IslandSolvers.IsValidIndex(GraphIndex))
+	{
+		Particles.DeactivateParticles(GetIslandParticles(IslandIndex));
+
+		//for (FConstraintHandle* ConstraintHandle : IslandSolvers[GraphIndex]->GetConstraints())
+		//{
+		//	ConstraintHandle->SetIsSleeping(true);
+		//}
+	}
+
+}
+
+void FPBDIslandManager::WakeIsland(FPBDRigidsSOAs& Particles, const int32 IslandIndex)
+{
+	const int32 GraphIndex = GetGraphIndex(IslandIndex);
+	if (IslandSolvers.IsValidIndex(GraphIndex))
+	{
+		for (FGeometryParticleHandle* ParticleHandle : IslandSolvers[GraphIndex]->GetParticles())
+		{
+			if (ParticleHandle)
+			{
+				if (FPBDRigidParticleHandle* PBDRigid = ParticleHandle->CastToRigidParticle())
+				{
+					if (IsDynamicParticle(ParticleHandle) && PBDRigid->Sleeping())
+					{
+						PBDRigid->SetSleeping(false);
+
+						// Do we need  to enable the particle since the state is already dynamic/sleeping?
+						Particles.EnableParticle(ParticleHandle);
+					}
+				}
+			}
+
+			//for (FConstraintHandle* ConstraintHandle : IslandSolvers[GraphIndex]->GetConstraints())
+			//{
+			//	ConstraintHandle->SetIsSleeping(false);
+			//}
+		}
+		IslandSolvers[GraphIndex]->SetSleepCounter(0);
+	}
+}
+	
+const TArray<FGeometryParticleHandle*>& FPBDIslandManager::GetIslandParticles(const int32 IslandIndex) const
+{
+	return IslandSolvers[GetGraphIndex(IslandIndex)]->GetParticles();
+}
+
+const TArray<FConstraintHandle*>& FPBDIslandManager::GetIslandConstraints(const int32 IslandIndex) const
+{
+	return IslandSolvers[GetGraphIndex(IslandIndex)]->GetConstraints();
+}
+
+bool FPBDIslandManager::IslandNeedsResim(const int32 IslandIndex) const
+{
+	return IslandSolvers[GetGraphIndex(IslandIndex)]->GetNeedsResim();
+}
+
+}
\ No newline at end of file
diff --git a/Engine/Source/Runtime/Experimental/Chaos/Private/Chaos/Island/SolverIsland.cpp b/Engine/Source/Runtime/Experimental/Chaos/Private/Chaos/Island/SolverIsland.cpp
new file mode 100644
--- /dev/null
+++ b/Engine/Source/Runtime/Experimental/Chaos/Private/Chaos/Island/SolverIsland.cpp
@@ -0,0 +1,86 @@
+// Copyright Epic Games, Inc. All Rights Reserved.
+#include "Chaos/Island/SolverIsland.h"
+#include "Chaos/PBDConstraintRule.h"
+
+namespace Chaos
+{
+	
+FPBDIslandSolver::FPBDIslandSolver(const FPBDIslandManager* IslandManagerIn, const int32 IslandIndexIn) : FPBDIslandSolverData(IslandIndexIn),
+	IslandManager(IslandManagerIn), bIsSleeping(false), bNeedsResim(false), bIsPersistent(true), SleepCounter(0),
+	IslandParticles(), IslandConstraints()
+{}
+
+void FPBDIslandSolver::UpdateParticles()
+{
+	for (FGenericParticleHandle ParticleHandle : IslandParticles)
+	{
+		if (ParticleHandle.IsValid() && ParticleHandle->IsDynamic())
+		{
+			ParticleHandle->SetIsland(IslandIndex);
+		}
+	}
+}
+	
+void FPBDIslandSolver::ClearParticles()
+{
+	IslandParticles.Reset();
+}
+
+void FPBDIslandSolver::AddParticle(FGenericParticleHandle ParticleHandle)
+{
+	if (ParticleHandle.IsValid())
+	{
+		if (ParticleHandle->IsDynamic())
+		{
+			ParticleHandle->SetIsland(IslandIndex);
+		}
+		IslandParticles.Add(ParticleHandle->Handle());
+	}
+}
+
+void FPBDIslandSolver::RemoveParticle(FGenericParticleHandle ParticleHandle)
+{
+	if (ParticleHandle.IsValid())
+	{
+		if (ParticleHandle->IsDynamic())
+		{
+			ParticleHandle->SetIsland(INDEX_NONE);
+		}
+		IslandParticles.Remove(ParticleHandle->Handle());
+	}
+}
+
+void FPBDIslandSolver::ReserveParticles(const int32 NumParticles)
+{
+	ClearParticles();
+	IslandParticles.Reserve(NumParticles);
+}
+
+void FPBDIslandSolver::AddConstraint(FConstraintHandle* ConstraintHandle)
+{
+	if (ConstraintHandle)
+	{
+		IslandConstraints.Add(ConstraintHandle);
+	}
+}
+
+void FPBDIslandSolver::RemoveConstraint(FConstraintHandle* ConstraintHandle)
+{
+	if (ConstraintHandle)
+	{
+		IslandConstraints.Remove(ConstraintHandle);
+	}
+}
+
+void FPBDIslandSolver::ReserveConstraints(const int32 NumConstraints)
+{
+	ClearConstraints();
+	IslandConstraints.Reserve(NumConstraints);
+}
+
+void FPBDIslandSolver::ClearConstraints()
+{
+	IslandConstraints.Reset();
+}
+
+}
\ No newline at end of file
diff --git a/Engine/Source/Runtime/Experimental/Chaos/Private/Chaos/Joint/JointConstraintsCVars.cpp b/Engine/Source/Runtime/Experimental/Chaos/Private/Chaos/Joint/JointConstraintsCVars.cpp
--- a/Engine/Source/Runtime/Experimental/Chaos/Private/Chaos/Joint/JointConstraintsCVars.cpp
+++ b/Engine/Source/Runtime/Experimental/Chaos/Private/Chaos/Joint/JointConstraintsCVars.cpp
@@ -3,11 +3,12 @@
 #include "HAL/IConsoleManager.h"
 
 #if INTEL_ISPC && !UE_BUILD_SHIPPING
-bool bChaos_Joint_ISPC_Enabled = true;
+bool bChaos_Joint_ISPC_Enabled = false;
 FAutoConsoleVariableRef CVarChaosJointISPCEnabled(TEXT("p.Chaos.Joint.ISPC"), bChaos_Joint_ISPC_Enabled, TEXT("Whether to use ISPC optimizations in the Joint Solver"));
 #endif
 
-bool bChaos_Joint_EarlyOut_Enabled = true;
+// @todo(chaos): fix joint early out. Joints that are "solved" can be "unsolved" by changes to the bodies they are connected to.
+bool bChaos_Joint_EarlyOut_Enabled = false;
 FAutoConsoleVariableRef CVarChaosJointEarlyOutEnabled(TEXT("p.Chaos.Joint.EarlyOut"), bChaos_Joint_EarlyOut_Enabled, TEXT("Whether to iterating when joints report being solved"));
 
 float Chaos_Joint_DegenerateRotationLimit = -0.998f;	// Cos(176deg)
@@ -18,3 +19,9 @@ FAutoConsoleVariableRef CVarChaosJointVelProjectionScale(TEXT("p.Chaos.Joint.Vel
 
 bool bChaos_Joint_DisableSoftLimits = false;
 FAutoConsoleVariableRef CVarChaosJointDisableSoftLimits(TEXT("p.Chaos.Joint.DisableSoftLimits"), bChaos_Joint_DisableSoftLimits, TEXT("Disable soft limits (for debugging only)"));
+
+float Chaos_Joint_LinearVelocityThresholdToApplyRestitution = 1e-2f;
+FAutoConsoleVariableRef CVarChaosJointLinearVelocityThresholdToApplyRestitution(TEXT("p.Chaos.Joint.LinearVelocityThresholdToApplyRestitution"), Chaos_Joint_LinearVelocityThresholdToApplyRestitution, TEXT("Apply restitution only if initial velocity is higher than this threshold (used in Quasipbd)"));
+
+float Chaos_Joint_AngularVelocityThresholdToApplyRestitution = 1e-2f;
+FAutoConsoleVariableRef CVarChaosJointAngularVelocityThresholdToApplyRestitution(TEXT("p.Chaos.Joint.AngularVelocityThresholdToApplyRestitution"), Chaos_Joint_AngularVelocityThresholdToApplyRestitution, TEXT("Apply restitution only if initial velocity is higher than this threshold (used in Quasipbd)"));
\ No newline at end of file
diff --git a/Engine/Source/Runtime/Experimental/Chaos/Private/Chaos/Joint/JointConstraintsCVars.h b/Engine/Source/Runtime/Experimental/Chaos/Private/Chaos/Joint/JointConstraintsCVars.h
--- a/Engine/Source/Runtime/Experimental/Chaos/Private/Chaos/Joint/JointConstraintsCVars.h
+++ b/Engine/Source/Runtime/Experimental/Chaos/Private/Chaos/Joint/JointConstraintsCVars.h
@@ -7,7 +7,7 @@
 #if !INTEL_ISPC
 	const bool bChaos_Joint_ISPC_Enabled = false;
 #elif UE_BUILD_SHIPPING
-	const bool bChaos_Joint_ISPC_Enabled = true;
+	const bool bChaos_Joint_ISPC_Enabled = false;
 #else
 	extern bool bChaos_Joint_ISPC_Enabled;
 #endif
@@ -19,3 +19,8 @@ extern float Chaos_Joint_DegenerateRotationLimit;
 extern float Chaos_Joint_VelProjectionAlpha;
 
 extern bool bChaos_Joint_DisableSoftLimits;
+
+extern float Chaos_Joint_LinearVelocityThresholdToApplyRestitution;
+
+extern float Chaos_Joint_AngularVelocityThresholdToApplyRestitution;
+
