Chaos

diff --git a/Engine/Source/Runtime/Experimental/Chaos/Private/Chaos/Joint/PBDJointSolverGaussSeidel.cpp b/Engine/Source/Runtime/Experimental/Chaos/Private/Chaos/Joint/PBDJointSolverGaussSeidel.cpp
--- a/Engine/Source/Runtime/Experimental/Chaos/Private/Chaos/Joint/PBDJointSolverGaussSeidel.cpp
+++ b/Engine/Source/Runtime/Experimental/Chaos/Private/Chaos/Joint/PBDJointSolverGaussSeidel.cpp
@@ -23,64 +23,61 @@ namespace Chaos
 	//
 	//
 
-	FJointSolverGaussSeidel::FJointSolverGaussSeidel()
+	FPBDJointSolver::FPBDJointSolver()
 	{
 		if (bChaos_Joint_ISPC_Enabled)
 		{
 #if INTEL_ISPC
-			check(sizeof(FJointSolverGaussSeidel) == ispc::SizeofFJointSolverGaussSeidel());
+			check(sizeof(FPBDJointSolver) == ispc::SizeofFPBDJointSolver());
 #endif
 		}
 	}
 
-
-	void FJointSolverGaussSeidel::InitDerivedState()
+	void FPBDJointSolver::InitDerivedState()
 	{
-		InitXs[0] = InitPs[0] + InitQs[0] * XLs[0].GetTranslation();
-		InitXs[1] = InitPs[1] + InitQs[1] * XLs[1].GetTranslation();
-		InitRs[0] = InitQs[0] * XLs[0].GetRotation();
-		InitRs[1] = InitQs[1] * XLs[1].GetRotation();
-		InitRs[1].EnforceShortestArcWith(InitRs[0]);
+		InitConnectorXs[0] = X(0) + R(0) * LocalConnectorXs[0].GetTranslation();
+		InitConnectorXs[1] = X(1) + R(1) * LocalConnectorXs[1].GetTranslation();
+		InitConnectorRs[0] = R(0) * LocalConnectorXs[0].GetRotation();
+		InitConnectorRs[1] = R(1) * LocalConnectorXs[1].GetRotation();
+		InitConnectorRs[1].EnforceShortestArcWith(InitConnectorRs[0]);
 
-		Xs[0] = InitXs[0];
-		Rs[0] = InitRs[0];
-		InvIs[0] = (InvMs[0] > 0.0f) ? Utilities::ComputeWorldSpaceInertia(InitQs[0], InvILs[0]) : FMatrix33(0, 0, 0);
+		ConnectorXs[0] = InitConnectorXs[0];
+		ConnectorRs[0] = InitConnectorRs[0];
 
-		Xs[1] = InitXs[1];
-		Rs[1] = InitRs[1];
-		InvIs[1] = (InvMs[1] > 0.0f) ? Utilities::ComputeWorldSpaceInertia(InitQs[1], InvILs[1]) : FMatrix33(0, 0, 0);
+		ConnectorXs[1] = InitConnectorXs[1];
+		ConnectorRs[1] = InitConnectorRs[1];
 	}
 
 
-	void FJointSolverGaussSeidel::UpdateDerivedState()
+	void FPBDJointSolver::UpdateDerivedState()
 	{
 		// Kinematic bodies will not be moved, so we don't update derived state during iterations
-		if (InvMs[0] > 0.0f)
+		if (InvM(0) > SMALL_NUMBER)
 		{
-			Xs[0] = Ps[0] + Qs[0] * XLs[0].GetTranslation();
-			Rs[0] = Qs[0] * XLs[0].GetRotation();
-			InvIs[0] = Utilities::ComputeWorldSpaceInertia(Qs[0], InvILs[0]);
+			ConnectorXs[0] = P(0) + Q(0) * LocalConnectorXs[0].GetTranslation();
+			ConnectorRs[0] = Q(0) * LocalConnectorXs[0].GetRotation();
+			SolverBodies[0].UpdateRotationDependentState();
 		}
-		if (InvMs[1] > 0.0f)
+		if (InvM(1) > SMALL_NUMBER)
 		{
-			Xs[1] = Ps[1] + Qs[1] * XLs[1].GetTranslation();
-			Rs[1] = Qs[1] * XLs[1].GetRotation();
-			InvIs[1] = Utilities::ComputeWorldSpaceInertia(Qs[1], InvILs[1]);
+			ConnectorXs[1] = P(1) + Q(1) * LocalConnectorXs[1].GetTranslation();
+			ConnectorRs[1] = Q(1) * LocalConnectorXs[1].GetRotation();
+			SolverBodies[1].UpdateRotationDependentState();
 		}
-		Rs[1].EnforceShortestArcWith(Rs[0]);
+		ConnectorRs[1].EnforceShortestArcWith(ConnectorRs[0]);
 	}
 
 
-	void FJointSolverGaussSeidel::UpdateDerivedState(const int32 BodyIndex)
+	void FPBDJointSolver::UpdateDerivedState(const int32 BodyIndex)
 	{
-		Xs[BodyIndex] = Ps[BodyIndex] + Qs[BodyIndex] * XLs[BodyIndex].GetTranslation();
-		Rs[BodyIndex] = Qs[BodyIndex] * XLs[BodyIndex].GetRotation();
-		Rs[1].EnforceShortestArcWith(Rs[0]);
-	
-		InvIs[BodyIndex] = Utilities::ComputeWorldSpaceInertia(Qs[BodyIndex], InvILs[BodyIndex]);
+		ConnectorXs[BodyIndex] = P(BodyIndex) + Q(BodyIndex) * LocalConnectorXs[BodyIndex].GetTranslation();
+		ConnectorRs[BodyIndex] = Q(BodyIndex) * LocalConnectorXs[BodyIndex].GetRotation();
+		ConnectorRs[1].EnforceShortestArcWith(ConnectorRs[0]);
+		
+		SolverBodies[BodyIndex].UpdateRotationDependentState();
 	}
 
-	bool FJointSolverGaussSeidel::UpdateIsActive(
+	bool FPBDJointSolver::UpdateIsActive(
 		const FReal Dt,
 		const FPBDJointSolverSettings& SolverSettings,
 		const FPBDJointSettings& JointSettings)
@@ -90,51 +87,42 @@ namespace Chaos
 		if ((NumActiveConstraints >= 0) && (SolverStiffness > 0.0f))
 		{
 			bool bIsSolved =
-				FVec3::IsNearlyEqual(Ps[0], LastPs[0], PositionTolerance)
-				&& FVec3::IsNearlyEqual(Ps[1], LastPs[1], PositionTolerance)
-				&& FRotation3::IsNearlyEqual(Qs[0], LastQs[0], 0.5f * AngleTolerance)
-				&& FRotation3::IsNearlyEqual(Qs[1], LastQs[1], 0.5f * AngleTolerance);
+				FVec3::IsNearlyEqual(P(0), LastPs[0], PositionTolerance)
+				&& FVec3::IsNearlyEqual(P(1), LastPs[1], PositionTolerance)
+				&& FRotation3::IsNearlyEqual(Q(0), LastQs[0], 0.5f * AngleTolerance)
+				&& FRotation3::IsNearlyEqual(Q(1), LastQs[1], 0.5f * AngleTolerance);
 			bIsActive = !bIsSolved;
 		}
 
-		LastPs[0] = Ps[0];
-		LastPs[1] = Ps[1];
-		LastQs[0] = Qs[0];
-		LastQs[1] = Qs[1];
+		LastPs[0] = P(0);
+		LastPs[1] = P(1);
+		LastQs[0] = Q(0);
+		LastQs[1] = Q(1);
 
 		return bIsActive;
 	}
 
 
-	void FJointSolverGaussSeidel::Init(
+	void FPBDJointSolver::Init(
 		const FReal Dt,
+		const FSolverBodyPtrPair& SolverBodyPair,
 		const FPBDJointSolverSettings& SolverSettings,
 		const FPBDJointSettings& JointSettings,
-		const FVec3& PrevP0,
-		const FVec3& PrevP1,
-		const FRotation3& PrevQ0,
-		const FRotation3& PrevQ1,
-		const FReal InvM0,
-		const FVec3& InvIL0,
-		const FReal InvM1,
-		const FVec3& InvIL1,
 		const FRigidTransform3& XL0,
 		const FRigidTransform3& XL1)
 	{
-		XLs[0] = XL0;
-		XLs[1] = XL1;
+		SolverBodies[0] = *SolverBodyPair[0];
+		SolverBodies[1] = *SolverBodyPair[1];
 
-		InvILs[0] = JointSettings.ParentInvMassScale * InvIL0;
-		InvILs[1] = InvIL1;
-		InvMs[0] = JointSettings.ParentInvMassScale * InvM0;
-		InvMs[1] = InvM1;
+		LocalConnectorXs[0] = XL0;
+		LocalConnectorXs[1] = XL1;
 
-		FPBDJointUtilities::ConditionInverseMassAndInertia(InvMs[0], InvMs[1], InvILs[0], InvILs[1], SolverSettings.MinParentMassRatio, SolverSettings.MaxInertiaRatio);
+		// \todo(chaos): joint should support parent/child in either order
+		SolverBodies[0].SetInvMassScale(JointSettings.ParentInvMassScale);
+		SolverBodies[1].SetInvMassScale(FReal(1));
 
-		InitPs[0] = PrevP0;
-		InitPs[1] = PrevP1;
-		InitQs[0] = PrevQ0;
-		InitQs[1] = PrevQ1;
+		// \chaos(todo): Add support for mass conditioning in ConstraintSolverBody
+		//FPBDJointUtilities::ConditionInverseMassAndInertia(InvM(0), InvM(1), InvILs[0], InvILs[1], SolverSettings.MinParentMassRatio, SolverSettings.MaxInertiaRatio);
 
 		NetLinearImpulse = FVec3(0);
 		NetAngularImpulse = FVec3(0);
@@ -168,48 +156,44 @@ namespace Chaos
 		AngularHardLambda = FVec3(0);
 
 		InitDerivedState();
+
+		if (JointSettings.LinearRestitution != 0.0f)
+		{
+			CalculateConstraintAxisLinearVelocities(JointSettings, InitConstraintAxisLinearVelocities);
+		}
+		if (JointSettings.TwistRestitution != 0.0f || JointSettings.SwingRestitution != 0.0f)
+		{
+			CalculateConstraintAxisAngularVelocities(SolverSettings, JointSettings, InitConstraintAxisAngularVelocities);
+		}
 	}
 
 
-	void FJointSolverGaussSeidel::Update(
-		const FReal Dt,
-		const FReal InSolverStiffness,
-		const FPBDJointSolverSettings& SolverSettings,
-		const FPBDJointSettings& JointSettings,
-		const FVec3& P0,
-		const FRotation3& Q0,
-		const FVec3& V0,
-		const FVec3& W0,
-		const FVec3& P1,
-		const FRotation3& Q1,
-		const FVec3& V1,
-		const FVec3& W1)
-	{
-		Ps[0] = P0;
-		Ps[1] = P1;
-		Qs[0] = Q0;
-		Qs[1] = Q1;
-		Qs[1].EnforceShortestArcWith(Qs[0]);
-
-		Vs[0] = V0;
-		Vs[1] = V1;
-		Ws[0] = W0;
-		Ws[1] = W1;
+	void FPBDJointSolver::Deinit()
+	{
+		SolverBodies[0].Reset();
+		SolverBodies[1].Reset();
+	}
 
-		SolverStiffness = InSolverStiffness;
 
+	void FPBDJointSolver::Update(
+		const FReal Dt,
+		const FPBDJointSolverSettings& SolverSettings,
+		const FPBDJointSettings& JointSettings)
+	{
 		UpdateDerivedState();
 
 		UpdateIsActive(Dt, SolverSettings, JointSettings);
 	}
 
 
-	void FJointSolverGaussSeidel::ApplyConstraints(
+	void FPBDJointSolver::ApplyConstraints(
 		const FReal Dt,
+		const FReal InSolverStiffness,
 		const FPBDJointSolverSettings& SolverSettings,
 		const FPBDJointSettings& JointSettings)
 	{
 		NumActiveConstraints = 0;
+		SolverStiffness = InSolverStiffness;
 
 		ApplyPositionConstraints(Dt, SolverSettings, JointSettings);
 		ApplyRotationConstraints(Dt, SolverSettings, JointSettings);
@@ -221,11 +205,14 @@ namespace Chaos
 	}
 
 
-	void FJointSolverGaussSeidel::ApplyVelocityConstraints(
+	void FPBDJointSolver::ApplyVelocityConstraints(
 		const FReal Dt,
+		const FReal InSolverStiffness,
 		const FPBDJointSolverSettings& SolverSettings,
 		const FPBDJointSettings& JointSettings)
 	{
+		SolverStiffness = InSolverStiffness;
+
 		// This is used for the QuasiPbd solver. If the Pbd step applied impulses to
 		// correct position errors, it will have introduced a velocity equal to the 
 		// corection divided by the timestep. We ensure that the velocity constraints
@@ -235,21 +222,25 @@ namespace Chaos
 		ApplyLinearVelocityConstraints(Dt, SolverSettings, JointSettings);
 		ApplyAngularVelocityConstraints(Dt, SolverSettings, JointSettings);
 
-		// @todo(chaos): We can also apply velocity drrives here rather than in the Pbd pass
+		// @todo(chaos): We can also apply velocity drives here rather than in the Pbd pass
 	}
 
 
-	void FJointSolverGaussSeidel::ApplyProjections(
+	void FPBDJointSolver::ApplyProjections(
 		const FReal Dt,
+		const FReal InSolverStiffness,
 		const FPBDJointSolverSettings& SolverSettings,
 		const FPBDJointSettings& JointSettings)
 	{
-		if (InvMs[1] < SMALL_NUMBER)
+		// @todo(chaos): We need to handle parent/child being the other way round
+		if (InvM(1) < SMALL_NUMBER)
 		{
 			// If child is kinematic, return. 
 			return;
 		}
-		
+
+		SolverStiffness = InSolverStiffness;
+
 		const FReal LinearProjection = FPBDJointUtilities::GetLinearProjection(SolverSettings, JointSettings);
 		const FReal AngularProjection = FPBDJointUtilities::GetAngularProjection(SolverSettings, JointSettings);
 
@@ -352,7 +343,7 @@ namespace Chaos
 	}
 
 
-	void FJointSolverGaussSeidel::ApplyRotationConstraints(
+	void FPBDJointSolver::ApplyRotationConstraints(
 		const FReal Dt,
 		const FPBDJointSolverSettings& SolverSettings,
 		const FPBDJointSettings& JointSettings)
@@ -374,8 +365,8 @@ namespace Chaos
 		bool bSwingSoft = FPBDJointUtilities::GetSoftSwingLimitEnabled(SolverSettings, JointSettings);
 
 		// If the twist axes are opposing, we cannot decompose the orientation into swing and twist angles, so just give up
-		const FVec3 Twist0 = Rs[0] * FJointConstants::TwistAxis();
-		const FVec3 Twist1 = Rs[1] * FJointConstants::TwistAxis();
+		const FVec3 Twist0 = ConnectorRs[0] * FJointConstants::TwistAxis();
+		const FVec3 Twist1 = ConnectorRs[1] * FJointConstants::TwistAxis();
 		const FReal Twist01Dot = FVec3::DotProduct(Twist0, Twist1);
 		const bool bDegenerate = (Twist01Dot < Chaos_Joint_DegenerateRotationLimit);
 		if (bDegenerate)
@@ -468,7 +459,7 @@ namespace Chaos
 	}
 
 
-	void FJointSolverGaussSeidel::ApplyRotationDrives(
+	void FPBDJointSolver::ApplyRotationDrives(
 		const FReal Dt,
 		const FPBDJointSolverSettings& SolverSettings,
 		const FPBDJointSettings& JointSettings)
@@ -516,7 +507,7 @@ namespace Chaos
 	}
 
 
-	void FJointSolverGaussSeidel::ApplyPositionConstraints(
+	void FPBDJointSolver::ApplyPositionConstraints(
 		const FReal Dt,
 		const FPBDJointSolverSettings& SolverSettings,
 		const FPBDJointSettings& JointSettings)
@@ -547,11 +538,11 @@ namespace Chaos
 		if (bLinearLocked[0] && bLinearLocked[1] && bLinearLocked[2])
 		{
 			// Hard point constraint (most common case)
-			if (InvMs[0] == 0)
+			if (InvM(0) < SMALL_NUMBER)
 			{
 				ApplyPointPositionConstraintKD(0, 1, Dt, SolverSettings, JointSettings);
 			}
-			else if (InvMs[1] == 0)
+			else if (InvM(1) < SMALL_NUMBER)
 			{
 				ApplyPointPositionConstraintKD(1, 0, Dt, SolverSettings, JointSettings);
 			}
@@ -613,7 +604,7 @@ namespace Chaos
 	}
 
 
-	void FJointSolverGaussSeidel::ApplyPositionDrives(
+	void FPBDJointSolver::ApplyPositionDrives(
 		const FReal Dt,
 		const FPBDJointSolverSettings& SolverSettings,
 		const FPBDJointSettings& JointSettings)
@@ -630,10 +621,10 @@ namespace Chaos
 			// Rectangular position drives
 			if (bDriven[0] || bDriven[1] || bDriven[2])
 			{
-				const FMatrix33 R0M = Rs[0].ToMatrix();
-				const FVec3 XTarget = Xs[0] + Rs[0] * JointSettings.LinearDrivePositionTarget;
-				const FVec3 VTarget = Rs[0] * JointSettings.LinearDriveVelocityTarget;
-				const FVec3 CX = Xs[1] - XTarget;
+				const FMatrix33 R0M = ConnectorRs[0].ToMatrix();
+				const FVec3 XTarget = ConnectorXs[0] + ConnectorRs[0] * JointSettings.LinearDrivePositionTarget;
+				const FVec3 VTarget = ConnectorRs[0] * JointSettings.LinearDriveVelocityTarget;
+				const FVec3 CX = ConnectorXs[1] - XTarget;
 
 				for (int32 AxisIndex = 0; AxisIndex < 3; ++AxisIndex)
 				{
@@ -658,7 +649,7 @@ namespace Chaos
 	//
 
 
-	void FJointSolverGaussSeidel::ApplyLinearVelocityConstraints(
+	void FPBDJointSolver::ApplyLinearVelocityConstraints(
 		const FReal Dt,
 		const FPBDJointSolverSettings& SolverSettings,
 		const FPBDJointSettings& JointSettings)
@@ -742,7 +733,7 @@ namespace Chaos
 		}
 	}
 
-	void FJointSolverGaussSeidel::ApplyAngularVelocityConstraints(
+	void FPBDJointSolver::ApplyAngularVelocityConstraints(
 		const FReal Dt,
 		const FPBDJointSolverSettings& SolverSettings,
 		const FPBDJointSettings& JointSettings)
@@ -764,8 +755,8 @@ namespace Chaos
 		bool bSwingSoft = FPBDJointUtilities::GetSoftSwingLimitEnabled(SolverSettings, JointSettings);
 
 		// If the twist axes are opposing, we cannot decompose the orientation into swing and twist angles, so just give up
-		const FVec3 Twist0 = Rs[0] * FJointConstants::TwistAxis();
-		const FVec3 Twist1 = Rs[1] * FJointConstants::TwistAxis();
+		const FVec3 Twist0 = ConnectorRs[0] * FJointConstants::TwistAxis();
+		const FVec3 Twist1 = ConnectorRs[1] * FJointConstants::TwistAxis();
 		const FReal Twist01Dot = FVec3::DotProduct(Twist0, Twist1);
 		const bool bDegenerate = (Twist01Dot < Chaos_Joint_DegenerateRotationLimit);
 		if (bDegenerate)
@@ -857,7 +848,7 @@ namespace Chaos
 		}
 	}
 
-	void FJointSolverGaussSeidel::ApplyTwistVelocityConstraint(
+	void FPBDJointSolver::ApplyTwistVelocityConstraint(
 		const FReal Dt,
 		const FPBDJointSolverSettings& SolverSettings,
 		const FPBDJointSettings& JointSettings,
@@ -865,14 +856,28 @@ namespace Chaos
 	{
 		if(!NetAngularImpulse.IsNearlyZero() && !bUseSoftLimit && AngularHardLambda[(int32)EJointAngularConstraintIndex::Twist] > SMALL_NUMBER)
 		{
-			const FVec3 TwistAxis = Rs[1] * FJointConstants::TwistAxis();
+			FVec3 TwistAxis;
+			FReal TwistAngle;
+			FPBDJointUtilities::GetTwistAxisAngle(ConnectorRs[0], ConnectorRs[1], TwistAxis, TwistAngle);
+
+			if (TwistAngle < 0.0f)
+			{
+				TwistAxis = -TwistAxis;
+			}
 
 			const FReal TwistStiffness = SolverStiffness * FPBDJointUtilities::GetTwistStiffness(SolverSettings, JointSettings);
-			ApplyAngularVelocityConstraint(TwistStiffness, TwistAxis);
+
+			FReal TargetVel = 0.0f;
+			if (JointSettings.TwistRestitution != 0.0f)
+			{
+				const FReal InitVel = InitConstraintAxisAngularVelocities[(int32)EJointAngularConstraintIndex::Twist];
+				TargetVel = InitVel > Chaos_Joint_AngularVelocityThresholdToApplyRestitution ? -JointSettings.TwistRestitution * InitVel : 0.0f;
+			}
+			ApplyAngularVelocityConstraint(TwistStiffness, TwistAxis, TargetVel);
 		}
 	}
 
-	void FJointSolverGaussSeidel::ApplyConeVelocityConstraint(
+	void FPBDJointSolver::ApplyConeVelocityConstraint(
 		const FReal Dt,
 		const FPBDJointSolverSettings& SolverSettings,
 		const FPBDJointSettings& JointSettings,
@@ -885,16 +890,22 @@ namespace Chaos
 			const FReal LimitPadding = GetAngularConstraintPadding(EJointAngularConstraintIndex::Swing1);
 			const FReal Swing1Limit = FMath::Max(JointSettings.AngularLimits[(int32)EJointAngularConstraintIndex::Swing1] - LimitPadding, 0.0f);
 			const FReal Swing2Limit = FMath::Max(JointSettings.AngularLimits[(int32)EJointAngularConstraintIndex::Swing2] - LimitPadding, 0.0f);
-			FPBDJointUtilities::GetEllipticalConeAxisErrorLocal(Rs[0], Rs[1], Swing2Limit, Swing1Limit, SwingAxisLocal, DSwingAngle);
+			FPBDJointUtilities::GetEllipticalConeAxisErrorLocal(ConnectorRs[0], ConnectorRs[1], Swing2Limit, Swing1Limit, SwingAxisLocal, DSwingAngle);
 			// SwingAxisLocal has the size of Sine(SwingAngle), thus we need to normalize it
 			SwingAxisLocal.SafeNormalize();
-			const FVec3 SwingAxis = Rs[0] * SwingAxisLocal;
+			const FVec3 SwingAxis = ConnectorRs[0] * SwingAxisLocal;
 			const FReal SwingStiffness = SolverStiffness * FPBDJointUtilities::GetSwingStiffness(SolverSettings, JointSettings);
-			ApplyAngularVelocityConstraint(SwingStiffness, SwingAxis);
+			FReal TargetVel = 0.0f;
+			if (JointSettings.SwingRestitution != 0.0f)
+			{
+				const FReal InitVel = InitConstraintAxisAngularVelocities[(int32)EJointAngularConstraintIndex::Swing1];
+				TargetVel = InitVel > Chaos_Joint_AngularVelocityThresholdToApplyRestitution ? -JointSettings.SwingRestitution * InitVel : 0.0f;
+			}
+			ApplyAngularVelocityConstraint(SwingStiffness, SwingAxis, TargetVel);
 		}
 	}
 
-	void FJointSolverGaussSeidel::ApplySingleLockedSwingVelocityConstraint(
+	void FPBDJointSolver::ApplySingleLockedSwingVelocityConstraint(
 		const FReal Dt,
 		const FPBDJointSolverSettings& SolverSettings,
 		const FPBDJointSettings& JointSettings,
@@ -905,7 +916,7 @@ namespace Chaos
 		{
 			FVec3 SwingAxis;
 			FReal SwingAngle;
-			FPBDJointUtilities::GetLockedSwingAxisAngle(Rs[0], Rs[1], SwingConstraintIndex, SwingAxis, SwingAngle);
+			FPBDJointUtilities::GetLockedSwingAxisAngle(ConnectorRs[0], ConnectorRs[1], SwingConstraintIndex, SwingAxis, SwingAngle);
 			// Swing axis has the size of Sine(SwingAngle), thus we need to normalize it
 			SwingAxis.SafeNormalize();
 
@@ -914,7 +925,7 @@ namespace Chaos
 		}
 	}
 
-	void FJointSolverGaussSeidel::ApplyDualConeSwingVelocityConstraint(
+	void FPBDJointSolver::ApplyDualConeSwingVelocityConstraint(
 		const FReal Dt,
 		const FPBDJointSolverSettings& SolverSettings,
 		const FPBDJointSettings& JointSettings,
@@ -925,14 +936,20 @@ namespace Chaos
 		{
 			FVec3 SwingAxis;
 			FReal SwingAngle;
-			FPBDJointUtilities::GetDualConeSwingAxisAngle(Rs[0], Rs[1], SwingConstraintIndex, SwingAxis, SwingAngle);
+			FPBDJointUtilities::GetDualConeSwingAxisAngle(ConnectorRs[0], ConnectorRs[1], SwingConstraintIndex, SwingAxis, SwingAngle);
 
 			const FReal SwingStiffness = SolverStiffness * FPBDJointUtilities::GetSwingStiffness(SolverSettings, JointSettings);
-			ApplyAngularVelocityConstraint(SwingStiffness, SwingAxis);
+			FReal TargetVel = 0.0f;
+			if (JointSettings.SwingRestitution != 0.0f)
+			{
+				const FReal InitVel = InitConstraintAxisAngularVelocities[(int32)SwingConstraintIndex];
+				TargetVel = InitVel > Chaos_Joint_AngularVelocityThresholdToApplyRestitution ? -JointSettings.SwingRestitution * InitVel : 0.0f;
+			}
+			ApplyAngularVelocityConstraint(SwingStiffness, SwingAxis, TargetVel);
 		}
 	}
 
-	void FJointSolverGaussSeidel::ApplySwingVelocityConstraint(
+	void FPBDJointSolver::ApplySwingVelocityConstraint(
 		const FReal Dt,
 		const FPBDJointSolverSettings& SolverSettings,
 		const FPBDJointSettings& JointSettings,
@@ -943,31 +960,38 @@ namespace Chaos
 		{
 			FVec3 SwingAxis;
 			FReal SwingAngle;
-			FPBDJointUtilities::GetSwingAxisAngle(Rs[0], Rs[1], SolverSettings.SwingTwistAngleTolerance, SwingConstraintIndex, SwingAxis, SwingAngle);
+			FPBDJointUtilities::GetSwingAxisAngle(ConnectorRs[0], ConnectorRs[1], SolverSettings.SwingTwistAngleTolerance, SwingConstraintIndex, SwingAxis, SwingAngle);
 
 			const FReal SwingStiffness = SolverStiffness * FPBDJointUtilities::GetSwingStiffness(SolverSettings, JointSettings);
-			ApplyAngularVelocityConstraint(SwingStiffness, SwingAxis);
+			FReal TargetVel = 0.0f;
+			if (JointSettings.SwingRestitution != 0.0f)
+			{
+				const FReal InitVel = InitConstraintAxisAngularVelocities[(int32)SwingConstraintIndex];
+				TargetVel = InitVel > Chaos_Joint_AngularVelocityThresholdToApplyRestitution ? -JointSettings.SwingRestitution * InitVel : 0.0f;
+			}
+			ApplyAngularVelocityConstraint(SwingStiffness, SwingAxis, TargetVel);
 		}
 	}
 
-	void FJointSolverGaussSeidel::ApplyAngularVelocityConstraint(
+	void FPBDJointSolver::ApplyAngularVelocityConstraint(
 		const FReal Stiffness,
-		const FVec3& Axis)
+		const FVec3& Axis,
+		const FReal TargetVel)
 	{
-		const FVec3 CW = Ws[1] - Ws[0];
-		const FVec3 IA0 = Utilities::Multiply(InvIs[0], Axis);
-		const FVec3 IA1 = Utilities::Multiply(InvIs[1], Axis);
+		const FVec3 CW = W(1) - W(0);
+		const FVec3 IA0 = Utilities::Multiply(InvI(0), Axis);
+		const FVec3 IA1 = Utilities::Multiply(InvI(1), Axis);
 		const FReal II0 = FVec3::DotProduct(Axis, IA0);
 		const FReal II1 = FVec3::DotProduct(Axis, IA1);
 
-		const FVec3 AngularImpulse = Stiffness * FVec3::DotProduct(CW, Axis) / (II0 + II1) * Axis;
-		const FVec3 DW0 = Utilities::Multiply(InvIs[0], AngularImpulse);
-		const FVec3 DW1 = Utilities::Multiply(InvIs[1], -AngularImpulse);
+		const FVec3 AngularImpulse = Stiffness * (FVec3::DotProduct(CW, Axis) - TargetVel) / (II0 + II1) * Axis;
+		const FVec3 DW0 = Utilities::Multiply(InvI(0), AngularImpulse);
+		const FVec3 DW1 = Utilities::Multiply(InvI(1), -AngularImpulse);
 
 		ApplyAngularVelocityDelta(DW0, DW1);
 	}
 
-	void FJointSolverGaussSeidel::ApplyLockedRotationVelocityConstraints(
+	void FPBDJointSolver::ApplyLockedRotationVelocityConstraints(
 		const FReal Dt,
 		const FPBDJointSolverSettings& SolverSettings,
 		const FPBDJointSettings& JointSettings,
@@ -977,8 +1001,8 @@ namespace Chaos
 		if(!NetAngularImpulse.IsNearlyZero())
 		{
 			FVec3 Axis0, Axis1, Axis2;
-			FPBDJointUtilities::GetLockedRotationAxes(Rs[0], Rs[1], Axis0, Axis1, Axis2);
-			const FRotation3 R01 = Rs[0].Inverse() * Rs[1];
+			FPBDJointUtilities::GetLockedRotationAxes(ConnectorRs[0], ConnectorRs[1], Axis0, Axis1, Axis2);
+			const FRotation3 R01 = ConnectorRs[0].Inverse() * ConnectorRs[1];
 
 			if (bApplyTwist && AngularHardLambda[(int32)EJointAngularConstraintIndex::Twist] > SMALL_NUMBER)
 			{
@@ -1006,48 +1030,47 @@ namespace Chaos
 	//
 
 
-	void FJointSolverGaussSeidel::ApplyPositionDelta(
+	void FPBDJointSolver::ApplyPositionDelta(
 		const int32 BodyIndex,
 		const FVec3& DP)
 	{
 		//UE_LOG(LogChaosJoint, VeryVerbose, TEXT("      Apply DP%d %f %f %f"), BodyIndex, DP.X, DP.Y, DP.Z);
 
-		Ps[BodyIndex] += DP;
+		Body(BodyIndex).ApplyPositionDelta(DP);
 
-		Xs[BodyIndex] += DP;
+		ConnectorXs[BodyIndex] += DP;
 	}
 
 
-	void FJointSolverGaussSeidel::ApplyPositionDelta(
+	void FPBDJointSolver::ApplyPositionDelta(
 		const FVec3& DP0,
 		const FVec3& DP1)
 	{
 		//UE_LOG(LogChaosJoint, VeryVerbose, TEXT("      Apply DP%d %f %f %f"), 0, DP0.X, DP0.Y, DP0.Z);
 		//UE_LOG(LogChaosJoint, VeryVerbose, TEXT("      Apply DP%d %f %f %f"), 1, DP1.X, DP1.Y, DP1.Z);
 
-		Ps[0] += DP0;
-		Ps[1] += DP1;
+		Body(0).ApplyPositionDelta(DP0);
+		Body(1).ApplyPositionDelta(DP1);
 
-		Xs[0] += DP0;
-		Xs[1] += DP1;
+		ConnectorXs[0] += DP0;
+		ConnectorXs[1] += DP1;
 	}
 
 
-	void FJointSolverGaussSeidel::ApplyRotationDelta(
+	void FPBDJointSolver::ApplyRotationDelta(
 		const int32 BodyIndex,
 		const FVec3& DR)
 	{
 		//UE_LOG(LogChaosJoint, VeryVerbose, TEXT("      Apply DR%d %f %f %f"), BodyIndex, DR.X, DR.Y, DR.Z);
 
-		const FRotation3 DQ = (FRotation3::FromElements(DR, 0) * Qs[BodyIndex]) * (FReal)0.5;
-		Qs[BodyIndex] = (Qs[BodyIndex] + DQ).GetNormalized();
-		Qs[1].EnforceShortestArcWith(Qs[0]);
+		Body(BodyIndex).ApplyRotationDelta(DR);
+		Body(1).EnforceShortestRotationTo(Q(0));
 
 		UpdateDerivedState(BodyIndex);
 	}
 
 
-	void FJointSolverGaussSeidel::ApplyRotationDelta(
+	void FPBDJointSolver::ApplyRotationDelta(
 		const FVec3& DR0,
 		const FVec3& DR1)
 	{
@@ -1057,28 +1080,26 @@ namespace Chaos
 		if (bRealTypeCompatibleWithISPC && bChaos_Joint_ISPC_Enabled)
 		{
 #if INTEL_ISPC
-			ispc::ApplyRotationDelta2((ispc::FJointSolverGaussSeidel*)this, (ispc::FVector&)DR0, (ispc::FVector&)DR1);
+			ispc::ApplyRotationDelta2((ispc::FPBDJointSolver*)this, (ispc::FVector&)DR0, (ispc::FVector&)DR1);
 #endif
 		}
 		else
 		{
-			if (InvMs[0] > 0.0f)
+			if (InvM(0) > SMALL_NUMBER)
 			{
-				const FRotation3 DQ0 = (FRotation3::FromElements(DR0, 0) * Qs[0]) * (FReal)0.5;
-				Qs[0] = (Qs[0] + DQ0).GetNormalized();
+				Body(0).ApplyRotationDelta(DR0);
 			}
-			if (InvMs[1] > 0.0f)
+			if (InvM(1) > SMALL_NUMBER)
 			{
-				const FRotation3 DQ1 = (FRotation3::FromElements(DR1, 0) * Qs[1]) * (FReal)0.5;
-				Qs[1] = (Qs[1] + DQ1).GetNormalized();
+				Body(1).ApplyRotationDelta(DR1);
 			}
-			Qs[1].EnforceShortestArcWith(Qs[0]);
+			Body(1).EnforceShortestRotationTo(Q(0));
 
 			UpdateDerivedState();
 		}
 	}
 
-	void FJointSolverGaussSeidel::ApplyDelta(
+	void FPBDJointSolver::ApplyDelta(
 		const int32 BodyIndex,
 		const FVec3& DP,
 		const FVec3& DR)
@@ -1086,70 +1107,68 @@ namespace Chaos
 		//UE_LOG(LogChaosJoint, VeryVerbose, TEXT("      Apply DP%d %f %f %f"), BodyIndex, DP.X, DP.Y, DP.Z);
 		//UE_LOG(LogChaosJoint, VeryVerbose, TEXT("      Apply DR%d %f %f %f"), BodyIndex, DR.X, DR.Y, DR.Z);
 
-		Ps[BodyIndex] += DP;
-		const FRotation3 DQ = (FRotation3::FromElements(DR, 0) * Qs[BodyIndex]) * (FReal)0.5;
-		Qs[BodyIndex] = (Qs[BodyIndex] + DQ).GetNormalized();
-		Qs[1].EnforceShortestArcWith(Qs[0]);
+		Body(BodyIndex).ApplyTransformDelta(DP, DR);
+		Body(1).EnforceShortestRotationTo(Q(0));
 
 		UpdateDerivedState(BodyIndex);
 	}
 
 
-	void FJointSolverGaussSeidel::ApplyVelocityDelta(
+	void FPBDJointSolver::ApplyVelocityDelta(
 		const int32 BodyIndex,
 		const FVec3& DV,
 		const FVec3& DW)
 	{
-		Vs[BodyIndex] = Vs[BodyIndex] + DV;
-		Ws[BodyIndex] = Ws[BodyIndex] + DW;
+		Body(BodyIndex).ApplyVelocityDelta(DV, DW);
 	}
 
 
-	void FJointSolverGaussSeidel::ApplyVelocityDelta(
+	void FPBDJointSolver::ApplyVelocityDelta(
 		const FVec3& DV0,
 		const FVec3& DW0,
 		const FVec3& DV1,
 		const FVec3& DW1)
 	{
-		Vs[0] += DV0;
-		Vs[1] += DV1;
-		Ws[0] += DW0;
-		Ws[1] += DW1;
+		Body(0).ApplyVelocityDelta(DV0, DW0);
+		Body(1).ApplyVelocityDelta(DV1, DW1);
 	}
 
-	void FJointSolverGaussSeidel::ApplyAngularVelocityDelta(
+	void FPBDJointSolver::ApplyAngularVelocityDelta(
 		const FVec3& DW0,
 		const FVec3& DW1)
 	{
-		Ws[0] += DW0;
-		Ws[1] += DW1;
+		Body(0).ApplyAngularVelocityDelta(DW0);
+		Body(1).ApplyAngularVelocityDelta(DW1);
 	}
 
-	void FJointSolverGaussSeidel::ApplyPositionConstraint(
+	void FPBDJointSolver::ApplyPositionConstraint(
 		const FReal JointStiffness,
 		const FVec3& Axis,
 		const FReal Delta,
+		const FVec3& Connector0Correction,
 		const int32 LinearHardLambdaIndex)
 	{
 		const FReal Stiffness = SolverStiffness * JointStiffness;
-
-		const FVec3 AngularAxis0 = FVec3::CrossProduct(Xs[0] - Ps[0], Axis);
-		const FVec3 AngularAxis1 = FVec3::CrossProduct(Xs[1] - Ps[1], Axis);
-		const FVec3 IA0 = Utilities::Multiply(InvIs[0], AngularAxis0);
-		const FVec3 IA1 = Utilities::Multiply(InvIs[1], AngularAxis1);
+		// Project ConnectorXs[1] to the feasible space and apply impulse at the projected location. ConnectorXs[0] + Connector0Correction is the Projected location. The results are more stable this way.
+		const FVec3 Arm0 = ConnectorXs[0] + Connector0Correction - P(0);
+		const FVec3 Arm1 = ConnectorXs[1] - P(1);
+		const FVec3 AngularAxis0 = FVec3::CrossProduct(Arm0, Axis);
+		const FVec3 AngularAxis1 = FVec3::CrossProduct(Arm1, Axis);
+		const FVec3 IA0 = Utilities::Multiply(InvI(0), AngularAxis0);
+		const FVec3 IA1 = Utilities::Multiply(InvI(1), AngularAxis1);
 
 		// Joint-space inverse mass
 		const FReal II0 = FVec3::DotProduct(AngularAxis0, IA0);
 		const FReal II1 = FVec3::DotProduct(AngularAxis1, IA1);
-		const FReal IM = InvMs[0] + II0 + InvMs[1] + II1;
+		const FReal IM = InvM(0) + II0 + InvM(1) + II1;
 
 		// Apply constraint correction
 		const FReal DLambda = Stiffness * Delta / IM;
 		const FVec3 DX = Axis * DLambda;
-		const FVec3 DP0 = InvMs[0] * DX;
-		const FVec3 DP1 = -InvMs[1] * DX;
-		const FVec3 DR0 = Utilities::Multiply(InvIs[0], FVec3::CrossProduct(Xs[0] - Ps[0], DX));
-		const FVec3 DR1 = Utilities::Multiply(InvIs[1], FVec3::CrossProduct(Xs[1] - Ps[1], -DX));
+		const FVec3 DP0 = InvM(0) * DX;
+		const FVec3 DP1 = -InvM(1) * DX;
+		const FVec3 DR0 = Utilities::Multiply(InvI(0), FVec3::CrossProduct(Arm0, DX));
+		const FVec3 DR1 = Utilities::Multiply(InvI(1), FVec3::CrossProduct(Arm1, -DX));
 
 		ApplyPositionDelta(DP0, DP1);
 		ApplyRotationDelta(DR0, DR1);
@@ -1165,7 +1184,7 @@ namespace Chaos
 	}
 
 
-	void FJointSolverGaussSeidel::ApplyPositionConstraintSoft(
+	void FPBDJointSolver::ApplyPositionConstraintSoft(
 		const FReal Dt,
 		const FReal JointStiffness,
 		const FReal JointDamping,
@@ -1179,39 +1198,39 @@ namespace Chaos
 		{
 #if INTEL_ISPC
 			float ReturnedLambda = Lambda;
-			ispc::ApplyPositionConstraintSoft((ispc::FJointSolverGaussSeidel*)this, Dt, JointStiffness, JointDamping, bAccelerationMode, (ispc::FVector&)Axis, Delta, TargetVel, ReturnedLambda);
+			ispc::ApplyPositionConstraintSoft((ispc::FPBDJointSolver*)this, Dt, JointStiffness, JointDamping, bAccelerationMode, (ispc::FVector&)Axis, Delta, TargetVel, ReturnedLambda);
 			Lambda = ReturnedLambda;
 #endif
 		}
 		else
 		{
 			// Joint-space inverse mass
-			const FVec3 AngularAxis0 = FVec3::CrossProduct(Xs[0] - Ps[0], Axis);
-			const FVec3 AngularAxis1 = FVec3::CrossProduct(Xs[1] - Ps[1], Axis);
-			const FVec3 IA0 = Utilities::Multiply(InvIs[0], AngularAxis0);
-			const FVec3 IA1 = Utilities::Multiply(InvIs[1], AngularAxis1);
+			const FVec3 AngularAxis0 = FVec3::CrossProduct(ConnectorXs[0] - P(0), Axis);
+			const FVec3 AngularAxis1 = FVec3::CrossProduct(ConnectorXs[1] - P(1), Axis);
+			const FVec3 IA0 = Utilities::Multiply(InvI(0), AngularAxis0);
+			const FVec3 IA1 = Utilities::Multiply(InvI(1), AngularAxis1);
 			const FReal II0 = FVec3::DotProduct(AngularAxis0, IA0);
 			const FReal II1 = FVec3::DotProduct(AngularAxis1, IA1);
-			const FReal II = (InvMs[0] + II0 + InvMs[1] + II1);
+			const FReal II = (InvM(0) + II0 + InvM(1) + II1);
 	
 			FReal VelDt = 0;
 			if (JointDamping > KINDA_SMALL_NUMBER)
 			{
-				const FVec3 V0Dt = FVec3::CalculateVelocity(InitXs[0], Xs[0], 1.0f);
-				const FVec3 V1Dt = FVec3::CalculateVelocity(InitXs[1], Xs[1], 1.0f);
+				const FVec3 V0Dt = FVec3::CalculateVelocity(InitConnectorXs[0], ConnectorXs[0], 1.0f);
+				const FVec3 V1Dt = FVec3::CalculateVelocity(InitConnectorXs[1], ConnectorXs[1], 1.0f);
 				VelDt = TargetVel * Dt + FVec3::DotProduct(V0Dt - V1Dt, Axis);
 			}
 	
-			const FReal SpringMassScale = (bAccelerationMode) ? 1.0f / (InvMs[0] + InvMs[1]) : 1.0f;
+			const FReal SpringMassScale = (bAccelerationMode) ? 1.0f / (InvM(0) + InvM(1)) : 1.0f;
 			const FReal S = SpringMassScale * JointStiffness * Dt * Dt;
 			const FReal D = SpringMassScale * JointDamping * Dt;
 			const FReal Multiplier = (FReal)1 / ((S + D) * II + (FReal)1);
 			const FReal DLambda = SolverStiffness * Multiplier * (S * Delta - D * VelDt - Lambda);
 	
-			const FVec3 DP0 = (InvMs[0] * DLambda) * Axis;
-			const FVec3 DP1 = (-InvMs[1] * DLambda) * Axis;
-			const FVec3 DR0 = DLambda * Utilities::Multiply(InvIs[0], AngularAxis0);
-			const FVec3 DR1 = -DLambda * Utilities::Multiply(InvIs[1], AngularAxis1);
+			const FVec3 DP0 = (InvM(0) * DLambda) * Axis;
+			const FVec3 DP1 = (-InvM(1) * DLambda) * Axis;
+			const FVec3 DR0 = DLambda * Utilities::Multiply(InvI(0), AngularAxis0);
+			const FVec3 DR1 = -DLambda * Utilities::Multiply(InvI(1), AngularAxis1);
 	
 			ApplyPositionDelta(DP0, DP1);
 			ApplyRotationDelta(DR0, DR1);
@@ -1224,7 +1243,7 @@ namespace Chaos
 	}
 	
 
-	void FJointSolverGaussSeidel::ApplyRotationConstraintKD(
+	void FPBDJointSolver::ApplyRotationConstraintKD(
 		const int32 KIndex,
 		const int32 DIndex,
 		const FReal JointStiffness,
@@ -1234,7 +1253,7 @@ namespace Chaos
 	{
 		const FReal Stiffness = SolverStiffness * JointStiffness;
 
-		const FVec3 IA1 = Utilities::Multiply(InvIs[DIndex], Axis);
+		const FVec3 IA1 = Utilities::Multiply(InvI(DIndex), Axis);
 		const FReal II1 = FVec3::DotProduct(Axis, IA1);
 		const FReal DR = Stiffness * (Angle / II1);
 		const FVec3 DR1 = IA1 * -DR;
@@ -1248,7 +1267,7 @@ namespace Chaos
 	}
 
 
-	void FJointSolverGaussSeidel::ApplyRotationConstraintDD(
+	void FPBDJointSolver::ApplyRotationConstraintDD(
 		const FReal JointStiffness,
 		const FVec3& Axis,
 		const FReal Angle,
@@ -1257,8 +1276,8 @@ namespace Chaos
 		const FReal Stiffness = SolverStiffness * JointStiffness;
 
 		// Joint-space inverse mass
-		const FVec3 IA0 = Utilities::Multiply(InvIs[0], Axis);
-		const FVec3 IA1 = Utilities::Multiply(InvIs[1], Axis);
+		const FVec3 IA0 = Utilities::Multiply(InvI(0), Axis);
+		const FVec3 IA1 = Utilities::Multiply(InvI(1), Axis);
 		const FReal II0 = FVec3::DotProduct(Axis, IA0);
 		const FReal II1 = FVec3::DotProduct(Axis, IA1);
 
@@ -1276,17 +1295,17 @@ namespace Chaos
 	}
 
 
-	void FJointSolverGaussSeidel::ApplyRotationConstraint(
+	void FPBDJointSolver::ApplyRotationConstraint(
 		const FReal JointStiffness,
 		const FVec3& Axis,
 		const FReal Angle,
 		const int32 AngularHardLambdaIndex)
 	{
-		if (InvMs[0] == 0)
+		if (InvM(0) < SMALL_NUMBER)
 		{
 			ApplyRotationConstraintKD(0, 1, JointStiffness, Axis, Angle, AngularHardLambdaIndex);
 		}
-		else if (InvMs[1] == 0)
+		else if (InvM(1) < SMALL_NUMBER)
 		{
 			ApplyRotationConstraintKD(1, 0, JointStiffness, Axis, -Angle, AngularHardLambdaIndex);
 		}
@@ -1300,7 +1319,7 @@ namespace Chaos
 
 
 	// See "XPBD: Position-Based Simulation of Compliant Constrained Dynamics"
-	void FJointSolverGaussSeidel::ApplyRotationConstraintSoftKD(
+	void FPBDJointSolver::ApplyRotationConstraintSoftKD(
 		const int32 KIndex,
 		const int32 DIndex,
 		const FReal Dt,
@@ -1312,20 +1331,20 @@ namespace Chaos
 		const FReal AngVelTarget,
 		FReal& Lambda)
 	{
-		check(InvMs[DIndex] > 0);
+		check(InvM(DIndex) > SMALL_NUMBER);
 
 		if (bRealTypeCompatibleWithISPC && bChaos_Joint_ISPC_Enabled)
 		{
 #if INTEL_ISPC
 			float ReturnedLambda = Lambda;
-			ispc::ApplyRotationConstraintSoftKD((ispc::FJointSolverGaussSeidel*)this, KIndex, DIndex, Dt, JointStiffness, JointDamping, bAccelerationMode, (ispc::FVector&) Axis, Angle, AngVelTarget, ReturnedLambda);
+			ispc::ApplyRotationConstraintSoftKD((ispc::FPBDJointSolver*)this, KIndex, DIndex, Dt, JointStiffness, JointDamping, bAccelerationMode, (ispc::FVector&) Axis, Angle, AngVelTarget, ReturnedLambda);
 			Lambda = ReturnedLambda;
 #endif
 		}
 		else
 		{
 			// World-space inverse mass
-			const FVec3 IA1 = Utilities::Multiply(InvIs[DIndex], Axis);
+			const FVec3 IA1 = Utilities::Multiply(InvI(DIndex), Axis);
 
 			// Joint-space inverse mass
 			FReal II1 = FVec3::DotProduct(Axis, IA1);
@@ -1335,8 +1354,8 @@ namespace Chaos
 			FReal AngVelDt = 0;
 			if (JointDamping > KINDA_SMALL_NUMBER)
 			{
-				const FVec3 W0Dt = FRotation3::CalculateAngularVelocity(InitRs[KIndex], Rs[KIndex], 1.0f);
-				const FVec3 W1Dt = FRotation3::CalculateAngularVelocity(InitRs[DIndex], Rs[DIndex], 1.0f);
+				const FVec3 W0Dt = FRotation3::CalculateAngularVelocity(InitConnectorRs[KIndex], ConnectorRs[KIndex], 1.0f);
+				const FVec3 W1Dt = FRotation3::CalculateAngularVelocity(InitConnectorRs[DIndex], ConnectorRs[DIndex], 1.0f);
 				AngVelDt = AngVelTarget * Dt + FVec3::DotProduct(Axis, W0Dt - W1Dt);
 			}
 
@@ -1357,7 +1376,7 @@ namespace Chaos
 	}
 
 	// See "XPBD: Position-Based Simulation of Compliant Constrained Dynamics"
-	void FJointSolverGaussSeidel::ApplyRotationConstraintSoftDD(
+	void FPBDJointSolver::ApplyRotationConstraintSoftDD(
 		const FReal Dt,
 		const FReal JointStiffness,
 		const FReal JointDamping,
@@ -1367,22 +1386,22 @@ namespace Chaos
 		const FReal AngVelTarget,
 		FReal& Lambda)
 	{
-		check(InvMs[0] > 0);
-		check(InvMs[1] > 0);
+		check(InvM(0) > SMALL_NUMBER);
+		check(InvM(1) > SMALL_NUMBER);
 
 		if (bRealTypeCompatibleWithISPC && bChaos_Joint_ISPC_Enabled)
 		{
 #if INTEL_ISPC
 			float ReturnedLambda = Lambda;
-			ispc::ApplyRotationConstraintSoftDD((ispc::FJointSolverGaussSeidel*)this, Dt, JointStiffness, JointDamping, bAccelerationMode, (ispc::FVector&) Axis, Angle, AngVelTarget, ReturnedLambda);
+			ispc::ApplyRotationConstraintSoftDD((ispc::FPBDJointSolver*)this, Dt, JointStiffness, JointDamping, bAccelerationMode, (ispc::FVector&) Axis, Angle, AngVelTarget, ReturnedLambda);
 			Lambda = ReturnedLambda;
 #endif
 		}
 		else
 		{
 			// World-space inverse mass
-			const FVec3 IA0 = Utilities::Multiply(InvIs[0], Axis);
-			const FVec3 IA1 = Utilities::Multiply(InvIs[1], Axis);
+			const FVec3 IA0 = Utilities::Multiply(InvI(0), Axis);
+			const FVec3 IA1 = Utilities::Multiply(InvI(1), Axis);
 
 			// Joint-space inverse mass
 			FReal II0 = FVec3::DotProduct(Axis, IA0);
@@ -1393,8 +1412,8 @@ namespace Chaos
 			FReal AngVelDt = 0;
 			if (JointDamping > KINDA_SMALL_NUMBER)
 			{
-				const FVec3 W0Dt = FRotation3::CalculateAngularVelocity(InitRs[0], Rs[0], 1.0f);
-				const FVec3 W1Dt = FRotation3::CalculateAngularVelocity(InitRs[1], Rs[1], 1.0f);
+				const FVec3 W0Dt = FRotation3::CalculateAngularVelocity(InitConnectorRs[0], ConnectorRs[0], 1.0f);
+				const FVec3 W1Dt = FRotation3::CalculateAngularVelocity(InitConnectorRs[1], ConnectorRs[1], 1.0f);
 				AngVelDt = AngVelTarget * Dt + FVec3::DotProduct(Axis, W0Dt - W1Dt);
 			}
 
@@ -1416,7 +1435,7 @@ namespace Chaos
 		}
 	}
 
-	void FJointSolverGaussSeidel::ApplyRotationConstraintSoft(
+	void FPBDJointSolver::ApplyRotationConstraintSoft(
 		const FReal Dt,
 		const FReal JointStiffness,
 		const FReal JointDamping,
@@ -1426,11 +1445,11 @@ namespace Chaos
 		const FReal AngVelTarget,
 		FReal& Lambda)
 	{
-		if (InvMs[0] == 0)
+		if (InvM(0) < SMALL_NUMBER)
 		{
 			ApplyRotationConstraintSoftKD(0, 1, Dt, JointStiffness, JointDamping, bAccelerationMode, Axis, Angle, AngVelTarget, Lambda);
 		}
-		else if (InvMs[1] == 0)
+		else if (InvM(1) < SMALL_NUMBER)
 		{
 			ApplyRotationConstraintSoftKD(1, 0, Dt, JointStiffness, JointDamping, bAccelerationMode, Axis, -Angle, -AngVelTarget, Lambda);
 		}
@@ -1444,7 +1463,7 @@ namespace Chaos
 
 	// Used for non-zero restitution. We pad constraints by an amount such that the velocity
 	// calculated after solving constraint positions will as required for the restitution.
-	void FJointSolverGaussSeidel::CalculateLinearConstraintPadding(
+	void FPBDJointSolver::CalculateLinearConstraintPadding(
 		const FReal Dt,
 		const FPBDJointSolverSettings& SolverSettings,
 		const FPBDJointSettings& JointSettings,
@@ -1460,8 +1479,8 @@ namespace Chaos
 			SetLinearConstraintPadding(AxisIndex, 0.0f);
 
 			// Calculate the velocity we want to match
-			const FVec3 V0Dt = FVec3::CalculateVelocity(InitXs[0], Xs[0], 1.0f);
-			const FVec3 V1Dt = FVec3::CalculateVelocity(InitXs[1], Xs[1], 1.0f);
+			const FVec3 V0Dt = FVec3::CalculateVelocity(InitConnectorXs[0], ConnectorXs[0], 1.0f);
+			const FVec3 V1Dt = FVec3::CalculateVelocity(InitConnectorXs[1], ConnectorXs[1], 1.0f);
 			const FReal AxisVDt = FVec3::DotProduct(V1Dt - V0Dt, Axis);
 
 			// Calculate the padding to apply to the constraint that will result in the
@@ -1477,7 +1496,7 @@ namespace Chaos
 
 	// Used for non-zero restitution. We pad constraints by an amount such that the velocity
 	// calculated after solving constraint positions will as required for the restitution.
-	void FJointSolverGaussSeidel::CalculateAngularConstraintPadding(
+	void FPBDJointSolver::CalculateAngularConstraintPadding(
 		const FReal Dt,
 		const FPBDJointSolverSettings& SolverSettings,
 		const FPBDJointSettings& JointSettings,
@@ -1493,8 +1512,8 @@ namespace Chaos
 			SetAngularConstraintPadding(ConstraintIndex, 0.0f);
 
 			// Calculate the velocity we want to match
-			const FVec3 W0Dt = FRotation3::CalculateAngularVelocity(InitRs[0], Rs[0], 1.0f);
-			const FVec3 W1Dt = FRotation3::CalculateAngularVelocity(InitRs[1], Rs[1], 1.0f);
+			const FVec3 W0Dt = FRotation3::CalculateAngularVelocity(InitConnectorRs[0], ConnectorRs[0], 1.0f);
+			const FVec3 W1Dt = FRotation3::CalculateAngularVelocity(InitConnectorRs[1], ConnectorRs[1], 1.0f);
 			const FReal AxisWDt = FVec3::DotProduct(W1Dt - W0Dt, Axis);
 
 			// Calculate the padding to apply to the constraint that will result in the
@@ -1514,7 +1533,7 @@ namespace Chaos
 	//
 	//
 
-	void FJointSolverGaussSeidel::ApplyLockedRotationConstraints(
+	void FPBDJointSolver::ApplyLockedRotationConstraints(
 		const FReal Dt,
 		const FPBDJointSolverSettings& SolverSettings,
 		const FPBDJointSettings& JointSettings,
@@ -1522,9 +1541,9 @@ namespace Chaos
 		const bool bApplySwing)
 	{
 		FVec3 Axis0, Axis1, Axis2;
-		FPBDJointUtilities::GetLockedRotationAxes(Rs[0], Rs[1], Axis0, Axis1, Axis2);
+		FPBDJointUtilities::GetLockedRotationAxes(ConnectorRs[0], ConnectorRs[1], Axis0, Axis1, Axis2);
 
-		const FRotation3 R01 = Rs[0].Inverse() * Rs[1];
+		const FRotation3 R01 = ConnectorRs[0].Inverse() * ConnectorRs[1];
 
 		if (bApplyTwist)
 		{
@@ -1540,7 +1559,7 @@ namespace Chaos
 		}
 	}
 
-	void FJointSolverGaussSeidel::ApplyTwistConstraint(
+	void FPBDJointSolver::ApplyTwistConstraint(
 		const FReal Dt,
 		const FPBDJointSolverSettings& SolverSettings,
 		const FPBDJointSettings& JointSettings,
@@ -1548,7 +1567,7 @@ namespace Chaos
 	{
 		FVec3 TwistAxis;
 		FReal TwistAngle;
-		FPBDJointUtilities::GetTwistAxisAngle(Rs[0], Rs[1], TwistAxis, TwistAngle);
+		FPBDJointUtilities::GetTwistAxisAngle(ConnectorRs[0], ConnectorRs[1], TwistAxis, TwistAngle);
 
 		// Calculate the twist correction to apply to each body
 		const FReal LimitPadding = GetAngularConstraintPadding(EJointAngularConstraintIndex::Twist);
@@ -1590,7 +1609,7 @@ namespace Chaos
 		}
 	}
 
-	void FJointSolverGaussSeidel::ApplyConeConstraint(
+	void FPBDJointSolver::ApplyConeConstraint(
 		const FReal Dt,
 		const FPBDJointSolverSettings& SolverSettings,
 		const FPBDJointSettings& JointSettings,
@@ -1602,11 +1621,11 @@ namespace Chaos
 		const FReal LimitPadding = GetAngularConstraintPadding(EJointAngularConstraintIndex::Swing1);
 		const FReal Swing1Limit = FMath::Max(JointSettings.AngularLimits[(int32)EJointAngularConstraintIndex::Swing1] - LimitPadding, (FReal)0.);
 		const FReal Swing2Limit = FMath::Max(JointSettings.AngularLimits[(int32)EJointAngularConstraintIndex::Swing2] - LimitPadding, (FReal)0.);
-		FPBDJointUtilities::GetEllipticalConeAxisErrorLocal(Rs[0], Rs[1], Swing2Limit, Swing1Limit, SwingAxisLocal, DSwingAngle);
+		FPBDJointUtilities::GetEllipticalConeAxisErrorLocal(ConnectorRs[0], ConnectorRs[1], Swing2Limit, Swing1Limit, SwingAxisLocal, DSwingAngle);
 
 		UE_LOG(LogChaosJoint, VeryVerbose, TEXT("    Cone Error %f [Limits %f %f]"), FMath::RadiansToDegrees(DSwingAngle), FMath::RadiansToDegrees(Swing2Limit), FMath::RadiansToDegrees(Swing1Limit));
 
-		const FVec3 SwingAxis = Rs[0] * SwingAxisLocal;
+		const FVec3 SwingAxis = ConnectorRs[0] * SwingAxisLocal;
 
 		// Apply swing correction to each body
 		if (DSwingAngle > AngleTolerance)
@@ -1633,7 +1652,7 @@ namespace Chaos
 	}
 
 
-	void FJointSolverGaussSeidel::ApplySingleLockedSwingConstraint(
+	void FPBDJointSolver::ApplySingleLockedSwingConstraint(
 		const FReal Dt,
 		const FPBDJointSolverSettings& SolverSettings,
 		const FPBDJointSettings& JointSettings,
@@ -1644,7 +1663,7 @@ namespace Chaos
 		// Likewise, the SwingAngle is actually Sin(SwingAngle)
 		FVec3 SwingAxis;
 		FReal SwingAngle;
-		FPBDJointUtilities::GetLockedSwingAxisAngle(Rs[0], Rs[1], SwingConstraintIndex, SwingAxis, SwingAngle);
+		FPBDJointUtilities::GetLockedSwingAxisAngle(ConnectorRs[0], ConnectorRs[1], SwingConstraintIndex, SwingAxis, SwingAngle);
 
 		UE_LOG(LogChaosJoint, VeryVerbose, TEXT("    LockedSwing%d Angle %f [Tolerance %f]"), (SwingConstraintIndex == EJointAngularConstraintIndex::Swing1) ? 1 : 2, FMath::RadiansToDegrees(SwingAngle), FMath::RadiansToDegrees(AngleTolerance));
 
@@ -1668,7 +1687,7 @@ namespace Chaos
 	}
 
 
-	void FJointSolverGaussSeidel::ApplyDualConeSwingConstraint(
+	void FPBDJointSolver::ApplyDualConeSwingConstraint(
 		const FReal Dt,
 		const FPBDJointSolverSettings& SolverSettings,
 		const FPBDJointSettings& JointSettings,
@@ -1677,7 +1696,7 @@ namespace Chaos
 	{
 		FVec3 SwingAxis;
 		FReal SwingAngle;
-		FPBDJointUtilities::GetDualConeSwingAxisAngle(Rs[0], Rs[1], SwingConstraintIndex, SwingAxis, SwingAngle);
+		FPBDJointUtilities::GetDualConeSwingAxisAngle(ConnectorRs[0], ConnectorRs[1], SwingConstraintIndex, SwingAxis, SwingAngle);
 
 		// Calculate swing error we need to correct
 		const FReal LimitPadding = GetAngularConstraintPadding(SwingConstraintIndex);
@@ -1720,7 +1739,7 @@ namespace Chaos
 	}
 
 
-	void FJointSolverGaussSeidel::ApplySwingConstraint(
+	void FPBDJointSolver::ApplySwingConstraint(
 		const FReal Dt,
 		const FPBDJointSolverSettings& SolverSettings,
 		const FPBDJointSettings& JointSettings,
@@ -1729,7 +1748,7 @@ namespace Chaos
 	{
 		FVec3 SwingAxis;
 		FReal SwingAngle;
-		FPBDJointUtilities::GetSwingAxisAngle(Rs[0], Rs[1], SolverSettings.SwingTwistAngleTolerance, SwingConstraintIndex, SwingAxis, SwingAngle);
+		FPBDJointUtilities::GetSwingAxisAngle(ConnectorRs[0], ConnectorRs[1], SolverSettings.SwingTwistAngleTolerance, SwingConstraintIndex, SwingAxis, SwingAngle);
 
 		// Calculate swing error we need to correct
 		const FReal LimitPadding = GetAngularConstraintPadding(SwingConstraintIndex);
@@ -1772,7 +1791,7 @@ namespace Chaos
 	}
 
 
-	void FJointSolverGaussSeidel::ApplySwingTwistDrives(
+	void FPBDJointSolver::ApplySwingTwistDrives(
 		const FReal Dt,
 		const FPBDJointSolverSettings& SolverSettings,
 		const FPBDJointSettings& JointSettings,
@@ -1780,9 +1799,9 @@ namespace Chaos
 		const bool bSwing1DriveEnabled,
 		const bool bSwing2DriveEnabled)
 	{
-		FRotation3 R1Target = Rs[0] * JointSettings.AngularDrivePositionTarget;
-		R1Target.EnforceShortestArcWith(Rs[1]);
-		FRotation3 R1Error = R1Target.Inverse() * Rs[1];
+		FRotation3 R1Target = ConnectorRs[0] * JointSettings.AngularDrivePositionTarget;
+		R1Target.EnforceShortestArcWith(ConnectorRs[1]);
+		FRotation3 R1Error = R1Target.Inverse() * ConnectorRs[1];
 		FVec3 R1TwistAxisError = R1Error * FJointConstants::TwistAxis();
 
 		// Angle approximation Angle ~= Sin(Angle) for small angles, underestimates for large angles
@@ -1800,7 +1819,7 @@ namespace Chaos
 		if (bUseTwistDrive)
 		{
 			FReal ReturnedLambda = RotationDriveLambdas[(int32)EJointAngularConstraintIndex::Twist];
-			const FVec3 TwistAxis = Rs[1] * FJointConstants::TwistAxis();
+			const FVec3 TwistAxis = ConnectorRs[1] * FJointConstants::TwistAxis();
 			ApplyRotationConstraintSoft(Dt, AngularTwistDriveStiffness, AngularTwistDriveDamping, bAccelerationMode, TwistAxis, DTwistAngle, JointSettings.AngularDriveVelocityTarget[(int32)EJointAngularConstraintIndex::Twist], ReturnedLambda);
 			RotationDriveLambdas[(int32)EJointAngularConstraintIndex::Twist] = ReturnedLambda;
 		}
@@ -1809,7 +1828,7 @@ namespace Chaos
 		if (bUseSwing1Drive)
 		{
 			FReal ReturnedLambda = RotationDriveLambdas[(int32)EJointAngularConstraintIndex::Swing1];
-			const FVec3 Swing1Axis = Rs[1] * FJointConstants::Swing1Axis();
+			const FVec3 Swing1Axis = ConnectorRs[1] * FJointConstants::Swing1Axis();
 			ApplyRotationConstraintSoft(Dt, AngularSwingDriveStiffness, AngularSwingDriveDamping, bAccelerationMode, Swing1Axis, DSwing1Angle, JointSettings.AngularDriveVelocityTarget[(int32)EJointAngularConstraintIndex::Swing1], ReturnedLambda);
 			RotationDriveLambdas[(int32)EJointAngularConstraintIndex::Swing1] = ReturnedLambda;
 		}
@@ -1818,14 +1837,14 @@ namespace Chaos
 		if (bUseSwing2Drive)
 		{
 			FReal ReturnedLambda = RotationDriveLambdas[(int32)EJointAngularConstraintIndex::Swing2];
-			const FVec3 Swing2Axis = Rs[1] * FJointConstants::Swing2Axis();
+			const FVec3 Swing2Axis = ConnectorRs[1] * FJointConstants::Swing2Axis();
 			ApplyRotationConstraintSoft(Dt, AngularSwingDriveStiffness, AngularSwingDriveDamping, bAccelerationMode, Swing2Axis, DSwing2Angle, JointSettings.AngularDriveVelocityTarget[(int32)EJointAngularConstraintIndex::Swing2], ReturnedLambda);
 			RotationDriveLambdas[(int32)EJointAngularConstraintIndex::Swing2] = ReturnedLambda;
 		}
 	}
 
 
-	void FJointSolverGaussSeidel::ApplySLerpDrive(
+	void FPBDJointSolver::ApplySLerpDrive(
 		const FReal Dt,
 		const FPBDJointSolverSettings& SolverSettings,
 		const FPBDJointSettings& JointSettings)
@@ -1841,13 +1860,13 @@ namespace Chaos
 			FVec3 Axes[3] = { FVec3(1, 0, 0), FVec3(0, 1, 0), FVec3(0, 0, 1) };
 			if (AngularDriveStiffness > 0.0f)
 			{
-				FPBDJointUtilities::GetLockedRotationAxes(Rs[0], Rs[1], Axes[0], Axes[1], Axes[2]);
+				FPBDJointUtilities::GetLockedRotationAxes(ConnectorRs[0], ConnectorRs[1], Axes[0], Axes[1], Axes[2]);
 				Utilities::NormalizeSafe(Axes[0], KINDA_SMALL_NUMBER);
 				Utilities::NormalizeSafe(Axes[1], KINDA_SMALL_NUMBER);
 				Utilities::NormalizeSafe(Axes[2], KINDA_SMALL_NUMBER);
 			}
 
-			const FRotation3 R01 = Rs[0].Inverse() * Rs[1];
+			const FRotation3 R01 = ConnectorRs[0].Inverse() * ConnectorRs[1];
 			FRotation3 TargetAngPos = JointSettings.AngularDrivePositionTarget;
 			TargetAngPos.EnforceShortestArcWith(R01);
 			const FRotation3 R1Error = TargetAngPos.Inverse() * R01;
@@ -1858,7 +1877,7 @@ namespace Chaos
 				2.0f * FMath::Asin(R1Error.Z) 
 			};
 
-			const FVec3 TargetAngVel = Rs[0] * JointSettings.AngularDriveVelocityTarget;
+			const FVec3 TargetAngVel = ConnectorRs[0] * JointSettings.AngularDriveVelocityTarget;
 
 			for (int32 AxisIndex = 0; AxisIndex < 3; ++AxisIndex)
 			{
@@ -1870,8 +1889,8 @@ namespace Chaos
 		}
 		else
 		{
-			const FRotation3 TargetR1 = Rs[0] * JointSettings.AngularDrivePositionTarget;
-			const FRotation3 DR = TargetR1 * Rs[1].Inverse();
+			const FRotation3 TargetR1 = ConnectorRs[0] * JointSettings.AngularDrivePositionTarget;
+			const FRotation3 DR = TargetR1 * ConnectorRs[1].Inverse();
 
 			FVec3 SLerpAxis;
 			FReal SLerpAngle;
@@ -1887,7 +1906,7 @@ namespace Chaos
 				if (FMath::Abs(SLerpAngle) > AngleTolerance)
 				{
 					FReal ReturnedLambda = RotationDriveLambdas[(int32)EJointAngularConstraintIndex::Swing1];
-					FReal AngVelTarget = (JointSettings.AngularDriveDamping > 0.0f) ? FVec3::DotProduct(SLerpAxis, Rs[0] * JointSettings.AngularDriveVelocityTarget) : 0.0f;
+					FReal AngVelTarget = (JointSettings.AngularDriveDamping > FReal(0)) ? FVec3::DotProduct(SLerpAxis, ConnectorRs[0] * JointSettings.AngularDriveVelocityTarget) : 0.0f;
 					ApplyRotationConstraintSoft(Dt, AngularDriveStiffness, AngularDriveDamping, bAccelerationMode, SLerpAxis, -SLerpAngle, AngVelTarget, ReturnedLambda);
 					RotationDriveLambdas[(int32)EJointAngularConstraintIndex::Swing1] = ReturnedLambda;
 				}
@@ -1897,17 +1916,17 @@ namespace Chaos
 
 
 	// Kinematic-Dynamic bodies
-	void FJointSolverGaussSeidel::ApplyPointPositionConstraintKD(
+	void FPBDJointSolver::ApplyPointPositionConstraintKD(
 		const int32 KIndex,
 		const int32 DIndex,
 		const FReal Dt,
 		const FPBDJointSolverSettings& SolverSettings,
 		const FPBDJointSettings& JointSettings)
 	{
-		check(InvMs[DIndex] > 0);
+		check(InvM(DIndex) > SMALL_NUMBER);
 
 		FReal Stiffness = SolverStiffness * FPBDJointUtilities::GetLinearStiffness(SolverSettings, JointSettings);
-		const FVec3 CX = Xs[DIndex] - Xs[KIndex];
+		const FVec3 CX = ConnectorXs[DIndex] - ConnectorXs[KIndex];
 
 		UE_LOG(LogChaosJoint, VeryVerbose, TEXT("    PointKD Delta %f [Limit %f]"), CX.Size(), PositionTolerance);
 
@@ -1916,19 +1935,19 @@ namespace Chaos
 			if (bRealTypeCompatibleWithISPC && bChaos_Joint_ISPC_Enabled)
 			{
 #if INTEL_ISPC
-				ispc::ApplyPointPositionConstraintKD((ispc::FJointSolverGaussSeidel*)this, KIndex, DIndex, (ispc::FVector&)CX, Stiffness);
+				ispc::ApplyPointPositionConstraintKD((ispc::FPBDJointSolver*)this, KIndex, DIndex, (ispc::FVector&)CX, Stiffness);
 #endif
 			}
 			else
 			{
 				// Calculate constraint correction
-				FMatrix33 M = Utilities::ComputeJointFactorMatrix(Xs[DIndex] - Ps[DIndex], InvIs[DIndex], InvMs[DIndex]);
+				FMatrix33 M = Utilities::ComputeJointFactorMatrix(ConnectorXs[DIndex] - P(DIndex), InvI(DIndex), InvM(DIndex));
 				FMatrix33 MI = M.Inverse();
 				const FVec3 DX = Stiffness * Utilities::Multiply(MI, CX);
 
 				// Apply constraint correction
-				const FVec3 DP1 = -InvMs[DIndex] * DX;
-				const FVec3 DR1 = Utilities::Multiply(InvIs[DIndex], FVec3::CrossProduct(Xs[DIndex] - Ps[DIndex], -DX));
+				const FVec3 DP1 = -InvM(DIndex) * DX;
+				const FVec3 DR1 = Utilities::Multiply(InvI(DIndex), FVec3::CrossProduct(ConnectorXs[DIndex] - P(DIndex), -DX));
 
 				ApplyDelta(DIndex, DP1, DR1);
 
@@ -1941,16 +1960,16 @@ namespace Chaos
 
 
 	// Dynamic-Dynamic bodies
-	void FJointSolverGaussSeidel::ApplyPointPositionConstraintDD(
+	void FPBDJointSolver::ApplyPointPositionConstraintDD(
 		const FReal Dt,
 		const FPBDJointSolverSettings& SolverSettings,
 		const FPBDJointSettings& JointSettings)
 	{
-		check(InvMs[0] > 0);
-		check(InvMs[1] > 0);
+		check(InvM(0) > SMALL_NUMBER);
+		check(InvM(1) > SMALL_NUMBER);
 
 		FReal Stiffness = SolverStiffness * FPBDJointUtilities::GetLinearStiffness(SolverSettings, JointSettings);
-		const FVec3 CX = Xs[1] - Xs[0];
+		const FVec3 CX = ConnectorXs[1] - ConnectorXs[0];
 
 		UE_LOG(LogChaosJoint, VeryVerbose, TEXT("    PointDD Delta %f [Limit %f]"), CX.Size(), PositionTolerance);
 
@@ -1959,22 +1978,22 @@ namespace Chaos
 			if (bRealTypeCompatibleWithISPC && bChaos_Joint_ISPC_Enabled)
 			{
 #if INTEL_ISPC
-				ispc::ApplyPointPositionConstraintDD((ispc::FJointSolverGaussSeidel*)this, (ispc::FVector&)CX, Stiffness);
+				ispc::ApplyPointPositionConstraintDD((ispc::FPBDJointSolver*)this, (ispc::FVector&)CX, Stiffness);
 #endif
 			}
 			else
 			{
 				// Calculate constraint correction
-				FMatrix33 M0 = Utilities::ComputeJointFactorMatrix(Xs[0] - Ps[0], InvIs[0], InvMs[0]);
-				FMatrix33 M1 = Utilities::ComputeJointFactorMatrix(Xs[1] - Ps[1], InvIs[1], InvMs[1]);
+				FMatrix33 M0 = Utilities::ComputeJointFactorMatrix(ConnectorXs[0] - P(0), InvI(0), InvM(0));
+				FMatrix33 M1 = Utilities::ComputeJointFactorMatrix(ConnectorXs[1] - P(1), InvI(1), InvM(1));
 				FMatrix33 MI = (M0 + M1).Inverse();
 				const FVec3 DX = Stiffness * Utilities::Multiply(MI, CX);
 
 				// Apply constraint correction
-				const FVec3 DP0 = InvMs[0] * DX;
-				const FVec3 DP1 = -InvMs[1] * DX;
-				const FVec3 DR0 = Utilities::Multiply(InvIs[0], FVec3::CrossProduct(Xs[0] - Ps[0], DX));
-				const FVec3 DR1 = Utilities::Multiply(InvIs[1], FVec3::CrossProduct(Xs[1] - Ps[1], -DX));
+				const FVec3 DP0 = InvM(0) * DX;
+				const FVec3 DP1 = -InvM(1) * DX;
+				const FVec3 DR0 = Utilities::Multiply(InvI(0), FVec3::CrossProduct(ConnectorXs[0] - P(0), DX));
+				const FVec3 DR1 = Utilities::Multiply(InvI(1), FVec3::CrossProduct(ConnectorXs[1] - P(1), -DX));
 
 				ApplyPositionDelta(DP0, DP1);
 				ApplyRotationDelta(DR0, DR1);
@@ -1987,14 +2006,14 @@ namespace Chaos
 	}
 
 
-	void FJointSolverGaussSeidel::ApplySphericalPositionConstraint(
+	void FPBDJointSolver::ApplySphericalPositionConstraint(
 		const FReal Dt,
 		const FPBDJointSolverSettings& SolverSettings,
 		const FPBDJointSettings& JointSettings)
 	{
 		FVec3 Axis;
 		FReal Delta;
-		FPBDJointUtilities::GetSphericalAxisDelta(Xs[0], Xs[1], Axis, Delta);
+		FPBDJointUtilities::GetSphericalAxisDelta(ConnectorXs[0], ConnectorXs[1], Axis, Delta);
 
 		const FReal LimitPadding = GetLinearConstraintPadding(0);
 		const FReal Limit = FMath::Max(JointSettings.LinearLimit - LimitPadding, (FReal)0.);
@@ -2010,7 +2029,8 @@ namespace Chaos
 				}
 
 				const FReal JointStiffness = FPBDJointUtilities::GetLinearStiffness(SolverSettings, JointSettings);
-				ApplyPositionConstraint(JointStiffness, Axis, Error);
+				const FVec3 Connector0Correction = Axis * Limit;
+				ApplyPositionConstraint(JointStiffness, Axis, Error, Connector0Correction);
 			}
 			else
 			{
@@ -2023,7 +2043,7 @@ namespace Chaos
 	}
 
 
-	void FJointSolverGaussSeidel::ApplyCylindricalPositionConstraint(
+	void FPBDJointSolver::ApplyCylindricalPositionConstraint(
 		const FReal Dt,
 		const int32 AxisIndex,
 		const EJointMotionType AxialMotion,
@@ -2035,7 +2055,7 @@ namespace Chaos
 
 		FVec3 Axis, RadialAxis;
 		FReal AxialDelta, RadialDelta;
-		FPBDJointUtilities::GetCylindricalAxesDeltas(Rs[0], Xs[0], Xs[1], AxisIndex, Axis, AxialDelta, RadialAxis, RadialDelta);
+		FPBDJointUtilities::GetCylindricalAxesDeltas(ConnectorRs[0], ConnectorXs[0], ConnectorXs[1], AxisIndex, Axis, AxialDelta, RadialAxis, RadialDelta);
 
 		if (AxialDelta < 0.0f)
 		{
@@ -2066,7 +2086,8 @@ namespace Chaos
 				}
 
 				const FReal JointStiffness = FPBDJointUtilities::GetLinearStiffness(SolverSettings, JointSettings);
-				ApplyPositionConstraint(JointStiffness, Axis, AxialError, (int32)EJointCylindricalPositionConstraintType::Axial);
+				const FVec3 Connector0Correction = ConnectorXs[1] - ConnectorXs[0] - AxialError * Axis;
+				ApplyPositionConstraint(JointStiffness, Axis, AxialError, Connector0Correction, (int32)EJointCylindricalPositionConstraintType::Axial);
 			}
 		}
 
@@ -2093,13 +2114,14 @@ namespace Chaos
 				}
 
 				const FReal JointStiffness = FPBDJointUtilities::GetLinearStiffness(SolverSettings, JointSettings);
-				ApplyPositionConstraint(JointStiffness, RadialAxis, RadialError, (int32)EJointCylindricalPositionConstraintType::Radial);
+				const FVec3 Connector0Correction = ConnectorXs[1] - ConnectorXs[0] - RadialError * RadialAxis;
+				ApplyPositionConstraint(JointStiffness, RadialAxis, RadialError, Connector0Correction, (int32)EJointCylindricalPositionConstraintType::Radial);
 			}
 		}
 	}
 
 
-	void FJointSolverGaussSeidel::ApplyPlanarPositionConstraint(
+	void FPBDJointSolver::ApplyPlanarPositionConstraint(
 		const FReal Dt,
 		const int32 AxisIndex,
 		const EJointMotionType AxialMotion,
@@ -2108,7 +2130,7 @@ namespace Chaos
 	{
 		FVec3 Axis;
 		FReal Delta;
-		FPBDJointUtilities::GetPlanarAxisDelta(Rs[0], Xs[0], Xs[1], AxisIndex, Axis, Delta);
+		FPBDJointUtilities::GetPlanarAxisDelta(ConnectorRs[0], ConnectorXs[0], ConnectorXs[1], AxisIndex, Axis, Delta);
 
 		if (Delta < 0.0f)
 		{
@@ -2136,13 +2158,14 @@ namespace Chaos
 				}
 
 				const FReal JointStiffness = FPBDJointUtilities::GetLinearStiffness(SolverSettings, JointSettings);
-				ApplyPositionConstraint(JointStiffness, Axis, Error);
+				const FVec3 Connector0Correction = ConnectorXs[1] - ConnectorXs[0] - Axis * Error;
+				ApplyPositionConstraint(JointStiffness, Axis, Error, Connector0Correction);
 			}
 		}
 	}
 
 
-	void FJointSolverGaussSeidel::ApplyPositionDrive(
+	void FPBDJointSolver::ApplyPositionDrive(
 		const FReal Dt,
 		const int32 AxisIndex,
 		const FPBDJointSolverSettings& SolverSettings,
@@ -2164,7 +2187,7 @@ namespace Chaos
 	}
 
 
-	void FJointSolverGaussSeidel::ApplyPointProjection(
+	void FPBDJointSolver::ApplyPointProjection(
 		const FReal Dt,
 		const FPBDJointSolverSettings& SolverSettings,
 		const FPBDJointSettings& JointSettings,
@@ -2174,15 +2197,15 @@ namespace Chaos
 	{
 		const FReal ProjectionPositionTolerance = 0.0f;//PositionTolerance;
 
-		const FVec3 CX = Xs[1] - Xs[0];
+		const FVec3 CX = ConnectorXs[1] - ConnectorXs[0];
 		if (CX.Size() > ProjectionPositionTolerance)
 		{
-			FMatrix33 J = Utilities::ComputeJointFactorMatrix(Xs[1] - Ps[1], InvIs[1], InvMs[1]);
+			FMatrix33 J = Utilities::ComputeJointFactorMatrix(ConnectorXs[1] - P(1), InvI(1), InvM(1));
 			const FMatrix33 IJ = J.Inverse();
 			const FVec3 DX = Utilities::Multiply(IJ, CX);
 
-			const FVec3 DP1 = -Alpha * InvMs[1] * DX;
-			const FVec3 DR1 = -Alpha * Utilities::Multiply(InvIs[1], FVec3::CrossProduct(Xs[1] - Ps[1], DX));
+			const FVec3 DP1 = -Alpha * InvM(1) * DX;
+			const FVec3 DR1 = -Alpha * Utilities::Multiply(InvI(1), FVec3::CrossProduct(ConnectorXs[1] - P(1), DX));
 			ApplyDelta(1, DP1, DR1);
 
 			NetDP1 += DP1;
@@ -2190,7 +2213,7 @@ namespace Chaos
 		}
 	}
 
-	void FJointSolverGaussSeidel::ApplySphereProjection(
+	void FPBDJointSolver::ApplySphereProjection(
 		const FReal Dt,
 		const FPBDJointSolverSettings& SolverSettings,
 		const FPBDJointSettings& JointSettings,
@@ -2202,18 +2225,18 @@ namespace Chaos
 
 		FVec3 Axis;
 		FReal Delta;
-		FPBDJointUtilities::GetSphericalAxisDelta(Xs[0], Xs[1], Axis, Delta);
+		FPBDJointUtilities::GetSphericalAxisDelta(ConnectorXs[0], ConnectorXs[1], Axis, Delta);
 		const FReal Error = FMath::Max((FReal)0, Delta - JointSettings.LinearLimit);
 		if (FMath::Abs(Error) > ProjectionPositionTolerance)
 		{
-			const FVec3 AngularAxis1 = FVec3::CrossProduct(Xs[1] - Ps[1], Axis);
-			const FVec3 IA1 = Utilities::Multiply(InvIs[1], AngularAxis1);
+			const FVec3 AngularAxis1 = FVec3::CrossProduct(ConnectorXs[1] - P(1), Axis);
+			const FVec3 IA1 = Utilities::Multiply(InvI(1), AngularAxis1);
 			const FReal II1 = FVec3::DotProduct(AngularAxis1, IA1);
-			const FReal IM = InvMs[1] + II1;
+			const FReal IM = InvM(1) + II1;
 			const FVec3 DX = Axis * Error / IM;
 
-			const FVec3 DP1 = -Alpha * InvMs[1] * DX;
-			const FVec3 DR1 = -Alpha * Utilities::Multiply(InvIs[1], FVec3::CrossProduct(Xs[1] - Ps[1], DX));
+			const FVec3 DP1 = -Alpha * InvM(1) * DX;
+			const FVec3 DR1 = -Alpha * Utilities::Multiply(InvI(1), FVec3::CrossProduct(ConnectorXs[1] - P(1), DX));
 			ApplyDelta(1, DP1, DR1);
 
 			NetDP1 += DP1;
@@ -2221,7 +2244,7 @@ namespace Chaos
 		}
 	}
 
-	void FJointSolverGaussSeidel::ApplyTranslateProjection(
+	void FPBDJointSolver::ApplyTranslateProjection(
 		const FReal Dt,
 		const FPBDJointSolverSettings& SolverSettings,
 		const FPBDJointSettings& JointSettings,
@@ -2233,7 +2256,7 @@ namespace Chaos
 
 		FVec3 Axis;
 		FReal Delta;
-		FPBDJointUtilities::GetSphericalAxisDelta(Xs[0], Xs[1], Axis, Delta);
+		FPBDJointUtilities::GetSphericalAxisDelta(ConnectorXs[0], ConnectorXs[1], Axis, Delta);
 		const FReal Error = FMath::Max((FReal)0, Delta - JointSettings.LinearLimit);
 		if (Error > ProjectionPositionTolerance)
 		{
@@ -2244,7 +2267,7 @@ namespace Chaos
 		}
 	}
 
-	void FJointSolverGaussSeidel::ApplyConeProjection(
+	void FPBDJointSolver::ApplyConeProjection(
 		const FReal Dt,
 		const FPBDJointSolverSettings& SolverSettings,
 		const FPBDJointSettings& JointSettings,
@@ -2259,15 +2282,15 @@ namespace Chaos
 		FReal DSwingAngle = 0.0f;
 		const FReal Swing1Limit = JointSettings.AngularLimits[(int32)EJointAngularConstraintIndex::Swing1];
 		const FReal Swing2Limit = JointSettings.AngularLimits[(int32)EJointAngularConstraintIndex::Swing2];
-		FPBDJointUtilities::GetEllipticalConeAxisErrorLocal(Rs[0], Rs[1], Swing2Limit, Swing1Limit, SwingAxisLocal, DSwingAngle);
-		FVec3 SwingAxis = Rs[0] * SwingAxisLocal;
+		FPBDJointUtilities::GetEllipticalConeAxisErrorLocal(ConnectorRs[0], ConnectorRs[1], Swing2Limit, Swing1Limit, SwingAxisLocal, DSwingAngle);
+		FVec3 SwingAxis = ConnectorRs[0] * SwingAxisLocal;
 		if (DSwingAngle > ProjectionAngleTolerance)
 		{
 			const FVec3 DR1 = -Alpha * DSwingAngle * SwingAxis;
 			FVec3 DP1 = FVec3(0);
 			if (bPositionLocked)
 			{
-				DP1 = -Alpha * FVec3::CrossProduct(DR1, Xs[1] - Ps[1]);
+				DP1 = -Alpha * FVec3::CrossProduct(DR1, ConnectorXs[1] - P(1));
 			}
 			ApplyDelta(1, DP1, DR1);
 
@@ -2276,7 +2299,7 @@ namespace Chaos
 		}
 	}
 
-	void FJointSolverGaussSeidel::ApplySwingProjection(
+	void FPBDJointSolver::ApplySwingProjection(
 		const FReal Dt,
 		const FPBDJointSolverSettings& SolverSettings,
 		const FPBDJointSettings& JointSettings,
@@ -2290,7 +2313,7 @@ namespace Chaos
 
 		FVec3 SwingAxis;
 		FReal SwingAngle;
-		FPBDJointUtilities::GetSwingAxisAngle(Rs[0], Rs[1], SolverSettings.SwingTwistAngleTolerance, SwingConstraintIndex, SwingAxis, SwingAngle);
+		FPBDJointUtilities::GetSwingAxisAngle(ConnectorRs[0], ConnectorRs[1], SolverSettings.SwingTwistAngleTolerance, SwingConstraintIndex, SwingAxis, SwingAngle);
 
 		// Calculate swing error we need to correct
 		FReal DSwingAngle = 0;
@@ -2313,7 +2336,7 @@ namespace Chaos
 			FVec3 DP1 = FVec3(0);
 			if (bPositionLocked)
 			{
-				DP1 = -Alpha * FVec3::CrossProduct(DR1, Xs[1] - Ps[1]);
+				DP1 = -Alpha * FVec3::CrossProduct(DR1, ConnectorXs[1] - P(1));
 			}
 			ApplyDelta(1, DP1, DR1);
 
@@ -2322,7 +2345,7 @@ namespace Chaos
 		}
 	}
 
-	void FJointSolverGaussSeidel::ApplySingleLockedSwingProjection(
+	void FPBDJointSolver::ApplySingleLockedSwingProjection(
 		const FReal Dt,
 		const FPBDJointSolverSettings& SolverSettings,
 		const FPBDJointSettings& JointSettings,
@@ -2335,7 +2358,7 @@ namespace Chaos
 		ApplySwingProjection(Dt, SolverSettings, JointSettings, SwingConstraintIndex, Alpha, bPositionLocked, NetDP1, NetDR1);
 	}
 
-	void FJointSolverGaussSeidel::ApplyDoubleLockedSwingProjection(
+	void FPBDJointSolver::ApplyDoubleLockedSwingProjection(
 		const FReal Dt,
 		const FPBDJointSolverSettings& SolverSettings,
 		const FPBDJointSettings& JointSettings,
@@ -2348,15 +2371,15 @@ namespace Chaos
 
 		FVec3 SwingAxisLocal;
 		FReal DSwingAngle = 0.0f;
-		FPBDJointUtilities::GetCircularConeAxisErrorLocal(Rs[0], Rs[1], 0.0f, SwingAxisLocal, DSwingAngle);
-		FVec3 SwingAxis = Rs[0] * SwingAxisLocal;
+		FPBDJointUtilities::GetCircularConeAxisErrorLocal(ConnectorRs[0], ConnectorRs[1], 0.0f, SwingAxisLocal, DSwingAngle);
+		FVec3 SwingAxis = ConnectorRs[0] * SwingAxisLocal;
 		if (DSwingAngle > ProjectionAngleTolerance)
 		{
 			const FVec3 DR1 = -Alpha * DSwingAngle * SwingAxis;
 			FVec3 DP1 = FVec3(0);
 			if (bPositionLocked)
 			{
-				DP1 = -Alpha * FVec3::CrossProduct(DR1, Xs[1] - Ps[1]);
+				DP1 = -Alpha * FVec3::CrossProduct(DR1, ConnectorXs[1] - P(1));
 			}
 			ApplyDelta(1, DP1, DR1);
 
@@ -2365,7 +2388,7 @@ namespace Chaos
 		}
 	}
 
-	void FJointSolverGaussSeidel::ApplyDualConeSwingProjection(
+	void FPBDJointSolver::ApplyDualConeSwingProjection(
 		const FReal Dt,
 		const FPBDJointSolverSettings& SolverSettings,
 		const FPBDJointSettings& JointSettings,
@@ -2379,7 +2402,7 @@ namespace Chaos
 
 		FVec3 SwingAxis;
 		FReal SwingAngle;
-		FPBDJointUtilities::GetDualConeSwingAxisAngle(Rs[0], Rs[1], SwingConstraintIndex, SwingAxis, SwingAngle);
+		FPBDJointUtilities::GetDualConeSwingAxisAngle(ConnectorRs[0], ConnectorRs[1], SwingConstraintIndex, SwingAxis, SwingAngle);
 
 		// Calculate swing error we need to correct
 		FReal DSwingAngle = 0;
@@ -2400,7 +2423,7 @@ namespace Chaos
 			FVec3 DP1 = FVec3(0);
 			if (bPositionLocked)
 			{
-				DP1 = -Alpha * FVec3::CrossProduct(DR1, Xs[1] - Ps[1]);
+				DP1 = -Alpha * FVec3::CrossProduct(DR1, ConnectorXs[1] - P(1));
 			}
 			ApplyDelta(1, DP1, DR1);
 
@@ -2409,7 +2432,7 @@ namespace Chaos
 		}
 	}
 
-	void FJointSolverGaussSeidel::ApplyTwistProjection(
+	void FPBDJointSolver::ApplyTwistProjection(
 		const FReal Dt,
 		const FPBDJointSolverSettings& SolverSettings,
 		const FPBDJointSettings& JointSettings,
@@ -2422,7 +2445,7 @@ namespace Chaos
 
 		FVec3 TwistAxis;
 		FReal TwistAngle;
-		FPBDJointUtilities::GetTwistAxisAngle(Rs[0], Rs[1], TwistAxis, TwistAngle);
+		FPBDJointUtilities::GetTwistAxisAngle(ConnectorRs[0], ConnectorRs[1], TwistAxis, TwistAngle);
 		FReal DTwistAngle = 0;
 		const FReal TwistLimit = JointSettings.AngularLimits[(int32)EJointAngularConstraintIndex::Twist];
 		if (TwistAngle > TwistLimit)
@@ -2440,7 +2463,7 @@ namespace Chaos
 			FVec3 DP1 = FVec3(0);
 			if (bPositionLocked)
 			{
-				DP1 = -Alpha * FVec3::CrossProduct(DR1, Xs[1] - Ps[1]);
+				DP1 = -Alpha * FVec3::CrossProduct(DR1, ConnectorXs[1] - P(1));
 			}
 			ApplyDelta(1, DP1, DR1);
 
@@ -2449,7 +2472,7 @@ namespace Chaos
 		}
 	}
 
-	void FJointSolverGaussSeidel::ApplyVelocityProjection(
+	void FPBDJointSolver::ApplyVelocityProjection(
 		const FReal Dt,
 		const FPBDJointSolverSettings& SolverSettings,
 		const FPBDJointSettings& JointSettings,
@@ -2462,33 +2485,37 @@ namespace Chaos
 		ApplyVelocityDelta(1, DV1, DW1);
 	}
 
-	void FJointSolverGaussSeidel::ApplyLinearVelocityConstraint(
+	void FPBDJointSolver::ApplyLinearVelocityConstraint(
 		const FReal Stiffness,
-		const FVec3& Axis)
+		const FVec3& Axis,
+		const FVec3& Connector0Correction,
+		const FReal TargetVel)
 	{
-		const FVec3 CV0 = Vs[0] + FVec3::CrossProduct(Ws[0], Xs[0] - Ps[0]);
-		const FVec3 CV1 = Vs[1] + FVec3::CrossProduct(Ws[1], Xs[1] - Ps[1]);
+		const FVec3 Arm0 = ConnectorXs[0] + Connector0Correction - P(0);
+		const FVec3 Arm1 = ConnectorXs[1] - P(1);
+		const FVec3 CV0 = V(0) + FVec3::CrossProduct(W(0), Arm0);
+		const FVec3 CV1 = V(1) + FVec3::CrossProduct(W(1), Arm1);
 		const FVec3 CV = CV1 - CV0;
 
-		const FVec3 AngularAxis0 = FVec3::CrossProduct(Xs[0] - Ps[0], Axis);
-		const FVec3 AngularAxis1 = FVec3::CrossProduct(Xs[1] - Ps[1], Axis);
-		const FVec3 IA0 = Utilities::Multiply(InvIs[0], AngularAxis0);
-		const FVec3 IA1 = Utilities::Multiply(InvIs[1], AngularAxis1);
+		const FVec3 AngularAxis0 = FVec3::CrossProduct(Arm0, Axis);
+		const FVec3 AngularAxis1 = FVec3::CrossProduct(Arm1, Axis);
+		const FVec3 IA0 = Utilities::Multiply(InvI(0), AngularAxis0);
+		const FVec3 IA1 = Utilities::Multiply(InvI(1), AngularAxis1);
 		const FReal II0 = FVec3::DotProduct(AngularAxis0, IA0);
 		const FReal II1 = FVec3::DotProduct(AngularAxis1, IA1);
-		const FReal II = (InvMs[0] + II0 + InvMs[1] + II1);
+		const FReal II = (InvM(0) + II0 + InvM(1) + II1);
 
-		const FVec3 Impulse = Stiffness * FVec3::DotProduct(CV, Axis) / II * Axis;
+		const FVec3 Impulse = Stiffness * (FVec3::DotProduct(CV, Axis) - TargetVel) / II * Axis;
 
-		const FVec3 DV0 = InvMs[0] * Impulse;
-		const FVec3 DV1 = -InvMs[1] * Impulse;
-		const FVec3 DW0 = Utilities::Multiply(InvIs[0], FVec3::CrossProduct(Xs[0] - Ps[0], Impulse));
-		const FVec3 DW1 = Utilities::Multiply(InvIs[1], FVec3::CrossProduct(Xs[1] - Ps[1], -Impulse));
+		const FVec3 DV0 = InvM(0) * Impulse;
+		const FVec3 DV1 = -InvM(1) * Impulse;
+		const FVec3 DW0 = Utilities::Multiply(InvI(0), FVec3::CrossProduct(Arm0, Impulse));
+		const FVec3 DW1 = Utilities::Multiply(InvI(1), FVec3::CrossProduct(Arm1, -Impulse));
 
 		ApplyVelocityDelta(DV0, DW0, DV1, DW1);
 	}
 
-	void FJointSolverGaussSeidel::ApplyPointVelocityConstraint(
+	void FPBDJointSolver::ApplyPointVelocityConstraint(
 		const FReal Dt,
 		const FPBDJointSolverSettings& SolverSettings,
 		const FPBDJointSettings& JointSettings)
@@ -2496,36 +2523,36 @@ namespace Chaos
 		if (!NetLinearImpulse.IsNearlyZero() && !FPBDJointUtilities::GetSoftLinearLimitEnabled(SolverSettings, JointSettings))
 		{
 			FReal Stiffness = SolverStiffness * FPBDJointUtilities::GetLinearStiffness(SolverSettings, JointSettings);
-			const FVec3 CV0 = Vs[0] + FVec3::CrossProduct(Ws[0], Xs[0] - Ps[0]);
-			const FVec3 CV1 = Vs[1] + FVec3::CrossProduct(Ws[1], Xs[1] - Ps[1]);
+			const FVec3 CV0 = V(0) + FVec3::CrossProduct(W(0), ConnectorXs[0] - P(0));
+			const FVec3 CV1 = V(1) + FVec3::CrossProduct(W(1), ConnectorXs[1] - P(1));
 			const FVec3 CV = CV1 - CV0;
 
 			UE_LOG(LogChaosJoint, VeryVerbose, TEXT("    PointVel Delta %f"), CV.Size());
 			// Calculate constraint correction
 			FMatrix33 M0 = FMatrix33::Zero;
 			FMatrix33 M1 = FMatrix33::Zero;
-			if (InvMs[0] > SMALL_NUMBER)
+			if (InvM(0) > SMALL_NUMBER)
 			{
-				M0 = Utilities::ComputeJointFactorMatrix(Xs[0] - Ps[0], InvIs[0], InvMs[0]);
+				M0 = Utilities::ComputeJointFactorMatrix(ConnectorXs[0] - P(0), InvI(0), InvM(0));
 		}
-			if (InvMs[1] > SMALL_NUMBER)
+			if (InvM(1) > SMALL_NUMBER)
 			{
-				M1 = Utilities::ComputeJointFactorMatrix(Xs[1] - Ps[1], InvIs[1], InvMs[1]);
+				M1 = Utilities::ComputeJointFactorMatrix(ConnectorXs[1] - P(1), InvI(1), InvM(1));
 	}
 			FMatrix33 MI = (M0 + M1).Inverse();
 			const FVec3 Impulse = Stiffness * Utilities::Multiply(MI, CV);
 
 			// Apply constraint correction
-			const FVec3 DV0 = InvMs[0] * Impulse;
-			const FVec3 DV1 = -InvMs[1] * Impulse;
-			const FVec3 DW0 = Utilities::Multiply(InvIs[0], FVec3::CrossProduct(Xs[0] - Ps[0], Impulse));
-			const FVec3 DW1 = Utilities::Multiply(InvIs[1], FVec3::CrossProduct(Xs[1] - Ps[1], -Impulse));
+			const FVec3 DV0 = InvM(0) * Impulse;
+			const FVec3 DV1 = -InvM(1) * Impulse;
+			const FVec3 DW0 = Utilities::Multiply(InvI(0), FVec3::CrossProduct(ConnectorXs[0] - P(0), Impulse));
+			const FVec3 DW1 = Utilities::Multiply(InvI(1), FVec3::CrossProduct(ConnectorXs[1] - P(1), -Impulse));
 
 			ApplyVelocityDelta(DV0, DW0, DV1, DW1);
 		}
 	}
 
-	void FJointSolverGaussSeidel::ApplySphericalVelocityConstraint(
+	void FPBDJointSolver::ApplySphericalVelocityConstraint(
 		const FReal Dt,
 		const FPBDJointSolverSettings& SolverSettings,
 		const FPBDJointSettings& JointSettings)
@@ -2534,14 +2561,22 @@ namespace Chaos
 		{
 			FVec3 Axis;
 			FReal Delta;
-			FPBDJointUtilities::GetSphericalAxisDelta(Xs[0], Xs[1], Axis, Delta);
+			FPBDJointUtilities::GetSphericalAxisDelta(ConnectorXs[0], ConnectorXs[1], Axis, Delta);
 
-			FReal Stiffness = SolverStiffness * FPBDJointUtilities::GetLinearStiffness(SolverSettings, JointSettings);
-			ApplyLinearVelocityConstraint(Stiffness, Axis);
+			const FReal Stiffness = SolverStiffness * FPBDJointUtilities::GetLinearStiffness(SolverSettings, JointSettings);
+			const FVec3 DConnectorX = ConnectorXs[1] - ConnectorXs[0];
+			const FVec3 Connector0Correction = Delta > JointSettings.LinearLimit ? Axis * JointSettings.LinearLimit : DConnectorX;
+			FReal TargetVel = 0.0f;
+			if (JointSettings.LinearRestitution != 0.0f)
+			{
+				const FReal InitVel = InitConstraintAxisLinearVelocities[0];
+				TargetVel = InitVel > Chaos_Joint_LinearVelocityThresholdToApplyRestitution ? -JointSettings.LinearRestitution * InitVel : 0.0f; 
+			}
+			ApplyLinearVelocityConstraint(Stiffness, Axis, Connector0Correction, TargetVel);
 		}
 	}
 
-	void FJointSolverGaussSeidel::ApplyCylindricalVelocityConstraint(
+	void FPBDJointSolver::ApplyCylindricalVelocityConstraint(
 		const FReal Dt,
 		const int32 AxisIndex,
 		const EJointMotionType AxialMotion,
@@ -2555,23 +2590,54 @@ namespace Chaos
 		{
 			FVec3 Axis, RadialAxis;
 			FReal AxialDelta, RadialDelta;
-			FPBDJointUtilities::GetCylindricalAxesDeltas(Rs[0], Xs[0], Xs[1], AxisIndex, Axis, AxialDelta, RadialAxis, RadialDelta);
+			FPBDJointUtilities::GetCylindricalAxesDeltas(ConnectorRs[0], ConnectorXs[0], ConnectorXs[1], AxisIndex, Axis, AxialDelta, RadialAxis, RadialDelta);
+			const FVec3 DConnectorX = ConnectorXs[1] - ConnectorXs[0];
 
 			if (LinearHardLambda[(int32)EJointCylindricalPositionConstraintType::Axial] > SMALL_NUMBER)
 			{
-				FReal Stiffness = SolverStiffness * FPBDJointUtilities::GetLinearStiffness(SolverSettings, JointSettings);
-				ApplyLinearVelocityConstraint(Stiffness, Axis);
+				const FReal Stiffness = SolverStiffness * FPBDJointUtilities::GetLinearStiffness(SolverSettings, JointSettings);
+
+				// This flipping of axis might result in undesired restitution velocity if the current axis and the axis at the beginning of the frame are pointing at very different directions.
+				// On the other hand, Quasipbd restitution in general works only on the assumption that constraint axes do not change much within one frame. 
+				if (AxialDelta < 0.0f)
+				{
+					AxialDelta = -AxialDelta;
+					Axis = -Axis;
+				}
+				
+				const FReal AxialLimit = (AxialMotion == EJointMotionType::Locked) ? 0.0f : JointSettings.LinearLimit;
+				const FReal AxialError = AxialDelta - AxialLimit;
+
+				const FVec3 Connector0Correction = DConnectorX - AxialError * Axis;
+				FReal TargetVel = 0.0f;
+				if (AxialMotion == EJointMotionType::Limited && JointSettings.LinearRestitution != 0.0f)
+				{
+					const FReal InitVel = InitConstraintAxisLinearVelocities[(int32)EJointCylindricalPositionConstraintType::Axial];
+					TargetVel = InitVel > Chaos_Joint_LinearVelocityThresholdToApplyRestitution ? -JointSettings.LinearRestitution * InitVel : 0.0f; 
+				}
+
+				ApplyLinearVelocityConstraint(Stiffness, Axis, Connector0Correction, TargetVel);
 			}
 
 			if (LinearHardLambda[(int32)EJointCylindricalPositionConstraintType::Radial] > SMALL_NUMBER)
 			{
-				FReal Stiffness = SolverStiffness * FPBDJointUtilities::GetLinearStiffness(SolverSettings, JointSettings);
-				ApplyLinearVelocityConstraint(Stiffness, RadialAxis);
+				const FReal Stiffness = SolverStiffness * FPBDJointUtilities::GetLinearStiffness(SolverSettings, JointSettings);
+
+				const FReal RadialLimit = (RadialMotion == EJointMotionType::Locked) ? 0.0f : JointSettings.LinearLimit;
+				const FReal RadialError = FMath::Max(RadialDelta - RadialLimit, 0.0f);
+				const FVec3 Connector0Correction = DConnectorX - RadialError * RadialAxis;
+				FReal TargetVel = 0.0f;
+				if (RadialMotion == EJointMotionType::Limited && JointSettings.LinearRestitution != 0.0f)
+				{
+					const FReal InitVel = InitConstraintAxisLinearVelocities[(int32)EJointCylindricalPositionConstraintType::Radial];
+					TargetVel = InitVel > Chaos_Joint_LinearVelocityThresholdToApplyRestitution ? -JointSettings.LinearRestitution * InitVel : 0.0f; 
+				}
+				ApplyLinearVelocityConstraint(Stiffness, RadialAxis, Connector0Correction, TargetVel);
 			}
 		}
 	}
 
-	void FJointSolverGaussSeidel::ApplyPlanarVelocityConstraint(
+	void FPBDJointSolver::ApplyPlanarVelocityConstraint(
 		const FReal Dt,
 		const int32 AxisIndex,
 		const EJointMotionType AxialMotion,
@@ -2582,11 +2648,352 @@ namespace Chaos
 		{
 			FVec3 Axis;
 			FReal Delta;
-			FPBDJointUtilities::GetPlanarAxisDelta(Rs[0], Xs[0], Xs[1], AxisIndex, Axis, Delta);
+			FPBDJointUtilities::GetPlanarAxisDelta(ConnectorRs[0], ConnectorXs[0], ConnectorXs[1], AxisIndex, Axis, Delta);
 
-			FReal Stiffness = SolverStiffness * FPBDJointUtilities::GetLinearStiffness(SolverSettings, JointSettings);
-			ApplyLinearVelocityConstraint(Stiffness, Axis);
+			const FReal Stiffness = SolverStiffness * FPBDJointUtilities::GetLinearStiffness(SolverSettings, JointSettings);
+
+			if (Delta < 0.0f)
+			{
+				Delta = -Delta;
+				Axis = -Axis;
+			}
+			const FReal Limit = (AxialMotion == EJointMotionType::Locked) ? 0 : JointSettings.LinearLimit;
+			const FReal Error = Delta - Limit;
+
+			const FVec3 DConnectorX = ConnectorXs[1] - ConnectorXs[0];
+			const FVec3 Connector0Correction = DConnectorX - Axis * Error;
+
+			FReal TargetVel = 0.0f;
+			if (AxialMotion == EJointMotionType::Limited && JointSettings.LinearRestitution != 0.0f)
+			{
+				const FReal InitVel = InitConstraintAxisLinearVelocities[0];
+				TargetVel = InitVel > Chaos_Joint_LinearVelocityThresholdToApplyRestitution ? -JointSettings.LinearRestitution * InitVel : 0.0f; 
+			}
+			ApplyLinearVelocityConstraint(Stiffness, Axis, Connector0Correction, TargetVel);
+		}
+	}
+
+	void FPBDJointSolver::CalculateConstraintAxisLinearVelocities(
+		const FPBDJointSettings& JointSettings,
+		FVec3& ConstraintAxisLinearVelocities) const
+	{
+		bool bHasPositionConstraints =
+			(JointSettings.LinearMotionTypes[0] != EJointMotionType::Free)
+			|| (JointSettings.LinearMotionTypes[1] != EJointMotionType::Free)
+			|| (JointSettings.LinearMotionTypes[2] != EJointMotionType::Free);
+		if (!bHasPositionConstraints)
+		{
+			return;
+		}
+
+		const TVector<EJointMotionType, 3>& LinearMotion = JointSettings.LinearMotionTypes;
+		const TVector<bool, 3> bLinearLocked =
+		{
+			(LinearMotion[0] == EJointMotionType::Locked),
+			(LinearMotion[1] == EJointMotionType::Locked),
+			(LinearMotion[2] == EJointMotionType::Locked),
+		};
+		const TVector<bool, 3> bLinearLimited =
+		{
+			(LinearMotion[0] == EJointMotionType::Limited),
+			(LinearMotion[1] == EJointMotionType::Limited),
+			(LinearMotion[2] == EJointMotionType::Limited),
+		};
+
+		if (bLinearLocked[0] && bLinearLocked[1] && bLinearLocked[2])
+		{
+			
+		}
+		else if (bLinearLimited[0] && bLinearLimited[1] && bLinearLimited[2])
+		{
+			CalculateSphericalConstraintAxisLinearVelocities(JointSettings, ConstraintAxisLinearVelocities);
+		}
+		else if (bLinearLocked[1] && bLinearLocked[2] && !bLinearLocked[0])
+		{
+			// Line constraint along X axis
+			CalculateCylindricalConstraintAxisLinearVelocities(0, LinearMotion[0], EJointMotionType::Locked, JointSettings, ConstraintAxisLinearVelocities);
+		}
+		else if (bLinearLocked[0] && bLinearLocked[2] && !bLinearLocked[1])
+		{
+			// Line constraint along Y axis
+			CalculateCylindricalConstraintAxisLinearVelocities(1, LinearMotion[1], EJointMotionType::Locked, JointSettings, ConstraintAxisLinearVelocities);
+		}
+		else if (bLinearLocked[0] && bLinearLocked[1] && !bLinearLocked[2])
+		{
+			// Line constraint along Z axis
+			CalculateCylindricalConstraintAxisLinearVelocities(2, LinearMotion[2], EJointMotionType::Locked, JointSettings, ConstraintAxisLinearVelocities);
+		}
+		else if (bLinearLimited[1] && bLinearLimited[2] && !bLinearLimited[0])
+		{
+			// Cylindrical constraint along X axis
+			CalculateCylindricalConstraintAxisLinearVelocities(0, LinearMotion[0], EJointMotionType::Limited, JointSettings, ConstraintAxisLinearVelocities);
+		}
+		else if (bLinearLimited[0] && bLinearLimited[2] && !bLinearLimited[1])
+		{
+			// Cylindrical constraint along Y axis
+			CalculateCylindricalConstraintAxisLinearVelocities(1, LinearMotion[1], EJointMotionType::Limited, JointSettings, ConstraintAxisLinearVelocities);
+		}
+		else if (bLinearLimited[0] && bLinearLimited[1] && !bLinearLimited[2])
+		{
+			// Cylindrical constraint along Z axis
+			CalculateCylindricalConstraintAxisLinearVelocities(2, LinearMotion[2], EJointMotionType::Limited, JointSettings, ConstraintAxisLinearVelocities);
+		}
+		else if (bLinearLimited[0])
+		{
+			// Planar constraint along X axis
+			CalculatePlanarConstraintAxisLinearVelocities(0, LinearMotion[0], JointSettings, ConstraintAxisLinearVelocities);
+		}
+		else if (bLinearLimited[1])
+		{
+			// Planar constraint along Y axis
+			CalculatePlanarConstraintAxisLinearVelocities(1, LinearMotion[1], JointSettings, ConstraintAxisLinearVelocities);
+		}
+		else if (bLinearLimited[2])
+		{
+			// Planar constraint along Z axis
+			CalculatePlanarConstraintAxisLinearVelocities(2, LinearMotion[2], JointSettings, ConstraintAxisLinearVelocities);
+		}
+	}
+
+	void FPBDJointSolver::CalculateSphericalConstraintAxisLinearVelocities(
+			const FPBDJointSettings& JointSettings,
+			FVec3& ConstraintAxisLinearVelocities) const
+	{
+		FVec3 Axis;
+		FReal Delta;
+		FPBDJointUtilities::GetSphericalAxisDelta(ConnectorXs[0], ConnectorXs[1], Axis, Delta);
+		const FVec3 DConnectorX = ConnectorXs[1] - ConnectorXs[0];
+		const FVec3 Connector0Correction = Delta > JointSettings.LinearLimit ? Axis * JointSettings.LinearLimit : DConnectorX;
+
+		const FVec3 Arm0 = ConnectorXs[0] + Connector0Correction - P(0);
+		const FVec3 Arm1 = ConnectorXs[1] - P(1);
+		const FVec3 CV0 = V(0) + FVec3::CrossProduct(W(0), Arm0);
+		const FVec3 CV1 = V(1) + FVec3::CrossProduct(W(1), Arm1);
+		const FVec3 CV = CV1 - CV0;
+
+		ConstraintAxisLinearVelocities[0] = FVec3::DotProduct(CV, Axis);
+	}
+
+	void FPBDJointSolver::CalculateCylindricalConstraintAxisLinearVelocities(
+			const int32 AxisIndex,
+			const EJointMotionType AxialMotion,
+			const EJointMotionType RadialMotion,
+			const FPBDJointSettings& JointSettings,
+			FVec3& ConstraintAxisLinearVelocities) const
+	{
+		FVec3 Axis, RadialAxis;
+		FReal AxialDelta, RadialDelta;
+		FPBDJointUtilities::GetCylindricalAxesDeltas(ConnectorRs[0], ConnectorXs[0], ConnectorXs[1], AxisIndex, Axis, AxialDelta, RadialAxis, RadialDelta);
+		const FVec3 DConnectorX = ConnectorXs[1] - ConnectorXs[0];
+
+		if (AxialMotion == EJointMotionType::Limited)
+		{
+			if (AxialDelta < 0.0f)
+			{
+				AxialDelta = -AxialDelta;
+				Axis = -Axis;
+			}
+			const FReal AxialLimit = (AxialMotion == EJointMotionType::Locked) ? 0.0f :JointSettings.LinearLimit;
+			const FReal AxialError = AxialDelta - AxialLimit;
+			
+			const FVec3 Connector0Correction = DConnectorX - AxialError * Axis;
+
+			const FVec3 Arm0 = ConnectorXs[0] + Connector0Correction - P(0);
+			const FVec3 Arm1 = ConnectorXs[1] - P(1);
+			const FVec3 CV0 = V(0) + FVec3::CrossProduct(W(0), Arm0);
+			const FVec3 CV1 = V(1) + FVec3::CrossProduct(W(1), Arm1);
+			const FVec3 CV = CV1 - CV0;
+
+			ConstraintAxisLinearVelocities[(int32)EJointCylindricalPositionConstraintType::Axial] = FVec3::DotProduct(CV, Axis);
+		}
+
+		if (RadialMotion == EJointMotionType::Limited)
+		{
+			const FReal RadialLimit = (RadialMotion == EJointMotionType::Locked) ? 0.0f : JointSettings.LinearLimit;
+			const FReal RadialError = FMath::Max(RadialDelta - RadialLimit, 0.0f);
+			const FVec3 Connector0Correction = DConnectorX - RadialError * RadialAxis;
+			
+			const FVec3 Arm0 = ConnectorXs[0] + Connector0Correction - P(0);
+			const FVec3 Arm1 = ConnectorXs[1] - P(1);
+			const FVec3 CV0 = V(0) + FVec3::CrossProduct(W(0), Arm0);
+			const FVec3 CV1 = V(1) + FVec3::CrossProduct(W(1), Arm1);
+			const FVec3 CV = CV1 - CV0;
+
+			ConstraintAxisLinearVelocities[(int32)EJointCylindricalPositionConstraintType::Radial] = FVec3::DotProduct(CV, RadialAxis);
+		}
+	}
+
+	void FPBDJointSolver::CalculatePlanarConstraintAxisLinearVelocities(
+			const int32 AxisIndex,
+			const EJointMotionType AxialMotion,
+			const FPBDJointSettings& JointSettings,
+			FVec3& ConstraintAxisLinearVelocities) const
+	{
+		FVec3 Axis;
+		FReal Delta;
+		FPBDJointUtilities::GetPlanarAxisDelta(ConnectorRs[0], ConnectorXs[0], ConnectorXs[1], AxisIndex, Axis, Delta);
+
+		if (Delta < 0.0f)
+		{
+			Delta = -Delta;
+			Axis = -Axis;
+		}
+
+		const FReal Limit = (AxialMotion == EJointMotionType::Locked) ? 0 : JointSettings.LinearLimit;
+		const FReal Error = Delta - Limit;
+
+		const FVec3 DConnectorX = ConnectorXs[1] - ConnectorXs[0];
+		const FVec3 Connector0Correction = DConnectorX - Axis * Error; 
+
+		const FVec3 Arm0 = ConnectorXs[0] + Connector0Correction - P(0);
+		const FVec3 Arm1 = ConnectorXs[1] - P(1);
+		const FVec3 CV0 = V(0) + FVec3::CrossProduct(W(0), Arm0);
+		const FVec3 CV1 = V(1) + FVec3::CrossProduct(W(1), Arm1);
+		const FVec3 CV = CV1 - CV0;
+
+		ConstraintAxisLinearVelocities[0] = FVec3::DotProduct(CV, Axis);
+	}
+
+	void FPBDJointSolver::CalculateConstraintAxisAngularVelocities(
+		const FPBDJointSolverSettings& SolverSettings,
+		const FPBDJointSettings& JointSettings,
+		FVec3& ConstraintAxisAngularVelocities
+		) const
+	{
+		bool bHasRotationConstraints =
+			(JointSettings.AngularMotionTypes[0] != EJointMotionType::Free)
+			|| (JointSettings.AngularMotionTypes[1] != EJointMotionType::Free)
+			|| (JointSettings.AngularMotionTypes[2] != EJointMotionType::Free);
+		if (!bHasRotationConstraints)
+		{
+			return;
+		}
+
+		// Locked axes always use hard constraints. Limited axes use hard or soft depending on settings
+		EJointMotionType TwistMotion = JointSettings.AngularMotionTypes[(int32)EJointAngularConstraintIndex::Twist];
+		EJointMotionType Swing1Motion = JointSettings.AngularMotionTypes[(int32)EJointAngularConstraintIndex::Swing1];
+		EJointMotionType Swing2Motion = JointSettings.AngularMotionTypes[(int32)EJointAngularConstraintIndex::Swing2];
+
+		// If the twist axes are opposing, we cannot decompose the orientation into swing and twist angles, so just give up
+		const FVec3 Twist0 = ConnectorRs[0] * FJointConstants::TwistAxis();
+		const FVec3 Twist1 = ConnectorRs[1] * FJointConstants::TwistAxis();
+		const FReal Twist01Dot = FVec3::DotProduct(Twist0, Twist1);
+		const bool bDegenerate = (Twist01Dot < Chaos_Joint_DegenerateRotationLimit);
+		if (bDegenerate)
+		{
+			UE_LOG(LogChaosJoint, VeryVerbose, TEXT(" Degenerate rotation at Swing %f deg"), FMath::RadiansToDegrees(FMath::Acos(Twist01Dot)));
+		}
+
+		// Apply twist constraint
+		// NOTE: Cannot calculate twist angle at 180degree swing
+		if (SolverSettings.bEnableTwistLimits && !bDegenerate)
+		{
+			if (TwistMotion == EJointMotionType::Limited)
+			{
+				CalculateTwistConstraintAxisAngularVelocities(JointSettings, ConstraintAxisAngularVelocities);
+			}
+		}
+
+		// Apply swing constraints
+		// NOTE: Cannot separate swing angles at 180degree swing (but we can still apply locks)
+		if (SolverSettings.bEnableSwingLimits)
+		{
+			if ((Swing1Motion == EJointMotionType::Limited) && (Swing2Motion == EJointMotionType::Limited))
+			{
+				CalculateConeConstraintAxisAngularVelocities(JointSettings, ConstraintAxisAngularVelocities);
+			}
+			else if ((Swing1Motion == EJointMotionType::Limited) && (Swing2Motion == EJointMotionType::Locked))
+			{
+				if (!bDegenerate)
+				{
+					CalculateSwingConstraintAxisAngularVelocities(SolverSettings, JointSettings, EJointAngularConstraintIndex::Swing1, ConstraintAxisAngularVelocities);
+				}
+			}
+			else if ((Swing1Motion == EJointMotionType::Limited) && (Swing2Motion == EJointMotionType::Free))
+			{
+				if (!bDegenerate)
+				{
+					CalculateDualConeSwingConstraintAxisAngularVelocities(JointSettings, EJointAngularConstraintIndex::Swing1, ConstraintAxisAngularVelocities);
+				}
+			}
+			else if ((Swing1Motion == EJointMotionType::Locked) && (Swing2Motion == EJointMotionType::Limited))
+			{
+				if (!bDegenerate)
+				{
+					CalculateSwingConstraintAxisAngularVelocities(SolverSettings, JointSettings, EJointAngularConstraintIndex::Swing2, ConstraintAxisAngularVelocities);
+				}
+			}
+			else if ((Swing1Motion == EJointMotionType::Locked) && (Swing2Motion == EJointMotionType::Locked))
+			{
+			}
+			else if ((Swing1Motion == EJointMotionType::Locked) && (Swing2Motion == EJointMotionType::Free))
+			{
+			}
+			else if ((Swing1Motion == EJointMotionType::Free) && (Swing2Motion == EJointMotionType::Limited))
+			{
+				if (!bDegenerate)
+				{
+					CalculateDualConeSwingConstraintAxisAngularVelocities(JointSettings, EJointAngularConstraintIndex::Swing2, ConstraintAxisAngularVelocities);
+				}
+			}
+			else if ((Swing1Motion == EJointMotionType::Free) && (Swing2Motion == EJointMotionType::Locked))
+			{
+			}
+			else if ((Swing1Motion == EJointMotionType::Free) && (Swing2Motion == EJointMotionType::Free))
+			{
+			}
 		}
 	}
 
+	void FPBDJointSolver::CalculateTwistConstraintAxisAngularVelocities(
+		const FPBDJointSettings& JointSettings,
+		FVec3& ConstraintAxisAngularVelocities) const
+	{
+		FVec3 TwistAxis;
+		FReal TwistAngle;
+		FPBDJointUtilities::GetTwistAxisAngle(ConnectorRs[0], ConnectorRs[1], TwistAxis, TwistAngle);
+
+		if (TwistAngle < 0.0f)
+		{
+			TwistAxis = -TwistAxis;
+		}
+		ConstraintAxisAngularVelocities[(int32)EJointAngularConstraintIndex::Twist] = FVec3::DotProduct(W(1) - W(0), TwistAxis);
+	}
+
+	void FPBDJointSolver::CalculateConeConstraintAxisAngularVelocities(
+			const FPBDJointSettings& JointSettings,
+			FVec3& ConstraintAxisAngularVelocities) const
+	{
+		FVec3 SwingAxisLocal;
+		FReal DSwingAngle = 0.0f;
+		const FReal Swing1Limit = JointSettings.AngularLimits[(int32)EJointAngularConstraintIndex::Swing1];
+		const FReal Swing2Limit = JointSettings.AngularLimits[(int32)EJointAngularConstraintIndex::Swing2];
+		FPBDJointUtilities::GetEllipticalConeAxisErrorLocal(ConnectorRs[0], ConnectorRs[1], Swing2Limit, Swing1Limit, SwingAxisLocal, DSwingAngle);
+		SwingAxisLocal.SafeNormalize();
+		const FVec3 SwingAxis = ConnectorRs[0] * SwingAxisLocal;
+		ConstraintAxisAngularVelocities[(int32)EJointAngularConstraintIndex::Swing1] = FVec3::DotProduct(W(1) - W(0), SwingAxis);	
+	}
+
+	void FPBDJointSolver::CalculateDualConeSwingConstraintAxisAngularVelocities(
+			const FPBDJointSettings& JointSettings,
+			const EJointAngularConstraintIndex SwingConstraintIndex,
+			FVec3& ConstraintAxisAngularVelocities) const
+	{
+		FVec3 SwingAxis;
+		FReal SwingAngle;
+		FPBDJointUtilities::GetDualConeSwingAxisAngle(ConnectorRs[0], ConnectorRs[1], SwingConstraintIndex, SwingAxis, SwingAngle);
+		ConstraintAxisAngularVelocities[(int32)SwingConstraintIndex] = FVec3::DotProduct(W(1) - W(0), SwingAxis);
+	}
+
+	void FPBDJointSolver::CalculateSwingConstraintAxisAngularVelocities(
+			const FPBDJointSolverSettings& SolverSettings,
+			const FPBDJointSettings& JointSettings,
+			const EJointAngularConstraintIndex SwingConstraintIndex,
+			FVec3& ConstraintAxisAngularVelocities) const
+	{
+		FVec3 SwingAxis;
+		FReal SwingAngle;
+		FPBDJointUtilities::GetSwingAxisAngle(ConnectorRs[0], ConnectorRs[1], SolverSettings.SwingTwistAngleTolerance, SwingConstraintIndex, SwingAxis, SwingAngle);
+		ConstraintAxisAngularVelocities[(int32)SwingConstraintIndex] = FVec3::DotProduct(W(1) - W(0), SwingAxis);
+	}
 }
diff --git a/Engine/Source/Runtime/Experimental/Chaos/Private/Chaos/Joint/PBDJointSolverGaussSeidel.ispc b/Engine/Source/Runtime/Experimental/Chaos/Private/Chaos/Joint/PBDJointSolverGaussSeidel.ispc
--- a/Engine/Source/Runtime/Experimental/Chaos/Private/Chaos/Joint/PBDJointSolverGaussSeidel.ispc
+++ b/Engine/Source/Runtime/Experimental/Chaos/Private/Chaos/Joint/PBDJointSolverGaussSeidel.ispc
@@ -15,7 +15,7 @@ struct FRigidTransform
 
 static const uniform int32 MaxConstrainedBodies = 2;
 
-struct FJointSolverGaussSeidel
+struct FPBDJointSolver
 {
 	// Local-space constraint settings
 	FRigidTransform XLs[MaxConstrainedBodies];	// Local-space joint connector transforms
@@ -71,9 +71,9 @@ struct FJointSolverGaussSeidel
 	bool bIsActive;								// Whether any constraints actually moved any bodies in last iteration
 };
 
-export uniform int32 SizeofFJointSolverGaussSeidel()
+export uniform int32 SizeofFPBDJointSolver()
 {
-	return sizeof(uniform FJointSolverGaussSeidel);
+	return sizeof(uniform FPBDJointSolver);
 }
 
 static inline uniform FMatrix33 ComputeWorldSpaceInertia(const uniform FVector4 &Q, const uniform FVector &I)
@@ -132,7 +132,7 @@ static inline uniform FVector CalculateVelocity(const uniform FVector& P0, const
 }
 
 static inline void ApplyPositionDelta(
-	uniform FJointSolverGaussSeidel * uniform M,
+	uniform FPBDJointSolver * uniform M,
 	const uniform int Index,
 	const uniform FVector& DP)
 {
@@ -140,7 +140,7 @@ static inline void ApplyPositionDelta(
 }
 
 static inline void ApplyPositionDelta(
-	uniform FJointSolverGaussSeidel * uniform M,
+	uniform FPBDJointSolver * uniform M,
 	const uniform FVector& DP0,
 	const uniform FVector& DP1)
 {
@@ -149,7 +149,7 @@ static inline void ApplyPositionDelta(
 }
 
 static inline void ApplyRotationDelta(
-	uniform FJointSolverGaussSeidel * uniform M,
+	uniform FPBDJointSolver * uniform M,
 	const uniform int Index,
 	const uniform FVector& InDR)
 {
@@ -175,7 +175,7 @@ static inline void ApplyRotationDelta(
 }
 
 static inline void ApplyRotationDelta(
-	uniform FJointSolverGaussSeidel * uniform M,
+	uniform FPBDJointSolver * uniform M,
 	const uniform FVector& InDR0,
 	const uniform FVector& InDR1)
 {
@@ -208,7 +208,7 @@ static inline void ApplyRotationDelta(
 }
 
 export void ApplyRotationDelta2(
-	uniform FJointSolverGaussSeidel * uniform M,
+	uniform FPBDJointSolver * uniform M,
 	const uniform FVector& DR0,
 	const uniform FVector& DR1)
 {
@@ -216,7 +216,7 @@ export void ApplyRotationDelta2(
 }
 
 static inline void ApplyDelta(
-	uniform FJointSolverGaussSeidel * uniform M,
+	uniform FPBDJointSolver * uniform M,
 	const uniform int Index,
 	const uniform FVector& DP,
 	const uniform FVector& DR)
@@ -226,7 +226,7 @@ static inline void ApplyDelta(
 }
 
 static inline void ApplyVelocityDelta(
-	uniform FJointSolverGaussSeidel * uniform M,
+	uniform FPBDJointSolver * uniform M,
 	const uniform FVector& DV0,
 	const uniform FVector& DW0,
 	const uniform FVector& DV1,
@@ -240,7 +240,7 @@ static inline void ApplyVelocityDelta(
 
 
 static inline void ApplyVelocityDelta(
-	uniform FJointSolverGaussSeidel * uniform M,
+	uniform FPBDJointSolver * uniform M,
 	const uniform int Index,
 	const uniform FVector& DV,
 	const uniform FVector& DW)
@@ -250,7 +250,7 @@ static inline void ApplyVelocityDelta(
 }
 
 export void ApplyPositionConstraintSoft(
-	uniform FJointSolverGaussSeidel * uniform M,
+	uniform FPBDJointSolver * uniform M,
 	const uniform float Dt,
 	const uniform float JointStiffness,
 	const uniform float JointDamping,
@@ -297,7 +297,7 @@ export void ApplyPositionConstraintSoft(
 
 // See "XPBD: Position-Based Simulation of Compliant Constrained Dynamics"
 export void ApplyRotationConstraintSoftKD(
-	uniform FJointSolverGaussSeidel * uniform M,
+	uniform FPBDJointSolver * uniform M,
 	const uniform int KIndex,
 	const uniform int DIndex,
 	const uniform float Dt,
@@ -341,7 +341,7 @@ export void ApplyRotationConstraintSoftKD(
 }
 
 export void ApplyRotationConstraintSoftDD(
-	uniform FJointSolverGaussSeidel * uniform M,
+	uniform FPBDJointSolver * uniform M,
 	const uniform float Dt,
 	const uniform float JointStiffness,
 	const uniform float JointDamping,
@@ -388,7 +388,7 @@ export void ApplyRotationConstraintSoftDD(
 
 
 export void ApplyPointPositionConstraintKD(
-	uniform FJointSolverGaussSeidel * uniform M,
+	uniform FPBDJointSolver * uniform M,
 	const uniform int KIndex,
 	const uniform int DIndex,
 	const uniform FVector &CX,
@@ -410,7 +410,7 @@ export void ApplyPointPositionConstraintKD(
 }
 
 export void ApplyPointPositionConstraintDD(
-	uniform FJointSolverGaussSeidel * uniform M,
+	uniform FPBDJointSolver * uniform M,
 	const uniform FVector &CX,
 	const uniform float InStiffness)
 {
diff --git a/Engine/Source/Runtime/Experimental/Chaos/Private/Chaos/PBDCollisionConstraintUtil.cpp b/Engine/Source/Runtime/Experimental/Chaos/Private/Chaos/PBDCollisionConstraintUtil.cpp
--- a/Engine/Source/Runtime/Experimental/Chaos/Private/Chaos/PBDCollisionConstraintUtil.cpp
+++ b/Engine/Source/Runtime/Experimental/Chaos/Private/Chaos/PBDCollisionConstraintUtil.cpp
@@ -14,7 +14,7 @@
 
 namespace Chaos
 {
-	void ComputeHashTable(const TArray<Chaos::FRigidBodyPointContactConstraint>& ConstraintsArray,
+	void ComputeHashTable(const TArray<Chaos::FPBDCollisionConstraint>& ConstraintsArray,
 						  const FBox& BoundingBox, TMultiMap<int32, int32>& HashTableMap, const FReal SpatialHashRadius)
 	{
 		FReal CellSize = 2 * SpatialHashRadius;
