Chaos

diff --git a/Engine/Source/Runtime/Experimental/Chaos/Private/Chaos/PBDCollisionConstraints.cpp b/Engine/Source/Runtime/Experimental/Chaos/Private/Chaos/PBDCollisionConstraints.cpp
--- a/Engine/Source/Runtime/Experimental/Chaos/Private/Chaos/PBDCollisionConstraints.cpp
+++ b/Engine/Source/Runtime/Experimental/Chaos/Private/Chaos/PBDCollisionConstraints.cpp
@@ -9,7 +9,9 @@
 #include "Chaos/CollisionResolutionUtil.h"
 #include "Chaos/CollisionResolution.h"
 #include "Chaos/Collision/CollisionContext.h"
+#include "Chaos/Collision/SolverCollisionContainer.h"
 #include "Chaos/Defines.h"
+#include "Chaos/Evolution/SolverBodyContainer.h"
 #include "Chaos/GeometryQueries.h"
 #include "Chaos/ImplicitObjectUnion.h"
 #include "Chaos/ImplicitObjectScaled.h"
@@ -23,12 +25,13 @@
 #include "Chaos/CastingUtilities.h"
 #include "ChaosLog.h"
 #include "ChaosStats.h"
+#include "Chaos/Evolution/SolverDatas.h"
 #include "Containers/Queue.h"
 #include "ProfilingDebugging/ScopedTimers.h"
-#include "ProfilingDebugging/CsvProfiler.h"
 #include "Algo/Sort.h"
 #include "Algo/StableSort.h"
 
+
 #if INTEL_ISPC
 #include "PBDCollisionConstraints.ispc.generated.h"
 #endif
@@ -39,11 +42,6 @@ namespace Chaos
 {
 	extern int32 UseLevelsetCollision;
 
-	namespace Collisions
-	{
-		extern int32 Chaos_Collision_UseAccumulatedImpulseClipSolve;
-	}
-
 	int32 CollisionParticlesBVHDepth = 4;
 	FAutoConsoleVariableRef CVarCollisionParticlesBVHDepth(TEXT("p.CollisionParticlesBVHDepth"), CollisionParticlesBVHDepth, TEXT("The maximum depth for collision particles bvh"));
 
@@ -83,7 +81,6 @@ namespace Chaos
 	bool CollisionsAllowParticleTracking = true;
 	FAutoConsoleVariableRef CVarCollisionsAllowParticleTracking(TEXT("p.Chaos.Collision.AllowParticleTracking"), CollisionsAllowParticleTracking, TEXT("Allow particles to track their collisions constraints when their DoBufferCollisions flag is enable [def:true]"));
 
-
 #if INTEL_ISPC
 	bool bChaos_Collision_ISPC_Enabled = false;
 	FAutoConsoleVariableRef CVarChaosCollisionISPCEnabled(TEXT("p.Chaos.Collision.ISPC"), bChaos_Collision_ISPC_Enabled, TEXT("Whether to use ISPC optimizations in the Collision Solver"));
@@ -92,8 +89,12 @@ namespace Chaos
 
 	DECLARE_CYCLE_STAT(TEXT("Collisions::Reset"), STAT_Collisions_Reset, STATGROUP_ChaosCollision);
 	DECLARE_CYCLE_STAT(TEXT("Collisions::UpdatePointConstraints"), STAT_Collisions_UpdatePointConstraints, STATGROUP_ChaosCollision);
-	DECLARE_CYCLE_STAT(TEXT("Collisions::Apply"), STAT_Collisions_Apply, STATGROUP_ChaosCollision);
-	DECLARE_CYCLE_STAT(TEXT("Collisions::ApplyPushOut"), STAT_Collisions_ApplyPushOut, STATGROUP_ChaosCollision);
+	DECLARE_CYCLE_STAT(TEXT("Collisions::BeginDetect"), STAT_Collisions_BeginDetect, STATGROUP_ChaosCollision);
+	DECLARE_CYCLE_STAT(TEXT("Collisions::EndDetect"), STAT_Collisions_EndDetect, STATGROUP_ChaosCollision);
+	DEFINE_STAT(STAT_Collisions_Gather);
+	DEFINE_STAT(STAT_Collisions_Scatter);
+	DEFINE_STAT(STAT_Collisions_Apply);
+	DEFINE_STAT(STAT_Collisions_ApplyPushOut);
 
 	//
 	// Collision Constraint Container
@@ -107,26 +108,20 @@ namespace Chaos
 		const int32 InApplyPairIterations /*= 1*/,
 		const int32 InApplyPushOutPairIterations /*= 1*/,
 		const FReal InRestitutionThreshold /*= (FReal)0*/)
-		: FPBDConstraintContainer(EConstraintContainerType::Collision)
-		, bInAppendOperation(false)
+		: FPBDConstraintContainer(FConstraintContainerHandle::StaticType())
 		, Particles(InParticles)
 		, NumActivePointConstraints(0)
-		, NumActiveSweptPointConstraints(0)
 		, MCollided(Collided)
 		, MPhysicsMaterials(InPhysicsMaterials)
 		, MPerParticlePhysicsMaterials(InPerParticlePhysicsMaterials)
 		, MApplyPairIterations(InApplyPairIterations)
 		, MApplyPushOutPairIterations(InApplyPushOutPairIterations)
 		, RestitutionThreshold(InRestitutionThreshold)	// @todo(chaos): expose as property
-		, bUseCCD(false)
 		, bEnableCollisions(true)
 		, bEnableRestitution(true)
 		, bHandlesEnabled(true)
 		, bCanDisableContacts(true)
 		, SolverType(EConstraintSolverType::GbfPbd)
-		, LifespanCounter(0)
-		, PostApplyCallback(nullptr)
-		, PostApplyPushOutCallback(nullptr)
 	{
 #if INTEL_ISPC
 		if (bRealTypeCompatibleWithISPC && bChaos_Collision_ISPC_Enabled)
@@ -136,31 +131,24 @@ namespace Chaos
 #endif
 	}
 
-	void FPBDCollisionConstraints::DisableHandles()
+	FPBDCollisionConstraints::~FPBDCollisionConstraints()
 	{
-		check(NumConstraints() == 0);
-		bHandlesEnabled = false;
 	}
 
-
-	void FPBDCollisionConstraints::SetPostApplyCallback(const FRigidBodyContactConstraintsPostApplyCallback& Callback)
+	void FPBDCollisionConstraints::DisableHandles()
 	{
-		PostApplyCallback = Callback;
+		check(NumConstraints() == 0);
+		bHandlesEnabled = false;
 	}
 
-	void FPBDCollisionConstraints::ClearPostApplyCallback()
+	FPBDCollisionConstraints::FHandles FPBDCollisionConstraints::GetConstraintHandles() const
 	{
-		PostApplyCallback = nullptr;
+		return ConstraintAllocator.GetConstraints();
 	}
 
-	void FPBDCollisionConstraints::SetPostApplyPushOutCallback(const FRigidBodyContactConstraintsPostApplyPushOutCallback& Callback)
+	FPBDCollisionConstraints::FConstHandles FPBDCollisionConstraints::GetConstConstraintHandles() const
 	{
-		PostApplyPushOutCallback = Callback;
-	}
-	
-	void FPBDCollisionConstraints::ClearPostApplyPushOutCallback()
-	{
-		PostApplyPushOutCallback = nullptr;
+		return ConstraintAllocator.GetConstConstraints();
 	}
 
 	const FChaosPhysicsMaterial* GetPhysicsMaterial(const TGeometryParticleHandle<FReal, 3>* Particle, const FImplicitObject* Geom, const TArrayCollectionArray<TSerializablePtr<FChaosPhysicsMaterial>>& PhysicsMaterials, const TArrayCollectionArray<TUniquePtr<FChaosPhysicsMaterial>>& PerParticlePhysicsMaterials)
@@ -202,7 +190,7 @@ namespace Chaos
 		return nullptr;
 	}
 
-	void FPBDCollisionConstraints::UpdateConstraintMaterialProperties(FCollisionConstraintBase& Constraint)
+	void FPBDCollisionConstraints::UpdateConstraintMaterialProperties(FPBDCollisionConstraint& Constraint)
 	{
 		const FChaosPhysicsMaterial* PhysicsMaterial0 = GetPhysicsMaterial(Constraint.Particle[0], Constraint.Manifold.Implicit[0], MPhysicsMaterials, MPerParticlePhysicsMaterials);
 		const FChaosPhysicsMaterial* PhysicsMaterial1 = GetPhysicsMaterial(Constraint.Particle[1], Constraint.Manifold.Implicit[1], MPhysicsMaterials, MPerParticlePhysicsMaterials);
@@ -240,6 +228,8 @@ namespace Chaos
 			Contact.Restitution = DefaultCollisionRestitution;
 		}
 
+		Contact.RestitutionThreshold = (CollisionRestitutionThresholdOverride >= 0.0f) ? CollisionRestitutionThresholdOverride : RestitutionThreshold;
+
 		if (!bEnableRestitution)
 		{
 			Contact.Restitution = 0.0f;
@@ -249,6 +239,7 @@ namespace Chaos
 		if (CollisionFrictionOverride >= 0)
 		{
 			Contact.Friction = CollisionFrictionOverride;
+			Contact.AngularFriction = CollisionFrictionOverride;
 		}
 		if (CollisionRestitutionOverride >= 0)
 		{
@@ -260,135 +251,96 @@ namespace Chaos
 		}
 	}
 
-	FPBDCollisionConstraints::FConstraintAppendScope FPBDCollisionConstraints::BeginAppendScope()
+	void FPBDCollisionConstraints::UpdatePositionBasedState(const FReal Dt)
 	{
-		check(!bInAppendOperation);
-		return FPBDCollisionConstraints::FConstraintAppendScope(this);
 	}
 
-	void FPBDCollisionConstraints::AddConstraint(const FRigidBodyPointContactConstraint& InConstraint)
+	void FPBDCollisionConstraints::BeginFrame()
 	{
-		check(!bInAppendOperation);
-
-		int32 Idx = Constraints.SinglePointConstraints.Add(InConstraint);
-
-		if (bHandlesEnabled)
-		{
-			FPBDCollisionConstraintHandle* Handle = HandleAllocator.template AllocHandle<FRigidBodyPointContactConstraint>(this, Idx);
-			Handle->GetContact().Timestamp = -INT_MAX; // force point constraints to be deleted.
-
-			Constraints.SinglePointConstraints[Idx].SetConstraintHandle(Handle);
-
-			check(Handle != nullptr);
-			Handles.Add(Handle);
-				
-#if CHAOS_COLLISION_PERSISTENCE_ENABLED
-			check(!Manifolds.Contains(Handle->GetKey()));
-			Manifolds.Add(Handle->GetKey(), Handle);
-#endif
-		}
+		ConstraintAllocator.BeginFrame();
 	}
 
-	void FPBDCollisionConstraints::AddConstraint(const FRigidBodySweptPointContactConstraint& InConstraint)
+	void FPBDCollisionConstraints::Reset()
 	{
-		check(!bInAppendOperation);
-
-		int32 Idx = Constraints.SinglePointSweptConstraints.Add(InConstraint);
-
-		if (bHandlesEnabled)
-		{
-			FPBDCollisionConstraintHandle* Handle = HandleAllocator.template AllocHandle<FRigidBodySweptPointContactConstraint>(this, Idx);
-			Handle->GetContact().Timestamp = -INT_MAX; // force point constraints to be deleted.
+		SCOPE_CYCLE_COUNTER(STAT_Collisions_Reset);
 
-			Constraints.SinglePointSweptConstraints[Idx].SetConstraintHandle(Handle);
+		ConstraintAllocator.Reset();
+		ParticleCollisionsMap.Reset();
+	}
 
-			if(ensure(Handle != nullptr))
-			{			
-				Handles.Add(Handle);
+	void FPBDCollisionConstraints::BeginDetectCollisions()
+	{
+		SCOPE_CYCLE_COUNTER(STAT_Collisions_BeginDetect);
 
-#if CHAOS_COLLISION_PERSISTENCE_ENABLED
-				check(!Manifolds.Contains(Handle->GetKey()));
-				Manifolds.Add(Handle->GetKey(), Handle);
-#endif
-			}
-		}
+		ConstraintAllocator.BeginDetectCollisions();
 	}
 
-	void FPBDCollisionConstraints::PrepareIteration(FReal dt)
+	void FPBDCollisionConstraints::EndDetectCollisions()
 	{
-		// NOTE: We could set material properties as we add constraints, but the ParticlePairBroadphase
-		// skips the call to AddConstraint and writes directly to the constraint array, so we
-		// need to do it after all constraints are added.
+		SCOPE_CYCLE_COUNTER(STAT_Collisions_EndDetect);
 
-		for (FRigidBodyPointContactConstraint& Contact : Constraints.SinglePointConstraints)
-		{
-			UpdateConstraintMaterialProperties(Contact);
-		}
+		// Prune the unused contacts
+		ConstraintAllocator.EndDetectCollisions();
 
-		for (FRigidBodySweptPointContactConstraint& Contact : Constraints.SinglePointSweptConstraints)
+		// Initialize container properties
+		for (FPBDCollisionConstraint* Contact : GetConstraints())
 		{
-			UpdateConstraintMaterialProperties(Contact);
-		}
-	}
+			// @todo(chaos): this could be set on creation if the allocator knew about the container
+			Contact->SetContainer(this);
+			UpdateConstraintMaterialProperties(*Contact);
 
-	void FPBDCollisionConstraints::UpdatePositionBasedState(const FReal Dt)
-	{
-		check(!bInAppendOperation);
-
-		Reset();
-	
-		LifespanCounter++;
+			// @todo(chaos): move to constraint allocator with the other maps
+			AddToParticleCollisionMap(Contact);
+		}
 	}
 
-	void FPBDCollisionConstraints::Reset()
+	void FPBDCollisionConstraints::AddToParticleCollisionMap(FPBDCollisionConstraint* Constraint)
 	{
-		check(!bInAppendOperation);
-
-		SCOPE_CYCLE_COUNTER(STAT_Collisions_Reset);
-
-#if CHAOS_COLLISION_PERSISTENCE_ENABLED
-		check(bHandlesEnabled);	// This will need fixing for handle-free mode
-		TArray<FPBDCollisionConstraintHandle*> CopyOfHandles = Handles;
-		int32 LifespanWindow = LifespanCounter - 1;
-		for (FPBDCollisionConstraintHandle* ContactHandle : CopyOfHandles)
+		// @todo(chaos): move the particle-> constraints mapping fucntionality into the collision allocator
+		if (CollisionsAllowParticleTracking)
 		{
-			if (!bEnableCollisions || ContactHandle->GetContact().Timestamp< LifespanWindow)
+			auto AddToParticleCollisionsMap = [](FPBDCollisionConstraintHandle* InHandle, FGeometryParticleHandle* InParticle, TMap<FGeometryParticleHandle*, TArray<FPBDCollisionConstraintHandle*>>& InParticleCollisionsMap)
 			{
-				RemoveConstraint(ContactHandle);
-			}
-		}
-#else
-		for (FPBDCollisionConstraintHandle* Handle : Handles)
-		{
-			HandleAllocator.FreeHandle(Handle);
+				if (InParticle)
+				{
+					if (FPBDRigidParticleHandle* RigidParticle = InParticle->CastToRigidParticle())
+					{
+						if (RigidParticle->HasCollisionConstraintFlag(ECollisionConstraintFlags::CCF_DoBufferCollisions))
+						{
+							TArray<FPBDCollisionConstraintHandle*>* HandleArray = InParticleCollisionsMap.Find(RigidParticle);
+							if (HandleArray == nullptr)
+							{
+								InParticleCollisionsMap.Add(RigidParticle, TArray<FPBDCollisionConstraintHandle*>());
+								HandleArray = InParticleCollisionsMap.Find(RigidParticle);
+							}
+							HandleArray->Add(InHandle);
+						}
+					}
+				}
+			};
+			AddToParticleCollisionsMap(Constraint, Constraint->GetConstrainedParticles()[0], ParticleCollisionsMap);
+			AddToParticleCollisionsMap(Constraint, Constraint->GetConstrainedParticles()[1], ParticleCollisionsMap);
 		}
-		Constraints.Reset();
-		Handles.Reset();
-		ParticleCollisionsMap.Reset();
-#endif
-
-		bUseCCD = false;
 	}
 
 	void FPBDCollisionConstraints::ApplyCollisionModifier(const TArray<ISimCallbackObject*>& CollisionModifiers)
 	{
-		check(!bInAppendOperation);
-		if (Handles.Num())
+		if (GetConstraints().Num() > 0)
 		{
 			for(ISimCallbackObject* Modifier : CollisionModifiers)
 			{
 				TArray<FPBDCollisionConstraintHandleModification> ModificationResults;
-				ModificationResults.Reserve(Handles.Num());
-				for (FPBDCollisionConstraintHandle* Handle : Handles)
+				ModificationResults.Reserve(GetConstraints().Num());
+				for (FPBDCollisionConstraintHandle* Handle : GetConstraintHandles())
 				{
 					ModificationResults.Emplace(Handle);
 				}
 
 				Modifier->ContactModification_Internal(TArrayView<FPBDCollisionConstraintHandleModification>(ModificationResults.GetData(), ModificationResults.Num()));
-
+				
 				for (const FPBDCollisionConstraintHandleModification& Modification : ModificationResults)
 				{
-						if (Modification.GetResult() == ECollisionModifierResult::Disabled)
+					if (Modification.GetResult() == ECollisionModifierResult::Disabled)
 					{
 						RemoveConstraint(Modification.GetHandle());
 					}
@@ -399,12 +351,10 @@ namespace Chaos
 
 	void FPBDCollisionConstraints::RemoveConstraints(const TSet<TGeometryParticleHandle<FReal, 3>*>&  InHandleSet)
 	{
-		check(!bInAppendOperation);
-
 		const TArray<TGeometryParticleHandle<FReal, 3>*> HandleArray = InHandleSet.Array();
 		for (auto ParticleHandle : HandleArray)
 		{
-			TArray<FPBDCollisionConstraintHandle*> CopyOfHandles = Handles;
+			FHandles CopyOfHandles = GetConstraintHandles();
 
 			for (FPBDCollisionConstraintHandle* ContactHandle : CopyOfHandles)
 			{
@@ -419,98 +369,7 @@ namespace Chaos
 
 	void FPBDCollisionConstraints::RemoveConstraint(FPBDCollisionConstraintHandle* Handle)
 	{
-		check(!bInAppendOperation);
-
-		FConstraintContainerHandleKey KeyToRemove = Handle->GetKey();
-		int32 Idx = Handle->GetConstraintIndex(); // index into specific array
-		typename FCollisionConstraintBase::FType ConstraintType = Handle->GetType();
-
-		if (ConstraintType == FCollisionConstraintBase::FType::SinglePoint)
-		{
-#if CHAOS_COLLISION_PERSISTENCE_ENABLED
-			if (Idx < Constraints.SinglePointConstraints.Num() - 1)
-			{
-				// update the handle
-				FConstraintContainerHandleKey Key = FPBDCollisionConstraintHandle::MakeKey(&Constraints.SinglePointConstraints.Last());
-				Manifolds[Key]->SetConstraintIndex(Idx, ConstraintType);
-			}
-#endif
-			Constraints.SinglePointConstraints.RemoveAtSwap(Idx);
-			if (bHandlesEnabled && (Idx < Constraints.SinglePointConstraints.Num()))
-			{
-				Constraints.SinglePointConstraints[Idx].GetConstraintHandle()->SetConstraintIndex(Idx, FCollisionConstraintBase::FType::SinglePoint);
-			}
-
-		}
-		else if (ConstraintType == FCollisionConstraintBase::FType::SinglePointSwept)
-		{
-#if CHAOS_COLLISION_PERSISTENCE_ENABLED
-			if (Idx < Constraints.SinglePointSweptConstraints.Num() - 1)
-			{
-				// update the handle
-				FConstraintContainerHandleKey Key = FPBDCollisionConstraintHandle::MakeKey(&Constraints.SinglePointSweptConstraints.Last());
-				Manifolds[Key]->SetConstraintIndex(Idx, ConstraintType);
-			}
-#endif
-			Constraints.SinglePointSweptConstraints.RemoveAtSwap(Idx);
-			if (bHandlesEnabled && (Idx < Constraints.SinglePointSweptConstraints.Num()))
-			{
-				Constraints.SinglePointSweptConstraints[Idx].GetConstraintHandle()->SetConstraintIndex(Idx, FCollisionConstraintBase::FType::SinglePointSwept);
-			}
-		}
-		else 
-		{
-			check(false);
-		}
-
-		if (bHandlesEnabled)
-		{
-			// @todo(chaos): Collision Manifold
-			//   Add an index to the handle in the Manifold.Value 
-			//   to prevent the search in Handles when removed.
-#if CHAOS_COLLISION_PERSISTENCE_ENABLED
-			Manifolds.Remove(KeyToRemove);
-#endif
-			Handles.Remove(Handle);
-			check(Handles.Num() == Constraints.SinglePointConstraints.Num() + Constraints.SinglePointSweptConstraints.Num());
-
-			HandleAllocator.FreeHandle(Handle);
-		}
-	}
-
-
-	void FPBDCollisionConstraints::UpdateConstraints(FReal Dt, const TSet<TGeometryParticleHandle<FReal, 3>*>& ParticlesSet)
-	{
-		// Clustering uses update constraints to force a re-evaluation. 
-	}
-
-	// Called once per frame to update persistent constraints (reruns collision detection, or selects the best manifold point)
-	void FPBDCollisionConstraints::UpdateConstraints(FReal Dt)
-	{
-		SCOPE_CYCLE_COUNTER(STAT_Collisions_UpdatePointConstraints);
-
-		// @todo(chaos): parallelism needs to be optional
-
-		//PhysicsParallelFor(Handles.Num(), [&](int32 ConstraintHandleIndex)
-		//{
-		//	FPBDCollisionConstraintHandle* ConstraintHandle = Handles[ConstraintHandleIndex];
-		//	check(ConstraintHandle != nullptr);
-		//	Collisions::Update(MCullDistance, MShapePadding, ConstraintHandle->GetContact());
-
-		//	if (ConstraintHandle->GetContact().GetPhi() < MCullDistance) 
-		//	{
-		//		ConstraintHandle->GetContact().Timestamp = LifespanCounter;
-		//	}
-		//}, bDisableCollisionParallelFor);
-
-		for (FRigidBodyPointContactConstraint& Contact : Constraints.SinglePointConstraints)
-		{
-			Collisions::Update(Contact, Dt);
-			if (Contact.GetPhi() < Contact.GetCullDistance())
-			{
-				Contact.Timestamp = LifespanCounter;
-			}
-		}
+		ConstraintAllocator.DestroyConstraint(&Handle->GetContact());
 	}
 
 	Collisions::FContactParticleParameters FPBDCollisionConstraints::GetContactParticleParameters(const FReal Dt)
@@ -535,296 +394,309 @@ namespace Chaos
 		};
 	}
 
-	bool FPBDCollisionConstraints::Apply(const FReal Dt, const int32 Iterations, const int32 NumIterations)
+	void FPBDCollisionConstraints::SetNumIslandConstraints(const int32 NumIslandConstraints, FPBDIslandSolverData& SolverData)
 	{
-		SCOPE_CYCLE_COUNTER(STAT_Collisions_Apply);
-		CSV_SCOPED_TIMING_STAT(Chaos, ApplyCollisionConstraints);
-
-		bool bNeedsAnotherIteration = false;
-		if (MApplyPairIterations > 0)
+		if (SolverType == EConstraintSolverType::QuasiPbd)
 		{
-			const Collisions::FContactParticleParameters ParticleParameters = GetContactParticleParameters(Dt);
-			const Collisions::FContactIterationParameters IterationParameters = GetContactIterationParameters(Dt, Iterations, NumIterations, MApplyPairIterations, bNeedsAnotherIteration);
-
-			NumActivePointConstraints = 0;
-			for (FRigidBodyPointContactConstraint& Contact : Constraints.SinglePointConstraints)
-			{
-				if (!Contact.GetDisabled())
-				{
-					Collisions::Apply(Contact, IterationParameters, ParticleParameters);
-					++NumActivePointConstraints;
-				}
-			}
-
-			// Swept apply may significantly change particle position, invalidating other constraint's manifolds.
-			// We don't update manifolds on first apply iteration, so make sure we apply swept constraints last.
-			NumActiveSweptPointConstraints = 0;
-			for (FRigidBodySweptPointContactConstraint& Contact : Constraints.SinglePointSweptConstraints)
-			{
-				if (!Contact.GetDisabled())
-				{
-					Collisions::Apply(Contact, IterationParameters, ParticleParameters);
-					++NumActiveSweptPointConstraints;
-				}
-			}
+			GetConstraintSolverContainer(SolverData).Reset(NumIslandConstraints);
 		}
-
-		if (PostApplyCallback != nullptr)
+		else
 		{
-			PostApplyCallback(Dt, Handles);
+			SolverData.GetConstraintHandles(ContainerId).Reset(NumIslandConstraints);
 		}
-
-		return bNeedsAnotherIteration;
 	}
 
-	bool FPBDCollisionConstraints::ApplyPushOut(const FReal Dt, const int32 Iterations, const int32 NumIterations)
+	FPBDCollisionSolverContainer& FPBDCollisionConstraints::GetConstraintSolverContainer(FPBDIslandSolverData& SolverData)
 	{
-		SCOPE_CYCLE_COUNTER(STAT_Collisions_ApplyPushOut);
+		check(SolverType == EConstraintSolverType::QuasiPbd);
+		return SolverData.GetConstraintContainer<FPBDCollisionSolverContainer>(ContainerId);
+	}
 
-		TSet<const TGeometryParticleHandle<FReal, 3>*> TempStatic;
-		bool bNeedsAnotherIteration = false;
-		if (MApplyPushOutPairIterations > 0)
+	void FPBDCollisionConstraints::GatherInput(const FReal Dt, FPBDCollisionConstraint& Constraint, const int32 Particle0Level, const int32 Particle1Level, FPBDIslandSolverData& SolverData)
+	{
+		if (SolverType == EConstraintSolverType::QuasiPbd)
 		{
-			const Collisions::FContactParticleParameters ParticleParameters = GetContactParticleParameters(Dt);
-			const Collisions::FContactIterationParameters IterationParameters = GetContactIterationParameters(Dt, Iterations, NumIterations, MApplyPushOutPairIterations, bNeedsAnotherIteration);
+			// We shouldn't be adding disabled constraints to the solver list. The check needs to be at caller site or we should return success/fail - see TPBDConstraintColorRule::GatherSolverInput
+			check(Constraint.IsEnabled());
+			GetConstraintSolverContainer(SolverData).AddConstraintSolver(Dt, Constraint, Particle0Level, Particle1Level, SolverData.GetBodyContainer());
+		}
+		else
+		{
+			LegacyGatherInput(Dt, Constraint, Particle0Level, Particle1Level, SolverData);
+		}
+	}
 
-			for (FRigidBodyPointContactConstraint& Contact : Constraints.SinglePointConstraints)
+	void FPBDCollisionConstraints::GatherInput(const FReal Dt, FPBDIslandSolverData& SolverData)
+	{
+		if (SolverType == EConstraintSolverType::QuasiPbd)
+		{
+			for (FPBDCollisionConstraint* Constraint : GetConstraints())
 			{
-				if (!Contact.GetDisabled())
+				if (Constraint->IsEnabled())
 				{
-					Collisions::ApplyPushOut(Contact, TempStatic, IterationParameters, ParticleParameters);
+					GatherInput(Dt, *Constraint, INDEX_NONE, INDEX_NONE, SolverData);
 				}
 			}
-
-			for (FRigidBodySweptPointContactConstraint& Contact : Constraints.SinglePointSweptConstraints)
+		}
+		else
+		{
+			for (FPBDCollisionConstraint* Constraint : GetConstraints())
 			{
-				if (!Contact.GetDisabled())
+				if (Constraint->IsEnabled())
 				{
-					Collisions::ApplyPushOut(Contact, TempStatic, IterationParameters, ParticleParameters);
+					LegacyGatherInput(Dt, *Constraint, INDEX_NONE, INDEX_NONE, SolverData);
 				}
 			}
 		}
+	}
 
-		if (PostApplyPushOutCallback != nullptr)
+	void FPBDCollisionConstraints::ScatterOutput(const FReal Dt, const int32 BeginIndex, const int32 EndIndex, FPBDIslandSolverData& SolverData)
+	{
+		if (SolverType == EConstraintSolverType::QuasiPbd)
 		{
-			PostApplyPushOutCallback(Dt, Handles, bNeedsAnotherIteration);
+			GetConstraintSolverContainer(SolverData).ScatterOutput(Dt, BeginIndex, EndIndex);
+		}
+		else
+		{
+			LegacyScatterOutput(Dt, BeginIndex, EndIndex, SolverData);
 		}
-
-		return bNeedsAnotherIteration;
 	}
 
-	void FPBDCollisionConstraints::SortConstraints()
+	void FPBDCollisionConstraints::ScatterOutput(const FReal Dt, FPBDIslandSolverData& SolverData)
 	{
-		check(!bInAppendOperation);
+		if (SolverType == EConstraintSolverType::QuasiPbd)
+		{
+			FPBDCollisionSolverContainer& SolverContainer = GetConstraintSolverContainer(SolverData);
+			SolverContainer.ScatterOutput(Dt, 0, SolverContainer.NumSolvers());
+		}
+		else
+		{
+			LegacyScatterOutput(Dt, 0, SolverData.GetConstraintHandles(ContainerId).Num(), SolverData);
+		}
+	}
 
-		//We have to use StableSort so that constraints of the same pair stay in the same order
-		//Otherwise the order within each pair can change due to where they start out in the array
-		//TODO: we should label each contact (and shape) for things like warm starting GJK
-		//and so we could use that label as part of the key
-		//and then we could use regular Sort (which is faster)
-		Algo::StableSort(Handles, [](const FPBDCollisionConstraintHandle* A, const FPBDCollisionConstraintHandle* B)
+	void FPBDCollisionConstraints::ApplySwept(const FReal Dt, FPBDIslandSolverData& SolverData)
+	{
+		if (SolverType == EConstraintSolverType::QuasiPbd)
 		{
-			if(A->GetType() == B->GetType())
-			{
-				return A->GetContact() < B->GetContact();
-			}
-			else
-			{
-				return A->GetType() < B->GetType();
-			}
-		});
+			FPBDCollisionSolverContainer& SolverContainer = GetConstraintSolverContainer(SolverData);
+			SolverContainer.SolveSwept(Dt);
+		}
 	}
 
-	bool FPBDCollisionConstraints::Apply(const FReal Dt, const TArray<FPBDCollisionConstraintHandle*>& InConstraintHandles, const int32 Iterations, const int32 NumIterations)
+	// Simple Rule version
+	bool FPBDCollisionConstraints::ApplyPhase1(const FReal Dt, const int32 It, const int32 NumIts, FPBDIslandSolverData& SolverData)
 	{
 		SCOPE_CYCLE_COUNTER(STAT_Collisions_Apply);
-		CSV_SCOPED_TIMING_STAT(Chaos, ApplyCollisionConstraints);
 
-		TAtomic<bool> bNeedsAnotherIterationAtomic;
-		bNeedsAnotherIterationAtomic.Store(false);
-		if (MApplyPairIterations > 0)
+		if (SolverType == EConstraintSolverType::QuasiPbd)
 		{
-			PhysicsParallelFor(InConstraintHandles.Num(), [&](int32 ConstraintHandleIndex) 
-			{
-				FPBDCollisionConstraintHandle* ConstraintHandle = InConstraintHandles[ConstraintHandleIndex];
-				check(ConstraintHandle != nullptr);
+			FPBDCollisionSolverContainer& SolverContainer = GetConstraintSolverContainer(SolverData);
+			return SolverContainer.SolvePositionSerial(Dt, It, NumIts, 0, SolverContainer.NumSolvers());
+		}
+		else
+		{
+			return LegacyApplyPhase1Serial(Dt, It, NumIts, 0, SolverData.GetConstraintHandles(ContainerId).Num(), SolverData);
+		}
+	}
 
-				TVector<const TGeometryParticleHandle<FReal, 3>*, 2> ConstrainedParticles = ConstraintHandle->GetConstrainedParticles();
-				bool bNeedsAnotherIteration = false;
+	// Island Rule version
+	bool FPBDCollisionConstraints::ApplyPhase1Serial(const FReal Dt, const int32 It, const int32 NumIts, FPBDIslandSolverData& SolverData)
+	{
+		SCOPE_CYCLE_COUNTER(STAT_Collisions_Apply);
 
-				if (!ConstraintHandle->GetContact().GetDisabled())
-				{
-					const Collisions::FContactParticleParameters ParticleParameters = GetContactParticleParameters(Dt);
-					const Collisions::FContactIterationParameters IterationParameters = GetContactIterationParameters(Dt, Iterations, NumIterations, MApplyPairIterations, bNeedsAnotherIteration);
-					Collisions::Apply(ConstraintHandle->GetContact(), IterationParameters, ParticleParameters);
+		if (SolverType == EConstraintSolverType::QuasiPbd)
+		{
+			FPBDCollisionSolverContainer& SolverContainer = GetConstraintSolverContainer(SolverData);
+			return SolverContainer.SolvePositionSerial(Dt, It, NumIts, 0, SolverContainer.NumSolvers());
+		}
+		else
+		{
+			return LegacyApplyPhase1Serial(Dt, It, NumIts, 0, SolverData.GetConstraintHandles(ContainerId).Num(), SolverData);
+		}
+	}
 
-					if (bNeedsAnotherIteration)
-					{
-						bNeedsAnotherIterationAtomic.Store(true);
-					}
-				}
+	// Color Rule version
+	bool FPBDCollisionConstraints::ApplyPhase1Serial(const FReal Dt, const int32 It, const int32 NumIts, const int32 BeginIndex, const int32 EndIndex, FPBDIslandSolverData& SolverData)
+	{
+		SCOPE_CYCLE_COUNTER(STAT_Collisions_Apply);
 
-			}, bDisableCollisionParallelFor);
+		if (SolverType == EConstraintSolverType::QuasiPbd)
+		{
+			FPBDCollisionSolverContainer& SolverContainer = GetConstraintSolverContainer(SolverData);
+			return SolverContainer.SolvePositionSerial(Dt, It, NumIts, BeginIndex, EndIndex);
 		}
-
-		if (PostApplyCallback != nullptr)
+		else
 		{
-			PostApplyCallback(Dt, InConstraintHandles);
+			return LegacyApplyPhase1Serial(Dt, It, NumIts, BeginIndex, EndIndex, SolverData);
 		}
-
-		return bNeedsAnotherIterationAtomic.Load();
 	}
 
+	// Color Rule version
+	bool FPBDCollisionConstraints::ApplyPhase1Parallel(const FReal Dt, const int32 It, const int32 NumIts, const int32 BeginIndex, const int32 EndIndex, FPBDIslandSolverData& SolverData)
+	{
+		SCOPE_CYCLE_COUNTER(STAT_Collisions_Apply);
 
-	bool FPBDCollisionConstraints::ApplyPushOut(const FReal Dt, const TArray<FPBDCollisionConstraintHandle*>& InConstraintHandles, const TSet< const TGeometryParticleHandle<FReal, 3>*>& IsTemporarilyStatic, int32 Iteration, int32 NumIterations)
+		if (SolverType == EConstraintSolverType::QuasiPbd)
+		{
+			FPBDCollisionSolverContainer& SolverContainer = GetConstraintSolverContainer(SolverData);
+			return SolverContainer.SolvePositionParallel(Dt, It, NumIts, BeginIndex, EndIndex);
+		}
+		else
+		{
+			return LegacyApplyPhase1Parallel(Dt, It, NumIts, BeginIndex, EndIndex, SolverData);
+		}
+	}
+
+	// Simple Rule version
+	bool FPBDCollisionConstraints::ApplyPhase2(const FReal Dt, const int32 It, const int32 NumIts, FPBDIslandSolverData& SolverData)
 	{
 		SCOPE_CYCLE_COUNTER(STAT_Collisions_ApplyPushOut);
 
-		bool bNeedsAnotherIteration = false;
-		if (MApplyPushOutPairIterations > 0)
+		if (SolverType == EConstraintSolverType::QuasiPbd)
 		{
-			PhysicsParallelFor(InConstraintHandles.Num(), [&](int32 ConstraintHandleIndex)
-			{
-				FPBDCollisionConstraintHandle* ConstraintHandle = InConstraintHandles[ConstraintHandleIndex];
-				check(ConstraintHandle != nullptr);
+			FPBDCollisionSolverContainer& SolverContainer = GetConstraintSolverContainer(SolverData);
+			return SolverContainer.SolveVelocitySerial(Dt, It, NumIts, 0, SolverContainer.NumSolvers());
+		}
+		else
+		{
+			return LegacyApplyPhase2Serial(Dt, It, NumIts, 0, SolverData.GetConstraintHandles(ContainerId).Num(), SolverData);
+		}
+	}
 
-				if (!ConstraintHandle->GetContact().GetDisabled())
-				{
-					const Collisions::FContactParticleParameters ParticleParameters = GetContactParticleParameters(Dt);
-					const Collisions::FContactIterationParameters IterationParameters = GetContactIterationParameters(Dt, Iteration, NumIterations, MApplyPushOutPairIterations, bNeedsAnotherIteration);
-					Collisions::ApplyPushOut(ConstraintHandle->GetContact(), IsTemporarilyStatic, IterationParameters, ParticleParameters);
-				}
+	// Island Rule version
+	bool FPBDCollisionConstraints::ApplyPhase2Serial(const FReal Dt, const int32 It, const int32 NumIts, FPBDIslandSolverData& SolverData)
+	{
+		SCOPE_CYCLE_COUNTER(STAT_Collisions_ApplyPushOut);
 
-			}, bDisableCollisionParallelFor);
+		if (SolverType == EConstraintSolverType::QuasiPbd)
+		{
+			FPBDCollisionSolverContainer& SolverContainer = GetConstraintSolverContainer(SolverData);
+			return SolverContainer.SolveVelocitySerial(Dt, It, NumIts, 0, SolverContainer.NumSolvers());
 		}
-
-		if (PostApplyPushOutCallback != nullptr)
+		else
 		{
-			PostApplyPushOutCallback(Dt, InConstraintHandles, bNeedsAnotherIteration);
+			return LegacyApplyPhase2Serial(Dt, It, NumIts, 0, SolverData.GetConstraintHandles(ContainerId).Num(), SolverData);
 		}
-
-		return bNeedsAnotherIteration;
 	}
 
-	const FCollisionConstraintBase& FPBDCollisionConstraints::GetConstraint(int32 Index) const
+	// Color Rule version
+	bool FPBDCollisionConstraints::ApplyPhase2Serial(const FReal Dt, const int32 It, const int32 NumIts, const int32 BeginIndex, const int32 EndIndex, FPBDIslandSolverData& SolverData)
 	{
-		check(Index < NumConstraints());
-		
-		if (Index < Constraints.SinglePointConstraints.Num())
+		SCOPE_CYCLE_COUNTER(STAT_Collisions_ApplyPushOut);
+
+		if (SolverType == EConstraintSolverType::QuasiPbd)
 		{
-			return Constraints.SinglePointConstraints[Index];
+			FPBDCollisionSolverContainer& SolverContainer = GetConstraintSolverContainer(SolverData);
+			return SolverContainer.SolveVelocitySerial(Dt, It, NumIts, BeginIndex, EndIndex);
+		}
+		else
+		{
+			return LegacyApplyPhase2Serial(Dt, It, NumIts, BeginIndex, EndIndex, SolverData);
 		}
-		Index -= Constraints.SinglePointConstraints.Num();
-
-		return Constraints.SinglePointSweptConstraints[Index];
 	}
 
-	FPBDCollisionConstraints::FConstraintAppendScope::FConstraintAppendScope(FPBDCollisionConstraints* InOwner)
-		: Owner(InOwner)
+	// Color Rule version
+	bool FPBDCollisionConstraints::ApplyPhase2Parallel(const FReal Dt,  const int32 It, const int32 NumIts, const int32 BeginIndex, const int32 EndIndex, FPBDIslandSolverData& SolverData)
 	{
-		check(Owner);
-		Owner->bInAppendOperation = true;
-		Constraints = &(Owner->Constraints);
+		SCOPE_CYCLE_COUNTER(STAT_Collisions_ApplyPushOut);
 
-		NumBeginSingle = Owner->Constraints.SinglePointConstraints.Num();
-		NumBeginSingleSwept = Owner->Constraints.SinglePointSweptConstraints.Num();
+		if (SolverType == EConstraintSolverType::QuasiPbd)
+		{
+			FPBDCollisionSolverContainer& SolverContainer = GetConstraintSolverContainer(SolverData);
+			return SolverContainer.SolveVelocityParallel(Dt, It, NumIts, BeginIndex, EndIndex);
+		}
+		else
+		{
+			return LegacyApplyPhase2Parallel(Dt, It, NumIts, BeginIndex, EndIndex, SolverData);
+		}
 	}
 
-	FPBDCollisionConstraints::FConstraintAppendScope::~FConstraintAppendScope()
+	void FPBDCollisionConstraints::LegacyGatherInput(const FReal Dt, FPBDCollisionConstraint& Constraint, const int32 Particle0Level, const int32 Particle1Level, FPBDIslandSolverData& SolverData)
 	{
-		FPBDCollisionConstraints::FConstraintHandleAllocator& HandleAlloc = Owner->HandleAllocator;
-		int32 HandlesBeginIndex = Owner->Handles.Num();
-		const int32 TotalAdded = NumAddedSingle + NumAddedSingleSwept;
+		SolverData.GetConstraintHandles(ContainerId).Add(&Constraint);
 
-		Owner->Handles.AddUninitialized(TotalAdded);
-		const int32 NumHandles = Owner->Handles.Num();
+		FSolverBody* SolverBody0 = SolverData.GetBodyContainer().FindOrAdd(Constraint.Particle[0]);
+		FSolverBody* SolverBody1 = SolverData.GetBodyContainer().FindOrAdd(Constraint.Particle[1]);
 
-		for(int32 HandleIndex = 0; HandleIndex < NumAddedSingle ; ++HandleIndex)
-		{
-			FPBDCollisionConstraintHandle* NewHandle = HandleAlloc.template AllocHandle<FRigidBodyPointContactConstraint>(Owner, NumBeginSingle + HandleIndex);
-			
-			const int32 FullHandleIndex = HandlesBeginIndex + HandleIndex;
-			Owner->Handles[FullHandleIndex] = NewHandle;
-
-			NewHandle->GetContact().Timestamp = -INT_MAX;
-			Constraints->SinglePointConstraints[NumBeginSingle + HandleIndex].SetConstraintHandle(NewHandle);
+		SolverBody0->SetLevel(Particle0Level);
+		SolverBody1->SetLevel(Particle1Level);
 
-			if (CollisionsAllowParticleTracking)
-			{
-				// Add reverse mappings for particles that require it. 
-				using FHandle = FPBDCollisionConstraintHandle;
-				using FParticle = TPBDRigidParticleHandleImp<FReal, 3, true>;
+		Constraint.SetSolverBodies(SolverBody0, SolverBody1);
+	}
 
-				auto AddToParticleCollisionsMap = [](FHandle* InHandle, FGeometryParticleHandle* InParticle, TMap<FParticleHandle*, FHandles >& InParticleCollisionsMap)
-				{
-					if (InParticle)
-					{
-						if (FParticle* RigidParticle = InParticle->CastToRigidParticle())
-						{
-							if (RigidParticle->HasCollisionConstraintFlag(ECollisionConstraintFlags::CCF_DoBufferCollisions))
-							{
-								FHandles* HandleArray = InParticleCollisionsMap.Find(RigidParticle);
-								if (!HandleArray)
-								{
-									InParticleCollisionsMap.Add(RigidParticle, FHandles());
-									HandleArray = InParticleCollisionsMap.Find(RigidParticle);
-								}
-								HandleArray->Add(InHandle);
-							}
-						}
-					}
-				};
-				AddToParticleCollisionsMap(NewHandle, NewHandle->GetConstrainedParticles()[0], Owner->GetParticleCollisionsMap());
-				AddToParticleCollisionsMap(NewHandle, NewHandle->GetConstrainedParticles()[1], Owner->GetParticleCollisionsMap());
-			}
+	void FPBDCollisionConstraints::LegacyScatterOutput(const FReal Dt, const int32 BeginIndex, const int32 EndIndex, FPBDIslandSolverData& SolverData)
+	{
+		for (int32 Index = BeginIndex; Index < EndIndex; ++Index)
+		{
+			FPBDCollisionConstraint* Constraint = SolverData.GetConstraintHandle<FPBDCollisionConstraint>(ContainerId,Index);
+			Constraint->SetSolverBodies(nullptr, nullptr);
 		}
-		HandlesBeginIndex += NumAddedSingle;
+	}
 
-		for(int32 HandleIndex = 0; HandleIndex < NumAddedSingleSwept; ++HandleIndex)
+	bool FPBDCollisionConstraints::LegacyApplyPhase1Serial(const FReal Dt, const int32 Iterations, const int32 NumIterations, const int32 BeginIndex, const int32 EndIndex, FPBDIslandSolverData& SolverData)
+	{
+		bool bNeedsAnotherIteration = false;
+		if (MApplyPairIterations > 0)
 		{
-			FPBDCollisionConstraintHandle* NewHandle = HandleAlloc.template AllocHandle<FRigidBodySweptPointContactConstraint>(Owner, NumBeginSingleSwept + HandleIndex);
-
-			const int32 FullHandleIndex = HandlesBeginIndex + HandleIndex;
-			Owner->Handles[FullHandleIndex] = NewHandle;
+			NumActivePointConstraints = 0;
+			const Collisions::FContactParticleParameters ParticleParameters = GetContactParticleParameters(Dt);
+			const Collisions::FContactIterationParameters IterationParameters = GetContactIterationParameters(Dt, Iterations, NumIterations, MApplyPairIterations, bNeedsAnotherIteration);
 
-			NewHandle->GetContact().Timestamp = -INT_MAX;
-			Constraints->SinglePointSweptConstraints[NumBeginSingleSwept + HandleIndex].SetConstraintHandle(NewHandle);
+			for (int32 Index = BeginIndex; Index < EndIndex; ++Index)
+			{
+				FPBDCollisionConstraint* Constraint = SolverData.GetConstraintHandle<FPBDCollisionConstraint>(ContainerId,Index);
+				if (!Constraint->GetDisabled())
+				{
+					Collisions::Apply(*Constraint, IterationParameters, ParticleParameters);
+					++NumActivePointConstraints;
+				}
+			}
 		}
-		HandlesBeginIndex += NumAddedSingle;
-
-		Owner->bInAppendOperation = false;
+		return bNeedsAnotherIteration;
 	}
 
-	void FPBDCollisionConstraints::FConstraintAppendScope::ReserveSingle(int32 NumToAdd)
+	bool FPBDCollisionConstraints::LegacyApplyPhase1Parallel(const FReal Dt, const int32 Iterations, const int32 NumIterations, const int32 BeginIndex, const int32 EndIndex, FPBDIslandSolverData& SolverData)
 	{
-		Constraints->SinglePointConstraints.Reserve(Constraints->SinglePointConstraints.Num() + NumToAdd);
+		return LegacyApplyPhase1Serial(Dt, Iterations, NumIterations, BeginIndex, EndIndex, SolverData);
 	}
 
-	void FPBDCollisionConstraints::FConstraintAppendScope::ReserveSingleSwept(int32 NumToAdd)
+	bool FPBDCollisionConstraints::LegacyApplyPhase2Serial(const FReal Dt, const int32 Iterations, const int32 NumIterations, const int32 BeginIndex, const int32 EndIndex, FPBDIslandSolverData& SolverData)
 	{
-		Constraints->SinglePointSweptConstraints.Reserve(Constraints->SinglePointConstraints.Num() + NumToAdd);
+		bool bNeedsAnotherIteration = false;
+		if (MApplyPairIterations > 0)
+		{
+			const Collisions::FContactParticleParameters ParticleParameters = GetContactParticleParameters(Dt);
+			const Collisions::FContactIterationParameters IterationParameters = GetContactIterationParameters(Dt, Iterations, NumIterations, MApplyPushOutPairIterations, bNeedsAnotherIteration);
+
+			for (int32 Index = BeginIndex; Index < EndIndex; ++Index)
+			{
+				FPBDCollisionConstraint* Constraint = SolverData.GetConstraintHandle<FPBDCollisionConstraint>(ContainerId,Index);
+				if (!Constraint->GetDisabled())
+				{
+					Collisions::ApplyPushOut(*Constraint, IterationParameters, ParticleParameters);
+				}
+			}
+		}
+		return bNeedsAnotherIteration;
 	}
 
-	void FPBDCollisionConstraints::FConstraintAppendScope::Append(TArray<FRigidBodyPointContactConstraint>&& InConstraints)
+	bool FPBDCollisionConstraints::LegacyApplyPhase2Parallel(const FReal Dt, const int32 Iterations, const int32 NumIterations, const int32 BeginIndex, const int32 EndIndex, FPBDIslandSolverData& SolverData)
 	{
-		if(InConstraints.Num() == 0)
-		{
-			return;
-		}
+		return LegacyApplyPhase2Serial(Dt,  Iterations, NumIterations, BeginIndex, EndIndex, SolverData);
+	}
 
-		NumAddedSingle += InConstraints.Num();
-		Constraints->SinglePointConstraints.Append(MoveTemp(InConstraints));
+	const FPBDCollisionConstraint& FPBDCollisionConstraints::GetConstraint(int32 Index) const
+	{
+		check(Index < NumConstraints());
+		
+		return *GetConstraints()[Index];
 	}
 
-	void FPBDCollisionConstraints::FConstraintAppendScope::Append(TArray<FRigidBodySweptPointContactConstraint>&& InConstraints)
+	FPBDCollisionConstraint& FPBDCollisionConstraints::GetConstraint(int32 Index)
 	{
-		if(InConstraints.Num() == 0)
-		{
-			return;
-		}
+		check(Index < NumConstraints());
 
-		NumAddedSingleSwept += InConstraints.Num();
-		Constraints->SinglePointSweptConstraints.Append(MoveTemp(InConstraints));
+		return *GetConstraints()[Index];
 	}
-
 }
diff --git a/Engine/Source/Runtime/Experimental/Chaos/Private/Chaos/PBDCollisionConstraintsContact.cpp b/Engine/Source/Runtime/Experimental/Chaos/Private/Chaos/PBDCollisionConstraintsContact.cpp
--- a/Engine/Source/Runtime/Experimental/Chaos/Private/Chaos/PBDCollisionConstraintsContact.cpp
+++ b/Engine/Source/Runtime/Experimental/Chaos/Private/Chaos/PBDCollisionConstraintsContact.cpp
@@ -21,7 +21,7 @@ namespace Chaos
 	extern bool bChaos_Collision_ISPC_Enabled;
 #endif
 
-	namespace Collisions
+	namespace CVars
 	{
 		int32 Chaos_Collision_EnergyClampEnabled = 1;
 		FAutoConsoleVariableRef CVarChaosCollisionEnergyClampEnabled(TEXT("p.Chaos.Collision.EnergyClampEnabled"), Chaos_Collision_EnergyClampEnabled, TEXT("Whether to use energy clamping in collision apply step"));
@@ -33,7 +33,7 @@ namespace Chaos
 		FAutoConsoleVariableRef CVarChaosCollisionPrevVelocityRestitutionEnabled(TEXT("p.Chaos.Collision.PrevVelocityRestitutionEnabled"), Chaos_Collision_PrevVelocityRestitutionEnabled, TEXT("If enabled restitution will be calculated on previous frame velocities instead of current frame velocities"));
 
 		FRealSingle Chaos_Collision_ContactMovementAllowance = 0.05f;
-		FAutoConsoleVariableRef CVarChaosCollisionContactMovementAllowance(TEXT("p.Chaos.Collision.AntiJitterContactMovementAllowance"), Chaos_Collision_ContactMovementAllowance, 
+		FAutoConsoleVariableRef CVarChaosCollisionContactMovementAllowance(TEXT("p.Chaos.Collision.AntiJitterContactMovementAllowance"), Chaos_Collision_ContactMovementAllowance,
 			TEXT("If a contact is close to where it was during a previous iteration, we will assume it is the same contact that moved (to reduce jitter). Expressed as the fraction of movement distance and Centre of Mass distance to the contact point"));
 
 		int32 Chaos_Collision_UseAccumulatedImpulseClipSolve = 0; // This requires multiple contact points per iteration per pair and contact points that don't move too much (in body space) to have an effect
@@ -47,50 +47,55 @@ namespace Chaos
 
 		bool Chaos_Collision_CheckManifoldComplete = false;
 		FAutoConsoleVariableRef CVarChaosCollisionCheckManifoldComplete(TEXT("p.Chaos.Collision.CheckManifoldComplete"), Chaos_Collision_CheckManifoldComplete, TEXT(""));
+	}
+	using namespace CVars;
 
-		void Update(FRigidBodyPointContactConstraint& Constraint, const FReal Dt)
+	namespace Collisions
+	{
+		void Update(FPBDCollisionConstraint& Constraint, const FReal Dt)
 		{
-			const FRigidTransform3 Transform0 = GetTransform(Constraint.Particle[0]);
-			const FRigidTransform3 Transform1 = GetTransform(Constraint.Particle[1]);
+			// NOTE: These are actor transforms, not CoM transforms
+			// \todo(chaos): see if we can easily switch to CoM transforms now in collision loop (shapes are held in actor space)
+			const FSolverBody& Body0 = *Constraint.GetSolverBody0();
+			const FSolverBody& Body1 = *Constraint.GetSolverBody1();
+			const FRigidTransform3 Transform0 = FRigidTransform3(Body0.ActorP(), Body0.ActorQ());
+			const FRigidTransform3 Transform1 = FRigidTransform3(Body1.ActorP(), Body1.ActorQ());
 
 			Constraint.ResetPhi(Constraint.GetCullDistance());
-			UpdateConstraintFromGeometry<ECollisionUpdateType::Deepest, FRigidBodyPointContactConstraint>(Constraint, Transform0, Transform1, Dt);
+			UpdateConstraintFromGeometry<ECollisionUpdateType::Deepest>(Constraint, Transform0, Transform1, Dt);
 		}
 
-		void UpdateSwept(FRigidBodySweptPointContactConstraint& Constraint, const FReal Dt)
+		void UpdateSwept(FPBDCollisionConstraint& Constraint, const FReal Dt)
 		{
-			FGenericParticleHandle Particle0 = FGenericParticleHandle(Constraint.Particle[0]);
 			// Note: This is unusual but we are using a mix of the previous and current transform
-			// This is due to how CCD only rewinds the position
-			const FRigidTransform3 TransformXQ0(Particle0->X(), Particle0->Q());
-			const FRigidTransform3 Transform1 = GetTransform(Constraint.Particle[1]);
+			// This is due to how CCD rewinds the position (not rotation) and then sweeps to find the first contact at the current orientation
+			// NOTE: These are actor transforms, not CoM transforms
+			// \todo(chaos): see if we can easily switch to CoM transforms now in collision loop (shapes are held in actor space)
+			const FSolverBody& Body0 = *Constraint.GetSolverBody0();
+			const FSolverBody& Body1 = *Constraint.GetSolverBody1();
+			FGenericParticleHandle Particle0 = FGenericParticleHandle(Constraint.Particle[0]);
+			const FRigidTransform3 TransformXQ0(Body0.X() - Body0.ActorQ().RotateVector(Particle0->CenterOfMass()), Body0.ActorQ());
+			const FRigidTransform3 Transform1 = FRigidTransform3(Body1.ActorP(), Body1.ActorQ());
+
 			Constraint.ResetPhi(Constraint.GetCullDistance());
-			
-			// Update as a point constraint (base class).
-			UpdateConstraintFromGeometry<ECollisionUpdateType::Deepest, FRigidBodySweptPointContactConstraint>(Constraint, TransformXQ0, Transform1, Dt);
+			UpdateConstraintFromGeometrySwept<ECollisionUpdateType::Deepest>(Constraint, TransformXQ0, Transform1, Dt);
 		}
 
 
 		// A PBD collision penetration correction.
 		// Currently only used by RBAN
-		FVec3 ApplyContact2(FCollisionContact& Contact,
-			FGenericParticleHandle Particle0,
-			FGenericParticleHandle Particle1,
+		FVec3 ApplyContact2(FPBDCollisionConstraint& Constraint,
 			const FContactIterationParameters& IterationParameters,
 			const FContactParticleParameters& ParticleParameters)
 		{
 			FVec3 AccumulatedImpulse(0);
-			TPBDRigidParticleHandle<FReal, 3>* PBDRigid0 = Particle0->CastToRigidParticle();
-			TPBDRigidParticleHandle<FReal, 3>* PBDRigid1 = Particle1->CastToRigidParticle();
-			bool bIsRigidDynamic0 = PBDRigid0 && PBDRigid0->ObjectState() == EObjectStateType::Dynamic;
-			bool bIsRigidDynamic1 = PBDRigid1 && PBDRigid1->ObjectState() == EObjectStateType::Dynamic;
-
-			FVec3 P0 = FParticleUtilities::GetCoMWorldPosition(Particle0);
-			FVec3 P1 = FParticleUtilities::GetCoMWorldPosition(Particle1);
-			FRotation3 Q0 = FParticleUtilities::GetCoMWorldRotation(Particle0);
-			FRotation3 Q1 = FParticleUtilities::GetCoMWorldRotation(Particle1);
-			FVec3 VectorToPoint0 = Contact.Location - P0;
-			FVec3 VectorToPoint1 = Contact.Location - P1;
+
+			FCollisionContact& Contact = Constraint.Manifold;
+			FSolverBody& Body0 = *Constraint.GetSolverBody0();
+			FSolverBody& Body1 = *Constraint.GetSolverBody1();
+
+			FVec3 VectorToPoint0 = Contact.Location - Body0.P();
+			FVec3 VectorToPoint1 = Contact.Location - Body1.P();
 
 			if (Contact.Phi < 0)
 			{
@@ -100,19 +105,15 @@ namespace Chaos
 				bool bHaveRestitutionPadding = (Contact.RestitutionPadding > 0.0f);
 				bool bApplyFriction = (Contact.Friction > 0) && (IterationParameters.Dt > SMALL_NUMBER);
 
-				// If we have restitution, padd the constraint by an amount that enforces the outgoing velocity constraint
+				// If we have restitution, pad the constraint by an amount that enforces the outgoing velocity constraint
 				// Really this should be per contact point, not per constraint.
 				// NOTE: once we have calculated a padding, it is locked in for the rest of the iterations, and automatically
 				// included in the Phi we get back from collision detection. The first time we calculate it, we must also
 				// add the padding to the Phi (since it was from pre-padded collision detection).
 				if (bApplyResitution && !bHaveRestitutionPadding)
 				{
-					FVec3 V0 = Particle0->V();
-					FVec3 W0 = Particle0->W();
-					FVec3 V1 = Particle1->V();
-					FVec3 W1 = Particle1->W();
-					FVec3 CV0 = V0 + FVec3::CrossProduct(W0, VectorToPoint0);
-					FVec3 CV1 = V1 + FVec3::CrossProduct(W1, VectorToPoint1);
+					FVec3 CV0 = Body0.V() + FVec3::CrossProduct(Body0.W(), VectorToPoint0);
+					FVec3 CV1 = Body1.V() + FVec3::CrossProduct(Body1.W(), VectorToPoint1);
 					FVec3 CV = CV0 - CV1;
 					FReal CVNormal = FVec3::DotProduct(CV, Contact.Normal);
 
@@ -124,13 +125,9 @@ namespace Chaos
 					}
 				}
 			
-				FReal InvM0 = bIsRigidDynamic0 ? PBDRigid0->InvM() : 0.0f;
-				FReal InvM1 = bIsRigidDynamic1 ? PBDRigid1->InvM() : 0.0f;
-				FMatrix33 InvI0 = bIsRigidDynamic0 ? Utilities::ComputeWorldSpaceInertia(Q0, PBDRigid0->InvI()) * Contact.InvInertiaScale0 : FMatrix33(0);
-				FMatrix33 InvI1 = bIsRigidDynamic1 ? Utilities::ComputeWorldSpaceInertia(Q1, PBDRigid1->InvI()) * Contact.InvInertiaScale1 : FMatrix33(0);
 				FMatrix33 ContactInvI =
-					(bIsRigidDynamic0 ? ComputeFactorMatrix3(VectorToPoint0, InvI0, InvM0) : FMatrix33(0)) +
-					(bIsRigidDynamic1 ? ComputeFactorMatrix3(VectorToPoint1, InvI1, InvM1) : FMatrix33(0));
+					(Body0.IsDynamic() ? ComputeFactorMatrix3(VectorToPoint0, Body0.InvI(), Body0.InvM()) : FMatrix33(0)) +
+					(Body1.IsDynamic() ? ComputeFactorMatrix3(VectorToPoint1, Body1.InvI(), Body1.InvM()) : FMatrix33(0));
 
 				// Calculate the normal correction
 				FVec3 NormalError = Contact.Phi * Contact.Normal;
@@ -143,14 +140,10 @@ namespace Chaos
 				if (bApplyFriction)
 				{
 					// @todo(ccaulfield): use initial velocity (as for restitution) and accumulate friction force per contact point
-					FVec3 X0 = FParticleUtilitiesXR::GetCoMWorldPosition(Particle0);
-					FVec3 X1 = FParticleUtilitiesXR::GetCoMWorldPosition(Particle1);
-					FRotation3 R0 = FParticleUtilitiesXR::GetCoMWorldRotation(Particle0);
-					FRotation3 R1 = FParticleUtilitiesXR::GetCoMWorldRotation(Particle1);
-					FVec3 V0 = FVec3::CalculateVelocity(X0, P0, IterationParameters.Dt);
-					FVec3 W0 = FRotation3::CalculateAngularVelocity(R0, Q0, IterationParameters.Dt);
-					FVec3 V1 = FVec3::CalculateVelocity(X1, P1, IterationParameters.Dt);
-					FVec3 W1 = FRotation3::CalculateAngularVelocity(R1, Q1, IterationParameters.Dt);
+					FVec3 V0 = FVec3::CalculateVelocity(Body0.X(), Body0.P(), IterationParameters.Dt);
+					FVec3 W0 = FRotation3::CalculateAngularVelocity(Body0.R(), Body0.Q(), IterationParameters.Dt);
+					FVec3 V1 = FVec3::CalculateVelocity(Body1.X(), Body1.P(), IterationParameters.Dt);
+					FVec3 W1 = FRotation3::CalculateAngularVelocity(Body1.R(), Body1.Q(), IterationParameters.Dt);
 					FVec3 CV0 = V0 + FVec3::CrossProduct(W0, VectorToPoint0);
 					FVec3 CV1 = V1 + FVec3::CrossProduct(W1, VectorToPoint1);
 					FVec3 CV = CV0 - CV1;
@@ -178,23 +171,19 @@ namespace Chaos
 				// Net Correction
 				FVec3 DX = NormalCorrection + LateralCorrection;
 				
-				if (bIsRigidDynamic0)
+				if (Body0.IsDynamic())
 				{
-					FVec3 DP0 = InvM0 * DX;
-					FVec3 DR0 = Utilities::Multiply(InvI0, FVec3::CrossProduct(VectorToPoint0, DX));
-					P0 += DP0;
-					Q0 += FRotation3::FromElements(DR0, 0.f) * Q0 * FReal(0.5);
-					Q0.Normalize();
-					FParticleUtilities::SetCoMWorldTransform(PBDRigid0, P0, Q0);
+					FVec3 DP0 = Body0.InvM() * DX;
+					FVec3 DR0 = Utilities::Multiply(Body0.InvI(), FVec3::CrossProduct(VectorToPoint0, DX));
+					Body0.ApplyTransformDelta(DP0, DR0);
+					Body0.UpdateRotationDependentState();
 				}
-				if (bIsRigidDynamic1)
+				if (Body1.IsDynamic())
 				{
-					FVec3 DP1 = InvM1 * -DX;
-					FVec3 DR1 = Utilities::Multiply(InvI1, FVec3::CrossProduct(VectorToPoint1, -DX));
-					P1 += DP1;
-					Q1 += FRotation3::FromElements(DR1, 0.f) * Q1 * FReal(0.5);
-					Q1.Normalize();
-					FParticleUtilities::SetCoMWorldTransform(PBDRigid1, P1, Q1);
+					FVec3 DP1 = Body1.InvM() * -DX;
+					FVec3 DR1 = Utilities::Multiply(Body1.InvI(), FVec3::CrossProduct(VectorToPoint1, -DX));
+					Body1.ApplyTransformDelta(DP1, DR1);
+					Body1.UpdateRotationDependentState();
 				}
 			}
 			return AccumulatedImpulse;
@@ -204,17 +193,10 @@ namespace Chaos
 		template<typename T_CONSTRAINT>
 		void ApplyImpl(T_CONSTRAINT& Constraint, const FContactIterationParameters & IterationParameters, const FContactParticleParameters & ParticleParameters)
 		{
-			FGenericParticleHandle Particle0 = FGenericParticleHandle(Constraint.Particle[0]);
-			FGenericParticleHandle Particle1 = FGenericParticleHandle(Constraint.Particle[1]);
-
 			for (int32 PairIt = 0; PairIt < IterationParameters.NumPairIterations; ++PairIt)
 			{
 				// Update the contact information based on current particles' positions
-				bool bRequiresCollisionUpdate = true;
-				if (bRequiresCollisionUpdate)
-				{
-					Collisions::Update(Constraint, IterationParameters.Dt);
-				}
+				Collisions::Update(Constraint, IterationParameters.Dt);
 
 				// Permanently disable a constraint that is beyond the cull distance
 				if (Constraint.GetPhi() >= Constraint.GetCullDistance())
@@ -235,8 +217,12 @@ namespace Chaos
 				}
 
 				// @todo(chaos): fix the collided flag - it will sometimes be set if clipping is enabled, even if there was no contact...
+				// @todo(chaos): this doesn't seem to be being set or used...remove Collided and the Aux Collided array if so. If it is supposed 
+				// to work, it should be set via the Scatter method of the SolverBody
 				if (ParticleParameters.Collided)
 				{
+					FGenericParticleHandle Particle0 = FGenericParticleHandle(Constraint.Particle[0]);
+					FGenericParticleHandle Particle1 = FGenericParticleHandle(Constraint.Particle[1]);
 					Particle0->AuxilaryValue(*ParticleParameters.Collided) = true;
 					Particle1->AuxilaryValue(*ParticleParameters.Collided) = true;
 				}
@@ -246,18 +232,17 @@ namespace Chaos
 				{
 				case EConstraintSolverType::GbfPbd:
 					{
-						ApplyContactManifold(Constraint, Particle0, Particle1, IterationParameters, ParticleParameters);
+						ApplyContactManifold(Constraint, IterationParameters, ParticleParameters);
 					}
 					break;
 				case EConstraintSolverType::StandardPbd:
 					{
-						Constraint.AccumulatedImpulse += ApplyContact2(Constraint.Manifold, Particle0, Particle1, IterationParameters, ParticleParameters);
+						Constraint.AccumulatedImpulse += ApplyContact2(Constraint, IterationParameters, ParticleParameters);
 					}
 					break;
 				case EConstraintSolverType::QuasiPbd:
 					{
-						FPBDCollisionSolver CollisionSolver;
-						CollisionSolver.SolvePosition(Constraint, IterationParameters, ParticleParameters);
+						check(false);
 					}
 					break;
 				default:
@@ -266,10 +251,12 @@ namespace Chaos
 			}
 		}
 		
-		void ApplySweptImpl(FRigidBodySweptPointContactConstraint& Constraint, const FContactIterationParameters & IterationParameters, const FContactParticleParameters & ParticleParameters)
+		void ApplySweptImpl(FPBDCollisionConstraint& Constraint, const FContactIterationParameters & IterationParameters, const FContactParticleParameters & ParticleParameters)
 		{
-			FGenericParticleHandle Particle0 = FGenericParticleHandle(Constraint.Particle[0]);
-			FGenericParticleHandle Particle1 = FGenericParticleHandle(Constraint.Particle[1]);
+			check(Constraint.GetType() == ECollisionConstraintType::Swept);
+
+			FSolverBody& Body0 = *Constraint.GetSolverBody0();
+			FSolverBody& Body1 = *Constraint.GetSolverBody1();
 
 			Collisions::UpdateSwept(Constraint, IterationParameters.Dt);
 			
@@ -291,53 +278,53 @@ namespace Chaos
 			const FContactIterationParameters IterationParametersRemainingDT{ RemainingDT, FakeIteration, IterationParameters.NumIterations, IterationParameters.NumPairIterations, IterationParameters.SolverType, IterationParameters.NeedsAnotherIteration };
 
 			// Rewind P to TOI and Apply
-			Particle0->P() = FMath::Lerp(Particle0->X(), Particle0->P(), Constraint.TimeOfImpact);
+			Body0.SetP(FMath::Lerp(Body0.X(), Body0.P(), Constraint.TimeOfImpact));
 			ApplyImpl(Constraint, IterationParametersPartialDT, CCDParticleParamaters);
 
-			// Advance P to end of frame from TOI, and Apply
-			if (IterationParameters.Iteration + 1 < IterationParameters.NumIterations)
-			{
-				Particle0->P() = Particle0->P() + Particle0->V() * RemainingDT; // If we are tunneling through something else due to this, it will be resolved in the next iteration
-				ApplyImpl(Constraint, IterationParametersRemainingDT, CCDParticleParamaters);
-			}
-			else
+			// @todo(chaos): Make this work properly for real Stanard and QPBD implementations (they do not alter velocity)
+			if (IterationParameters.SolverType == EConstraintSolverType::GbfPbd)
 			{
-				// We get here if we cannot solve CCD collisions with the given number of iterations and restitution settings.
-				// So don't do the remaining dt update. This will bleed the energy! (also: Ignore rotation)
-				// To prevent this condition: increase number of iterations and/or reduce restitution and/or reduce velocities
-				if (IterationParameters.Dt > SMALL_NUMBER)
+				// Advance P to end of frame from TOI, and Apply
+				if (IterationParameters.Iteration + 1 < IterationParameters.NumIterations)
 				{
-					// Update velocity to be consistent with PBD
-					Particle0->SetV((1.f / IterationParameters.Dt) * (Particle0->P() - Particle0->X()));
+					Body0.SetP(Body0.P() + Body0.V() * RemainingDT); // If we are tunneling through something else due to this, it will be resolved in the next iteration
+					ApplyImpl(Constraint, IterationParametersRemainingDT, CCDParticleParamaters);
 				}
 				else
 				{
-					Particle0->SetV(FVec3(0));
+					// We get here if we cannot solve CCD collisions with the given number of iterations and restitution settings.
+					// So don't do the remaining dt update. This will bleed the energy! (also: Ignore rotation)
+					// To prevent this condition: increase number of iterations and/or reduce restitution and/or reduce velocities
+					if (IterationParameters.Dt > SMALL_NUMBER)
+					{
+						// Update velocity to be consistent with PBD
+						const FReal InvDt = FReal(1) / IterationParameters.Dt;
+						Body0.SetV((Body0.P() - Body0.X()) * InvDt);
+					}
+					else
+					{
+						Body0.SetV(FVec3(0));
+					}
 				}
 			}
 		}
 
 
-		void Apply(FCollisionConstraintBase& Constraint, const FContactIterationParameters & IterationParameters, const FContactParticleParameters & ParticleParameters)
+		void Apply(FPBDCollisionConstraint& Constraint, const FContactIterationParameters & IterationParameters, const FContactParticleParameters & ParticleParameters)
 		{
-			if (Constraint.GetType() == FCollisionConstraintBase::FType::SinglePoint)
+			if (Constraint.GetType() == ECollisionConstraintType::Standard)
 			{
-				ApplyImpl(*Constraint.As<FRigidBodyPointContactConstraint>(), IterationParameters, ParticleParameters);
+				ApplyImpl(Constraint, IterationParameters, ParticleParameters);
 			}
-			else if (Constraint.GetType() == FCollisionConstraintBase::FType::SinglePointSwept)
+			else if (Constraint.GetType() == ECollisionConstraintType::Swept)
 			{
-				ApplySweptImpl(*Constraint.As<FRigidBodySweptPointContactConstraint>(), IterationParameters, ParticleParameters);
+				ApplySweptImpl(Constraint, IterationParameters, ParticleParameters);
 			}
 		}
 
 
-		template<typename T_CONSTRAINT>
-		void ApplyPushOutImpl(T_CONSTRAINT& Constraint, const TSet<const TGeometryParticleHandle<FReal, 3>*>& IsTemporarilyStatic,
-			const FContactIterationParameters & IterationParameters, const FContactParticleParameters & ParticleParameters, const FVec3& GravityDir)
+		void ApplyPushOutImpl(FPBDCollisionConstraint& Constraint, const FContactIterationParameters & IterationParameters, const FContactParticleParameters & ParticleParameters)
 		{
-			FGenericParticleHandle Particle0 = FGenericParticleHandle(Constraint.Particle[0]);
-			FGenericParticleHandle Particle1 = FGenericParticleHandle(Constraint.Particle[1]);
-
 			for (int32 PairIt = 0; PairIt < IterationParameters.NumPairIterations; ++PairIt)
 			{
 				Update(Constraint, IterationParameters.Dt);
@@ -353,25 +340,25 @@ namespace Chaos
 					return;
 				}
 
-				if (FRigidBodyPointContactConstraint* PointConstraint = Constraint.template As<FRigidBodyPointContactConstraint>())
+				// @todo(chaos): is this supposed to run for swept contacts as well?
+				if (Constraint.GetType() == ECollisionConstraintType::Standard)
 				{
 					switch (IterationParameters.SolverType)
 					{
 					case EConstraintSolverType::GbfPbd:
 						{
-							ApplyPushOutManifold(*PointConstraint, IsTemporarilyStatic, IterationParameters, ParticleParameters, GravityDir);
+							ApplyPushOutManifold(Constraint, IterationParameters, ParticleParameters);
 						}
 						break;
 					case EConstraintSolverType::StandardPbd:
 						{
 							// Shouldn't have PushOut for Standard Pbd, but this is for experimentation (Collision PushOut Iterations should normally be set to 0 instead)
-							ApplyPushOutManifold(*PointConstraint, IsTemporarilyStatic, IterationParameters, ParticleParameters, GravityDir);
+							ApplyPushOutManifold(Constraint, IterationParameters, ParticleParameters);
 						}
 						break;
 					case EConstraintSolverType::QuasiPbd:
 						{
-							FPBDCollisionSolver CollisionSolver;
-							CollisionSolver.SolveVelocity(Constraint, IterationParameters, ParticleParameters);
+							check(false);
 						}
 						break;
 					default:
@@ -381,17 +368,9 @@ namespace Chaos
 			}
 		}
 
-		void ApplyPushOut(FCollisionConstraintBase& Constraint, const TSet<const TGeometryParticleHandle<FReal, 3>*>& IsTemporarilyStatic, 
-			const FContactIterationParameters & IterationParameters, const FContactParticleParameters & ParticleParameters, const FVec3& GravityDir)
+		void ApplyPushOut(FPBDCollisionConstraint& Constraint, const FContactIterationParameters & IterationParameters, const FContactParticleParameters & ParticleParameters)
 		{
-			if (Constraint.GetType() == FCollisionConstraintBase::FType::SinglePoint)
-			{
-				ApplyPushOutImpl<FRigidBodyPointContactConstraint>(*Constraint.As<FRigidBodyPointContactConstraint>(), IsTemporarilyStatic, IterationParameters, ParticleParameters, GravityDir);
-			}
-			else if (Constraint.GetType() == FCollisionConstraintBase::FType::SinglePointSwept)
-			{
-				ApplyPushOutImpl(*Constraint.As<FRigidBodySweptPointContactConstraint>(), IsTemporarilyStatic, IterationParameters, ParticleParameters, GravityDir);
-			}
+			ApplyPushOutImpl(Constraint, IterationParameters, ParticleParameters);
 		}
 
 	} // Collisions
diff --git a/Engine/Source/Runtime/Experimental/Chaos/Private/Chaos/PBDConstraintColor.cpp b/Engine/Source/Runtime/Experimental/Chaos/Private/Chaos/PBDConstraintColor.cpp
--- a/Engine/Source/Runtime/Experimental/Chaos/Private/Chaos/PBDConstraintColor.cpp
+++ b/Engine/Source/Runtime/Experimental/Chaos/Private/Chaos/PBDConstraintColor.cpp
@@ -7,6 +7,7 @@
 #include "Chaos/Framework/Parallel.h"
 #include "Chaos/ParticleHandle.h"
 #include "Chaos/PBDRigidParticles.h"
+#include "Chaos/PBDCollisionConstraints.h"
 
 #include "ProfilingDebugging/ScopedTimers.h"
 #include "ChaosStats.h"
@@ -22,32 +23,47 @@ using namespace Chaos;
 
 DECLARE_CYCLE_STAT(TEXT("FPBDConstraintColor::ComputeColors"), STAT_Constraint_ComputeColor, STATGROUP_Chaos);
 DECLARE_CYCLE_STAT(TEXT("FPBDConstraintColor::ComputeContactGraph"), STAT_Constraint_ComputeContactGraph, STATGROUP_Chaos);
+DECLARE_CYCLE_STAT(TEXT("FPBDConstraintColor::ComputeContactGraphGBF"), STAT_Constraint_ComputeContactGraphGBF, STATGROUP_Chaos);
 DECLARE_CYCLE_STAT(TEXT("FPBDConstraintColor::ComputeIslandColoring"), STAT_Constraint_ComputeIslandColoring, STATGROUP_Chaos);
 
+bool bUseContactGraphGBF = false;
+FAutoConsoleVariableRef CVarUseContactGraphGBF(TEXT("p.Chaos.UseContactGraphGBF"), bUseContactGraphGBF, TEXT(""));
+
 void FPBDConstraintColor::ComputeIslandColoring(const int32 Island, const FPBDConstraintGraph& ConstraintGraph, uint32 ContainerId)
 {
 	SCOPE_CYCLE_COUNTER(STAT_Constraint_ComputeIslandColoring);
-	const TArray<TGeometryParticleHandle<FReal, 3>*>& IslandParticles = ConstraintGraph.GetIslandParticles(Island);
+
+	// We need to sort the particles for the coloring to be deterministic
+	// @todo(chaos): get rid of this sort and solve the problem at a higher level
+	//const TArray<TGeometryParticleHandle<FReal, 3>*>& IslandParticles = ConstraintGraph.GetIslandParticles(Island);
+	TArray<TGeometryParticleHandle<FReal, 3>*> IslandParticles = ConstraintGraph.GetIslandParticles(Island);
+	IslandParticles.Sort([](const FGeometryParticleHandle& A, const FGeometryParticleHandle& B)
+		{
+			return A.ParticleID() < B.ParticleID();
+		});
+
 	FLevelToColorToConstraintListMap& LevelToColorToConstraintListMap = IslandData[Island].LevelToColorToConstraintListMap;
 	int32& MaxColor = IslandData[Island].MaxColor;
+	int32& MaxEdges = IslandData[Island].NumEdges;
 
 	const int32 MaxLevel = IslandData[Island].MaxLevel;
 	
 	LevelToColorToConstraintListMap.Reset();
 	LevelToColorToConstraintListMap.SetNum(MaxLevel + 1);
 	MaxColor = -1;
+	MaxEdges = 0;
 
 	TSet<int32> ProcessedNodes;
 	TArray<int32> NodesToProcess;
 
 	for (const TGeometryParticleHandle<FReal, 3>* Particle :IslandParticles)
 	{
-		if (!ConstraintGraph.ParticleToNodeIndex.Find(Particle))
+		if (!ConstraintGraph.GetParticleNodes().Find(Particle))
 		{
 			continue;
 		}
 
-		const int32 ParticleNodeIndex = ConstraintGraph.ParticleToNodeIndex[Particle];
+		const int32 ParticleNodeIndex = ConstraintGraph.GetParticleNodes()[Particle];
 
 		const bool bIsParticleDynamic = Particle->CastToRigidParticle() && Particle->ObjectState() == EObjectStateType::Dynamic;
 		if (ProcessedNodes.Contains(ParticleNodeIndex) || bIsParticleDynamic == false)
@@ -60,19 +76,19 @@ void FPBDConstraintColor::ComputeIslandColoring(const int32 Island, const FPBDCo
 		while (NodesToProcess.Num())
 		{
 			const int32 NodeIndex = NodesToProcess.Last();
-			const typename FPBDConstraintGraph::FGraphNode& GraphNode = ConstraintGraph.Nodes[NodeIndex];
+			const typename FPBDConstraintGraph::FGraphNode& GraphNode = ConstraintGraph.GetGraphNodes()[NodeIndex];
 			FGraphNodeColor& ColorNode = Nodes[NodeIndex];
 
 			NodesToProcess.SetNum(NodesToProcess.Num() - 1, /*bAllowShrinking=*/false);
-			ProcessedNodes.Add(NodeIndex);
+			ProcessedNodes.Add(NodeIndex); 
 
-			for (const int32 EdgeIndex : GraphNode.Edges)
+			for (const int32 EdgeIndex : GraphNode.NodeEdges)
 			{
-				const typename FPBDConstraintGraph::FGraphEdge& GraphEdge = ConstraintGraph.Edges[EdgeIndex];
+				const typename FPBDConstraintGraph::FGraphEdge& GraphEdge = ConstraintGraph.GetGraphEdges()[EdgeIndex];
 				FGraphEdgeColor& ColorEdge = Edges[EdgeIndex];
 
 				// If this is not from our rule, ignore it
-				if (GraphEdge.Data.GetContainerId() != ContainerId)
+				if (GraphEdge.ItemContainer != ContainerId)
 				{
 					continue;
 				}
@@ -106,8 +122,8 @@ void FPBDConstraintColor::ComputeIslandColoring(const int32 Island, const FPBDCo
 				{
 					FGraphNodeColor& OtherColorNode = Nodes[OtherNodeIndex];
 
-					const typename FPBDConstraintGraph::FGraphNode& OtherGraphNode = ConstraintGraph.Nodes[OtherNodeIndex];
-					const bool bIsOtherGraphNodeDynamic = OtherGraphNode.Particle->CastToRigidParticle() && OtherGraphNode.Particle->ObjectState() == EObjectStateType::Dynamic;
+					const typename FPBDConstraintGraph::FGraphNode& OtherGraphNode = ConstraintGraph.GetGraphNodes()[OtherNodeIndex];
+					const bool bIsOtherGraphNodeDynamic = OtherGraphNode.NodeItem->CastToRigidParticle() && OtherGraphNode.NodeItem->ObjectState() == EObjectStateType::Dynamic;
 					if (bIsOtherGraphNodeDynamic)
 					{
 						while (OtherColorNode.UsedColors.Contains(ColorToUse) || ColorNode.UsedColors.Contains(ColorToUse))
@@ -141,12 +157,13 @@ void FPBDConstraintColor::ComputeIslandColoring(const int32 Island, const FPBDCo
 					LevelToColorToConstraintListMap[Level].Add(ColorEdge.Color, {});
 				}
 
-				LevelToColorToConstraintListMap[Level][ColorEdge.Color].Add(GraphEdge.Data.GetConstraintHandle());
+				LevelToColorToConstraintListMap[Level][ColorEdge.Color].Add(GraphEdge.EdgeItem);
+				++MaxEdges;
 
 				if (OtherNodeIndex != INDEX_NONE)
 				{
-					const typename FPBDConstraintGraph::FGraphNode& OtherGraphNode = ConstraintGraph.Nodes[OtherNodeIndex];
-					const bool bIsOtherGraphNodeDynamic = OtherGraphNode.Particle->CastToRigidParticle() && OtherGraphNode.Particle->ObjectState() == EObjectStateType::Dynamic;
+					const typename FPBDConstraintGraph::FGraphNode& OtherGraphNode = ConstraintGraph.GetGraphNodes()[OtherNodeIndex];
+					const bool bIsOtherGraphNodeDynamic = OtherGraphNode.NodeItem->CastToRigidParticle() && OtherGraphNode.NodeItem->ObjectState() == EObjectStateType::Dynamic;
 					if (bIsOtherGraphNodeDynamic)
 					{
 						FGraphNodeColor& OtherColorNode = Nodes[OtherNodeIndex];
@@ -160,8 +177,8 @@ void FPBDConstraintColor::ComputeIslandColoring(const int32 Island, const FPBDCo
 						// Queue other node for processing
 						if (!ProcessedNodes.Contains(OtherNodeIndex))
 						{
-							ensure(OtherGraphNode.Island == GraphNode.Island);
-							checkSlow(IslandParticles.Find(OtherGraphNode.Particle) != INDEX_NONE);
+							ensure(OtherGraphNode.IslandIndex == GraphNode.IslandIndex);
+							checkSlow(IslandParticles.Find(OtherGraphNode.NodeItem) != INDEX_NONE);
 							NodesToProcess.Add(OtherNodeIndex);
 						}
 					}
@@ -174,9 +191,9 @@ void FPBDConstraintColor::ComputeIslandColoring(const int32 Island, const FPBDCo
 void FPBDConstraintColor::ComputeContactGraph(const int32 Island, const FPBDConstraintGraph& ConstraintGraph, uint32 ContainerId)
 {
 	SCOPE_CYCLE_COUNTER(STAT_Constraint_ComputeContactGraph);
-	const TArray<int32>& ConstraintDataIndices = ConstraintGraph.GetIslandConstraintData(Island);
+	const TArray<FConstraintHandle*>& IslandConstraints = ConstraintGraph.GetIslandConstraints(Island);
 
-	IslandData[Island].MaxLevel = ConstraintDataIndices.Num() ? 0 : -1;
+	IslandData[Island].MaxLevel = IslandConstraints.Num() ? 0 : -1;
 	
 	struct FLevelNodePair
 	{
@@ -196,33 +213,41 @@ void FPBDConstraintColor::ComputeContactGraph(const int32 Island, const FPBDCons
 
 	for(const TGeometryParticleHandle<FReal, 3> * Particle : ConstraintGraph.GetIslandParticles(Island))
 	{
-		const int32* NodeIndexPtr = ConstraintGraph.ParticleToNodeIndex.Find(Particle);
+		const int32* NodeIndexPtr = ConstraintGraph.GetParticleNodes().Find(Particle);
 		const bool bIsParticleDynamic = Particle->CastToRigidParticle() && Particle->ObjectState() == EObjectStateType::Dynamic;
 
+		if (NodeIndexPtr)
+		{
+			// To be consistent with the solver body interface, we need to return a level given a particle. But internally we compute levels on nodes, therefore a mapping from nodes to particles is needed. 
+			if (bIsParticleDynamic)
+			{
+				NodeToParticle[*NodeIndexPtr] = Particle;
+			}
+		}
 		// We're only interested in static particles here to generate the graph (graph of dynamic objects touching static)
 		if(bIsParticleDynamic == false && NodeIndexPtr)
 		{
 			const int32 NodeIndex = *NodeIndexPtr;
-
-			const typename FPBDConstraintGraph::FGraphNode& GraphNode = ConstraintGraph.Nodes[NodeIndex];
+			const typename FPBDConstraintGraph::FGraphNode& GraphNode = ConstraintGraph.GetGraphNodes()[NodeIndex];
 			
-			for(int32 EdgeIndex : GraphNode.Edges)
+
+			for (int32 EdgeIndex : GraphNode.NodeEdges)
 			{
-				const typename FPBDConstraintGraph::FGraphEdge& GraphEdge = ConstraintGraph.Edges[EdgeIndex];
+				const typename FPBDConstraintGraph::FGraphEdge& GraphEdge = ConstraintGraph.GetGraphEdges()[EdgeIndex];
 
 				// If this is not from our rule, ignore it
-				if(GraphEdge.Data.GetContainerId() != ContainerId)
+				if (GraphEdge.ItemContainer != ContainerId)
 				{
 					continue;
 				}
 
 				// Find adjacent node
 				int32 OtherNode = INDEX_NONE;
-				if(GraphEdge.FirstNode == NodeIndex)
+				if (GraphEdge.FirstNode == NodeIndex)
 				{
 					OtherNode = GraphEdge.SecondNode;
 				}
-				else if(GraphEdge.SecondNode == NodeIndex)
+				else if (GraphEdge.SecondNode == NodeIndex)
 				{
 					OtherNode = GraphEdge.FirstNode;
 				}
@@ -230,7 +255,7 @@ void FPBDConstraintColor::ComputeContactGraph(const int32 Island, const FPBDCons
 				// If we have a node, add it to the queue only if it matches our island. Statics have no island and can be touching dynamics of many islands
 				// so we need to pick out only the edges that lead to the requested island to correctly build the graph. Implicitly all further edges must
 				// be of the same island so we only need to do this check for level 1
-				if(OtherNode != INDEX_NONE && ConstraintGraph.Nodes[OtherNode].Island == Island)
+				if (OtherNode != INDEX_NONE && ConstraintGraph.GetGraphNodes()[OtherNode].IslandIndex == ConstraintGraph.GetGraphIndex(Island))
 				{
 					Edges[EdgeIndex].Level = 0;
 					NodeQueue.Emplace(1, OtherNode);
@@ -245,34 +270,39 @@ void FPBDConstraintColor::ComputeContactGraph(const int32 Island, const FPBDCons
 		Current = NodeQueue.First();
 		NodeQueue.PopFrontNoCheck();
 
-		int32 Level = Current.Level;
-		int32 NodeIndex = Current.NodeIndex;
-		const typename FPBDConstraintGraph::FGraphNode& GraphNode = ConstraintGraph.Nodes[NodeIndex];
-
-		for(int32 EdgeIndex : GraphNode.Edges)
+		const int32 Level = Current.Level;
+		const int32 NodeIndex = Current.NodeIndex;
+		const TGeometryParticleHandle<FReal, 3> *Particle = NodeToParticle[NodeIndex];
+		if (Particle)
 		{
-			const typename FPBDConstraintGraph::FGraphEdge& GraphEdge = ConstraintGraph.Edges[EdgeIndex];
+			ParticleToLevel[Particle->UniqueIdx().Idx] = Level;
+		} 
+		const typename FPBDConstraintGraph::FGraphNode& GraphNode = ConstraintGraph.GetGraphNodes()[NodeIndex];
+		
+		for (int32 EdgeIndex : GraphNode.NodeEdges)
+		{
+			const typename FPBDConstraintGraph::FGraphEdge& GraphEdge = ConstraintGraph.GetGraphEdges()[EdgeIndex];
 			FGraphEdgeColor& ColorEdge = Edges[EdgeIndex];
 
 			// If this is not from our rule, ignore it
-			if(GraphEdge.Data.GetContainerId() != ContainerId)
+			if (GraphEdge.ItemContainer != ContainerId)
 			{
 				continue;
 			}
 
 			// If we have already been assigned a level, move on
-			if(ColorEdge.Level >= 0)
+			if (ColorEdge.Level >= 0)
 			{
 				continue;
 			}
 
 			// Find adjacent node and recurse
 			int32 OtherNode = INDEX_NONE;
-			if(GraphEdge.FirstNode == NodeIndex)
+			if (GraphEdge.FirstNode == NodeIndex)
 			{
 				OtherNode = GraphEdge.SecondNode;
 			}
-			else if(GraphEdge.SecondNode == NodeIndex)
+			else if (GraphEdge.SecondNode == NodeIndex)
 			{
 				OtherNode = GraphEdge.FirstNode;
 			}
@@ -283,7 +313,7 @@ void FPBDConstraintColor::ComputeContactGraph(const int32 Island, const FPBDCons
 			// may not all be in the same island, which is ok (e.g., two separated
 			// boxes sat on a large plane). We need to ignore edges that are in other islands
 			// @todo(chaos): we should probably store the island index with each edge.
-			if (ConstraintGraph.Nodes[OtherNode].Island == Island)
+			if (ConstraintGraph.GetGraphNodes()[OtherNode].IslandIndex == ConstraintGraph.GetGraphIndex(Island))
 			{
 				ColorEdge.Level = Level;
 				IslandData[Island].MaxLevel = FGenericPlatformMath::Max(IslandData[Island].MaxLevel, ColorEdge.Level);
@@ -300,52 +330,451 @@ void FPBDConstraintColor::ComputeContactGraph(const int32 Island, const FPBDCons
 	// An isolated island that is only dynamics will not have been processed above, put everything without a level into level zero
 	// #BGTODO this can surely be done as we build the edges, after this function everything will be at least level 0 so we can probably construct them
 	//         in that level to avoid a potentially large iteration here.
+	for(const FConstraintHandle* IslandConstraint : IslandConstraints)
+	{
+		if(IslandConstraint)
+		{
+			const int32 EdgeIndex = IslandConstraint->GetConstraintGraphIndex();
+			if(Edges.IsValidIndex(EdgeIndex))
+			{
+				check(Edges[EdgeIndex].Level <= IslandData[Island].MaxLevel);
+				if(Edges[EdgeIndex].Level < 0)
+				{
+					Edges[EdgeIndex].Level = 0;
+				}
+			}
+		}
+	}
+
+	check(IslandData[Island].MaxLevel >= 0 || !IslandConstraints.Num());
+}
+
+void FPBDConstraintColor::ComputeContactGraphGBF(const FReal Dt, const int32 Island, const FPBDConstraintGraph& ConstraintGraph, const uint32 ContainerId)
+{
+	SCOPE_CYCLE_COUNTER(STAT_Constraint_ComputeContactGraphGBF);
+	FGBFContactGraphData& Data = IslandToGBFContactGraphData[Island];
+	// First create an island-local graph. The mapping is stored in DynamicGlobalToIslandLocalNodeIndicesArray and KinematicGlobalToIslandLocalNodeIndices.
+	ComputeGlobalToIslandLocalNodeMapping(Island, ConstraintGraph, DynamicGlobalToIslandLocalNodeIndicesArray, Data.KinematicGlobalToIslandLocalNodeIndices);
+	CollectIslandDirectedEdges(Dt, Island, ConstraintGraph, DynamicGlobalToIslandLocalNodeIndicesArray, Data.KinematicGlobalToIslandLocalNodeIndices, Data.DigraphEdges);
+
+	// Build compressed sparse row (CSR) representation of the directed graph (digraph).
+	const int32 NumNodes=ConstraintGraph.GetIslandParticles(Island).Num();
+	BuildGraphCSR(Data.DigraphEdges, NumNodes, Data.ChildrenCSRIndices, Data.ChildrenCSRValues, Data.CSRCurrentIndices);
+
+	// Root nodes: nodes have no parents.
+	ComputeIsRootNode(Data.DigraphEdges, NumNodes, Data.IsRootNode);
+
+	// Collapse directed graph (digraph) to direct acyclic graph (DAG). This is a many-to-one mapping.
+	CollapseDigraphToDAG(Data.ChildrenCSRIndices, Data.ChildrenCSRValues, Data.IsRootNode, Data.DigraphToDAGIndices, Data.DAGToDigraphIndices, Data.Visited, Data.TraversalStack, Data.NumOnStack);
+
+	// Compute the DAG edges from the digraph edges 
+	ComputeDAGEdges(Data.DigraphEdges, Data.DigraphToDAGIndices, Data.DAGEdges);
+	BuildGraphCSR(Data.DAGEdges, NumNodes, Data.DAGChildrenCSRIndices, Data.DAGChildrenCSRValues, Data.CSRCurrentIndices);
+
+	// Topological sort on DAG nodes
+	TopologicalSortDAG(Data.DAGChildrenCSRIndices, Data.DAGChildrenCSRValues, Data.DAGToDigraphIndices, Data.SortedDAGNodes, Data.Visited);
+
+	AssignDAGNodeLevels(Data.DAGChildrenCSRIndices, Data.DAGChildrenCSRValues, Data.SortedDAGNodes, Data.DAGNodeLevels);
+	AssignEdgeLevels(Island, ConstraintGraph, DynamicGlobalToIslandLocalNodeIndicesArray, Data.KinematicGlobalToIslandLocalNodeIndices, Data.DigraphToDAGIndices, Data.DAGNodeLevels);
+	UpdateParticleToLevel(Island, ConstraintGraph, DynamicGlobalToIslandLocalNodeIndicesArray, Data.KinematicGlobalToIslandLocalNodeIndices, Data.DigraphToDAGIndices, Data.DAGNodeLevels);
+}
+
+void FPBDConstraintColor::ComputeGlobalToIslandLocalNodeMapping(const int32 Island, const FPBDConstraintGraph& ConstraintGraph, TArray<int32>& DynamicGlobalToIslandLocalNodeIndices, TMap<int32, int32>& KinematicGlobalToIslandLocalNodeIndices) const
+{
+	const TArray<FGeometryParticleHandle*>& IslandParticles = ConstraintGraph.GetIslandParticles(Island);
+	const int32 NumIslandParticles = IslandParticles.Num();
+	KinematicGlobalToIslandLocalNodeIndices.Reset();
+	for(int32 GraphNodeI = 0; GraphNodeI < NumIslandParticles; GraphNodeI++)
 	{
-		for(const int32 EdgeIndex : ConstraintDataIndices)
+		const TGeometryParticleHandle<FReal, 3>* Particle = IslandParticles[GraphNodeI];
+		const int32* NodeIndexPtr = ConstraintGraph.GetParticleNodes().Find(Particle);
+		if (NodeIndexPtr)
 		{
-			check(Edges[EdgeIndex].Level <= IslandData[Island].MaxLevel);
-			if(Edges[EdgeIndex].Level < 0)
+			const int32 NodeIndex = *NodeIndexPtr;
+			const bool IsDynamic = ConstraintGraph.GetGraphNodes()[NodeIndex].IslandIndex != INDEX_NONE;
+			if (IsDynamic)
 			{
-				Edges[EdgeIndex].Level = 0;
+				DynamicGlobalToIslandLocalNodeIndices[NodeIndex] = GraphNodeI;
+			}
+			else
+			{
+				KinematicGlobalToIslandLocalNodeIndices.FindOrAdd(NodeIndex) = GraphNodeI;
 			}
 		}
 	}
+}
 
-	check(IslandData[Island].MaxLevel >= 0 || !ConstraintDataIndices.Num());
+int32 FPBDConstraintColor::GetIslandLocalNodeIdx(const int32 GlobalNodeIndex, const TArray<int32>& DynamicGlobalToIslandLocalNodeIndices, const TMap<int32, int32>& KinematicGlobalToIslandLocalNodeIndices, const bool IsNodeDynamic) const
+{
+	if (IsNodeDynamic)
+	{
+		return DynamicGlobalToIslandLocalNodeIndices[GlobalNodeIndex];
+	}
+	else
+	{
+		return KinematicGlobalToIslandLocalNodeIndices[GlobalNodeIndex];
+	}
+}
+
+void FPBDConstraintColor::CollectIslandDirectedEdges(const FReal Dt, const int32 Island, const FPBDConstraintGraph& ConstraintGraph, const TArray<int32>& DynamicGlobalToIslandLocalNodeIndices, const TMap<int32, int32>& KinematicGlobalToIslandLocalNodeIndices, TArray<FDirectedEdge>& DirectedEdges) const
+{
+	DirectedEdges.Reset();
+	const int32 NumIslandParticles = ConstraintGraph.GetIslandParticles(Island).Num();
+
+	for (const FConstraintHandle* IslandConstraint : ConstraintGraph.GetIslandConstraints(Island))
+	{
+		if(IslandConstraint && ConstraintGraph.GetGraphEdges().IsValidIndex(IslandConstraint->GetConstraintGraphIndex()))
+		{
+			const typename FPBDConstraintGraph::FGraphEdge& GraphEdge = ConstraintGraph.GetGraphEdges()[IslandConstraint->GetConstraintGraphIndex()];
+
+			const FPBDCollisionConstraintHandle *CollisionHandle = GraphEdge.EdgeItem->As<FPBDCollisionConstraintHandle>();
+			if (CollisionHandle)
+			{
+				const bool IsFirstNodeDynamic = ConstraintGraph.GetGraphNodes()[GraphEdge.FirstNode].IslandIndex != INDEX_NONE;
+				const int32 FirstLocalNodeIndex = GetIslandLocalNodeIdx(GraphEdge.FirstNode, DynamicGlobalToIslandLocalNodeIndices, KinematicGlobalToIslandLocalNodeIndices, IsFirstNodeDynamic);
+				const bool IsSecondNodeDynamic = ConstraintGraph.GetGraphNodes()[GraphEdge.SecondNode].IslandIndex != INDEX_NONE;
+				const int32 SecondLocalNodeIndex = GetIslandLocalNodeIdx(GraphEdge.SecondNode, DynamicGlobalToIslandLocalNodeIndices, KinematicGlobalToIslandLocalNodeIndices, IsSecondNodeDynamic);
+
+				const ECollisionConstraintDirection ConstraintDirection = CollisionHandle->GetContact().GetConstraintDirection(Dt);
+				if(ConstraintDirection == ECollisionConstraintDirection::Particle0ToParticle1)
+				{
+					DirectedEdges.Push(FDirectedEdge(FirstLocalNodeIndex, SecondLocalNodeIndex));
+				}
+				else if(ConstraintDirection == ECollisionConstraintDirection::Particle1ToParticle0)
+				{
+					DirectedEdges.Push(FDirectedEdge(SecondLocalNodeIndex, FirstLocalNodeIndex));
+				}
+			}
+		}
+	}
+}
+
+/**
+ * The children of node i is stored from CSRValues[CSRIndices[i]] to CSRValues[CSRIndices[i+1]]
+ */
+void FPBDConstraintColor::BuildGraphCSR(const TArray<FDirectedEdge>& GraphEdges, const int32 NumNodes, TArray<int32>& CSRIndices, TArray<int32>& CSRValues, TArray<int32>& CurrentIndices) const
+{
+	CSRIndices.Init(0, NumNodes + 1);
+	if (NumNodes == 0)
+	{
+		return;
+	}
+	const int32 NumEdges = GraphEdges.Num();
+	CSRValues.Init(-1, NumEdges);
+	// Collect the number of children of each node and store them in CSRIndices[i+1]
+	for(const FDirectedEdge& Edge : GraphEdges)
+	{
+		CSRIndices[Edge.FirstNode + 1] += 1;
+	}
+	// CurrentIndices[i] is the first available vacancy for node i to populate its children
+	CurrentIndices.Init(0, NumNodes);
+	CurrentIndices[0] = 0;
+	if (NumNodes > 1)
+	{
+		CurrentIndices[1] = CSRIndices[1];
+		for(int32 i = 2; i < NumNodes; i++)
+		{
+			CSRIndices[i] += CSRIndices[i-1];
+			CurrentIndices[i] = CSRIndices[i];
+		}
+	}
+
+	// Populate edges
+	CSRIndices[NumNodes] = NumEdges;
+	for(const FDirectedEdge& Edge : GraphEdges)
+	{
+		CSRValues[CurrentIndices[Edge.FirstNode]] = Edge.SecondNode;
+		CurrentIndices[Edge.FirstNode]++;
+	}
+}
+
+void FPBDConstraintColor::ComputeIsRootNode(const TArray<FDirectedEdge>& GraphEdges, const int32 NumNodes, TArray<bool> &IsRootNode) const
+{
+	IsRootNode.SetNum(NumNodes);
+	for (bool& B : IsRootNode)
+	{
+		B = true;
+	}
+	for(const FDirectedEdge& Edge : GraphEdges)
+	{
+		if (IsRootNode[Edge.SecondNode])
+		{
+			IsRootNode[Edge.SecondNode] = false;
+		}
+	}
+}
+
+void FPBDConstraintColor::CollapseDigraphToDAG(const TArray<int32>& ChildrenCSRIndices, const TArray<int32>& ChildrenCSRValues, const TArray<bool>& IsRootNode, TArray<int32>& DigraphToDAGIndices, TArray<TArray<int32>>& DAGToDigraphIndices, TArray<bool>& Visited, TArray<int32>& TraversalStack, TArray<int32>& NumOnStack) const
+{
+	const int32 NumDigraphNodes = ChildrenCSRIndices.Num() - 1;
+	DigraphToDAGIndices.SetNum(NumDigraphNodes);
+	DAGToDigraphIndices.SetNum(NumDigraphNodes);
+
+	// Initialize identity mapping between Digraph and DAG
+	for(int32 i = 0; i < NumDigraphNodes; i++)
+	{
+		DigraphToDAGIndices[i] = i;
+		DAGToDigraphIndices[i].SetNum(1);
+		DAGToDigraphIndices[i][0] = i; 
+	}
+
+	// If Visited[**DigraphI**] == true, the node DigraphI and its children are processed and will not be traversed in the future.
+	Visited.Init(false, NumDigraphNodes);
+
+	// Every **Digraph** node will be pushed to TraversalStack before they are processed and will be popped from the stack after the process finished.
+	TraversalStack.Reserve(64);
+	TraversalStack.Reset();
+
+	// NumOnStack[**DAGI**] is the number of Digraph nodes that are collapsed to DAGI and are currently on TraversalStack.
+	NumOnStack.Init(0, NumDigraphNodes);
+
+
+	TFunctionRef<void(const int32)> CollapseLoop = [&DigraphToDAGIndices, &DAGToDigraphIndices, &TraversalStack, &NumOnStack](const int32 DigraphI)
+	{
+		const int32 DAGI = DigraphToDAGIndices[DigraphI];
+		// Go through the traversal stack in reverse order until we find another node that is pointed to DAGI. Collapse all the nodes traversed this way to DAGI.
+		// We collapse these nodes to DAGI, but we don't remove them from the traversal stack. We will still need to visit their children later.
+		for(int32 j = TraversalStack.Num() - 1; DigraphToDAGIndices[TraversalStack[j]] != DAGI; j--)
+		{
+			const int32 DigraphJ = TraversalStack[j];
+			const int32 DAGJ = DigraphToDAGIndices[DigraphJ];
+			// Merge DAGI and DAGJ. Move all the nodes in DAGToDigraphIndices[DAGJ] to DAGToDigraphIndices[DAGI].
+			const int32 Num = DAGToDigraphIndices[DAGJ].Num();
+			if (Num > 0)
+			{
+				for(int32 k = 0; k < Num; k++)
+				{
+					const int32 DigraphK = DAGToDigraphIndices[DAGJ][k];
+					DigraphToDAGIndices[DigraphK] = DAGI;
+					DAGToDigraphIndices[DAGI].Push(DigraphK);
+				}
+				DAGToDigraphIndices[DAGJ].Empty();
+				NumOnStack[DAGJ] -= Num;
+				NumOnStack[DAGI] += Num;
+			}
+		}
+	};
+
+	// For simplicity, traverse on **Digraph** nodes so we don't need to change graph edges as we collapsed digraph to DAG.
+	TFunctionRef<void(const int32)> TraverseAndCollapse = [&DigraphToDAGIndices, &NumOnStack, &CollapseLoop, &TraversalStack, &Visited, &ChildrenCSRIndices, &ChildrenCSRValues, &TraverseAndCollapse](const int32 DigraphI)
+	{
+		int32 DAGI = DigraphToDAGIndices[DigraphI];
+		// If there are already Digraph nodes that point to DAGI on the stack, collapse the loop.
+		if (NumOnStack[DAGI] > 0)
+		{
+			CollapseLoop(DigraphI);
+		}
+		else
+		{
+			const int32 NumEdges = ChildrenCSRIndices[DigraphI + 1] - ChildrenCSRIndices[DigraphI];
+			// For non-leaf nodes
+			if (NumEdges > 0)
+			{
+				TraversalStack.Push(DigraphI);
+				NumOnStack[DAGI]++;
+				for(int32 j = ChildrenCSRIndices[DigraphI]; j < ChildrenCSRIndices[DigraphI + 1]; j++)
+				{
+					const int32 DigraphJ = ChildrenCSRValues[j];
+					if (Visited[DigraphJ])
+					{
+						continue;
+					}
+					TraverseAndCollapse(DigraphJ);
+				}
+				// Since there might have been collapsing of the loops and thereby the mapped DAG node might be modified, DAGI needs to be updated here.
+				DAGI = DigraphToDAGIndices[DigraphI];
+				// Popping from the stack
+				NumOnStack[DAGI]--;
+				TraversalStack.Pop(false/*bAllowShrinking*/); 
+			}
+			Visited[DigraphI]=true;
+		}
+	};
+
+	// First traverse all the root nodes that have no parents.
+	for (int32 DigraphI = 0; DigraphI < NumDigraphNodes; DigraphI++)
+	{
+		if (IsRootNode[DigraphI])
+		{
+			TraverseAndCollapse(DigraphI);
+		}
+	}
+
+	// Next traverse the rest loops.
+	for(int32 DigraphI = 0; DigraphI < NumDigraphNodes; DigraphI++)
+	{
+		if(!Visited[DigraphI])
+		{
+			TraverseAndCollapse(DigraphI);
+		}
+	}
+}
+
+void FPBDConstraintColor::ComputeDAGEdges(const TArray<FDirectedEdge>& DigraphEdges, const TArray<int32>& DigraphToDAGIndices, TArray<FDirectedEdge> &DAGEdges) const
+{
+	DAGEdges.Reset(DigraphEdges.Num());
+	for (const FDirectedEdge &Edge : DigraphEdges)
+	{
+		const int32 DAGFirst = DigraphToDAGIndices[Edge.FirstNode];
+		const int32 DAGSecond = DigraphToDAGIndices[Edge.SecondNode];
+		// Remove self loops
+		if (DAGFirst != DAGSecond)
+		{
+			DAGEdges.Push({DAGFirst, DAGSecond});
+		}
+	}
+}
+
+void FPBDConstraintColor::TopologicalSortDAG(const TArray<int32>& DAGChildrenCSRIndices, const TArray<int32>& DAGChildrenCSRValues, const TArray<TArray<int32>>& DAGToDigraphIndices, TArray<int32>& SortedDAGNodes, TArray<bool>& Visited) const
+{
+	const int32 NumNodes = DAGToDigraphIndices.Num();
+	SortedDAGNodes.Reset(NumNodes);
+	Visited.Init(false, NumNodes);
+	TFunctionRef<void(const int32)> Traverse = [&DAGChildrenCSRIndices, &DAGChildrenCSRValues, &SortedDAGNodes, &Visited, &Traverse](const int32 DAGNodeI)
+	{
+		Visited[DAGNodeI] = true;
+		for (int32 i = DAGChildrenCSRIndices[DAGNodeI]; i < DAGChildrenCSRIndices[DAGNodeI + 1]; i++)
+		{
+			const int32 Child = DAGChildrenCSRValues[i];
+			if (!Visited[Child])
+			{
+				Traverse(Child);
+			}
+		}
+		// Push DAGNodeI after all its children are pushed.
+		SortedDAGNodes.Push(DAGNodeI);
+	};
+	for(int32 i = 0; i < NumNodes; i++)
+	{
+		const bool IsValidDAGNode = DAGToDigraphIndices[i].Num() > 0; 
+		if (!IsValidDAGNode)
+		{
+			continue;
+		}
+		if (Visited[i])
+		{
+			continue;
+		}
+		Traverse(i);
+	}
+}
+
+void FPBDConstraintColor::AssignDAGNodeLevels(const TArray<int32>& DAGChildrenCSRIndices, const TArray<int32>& DAGChildrenCSRValues, const TArray<int32>& SortedDAGNodes, TArray<int32> &DAGNodeLevels) const
+{
+	const int32 NumNodes = DAGChildrenCSRIndices.Num() - 1;
+	DAGNodeLevels.Init(0, NumNodes);
+	for (int32 i = SortedDAGNodes.Num() - 1; i >= 0; i--)
+	{
+		const int32 DAGNodeI = SortedDAGNodes[i];
+		for (int32 j = DAGChildrenCSRIndices[DAGNodeI]; j < DAGChildrenCSRIndices[DAGNodeI + 1]; j++)
+		{
+			const int32 Child = DAGChildrenCSRValues[j];
+			// Level of DAGI is the max distance to root nodes
+			DAGNodeLevels[Child] = FMath::Max(DAGNodeLevels[Child], DAGNodeLevels[DAGNodeI] + 1);
+		}
+	}
+}
+
+void FPBDConstraintColor::AssignEdgeLevels(const int32 Island, const FPBDConstraintGraph& ConstraintGraph, const TArray<int32>& DynamicGlobalToIslandLocalNodeIndices, const TMap<int32, int32>& KinematicGlobalToIslandLocalNodeIndices, const TArray<int32>& DigraphToDAGIndices, const TArray<int32> DAGNodeLevels)
+{
+	for(const FConstraintHandle* IslandConstraint : ConstraintGraph.GetIslandConstraints(Island))
+	{
+		if(IslandConstraint && ConstraintGraph.GetGraphEdges().IsValidIndex(IslandConstraint->GetConstraintGraphIndex()))
+		{
+			const int32 EdgeIndex = IslandConstraint->GetConstraintGraphIndex();
+			const typename FPBDConstraintGraph::FGraphEdge& GraphEdge = ConstraintGraph.GetGraphEdges()[EdgeIndex];
+
+			const bool IsFirstNodeDynamic = ConstraintGraph.GetGraphNodes()[GraphEdge.FirstNode].IslandIndex != INDEX_NONE;
+			const int32 FirstLocalNodeIndex = GetIslandLocalNodeIdx(GraphEdge.FirstNode, DynamicGlobalToIslandLocalNodeIndices, KinematicGlobalToIslandLocalNodeIndices, IsFirstNodeDynamic);
+			const int32 FirstLevel = DAGNodeLevels[DigraphToDAGIndices[FirstLocalNodeIndex]];
+			const bool IsSecondNodeDynamic = ConstraintGraph.GetGraphNodes()[GraphEdge.SecondNode].IslandIndex != INDEX_NONE;
+			const int32 SecondLocalNodeIndex = GetIslandLocalNodeIdx(GraphEdge.SecondNode, DynamicGlobalToIslandLocalNodeIndices, KinematicGlobalToIslandLocalNodeIndices, IsSecondNodeDynamic);
+			const int32 SecondLevel = DAGNodeLevels[DigraphToDAGIndices[SecondLocalNodeIndex]];
+
+			// Edge level is the max of node level
+			const int32 EdgeLevel = FMath::Max(FirstLevel, SecondLevel);
+			Edges[EdgeIndex].Level = EdgeLevel;
+			IslandData[Island].MaxLevel = FMath::Max(IslandData[Island].MaxLevel, EdgeLevel);
+		}
+	}
+}
+
+void FPBDConstraintColor::UpdateParticleToLevel(const int32 Island, const FPBDConstraintGraph& ConstraintGraph, const TArray<int32>& DynamicGlobalToIslandLocalNodeIndices, const TMap<int32, int32>& KinematicGlobalToIslandLocalNodeIndices, const TArray<int32>& DigraphToDAGIndices, const TArray<int32>& DAGNodeLevels)
+{
+	for(const TGeometryParticleHandle<FReal, 3> * Particle : ConstraintGraph.GetIslandParticles(Island))
+	{
+		const int32* NodeIndexPtr = ConstraintGraph.GetParticleNodes().Find(Particle);
+		if (NodeIndexPtr)
+		{
+			const int32 NodeIndex = *NodeIndexPtr;
+			const int32 ParticleIdx = Particle->UniqueIdx().Idx;
+			const typename FPBDConstraintGraph::FGraphNode& GraphNode = ConstraintGraph.GetGraphNodes()[NodeIndex];
+			const bool IsNodeDynamic = GraphNode.IslandIndex != INDEX_NONE;
+			if (IsNodeDynamic)
+			{
+				const int32 IslandDigraphI = GetIslandLocalNodeIdx(NodeIndex, DynamicGlobalToIslandLocalNodeIndices, KinematicGlobalToIslandLocalNodeIndices, IsNodeDynamic);
+				const int32 Level = DAGNodeLevels[DigraphToDAGIndices[IslandDigraphI]];
+				ParticleToLevel[ParticleIdx] = Level;	
+			}
+			else
+			{
+				// Set kinematic particle levels to 0
+				ParticleToLevel[ParticleIdx] = 0;
+			}
+		}
+	}
+}
+
+int32 FPBDConstraintColor::GetParticleLevel(const FGeometryParticleHandle* ParticleHandle) const
+{
+	return ParticleToLevel[ParticleHandle->UniqueIdx().Idx];
 }
 
 void FPBDConstraintColor::InitializeColor(const FPBDConstraintGraph& ConstraintGraph)
 {
 	// The Number of nodes is large and fairly constant so persist rather than resetting every frame
-	if (Nodes.Num() != ConstraintGraph.Nodes.Num())
+	if (Nodes.Num() != ConstraintGraph.GetGraphNodes().GetMaxIndex())
 	{
 		// Nodes need to grow when the nodes of the constraint graph grows
-		Nodes.AddDefaulted(ConstraintGraph.Nodes.Num() - Nodes.Num());
+		Nodes.AddDefaulted(ConstraintGraph.GetGraphNodes().GetMaxIndex() - Nodes.Num());
 	}
-	
-	// Reset the existing Nodes - so colors are all reset to zero
-	for (int32 UpdatedNode : UpdatedNodes)
+
+	for (auto& Node : Nodes)
 	{
-		Nodes[UpdatedNode].NextColor = 0;
-		Nodes[UpdatedNode].UsedColors.Empty();
+		Node.NextColor = 0;
+		Node.UsedColors.Empty();
 	}
 
 	// edges are not persistent right now so we still reset them
 	Edges.Reset();
 	IslandData.Reset();
 
-	Edges.SetNum(ConstraintGraph.Edges.Num());
-	IslandData.SetNum(ConstraintGraph.IslandToData.Num());
-
-	UpdatedNodes = ConstraintGraph.GetUpdatedNodes();
+	Edges.SetNum(ConstraintGraph.GetGraphEdges().GetMaxIndex());
+	IslandData.SetNum(ConstraintGraph.NumIslands());
+	ParticleToLevel.Init(0, ConstraintGraph.GetMaxParticleUniqueIdx() + 1);
+	DynamicGlobalToIslandLocalNodeIndicesArray.Init(-1, ConstraintGraph.GetIslandGraph()->NumNodes());
+	NodeToParticle.Init(nullptr, ConstraintGraph.GetIslandGraph()->NumNodes());
+	IslandToGBFContactGraphData.SetNum(ConstraintGraph.NumIslands());
 }
 
-void FPBDConstraintColor::ComputeColor(const int32 Island, const FPBDConstraintGraph& ConstraintGraph, uint32 ContainerId)
+void FPBDConstraintColor::ComputeColor(const FReal Dt, const int32 Island, const FPBDConstraintGraph& ConstraintGraph, uint32 ContainerId)
 {
 	SCOPE_CYCLE_COUNTER(STAT_Constraint_ComputeColor);
 	if (bUseContactGraph)
 	{
-		ComputeContactGraph(Island, ConstraintGraph, ContainerId);
+		if (bUseContactGraphGBF)
+		{
+			// Hack: extra cautious to make sure ComputeContactGraphGBF is only called once per frame.
+			// Ideally this would be ContainerType == CollisionConstraints
+			if (ContainerId == 0) 
+			{
+				ComputeContactGraphGBF(Dt, Island, ConstraintGraph, ContainerId);
+			}
+		}
+		else
+		{
+			ComputeContactGraph(Island, ConstraintGraph, ContainerId);
+		}
 	}
 	else
 	{
diff --git a/Engine/Source/Runtime/Experimental/Chaos/Private/Chaos/PBDConstraintContainer.cpp b/Engine/Source/Runtime/Experimental/Chaos/Private/Chaos/PBDConstraintContainer.cpp
--- a/Engine/Source/Runtime/Experimental/Chaos/Private/Chaos/PBDConstraintContainer.cpp
+++ b/Engine/Source/Runtime/Experimental/Chaos/Private/Chaos/PBDConstraintContainer.cpp
@@ -3,8 +3,8 @@
 
 namespace Chaos
 {
-	FPBDConstraintContainer::FPBDConstraintContainer(EConstraintContainerType InConstraintType)
-		: ConstraintContainerType(InConstraintType)
+	FPBDConstraintContainer::FPBDConstraintContainer(FConstraintHandleTypeID InConstraintHandleType)
+		: ConstraintHandleType(InConstraintHandleType)
 		, ContainerId(INDEX_NONE)
 	{
 	}
@@ -13,12 +13,12 @@ namespace Chaos
 	{
 	}
 
-	int32 FPBDConstraintContainer::GetConstraintIndex(const FConstraintHandle* ConstraintHandle) const
+	int32 FPBDIndexedConstraintContainer::GetConstraintIndex(const FIndexedConstraintHandle* ConstraintHandle) const
 	{
 		return ConstraintHandle->GetConstraintIndex();
 	}
 
-	void FPBDConstraintContainer::SetConstraintIndex(FConstraintHandle* ConstraintHandle, int32 ConstraintIndex) const
+	void FPBDIndexedConstraintContainer::SetConstraintIndex(FIndexedConstraintHandle* ConstraintHandle, int32 ConstraintIndex) const
 	{
 		ConstraintHandle->ConstraintIndex = ConstraintIndex;
 	}
