Chaos

diff --git a/Engine/Source/Runtime/Experimental/Chaos/Private/Chaos/PBDConstraintGraph.cpp b/Engine/Source/Runtime/Experimental/Chaos/Private/Chaos/PBDConstraintGraph.cpp
--- a/Engine/Source/Runtime/Experimental/Chaos/Private/Chaos/PBDConstraintGraph.cpp
+++ b/Engine/Source/Runtime/Experimental/Chaos/Private/Chaos/PBDConstraintGraph.cpp
@@ -1,1053 +1 @@
 // Copyright Epic Games, Inc. All Rights Reserved.
-
-#include "Chaos/PBDConstraintGraph.h"
-#include "ProfilingDebugging/ScopedTimers.h"
-#include "ChaosStats.h"
-#include "ChaosLog.h"
-#include "Chaos/Framework/Parallel.h"
-#include "Chaos/PBDRigidParticles.h"
-#include "Containers/Queue.h"
-#include "Chaos/ParticleHandle.h"
-#include "Chaos/ConstraintHandle.h"
-#include "Chaos/PBDRigidsSOAs.h"
-#include "ChaosLog.h"
-
-using namespace Chaos;
-
-bool ChaosSolverSleepEnabled = true;
-FAutoConsoleVariableRef CVarChaosSolverSleepEnabled(TEXT("p.Chaos.Solver.SleepEnabled"), ChaosSolverSleepEnabled, TEXT(""));
-
-bool ChaosSolverCollisionDefaultUseMaterialSleepThresholdsCVar = true;
-FAutoConsoleVariableRef CVarChaosSolverCollisionDefaultUseMaterialSleepThresholds(TEXT("p.ChaosSolverCollisionDefaultUseMaterialSleepThresholds"), ChaosSolverCollisionDefaultUseMaterialSleepThresholdsCVar, TEXT("Enable material support for sleeping thresholds[def:true]"));
-
-int32 ChaosSolverCollisionDefaultSleepCounterThresholdCVar = 20; 
-FAutoConsoleVariableRef CVarChaosSolverCollisionDefaultSleepCounterThreshold(TEXT("p.ChaosSolverCollisionDefaultSleepCounterThreshold"), ChaosSolverCollisionDefaultSleepCounterThresholdCVar, TEXT("Default counter threshold for sleeping.[def:20]"));
-
-Chaos::FRealSingle ChaosSolverCollisionDefaultLinearSleepThresholdCVar = 0.001f; // .001 unit mass cm
-FAutoConsoleVariableRef CVarChaosSolverCollisionDefaultLinearSleepThreshold(TEXT("p.ChaosSolverCollisionDefaultLinearSleepThreshold"), ChaosSolverCollisionDefaultLinearSleepThresholdCVar, TEXT("Default linear threshold for sleeping.[def:0.001]"));
-
-Chaos::FRealSingle ChaosSolverCollisionDefaultAngularSleepThresholdCVar = 0.0087f;  //~1/2 unit mass degree
-FAutoConsoleVariableRef CVarChaosSolverCollisionDefaultAngularSleepThreshold(TEXT("p.ChaosSolverCollisionDefaultAngularSleepThreshold"), ChaosSolverCollisionDefaultAngularSleepThresholdCVar, TEXT("Default angular threshold for sleeping.[def:0.0087]"));
-
-namespace Chaos
-{
-	bool ChaosPerfHackIgnoreSleepingContacts = false;
-	FAutoConsoleVariableRef CVarChaosPerfHackDisableSleepingConstraints(TEXT("p.Chaos.PerfHackIgnoreSleepingContacts"), ChaosPerfHackIgnoreSleepingContacts, TEXT("Hack to improve perf by not generating all contacts for sleeping particles"));
-}
-
-FPBDConstraintGraph::FPBDConstraintGraph() : VisitToken(0)
-{
-}
-
-
-FPBDConstraintGraph::FPBDConstraintGraph(const TParticleView<FGeometryParticles>& Particles) : VisitToken(0)
-{
-	InitializeGraph(Particles);
-}
-
-int32 FPBDConstraintGraph::ReserveParticles(const int32 Num)
-{
-	const int32 NumFree = FreeIndexList.Num();
-	int32 NumToAdd = NumFree - Num; // Can be negative
-	if (NumToAdd > 0)
-	{
-		return 0;
-	}
-	NumToAdd = -NumToAdd; // Flip sign
-	const int32 NumNodes = Nodes.Num();
-	Nodes.SetNum(NumNodes + NumToAdd);
-	FreeIndexList.Reserve(NumToAdd);
-	for (int32 Counter = 0; Counter < NumToAdd; Counter++)
-	{
-		FreeIndexList.Push(NumNodes + Counter);
-	}
-
-	ParticleToNodeIndex.Reserve(ParticleToNodeIndex.Num() + NumToAdd);
-	Visited.Reserve(Visited.Num() + NumToAdd);
-
-	return NumToAdd;
-}
-
-
-void FPBDConstraintGraph::ParticleAdd(FGeometryParticleHandle* AddedParticle)
-{
-	// GC Code creates particle then enables particle so end up calling this twice triggering ensure
-	if (/*ensure*/ (!ParticleToNodeIndex.Contains(AddedParticle)))
-	{
-		const int32 NewNodeIndex = GetNextNodeIndex();
-		FGraphNode& Node = Nodes[NewNodeIndex];
-		ensure(Node.Edges.Num() == 0);
-		ensure(Node.Island == INDEX_NONE);
-
-		Node.Particle = AddedParticle->Handle();
-		check(Node.Particle);
-		ParticleToNodeIndex.Add(Node.Particle, NewNodeIndex);
-
-		const int32 NewMinNum = NewNodeIndex + 1;
-		const int32 NumToAdd = NewMinNum - Visited.Num();
-		if (NumToAdd > 0)
-			Visited.AddZeroed(NumToAdd);
-		else
-			Visited[NewNodeIndex] = 0;
-	}
-}
-
-
-void FPBDConstraintGraph::ParticleRemove(FGeometryParticleHandle* RemovedParticle)
-{
-	if (ParticleToNodeIndex.Contains(RemovedParticle))
-	{
-		const int32 NodeIdx = ParticleToNodeIndex[RemovedParticle];
-		FreeIndexList.Push(NodeIdx);
-
-		FGraphNode& NodeRemoved = Nodes[NodeIdx];
-
-		// remove associated edges
-		for (int32 EdgeIndex : NodeRemoved.Edges)
-		{
-			// go to other particle and remove its edge reference
-			FGraphEdge& Edge = Edges[EdgeIndex];
-			if (Edge.FirstNode != Edge.SecondNode)
-			{
-				if (Edge.FirstNode == NodeIdx && Edge.SecondNode != INDEX_NONE)
-				{
-					FGraphNode& OtherNode = Nodes[Edge.SecondNode];
-					OtherNode.Edges.Remove(EdgeIndex);
-				}
-				else if (Edge.SecondNode == NodeIdx && Edge.FirstNode != INDEX_NONE)
-				{
-					FGraphNode& OtherNode = Nodes[Edge.FirstNode];
-					OtherNode.Edges.Remove(EdgeIndex);
-				}
-			}
-
-			// do we want to free the edge here also or leave it for later?
-			//Edges[EdgeIndex] = FGraphEdge();
-		}
-
-		NodeRemoved.Edges.Empty();
-		NodeRemoved.Particle = nullptr;
-		NodeRemoved.Island = INDEX_NONE;
-
-		Visited[NodeIdx]=0;
-		ParticleToNodeIndex.Remove(RemovedParticle);
-		UpdatedNodes.RemoveSwap(NodeIdx, false);
-	}
-}
-
-
-int32 FPBDConstraintGraph::GetNextNodeIndex()
-{
-	int32 NewNodeIndex = Nodes.Num();
-	if (FreeIndexList.Num() > 0)
-	{
-		NewNodeIndex = FreeIndexList.Pop();
-	}
-
-	if (Nodes.Num() <= NewNodeIndex)
-	{
-		Nodes.SetNum(NewNodeIndex + 1);
-	}
-
-	return NewNodeIndex;
-}
-
-/**
- * Called every frame, it used to clear all Nodes, Edges, ParticleToNodeIndex
- * then AddToGraph on constraint rules would fill up
- *
- * Now clears Edges and attempts to retain Nodes and ParticleToNodeIndex
- * It is still setting up nodes that don't have a constraint so wasted effort iterating over all nodes, better to iterate over constraints or don't fill out Nodes if they don't have a constraint
- */
-
-void FPBDConstraintGraph::InitializeGraph(const TParticleView<FGeometryParticles>& Particles)
-{
-	const int32 NumNonDisabledParticles = Particles.Num();
-
-	if(NumNonDisabledParticles && Nodes.Num()==0)
-	{
-		ensure(FreeIndexList.Num()==0);
-		Nodes.Reset();
-		Nodes.AddDefaulted(NumNonDisabledParticles);
-
-		ParticleToNodeIndex.Reset();
-		ParticleToNodeIndex.Reserve(NumNonDisabledParticles);
-		{
-			int32 Index = 0;
-			for (auto& Particle : Particles)
-			{
-				FGraphNode& Node = Nodes[Index];
-				Node.Particle = Particle.Handle();
-				ParticleToNodeIndex.Add(Node.Particle, Index);
-				++Index;
-			}
-		}
-
-		Visited.Reset();
-		Visited.AddZeroed(NumNonDisabledParticles);
-	}
-	else
-	{
-		if (!(NumNonDisabledParticles <= Nodes.Num())) 
-		{
-			for (auto& Particle : Particles)
-			{
-				if(!ParticleToNodeIndex.Contains(Particle.Handle()))
-					ParticleAdd(Particle.Handle());
-			}
-		}
-
-		// Update nodes may contain duplicate entries. To process in parallel we either need to prevent that, 
-		// or ensure we only process the first entry of a dupe. We are doing the latter for now...
-		UpdatedNodes.Sort();
-		PhysicsParallelFor(UpdatedNodes.Num(), 
-			[&](int32 Index)
-			{
-				if ((Index == 0) || (UpdatedNodes[Index] != UpdatedNodes[Index - 1]))
-				{
-					Nodes[UpdatedNodes[Index]].Island = INDEX_NONE;
-					Nodes[UpdatedNodes[Index]].Edges.Empty();
-					auto* Particle = Nodes[UpdatedNodes[Index]].Particle;
-				    if (CHAOS_ENSURE(Particle))
-				    {
-						// This does NOT check if particle is actually dynamic. TODO confirm that this should happen to all nonDisabledParticles, not just dynamics.
-						FPBDRigidParticleHandle* PBDRigid = Particle->CastToRigidParticle();
-						if(PBDRigid && PBDRigid->ObjectState() != EObjectStateType::Kinematic)
-					    {
-						    PBDRigid->Island() = INDEX_NONE;
-					    }
-				    }
-				}
-			});
-		UpdatedNodes.Empty();
-
-	}
-
-	Edges.Reset();
-
-	//@todo(ocohen): Should we reset more than just the edges? What about bIsIslandPersistant?
-	for (TArray<int32>& IslandConstraintList : IslandToConstraints)
-	{
-		IslandConstraintList.Reset();
-	}
-}
-
-
-void FPBDConstraintGraph::ResetIslands(const TParticleView<FPBDRigidParticles>& PBDRigids)
-{
-	//@todo(ocohen): Should we reset more than just the edges? What about bIsIslandPersistant?
-	for (TArray<int32>& IslandConstraintList : IslandToConstraints)
-	{
-		IslandConstraintList.Reset();
-	}
-
-	for(auto& PBDRigid : PBDRigids)	//@todo(ocohen): could go wide per island if we can get at the sets
-	{
-		const int32 Island = PBDRigid.Island();
-		if (Island >= 0)
-		{
-			FGraphNode& Node = Nodes[ParticleToNodeIndex[PBDRigid.Handle()]];
-			Node.Island = Island;
-			for (int32 ConstraintDataIndex : Node.Edges)
-			{
-				IslandToConstraints[Island].Add(ConstraintDataIndex);
-			}
-		}
-	}
-}
-
-
-void FPBDConstraintGraph::ReserveConstraints(const int32 NumConstraints)
-{
-	Edges.Reserve(Edges.Num() + NumConstraints);
-}
-
-
-void FPBDConstraintGraph::AddConstraint(const uint32 InContainerId, FConstraintHandle* InConstraintHandle, const TVec2<FGeometryParticleHandle*>& ConstrainedParticles)
-{
-	// Must have at least one constrained particle
-	check((ConstrainedParticles[0]) || (ConstrainedParticles[1]));
-
-	const int32 NewEdgeIndex = Edges.Num();
-	
-	InConstraintHandle->SetConstraintGraphIndex(NewEdgeIndex);
-
-	FGraphEdge NewEdge;
-	NewEdge.Data = { InContainerId, InConstraintHandle };
-
-	if (ConstrainedParticles[0] && !ParticleToNodeIndex.Contains(ConstrainedParticles[0]))
-	{
-		ParticleAdd(ConstrainedParticles[0]);
-	}
-	if (ConstrainedParticles[1] && !ParticleToNodeIndex.Contains(ConstrainedParticles[1]))
-	{
-		ParticleAdd(ConstrainedParticles[1]);
-	}
-
-	int32* PNodeIndex0 = (ConstrainedParticles[0])? ParticleToNodeIndex.Find(ConstrainedParticles[0]) : nullptr;
-	int32* PNodeIndex1 = (ConstrainedParticles[1])? ParticleToNodeIndex.Find(ConstrainedParticles[1]) : nullptr;
-	if (ensure(PNodeIndex0 || PNodeIndex1))
-	{
-		if (PNodeIndex0)
-		{
-			NewEdge.FirstNode = *PNodeIndex0;
-			Nodes[NewEdge.FirstNode].Particle = ConstrainedParticles[0];
-			Nodes[NewEdge.FirstNode].Edges.Add(NewEdgeIndex);
-			UpdatedNodes.Add(NewEdge.FirstNode);
-		}
-		if (PNodeIndex1)
-		{
-			NewEdge.SecondNode = *PNodeIndex1;
-			Nodes[NewEdge.SecondNode].Particle = ConstrainedParticles[1];
-			Nodes[NewEdge.SecondNode].Edges.Add(NewEdgeIndex);
-			UpdatedNodes.Add(NewEdge.SecondNode);
-		}
-
-		Edges.Add(MoveTemp(NewEdge));
-	}
-}
-
-int32 FPBDConstraintGraph::GetConstraintIndexFromNodeAndConstraintHandle(const FGraphNode& Node, const FConstraintHandle* ConstraintHandle)
-{
-	for (const int32 NodeEdgeIndex : Node.Edges)
-	{
-		const FConstraintData& EdgeData = Edges[NodeEdgeIndex].Data;
-		if (EdgeData.GetConstraintHandle() == ConstraintHandle)
-		{
-			return NodeEdgeIndex;
-		}
-	}
-
-	return INDEX_NONE;
-}
-
-void FPBDConstraintGraph::RemoveConstraint(const uint32 InContainerId, FConstraintHandle* InConstraintHandle, const TVector<FGeometryParticleHandle*, 2>& ConstrainedParticles)
-{
-	check(InConstraintHandle);
-
-	const int32* PNodeIndex0 = (ConstrainedParticles[0]) ? ParticleToNodeIndex.Find(ConstrainedParticles[0]) : nullptr;
-	const int32* PNodeIndex1 = (ConstrainedParticles[1]) ? ParticleToNodeIndex.Find(ConstrainedParticles[1]) : nullptr;
-
-	if (ensure(PNodeIndex0 || PNodeIndex1))
-	{ 
-		// we need to find the edge index that matches this constraint
-		// we could go through the entire edge list and find the matching InContainerId, 
-		// but each node has an edge array that will contain it 
-		int32 EdgeIndex = INDEX_NONE;
-		if (PNodeIndex0)
-		{
-			EdgeIndex = GetConstraintIndexFromNodeAndConstraintHandle(Nodes[*PNodeIndex0], InConstraintHandle);
-		}
-		if (PNodeIndex1 && EdgeIndex == INDEX_NONE)
-		{
-			EdgeIndex = GetConstraintIndexFromNodeAndConstraintHandle(Nodes[*PNodeIndex1], InConstraintHandle);
-		}
-
-		if (EdgeIndex != INDEX_NONE)
-		{
-			int32 Index0 = INDEX_NONE, Index1 = INDEX_NONE;
-			if (PNodeIndex0)
-			{
-				Nodes[*PNodeIndex0].Edges.Remove(EdgeIndex);
-				Index0 = *PNodeIndex0;
-			}
-			if (PNodeIndex1)
-			{
-				Nodes[*PNodeIndex1].Edges.Remove(EdgeIndex);
-				Index1 = *PNodeIndex1;
-			}
-			if (Edges[EdgeIndex].FirstNode == Index0 && Edges[EdgeIndex].SecondNode == Index1)
-			{
-				Edges[EdgeIndex] = FGraphEdge();
-			}
-
-		}
-
-		InConstraintHandle->SetConstraintGraphIndex(INDEX_NONE);
-	}
-}
-
-const typename FPBDConstraintGraph::FConstraintData& FPBDConstraintGraph::GetConstraintData(int32 ConstraintDataIndex) const
-{
-	return Edges[ConstraintDataIndex].Data;
-}
-
-
-void FPBDConstraintGraph::UpdateIslands(const TParticleView<FPBDRigidParticles>& PBDRigids, FPBDRigidsSOAs& Particles)
-{
-	// Maybe expose a memset style function for this instead of iterating
-	for (auto& PBDRigid : PBDRigids)
-	{
-		PBDRigid.Island() = INDEX_NONE;
-		// When enabling particle from a break, if the object state is static, then when Enabled the paticle doesn't get added
-		// to the constraint graph however the particle appears in GetNonDisabledDynamicView() so this ensure fires
-		if (!/*ensure*/(ParticleToNodeIndex.Contains(PBDRigid.Handle())))
-		{
-			ParticleAdd(PBDRigid.Handle());
-		}
-	}
-	ComputeIslands(PBDRigids, Particles);
-}
-
-
-DECLARE_CYCLE_STAT(TEXT("IslandGeneration2"), STAT_IslandGeneration2, STATGROUP_Chaos);
-
-
-void FPBDConstraintGraph::ComputeIslands(const TParticleView<FPBDRigidParticles>& PBDRigids, FPBDRigidsSOAs& Particles)
-{
-	SCOPE_CYCLE_COUNTER(STAT_IslandGeneration2);
-
-	int32 NextIsland = 0;
-	TArray<TSet<FGeometryParticleHandle*>> NewIslandParticles;
-	TArray<int32> NewIslandToSleepCount;
-
-	VisitToken++;
-	if (VisitToken==0)
-	{
-		VisitToken++;
-	}
-
-	// Instead of iterating over every node to reset Island, only iterate over the ones we care about for the following ComputeIslands algorithm to work 
-	PhysicsParallelFor(Edges.Num(), [&](int32 Idx)
-	{
-		const FGraphEdge& Edge = Edges[Idx];
-		Nodes[Edge.FirstNode].Island = INDEX_NONE;
-		if (Edge.SecondNode != -1)
-		{
-			Nodes[Edge.SecondNode].Island = INDEX_NONE;
-		}
-	});
-
-	IslandToData.Reset();
-
-	for (auto& Particle : PBDRigids)
-	{
-		auto* ParticleHandle = Particle.Handle();
-		int32 Idx = ParticleToNodeIndex[ParticleHandle];  // ryan - FAILS!
-		// selective reset of islands, don't reset if has been visited due to being edge connected to earlier processed node
-		if (Visited[Idx] && Visited[Idx] != VisitToken)
-		{
-			Nodes[Idx].Island = INDEX_NONE;
-			Visited[Idx] = VisitToken;
-		}
-
-		if (Nodes[Idx].Island >= 0)
-		{
-			// Island is already known - it was visited in ComputeIsland for a previous node
-			continue;
-		}
-
-		TSet<FGeometryParticleHandle*> SingleIslandParticles;
-		const bool bNeedsResim = ComputeIsland(Idx, NextIsland, SingleIslandParticles);
-
-		if (SingleIslandParticles.Num())
-		{
-			NewIslandParticles.SetNum(NextIsland + 1);
-			NewIslandParticles[NextIsland] = MoveTemp(SingleIslandParticles);
-			NextIsland++;
-			//if this is too slow when not doing resim, pass template in
-			IslandToData.AddDefaulted();
-			IslandToData.Last().bNeedsResim = bNeedsResim;
-		}
-	}
-
-	check(IslandToData.Num() == NextIsland);
-	IslandToConstraints.SetNum(NextIsland);
-
-	for (int32 EdgeIndex = 0; EdgeIndex < Edges.Num(); ++EdgeIndex)
-	{
-		const FGraphEdge& Edge = Edges[EdgeIndex];
-		int32 FirstIsland = (Edge.FirstNode != INDEX_NONE) ? Nodes[Edge.FirstNode].Island : INDEX_NONE;
-		int32 SecondIsland = (Edge.SecondNode != INDEX_NONE) ? Nodes[Edge.SecondNode].Island : INDEX_NONE;
-		check(FirstIsland == SecondIsland || FirstIsland == INDEX_NONE || SecondIsland == INDEX_NONE);
-
-		int32 Island = (FirstIsland != INDEX_NONE) ? FirstIsland : SecondIsland;
-		
-		// @todo(ccaulfield): should check(Island >= 0) when we disable particles properly
-		if (Island >= 0)
-		{
-			IslandToConstraints[Island].Add(EdgeIndex);
-		}
-	}
-
-	NewIslandToSleepCount.SetNum(NewIslandParticles.Num());
-
-	if (NewIslandParticles.Num())
-	{
-		for (int32 Island = 0; Island < NewIslandParticles.Num(); ++Island)
-		{
-			NewIslandToSleepCount[Island] = 0;
-			for (FGeometryParticleHandle* Particle : NewIslandParticles[Island])
-			{
-				FPBDRigidParticleHandle* PBDRigid = Particle->CastToRigidParticle();
-				if (PBDRigid && PBDRigid->ObjectState() != EObjectStateType::Kinematic)
-				{
-					PBDRigid->Island() = Island;
-				}
-			}
-		}
-		// Force consistent state if no previous islands
-		if (!IslandToParticles.Num())
-		{
-			for (int32 Island = 0; Island < NewIslandParticles.Num(); ++Island)
-			{
-				IslandToData[Island].bIsIslandPersistant = true;
-				bool bSleepState = true;
-
-				for (FGeometryParticleHandle* Particle : NewIslandParticles[Island])
-				{
-					if (Particle->ObjectState() != EObjectStateType::Static && !Particle->Sleeping())
-					{
-						bSleepState = false;
-						break;
-					}
-				}
-
-				for (FGeometryParticleHandle* Particle : NewIslandParticles[Island])
-				{
-					if (Particle->Sleeping() && !bSleepState)
-					{
-						Particles.ActivateParticle(Particle); 	//todo: record state change for array reorder
-					}
-
-					FPBDRigidParticleHandle* PBDRigid = Particle->CastToRigidParticle();
-					if(PBDRigid)
-					{
-						const EObjectStateType CurrState = PBDRigid->ObjectState();
-						if(CurrState == EObjectStateType::Kinematic || CurrState == EObjectStateType::Static)
-						{
-							// Statics and kinematics can't have sleeping states so don't attempt to set one.
-							break;
-						}
-
-						if (!Particle->Sleeping() && bSleepState)
-						{
-							Particles.DeactivateParticle(Particle); 	//todo: record state change for array reorder
-							PBDRigid->V() = FVec3(0);
-							PBDRigid->W() = FVec3(0);
-						}
-
-						PBDRigid->SetSleeping(bSleepState);
-					}
-
-					if (Particle->Sleeping())
-					{
-						Particles.DeactivateParticle(Particle);	//todo: record state change for array reorder (function could return true/false)
-					}
-				}
-			}
-		}
-
-		for (int32 Island = 0; Island < IslandToParticles.Num(); ++Island)
-		{
-			bool bIsSameIsland = true;
-
-			// Non-kinematic particles were removed from the island
-			int32 OtherIsland = -1;
-
-			for (FGeometryParticleHandle* Particle : IslandToParticles[Island])
-			{
-				if (CHAOS_ENSURE(Particle))
-				{
-					FPBDRigidParticleHandle* PBDRigid = Particle->CastToRigidParticle();
-
-					if (PBDRigid && PBDRigid->ObjectState() == EObjectStateType::Kinematic && PBDRigid->V().SizeSquared() > 0)
-					{
-						bIsSameIsland = false;
-						break;
-					}
-
-					const bool bIsDynamic = PBDRigid && PBDRigid->ObjectState() != EObjectStateType::Kinematic;
-
-					if (bIsDynamic && PBDRigid->PreObjectState() == EObjectStateType::Kinematic)
-					{
-						bIsSameIsland = false;
-						break;
-					}
-					
-					int32 TmpIsland = bIsDynamic ? PBDRigid->Island() : INDEX_NONE; //question: should we even store non dynamics in this array?
-
-					if (OtherIsland == INDEX_NONE && TmpIsland >= 0)
-					{
-						OtherIsland = TmpIsland;
-					}
-					else
-					{
-						if (TmpIsland >= 0 && OtherIsland != TmpIsland)
-						{
-							bIsSameIsland = false;
-							break;
-						}
-					}
-				}
-			}
-
-			// Kinematic particles were removed from the island. This needs to be called after OtherIsland is available.
-			if (bIsSameIsland && OtherIsland >= 0)
-			{
-				for (FGeometryParticleHandle* Particle : IslandToParticles[Island])
-				{
-					if (CHAOS_ENSURE(Particle))
-					{
-						FPBDRigidParticleHandle* PBDRigid = Particle->CastToRigidParticle();
-
-						// If an island has many kinematic particles, this could be slow.
-						if (PBDRigid && PBDRigid->ObjectState() == EObjectStateType::Kinematic && !NewIslandParticles[OtherIsland].Contains(PBDRigid))
-						{
-							bIsSameIsland = false;
-							break;
-						}
-					}
-				}
-			}
-
-			// A new object entered the island or the island is entirely new particles
-			if (bIsSameIsland && (OtherIsland == INDEX_NONE || NewIslandParticles[OtherIsland].Num() != IslandToParticles[Island].Num()))
-			{
-				bIsSameIsland = false;
-			}
-
-			// Find out if we need to activate island
-			if (bIsSameIsland)
-			{
-				NewIslandToSleepCount[OtherIsland] = IslandToSleepCount[Island];
-			}
-			else
-			{
-				int32 NumRigidsInNewIsland = 0;
-				if (ChaosPerfHackIgnoreSleepingContacts && NewIslandParticles.IsValidIndex(OtherIsland))
-				{
-					for (FGeometryParticleHandle* Particle : NewIslandParticles[OtherIsland])
-					{
-						if (Particle)
-						{
-							if (FPBDRigidParticleHandle* PBDRigid = Particle->CastToRigidParticle())
-							{
-								NumRigidsInNewIsland++;
-							}
-						}
-					}
-				}
-
-				for (TGeometryParticleHandle<FReal, 3>* Particle : IslandToParticles[Island])
-				{
-					if (CHAOS_ENSURE(Particle))
-					{
-						TPBDRigidParticleHandle<FReal, 3>* PBDRigid = Particle->CastToRigidParticle();
-						if (PBDRigid && PBDRigid->ObjectState() != EObjectStateType::Kinematic)
-						{
-							const bool bShouldActivateParticle = ChaosPerfHackIgnoreSleepingContacts ? !(NumRigidsInNewIsland == 1 && PBDRigid->Sleeping()) : true;
-							if (bShouldActivateParticle)
-							{
-								Particles.ActivateParticle(Particle);
-							}
-						}
-					}
-				}
-			}
-
-			// #BG Necessary? Should we ever not find an island?
-			if (OtherIsland != INDEX_NONE)
-			{
-				IslandToData[OtherIsland].bIsIslandPersistant = bIsSameIsland;
-			}
-		}
-	}
-
-	IslandToParticles.Reset();
-	IslandToParticles.Reserve(NewIslandParticles.Num());
-	for (int32 Island = 0; Island < NewIslandParticles.Num(); ++Island)
-	{
-		IslandToParticles.Emplace(NewIslandParticles[Island].Array());
-	}
-	IslandToSleepCount = MoveTemp(NewIslandToSleepCount);
-
-	check(IslandToParticles.Num() == IslandToSleepCount.Num());
-	check(IslandToParticles.Num() == IslandToConstraints.Num());
-	check(IslandToParticles.Num() == IslandToData.Num());
-	// @todo(ccaulfield): make a more complex unit test to check island integrity
-	checkSlow(DebugCheckGraph());
-}
-
-
-bool FPBDConstraintGraph::ComputeIsland(const int32 InNode, const int32 Island, TSet<FGeometryParticleHandle *>& ParticlesInIsland)
-{
-	TQueue<int32> NodeQueue;
-	NodeQueue.Enqueue(InNode);
-	bool bIslandNeedsToResim = false;
-	while (!NodeQueue.IsEmpty())
-	{
-		int32 NodeIndex;
-		NodeQueue.Dequeue(NodeIndex);
-		FGraphNode& Node = Nodes[NodeIndex];
-
-		if (Node.Island >= 0)
-		{
-			check(Node.Island == Island);
-			continue;
-		}
-
-		if(!bIslandNeedsToResim)
-		{
-			//if even one particle is soft/hard desync we must resim the entire island (when resim is used)
-			//seems cheap enough so just always do it, if slow pass resim template in here
-			bIslandNeedsToResim = Node.Particle->SyncState() != ESyncState::InSync;
-		}
-
-		FPBDRigidParticleHandle* RigidHandle = Node.Particle->CastToRigidParticle();
-		const bool isRigidDynamic = RigidHandle && RigidHandle->ObjectState() != EObjectStateType::Kinematic;  //??
-
-		if (RigidHandle && RigidHandle->Disabled())
-		{
-			continue;
-		}
-
-		ParticlesInIsland.Add(Node.Particle);
-		if (isRigidDynamic == false)
-		{
-			continue;
-		}
-
-		Node.Island = Island;
-		Visited[NodeIndex] = VisitToken;
-
-		for (const int32 EdgeIndex : Node.Edges)
-		{
-			const FGraphEdge& Edge = Edges[EdgeIndex];
-			int32 OtherNode = INDEX_NONE;
-			if (NodeIndex == Edge.FirstNode)
-			{
-				OtherNode = Edge.SecondNode;
-			}
-			if (NodeIndex == Edge.SecondNode)
-			{
-				OtherNode = Edge.FirstNode;
-			}
-			if (OtherNode != INDEX_NONE)
-			{
-				NodeQueue.Enqueue(OtherNode);
-			}
-		}
-	}
-
-	return bIslandNeedsToResim;
-}
-
-bool FPBDConstraintGraph::SleepInactive(const int32 Island, const TArrayCollectionArray<TSerializablePtr<FChaosPhysicsMaterial>>& PerParticleMaterialAttributes, const THandleArray<FChaosPhysicsMaterial>& SolverPhysicsMaterials)
-{
-	if (!ChaosSolverSleepEnabled)
-	{
-		return false;
-	}
-
-	FReal LinearSleepingThreshold = FLT_MAX;
-	FReal AngularSleepingThreshold = FLT_MAX;
-	int32 SleepCounterThreshold = 0;
-
-	const TArray<FGeometryParticleHandle*>& IslandParticles = GetIslandParticles(Island);
-	check(IslandParticles.Num());
-
-	if (!IslandToData[Island].bIsIslandPersistant)
-	{
-		return false;
-	}
-
-	int32& IslandSleepCount = IslandToSleepCount[Island];
-
-	FReal MaxLinearSpeed2 = 0.f;
-	FReal MaxAngularSpeed2 = 0.f;
-	int32 NumDynamicParticles = 0;
-
-	for (FGeometryParticleHandle* Particle : IslandParticles)
-	{
-		if (FPBDRigidParticleHandle* PBDRigid = Particle->CastToRigidParticle())
-		{
-			if (PBDRigid->ObjectState() == EObjectStateType::Dynamic)
-			{
-				NumDynamicParticles++;
-
-				if (PBDRigid->SleepType() == ESleepType::NeverSleep)
-				{
-					IslandSleepCount = 0;
-					return false;
-				}
-
-				const FReal LinearSpeed2 = PBDRigid->VSmooth().SizeSquared();
-				MaxLinearSpeed2 = FMath::Max(LinearSpeed2,MaxLinearSpeed2);
-
-				const FReal AngularSpeed2 = PBDRigid->WSmooth().SizeSquared();
-				MaxAngularSpeed2 = FMath::Max(AngularSpeed2,MaxAngularSpeed2);
-
-				bool bThresholdsSet = false;
-				if (TSerializablePtr<FChaosPhysicsMaterial> PhysicsMaterial = Particle->AuxilaryValue(PerParticleMaterialAttributes))
-				{
-					LinearSleepingThreshold = FMath::Min(LinearSleepingThreshold, PhysicsMaterial->SleepingLinearThreshold);
-					AngularSleepingThreshold = FMath::Min(AngularSleepingThreshold, PhysicsMaterial->SleepingAngularThreshold);
-					SleepCounterThreshold = FMath::Max(SleepCounterThreshold, PhysicsMaterial->SleepCounterThreshold);
-					bThresholdsSet = true;
-				}
-				else if (PBDRigid->ShapesArray().Num())
-				{
-					if (FPerShapeData* PerShapeData = PBDRigid->ShapesArray()[0].Get())
-					{
-						if (PerShapeData->GetMaterials().Num())
-						{
-							if (FChaosPhysicsMaterial* Material = SolverPhysicsMaterials.Get(PBDRigid->ShapesArray()[0].Get()->GetMaterials()[0].InnerHandle))
-							{
-								LinearSleepingThreshold = FMath::Min(LinearSleepingThreshold, Material->SleepingLinearThreshold);
-								AngularSleepingThreshold = FMath::Min(AngularSleepingThreshold, Material->SleepingAngularThreshold);
-								SleepCounterThreshold = FMath::Max(SleepCounterThreshold, Material->SleepCounterThreshold);
-								bThresholdsSet = true;
-							}
-						}
-					}
-				}
-
-				if (!bThresholdsSet)
-				{
-					LinearSleepingThreshold = FMath::Min(LinearSleepingThreshold, (FReal)ChaosSolverCollisionDefaultLinearSleepThresholdCVar);
-					AngularSleepingThreshold = FMath::Min(AngularSleepingThreshold, (FReal)ChaosSolverCollisionDefaultAngularSleepThresholdCVar);
-					SleepCounterThreshold = FMath::Max(SleepCounterThreshold, ChaosSolverCollisionDefaultSleepCounterThresholdCVar);
-				}
-			}
-		}
-	}
-
-	if (NumDynamicParticles == 0)
-	{
-		// prevent divide by zero - all particles much be sleeping/disabled already
-		return false;
-	}
-
-	const FReal MaxLinearSpeed = FMath::Sqrt(MaxLinearSpeed2);
-	const FReal MaxAngularSpeed = FMath::Sqrt(MaxAngularSpeed2);
-
-	if (MaxLinearSpeed < LinearSleepingThreshold && MaxAngularSpeed < AngularSleepingThreshold)
-	{
-		if (IslandSleepCount >= SleepCounterThreshold)
-		{
-			return true;
-		}
-		else
-		{
-			IslandSleepCount++;
-		}
-	}
-	else
-	{
-		//Reset sleep count since island is awake
-		IslandSleepCount = 0;
-	}
-
-	return false;
-}
-
-
-void FPBDConstraintGraph::WakeIsland(FPBDRigidsSOAs& Particles, const int32 Island)
-{
-	if (Island < IslandToParticles.Num())
-	{
-		for (FGeometryParticleHandle * Particle : IslandToParticles[Island])
-		{
-			FPBDRigidParticleHandle* PBDRigid = Particle->CastToRigidParticle();
-			if (PBDRigid && PBDRigid->ObjectState() != EObjectStateType::Kinematic)
-			{
-				if (PBDRigid->Sleeping())
-				{
-					PBDRigid->SetSleeping(false);
-					Particles.EnableParticle(Particle);
-				}
-			}
-		}
-		IslandToSleepCount[Island] = 0;
-	}
-}
-
-
-/*
-void FPBDConstraintGraph::ReconcileIslands()
-{
-	for (int32 Island = 0; Island < IslandToParticles.Num(); ++Island)
-	{
-		bool IsSleeping = true;
-		bool IsSet = false;
-		for (FGeometryParticleHandle* Particle : IslandToParticles[Island])
-		{
-			if (Particle->ObjectState() == EObjectStateType::Static)
-			{
-				continue;
-			}
-			if (!IsSet)
-			{
-				IsSet = true;
-				IsSleeping = Particle->Sleeping();
-			}
-			if (Particle->Sleeping() != IsSleeping)
-			{
-				WakeIsland(Island);
-				break;
-			}
-		}
-	}
-}
-*/
-
-void FPBDConstraintGraph::EnableParticle(FGeometryParticleHandle* Particle, const FGeometryParticleHandle* ParentParticle)
-{
-	if (ParentParticle)
-	{
-		const FPBDRigidParticleHandle* ParentPBDRigid = ParentParticle->CastToRigidParticle();
-		if(ParentPBDRigid && ParentPBDRigid->ObjectState() == EObjectStateType::Dynamic)
-		{
-			ParticleAdd(Particle);
-
-			FPBDRigidParticleHandle* ChildPBDRigid = Particle->CastToRigidParticle();
-			if(ChildPBDRigid && ChildPBDRigid->ObjectState() == EObjectStateType::Dynamic)
-			{
-				const int32 Island = ParentPBDRigid->Island();
-				ChildPBDRigid->Island() = Island;
-
-				// If our parent had a valid island, add the child to it.
-				if (IslandToParticles.IsValidIndex(Island))
-				{
-					IslandToParticles[Island].Add(Particle);
-				}
-
-				const bool SleepState = ParentPBDRigid->Sleeping();
-				ChildPBDRigid->SetSleeping(SleepState);	//todo: need to let evolution know to reorder arrays
-			}
-			else
-			{
-				ensure(false);	//this should never happen
-			}
-		}
-	}
-}
-
-
-void FPBDConstraintGraph::DisableParticle(FGeometryParticleHandle* Particle)
-{
-	FPBDRigidParticleHandle* PBDRigid = Particle->CastToRigidParticle();
-	if(PBDRigid && PBDRigid->ObjectState() != EObjectStateType::Kinematic)
-	{
-		const int32 Island = PBDRigid->Island();
-		if (Island != INDEX_NONE)
-		{
-			PBDRigid->Island() = INDEX_NONE;
-
-			// @todo(ccaulfield): optimize
-			if (ensure(IslandToParticles.IsValidIndex(Island)))
-			{
-				int32 IslandParticleArrayIdx = IslandToParticles[Island].Find(Particle);
-				check(IslandParticleArrayIdx != INDEX_NONE);
-				IslandToParticles[Island].RemoveAtSwap(IslandParticleArrayIdx);
-			}
-		}
-
-	}
-	else
-	{
-		// Kinematic & Static particles are included in IslandToParticles, however we cannot use islands to look them up.
-		// TODO find faster removal method?
-		for (TArray<FGeometryParticleHandle*>& IslandParticles : IslandToParticles)
-		{
-			IslandParticles.RemoveSingleSwap(Particle);
-		}
-	}
-
-	ParticleRemove(Particle);
-}
-
-
-void FPBDConstraintGraph::DisableParticles(const TSet<FGeometryParticleHandle *>& Particles)
-{
-	// @todo(ccaulfield): optimize
-	for (FGeometryParticleHandle* Particle : Particles)
-	{
-		DisableParticle(Particle);
-	}
-}
-
-#if !UE_BUILD_SHIPPING
-bool FPBDConstraintGraph::DebugCheckGraph() const
-{
-	bool bIsValid = true;
-
-	// Check that no particles are in multiple islands
-	TSet<FGeometryParticleHandle*> IslandParticlesUnionSet;
-	for (int32 Island = 0; Island < IslandToParticles.Num(); ++Island)
-	{
-		TSet<FGeometryParticleHandle*> IslandParticlesSet = TSet<FGeometryParticleHandle*>(IslandToParticles[Island]);
-		TSet<FGeometryParticleHandle*> IslandParticlesIntersectSet = IslandParticlesUnionSet.Intersect(IslandParticlesSet);
-		if (IslandParticlesIntersectSet.Num() > 0)
-		{
-			// This islands contains particles that were in a previous island.
-			// This is ok only if those particles are static
-			for (FGeometryParticleHandle* Particle : IslandParticlesIntersectSet)
-			{
-				if (Particle->CastToRigidParticle() && Particle->ObjectState() == EObjectStateType::Dynamic)
-				{
-					UE_LOG(LogChaos, Error, TEXT("Island %d contains non-static particle that is also in another Island"), Island);	//todo: add better logging for bad particle
-					bIsValid = false;
-				}
-			}
-		}
-		IslandParticlesUnionSet = IslandParticlesUnionSet.Union(IslandParticlesSet);
-	}
-
-	// Check that no constraints are in multiple islands
-	TSet<int32> IslandConstraintDataUnionSet;
-	IslandConstraintDataUnionSet.Reserve(Edges.Num());
-	for (int32 Island = 0; Island < IslandToConstraints.Num(); ++Island)
-	{
-		TSet<int32> IslandConstraintDataSet = TSet<int32>(IslandToConstraints[Island]);
-		TSet<int32> IslandConstraintDataIntersectSet = IslandConstraintDataUnionSet.Intersect(IslandConstraintDataSet);
-		if (IslandConstraintDataIntersectSet.Num() > 0)
-		{
-			// This islands contains constraints that were in a previous island
-			UE_LOG(LogChaos, Error, TEXT("Island %d contains Constraints in another Island"), Island);
-			bIsValid = false;
-		}
-		IslandConstraintDataUnionSet = IslandConstraintDataUnionSet.Union(IslandConstraintDataSet);
-	}
-
-	// Check that no constraints connect dynamic particles in different islands
-	// (non-dynamic particles can be in multiple islands)
-	for (int32 Island = 0; Island < IslandToConstraints.Num(); ++Island)
-	{
-		const TArray<int32> IslandEdges = IslandToConstraints[Island];
-		for (int32 EdgeIndex : IslandEdges)
-		{
-			const FGraphEdge& Edge = Edges[EdgeIndex];
-
-			if (Edge.SecondNode == INDEX_NONE)
-			{
-				// We add "edges" for single-particle constraints as well. In this case, they have no second node
-				continue;
-			}
-
-			const FGraphNode& Node0 = Nodes[Edge.FirstNode];
-			const FGraphNode& Node1 = Nodes[Edge.SecondNode];
-			if ((Node0.Particle != nullptr) && (Node1.Particle != nullptr) && FConstGenericParticleHandle(Node0.Particle)->IsDynamic() && FConstGenericParticleHandle(Node1.Particle)->IsDynamic())
-			{
-				if (Node0.Island != Node1.Island)
-				{
-					UE_LOG(LogChaos, Error, TEXT("Island %d Constraint %d connects particles in different islands %d %d"), Island, EdgeIndex, Node0.Island, Node1.Island);
-					bIsValid = false;
-				}
-			}
-		}
-	}
-
-	for (const FGraphNode& Node : Nodes)
-	{
-		if ((Node.Particle != nullptr) && FConstGenericParticleHandle(Node.Particle)->IsDynamic() && (FConstGenericParticleHandle(Node.Particle)->Island() != Node.Island))
-		{
-			UE_LOG(LogChaos, Error, TEXT("Particle %d island mismatch %d %d"), Node.Particle->ParticleID().LocalID, FConstGenericParticleHandle(Node.Particle)->Island(), Node.Island);
-			bIsValid = false;
-		}
-	}
-
-	return bIsValid;
-}
-#endif
diff --git a/Engine/Source/Runtime/Experimental/Chaos/Private/Chaos/PBDConstraintRule.cpp b/Engine/Source/Runtime/Experimental/Chaos/Private/Chaos/PBDConstraintRule.cpp
--- a/Engine/Source/Runtime/Experimental/Chaos/Private/Chaos/PBDConstraintRule.cpp
+++ b/Engine/Source/Runtime/Experimental/Chaos/Private/Chaos/PBDConstraintRule.cpp
@@ -2,6 +2,7 @@
 
 #include "Chaos/PBDConstraintRule.h"
 
+#include "Chaos/Island/SolverIsland.h"
 #include "Chaos/PBDCollisionConstraints.h"
 #include "Chaos/PBDJointConstraints.h"
 #include "Chaos/PBDPositionConstraints.h"
@@ -11,32 +12,107 @@
 
 namespace Chaos
 {
+	int32 ChaosShockPropagationVelocityPerLevelIterations = 1;
+	int32 ChaosShockPropagationPositionPerLevelIterations = 1;
+	FAutoConsoleVariableRef CVarChaosShockPropagationPositionPerLevelIterations(TEXT("p.Chaos.ShockPropagation.Position.PerLevelIterations"), ChaosShockPropagationPositionPerLevelIterations, TEXT(""));
+	FAutoConsoleVariableRef CVarChaosShockPropagationVelocityPerLevelIterations(TEXT("p.Chaos.ShockPropagation.Velocity.PerLevelIterations"), ChaosShockPropagationVelocityPerLevelIterations, TEXT(""));
 
-	template<class T_CONSTRAINTS>
-	TPBDConstraintGraphRuleImpl<T_CONSTRAINTS>::TPBDConstraintGraphRuleImpl(FConstraints& InConstraints, int32 InPriority)
-		: FPBDConstraintGraphRule(InPriority)
+	int32 ChaosCollisionColorMinParticles = 2000;
+	FAutoConsoleVariableRef CVarChaosCollisionColorMinParticles(TEXT("p.Chaos.Collision.Color.MinParticles"), ChaosCollisionColorMinParticles, TEXT(""));
+
+	template<class ConstraintType>
+	TSimpleConstraintRule<ConstraintType>::TSimpleConstraintRule(int32 InPriority, FConstraints& InConstraints)
+		: FSimpleConstraintRule(InPriority)
 		, Constraints(InConstraints)
-		, ConstraintGraph(nullptr)
 	{
 	}
 
+	template<class ConstraintType>
+	TSimpleConstraintRule<ConstraintType>::~TSimpleConstraintRule()
+	{
+	}
+
+	template<class ConstraintType>
+	void TSimpleConstraintRule<ConstraintType>::PrepareTick()
+	{
+		Constraints.PrepareTick();
+	}
+
+	template<class ConstraintType>
+	void TSimpleConstraintRule<ConstraintType>::UnprepareTick()
+	{
+		Constraints.UnprepareTick();
+	}
+
+	template<class ConstraintType>
+	void TSimpleConstraintRule<ConstraintType>::UpdatePositionBasedState(const FReal Dt)
+	{
+		return Constraints.UpdatePositionBasedState(Dt);
+	}
 
-	template<class T_CONSTRAINTS>
-	void TPBDConstraintGraphRuleImpl<T_CONSTRAINTS>::BindToGraph(FPBDConstraintGraph& InContactGraph, uint32 InContainerId)
+	template<class ConstraintType>
+	void TSimpleConstraintRule<ConstraintType>::BindToDatas(FPBDIslandSolverData& InSolverDatas, const uint32 InContainerId)
 	{
 		Constraints.SetContainerId(InContainerId);
-		ConstraintGraph = &InContactGraph;
+
+		SolverData = &InSolverDatas;
+		if (SolverData != nullptr)
+		{
+			SolverData->template AddConstraintDatas<ConstraintType>(Constraints.GetContainerId());
+		}
 	}
 
-	template<class T_CONSTRAINTS>
-	void TPBDConstraintGraphRuleImpl<T_CONSTRAINTS>::UpdatePositionBasedState(const FReal Dt)
+	template<class ConstraintType>
+	void TSimpleConstraintRule<ConstraintType>::GatherSolverInput(const FReal Dt)
 	{
-		Constraints.UpdatePositionBasedState(Dt);
+		if(SolverData)
+		{
+			Constraints.SetNumIslandConstraints(Constraints.NumConstraints(), *SolverData);
+			Constraints.GatherInput(Dt, *SolverData);
+		}
 	}
 
+	template<class ConstraintType>
+	void TSimpleConstraintRule<ConstraintType>::ScatterSolverOutput(const FReal Dt)
+	{
+		if(SolverData) Constraints.ScatterOutput(Dt, *SolverData);
+	}
 
-	template<class T_CONSTRAINTS>
-	void TPBDConstraintGraphRuleImpl<T_CONSTRAINTS>::AddToGraph()
+	template<class ConstraintType>
+	bool TSimpleConstraintRule<ConstraintType>::ApplyConstraints(const FReal Dt, const int32 It, const int32 NumIts)
+	{
+		return SolverData ? Constraints.ApplyPhase1(Dt, It, NumIts, *SolverData) : false;
+	}
+
+	template<class ConstraintType>
+	bool TSimpleConstraintRule<ConstraintType>::ApplyPushOut(const FReal Dt, const int32 It, const int32 NumIts)
+	{
+		return SolverData ? Constraints.ApplyPhase2(Dt, It, NumIts, *SolverData) : false;
+	}
+	
+	template<class ConstraintType>
+	TPBDConstraintGraphRuleImpl<ConstraintType>::TPBDConstraintGraphRuleImpl(FConstraints& InConstraints, int32 InPriority)
+		: FPBDConstraintGraphRule(InPriority)
+		, Constraints(InConstraints)
+		, ConstraintGraph(nullptr)
+	{
+	}
+	
+	template<class ConstraintType>
+	void TPBDConstraintGraphRuleImpl<ConstraintType>::BindToGraph(FPBDConstraintGraph& InContactGraph, uint32 InContainerId)
+	{
+		Constraints.SetContainerId(InContainerId);
+		ConstraintGraph = &InContactGraph;
+	}
+
+	template<class ConstraintType>
+	void TPBDConstraintGraphRuleImpl<ConstraintType>::UpdatePositionBasedState(const FReal Dt)
+	{
+		Constraints.UpdatePositionBasedState(Dt);
+	}
+	
+	template<class ConstraintType>
+	void TPBDConstraintGraphRuleImpl<ConstraintType>::AddToGraph()
 	{
 		ConstraintGraph->ReserveConstraints(Constraints.NumConstraints());
 
@@ -47,15 +123,274 @@ namespace Chaos
 				ConstraintGraph->AddConstraint(GetContainerId(), ConstraintHandle, ConstraintHandle->GetConstrainedParticles());
 			}
 		}
+	}
+
+	template<class ConstraintType>
+	TPBDConstraintIslandRule<ConstraintType>::TPBDConstraintIslandRule(FConstraints& InConstraints, int32 InPriority)
+		: TPBDConstraintGraphRuleImpl<ConstraintType>(InConstraints, InPriority)
+	{
+	}
+
+	template<class ConstraintType>
+	TPBDConstraintIslandRule<ConstraintType>::~TPBDConstraintIslandRule()
+	{
+	}
+
+	template<class ConstraintType>
+	void TPBDConstraintIslandRule<ConstraintType>::GatherSolverInput(const FReal Dt, int32 Island)
+	{
+		if(FPBDIslandSolver* IslandSolver = ConstraintGraph->GetSolverIsland(Island))
+		{
+			const TArray<FConstraintHandle*>& IslandConstraints = ConstraintGraph->GetIslandConstraints(Island);
+
+			// This will reset the number od constraints inside the solver datas. For now we keep this function since according
+			// to the constraints we can use the handles, the indices or the container. when only the container will be used
+			// we can replace this call by :
+			// IslandSolver->template GetConstraintContainer<typename ConstraintType::FSolverConstraintContainerType>(ContainerId)->Reset();
+			Constraints.SetNumIslandConstraints(IslandConstraints.Num(), *IslandSolver);
+
+			for (FConstraintHandle* ConstraintHandle : IslandConstraints)
+			{
+				if (ConstraintHandle->GetContainerId() == GetContainerId())
+				{
+					FConstraintContainerHandle* Constraint = ConstraintHandle->As<FConstraintContainerHandle>();
+
+					// Note we are building the SolverBodies as we go, in the order that we visit them. Each constraint
+					// references two bodies, so we won't strictly be acessing only in cache order, but it's about as good as it can be.
+					if (Constraint->IsEnabled())
+					{
+						// @todo(chaos): we should provide Particle Levels in the island rule as well (see TPBDConstraintColorRule)
+						Constraint->GatherInput(Dt, INDEX_NONE, INDEX_NONE, *IslandSolver);
+					}
+				}
+			}
+		}
+	}
 
+	template<class ConstraintType>
+	void TPBDConstraintIslandRule<ConstraintType>::ScatterSolverOutput(const FReal Dt, int32 Island)
+	{
+		if(FPBDIslandSolver* IslandSolver = ConstraintGraph->GetSolverIsland(Island))
+		{
+			Constraints.ScatterOutput(Dt, *IslandSolver);
+		}
+	}
+	
+	template<class ConstraintType>
+	bool TPBDConstraintIslandRule<ConstraintType>::ApplyConstraints(const FReal Dt, int32 Island, const int32 It, const int32 NumIts)
+	{
+		FPBDIslandSolver* IslandSolver = ConstraintGraph->GetSolverIsland(Island);
+		return IslandSolver ? Constraints.ApplyPhase1Serial(Dt, It, NumIts, *IslandSolver) : false;
+	}
+
+	template<class ConstraintType>
+	bool TPBDConstraintIslandRule<ConstraintType>::ApplyPushOut(const FReal Dt, int32 Island, const int32 It, const int32 NumIts)
+	{
+		FPBDIslandSolver* IslandSolver = ConstraintGraph->GetSolverIsland(Island);
+		return IslandSolver ? Constraints.ApplyPhase2Serial(Dt, It, NumIts, *IslandSolver) : false;
+	}
+
+	template<class ConstraintType>
+	void TPBDConstraintIslandRule<ConstraintType>::InitializeAccelerationStructures()
+	{
+		ConstraintGraph->template AddConstraintDatas<ConstraintType>(Constraints.GetContainerId());
 	}
 
-	template<class T_CONSTRAINTS>
-	int32 TPBDConstraintGraphRuleImpl<T_CONSTRAINTS>::NumConstraints() const
+	template<class ConstraintType>
+	void TPBDConstraintIslandRule<ConstraintType>::UpdateAccelerationStructures(const FReal Dt, const int32 Island)
+	{
+	}
+	
+	template<class ConstraintType>
+	int32 TPBDConstraintGraphRuleImpl<ConstraintType>::NumConstraints() const
 	{ 
 		return Constraints.NumConstraints(); 
 	}
 
+	template<class ConstraintType>
+	TPBDConstraintColorRule<ConstraintType>::TPBDConstraintColorRule(FConstraints& InConstraints, int32 InPriority)
+		: TPBDConstraintGraphRuleImpl<ConstraintType>(InConstraints, InPriority)
+		, ConstraintSets()
+		, bUseColor(false)
+	{
+	}
+
+	template<class ConstraintType>
+	TPBDConstraintColorRule<ConstraintType>::~TPBDConstraintColorRule()
+	{
+	}
+
+	template<class ConstraintType>
+	void TPBDConstraintColorRule<ConstraintType>::GatherSolverInput(const FReal Dt, int32 Island)
+	{
+		if(FPBDIslandSolver* IslandSolver = ConstraintGraph->GetSolverIsland(Island))
+		{
+			const typename FPBDConstraintColor::FLevelToColorToConstraintListMap& LevelToColorToConstraintListMap = GraphColor.GetIslandLevelToColorToConstraintListMap(Island);
+			int32 MaxColor = GraphColor.GetIslandMaxColor(Island);
+			int32 MaxLevel = GraphColor.GetIslandMaxLevel(Island);
+
+			TArray<TPair<int32, int32>>& IslandConstraintSets = ConstraintSets[Island];
+			IslandConstraintSets.Reset(MaxLevel * MaxColor);	// Pessimistic array size - we could store the actual required size in coloring algorithm
+			int32 ConstraintSetEnd = 0;
+
+			Constraints.SetNumIslandConstraints(GraphColor.NumIslandEdges(Island), *IslandSolver);
+
+			for (int32 Level = 0; Level <= MaxLevel; ++Level)
+			{
+				const int32 OldConstraintSetEnd = ConstraintSetEnd;
+
+				for (int32 Color = 0; Color <= MaxColor; ++Color)
+				{
+					if (LevelToColorToConstraintListMap[Level].Contains(Color) && LevelToColorToConstraintListMap[Level][Color].Num())
+					{
+						// Calculate the range of indices for this color as a set of independent contacts
+						TPair<int32, int32> ColorConstrainSet(ConstraintSetEnd, ConstraintSetEnd);
+
+						const TArray<FConstraintContainerHandle*>& ConstraintHandles = GetLevelColorConstraints(LevelToColorToConstraintListMap, Level, Color);
+						for (FConstraintContainerHandle* Constraint : ConstraintHandles)
+						{
+							if (Constraint->IsEnabled())
+							{
+								// Levels that should be assigned to the bodies for shock propagation
+								// @todo(chaos): optimize the lookup
+								const TVector<TGeometryParticleHandle<FReal, 3>*, 2> ConstrainedParticles = Constraint->GetConstrainedParticles();
+								const int32 Particle0Level = GraphColor.GetParticleLevel(ConstrainedParticles[0]);
+								const int32 Particle1Level = GraphColor.GetParticleLevel(ConstrainedParticles[1]);
+
+								// Note we are building the SolverBodies as we go, in the order that we visit them. Each constraint
+								// references two bodies, so we won't strictly be acessing only in cache order, but it's about as good as it can be.
+								Constraint->GatherInput(Dt, Particle0Level, Particle1Level, *IslandSolver);
+
+								// Update the current constraint set of this color
+								ColorConstrainSet.Value = ++ConstraintSetEnd;
+							}
+						}
+
+						// Remember the set of constraints of this color
+						if (bUseColor)
+						{
+							IslandConstraintSets.Add(ColorConstrainSet);
+						}
+					}
+				}
+
+				const int32 NumAdded = ConstraintSetEnd - OldConstraintSetEnd;
+				if (NumAdded > 1)
+				{
+					static int32 sHere = 0;
+					++sHere;
+				}
+			}
+
+			// If we aren't coloring, we have a single group of all constraints (they have been created in level order above)
+			if (!bUseColor)
+			{
+				const TPair<int32, int32> ConstrainSet(0, ConstraintSetEnd);
+				IslandConstraintSets.Add(ConstrainSet);
+			}
+		}
+	}
+
+	template<class ConstraintType>
+	void TPBDConstraintColorRule<ConstraintType>::ScatterSolverOutput(const FReal Dt, int32 Island)
+	{
+		if(FPBDIslandSolver* IslandSolver = ConstraintGraph->GetSolverIsland(Island))
+		{
+			const TArray<TPair<int32, int32>>& IslandConstraintSets = ConstraintSets[Island];
+			for (const TPair<int32, int32>& ConstraintSet : IslandConstraintSets)
+			{
+				Constraints.ScatterOutput(Dt, ConstraintSet.Key, ConstraintSet.Value, *IslandSolver);
+			}
+		}
+	}
+
+	template<class ConstraintType>
+	void TPBDConstraintColorRule<ConstraintType>::ApplySwept(const FReal Dt, int32 Island)
+	{
+		if(FPBDIslandSolver* IslandSolver = ConstraintGraph->GetSolverIsland(Island))
+		{
+			Constraints.ApplySwept(Dt, *IslandSolver);
+		}
+	}
+
+	template<class ConstraintType>
+	bool TPBDConstraintColorRule<ConstraintType>::ApplyConstraints(const FReal Dt, int32 Island, const int32 It, const int32 NumIts)
+	{
+		bool bNeedsAnotherIteration = false;
+		if(FPBDIslandSolver* IslandSolver = ConstraintGraph->GetSolverIsland(Island))
+		{
+			const TArray<TPair<int32, int32>>& IslandConstraintSets = ConstraintSets[Island];
+			if (!bUseColor)
+			{
+				for (const TPair<int32, int32>& ConstraintSet : IslandConstraintSets)
+				{
+					bNeedsAnotherIteration |= Constraints.ApplyPhase1Serial(Dt, It, NumIts, ConstraintSet.Key, ConstraintSet.Value, *IslandSolver);
+				}
+			}
+			else
+			{
+				for (const TPair<int32, int32>& ConstraintSet : IslandConstraintSets)
+				{
+					bNeedsAnotherIteration |= Constraints.ApplyPhase1Parallel(Dt, It, NumIts, ConstraintSet.Key, ConstraintSet.Value, *IslandSolver);
+				}
+			}
+		}
+		return bNeedsAnotherIteration;
+	}
+
+	template<class ConstraintType>
+	bool TPBDConstraintColorRule<ConstraintType>::ApplyPushOut(const FReal Dt, int32 Island, const int32 It, const int32 NumIts)
+	{
+		bool bNeedsAnotherIteration = false;
+		if(FPBDIslandSolver* IslandSolver = ConstraintGraph->GetSolverIsland(Island))
+		{
+			const TArray<TPair<int32, int32>>& IslandConstraintSets = ConstraintSets[Island];
+			if (!bUseColor)
+			{
+				for (const TPair<int32, int32>& ConstraintSet : IslandConstraintSets)
+				{
+					bNeedsAnotherIteration |= Constraints.ApplyPhase2Serial(Dt, It, NumIts, ConstraintSet.Key, ConstraintSet.Value, *IslandSolver);
+				}
+			}
+			else
+			{
+				for (const TPair<int32, int32>& ConstraintSet : IslandConstraintSets)
+				{
+					bNeedsAnotherIteration |= Constraints.ApplyPhase2Parallel(Dt, It, NumIts, ConstraintSet.Key, ConstraintSet.Value, *IslandSolver);
+				}
+			}
+		}
+		return bNeedsAnotherIteration;
+	}
+
+	template<class ConstraintType>
+	void TPBDConstraintColorRule<ConstraintType>::InitializeAccelerationStructures()
+	{
+		GraphColor.InitializeColor(*ConstraintGraph);
+
+		ConstraintSets.SetNum(GraphColor.NumIslands());
+		ConstraintGraph->template AddConstraintDatas<ConstraintType>(Constraints.GetContainerId());
+	}
+
+	template<class ConstraintType>
+	void TPBDConstraintColorRule<ConstraintType>::UpdateAccelerationStructures(const FReal Dt, const int32 Island)
+	{
+		// Decide if we want to attempt to solve the constraints in parallel - this requires coloring the graph which
+		// is expensive, so we need to get a good number of contacts per color (per level) for it to be worthwhile.
+		// @todo(chaos): if enabled here, it should get disabled again if the level/coloring doesn't produce enough parallelism
+		bUseColor = (ConstraintGraph->GetIslandParticles(Island).Num() > ChaosCollisionColorMinParticles);
+
+		// @todo(chaos): disable coloring, but still generate levels and ordering, when we don't need parallelism
+		GraphColor.ComputeColor(Dt, Island, *ConstraintGraph, Constraints.GetContainerId());
+	}
+
+	template<class ConstraintType>
+	void TPBDConstraintColorRule<ConstraintType>::SetUseContactGraph(const bool bInUseContactGraph)
+	{
+		GraphColor.SetUseContactGraph(bInUseContactGraph);
+	}
+
+
 	template class TSimpleConstraintRule<FPBDCollisionConstraints>;
 	template class TSimpleConstraintRule<FPBDJointConstraints>;
 	template class TSimpleConstraintRule<FPBDRigidSpringConstraints>;
diff --git a/Engine/Source/Runtime/Experimental/Chaos/Private/Chaos/PBDJointConstraints.cpp b/Engine/Source/Runtime/Experimental/Chaos/Private/Chaos/PBDJointConstraints.cpp
--- a/Engine/Source/Runtime/Experimental/Chaos/Private/Chaos/PBDJointConstraints.cpp
+++ b/Engine/Source/Runtime/Experimental/Chaos/Private/Chaos/PBDJointConstraints.cpp
@@ -2,6 +2,7 @@
 #include "Chaos/PBDJointConstraints.h"
 #include "Chaos/ChaosDebugDraw.h"
 #include "Chaos/DebugDrawQueue.h"
+#include "Chaos/Evolution/SolverDatas.h"
 #include "Chaos/Joint/ChaosJointLog.h"
 #include "Chaos/Joint/ColoringGraph.h"
 #include "Chaos/Joint/JointConstraintsCVars.h"
@@ -11,7 +12,6 @@
 #include "Chaos/Utilities.h"
 #include "ChaosLog.h"
 #include "ChaosStats.h"
-#include "ProfilingDebugging/CsvProfiler.h"
 
 #include "HAL/IConsoleManager.h"
 
@@ -24,11 +24,12 @@ namespace Chaos
 	DECLARE_CYCLE_STAT(TEXT("Joints::Sort"), STAT_Joints_Sort, STATGROUP_ChaosJoint);
 	DECLARE_CYCLE_STAT(TEXT("Joints::PrepareTick"), STAT_Joints_PrepareTick, STATGROUP_ChaosJoint);
 	DECLARE_CYCLE_STAT(TEXT("Joints::UnprepareTick"), STAT_Joints_UnprepareTick, STATGROUP_ChaosJoint);
-	DECLARE_CYCLE_STAT(TEXT("Joints::PrepareIterations"), STAT_Joints_PrepareIteration, STATGROUP_ChaosJoint);
-	DECLARE_CYCLE_STAT(TEXT("Joints::UnprepareIteration"), STAT_Joints_UnprepareIteration, STATGROUP_ChaosJoint);
+	DECLARE_CYCLE_STAT(TEXT("Joints::Gather"), STAT_Joints_Gather, STATGROUP_ChaosJoint);
+	DECLARE_CYCLE_STAT(TEXT("Joints::Scatter"), STAT_Joints_Scatter, STATGROUP_ChaosJoint);
 	DECLARE_CYCLE_STAT(TEXT("Joints::Apply"), STAT_Joints_Apply, STATGROUP_ChaosJoint);
 	DECLARE_CYCLE_STAT(TEXT("Joints::ApplyPushOut"), STAT_Joints_ApplyPushOut, STATGROUP_ChaosJoint);
 
+
 	//
 	// Constraint Handle
 	//
@@ -40,7 +41,7 @@ namespace Chaos
 
 	
 	FPBDJointConstraintHandle::FPBDJointConstraintHandle(FConstraintContainer* InConstraintContainer, int32 InConstraintIndex)
-		: TContainerConstraintHandle<FPBDJointConstraints>(InConstraintContainer, InConstraintIndex)
+		: TIndexedContainerConstraintHandle<FPBDJointConstraints>(InConstraintContainer, InConstraintIndex)
 	{
 	}
 
@@ -68,12 +69,6 @@ namespace Chaos
 		return ConcreteContainer()->GetConstraintColor(ConstraintIndex);
 	}
 
-
-	bool FPBDJointConstraintHandle::IsConstraintEnabled() const
-	{
-		return ConcreteContainer()->IsConstraintEnabled(ConstraintIndex);
-	}
-
 	bool FPBDJointConstraintHandle::IsConstraintBreaking() const
 	{
 		return ConcreteContainer()->IsConstraintBreaking(ConstraintIndex);
@@ -94,6 +89,11 @@ namespace Chaos
 		return ConcreteContainer()->ClearDriveTargetChanged(ConstraintIndex);
 	}
 
+	bool FPBDJointConstraintHandle::IsConstraintEnabled() const
+	{
+		return ConcreteContainer()->IsConstraintEnabled(ConstraintIndex);
+	}
+
 	FVec3 FPBDJointConstraintHandle::GetLinearImpulse() const
 	{
 		return ConcreteContainer()->GetConstraintLinearImpulse(ConstraintIndex);
@@ -124,7 +124,6 @@ namespace Chaos
 		return ConcreteContainer()->GetConstraintResimType(ConstraintIndex);
 	}
 
-	
 	const FPBDJointSettings& FPBDJointConstraintHandle::GetSettings() const
 	{
 		return ConcreteContainer()->GetConstraintSettings(ConstraintIndex);
@@ -145,6 +144,11 @@ namespace Chaos
 		return ConcreteContainer()->SetConstraintEnabled(ConstraintIndex, bInEnabled);
 	}
 
+	void FPBDJointConstraintHandle::GatherInput(const FReal Dt, const int32 Particle0Level, const int32 Particle1Level, FPBDIslandSolverData& SolverData)
+	{
+		ConcreteContainer()->GatherInput(Dt, ConstraintIndex, Particle0Level, Particle1Level, SolverData);
+	}
+
 	//
 	// Constraint Settings
 	//
@@ -301,8 +305,7 @@ namespace Chaos
 		, IslandSize(0)
 		, bDisabled(false)
 		, bBreaking(false)
-		, bDriveTargetChanged(false)
-		, LinearImpulse(FVec3(0))
+		, bDriveTargetChanged(false), LinearImpulse(FVec3(0))
 		, AngularImpulse(FVec3(0))
 	{
 	}
@@ -352,12 +355,9 @@ namespace Chaos
 
 	
 	FPBDJointConstraints::FPBDJointConstraints(const FPBDJointSolverSettings& InSettings)
-		: FPBDConstraintContainer(EConstraintContainerType::Joint)
+		: FPBDIndexedConstraintContainer(FConstraintContainerHandle::StaticType())
 		, Settings(InSettings)
 		, bJointsDirty(false)
-		, bUpdateVelocityInApplyConstraints(false)
-		, PreApplyCallback(nullptr)
-		, PostApplyCallback(nullptr)
 		, SolverType(EConstraintSolverType::GbfPbd)
 	{
 	}
@@ -472,10 +472,11 @@ namespace Chaos
 		{
 			for (FConstraintHandle* ConstraintHandle : RemovedParticle->ParticleConstraints())
 			{
-				if (ConstraintHandle->As<FPBDJointConstraintHandle>())
+				if (FPBDJointConstraintHandle* JointHandle = ConstraintHandle->As<FPBDJointConstraintHandle>())
 				{
-					ConstraintHandle->SetEnabled(false); // constraint lifespan is managed by the proxy
-					int ConstraintIndex = ConstraintHandle->GetConstraintIndex();
+					JointHandle->SetEnabled(false); // constraint lifespan is managed by the proxy
+
+					int ConstraintIndex = JointHandle->GetConstraintIndex();
 					if (ConstraintIndex != INDEX_NONE)
 					{
 						if (ConstraintParticles[ConstraintIndex][0] == RemovedParticle)
@@ -492,6 +493,7 @@ namespace Chaos
 		}
 	}
 
+
 	void FPBDJointConstraints::SortConstraints()
 	{
 		// Sort constraints so that constraints with lower level (closer to a kinematic joint) are first
@@ -610,42 +612,6 @@ namespace Chaos
 	}
 
 	
-	void FPBDJointConstraints::SetPreApplyCallback(const FJointPreApplyCallback& Callback)
-	{
-		PreApplyCallback = Callback;
-	}
-
-	
-	void FPBDJointConstraints::ClearPreApplyCallback()
-	{
-		PreApplyCallback = nullptr;
-	}
-
-
-	void FPBDJointConstraints::SetPostApplyCallback(const FJointPostApplyCallback& Callback)
-	{
-		PostApplyCallback = Callback;
-	}
-
-
-	void FPBDJointConstraints::ClearPostApplyCallback()
-	{
-		PostApplyCallback = nullptr;
-	}
-
-
-	void FPBDJointConstraints::SetPostProjectCallback(const FJointPostApplyCallback& Callback)
-	{
-		PostProjectCallback = Callback;
-	}
-
-
-	void FPBDJointConstraints::ClearPostProjectCallback()
-	{
-		PostProjectCallback = nullptr;
-	}
-
-
 	void FPBDJointConstraints::SetBreakCallback(const FJointBreakCallback& Callback)
 	{
 		BreakCallback = Callback;
@@ -737,7 +703,6 @@ namespace Chaos
 		return ConstraintStates[ConstraintIndex].ResimType;
 	}
 
-
 	void FPBDJointConstraints::UpdatePositionBasedState(const FReal Dt)
 	{
 	}
@@ -763,98 +728,6 @@ namespace Chaos
 
 		ConstraintSolvers.Empty();
 	}
-
-	void FPBDJointConstraints::PrepareIteration(FReal Dt)
-	{
-		SCOPE_CYCLE_COUNTER(STAT_Joints_PrepareIteration);
-
-		for (int32 JointIndex = 0; JointIndex < NumConstraints(); ++JointIndex)
-		{
-			if (ConstraintStates[JointIndex].bDisabled) continue;
-
-			const FPBDJointSettings& JointSettings = ConstraintSettings[JointIndex];
-
-			const FTransformPair& JointFrames = JointSettings.ConnectorTransforms;
-			FJointSolverGaussSeidel& Solver = ConstraintSolvers[JointIndex];
-
-			int32 Index0, Index1;
-			GetConstrainedParticleIndices(JointIndex, Index0, Index1);
-			FGenericParticleHandle Particle0 = FGenericParticleHandle(ConstraintParticles[JointIndex][Index0]);
-			FGenericParticleHandle Particle1 = FGenericParticleHandle(ConstraintParticles[JointIndex][Index1]);
-
-			Solver.Init(
-				Dt,
-				Settings,
-				JointSettings,
-				FParticleUtilitiesXR::GetCoMWorldPosition(Particle0),	// Prev position
-				FParticleUtilitiesXR::GetCoMWorldPosition(Particle1),	// Prev position
-				FParticleUtilitiesXR::GetCoMWorldRotation(Particle0),	// Prev rotation
-				FParticleUtilitiesXR::GetCoMWorldRotation(Particle1),	// Prev rotation
-				Particle0->InvM(),
-				Particle0->InvI().GetDiagonal(),
-				Particle1->InvM(),
-				Particle1->InvI().GetDiagonal(),
-				FParticleUtilities::ParticleLocalToCoMLocal(Particle0, JointFrames[Index0]),
-				FParticleUtilities::ParticleLocalToCoMLocal(Particle1, JointFrames[Index1]));
-
-			// Plasticity should not be turned on in the middle of simulation.
-			const bool bUseLinearPlasticity = JointSettings.LinearPlasticityLimit != FLT_MAX;
-			if (bUseLinearPlasticity)
-			{
-				const bool bIsCOMDistanceInitialized = !FMath::IsNearlyEqual(JointSettings.LinearPlasticityInitialDistanceSquared, (FReal)FLT_MAX);
-				if (!bIsCOMDistanceInitialized)
-				{
-					// Joint plasticity is baed on the distance of one of the moment arms of the joint. Typically, plasticity
-					// will get setup from the joint pivot to the child COM (centor of mass), so that is found first. However, when 
-					// the pivot is at the child COM then we fall back to the distance between thge pivot and parent COM.
-					ConstraintSettings[JointIndex].LinearPlasticityInitialDistanceSquared = JointSettings.ConnectorTransforms[Index1].GetTranslation().SizeSquared();
-					if (FMath::IsNearlyZero(ConstraintSettings[JointIndex].LinearPlasticityInitialDistanceSquared))
-					{
-						ConstraintSettings[JointIndex].LinearPlasticityInitialDistanceSquared = JointSettings.ConnectorTransforms[Index0].GetTranslation().SizeSquared();
-					}
-					ensureMsgf(!FMath::IsNearlyZero(ConstraintSettings[JointIndex].LinearPlasticityInitialDistanceSquared), TEXT("Plasticity made inactive due to Zero length difference between parent and child rigid body."));
-				}
-			}
-		}
-	}
-
-
-	void FPBDJointConstraints::UnprepareIteration(FReal Dt)
-	{
-		SCOPE_CYCLE_COUNTER(STAT_Joints_UnprepareIteration);
-
-		for (int32 JointIndex = 0; JointIndex < NumConstraints(); ++JointIndex)
-		{
-			if (ConstraintStates[JointIndex].bDisabled) continue;
-
-			FPBDJointState& JointState = ConstraintStates[JointIndex];
-			FJointSolverGaussSeidel& Solver = ConstraintSolvers[JointIndex];
-
-			int32 Index0, Index1;
-			GetConstrainedParticleIndices(JointIndex, Index0, Index1);
-
-			// NOTE: LinearImpulse/AngularImpulse in the solver are not really impulses - they are mass-weighted position/rotation delta, or (impulse x dt).
-			if (Dt > SMALL_NUMBER)
-			{
-				JointState.LinearImpulse = Solver.GetNetLinearImpulse() / Dt;
-				JointState.AngularImpulse = Solver.GetNetAngularImpulse() / Dt;
-				if (Index0 != 0)
-				{
-					// Particles were flipped in the solver...
-					JointState.LinearImpulse = -JointState.LinearImpulse;
-					JointState.AngularImpulse = -JointState.AngularImpulse;
-				}
-			}
-			else
-			{
-				JointState.LinearImpulse = FVec3(0);
-				JointState.AngularImpulse = FVec3(0);
-			}
-
-			ApplyPlasticityLimits(JointIndex);
-		}
-	}
-
 	
 	void FPBDJointConstraints::CalculateConstraintSpace(int32 ConstraintIndex, FVec3& OutX0, FMatrix33& OutR0, FVec3& OutX1, FMatrix33& OutR1) const
 	{
@@ -881,60 +754,25 @@ namespace Chaos
 	//
 	//////////////////////////////////////////////////////////////////////////
 
-	bool FPBDJointConstraints::Apply(const FReal Dt, const int32 It, const int32 NumIts)
+	void FPBDJointConstraints::GatherInput(const FReal Dt, FPBDIslandSolverData& SolverData)
 	{
-		CSV_SCOPED_TIMING_STAT(Chaos, ApplyJointConstraints);
-
-		if (PreApplyCallback != nullptr)
-		{
-			PreApplyCallback(Dt, Handles);
-		}
-
-		bool bActive = false;
-		if (Settings.ApplyPairIterations > 0)
+		for (int32 ConstraintIndex = 0; ConstraintIndex < NumConstraints(); ++ConstraintIndex)
 		{
-			SCOPE_CYCLE_COUNTER(STAT_Joints_Apply);
-			for (int32 ConstraintIndex = 0; ConstraintIndex < NumConstraints(); ++ConstraintIndex)
+			if (!ConstraintStates[ConstraintIndex].bDisabled)
 			{
-				if ( ConstraintStates[ConstraintIndex].bDisabled) continue;
-
-				bActive |= ApplySingle(Dt, ConstraintIndex, Settings.ApplyPairIterations, It, NumIts);
+				GatherInput(Dt, ConstraintIndex, INDEX_NONE, INDEX_NONE, SolverData);
 			}
 		}
-
-		UE_LOG(LogChaosJoint, Verbose, TEXT("Apply Iteration: %d / %d; Active: %d"), It, NumIts, bActive);
-
-		if (PostApplyCallback != nullptr)
-		{
-			PostApplyCallback(Dt, Handles);
-		}
-
-		return bActive;
 	}
 
-	bool FPBDJointConstraints::ApplyPushOut(const FReal Dt, const int32 It, const int32 NumIts)
+	bool FPBDJointConstraints::ApplyPhase1(const FReal Dt, const int32 It, const int32 NumIts, FPBDIslandSolverData& SolverData)
 	{
+		return ApplyPhase1Serial(Dt, It, NumIts, SolverData);
+	}
 
-		bool bActive = false;
-		if (Settings.ApplyPushOutPairIterations > 0)
-		{
-			SCOPE_CYCLE_COUNTER(STAT_Joints_ApplyPushOut);
-			for (int32 ConstraintIndex = 0; ConstraintIndex < NumConstraints(); ++ConstraintIndex)
-			{
-				if ( ConstraintStates[ConstraintIndex].bDisabled) continue;
-
-				bActive |= ApplyPushOutSingle(Dt, ConstraintIndex, Settings.ApplyPushOutPairIterations, It, NumIts);
-			}
-		}
-
-		UE_LOG(LogChaosJoint, Verbose, TEXT("PushOut Iteration: %d / %d; Active: %d"), It, NumIts, bActive);
-
-		if (PostProjectCallback != nullptr)
-		{
-			PostProjectCallback(Dt, Handles);
-		}
-
-		return bActive;
+	bool FPBDJointConstraints::ApplyPhase2(const FReal Dt, const int32 It, const int32 NumIts, FPBDIslandSolverData& SolverData)
+	{
+		return ApplyPhase2Serial(Dt, It, NumIts, SolverData);
 	}
 
 	//////////////////////////////////////////////////////////////////////////
@@ -949,73 +787,133 @@ namespace Chaos
 	//
 	//////////////////////////////////////////////////////////////////////////
 
-	bool FPBDJointConstraints::Apply(const FReal Dt, const TArray<FConstraintContainerHandle*>& InConstraintHandles, const int32 It, const int32 NumIts)
+	void FPBDJointConstraints::SetNumIslandConstraints(const int32 NumIslandConstraints, FPBDIslandSolverData& SolverData)
 	{
-		SCOPE_CYCLE_COUNTER(STAT_Joints_Apply);
-		CSV_SCOPED_TIMING_STAT(Chaos, ApplyJointConstraints);
+		SolverData.GetConstraintIndices(ContainerId).Reset(NumIslandConstraints);
+	}
 
-		// @todo(ccaulfield): make sorting optional
-		// @todo(ccaulfield): handles should be sorted by level by the constraint rule/graph
-		// @todo(ccaulfield): the best sort order depends on whether we are freezing.
-		// If we are freezing we want the root-most (nearest to kinematic) bodies solved first.
-		// For normal update we want the root body last, otherwise it gets dragged away from the root by the other bodies
+	void FPBDJointConstraints::GatherInput(const FReal Dt, const int32 ConstraintIndex, const int32 Particle0Level, const int32 Particle1Level, FPBDIslandSolverData& SolverData)
+	{
+		SCOPE_CYCLE_COUNTER(STAT_Joints_Gather);
+		check(!ConstraintStates[ConstraintIndex].bDisabled);
 
-		TArray<FConstraintContainerHandle*> SortedConstraintHandles = InConstraintHandles;
-		SortedConstraintHandles.Sort([](const FConstraintContainerHandle& L, const FConstraintContainerHandle& R)
-			{
-				// Sort bodies from root to leaf
-				return L.GetConstraintLevel() < R.GetConstraintLevel();
-			});
+		SolverData.GetConstraintIndices(ContainerId).Add(ConstraintIndex);
 
-		if (PreApplyCallback != nullptr)
+		const FPBDJointSettings& JointSettings = ConstraintSettings[ConstraintIndex];
+		const FTransformPair& JointFrames = JointSettings.ConnectorTransforms;
+		FPBDJointSolver& Solver = ConstraintSolvers[ConstraintIndex];
+
+		int32 Index0, Index1;
+		GetConstrainedParticleIndices(ConstraintIndex, Index0, Index1);
+		FGenericParticleHandle Particle0 = FGenericParticleHandle(ConstraintParticles[ConstraintIndex][Index0]);
+		FGenericParticleHandle Particle1 = FGenericParticleHandle(ConstraintParticles[ConstraintIndex][Index1]);
+
+		// Find the solver bodies for the particles we constrain. This will add them to the container
+		// if they aren't there already, and ensure that they are populated with the latest data.
+		FSolverBody* Body0 = SolverData.GetBodyContainer().FindOrAdd(Particle0);
+		FSolverBody* Body1 = SolverData.GetBodyContainer().FindOrAdd(Particle1);
+
+		Solver.Init(
+			Dt,
+			{ Body0, Body1 },
+			Settings,
+			JointSettings,
+			FParticleUtilities::ParticleLocalToCoMLocal(Particle0, JointFrames[Index0]),
+			FParticleUtilities::ParticleLocalToCoMLocal(Particle1, JointFrames[Index1]));
+
+		// Plasticity should not be turned on in the middle of simulation.
+		const bool bUseLinearPlasticity = JointSettings.LinearPlasticityLimit != FLT_MAX;
+		if (bUseLinearPlasticity)
 		{
-			PreApplyCallback(Dt, SortedConstraintHandles);
+			const bool bIsCOMDistanceInitialized = !FMath::IsNearlyEqual(JointSettings.LinearPlasticityInitialDistanceSquared, (FReal)FLT_MAX);
+			if (!bIsCOMDistanceInitialized)
+			{
+				// Joint plasticity is baed on the distance of one of the moment arms of the joint. Typically, plasticity
+				// will get setup from the joint pivot to the child COM (centor of mass), so that is found first. However, when 
+				// the pivot is at the child COM then we fall back to the distance between thge pivot and parent COM.
+				ConstraintSettings[ConstraintIndex].LinearPlasticityInitialDistanceSquared = JointSettings.ConnectorTransforms[Index1].GetTranslation().SizeSquared();
+				if (FMath::IsNearlyZero(ConstraintSettings[ConstraintIndex].LinearPlasticityInitialDistanceSquared))
+				{
+					ConstraintSettings[ConstraintIndex].LinearPlasticityInitialDistanceSquared = JointSettings.ConnectorTransforms[Index0].GetTranslation().SizeSquared();
+				}
+				// @todo(chaos): move this to validation
+				ensureMsgf(!FMath::IsNearlyZero(ConstraintSettings[ConstraintIndex].LinearPlasticityInitialDistanceSquared), TEXT("Plasticity made inactive due to Zero length difference between parent and child rigid body."));
+			}
 		}
+	}
 
+	void FPBDJointConstraints::ScatterOutput(FReal Dt, FPBDIslandSolverData& SolverData)
+	{
+		SCOPE_CYCLE_COUNTER(STAT_Joints_Scatter);
 
-		int32 NumActive = 0;
-		if (Settings.ApplyPairIterations > 0)
+		for (int32 ConstraintIndex : SolverData.GetConstraintIndices(ContainerId))
 		{
-			for (FConstraintContainerHandle* ConstraintHandle : SortedConstraintHandles)
+			if (!ConstraintStates[ConstraintIndex].bDisabled)
 			{
-				NumActive += ApplySingle(Dt, ConstraintHandle->GetConstraintIndex(), Settings.ApplyPairIterations, It, NumIts);
+				FPBDJointState& JointState = ConstraintStates[ConstraintIndex];
+				FPBDJointSolver& Solver = ConstraintSolvers[ConstraintIndex];
+
+				int32 Index0, Index1;
+				GetConstrainedParticleIndices(ConstraintIndex, Index0, Index1);
+
+				// NOTE: LinearImpulse/AngularImpulse in the solver are not really impulses - they are mass-weighted position/rotation delta, or (impulse x dt).
+				if (Dt > SMALL_NUMBER)
+				{
+					if (Index0 == 0)
+					{
+						JointState.LinearImpulse = Solver.GetNetLinearImpulse() / Dt;
+						JointState.AngularImpulse = Solver.GetNetAngularImpulse() / Dt;
+					}
+					else
+					{
+						// Particles were flipped in the solver...
+						JointState.LinearImpulse = -Solver.GetNetLinearImpulse() / Dt;
+						JointState.AngularImpulse = -Solver.GetNetAngularImpulse() / Dt;
+					}
+				}
+				else
+				{
+					JointState.LinearImpulse = FVec3(0);
+					JointState.AngularImpulse = FVec3(0);
+				}
+
+				ApplyPlasticityLimits(ConstraintIndex);
+
+				// Remove our solver body reference (they are not valid between frames)
+				ConstraintSolvers[ConstraintIndex].Deinit();
 			}
 		}
+	}
+
+	bool FPBDJointConstraints::ApplyPhase1Serial(const FReal Dt, const int32 It, const int32 NumIts, FPBDIslandSolverData& SolverData)
+	{
+		CSV_SCOPED_TIMING_STAT(Chaos, ApplyJointConstraints);
+		SCOPE_CYCLE_COUNTER(STAT_Joints_Apply);
 
-		if (PostApplyCallback != nullptr)
+		int32 NumActive = 0;
+		if (Settings.ApplyPairIterations > 0)
 		{
-			PostApplyCallback(Dt, SortedConstraintHandles);
+			for (int32 ConstraintIndex : SolverData.GetConstraintIndices(ContainerId))
+			{
+				NumActive += ApplyPhase1Single(Dt, ConstraintIndex, Settings.ApplyPairIterations, It, NumIts);
+			}
 		}
-
 		return (NumActive > 0);
 	}
 
-	
-	bool FPBDJointConstraints::ApplyPushOut(const FReal Dt, const TArray<FConstraintContainerHandle*>& InConstraintHandles, const int32 It, const int32 NumIts)
+	bool FPBDJointConstraints::ApplyPhase2Serial(const FReal Dt, const int32 It, const int32 NumIts, FPBDIslandSolverData& SolverData)
 	{
+		CSV_SCOPED_TIMING_STAT(Chaos, PushOutJointConstraints);
 		SCOPE_CYCLE_COUNTER(STAT_Joints_ApplyPushOut);
 
-		TArray<FConstraintContainerHandle*> SortedConstraintHandles = InConstraintHandles;
-		SortedConstraintHandles.Sort([](const FConstraintContainerHandle& L, const FConstraintContainerHandle& R)
-			{
-				// Sort bodies from root to leaf
-				return L.GetConstraintLevel() < R.GetConstraintLevel();
-			});
-
 		int32 NumActive = 0;
-		if (Settings.ApplyPushOutPairIterations > 0)
+		if (Settings.ApplyPairIterations > 0)
 		{
-			for (FConstraintContainerHandle* ConstraintHandle : SortedConstraintHandles)
+			for (int32 ConstraintIndex : SolverData.GetConstraintIndices(ContainerId))
 			{
-				NumActive += ApplyPushOutSingle(Dt, ConstraintHandle->GetConstraintIndex(), Settings.ApplyPushOutPairIterations, It, NumIts);
+				NumActive += ApplyPhase2Single(Dt, ConstraintIndex, Settings.ApplyPairIterations, It, NumIts);
 			}
 		}
-
-		if (PostProjectCallback != nullptr)
-		{
-			PostProjectCallback(Dt, SortedConstraintHandles);
-		}
-
 		return (NumActive > 0);
 	}
 
@@ -1032,32 +930,6 @@ namespace Chaos
 	//
 	//////////////////////////////////////////////////////////////////////////
 
-	void FPBDJointConstraints::UpdateParticleState(TPBDRigidParticleHandle<FReal, 3>* Rigid, const FReal Dt, const FVec3& PrevP, const FRotation3& PrevQ, const FVec3& P, const FRotation3& Q, const bool bUpdateVelocity)
-	{
-		if ((Rigid != nullptr) && (Rigid->ObjectState() == EObjectStateType::Dynamic))
-		{
-			FParticleUtilities::SetCoMWorldTransform(Rigid, P, Q);
-			if (bUpdateVelocity && (Dt > SMALL_NUMBER))
-			{
-				const FVec3 V = FVec3::CalculateVelocity(PrevP, P, Dt);
-				const FVec3 W = FRotation3::CalculateAngularVelocity(PrevQ, Q, Dt);
-				Rigid->SetV(V);
-				Rigid->SetW(W);
-			}
-		}
-	}
-
-
-	void FPBDJointConstraints::UpdateParticleStateExplicit(TPBDRigidParticleHandle<FReal, 3>* Rigid, const FReal Dt, const FVec3& P, const FRotation3& Q, const FVec3& V, const FVec3& W)
-	{
-		if ((Rigid != nullptr) && (Rigid->ObjectState() == EObjectStateType::Dynamic))
-		{
-			FParticleUtilities::SetCoMWorldTransform(Rigid, P, Q);
-			Rigid->SetV(V);
-			Rigid->SetW(W);
-		}
-	}
-
 	FReal FPBDJointConstraints::CalculateIterationStiffness(int32 It, int32 NumIts) const
 	{
 		// Linearly interpolate betwwen MinStiffness and MaxStiffness over the first few iterations,
@@ -1071,42 +943,11 @@ namespace Chaos
 		return FMath::Clamp(IterationStiffness, 0.0f, 1.0f);
 	}
 
-	bool FPBDJointConstraints::CanEvaluate(const int32 ConstraintIndex) const
-	{
-		if (!IsConstraintEnabled(ConstraintIndex))
-		{
-			return false;
-		}
-
-		int32 Index0, Index1;
-		GetConstrainedParticleIndices(ConstraintIndex, Index0, Index1);
-		const FGenericParticleHandle Particle0 = FGenericParticleHandle(ConstraintParticles[ConstraintIndex][Index0]);
-		const FGenericParticleHandle Particle1 = FGenericParticleHandle(ConstraintParticles[ConstraintIndex][Index1]);
-
-		// check for valid and enabled particles
-		if (Particle0->Handle() == nullptr || Particle0->Disabled()
-			|| Particle1->Handle() == nullptr || Particle1->Disabled())
-		{
-			return false;
-		}
-
-		// check valid particle and solver state
-		const FJointSolverGaussSeidel& Solver = ConstraintSolvers[ConstraintIndex];
-		if ((Particle0->Sleeping() && Particle1->Sleeping())
-			|| (Particle0->IsKinematic() && Particle1->Sleeping())
-			|| (Particle0->Sleeping() && Particle1->IsKinematic())
-			|| (FMath::IsNearlyZero(Solver.InvM(0)) && FMath::IsNearlyZero(Solver.InvM(1))))
-		{
-			return false;
-		}
-		return true;
-	}
-
 	// This position solver iterates over each of the inner constraints (position, twist, swing) and solves them independently.
 	// This will converge slowly in some cases, particularly where resolving angular constraints violates position constraints and vice versa.
-	bool FPBDJointConstraints::ApplySingle(const FReal Dt, const int32 ConstraintIndex, const int32 NumPairIts, const int32 It, const int32 NumIts)
+	bool FPBDJointConstraints::ApplyPhase1Single(const FReal Dt, const int32 ConstraintIndex, const int32 NumPairIts, const int32 It, const int32 NumIts)
 	{
-		if (!CanEvaluate(ConstraintIndex))
+		if (!IsConstraintEnabled(ConstraintIndex))
 		{
 			return false;
 		}
@@ -1115,48 +956,37 @@ namespace Chaos
 		UE_LOG(LogChaosJoint, VeryVerbose, TEXT("Solve Joint Constraint %d %s %s (dt = %f; it = %d / %d)"), ConstraintIndex, *Constraint[0]->ToString(), *Constraint[1]->ToString(), Dt, It, NumIts);
 
 		const FPBDJointSettings& JointSettings = ConstraintSettings[ConstraintIndex];
-		FJointSolverGaussSeidel& Solver = ConstraintSolvers[ConstraintIndex];
+		FPBDJointSolver& Solver = ConstraintSolvers[ConstraintIndex];
 
+		// @todo(chaos): store this on the Solver object and don't access the particles here
 		int32 Index0, Index1;
 		GetConstrainedParticleIndices(ConstraintIndex, Index0, Index1);
 		FGenericParticleHandle Particle0 = FGenericParticleHandle(ConstraintParticles[ConstraintIndex][Index0]);
 		FGenericParticleHandle Particle1 = FGenericParticleHandle(ConstraintParticles[ConstraintIndex][Index1]);
+		if ((Particle0->Sleeping() && Particle1->Sleeping())
+			|| (Particle0->IsKinematic() && Particle1->Sleeping()) 
+			|| (Particle0->Sleeping() && Particle1->IsKinematic())
+			|| (FMath::IsNearlyZero(Solver.InvM(0)) && FMath::IsNearlyZero(Solver.InvM(1))))
+		{
+			return false;
+		}
 
-		const FVec3 P0 = FParticleUtilities::GetCoMWorldPosition(Particle0);
-		const FRotation3 Q0 = FParticleUtilities::GetCoMWorldRotation(Particle0);
-		const FVec3 P1 = FParticleUtilities::GetCoMWorldPosition(Particle1);
-		const FRotation3 Q1 = FParticleUtilities::GetCoMWorldRotation(Particle1);
-
+		// If we were solved last iteration and nothing has changed since, we are done
 		const bool bWasActive = Solver.GetIsActive();
+		Solver.Update(Dt, Settings, JointSettings);
 
-		const FReal IterationStiffness = CalculateIterationStiffness(It, NumIts);
-
-		Solver.Update(
-			Dt,
-			IterationStiffness,
-			Settings,
-			JointSettings,
-			P0,
-			Q0,
-			Particle0->V(),
-			Particle0->W(),
-			P1,
-			Q1,
-			Particle1->V(),
-			Particle1->W());
-
-		// If we were solved last iteration and nothing has changed since, we are done
 		if (!bWasActive && !Solver.GetIsActive() && bChaos_Joint_EarlyOut_Enabled)
 		{
 			return false;
 		}
 
+		const FReal IterationStiffness = CalculateIterationStiffness(It, NumIts);
 		for (int32 PairIt = 0; PairIt < NumPairIts; ++PairIt)
 		{
 			UE_LOG(LogChaosJoint, VeryVerbose, TEXT("  Pair Iteration %d / %d"), PairIt, NumPairIts);
 
-			// This is the same for all SolverType settings (which makes it wrong for the GbfPbd version)
-			Solver.ApplyConstraints(Dt, Settings, JointSettings);
+			// This is the same position solver for all SolverTypes (which makes it wrong for the GbfPbd solver which should be a pbd-esque velocity solve)
+			Solver.ApplyConstraints(Dt, IterationStiffness, Settings, JointSettings);
 
 			if (!Solver.GetIsActive() && bChaos_Joint_EarlyOut_Enabled)
 			{
@@ -1164,8 +994,13 @@ namespace Chaos
 			}
 		}
 
-		UpdateParticleState(Particle0->CastToRigidParticle(), Dt, Solver.GetInitP(0), Solver.GetInitQ(0), Solver.GetP(0), Solver.GetQ(0), bUpdateVelocityInApplyConstraints);
-		UpdateParticleState(Particle1->CastToRigidParticle(), Dt, Solver.GetInitP(1), Solver.GetInitQ(1), Solver.GetP(1), Solver.GetQ(1), bUpdateVelocityInApplyConstraints);
+		// If the solver we are running uses a velocity solver in the Apply step (for other constraint types), 
+		// we need to update the implicit velocity because we did a position solve above. E.g., GbfPbd collisions.
+		if (SolverType == EConstraintSolverType::GbfPbd)
+		{
+			Solver.Body(0).SetImplicitVelocity(Dt);
+			Solver.Body(1).SetImplicitVelocity(Dt);
+		}
 
 		// @todo(ccaulfield): The break limit should really be applied to the impulse in the solver to prevent 1-frame impulses larger than the threshold
 		if ((JointSettings.LinearBreakForce!=FLT_MAX) || (JointSettings.AngularBreakTorque!=FLT_MAX))
@@ -1176,9 +1011,9 @@ namespace Chaos
 		return Solver.GetIsActive() || !bChaos_Joint_EarlyOut_Enabled;
 	}
 
-	bool FPBDJointConstraints::ApplyPushOutSingle(const FReal Dt, const int32 ConstraintIndex, const int32 NumPairIts, const int32 It, const int32 NumIts)
+	bool FPBDJointConstraints::ApplyPhase2Single(const FReal Dt, const int32 ConstraintIndex, const int32 NumPairIts, const int32 It, const int32 NumIts)
 	{
-		if (!CanEvaluate(ConstraintIndex))
+		if (!IsConstraintEnabled(ConstraintIndex))
 		{
 			return false;
 		}
@@ -1187,42 +1022,32 @@ namespace Chaos
 		UE_LOG(LogChaosJoint, VeryVerbose, TEXT("Project Joint Constraint %d %s %s (dt = %f; it = %d / %d)"), ConstraintIndex, *Constraint[0]->ToString(), *Constraint[1]->ToString(), Dt, It, NumIts);
 
 		const FPBDJointSettings& JointSettings = ConstraintSettings[ConstraintIndex];
-		FJointSolverGaussSeidel& Solver = ConstraintSolvers[ConstraintIndex];
+		FPBDJointSolver& Solver = ConstraintSolvers[ConstraintIndex];
 
+		// @todo(chaos): store this on the Solver object and don't access the particles here
 		int32 Index0, Index1;
 		GetConstrainedParticleIndices(ConstraintIndex, Index0, Index1);
 		FGenericParticleHandle Particle0 = FGenericParticleHandle(ConstraintParticles[ConstraintIndex][Index0]);
 		FGenericParticleHandle Particle1 = FGenericParticleHandle(ConstraintParticles[ConstraintIndex][Index1]);
-
-		const FVec3 P0 = FParticleUtilities::GetCoMWorldPosition(Particle0);
-		const FRotation3 Q0 = FParticleUtilities::GetCoMWorldRotation(Particle0);
-		const FVec3 P1 = FParticleUtilities::GetCoMWorldPosition(Particle1);
-		const FRotation3 Q1 = FParticleUtilities::GetCoMWorldRotation(Particle1);
-
-		const bool bWasActive = Solver.GetIsActive();
-
-		const FReal IterationStiffness = CalculateIterationStiffness(It, NumIts);
-
-		Solver.Update(
-			Dt,
-			IterationStiffness,
-			Settings,
-			JointSettings,
-			P0,
-			Q0,
-			Particle0->V(),
-			Particle0->W(),
-			P1,
-			Q1,
-			Particle1->V(),
-			Particle1->W());
+		if ((Particle0->Sleeping() && Particle1->Sleeping())
+			|| (Particle0->IsKinematic() && Particle1->Sleeping())
+			|| (Particle0->Sleeping() && Particle1->IsKinematic()) 
+			|| (FMath::IsNearlyZero(Solver.InvM(0)) && FMath::IsNearlyZero(Solver.InvM(1))))
+		{
+			return false;
+		}
 
 		// If we were solved last iteration and nothing has changed since, we are done
-		if (!bWasActive && !Solver.GetIsActive() && bChaos_Joint_EarlyOut_Enabled)
+		const bool bWasActive = Solver.GetIsActive();
+		Solver.Update(Dt, Settings, JointSettings);
+		
+		// For quasipbd, we may still need to stabilize velocities and solve for restitution.
+		if (!bWasActive && !Solver.GetIsActive() && SolverType != EConstraintSolverType::QuasiPbd && bChaos_Joint_EarlyOut_Enabled)
 		{
 			return false;
 		}
 
+		const FReal IterationStiffness = CalculateIterationStiffness(It, NumIts);
 		for (int32 PairIt = 0; PairIt < NumPairIts; ++PairIt)
 		{
 			switch(SolverType)
@@ -1231,10 +1056,10 @@ namespace Chaos
 				break;
 			case EConstraintSolverType::GbfPbd:
 			case EConstraintSolverType::StandardPbd:
-				Solver.ApplyProjections(Dt, Settings, JointSettings);
+				Solver.ApplyProjections(Dt, IterationStiffness, Settings, JointSettings);
 				break;
 			case EConstraintSolverType::QuasiPbd:
-				Solver.ApplyVelocityConstraints(Dt, Settings, JointSettings);
+				Solver.ApplyVelocityConstraints(Dt, IterationStiffness, Settings, JointSettings);
 				break;
 			}
 
@@ -1244,9 +1069,6 @@ namespace Chaos
 			}
 		}
 
-		UpdateParticleStateExplicit(Particle0->CastToRigidParticle(), Dt, Solver.GetP(0), Solver.GetQ(0), Solver.GetV(0), Solver.GetW(0));
-		UpdateParticleStateExplicit(Particle1->CastToRigidParticle(), Dt, Solver.GetP(1), Solver.GetQ(1), Solver.GetV(1), Solver.GetW(1));
-
 		// @todo(ccaulfield): should probably add to net impulses in push out too...(for breaking etc)
 
 		return Solver.GetIsActive() || !bChaos_Joint_EarlyOut_Enabled;
@@ -1262,20 +1084,21 @@ namespace Chaos
 		bool bBreak = false;
 		if (!bBreak && JointSettings.LinearBreakForce!=FLT_MAX)
 		{
-			const FReal LinearThreshold = JointSettings.LinearBreakForce * Dt * Dt;
-			UE_LOG(LogChaosJoint, VeryVerbose, TEXT("Constraint %d Linear Break Check: %f / %f"), ConstraintIndex, LinearImpulse.Size(), LinearThreshold);
+			const FReal LinearForceSq = LinearImpulse.SizeSquared() / (Dt * Dt * Dt * Dt);
+			const FReal LinearThresholdSq = FMath::Square(JointSettings.LinearBreakForce);
+
+			UE_LOG(LogChaosJoint, VeryVerbose, TEXT("Constraint %d Linear Break Check: %f / %f at Dt = %f"), ConstraintIndex, FMath::Sqrt(LinearForceSq), FMath::Sqrt(LinearThresholdSq), Dt);
 
-			const FReal LinearThresholdSq = LinearThreshold * LinearThreshold;
-			bBreak = LinearImpulse.SizeSquared() > LinearThresholdSq;
+			bBreak = LinearForceSq > LinearThresholdSq;
 		}
 
 		if (!bBreak && JointSettings.AngularBreakTorque!=FLT_MAX)
 		{
-			const FReal AngularThreshold = JointSettings.AngularBreakTorque * Dt * Dt;
-			UE_LOG(LogChaosJoint, VeryVerbose, TEXT("Constraint %d Angular Break Check: %f / %f"), ConstraintIndex, AngularImpulse.Size(), AngularThreshold);
+			const FReal AngularForceSq = AngularImpulse.SizeSquared() / (Dt * Dt * Dt * Dt);
+			const FReal AngularThresholdSq = FMath::Square(JointSettings.AngularBreakTorque);
+			UE_LOG(LogChaosJoint, VeryVerbose, TEXT("Constraint %d Angular Break Check: %f / %f at Dt = %f"), ConstraintIndex, FMath::Sqrt(AngularForceSq), FMath::Sqrt(AngularThresholdSq), Dt);
 
-			const FReal AngularThresholdSq = AngularThreshold * AngularThreshold;
-			bBreak = AngularImpulse.SizeSquared() > AngularThresholdSq;
+			bBreak = AngularForceSq > AngularThresholdSq;
 		}
 
 		if (bBreak)
@@ -1285,33 +1108,38 @@ namespace Chaos
 	}
 
 
-	void FPBDJointConstraints::ApplyPlasticityLimits(const int32 ConstraintIndex)
+	void FPBDJointConstraints::ApplyPlasticityLimits(int32 ConstraintIndex)
 	{
 		FPBDJointSettings& JointSettings = ConstraintSettings[ConstraintIndex];
 		const bool bHasLinearPlasticityLimit = JointSettings.LinearPlasticityLimit != FLT_MAX;
 		const bool bHasAngularPlasticityLimit = JointSettings.AngularPlasticityLimit != FLT_MAX;
-		const bool bHasPlasticityLimits =  bHasLinearPlasticityLimit || bHasAngularPlasticityLimit;
-		if(!bHasPlasticityLimits)
+		const bool bHasPlasticityLimits = bHasLinearPlasticityLimit || bHasAngularPlasticityLimit;
+		if (!bHasPlasticityLimits)
 		{
 			return;
 		}
 
-		if(!Settings.bEnableDrives)
+		if (!Settings.bEnableDrives)
 		{
 			return;
 		}
 
 		int32 Index0, Index1;
 		GetConstrainedParticleIndices(ConstraintIndex, Index0, Index1);
-		FGenericParticleHandle Particle0 = FGenericParticleHandle(ConstraintParticles[ConstraintIndex][Index0]);
-		FGenericParticleHandle Particle1 = FGenericParticleHandle(ConstraintParticles[ConstraintIndex][Index1]);
-		if (Particle0->Disabled() || Particle1->Disabled())
 		{
-			return;
+			FGenericParticleHandle Particle0 = FGenericParticleHandle(ConstraintParticles[ConstraintIndex][Index0]);
+			FGenericParticleHandle Particle1 = FGenericParticleHandle(ConstraintParticles[ConstraintIndex][Index1]);
+			if (Particle0->Disabled() || Particle1->Disabled())
+			{
+				return;
+			}
 		}
 
-		const FTransformPair &ConstraintFramesLocal = JointSettings.ConnectorTransforms;
-		FTransformPair ConstraintFramesGlobal(ConstraintFramesLocal[Index0] * FRigidTransform3(Particle0->P(), Particle0->Q()), ConstraintFramesLocal[Index1] * FRigidTransform3(Particle1->P(), Particle1->Q()));
+		FConstraintSolverBody& Body0 = ConstraintSolvers[ConstraintIndex].Body(0);
+		FConstraintSolverBody& Body1 = ConstraintSolvers[ConstraintIndex].Body(1);
+
+		const FTransformPair& ConstraintFramesLocal = JointSettings.ConnectorTransforms;
+		FTransformPair ConstraintFramesGlobal(ConstraintFramesLocal[Index0] * FRigidTransform3(Body0.ActorP(), Body0.ActorQ()), ConstraintFramesLocal[Index1] * FRigidTransform3(Body1.ActorP(), Body1.ActorQ()));
 		FQuat Q1 = ConstraintFramesGlobal[1].GetRotation();
 		Q1.EnforceShortestArcWith(ConstraintFramesGlobal[0].GetRotation());
 		ConstraintFramesGlobal[1].SetRotation(Q1);
@@ -1342,7 +1170,7 @@ namespace Chaos
 					// Shrink and Grow are based on the distance between the joint pivot and the child. 
 					// Note, if the pivot is located at the COM of the child then shrink will not do anything. 
 					FVec3 StartDelta = ConstraintFramesLocal[Index1].InverseTransformPositionNoScale(JointSettings.LinearDrivePositionTarget);
-					FVec3 CurrentDelta = ConstraintFramesGlobal[Index1].InverseTransformPositionNoScale(Particle1->P());
+					FVec3 CurrentDelta = ConstraintFramesGlobal[Index1].InverseTransformPositionNoScale(Body1.P());
 
 					if (JointSettings.LinearPlasticityType == EPlasticityType::Shrink && CurrentDelta.SizeSquared() < StartDelta.SizeSquared())
 					{
@@ -1357,25 +1185,25 @@ namespace Chaos
 				}
 			}
 		}
-		if(bHasAngularPlasticityLimit)
+		if (bHasAngularPlasticityLimit)
 		{
 			FRotation3 Swing, Twist; FPBDJointUtilities::DecomposeSwingTwistLocal(ConstraintFramesGlobal[0].GetRotation(), ConstraintFramesGlobal[1].GetRotation(), Swing, Twist);
 
 			// @todo(chaos): still need to warn against the case where all position drives are not enabled or all dimensions are locked. Warning should print out the joint names and should only print out once to avoid spamming.
-			if((!JointSettings.bAngularSLerpPositionDriveEnabled && !JointSettings.bAngularTwistPositionDriveEnabled) || JointSettings.AngularMotionTypes[(int32)EJointAngularConstraintIndex::Twist] == EJointMotionType::Locked)
+			if ((!JointSettings.bAngularSLerpPositionDriveEnabled && !JointSettings.bAngularTwistPositionDriveEnabled) || JointSettings.AngularMotionTypes[(int32)EJointAngularConstraintIndex::Twist] == EJointMotionType::Locked)
 			{
 				Twist = FRotation3::Identity;
 			}
 			// @todo(chaos): clamp rotation if only swing1(swing2) is locked
-			if((!JointSettings.bAngularSLerpPositionDriveEnabled && !JointSettings.bAngularSwingPositionDriveEnabled) || (JointSettings.AngularMotionTypes[(int32)EJointAngularConstraintIndex::Swing1] == EJointMotionType::Locked && JointSettings.AngularMotionTypes[(int32)EJointAngularConstraintIndex::Swing2] == EJointMotionType::Locked))
+			if ((!JointSettings.bAngularSLerpPositionDriveEnabled && !JointSettings.bAngularSwingPositionDriveEnabled) || (JointSettings.AngularMotionTypes[(int32)EJointAngularConstraintIndex::Swing1] == EJointMotionType::Locked && JointSettings.AngularMotionTypes[(int32)EJointAngularConstraintIndex::Swing2] == EJointMotionType::Locked))
 			{
 				Swing = FRotation3::Identity;
 			}
-			
+
 			const FRotation3 AngularDisplacement = Swing * Twist;
 			// Assuming that the dimensions which are locked or have no targets are 0. in AngularDrivePositionTarget
-			const FReal AngleDeg = JointSettings.AngularDrivePositionTarget.AngularDistance(AngularDisplacement);
-			if (AngleDeg > JointSettings.AngularPlasticityLimit)
+			const FReal AngleRad = JointSettings.AngularDrivePositionTarget.AngularDistance(AngularDisplacement);
+			if (AngleRad > JointSettings.AngularPlasticityLimit)
 			{
 				JointSettings.AngularDrivePositionTarget = AngularDisplacement;
 				SetDriveTargetChanged(ConstraintIndex, true);
@@ -1383,7 +1211,6 @@ namespace Chaos
 		}
 	}
 
-
 	// Assign an Island, Level and Color to each constraint. Constraints must be processed in Level order, but
 	// constraints of the same color are independent and can be processed in parallel (SIMD or Task)
 	// NOTE: Constraints are the Vertices in this graph, and Edges connect constraints sharing a Particle. 
@@ -1511,5 +1338,5 @@ namespace Chaos
 
 namespace Chaos
 {
-	template class TContainerConstraintHandle<FPBDJointConstraints>;
+	template class TIndexedContainerConstraintHandle<FPBDJointConstraints>;
 }
diff --git a/Engine/Source/Runtime/Experimental/Chaos/Private/Chaos/PBDPositionConstraints.cpp b/Engine/Source/Runtime/Experimental/Chaos/Private/Chaos/PBDPositionConstraints.cpp
--- a/Engine/Source/Runtime/Experimental/Chaos/Private/Chaos/PBDPositionConstraints.cpp
+++ b/Engine/Source/Runtime/Experimental/Chaos/Private/Chaos/PBDPositionConstraints.cpp
@@ -1,6 +1,7 @@
 // Copyright Epic Games, Inc. All Rights Reserved.
 
 #include "Chaos/PBDPositionConstraints.h"
+#include "Chaos/Evolution/SolverDatas.h"
 
 namespace Chaos
 {
@@ -9,14 +10,39 @@ namespace Chaos
 		return ConcreteContainer()->GetConstrainedParticles(ConstraintIndex);
 	}
 
-	bool FPBDPositionConstraints::Apply(const FReal Dt, const TArray<FConstraintContainerHandle*>& ConstraintHandles, const int32 It, const int32 NumIts) const
+	void FPBDPositionConstraintHandle::GatherInput(const FReal Dt, const int32 Particle0Level, const int32 Particle1Level, FPBDIslandSolverData& SolverData)
 	{
-		for (FConstraintContainerHandle* ConstraintHandle : ConstraintHandles)
+		ConcreteContainer()->GatherInput(Dt, ConstraintIndex, Particle0Level, Particle1Level, SolverData);
+	}
+	
+	void FPBDPositionConstraints::SetNumIslandConstraints(const int32 NumIslandConstraints, FPBDIslandSolverData& SolverData)
+	{
+		SolverData.GetConstraintIndices(ContainerId).Reset(NumIslandConstraints);
+	}
+
+	void FPBDPositionConstraints::GatherInput(const FReal Dt, const int32 ConstraintIndex, const int32 Particle0Level, const int32 Particle1Level, FPBDIslandSolverData& SolverData)
+	{
+		SolverData.GetConstraintIndices(ContainerId).Add(ConstraintIndex);
+
+		ConstraintSolverBodies[ConstraintIndex] = SolverData.GetBodyContainer().FindOrAdd(ConstrainedParticles[ConstraintIndex]);
+	}
+	
+	void FPBDPositionConstraints::ScatterOutput(FReal Dt, FPBDIslandSolverData& SolverData)
+	{
+		for (int32 ConstraintIndex : SolverData.GetConstraintIndices(ContainerId))
+		{
+			ConstraintSolverBodies[ConstraintIndex] = nullptr;
+		}
+	}
+
+	bool FPBDPositionConstraints::ApplyPhase1Serial(const FReal Dt, const int32 It, const int32 NumIts, FPBDIslandSolverData& SolverData)
+	{
+		for (int32 ConstraintIndex : SolverData.GetConstraintIndices(ContainerId))
 		{
-			ApplySingle(Dt, ConstraintHandle->GetConstraintIndex());
+			ApplySingle(Dt, ConstraintIndex);
 		}
 
-		// TODO: Return true only if more iteration are needed
+		// @todo(chaos): early iteration termination in FPBDPositionConstraints
 		return true;
 	}
 }
diff --git a/Engine/Source/Runtime/Experimental/Chaos/Private/Chaos/PBDRigidClustering.cpp b/Engine/Source/Runtime/Experimental/Chaos/Private/Chaos/PBDRigidClustering.cpp
--- a/Engine/Source/Runtime/Experimental/Chaos/Private/Chaos/PBDRigidClustering.cpp
+++ b/Engine/Source/Runtime/Experimental/Chaos/Private/Chaos/PBDRigidClustering.cpp
@@ -76,7 +76,7 @@ namespace Chaos
 	//==========================================================================
 
 	template<class T, int d>
-	FVec3 GetContactLocation(const FRigidBodyPointContactConstraint& Contact)
+	FVec3 GetContactLocation(const FPBDCollisionConstraint& Contact)
 	{
 		return Contact.GetLocation();
 	}
@@ -172,9 +172,9 @@ namespace Chaos
 			// @todo(mlentine): We can precompute internal constraints which can filter some from the narrow phase tests but may not help much
 
 #if CHAOS_PARTICLEHANDLE_TODO
-			CollisionRule.UpdateConstraints(InParticles, Evolution.GetNonDisabledIndices(), Dt, AllActivatedChildren, AllIslandParticles.Array());
+			//CollisionRule.UpdateConstraints(InParticles, Evolution.GetNonDisabledIndices(), Dt, AllActivatedChildren, AllIslandParticles.Array());
 #else
-			CollisionRule.UpdateConstraints(Dt, AllIslandParticles);	//this seems wrong
+			//CollisionRule.UpdateConstraints(Dt, AllIslandParticles);	//this seems wrong
 #endif
 
 			Evolution.InitializeAccelerationStructures();
@@ -182,10 +182,10 @@ namespace Chaos
 			// Resolve collisions
 			PhysicsParallelFor(IslandsToRecollide.Num(), [&](int32 Island) {
 				// @todo(mlentine): This is heavy handed and probably can be simplified as we know only a little bit changed.
-				Evolution.UpdateAccelerationStructures(Island);
-				Evolution.ApplyConstraints(Dt, Island);
+				Evolution.UpdateAccelerationStructures(Dt, Island);
+				Evolution.ApplyConstraintsPhase1(Dt, Island);
 				// @todo(ccaulfield): should we also update velocities here? Evolution does...
-				Evolution.ApplyPushOut(Dt, Island);
+				Evolution.ApplyConstraintsPhase2(Dt, Island);
 				// @todo(ccaulfield): support sleep state update on evolution
 				//Evolution.UpdateSleepState(Island);
 			});
