Chaos

diff --git a/Engine/Source/Runtime/Experimental/Chaos/Private/Chaos/PBDRigidDynamicSpringConstraints.cpp b/Engine/Source/Runtime/Experimental/Chaos/Private/Chaos/PBDRigidDynamicSpringConstraints.cpp
--- a/Engine/Source/Runtime/Experimental/Chaos/Private/Chaos/PBDRigidDynamicSpringConstraints.cpp
+++ b/Engine/Source/Runtime/Experimental/Chaos/Private/Chaos/PBDRigidDynamicSpringConstraints.cpp
@@ -1,6 +1,7 @@
 // Copyright Epic Games, Inc. All Rights Reserved.
 
 #include "Chaos/PBDRigidDynamicSpringConstraints.h"
+#include "Chaos/Evolution/SolverDatas.h"
 #include "Chaos/Particle/ParticleUtilities.h"
 #include "Chaos/Utilities.h"
 
@@ -11,6 +12,14 @@ TVec2<FGeometryParticleHandle*> FPBDRigidDynamicSpringConstraintHandle::GetConst
 	return ConcreteContainer()->GetConstrainedParticles(ConstraintIndex);
 }
 
+void FPBDRigidDynamicSpringConstraintHandle::GatherInput(const FReal Dt, const int32 Particle0Level, const int32 Particle1Level, FPBDIslandSolverData& SolverData)
+{
+	ConcreteContainer()->GatherInput(Dt, ConstraintIndex, Particle0Level, Particle1Level, SolverData);
+}
+
+// Note: this is called outside the constraint solver loop (it generates constraints and can merge 
+// islands - the equivalent of collision detection) and therefore reads directly from Particles, 
+// rather than from SolverBodies.
 void FPBDRigidDynamicSpringConstraints::UpdatePositionBasedState(const FReal Dt)
 {
 	const int32 NumConstraints = Constraints.Num();
@@ -93,69 +102,93 @@ void FPBDRigidDynamicSpringConstraints::UpdatePositionBasedState(const FReal Dt)
 	}
 }
 
-FVec3 FPBDRigidDynamicSpringConstraints::GetDelta(const FVec3& WorldSpaceX1, const FVec3& WorldSpaceX2, const int32 ConstraintIndex, const int32 SpringIndex) const
+void FPBDRigidDynamicSpringConstraints::SetNumIslandConstraints(const int32 NumIslandConstraints, FPBDIslandSolverData& SolverData)
+{
+	SolverData.GetConstraintIndices(ContainerId).Reset(NumIslandConstraints);
+}
+
+void FPBDRigidDynamicSpringConstraints::GatherInput(const FReal Dt, const int32 ConstraintIndex, const int32 Particle0Level, const int32 Particle1Level, FPBDIslandSolverData& SolverData)
 {
-	FPBDRigidParticleHandle* PBDRigid0 = Constraints[ConstraintIndex][0]->CastToRigidParticle();
-	FPBDRigidParticleHandle* PBDRigid1 = Constraints[ConstraintIndex][1]->CastToRigidParticle();
-	const bool bIsRigidDynamic0 = PBDRigid0 && PBDRigid0->ObjectState() == EObjectStateType::Dynamic;
-	const bool bIsRigidDynamic1 = PBDRigid1 && PBDRigid1->ObjectState() == EObjectStateType::Dynamic;
+	SolverData.GetConstraintIndices(ContainerId).Add(ConstraintIndex);
+
+	ConstraintSolverBodies[ConstraintIndex] =
+	{
+		SolverData.GetBodyContainer().FindOrAdd(Constraints[ConstraintIndex][0]),
+		SolverData.GetBodyContainer().FindOrAdd(Constraints[ConstraintIndex][1])
+	};
+}
 
-	if (!bIsRigidDynamic0 && !bIsRigidDynamic1)
+void FPBDRigidDynamicSpringConstraints::ScatterOutput(FReal Dt, FPBDIslandSolverData& SolverData)
+{
+	for (int32 ConstraintIndex : SolverData.GetConstraintIndices(ContainerId))
 	{
-		return FVec3(0);
+		ConstraintSolverBodies[ConstraintIndex] = 
+		{
+			nullptr,
+			nullptr
+		};
+	}
+}
+
+bool FPBDRigidDynamicSpringConstraints::ApplyPhase1Serial(const FReal Dt, const int32 It, const int32 NumIts, FPBDIslandSolverData& SolverData)
+{
+	for (int32 ConstraintIndex : SolverData.GetConstraintIndices(ContainerId))
+	{
+		ApplySingle(Dt, ConstraintIndex);
 	}
 
+	// @todo(chaos): early iteration termination in FPBDRigidDynamicSpringConstraints
+	return true;
+}
+
+FVec3 FPBDRigidDynamicSpringConstraints::GetDelta(const FVec3& WorldSpaceX1, const FVec3& WorldSpaceX2, const int32 ConstraintIndex, const int32 SpringIndex) const
+{
+	FSolverBody& Body0 = *ConstraintSolverBodies[ConstraintIndex][0];
+	FSolverBody& Body1 = *ConstraintSolverBodies[ConstraintIndex][1];
+	check(Body0.IsDynamic() || Body1.IsDynamic());
+
 	const FVec3 Difference = WorldSpaceX2 - WorldSpaceX1;
 	FReal Distance = Difference.Size();
 	check(Distance > 1e-7);
 
-	const FReal InvM0 = bIsRigidDynamic0 ? PBDRigid0->InvM() : (FReal)0.;
-	const FReal InvM1 = bIsRigidDynamic1 ? PBDRigid1->InvM() : (FReal)0.;
 	const FVec3 Direction = Difference / Distance;
 	const FVec3 Delta = (Distance - SpringDistances[ConstraintIndex][SpringIndex]) * Direction;
-	return Stiffness * Delta / (InvM0 + InvM1);
+	return Stiffness * Delta / (Body0.InvM() + Body1.InvM());
 }
 
 void FPBDRigidDynamicSpringConstraints::ApplySingle(const FReal Dt, int32 ConstraintIndex) const
 {
-	FGenericParticleHandle Particle0 = Constraints[ConstraintIndex][0];
-	FGenericParticleHandle Particle1 = Constraints[ConstraintIndex][1];
-	const bool bIsRigidDynamic0 = Particle0->IsDynamic();
-	const bool bIsRigidDynamic1 = Particle1->IsDynamic();
-	check(bIsRigidDynamic0 || bIsRigidDynamic1);
-
-	FRotation3 Q0 = FParticleUtilities::GetCoMWorldRotation(Particle0);
-	FRotation3 Q1 = FParticleUtilities::GetCoMWorldRotation(Particle1);
-	FVec3 P0 = FParticleUtilities::GetCoMWorldPosition(Particle0);
-	FVec3 P1 = FParticleUtilities::GetCoMWorldPosition(Particle1);
+	check(ConstraintSolverBodies[ConstraintIndex][0] != nullptr);
+	check(ConstraintSolverBodies[ConstraintIndex][1] != nullptr);
+	FSolverBody& Body0 = *ConstraintSolverBodies[ConstraintIndex][0];
+	FSolverBody& Body1 = *ConstraintSolverBodies[ConstraintIndex][1];
+	check(Body0.IsDynamic() || Body1.IsDynamic());
 
 	const int32 NumSprings = SpringDistances[ConstraintIndex].Num();
-	const FMatrix33 WorldSpaceInvI1 = bIsRigidDynamic0 ? Utilities::ComputeWorldSpaceInertia(Q0, Particle0->InvI()) : FMatrix33(0);
-	const FMatrix33 WorldSpaceInvI2 = bIsRigidDynamic1 ? Utilities::ComputeWorldSpaceInertia(Q1, Particle1->InvI()) : FMatrix33(0);;
 	for (int32 SpringIndex = 0; SpringIndex < NumSprings; ++SpringIndex)
 	{
 		const FVec3& Distance0 = Distances[ConstraintIndex][SpringIndex][0];
 		const FVec3& Distance1 = Distances[ConstraintIndex][SpringIndex][1];
-		const FVec3 WorldSpaceX1 = Particle0->Q().RotateVector(Distance0) + Particle0->P();
-		const FVec3 WorldSpaceX2 = Particle1->Q().RotateVector(Distance1) + Particle1->P();
+		const FVec3 WorldSpaceX1 = Body0.Q().RotateVector(Distance0) + Body0.P();
+		const FVec3 WorldSpaceX2 = Body1.Q().RotateVector(Distance1) + Body1.P();
 		const FVec3 Delta = GetDelta(WorldSpaceX1, WorldSpaceX2, ConstraintIndex, SpringIndex);
 
-		if (bIsRigidDynamic0)
+		if (Body0.IsDynamic())
 		{
-			const FVec3 Radius = WorldSpaceX1 - P0;
-			P0 += Particle0->InvM() * Delta;
-			Q0 += FRotation3::FromElements(WorldSpaceInvI1 * FVec3::CrossProduct(Radius, Delta), 0.f) * Q0 * FReal(0.5);
-			Q0.Normalize();
-			FParticleUtilities::SetCoMWorldTransform(Particle0, P0, Q0);
+			const FVec3 Radius = WorldSpaceX1 - Body0.P();
+			const FVec3 DX = Body0.InvM() * Delta;
+			const FVec3 DR = Body0.InvI() * FVec3::CrossProduct(Radius, Delta);
+			Body0.ApplyTransformDelta(DX, DR);
+			Body0.UpdateRotationDependentState();
 		}
 
-		if (bIsRigidDynamic1)
+		if (Body1.IsDynamic())
 		{
-			const FVec3 Radius = WorldSpaceX2 - P1;
-			P1 -= Particle1->InvM() * Delta;
-			Q1 += FRotation3::FromElements(WorldSpaceInvI2 * FVec3::CrossProduct(Radius, -Delta), 0.f) * Q1 * FReal(0.5);
-			Q1.Normalize();
-			FParticleUtilities::SetCoMWorldTransform(Particle1, P1, Q1);
+			const FVec3 Radius = WorldSpaceX2 - Body1.P();
+			const FVec3 DX = Body1.InvM() * -Delta;
+			const FVec3 DR = Body1.InvI() * FVec3::CrossProduct(Radius, -Delta);
+			Body1.ApplyTransformDelta(DX, DR);
+			Body1.UpdateRotationDependentState();
 		}
 	}
 }
diff --git a/Engine/Source/Runtime/Experimental/Chaos/Private/Chaos/PBDRigidSpringConstraints.cpp b/Engine/Source/Runtime/Experimental/Chaos/Private/Chaos/PBDRigidSpringConstraints.cpp
--- a/Engine/Source/Runtime/Experimental/Chaos/Private/Chaos/PBDRigidSpringConstraints.cpp
+++ b/Engine/Source/Runtime/Experimental/Chaos/Private/Chaos/PBDRigidSpringConstraints.cpp
@@ -1,6 +1,7 @@
 // Copyright Epic Games, Inc. All Rights Reserved.
 
 #include "Chaos/PBDRigidSpringConstraints.h"
+#include "Chaos/Evolution/SolverDatas.h"
 #include "Chaos/Particle/ParticleUtilities.h"
 #include "Chaos/Utilities.h"
 
@@ -36,12 +37,18 @@ namespace Chaos
 		ConcreteContainer()->SetRestLength(ConstraintIndex, SpringLength);
 	}
 
+	void FPBDRigidSpringConstraintHandle::GatherInput(const FReal Dt, const int32 Particle0Level, const int32 Particle1Level, FPBDIslandSolverData& SolverData)
+	{
+		ConcreteContainer()->GatherInput(Dt, ConstraintIndex, Particle0Level, Particle1Level, SolverData);
+	}
+
+
 	//
 	// Container Impl
 	//
 
 	FPBDRigidSpringConstraints::FPBDRigidSpringConstraints()
-		: FPBDConstraintContainer(EConstraintContainerType::RigidSpring)
+		: FPBDIndexedConstraintContainer(FConstraintContainerHandle::StaticType())
 	{}
 
 	FPBDRigidSpringConstraints::~FPBDRigidSpringConstraints()
@@ -56,7 +63,9 @@ namespace Chaos
 		SpringSettings.Emplace(FSpringSettings({ Stiffness, Damping, RestLength }));
 
 		Distances.Add({});
-		UpdateDistance(ConstraintIndex, InLocations[0], InLocations[1]);
+		InitDistance(ConstraintIndex, InLocations[0], InLocations[1]);
+
+		ConstraintSolverBodies.Add({ nullptr, nullptr });
 
 		return Handles.Last();
 	}
@@ -75,18 +84,20 @@ namespace Chaos
 		Constraints.RemoveAtSwap(ConstraintIndex);
 		SpringSettings.RemoveAtSwap(ConstraintIndex);
 		Distances.RemoveAtSwap(ConstraintIndex);
+		ConstraintSolverBodies.RemoveAtSwap(ConstraintIndex);
 		Handles.RemoveAtSwap(ConstraintIndex);
 
 		// Update the handle for the constraint that was moved
 		if (ConstraintIndex < Handles.Num())
 		{
-			FConstraintHandle* Handle = Handles[ConstraintIndex];
-			SetConstraintIndex(Handle, ConstraintIndex);
+			SetConstraintIndex(Handles[ConstraintIndex], ConstraintIndex);
 		}
 	}
 
-	void FPBDRigidSpringConstraints::UpdateDistance(int32 ConstraintIndex, const FVec3& Location0, const FVec3& Location1)
+	void FPBDRigidSpringConstraints::InitDistance(int32 ConstraintIndex, const FVec3& Location0, const FVec3& Location1)
 	{
+		// \note: this is called during initialization, not suring the solve. It therefore accesses
+		// the Particles directly, rather than the SolverBodies
 		const TVector<TGeometryParticleHandle<FReal, 3>*, 2>& Constraint = Constraints[ConstraintIndex];
 		const TGeometryParticleHandle<FReal, 3>* Particle0 = Constraint[0];
 		const TGeometryParticleHandle<FReal, 3>* Particle1 = Constraint[1];
@@ -97,86 +108,106 @@ namespace Chaos
 
 	FVec3 FPBDRigidSpringConstraints::GetDelta(int32 ConstraintIndex, const FVec3& WorldSpaceX1, const FVec3& WorldSpaceX2) const
 	{
-		const TVector<TGeometryParticleHandle<FReal, 3>*, 2>& Constraint = Constraints[ConstraintIndex];
-		const TPBDRigidParticleHandle<FReal, 3>* PBDRigid0 = Constraint[0]->CastToRigidParticle();
-		const TPBDRigidParticleHandle<FReal, 3>* PBDRigid1 = Constraint[1]->CastToRigidParticle();
-		const bool bIsRigidDynamic0 = PBDRigid0 && PBDRigid0->ObjectState() == EObjectStateType::Dynamic;
-		const bool bIsRigidDynamic1 = PBDRigid1 && PBDRigid1->ObjectState() == EObjectStateType::Dynamic;
+		FSolverBody& Body0 = *ConstraintSolverBodies[ConstraintIndex][0];
+		FSolverBody& Body1 = *ConstraintSolverBodies[ConstraintIndex][1];
 
-		if (!bIsRigidDynamic0 && !bIsRigidDynamic1)
+		if (!Body0.IsDynamic() && !Body1.IsDynamic())
 		{
 			return FVec3(0);
 		}
 
 		const FVec3 Difference = WorldSpaceX2 - WorldSpaceX1;
-
 		const FReal Distance = Difference.Size();
-		check(Distance > 1e-7);
+		if (Distance < FReal(1e-7))
+		{
+			return FVec3(0);
+		}
 
 		const FVec3 Direction = Difference / Distance;
 		const FVec3 Delta = (Distance - SpringSettings[ConstraintIndex].RestLength) * Direction;
-		const FReal InvM0 = (bIsRigidDynamic0) ? PBDRigid0->InvM() : (FReal)0;
-		const FReal InvM1 = (bIsRigidDynamic1) ? PBDRigid1->InvM() : (FReal)0;
-		const FReal CombinedMass = InvM0 + InvM1;
+		const FReal CombinedInvMass = Body0.InvM() + Body1.InvM();
 
-		return SpringSettings[ConstraintIndex].Stiffness * Delta / CombinedMass;
+		return SpringSettings[ConstraintIndex].Stiffness * Delta / CombinedInvMass;
 	}
 
-	bool FPBDRigidSpringConstraints::Apply(const FReal Dt, const int32 It, const int32 NumIts)
+	void FPBDRigidSpringConstraints::GatherInput(const FReal Dt, FPBDIslandSolverData& SolverData)
 	{
 		for (int32 ConstraintIndex = 0; ConstraintIndex < NumConstraints(); ++ConstraintIndex)
 		{
-			ApplySingle(Dt, ConstraintIndex);
+			GatherInput(Dt, ConstraintIndex, INDEX_NONE, INDEX_NONE, SolverData);
 		}
-		return false;
 	}
 
-	bool FPBDRigidSpringConstraints::Apply(const FReal Dt, const TArray<FConstraintContainerHandle*>& InConstraintHandles, const int32 It, const int32 NumIts)
+	bool FPBDRigidSpringConstraints::ApplyPhase1(const FReal Dt, const int32 It, const int32 NumIts, FPBDIslandSolverData& SolverData)
 	{
-		for (FConstraintContainerHandle* ConstraintHandle : InConstraintHandles)
+		return ApplyPhase1Serial(Dt, It, NumIts, SolverData);
+	}
+
+	void FPBDRigidSpringConstraints::SetNumIslandConstraints(const int32 NumIslandConstraints, FPBDIslandSolverData& SolverData)
+	{
+		SolverData.GetConstraintIndices(ContainerId).Reset(NumIslandConstraints);
+	}
+
+	void FPBDRigidSpringConstraints::GatherInput(const FReal Dt, const int32 ConstraintIndex, const int32 Particle0Level, const int32 Particle1Level, FPBDIslandSolverData& SolverData)
+	{
+		SolverData.GetConstraintIndices(ContainerId).Add(ConstraintIndex);
+	
+		ConstraintSolverBodies[ConstraintIndex] =
+		{
+			SolverData.GetBodyContainer().FindOrAdd(Constraints[ConstraintIndex][0]),
+			SolverData.GetBodyContainer().FindOrAdd(Constraints[ConstraintIndex][1])
+		};
+	}
+
+	void FPBDRigidSpringConstraints::ScatterOutput(FReal Dt, FPBDIslandSolverData& SolverData)
+	{
+		for (int32 ConstraintIndex : SolverData.GetConstraintIndices(ContainerId))
 		{
-			ApplySingle(Dt, ConstraintHandle->GetConstraintIndex());
+			ConstraintSolverBodies[ConstraintIndex] = { nullptr, nullptr };
 		}
-		return false;
 	}
 
-	void FPBDRigidSpringConstraints::ApplySingle(const FReal Dt, int32 ConstraintIndex) const
+	bool FPBDRigidSpringConstraints::ApplyPhase1Serial(const FReal Dt, const int32 It, const int32 NumIts, FPBDIslandSolverData& SolverData)
 	{
-		const FConstrainedParticlePair& Constraint = Constraints[ConstraintIndex];
+		for (int32 ConstraintIndex : SolverData.GetConstraintIndices(ContainerId))
+		{
+			ApplyPhase1Single(Dt, ConstraintIndex);
+		}
+		return true;
+	}
 
-		FGenericParticleHandle Particle0 = Constraints[ConstraintIndex][0];
-		FGenericParticleHandle Particle1 = Constraints[ConstraintIndex][1];
-		const bool bIsRigidDynamic0 = Particle0->IsDynamic();
-		const bool bIsRigidDynamic1 = Particle1->IsDynamic();
+	void FPBDRigidSpringConstraints::ApplyPhase1Single(const FReal Dt, int32 ConstraintIndex) const
+	{
+		check(ConstraintSolverBodies[ConstraintIndex][0] != nullptr);
+		check(ConstraintSolverBodies[ConstraintIndex][1] != nullptr);
+		FSolverBody& Body0 = *ConstraintSolverBodies[ConstraintIndex][0];
+		FSolverBody& Body1 = *ConstraintSolverBodies[ConstraintIndex][1];
 
-		check((bIsRigidDynamic0 && bIsRigidDynamic1) || (!bIsRigidDynamic0 && bIsRigidDynamic1) || (bIsRigidDynamic0 && bIsRigidDynamic1));
+		if (!Body0.IsDynamic() && !Body1.IsDynamic())
+		{
+			return;
+		}
 
-		const FVec3 WorldSpaceX1 = Particle0->Q().RotateVector(Distances[ConstraintIndex][0]) + Particle0->P();
-		const FVec3 WorldSpaceX2 = Particle1->Q().RotateVector(Distances[ConstraintIndex][1]) + Particle1->P();
+		const FVec3 WorldSpaceX1 = Body0.Q().RotateVector(Distances[ConstraintIndex][0]) + Body0.P();
+		const FVec3 WorldSpaceX2 = Body1.Q().RotateVector(Distances[ConstraintIndex][1]) + Body1.P();
 		const FVec3 Delta = GetDelta(ConstraintIndex, WorldSpaceX1, WorldSpaceX2);
 
-		if (bIsRigidDynamic0)
+		if (Body0.IsDynamic())
 		{
-			FRotation3 Q0 = FParticleUtilities::GetCoMWorldRotation(Particle0);
-			FVec3 P0 = FParticleUtilities::GetCoMWorldPosition(Particle0);
-			const FMatrix33 WorldSpaceInvI1 = Utilities::ComputeWorldSpaceInertia(Q0, Particle0->InvI());
-			const FVec3 Radius = WorldSpaceX1 - P0;
-			P0 += Particle0->InvM() * Delta;
-			Q0 += TRotation<FReal, 3>::FromElements(WorldSpaceInvI1 * FVec3::CrossProduct(Radius, Delta), 0.f) * Q0 * FReal(0.5);
-			Q0.Normalize();
-			FParticleUtilities::SetCoMWorldTransform(Particle0, P0, Q0);
+			const FVec3 Radius = WorldSpaceX1 - Body0.P();
+			const FVec3 DX = Body0.InvM() * Delta;
+			const FVec3 DR = Body0.InvI() * FVec3::CrossProduct(Radius, Delta);
+			Body0.ApplyTransformDelta(DX, DR);
+			Body0.UpdateRotationDependentState();
 		}
 
-		if (bIsRigidDynamic1)
+		if (Body1.IsDynamic())
 		{
-			FRotation3 Q1 = FParticleUtilities::GetCoMWorldRotation(Particle1);
-			FVec3 P1 = FParticleUtilities::GetCoMWorldPosition(Particle1);
-			const FMatrix33 WorldSpaceInvI2 = Utilities::ComputeWorldSpaceInertia(Q1, Particle1->InvI());
-			const FVec3 Radius = WorldSpaceX2 - P1;
-			P1 -= Particle1->InvM() * Delta;
-			Q1 += TRotation<FReal, 3>::FromElements(WorldSpaceInvI2 * FVec3::CrossProduct(Radius, -Delta), 0.f) * Q1 * FReal(0.5);
-			Q1.Normalize();
-			FParticleUtilities::SetCoMWorldTransform(Particle1, P1, Q1);
+			const FVec3 Radius = WorldSpaceX2 - Body1.P();
+			const FVec3 DX = Body1.InvM() * -Delta;
+			const FVec3 DR = Body1.InvI() * FVec3::CrossProduct(Radius, -Delta);
+			Body1.ApplyTransformDelta(DX, DR);
+			Body1.UpdateRotationDependentState();
 		}
 	}
 }
diff --git a/Engine/Source/Runtime/Experimental/Chaos/Private/Chaos/PBDRigidsEvolution.cpp b/Engine/Source/Runtime/Experimental/Chaos/Private/Chaos/PBDRigidsEvolution.cpp
--- a/Engine/Source/Runtime/Experimental/Chaos/Private/Chaos/PBDRigidsEvolution.cpp
+++ b/Engine/Source/Runtime/Experimental/Chaos/Private/Chaos/PBDRigidsEvolution.cpp
@@ -24,12 +24,15 @@ namespace Chaos
 	CHAOS_API int32 FixBadAccelerationStructureRemoval = 1;
 	FAutoConsoleVariableRef CVarFixBadAccelerationStructureRemoval(TEXT("p.FixBadAccelerationStructureRemoval"), FixBadAccelerationStructureRemoval, TEXT(""));
 
-	CHAOS_API int32 ConstraintIterationsOverride = -1;
-	FAutoConsoleVariableRef CVarConstraintIterationsOverride(TEXT("p.ConstraintIterationsOverride"), ConstraintIterationsOverride, TEXT("Override for number of collision iterations."));
+	CHAOS_API int32 AccelerationStructureSplitStaticAndDynamic = 1;
+	FAutoConsoleVariableRef CVarAccelerationStructureSplitStaticAndDynamic(TEXT("p.Chaos.AccelerationStructureSplitStaticDynamic"), AccelerationStructureSplitStaticAndDynamic, TEXT("Set to 1: Sort Dynamic and Static bodies into seperate acceleration structures, any other value will disable the feature"));
+
+	CHAOS_API int32 AccelerationStructureTimeSlicingMaxQueueSizeBeforeForce = 1000;
+	FAutoConsoleVariableRef CVarAccelerationStructureTimeSlicingMaxQueueSizeBeforeForce(TEXT("p.Chaos.AccelerationStructureTimeSlicingMaxQueueSizeBeforeForce"), AccelerationStructureTimeSlicingMaxQueueSizeBeforeForce, TEXT("If the update queue reaches this limit, time slicing will be disabled, and the acceleration structure will be built at once"));
+
+	CHAOS_API int32 AccelerationStructureTimeSlicingMaxBytesCopy = 100000;
+	FAutoConsoleVariableRef CVarAccelerationStructureTimeSlicingMaxBytesCopy(TEXT("p.Chaos.AccelerationStructureTimeSlicingMaxBytesCopy"), AccelerationStructureTimeSlicingMaxBytesCopy, TEXT("The Maximum number of bytes to copy to the external acceleration structure during Copy Time Slicing"));
 
-	CHAOS_API int32 PushOutIterationsOverride = -1;
-	FAutoConsoleVariableRef CVarPushOutIterationsOverride(TEXT("p.PushOutIterationsOverride"), PushOutIterationsOverride, TEXT("Override for number of push out iterations."));
-	
 	struct FAccelerationConfig
 	{
 		int32 BroadphaseType;
@@ -80,7 +83,13 @@ namespace Chaos
 
 			for (uint16 BucketIdx = 0; BucketIdx < NumBuckets; ++BucketIdx)
 			{
-				Collection->AddSubstructure(CreateAccelerationPerBucket_Threaded(Empty, BucketIdx, true), BucketIdx);
+				// For static bodies
+				Collection->AddSubstructure(CreateAccelerationPerBucket_Threaded(Empty, BucketIdx, true), BucketIdx, 0);
+				if (AccelerationStructureSplitStaticAndDynamic == 1)
+				{
+					// Non static bodies
+					Collection->AddSubstructure(CreateAccelerationPerBucket_Threaded(Empty, BucketIdx, true), BucketIdx, 1);
+				}
 			}
 
 			return TUniquePtr<ISpatialAccelerationCollection<FAccelerationStructureHandle, FReal, 3>>(Collection);
@@ -216,6 +225,7 @@ namespace Chaos
 	DECLARE_CYCLE_STAT(TEXT("CopyAccelerationStructure"), STAT_CopyAccelerationStructure, STATGROUP_Chaos);
 	DECLARE_CYCLE_STAT(TEXT("SwapAccelerationStructures"), STAT_SwapAccelerationStructures, STATGROUP_Chaos);
 	DECLARE_CYCLE_STAT(TEXT("AccelerationStructureTimeSlice"), STAT_AccelerationStructureTimeSlice, STATGROUP_Chaos);
+	DECLARE_CYCLE_STAT(TEXT("AccelerationStructureTimeSliceCopy"), STAT_AccelerationStructureTimeSliceCopy, STATGROUP_Chaos);
 	DECLARE_CYCLE_STAT(TEXT("CreateInitialAccelerationStructure"), STAT_CreateInitialAccelerationStructure, STATGROUP_Chaos);
 	DECLARE_CYCLE_STAT(TEXT("CreateNonSlicedStructures"), STAT_CreateNonSlicedStructures, STATGROUP_Chaos);
 
@@ -267,111 +277,202 @@ namespace Chaos
 		if (Substructure->template As<BVType>())
 		{
 			auto Collection = MakeUnique<TSpatialAccelerationCollection<BVType>>();
-			Collection->AddSubstructure(MoveTemp(Substructure), 0);
+			Collection->AddSubstructure(MoveTemp(Substructure), 0, 0);
 			return Collection;
 		}
 		else if (Substructure->template As<AABBType>())
 		{
 			auto Collection = MakeUnique<TSpatialAccelerationCollection<AABBType>>();
-			Collection->AddSubstructure(MoveTemp(Substructure), 0);
+			Collection->AddSubstructure(MoveTemp(Substructure), 0, 0);
 			return Collection;
 		}
 		else
 		{
 			using AccelType = TAABBTree<FAccelerationStructureHandle, TBoundingVolume<FAccelerationStructureHandle>>;
 			auto Collection = MakeUnique<TSpatialAccelerationCollection<AccelType>>();
-			Collection->AddSubstructure(MoveTemp(Substructure), 0);
+			Collection->AddSubstructure(MoveTemp(Substructure), 0, 0);
 			return Collection;
 		}
 	}
 
-	void FPBDRigidsEvolutionBase::FChaosAccelerationStructureTask::UpdateStructure(FAccelerationStructure* AccelerationStructure)
+	void FPBDRigidsEvolutionBase::FChaosAccelerationStructureTask::UpdateStructure(FAccelerationStructure* AccelerationStructure, FAccelerationStructure* AccelerationStructureCopy)
 	{
 		LLM_SCOPE(ELLMTag::ChaosAcceleration);
 
-		if (bNeedsReset)
+		// This function (running in FChaosAccelerationStructureTask) can be in one of two modes
+		// 1) Create new time sliced acceleration structures and run the first slice  (CreatingNewTimeSlicedStructures)
+		// or
+		// 2) Progress time slices (!CreatingNewTimeSlicedStructures)
+		// In either case, if time slicing is done create the non time sliced acceleration structures
+
+
+		auto GetSubStructureCopy = [AccelerationStructureCopy](FSpatialAccelerationIdx Index)
 		{
-			AccelerationStructure->Reset();
-		}
+			ISpatialAcceleration<FAccelerationStructureHandle, FReal, 3>* Result = nullptr;
+			if (AccelerationStructureCopy)
+			{
+				Result = AccelerationStructureCopy->GetSubstructure(Index);
+				if (!Result)
+				{
+					ensure(false);
+				}
+			}
+			return Result;
+		};
+
+		LLM_SCOPE(ELLMTag::ChaosAcceleration);
 
 		uint8 ActiveBucketsMask = SpatialCollectionFactory.GetActiveBucketsMask();
 		TArray<TSOAView<FSpatialAccelerationCache>> ViewsPerBucket[FSpatialAccelerationIdx::MaxBuckets];
-		TArray<uint8> TimeSlicedBucketsToCreate;
-		TArray<uint8> NonTimeSlicedBucketsToCreate;
+		TArray<FSpatialAccelerationIdx> TimeSlicedBucketsToCreate;
+		TArray<FSpatialAccelerationIdx> NonTimeSlicedBucketsToCreate;
 
 		bool IsTimeSlicingProgressing = false;
+		bool CreatingNewTimeSlicedStructures = bNeedsReset;
 
-		//merge buckets. todo: support multiple entries per bucket (i.e. dynamic vs static)
 		for(const auto& Itr : SpatialAccelerationCache)
 		{
 			const FSpatialAccelerationIdx SpatialIdx = Itr.Key;
 			const FSpatialAccelerationCache& Cache = *Itr.Value;
 			const uint8 BucketIdx = static_cast<uint8>((1 << SpatialIdx.Bucket) & ActiveBucketsMask ? SpatialIdx.Bucket : 0);
-			if(AccelerationStructure->GetSubstructure(SpatialIdx) && !AccelerationStructure->GetSubstructure(SpatialIdx)->IsAsyncTimeSlicingComplete())
+			const uint16 BucketInnerIdx = SpatialIdx.InnerIdx;
+			ISpatialAcceleration<FAccelerationStructureHandle, FReal, 3>* SubStructure = AccelerationStructure->GetSubstructure(SpatialIdx);
+			// Get the sub structure we need to copy our results to
+			ISpatialAcceleration<FAccelerationStructureHandle, FReal, 3>* AccelerationSUBStructureCopy = GetSubStructureCopy(SpatialIdx);
+
+			if (!SubStructure)
 			{
-				SCOPE_CYCLE_COUNTER(STAT_AccelerationStructureTimeSlice);
+				ensure(false);
+				continue;
+			}
+			
+			if (CreatingNewTimeSlicedStructures && !SubStructure->ShouldRebuild() && AccelerationStructureSplitStaticAndDynamic == 1)
+			{
+				ensure(SubStructure->IsAsyncTimeSlicingComplete());
+
+				if (AccelerationSUBStructureCopy)
+				{
+					AccelerationSUBStructureCopy->PrepareCopyTimeSliced(*SubStructure);
+					ensure(!AccelerationSUBStructureCopy->IsAsyncTimeSlicingComplete());
+					AccelerationSUBStructureCopy->ProgressCopyTimeSliced(*SubStructure, IsForceFullBuild ? -1 : AccelerationStructureTimeSlicingMaxBytesCopy);
+					if (!AccelerationSUBStructureCopy->IsAsyncTimeSlicingComplete())
+					{
+						IsTimeSlicingProgressing = true;
+					}
+				}
+
+				// This structure has not changed and therefore does not need to rebuild 
+				continue;
+			}
 
-				AccelerationStructure->GetSubstructure(SpatialIdx)->ProgressAsyncTimeSlicing(IsForceFullBuild);
+			if (CreatingNewTimeSlicedStructures)
+			{
+				SubStructure->Reset();
+			}
+
+			if(!SubStructure->IsAsyncTimeSlicingComplete())
+			{
+				SCOPE_CYCLE_COUNTER(STAT_AccelerationStructureTimeSlice);
+				
+				SubStructure->ProgressAsyncTimeSlicing(IsForceFullBuild);
 
 				// is it still progressing or now complete
-				if(!AccelerationStructure->GetSubstructure(SpatialIdx)->IsAsyncTimeSlicingComplete())
+				if(!SubStructure->IsAsyncTimeSlicingComplete())
 				{
 					IsTimeSlicingProgressing = true;
 				}
+				else
+				{
+					if (AccelerationSUBStructureCopy)
+					{
+						AccelerationSUBStructureCopy->PrepareCopyTimeSliced(*SubStructure);
+						IsTimeSlicingProgressing |= !AccelerationSUBStructureCopy->IsAsyncTimeSlicingComplete();
+					}
+				}
+			}
+			else if (AccelerationSUBStructureCopy && !AccelerationSUBStructureCopy->IsAsyncTimeSlicingComplete())
+			{
+				SCOPE_CYCLE_COUNTER(STAT_AccelerationStructureTimeSliceCopy);
+				AccelerationSUBStructureCopy->ProgressCopyTimeSliced(*SubStructure, IsForceFullBuild ? -1 : AccelerationStructureTimeSlicingMaxBytesCopy);
+				if (!AccelerationSUBStructureCopy->IsAsyncTimeSlicingComplete())
+					IsTimeSlicingProgressing = true;
 			}
 			else
 			{
-				ViewsPerBucket[BucketIdx].Add(const_cast<FSpatialAccelerationCache*>(&Cache));
-				if(AccelerationStructure->IsBucketActive(SpatialIdx.Bucket))
-				{
-					AccelerationStructure->RemoveSubstructure(SpatialIdx);
-				}
-
+				// Queue for potential Creation
+				ViewsPerBucket[BucketIdx].SetNum(FMath::Max(BucketInnerIdx + 1, ViewsPerBucket[BucketIdx].Num()));
+				ViewsPerBucket[BucketIdx][BucketInnerIdx] = const_cast<FSpatialAccelerationCache*>(&Cache);
+				
 				if(SpatialCollectionFactory.IsBucketTimeSliced(BucketIdx))
 				{
-					TimeSlicedBucketsToCreate.Add(SpatialIdx.Bucket);
+					TimeSlicedBucketsToCreate.Add(SpatialIdx);
 				} else
 				{
-					NonTimeSlicedBucketsToCreate.Add(SpatialIdx.Bucket);
+					NonTimeSlicedBucketsToCreate.Add(SpatialIdx);
 				}
 			}
 		}
 
 		//todo: creation can go wide, insertion to collection cannot
-		for(uint8 BucketIdx : TimeSlicedBucketsToCreate)
+		if (CreatingNewTimeSlicedStructures)
 		{
-			if(ViewsPerBucket[BucketIdx].Num())
+			for (FSpatialAccelerationIdx SpatialAccIdx : TimeSlicedBucketsToCreate)
 			{
 				SCOPE_CYCLE_COUNTER(STAT_CreateInitialAccelerationStructure);
 
-				auto ParticleView = MakeConstParticleView(MoveTemp(ViewsPerBucket[BucketIdx]));
-				auto NewStruct = SpatialCollectionFactory.CreateAccelerationPerBucket_Threaded(ParticleView,BucketIdx,IsForceFullBuild);
+				auto ParticleView = MakeConstParticleView(MoveTemp(ViewsPerBucket[SpatialAccIdx.Bucket][SpatialAccIdx.InnerIdx]));
+				// Create and run the first slice here
+				auto NewStruct = SpatialCollectionFactory.CreateAccelerationPerBucket_Threaded(ParticleView, SpatialAccIdx.Bucket, IsForceFullBuild);
+				{
+					NewStruct->ClearShouldRebuild();  // Mark as pristine
+				}
+
+				ISpatialAcceleration<FAccelerationStructureHandle, FReal, 3>* AccelerationSUBStructureCopy = GetSubStructureCopy(SpatialAccIdx);
+
+				if (AccelerationSUBStructureCopy && NewStruct->IsAsyncTimeSlicingComplete())
+				{
+					AccelerationSUBStructureCopy->PrepareCopyTimeSliced(*NewStruct);
+					AccelerationSUBStructureCopy->ProgressCopyTimeSliced(*NewStruct, IsForceFullBuild ? -1 : AccelerationStructureTimeSlicingMaxBytesCopy);
+					IsTimeSlicingProgressing |= !AccelerationSUBStructureCopy->IsAsyncTimeSlicingComplete();
+				}
 
 				//If new structure is not done mark time slicing in progress
 				IsTimeSlicingProgressing |= !NewStruct->IsAsyncTimeSlicingComplete();
 
-				AccelerationStructure->AddSubstructure(MoveTemp(NewStruct),BucketIdx);
-
+				if (AccelerationStructure->IsBucketActive(SpatialAccIdx.Bucket))
+				{
+					AccelerationStructure->RemoveSubstructure(SpatialAccIdx);
+				}
+				AccelerationStructure->AddSubstructure(MoveTemp(NewStruct), SpatialAccIdx.Bucket, SpatialAccIdx.InnerIdx);
 			}
 		}
 
+		// If it's not progressing then it is finished so we can perform the final copy if required
 		AccelerationStructure->SetAllAsyncTasksComplete(!IsTimeSlicingProgressing);
 
-		// If it's not progressing then it is finished so we can perform the final copy if required
+		// Create and Build the non time sliced structures when all time sliced structures are complete
 		if(!IsTimeSlicingProgressing)
 		{
 			//todo: creation can go wide, insertion to collection cannot
-			for(uint8 BucketIdx : NonTimeSlicedBucketsToCreate)
+			for(FSpatialAccelerationIdx SpatialAccIdx : NonTimeSlicedBucketsToCreate)
 			{
-				if(ViewsPerBucket[BucketIdx].Num())
+				if(ViewsPerBucket[SpatialAccIdx.Bucket].Num() > SpatialAccIdx.InnerIdx)
 				{
 					SCOPE_CYCLE_COUNTER(STAT_CreateNonSlicedStructures);
 
-					auto ParticleView = MakeConstParticleView(MoveTemp(ViewsPerBucket[BucketIdx]));
-					auto NewStruct = SpatialCollectionFactory.CreateAccelerationPerBucket_Threaded(ParticleView,BucketIdx,IsForceFullBuild);
-
-					AccelerationStructure->AddSubstructure(MoveTemp(NewStruct),BucketIdx);
+					auto ParticleView = MakeConstParticleView(MoveTemp(ViewsPerBucket[SpatialAccIdx.Bucket][SpatialAccIdx.InnerIdx]));
+					auto NewStruct = SpatialCollectionFactory.CreateAccelerationPerBucket_Threaded(ParticleView,SpatialAccIdx.Bucket,IsForceFullBuild);
 
+					ISpatialAcceleration<FAccelerationStructureHandle, FReal, 3>* AccelerationSUBStructureCopy = GetSubStructureCopy(SpatialAccIdx);
+					if (AccelerationSUBStructureCopy)
+					{
+						*AccelerationSUBStructureCopy = *NewStruct; // Normal copy
+					}
+					if (AccelerationStructure->IsBucketActive(SpatialAccIdx.Bucket))
+					{
+						AccelerationStructure->RemoveSubstructure(SpatialAccIdx);
+					}
+					AccelerationStructure->AddSubstructure(MoveTemp(NewStruct),SpatialAccIdx.Bucket, SpatialAccIdx.InnerIdx);
 				}
 			}
 		}
@@ -381,9 +482,19 @@ namespace Chaos
 	{
 		LLM_SCOPE(ELLMTag::ChaosAcceleration);
 
-		//Rebuild both structures. TODO: probably faster to time slice the copy instead of doing two time sliced builds
-		UpdateStructure(InternalStructure);
-		UpdateStructure(ExternalStructure);
+#if !WITH_EDITOR
+		CSV_SCOPED_TIMING_STAT(ChaosPhysicsTimers, AccelerationStructBuildTask);
+#endif
+		if (AccelerationStructureSplitStaticAndDynamic == 1)
+		{
+			UpdateStructure(InternalStructure, ExternalStructure);
+		}
+		else
+		{
+			//Rebuild both structures. TODO: probably faster to time slice the copy instead of doing two time sliced builds (This happens with AccelerationStructureSplitStaticAndDynamic)
+			UpdateStructure(InternalStructure);
+			UpdateStructure(ExternalStructure);
+		}
 	}
 
 	void FPBDRigidsEvolutionBase::ApplyParticlePendingData(const FPendingSpatialData& SpatialData, FAccelerationStructure& AccelerationStructure, bool bUpdateCache)
@@ -475,9 +586,6 @@ namespace Chaos
 		UniqueIndicesPendingRelease.Reset();
 		AsyncAccelerationQueue.Reset();
 
-		//other queues are no longer needed since we've flushed all operations and now have a pristine structure
-		InternalAccelerationQueue.Reset();
-
 		AsyncInternalAcceleration->SetSyncTimestamp(LatestExternalTimestampConsumed_Internal);
 		AsyncExternalAcceleration->SetSyncTimestamp(LatestExternalTimestampConsumed_Internal);
 	}
@@ -555,7 +663,7 @@ namespace Chaos
 		SCOPE_CYCLE_COUNTER(STAT_ComputeIntermediateSpatialAcceleration);
 		CHAOS_SCOPED_TIMER(ComputeIntermediateSpatialAcceleration);
 
-		bool ForceFullBuild = InternalAccelerationQueue.Num() > 1000;
+		bool ForceFullBuild = InternalAccelerationQueue.Num() > AccelerationStructureTimeSlicingMaxQueueSizeBeforeForce;
 
 		if (!AccelerationStructureTaskComplete)
 		{
@@ -576,6 +684,8 @@ namespace Chaos
 			WaitOnAccelerationStructure();
 		}
 
+		// Readability note:
+		// AccelerationStructureTaskComplete It is a reference counted pointer to a TaskGraph event, that will be null if it is not set up
 		bool AsyncComplete = !AccelerationStructureTaskComplete || AccelerationStructureTaskComplete->IsComplete();
 		if (AsyncComplete && !IsResimming())
 		{
@@ -591,9 +701,32 @@ namespace Chaos
 				bNeedsReset = true;
 
 				FlushAsyncAccelerationQueue();
+				
+				if (AccelerationStructureSplitStaticAndDynamic == 1)
+				{
+					FlushInternalAccelerationQueue();
+				}
+				else
+				{
+					//other queues are no longer needed since we've flushed all operations and now have a pristine structure
+					InternalAccelerationQueue.Reset();
+				}
 
+				// At this point the InternalAcceleration and the AsyncAcceleration structures will all be consistent with each other
+				// (If FlushInternalAccelerationQueue was flushed above)
+				// but the AsyncAcceleration now probably have fewer dirty elements (they were just rebuilt)
+				
 				//swap acceleration structure for new one
 				std::swap(InternalAcceleration, AsyncInternalAcceleration);
+				// The old InternalAcceleration is now in AsyncInternalAcceleration, and the constituent structures will be reused if no changes have been detected
+
+				if (AccelerationStructureSplitStaticAndDynamic == 1)
+				{
+					// If the new InternalAcceleration structure have constituents with no changes, we can copy it to the  AsyncInternalAcceleration for reuse
+					// This step can be left out to remove a copy here, at the cost of an unnecessary recalculation in the ChaosAccelerationStructureTask
+					// For structures that don't change often both source and destination will be pristine and no copies will be made (most of the time)
+					//CopyPristineAccelerationStructures(SpatialAccelerationCache, InternalAcceleration, AsyncInternalAcceleration, true);
+				}
 
 				//mark structure as ready for external thread
 				ExternalStructuresQueue.Enqueue(AsyncExternalAcceleration);
@@ -609,7 +742,7 @@ namespace Chaos
 			if (bCanStartAsyncTasks)
 			{
 				// we run the task for both starting a new accel structure as well as for the time-slicing
-			AccelerationStructureTaskComplete = TGraphTask<FChaosAccelerationStructureTask>::CreateTask().ConstructAndDispatchWhenReady(*SpatialCollectionFactory, SpatialAccelerationCache, AsyncInternalAcceleration, AsyncExternalAcceleration, ForceFullBuild, bIsSingleThreaded, bNeedsReset);
+				AccelerationStructureTaskComplete = TGraphTask<FChaosAccelerationStructureTask>::CreateTask().ConstructAndDispatchWhenReady(*SpatialCollectionFactory, SpatialAccelerationCache, AsyncInternalAcceleration, AsyncExternalAcceleration, ForceFullBuild, bIsSingleThreaded, bNeedsReset);
 			}
 		}
 		else
@@ -618,6 +751,33 @@ namespace Chaos
 		}
 	}
 
+
+	// copy Pristine substructures from FromStructure to ToStructure, if the substructure is not already pristine
+	// Assumption: Both structures represent the same state, one can just be dirtier than the other
+	void FPBDRigidsEvolutionBase::CopyPristineAccelerationStructures(const TMap<FSpatialAccelerationIdx, TUniquePtr<FSpatialAccelerationCache>>& SpatialAccelerationCache, FAccelerationStructure* FromStructure, FAccelerationStructure* ToStructure, bool CheckPristine)
+	{
+		
+		for (const auto& Itr : SpatialAccelerationCache)
+		{
+			const FSpatialAccelerationIdx SpatialIdx = Itr.Key;
+			auto* FromSubStructure = FromStructure->GetSubstructure(SpatialIdx);
+			auto* ToSubStructure = ToStructure->GetSubstructure(SpatialIdx);
+
+			if (!FromSubStructure || !ToSubStructure)
+			{
+				ensure(false);
+				continue;
+			}
+
+			if (!CheckPristine || (ToSubStructure->ShouldRebuild() && !FromSubStructure->ShouldRebuild()))
+			{
+				*ToSubStructure = *FromSubStructure;
+			}
+		}
+		// This is a pretty shallow copy: It won't copy anything that we copied above
+		*ToStructure = *FromStructure;
+	}
+
 	void FPBDRigidsEvolutionBase::UpdateExternalAccelerationStructure_External(
 		ISpatialAccelerationCollection<FAccelerationStructureHandle, FReal, 3>*& StructToUpdate, FPendingSpatialDataQueue& PendingExternal)
 	{
@@ -716,7 +876,7 @@ namespace Chaos
 				{
 					SubStructure = InternalAcceleration->RemoveSubstructure(FSpatialAccelerationIdx{ 0,0 });
 					Ar << SubStructure;
-					InternalAcceleration->AddSubstructure(MoveTemp(SubStructure), 0);
+					InternalAcceleration->AddSubstructure(MoveTemp(SubStructure), 0, 0);
 				}
 				else
 				{
@@ -819,9 +979,9 @@ namespace Chaos
 		{
 			Particles.DisableParticle(Particle);
 			RemoveParticleFromAccelerationStructure(*Particle);
+			ConstraintGraph.DisableParticle(Particle);
 		}
 
-		ConstraintGraph.DisableParticles(InParticles);
 		DisableConstraints(InParticles);
 	}
 
diff --git a/Engine/Source/Runtime/Experimental/Chaos/Private/Chaos/PBDRigidsEvolutionGBF.cpp b/Engine/Source/Runtime/Experimental/Chaos/Private/Chaos/PBDRigidsEvolutionGBF.cpp
--- a/Engine/Source/Runtime/Experimental/Chaos/Private/Chaos/PBDRigidsEvolutionGBF.cpp
+++ b/Engine/Source/Runtime/Experimental/Chaos/Private/Chaos/PBDRigidsEvolutionGBF.cpp
@@ -1,6 +1,7 @@
 // Copyright Epic Games, Inc. All Rights Reserved.
 #include "Chaos/PBDRigidsEvolutionGBF.h"
 #include "Chaos/Defines.h"
+#include "Chaos/Evolution/SolverBodyContainer.h"
 #include "Chaos/Framework/Parallel.h"
 #include "Chaos/ImplicitObjectTransformed.h"
 #include "Chaos/ImplicitObjectUnion.h"
@@ -22,6 +23,8 @@
 #include "Misc/Paths.h"
 #include "HAL/FileManager.h"
 
+//PRAGMA_DISABLE_OPTIMIZATION
+
 namespace Chaos
 {
 #if !UE_BUILD_SHIPPING
@@ -30,72 +33,76 @@ namespace Chaos
 	const bool bPendingHierarchyDump = false;
 #endif
 
-FRealSingle HackMaxAngularVelocity = 1000.f;
-FAutoConsoleVariableRef CVarHackMaxAngularVelocity(TEXT("p.HackMaxAngularVelocity"), HackMaxAngularVelocity, TEXT("Max cap on angular velocity: rad/s. This is only a temp solution and should not be relied on as a feature. -1.f to disable"));
-
-FRealSingle HackMaxVelocity = -1.f;
-FAutoConsoleVariableRef CVarHackMaxVelocity(TEXT("p.HackMaxVelocity2"), HackMaxVelocity, TEXT("Max cap on velocity: cm/s. This is only a temp solution and should not be relied on as a feature. -1.f to disable"));
-
-
-FRealSingle HackLinearDrag = 0.f;
-FAutoConsoleVariableRef CVarHackLinearDrag(TEXT("p.HackLinearDrag2"), HackLinearDrag, TEXT("Linear drag used to slow down objects. This is a hack and should not be relied on as a feature."));
-
-FRealSingle HackAngularDrag = 0.f;
-FAutoConsoleVariableRef CVarHackAngularDrag(TEXT("p.HackAngularDrag2"), HackAngularDrag, TEXT("Angular drag used to slow down objects. This is a hack and should not be relied on as a feature."));
-
-int DisableThreshold = 5;
-FAutoConsoleVariableRef CVarDisableThreshold(TEXT("p.DisableThreshold2"), DisableThreshold, TEXT("Disable threshold frames to transition to sleeping"));
-
-int CollisionDisableCulledContacts = 0;
-FAutoConsoleVariableRef CVarDisableCulledContacts(TEXT("p.CollisionDisableCulledContacts"), CollisionDisableCulledContacts, TEXT("Allow the PBDRigidsEvolutionGBF collision constraints to throw out contacts mid solve if they are culled."));
-
-FRealSingle BoundsThicknessVelocityMultiplier = 2.0f;	// @todo(chaos): more to FChaosSolverConfiguration
-FAutoConsoleVariableRef CVarBoundsThicknessVelocityMultiplier(TEXT("p.CollisionBoundsVelocityInflation"), BoundsThicknessVelocityMultiplier, TEXT("Collision velocity inflation for speculatibe contact generation.[def:2.0]"));
-
-FRealSingle SmoothedPositionLerpRate = 0.1f;
-FAutoConsoleVariableRef CVarSmoothedPositionLerpRate(TEXT("p.Chaos.SmoothedPositionLerpRate"), SmoothedPositionLerpRate, TEXT("The interpolation rate for the smoothed position calculation. Used for sleeping."));
-
-int DisableParticleUpdateVelocityParallelFor = 0;
-FAutoConsoleVariableRef CVarDisableParticleUpdateVelocityParallelFor(TEXT("p.DisableParticleUpdateVelocityParallelFor"), DisableParticleUpdateVelocityParallelFor, TEXT("Disable Particle Update Velocity ParallelFor and run the update on a single thread"));
-
-int32 ChaosSolverCollisionPriority = 0;
-FAutoConsoleVariableRef CVarChaosSolverCollisionPriority(TEXT("p.Chaos.Solver.Collision.Priority"), ChaosSolverCollisionPriority, TEXT("Set constraint priority. Larger values are evaluated later [def:0]"));
-
-int32 ChaosSolverJointPriority = 0;
-FAutoConsoleVariableRef CVarChaosSolverJointPriority(TEXT("p.Chaos.Solver.Joint.Priority"), ChaosSolverJointPriority, TEXT("Set constraint priority. Larger values are evaluated later [def:0]"));
-
-int32 ChaosSolverSuspensionPriority = 0;
-FAutoConsoleVariableRef CVarChaosSolverSuspensionPriority(TEXT("p.Chaos.Solver.Suspension.Priority"), ChaosSolverSuspensionPriority, TEXT("Set constraint priority. Larger values are evaluated later [def:0]"));
-
-bool DoTransferJointConstraintCollisions = true;
-FAutoConsoleVariableRef CVarDoTransferJointConstraintCollisions(TEXT("p.Chaos.Solver.Joint.TransferCollisions"), DoTransferJointConstraintCollisions, TEXT("Allows joints to apply collisions to the parent from the child when the Joints TransferCollisionScale is not 0 [def:true]"));
-
-
-
-DECLARE_CYCLE_STAT(TEXT("FPBDRigidsEvolutionGBF::AdvanceOneTimeStep"), STAT_Evolution_AdvanceOneTimeStep, STATGROUP_Chaos);
-DECLARE_CYCLE_STAT(TEXT("FPBDRigidsEvolutionGBF::UnclusterUnions"), STAT_Evolution_UnclusterUnions, STATGROUP_Chaos);
-DECLARE_CYCLE_STAT(TEXT("FPBDRigidsEvolutionGBF::Integrate"), STAT_Evolution_Integrate, STATGROUP_Chaos);
-DECLARE_CYCLE_STAT(TEXT("FPBDRigidsEvolutionGBF::KinematicTargets"), STAT_Evolution_KinematicTargets, STATGROUP_Chaos);
-DECLARE_CYCLE_STAT(TEXT("FPBDRigidsEvolutionGBF::PostIntegrateCallback"), STAT_Evolution_PostIntegrateCallback, STATGROUP_Chaos);
-DECLARE_CYCLE_STAT(TEXT("FPBDRigidsEvolutionGBF::PrepareConstraints"), STAT_Evolution_PrepareConstraints, STATGROUP_Chaos);
-DECLARE_CYCLE_STAT(TEXT("FPBDRigidsEvolutionGBF::CollisionModifierCallback"), STAT_Evolution_CollisionModifierCallback, STATGROUP_Chaos);
-DECLARE_CYCLE_STAT(TEXT("FPBDRigidsEvolutionGBF::UnprepareConstraints"), STAT_Evolution_UnprepareConstraints, STATGROUP_Chaos);
-DECLARE_CYCLE_STAT(TEXT("FPBDRigidsEvolutionGBF::ApplyConstraints"), STAT_Evolution_ApplyConstraints, STATGROUP_Chaos);
-DECLARE_CYCLE_STAT(TEXT("FPBDRigidsEvolutionGBF::UpdateVelocities"), STAT_Evolution_UpdateVelocites, STATGROUP_Chaos);
-DECLARE_CYCLE_STAT(TEXT("FPBDRigidsEvolutionGBF::ApplyPushOut"), STAT_Evolution_ApplyPushOut, STATGROUP_Chaos);
-DECLARE_CYCLE_STAT(TEXT("FPBDRigidsEvolutionGBF::DetectCollisions"), STAT_Evolution_DetectCollisions, STATGROUP_Chaos);
-DECLARE_CYCLE_STAT(TEXT("FPBDRigidsEvolutionGBF::TransferJointCollisions"), STAT_Evolution_TransferJointCollisions, STATGROUP_Chaos);
-DECLARE_CYCLE_STAT(TEXT("FPBDRigidsEvolutionGBF::PostDetectCollisionsCallback"), STAT_Evolution_PostDetectCollisionsCallback, STATGROUP_Chaos);
-DECLARE_CYCLE_STAT(TEXT("FPBDRigidsEvolutionGBF::UpdateConstraintPositionBasedState"), STAT_Evolution_UpdateConstraintPositionBasedState, STATGROUP_Chaos);
-DECLARE_CYCLE_STAT(TEXT("FPBDRigidsEvolutionGBF::ComputeIntermediateSpatialAcceleration"), STAT_Evolution_ComputeIntermediateSpatialAcceleration, STATGROUP_Chaos);
-DECLARE_CYCLE_STAT(TEXT("FPBDRigidsEvolutionGBF::CreateConstraintGraph"), STAT_Evolution_CreateConstraintGraph, STATGROUP_Chaos);
-DECLARE_CYCLE_STAT(TEXT("FPBDRigidsEvolutionGBF::CreateIslands"), STAT_Evolution_CreateIslands, STATGROUP_Chaos);
-DECLARE_CYCLE_STAT(TEXT("FPBDRigidsEvolutionGBF::PreApplyCallback"), STAT_Evolution_PreApplyCallback, STATGROUP_Chaos);
-DECLARE_CYCLE_STAT(TEXT("FPBDRigidsEvolutionGBF::ParallelSolve"), STAT_Evolution_ParallelSolve, STATGROUP_Chaos);
-DECLARE_CYCLE_STAT(TEXT("FPBDRigidsEvolutionGBF::DeactivateSleep"), STAT_Evolution_DeactivateSleep, STATGROUP_Chaos);
+	namespace CVars
+	{
+		FRealSingle HackMaxAngularVelocity = 1000.f;
+		FAutoConsoleVariableRef CVarHackMaxAngularVelocity(TEXT("p.HackMaxAngularVelocity"), HackMaxAngularVelocity, TEXT("Max cap on angular velocity: rad/s. This is only a temp solution and should not be relied on as a feature. -1.f to disable"));
+
+		FRealSingle HackMaxVelocity = -1.f;
+		FAutoConsoleVariableRef CVarHackMaxVelocity(TEXT("p.HackMaxVelocity2"), HackMaxVelocity, TEXT("Max cap on velocity: cm/s. This is only a temp solution and should not be relied on as a feature. -1.f to disable"));
+
+
+		int DisableThreshold = 5;
+		FAutoConsoleVariableRef CVarDisableThreshold(TEXT("p.DisableThreshold2"), DisableThreshold, TEXT("Disable threshold frames to transition to sleeping"));
+
+		int CollisionDisableCulledContacts = 0;
+		FAutoConsoleVariableRef CVarDisableCulledContacts(TEXT("p.CollisionDisableCulledContacts"), CollisionDisableCulledContacts, TEXT("Allow the PBDRigidsEvolutionGBF collision constraints to throw out contacts mid solve if they are culled."));
+
+		// @todo(chaos): this should be 0 but we need it for CCD atm
+		FRealSingle BoundsThicknessVelocityMultiplier = 1.0f;	// @todo(chaos): more to FChaosSolverConfiguration. The required value depends on the solver type
+		FAutoConsoleVariableRef CVarBoundsThicknessVelocityMultiplier(TEXT("p.CollisionBoundsVelocityInflation"), BoundsThicknessVelocityMultiplier, TEXT("Collision velocity inflation for speculatibe contact generation.[def:2.0]"));
+
+		FRealSingle SmoothedPositionLerpRate = 0.1f;
+		FAutoConsoleVariableRef CVarSmoothedPositionLerpRate(TEXT("p.Chaos.SmoothedPositionLerpRate"), SmoothedPositionLerpRate, TEXT("The interpolation rate for the smoothed position calculation. Used for sleeping."));
+
+		int DisableParticleUpdateVelocityParallelFor = 0;
+		FAutoConsoleVariableRef CVarDisableParticleUpdateVelocityParallelFor(TEXT("p.DisableParticleUpdateVelocityParallelFor"), DisableParticleUpdateVelocityParallelFor, TEXT("Disable Particle Update Velocity ParallelFor and run the update on a single thread"));
+
+		int32 ChaosSolverCollisionPriority = 0;
+		FAutoConsoleVariableRef CVarChaosSolverCollisionPriority(TEXT("p.Chaos.Solver.Collision.Priority"), ChaosSolverCollisionPriority, TEXT("Set constraint priority. Larger values are evaluated later [def:0]"));
+
+		int32 ChaosSolverJointPriority = 0;
+		FAutoConsoleVariableRef CVarChaosSolverJointPriority(TEXT("p.Chaos.Solver.Joint.Priority"), ChaosSolverJointPriority, TEXT("Set constraint priority. Larger values are evaluated later [def:0]"));
+
+		int32 ChaosSolverSuspensionPriority = 0;
+		FAutoConsoleVariableRef CVarChaosSolverSuspensionPriority(TEXT("p.Chaos.Solver.Suspension.Priority"), ChaosSolverSuspensionPriority, TEXT("Set constraint priority. Larger values are evaluated later [def:0]"));
+
+		bool DoTransferJointConstraintCollisions = true;
+		FAutoConsoleVariableRef CVarDoTransferJointConstraintCollisions(TEXT("p.Chaos.Solver.Joint.TransferCollisions"), DoTransferJointConstraintCollisions, TEXT("Allows joints to apply collisions to the parent from the child when the Joints TransferCollisionScale is not 0 [def:true]"));
+
+		DECLARE_CYCLE_STAT(TEXT("FPBDRigidsEvolutionGBF::AdvanceOneTimeStep"), STAT_Evolution_AdvanceOneTimeStep, STATGROUP_Chaos);
+		DECLARE_CYCLE_STAT(TEXT("FPBDRigidsEvolutionGBF::UnclusterUnions"), STAT_Evolution_UnclusterUnions, STATGROUP_Chaos);
+		DECLARE_CYCLE_STAT(TEXT("FPBDRigidsEvolutionGBF::Integrate"), STAT_Evolution_Integrate, STATGROUP_Chaos);
+		DECLARE_CYCLE_STAT(TEXT("FPBDRigidsEvolutionGBF::KinematicTargets"), STAT_Evolution_KinematicTargets, STATGROUP_Chaos);
+		DECLARE_CYCLE_STAT(TEXT("FPBDRigidsEvolutionGBF::PostIntegrateCallback"), STAT_Evolution_PostIntegrateCallback, STATGROUP_Chaos);
+		DECLARE_CYCLE_STAT(TEXT("FPBDRigidsEvolutionGBF::CollisionModifierCallback"), STAT_Evolution_CollisionModifierCallback, STATGROUP_Chaos);
+		DECLARE_CYCLE_STAT(TEXT("FPBDRigidsEvolutionGBF::GraphColor"), STAT_Evolution_GraphColor, STATGROUP_Chaos);
+		DECLARE_CYCLE_STAT(TEXT("FPBDRigidsEvolutionGBF::Gather"), STAT_Evolution_Gather, STATGROUP_Chaos);
+		DECLARE_CYCLE_STAT(TEXT("FPBDRigidsEvolutionGBF::Scatter"), STAT_Evolution_Scatter, STATGROUP_Chaos);
+		DECLARE_CYCLE_STAT(TEXT("FPBDRigidsEvolutionGBF::ApplyConstraintsPhase0"), STAT_Evolution_ApplyConstraintsPhase0, STATGROUP_Chaos);
+		DECLARE_CYCLE_STAT(TEXT("FPBDRigidsEvolutionGBF::ApplyConstraintsPhase1"), STAT_Evolution_ApplyConstraintsPhase1, STATGROUP_Chaos);
+		DECLARE_CYCLE_STAT(TEXT("FPBDRigidsEvolutionGBF::UpdateVelocities"), STAT_Evolution_UpdateVelocites, STATGROUP_Chaos);
+		DECLARE_CYCLE_STAT(TEXT("FPBDRigidsEvolutionGBF::ApplyConstraintsPhase2"), STAT_Evolution_ApplyConstraintsPhase2, STATGROUP_Chaos);
+		DECLARE_CYCLE_STAT(TEXT("FPBDRigidsEvolutionGBF::DetectCollisions"), STAT_Evolution_DetectCollisions, STATGROUP_Chaos);
+		DECLARE_CYCLE_STAT(TEXT("FPBDRigidsEvolutionGBF::TransferJointCollisions"), STAT_Evolution_TransferJointCollisions, STATGROUP_Chaos);
+		DECLARE_CYCLE_STAT(TEXT("FPBDRigidsEvolutionGBF::PostDetectCollisionsCallback"), STAT_Evolution_PostDetectCollisionsCallback, STATGROUP_Chaos);
+		DECLARE_CYCLE_STAT(TEXT("FPBDRigidsEvolutionGBF::UpdateConstraintPositionBasedState"), STAT_Evolution_UpdateConstraintPositionBasedState, STATGROUP_Chaos);
+		DECLARE_CYCLE_STAT(TEXT("FPBDRigidsEvolutionGBF::ComputeIntermediateSpatialAcceleration"), STAT_Evolution_ComputeIntermediateSpatialAcceleration, STATGROUP_Chaos);
+		DECLARE_CYCLE_STAT(TEXT("FPBDRigidsEvolutionGBF::CreateConstraintGraph"), STAT_Evolution_CreateConstraintGraph, STATGROUP_Chaos);
+		DECLARE_CYCLE_STAT(TEXT("FPBDRigidsEvolutionGBF::CreateIslands"), STAT_Evolution_CreateIslands, STATGROUP_Chaos);
+		DECLARE_CYCLE_STAT(TEXT("FPBDRigidsEvolutionGBF::PreApplyCallback"), STAT_Evolution_PreApplyCallback, STATGROUP_Chaos);
+		DECLARE_CYCLE_STAT(TEXT("FPBDRigidsEvolutionGBF::ParallelSolve"), STAT_Evolution_ParallelSolve, STATGROUP_Chaos);
+		DECLARE_CYCLE_STAT(TEXT("FPBDRigidsEvolutionGBF::SaveParticlePostSolve"), STAT_Evolution_SavePostSolve, STATGROUP_Chaos);
+		DECLARE_CYCLE_STAT(TEXT("FPBDRigidsEvolutionGBF::DeactivateSleep"), STAT_Evolution_DeactivateSleep, STATGROUP_Chaos);
+
+		int32 SerializeEvolution = 0;
+		FAutoConsoleVariableRef CVarSerializeEvolution(TEXT("p.SerializeEvolution"), SerializeEvolution, TEXT(""));
+
+		bool bChaos_CollisionStore_Enabled = true;
+		FAutoConsoleVariableRef CVarCollisionStoreEnabled(TEXT("p.Chaos.CollisionStore.Enabled"), bChaos_CollisionStore_Enabled, TEXT(""));
+	}
 
-int32 SerializeEvolution = 0;
-FAutoConsoleVariableRef CVarSerializeEvolution(TEXT("p.SerializeEvolution"), SerializeEvolution, TEXT(""));
+	using namespace CVars;
 
 #if !UE_BUILD_SHIPPING
 template <typename TEvolution>
@@ -149,7 +156,7 @@ void FPBDRigidsEvolutionGBF::Advance(const FReal Dt,const FReal MaxStepDt,const
 		{
 			// StepFraction: how much of the remaining time this step represents, used to interpolate kinematic targets
 			// E.g., for 4 steps this will be: 1/4, 1/2, 3/4, 1
-			const FReal StepFraction = (FReal)(Step + 1) / (FReal)(NumSteps);
+			const FReal StepFraction = (FReal)(Step + 1) / (NumSteps);
 		
 			UE_LOG(LogChaos, Verbose, TEXT("Advance dt = %f [%d/%d]"), StepDt, Step + 1, NumSteps);
 
@@ -173,7 +180,7 @@ void FPBDRigidsEvolutionGBF::AdvanceOneTimeStep(const FReal Dt,const FSubStepInf
 int32 DrawAwake = 0;
 FAutoConsoleVariableRef CVarDrawAwake(TEXT("p.chaos.DebugDrawAwake"),DrawAwake,TEXT("Draw particles that are awake"));
 
-void FPBDRigidsEvolutionGBF::AdvanceOneTimeStepImpl(const FReal Dt,const FSubStepInfo& SubStepInfo)
+void FPBDRigidsEvolutionGBF::AdvanceOneTimeStepImpl(const FReal Dt, const FSubStepInfo& SubStepInfo)
 {
 	SCOPE_CYCLE_COUNTER(STAT_Evolution_AdvanceOneTimeStep);
 
@@ -228,17 +235,7 @@ void FPBDRigidsEvolutionGBF::AdvanceOneTimeStepImpl(const FReal Dt,const FSubSte
 		SCOPE_CYCLE_COUNTER(STAT_Evolution_DetectCollisions);
 		CollisionDetector.GetBroadPhase().SetSpatialAcceleration(InternalAcceleration);
 
-		CollisionStats::FStatData StatData(bPendingHierarchyDump);
-
-		CollisionDetector.DetectCollisionsWithStats(Dt, StatData, GetCurrentStepResimCache());
-
-		CHAOS_COLLISION_STAT(StatData.Print());
-	}
-
-	{
-		SCOPE_CYCLE_COUNTER(STAT_Evolution_TransferJointCollisions);
-
-		TransferJointConstraintCollisions();
+		CollisionDetector.DetectCollisions(Dt, GetCurrentStepResimCache());
 	}
 
 	if (PostDetectCollisionsCallback != nullptr)
@@ -248,8 +245,8 @@ void FPBDRigidsEvolutionGBF::AdvanceOneTimeStepImpl(const FReal Dt,const FSubSte
 	}
 
 	{
-		SCOPE_CYCLE_COUNTER(STAT_Evolution_PrepareConstraints);
-		PrepareIteration(Dt);
+		SCOPE_CYCLE_COUNTER(STAT_Evolution_TransferJointCollisions);
+		TransferJointConstraintCollisions();
 	}
 
 	if(CollisionModifiers)
@@ -291,14 +288,14 @@ void FPBDRigidsEvolutionGBF::AdvanceOneTimeStepImpl(const FReal Dt,const FSubSte
 			{
 				bHasCachedData = true;
 			}
-
+			
 			const TArray<FGeometryParticleHandle*>& IslandParticles = GetConstraintGraph().GetIslandParticles(Island);
 
 			if (bHasCachedData)
 			{
 				for (auto Particle : IslandParticles)
 				{
-					if(auto Rigid = Particle->CastToRigidParticle())
+					if (auto Rigid = Particle->CastToRigidParticle())
 					{
 						ResimCache->ReloadParticlePostSolve(*Rigid);
 					}
@@ -306,9 +303,30 @@ void FPBDRigidsEvolutionGBF::AdvanceOneTimeStepImpl(const FReal Dt,const FSubSte
 			}
 			else
 			{
+				// Update constraint graphs, coloring etc as required by the different constraint types in this island
 				{
-					SCOPE_CYCLE_COUNTER(STAT_Evolution_ApplyConstraints);
-					ApplyConstraints(Dt, Island);
+					SCOPE_CYCLE_COUNTER(STAT_Evolution_GraphColor);
+					UpdateAccelerationStructures(Dt, Island);
+				}
+
+				// Collect all the data that the constraint solvers operate on
+				{
+					SCOPE_CYCLE_COUNTER(STAT_Evolution_Gather);
+					GatherSolverInput(Dt, Island);
+				}
+
+				{
+					SCOPE_CYCLE_COUNTER(STAT_Evolution_ApplyConstraintsPhase0);
+					ApplyConstraintsPhase0(Dt, Island);
+				}
+
+				// Run the first phase of the constraint solvers
+				// For GBF this is the hybrid velocity solving step (which also moves the bodies to make the implicit velocity be what it should be)
+				// For PBD/QPBD this is the position solve step
+				{
+					SCOPE_CYCLE_COUNTER(STAT_Evolution_ApplyConstraintsPhase1);
+					CSV_SCOPED_TIMING_STAT(Chaos, ApplyConstraints);
+					ApplyConstraintsPhase1(Dt, Island);
 				}
 
 				if (PostApplyCallback != nullptr)
@@ -316,14 +334,26 @@ void FPBDRigidsEvolutionGBF::AdvanceOneTimeStepImpl(const FReal Dt,const FSubSte
 					PostApplyCallback(Island);
 				}
 
+				// Update implicit velocities from results of constraint solver phase 1
 				{
 					SCOPE_CYCLE_COUNTER(STAT_Evolution_UpdateVelocites);
-					UpdateVelocities(Dt, Island);
+					SetImplicitVelocities(Dt, Island);
 				}
 
+				// Run the second phase of the constraint solvers
+				// For GBF this is the pushout step
+				// For PBD this does nothing
+				// For QPBD this is the velocity solve step
 				{
-					SCOPE_CYCLE_COUNTER(STAT_Evolution_ApplyPushOut);
-					ApplyPushOut(Dt, Island);
+					SCOPE_CYCLE_COUNTER(STAT_Evolution_ApplyConstraintsPhase2);
+					ApplyConstraintsPhase2(Dt, Island);
+				}
+
+				// Update the particles with the results of the constraint solvers, and also update constraint data
+				// that is accessed externally (net impusles, break info, etc)
+				{
+					SCOPE_CYCLE_COUNTER(STAT_Evolution_Scatter);
+					ScatterSolverOutput(Dt, Island);
 				}
 
 				if (PostApplyPushOutCallback != nullptr)
@@ -339,7 +369,7 @@ void FPBDRigidsEvolutionGBF::AdvanceOneTimeStepImpl(const FReal Dt,const FSubSte
 
 				// increment the disable count for the particle
 				auto PBDRigid = Particle->CastToRigidParticle();
-				if(PBDRigid && PBDRigid->ObjectState() == EObjectStateType::Dynamic)
+				if (PBDRigid && PBDRigid->ObjectState() == EObjectStateType::Dynamic)
 				{
 					if (PBDRigid->AuxilaryValue(PhysicsMaterials) && PBDRigid->V().SizeSquared() < PBDRigid->AuxilaryValue(PhysicsMaterials)->DisabledLinearThreshold &&
 						PBDRigid->W().SizeSquared() < PBDRigid->AuxilaryValue(PhysicsMaterials)->DisabledAngularThreshold)
@@ -371,33 +401,35 @@ void FPBDRigidsEvolutionGBF::AdvanceOneTimeStepImpl(const FReal Dt,const FSubSte
 	}
 
 	{
-		SCOPE_CYCLE_COUNTER(STAT_Evolution_UnprepareConstraints);
-		UnprepareIteration(Dt);
-	}
-
-
-	{
-		SCOPE_CYCLE_COUNTER(STAT_Evolution_DeactivateSleep);
+		SCOPE_CYCLE_COUNTER(STAT_Evolution_SavePostSolve);
 		FEvolutionResimCache* ResimCache = GetCurrentStepResimCache();
-		for (int32 Island = 0; Island < GetConstraintGraph().NumIslands(); ++Island)
+		if (ResimCache)
 		{
-			if(ResimCache)
+			for (const auto& Particle : Particles.GetActiveKinematicParticlesView())
 			{
-				const TArray<FGeometryParticleHandle*>& IslandParticles = GetConstraintGraph().GetIslandParticles(Island);
-				for(const auto Particle : IslandParticles)
+				if (const auto* Rigid = Particle.CastToRigidParticle())
 				{
-					if (auto PBDRigid = Particle->CastToRigidParticle())
-					{
-						//NOTE: this assumes the cached values have not changed after the solve (V, W, P, Q should be untouched, otherwise we'll use the wrong values when resim happens)
-						ResimCache->SaveParticlePostSolve(*PBDRigid);
-					}
+					//NOTE: this assumes the cached values have not changed after the solve (V, W, P, Q should be untouched, otherwise we'll use the wrong values when resim happens)
+					ResimCache->SaveParticlePostSolve(*Rigid->Handle());
 				}
 			}
+			for (const auto& Particle : Particles.GetNonDisabledDynamicView())
+			{
+				//NOTE: this assumes the cached values have not changed after the solve (V, W, P, Q should be untouched, otherwise we'll use the wrong values when resim happens)
+				ResimCache->SaveParticlePostSolve(*Particle.Handle());
+			}
+		}
+	}
 
+	{
+		SCOPE_CYCLE_COUNTER(STAT_Evolution_DeactivateSleep);
+		for (int32 Island = 0; Island < GetConstraintGraph().NumIslands(); ++Island)
+		{
 			if (SleepedIslands[Island])
 			{
-				Particles.DeactivateParticles(GetConstraintGraph().GetIslandParticles(Island));
+				GetConstraintGraph().SleepIsland(Particles, Island);
 			}
+			
 			for (const auto Particle : DisabledParticles[Island])
 			{
 				DisableParticle(Particle);
@@ -451,6 +483,83 @@ void FPBDRigidsEvolutionGBF::AdvanceOneTimeStepImpl(const FReal Dt,const FSubSte
 #endif
 }
 
+void FPBDRigidsEvolutionGBF::GatherSolverInput(FReal Dt, int32 Island)
+{
+	// We must initialize the solver body container to be large enough to hold all particles in the
+	// island so that the pointers remain valid (the array should not grow and relocate)
+	ConstraintGraph.GetSolverIsland(Island)->GetBodyContainer().Reset(ConstraintGraph.GetIslandParticles(Island).Num());
+
+	// NOTE: SolverBodies are gathered as part of the constraint gather, in the order that they are first seen
+	for (FPBDConstraintGraphRule* ConstraintRule : PrioritizedConstraintRules)
+	{
+		ConstraintRule->GatherSolverInput(Dt, Island);
+	}
+}
+
+void FPBDRigidsEvolutionGBF::ScatterSolverOutput(FReal Dt, int32 Island)
+{
+	// Scatter solver results for constraints (impulses, break events, etc)
+	for (FPBDConstraintGraphRule* ConstraintRule : PrioritizedConstraintRules)
+	{
+		ConstraintRule->ScatterSolverOutput(Dt, Island);
+	}
+
+	// Scatter body results back to particles (position, rotation, etc)
+	ConstraintGraph.GetSolverIsland(Island)->GetBodyContainer().ScatterOutput();
+}
+
+void FPBDRigidsEvolutionGBF::ApplyConstraintsPhase0(const FReal Dt, int32 Island)
+{
+	for (FPBDConstraintGraphRule* ConstraintRule : PrioritizedConstraintRules)
+	{
+		ConstraintRule->ApplySwept(Dt, Island);
+	}
+}
+
+void FPBDRigidsEvolutionGBF::ApplyConstraintsPhase1(const FReal Dt, int32 Island)
+{
+	int32 LocalNumIterations = ChaosNumContactIterationsOverride >= 0 ? ChaosNumContactIterationsOverride : NumIterations;
+	// @todo(ccaulfield): track whether we are sufficiently solved and can early-out
+	for (int i = 0; i < LocalNumIterations; ++i)
+	{
+		bool bNeedsAnotherIteration = false;
+		for (FPBDConstraintGraphRule* ConstraintRule : PrioritizedConstraintRules)
+		{
+			bNeedsAnotherIteration |= ConstraintRule->ApplyConstraints(Dt, Island, i, LocalNumIterations);
+		}
+
+		if (ChaosRigidsEvolutionApplyAllowEarlyOutCVar && !bNeedsAnotherIteration)
+		{
+			break;
+		}
+	}
+}
+
+void FPBDRigidsEvolutionGBF::SetImplicitVelocities(const FReal Dt, int32 Island)
+{
+	ConstraintGraph.GetSolverIsland(Island)->GetBodyContainer().SetImplicitVelocities(Dt);
+}
+
+void FPBDRigidsEvolutionGBF::ApplyConstraintsPhase2(const FReal Dt, int32 Island)
+{
+	int32 LocalNumPushOutIterations = ChaosNumPushOutIterationsOverride >= 0 ? ChaosNumPushOutIterationsOverride : NumPushOutIterations;
+	bool bNeedsAnotherIteration = true;
+	for (int32 It = 0; It < LocalNumPushOutIterations; ++It)
+	{
+		bNeedsAnotherIteration = false;
+		for (FPBDConstraintGraphRule* ConstraintRule : PrioritizedConstraintRules)
+		{
+			bNeedsAnotherIteration |= ConstraintRule->ApplyPushOut(Dt, Island, It, LocalNumPushOutIterations);
+		}
+
+		if (ChaosRigidsEvolutionApplyPushoutAllowEarlyOutCVar && !bNeedsAnotherIteration)
+		{
+			break;
+		}
+	}
+}
+
+
 FPBDRigidsEvolutionGBF::FPBDRigidsEvolutionGBF(FPBDRigidsSOAs& InParticles,THandleArray<FChaosPhysicsMaterial>& SolverPhysicsMaterials, const TArray<ISimCallbackObject*>* InCollisionModifiers, bool InIsSingleThreaded)
 	: Base(InParticles, SolverPhysicsMaterials, DefaultNumIterations, DefaultNumPushOutIterations, InIsSingleThreaded)
 	, Clustering(*this, Particles.GetClusteredParticles())
@@ -470,15 +579,6 @@ FPBDRigidsEvolutionGBF::FPBDRigidsEvolutionGBF(FPBDRigidsSOAs& InParticles,THand
 {
 	CollisionConstraints.SetCanDisableContacts(!!CollisionDisableCulledContacts);
 
-	SetParticleUpdateVelocityFunction([PBDUpdateRule = FPerParticlePBDUpdateFromDeltaPosition(), this](const TArray<FGeometryParticleHandle*>& ParticlesInput, const FReal Dt) {
-		ParticlesParallelFor(ParticlesInput, [&](auto& Particle, int32 Index) {
-			if (Particle->CastToRigidParticle() && Particle->ObjectState() == EObjectStateType::Dynamic)
-			{
-				PBDUpdateRule.Apply(Particle->CastToRigidParticle(), Dt);
-			}
-		}, (DisableParticleUpdateVelocityParallelFor != 0));
-	});
-
 	SetParticleUpdatePositionFunction([this](const TParticleView<FPBDRigidParticles>& ParticlesInput, const FReal Dt)
 	{
 		ParticlesInput.ParallelFor([&](auto& Particle, int32 Index)
@@ -509,12 +609,19 @@ FPBDRigidsEvolutionGBF::FPBDRigidsEvolutionGBF(FPBDRigidsSOAs& InParticles,THand
 	AddConstraintRule(&CollisionRule);
 	AddConstraintRule(&JointConstraintRule);
 
-	SetInternalParticleInitilizationFunction([](const FGeometryParticleHandle*, FGeometryParticleHandle*) {});
+	SetInternalParticleInitilizationFunction([](const FGeometryParticleHandle*, const FGeometryParticleHandle*) {});
+
 	NarrowPhase.GetContext().bFilteringEnabled = true;
 	NarrowPhase.GetContext().bDeferUpdate = true;
 	NarrowPhase.GetContext().bAllowManifolds = false;
+	NarrowPhase.GetContext().CollisionAllocator = &CollisionConstraints.GetConstraintAllocator();
 }
 
+FPBDRigidsEvolutionGBF::~FPBDRigidsEvolutionGBF()
+{
+}
+
+
 void FPBDRigidsEvolutionGBF::Serialize(FChaosArchive& Ar)
 {
 	Base::Serialize(Ar);
@@ -530,20 +637,24 @@ void FPBDRigidsEvolutionGBF::SetCurrentStepResimCache(IResimCacheBase* InCurrent
 	CurrentStepResimCacheImp = static_cast<FEvolutionResimCache*>(InCurrentStepResimCache);
 }
 
-
 void FPBDRigidsEvolutionGBF::TransferJointConstraintCollisions()
 {
-	//
+#	//
 	// Append transfer constraints. 
 	//
+	// @todo(chaos): this should be implementable with a collision modifier system
 	using FRigidHandle = TPBDRigidParticleHandleImp<FReal, 3, true>;
 	if (DoTransferJointConstraintCollisions)
 	{
+		FCollisionConstraintAllocator& CollisionAllocator = CollisionConstraints.GetConstraintAllocator();
+
 		TMap<TGeometryParticleHandle<FReal, 3>*, TArray<FPBDCollisionConstraintHandle*> >& CollisionMap = CollisionConstraints.GetParticleCollisionsMap();
 		for (auto& CollisionPair : CollisionMap)
 		{
 			if (FRigidHandle* ChildParticle = CollisionPair.Key->CastToRigidParticle())
 			{
+				const FRigidTransform3 ChildTransform = FParticleUtilities::GetActorWorldTransform(FConstGenericParticleHandle(ChildParticle));
+
 				for (FConstraintHandle* Constraint : ChildParticle->ParticleConstraints())
 				{
 					if (FPBDJointConstraintHandle* JointConstraint = Constraint->As<FPBDJointConstraintHandle>())
@@ -552,63 +663,65 @@ void FPBDRigidsEvolutionGBF::TransferJointConstraintCollisions()
 						{
 							if (FRigidHandle* ParentParticle = JointConstraint->GetConstrainedParticles()[1]->CastToRigidParticle()) // Parent
 							{
-								for (auto& CollisionHandel : CollisionPair.Value)
+								const FRigidTransform3 ParentTransform = FParticleUtilities::GetActorWorldTransform(FConstGenericParticleHandle(ParentParticle));
+								const FTransform ChildToParentTransform = ChildTransform.GetRelativeTransform(ParentTransform);
+
+								for (FPBDCollisionConstraintHandle* ContactHandle : CollisionPair.Value)
 								{
-									if (FPBDCollisionConstraintHandle* ContactHandle = CollisionHandel->As<FPBDCollisionConstraintHandle>())
+									if (ContactHandle->GetType() == ECollisionConstraintType::Standard)
 									{
-										if (ContactHandle->GetType() == FCollisionConstraintBase::FType::SinglePoint)
-										{
-											FRigidBodyPointContactConstraint CurrConstraint = ContactHandle->GetPointContact();
-											FTransform ParentTransform(ParentParticle->R(), ParentParticle->X()), ChildTransform(ChildParticle->R(), ChildParticle->X()),
-												ChildToParentTransform = ChildTransform.GetRelativeTransform(ParentTransform);
-
-											FRigidBodyPointContactConstraint TransferedConstraint;
-											if (ContactHandle->GetConstrainedParticles()[0] == ChildParticle)
-											{
-												TransferedConstraint = FRigidBodyPointContactConstraint(
-													ParentParticle,
-													CurrConstraint.GetManifold().Implicit[0],
-													CurrConstraint.GetManifold().Simplicial[0],
-													ChildToParentTransform.GetRelativeTransform(CurrConstraint.ImplicitTransform[0]),
-													CurrConstraint.Particle[1],
-													CurrConstraint.GetManifold().Implicit[1],
-													CurrConstraint.GetManifold().Simplicial[1],
-													CurrConstraint.ImplicitTransform[1],
-													CurrConstraint.GetCullDistance(),
-													CurrConstraint.GetManifold().ShapesType,
-													CurrConstraint.GetUseManifold()
-												);
-											}
-											else
-											{
-												TransferedConstraint = FRigidBodyPointContactConstraint(
-													CurrConstraint.Particle[0],
-													CurrConstraint.GetManifold().Implicit[0],
-													CurrConstraint.GetManifold().Simplicial[0],
-													CurrConstraint.ImplicitTransform[0],
-													ParentParticle,
-													CurrConstraint.GetManifold().Implicit[1],
-													CurrConstraint.GetManifold().Simplicial[1],
-													ChildToParentTransform.GetRelativeTransform(CurrConstraint.ImplicitTransform[1]),
-													CurrConstraint.GetCullDistance(),
-													CurrConstraint.GetManifold().ShapesType,
-													CurrConstraint.GetUseManifold()
-												);
-											}
-											TArray<FManifoldPoint> TransferManifolds;
-											for (FManifoldPoint& Manifold : CurrConstraint.GetManifoldPoints())
-											{
-												TransferManifolds.Add(Manifold);
-											}
-											TransferedConstraint.SetManifoldPoints(TransferManifolds);
-											for (FManifoldPoint& Manifold : CurrConstraint.GetManifoldPoints())
-											{
-												TransferedConstraint.UpdateManifoldPointFromContact(Manifold);
-											}
-											TransferedConstraint.SetStiffness(JointConstraint->GetSettings().ContactTransferScale);
-											CollisionConstraints.AddConstraint(TransferedConstraint);
+										const FPBDCollisionConstraint& CurrConstraint = ContactHandle->GetContact();
 
+										FPBDCollisionConstraint* TransferedConstraint;
+										if (ContactHandle->GetConstrainedParticles()[0] == ChildParticle)
+										{
+											TransferedConstraint = FPBDCollisionConstraint::Make(
+												ParentParticle,
+												CurrConstraint.GetManifold().Implicit[0],
+												CurrConstraint.GetManifold().Simplicial[0],
+												ChildTransform,
+												ChildToParentTransform.GetRelativeTransform(CurrConstraint.ImplicitTransform[0]),
+												CurrConstraint.Particle[1],
+												CurrConstraint.GetManifold().Implicit[1],
+												CurrConstraint.GetManifold().Simplicial[1],
+												ParentTransform,
+												CurrConstraint.ImplicitTransform[1],
+												CurrConstraint.GetCullDistance(),
+												CurrConstraint.GetManifold().ShapesType,
+												CurrConstraint.GetUseManifold(),
+												CollisionAllocator
+											);
 										}
+										else
+										{
+											TransferedConstraint = FPBDCollisionConstraint::Make(
+												CurrConstraint.Particle[0],
+												CurrConstraint.GetManifold().Implicit[0],
+												CurrConstraint.GetManifold().Simplicial[0],
+												ParentTransform,
+												CurrConstraint.ImplicitTransform[0],
+												ParentParticle,
+												CurrConstraint.GetManifold().Implicit[1],
+												CurrConstraint.GetManifold().Simplicial[1],
+												ChildTransform,
+												ChildToParentTransform.GetRelativeTransform(CurrConstraint.ImplicitTransform[1]),
+												CurrConstraint.GetCullDistance(),
+												CurrConstraint.GetManifold().ShapesType,
+												CurrConstraint.GetUseManifold(),
+												CollisionAllocator
+											);
+										}
+										TArray<FManifoldPoint> TransferManifolds;
+										for (const FManifoldPoint& Manifold : CurrConstraint.GetManifoldPoints())
+										{
+											TransferManifolds.Add(Manifold);
+										}
+										TransferedConstraint->SetManifoldPoints(TransferManifolds);
+										for (FManifoldPoint& Manifold : TransferedConstraint->GetManifoldPoints())
+										{
+											TransferedConstraint->UpdateManifoldPointFromContact(Manifold);
+										}
+										TransferedConstraint->SetStiffness(JointConstraint->GetSettings().ContactTransferScale);
 									}
 								}
 							}
@@ -617,9 +730,10 @@ void FPBDRigidsEvolutionGBF::TransferJointConstraintCollisions()
 				}
 			}
 		}
+
+		CollisionAllocator.ProcessInjectedConstraints();
 	}
 }
 
-
 }
 
diff --git a/Engine/Source/Runtime/Experimental/Chaos/Private/Chaos/PBDSuspensionConstraints.cpp b/Engine/Source/Runtime/Experimental/Chaos/Private/Chaos/PBDSuspensionConstraints.cpp
--- a/Engine/Source/Runtime/Experimental/Chaos/Private/Chaos/PBDSuspensionConstraints.cpp
+++ b/Engine/Source/Runtime/Experimental/Chaos/Private/Chaos/PBDSuspensionConstraints.cpp
@@ -1,11 +1,12 @@
 // Copyright Epic Games, Inc. All Rights Reserved.
 
 #include "Chaos/PBDSuspensionConstraints.h"
+#include "Chaos/Evolution/SolverDatas.h"
 
 namespace Chaos
 {
 	FPBDSuspensionConstraintHandle::FPBDSuspensionConstraintHandle(FConstraintContainer* InConstraintContainer, int32 InConstraintIndex) : 
-		TContainerConstraintHandle<FPBDSuspensionConstraints>(InConstraintContainer, InConstraintIndex)
+		TIndexedContainerConstraintHandle<FPBDSuspensionConstraints>(InConstraintContainer, InConstraintIndex)
 	{
 	}
 
@@ -24,25 +25,17 @@ namespace Chaos
 		ConcreteContainer()->SetSettings(ConstraintIndex, Settings);
 	}
 
-
 	TVec2<FGeometryParticleHandle*> FPBDSuspensionConstraintHandle::GetConstrainedParticles() const
 	{
 		return ConcreteContainer()->GetConstrainedParticles(ConstraintIndex);
 	}
 
-
-
-	bool FPBDSuspensionConstraints::Apply(const FReal Dt, const TArray<FConstraintContainerHandle*>& ConstraintHandles, const int32 It, const int32 NumIts) const
+	void FPBDSuspensionConstraintHandle::GatherInput(const FReal Dt, const int32 Particle0Level, const int32 Particle1Level, FPBDIslandSolverData& SolverData)
 	{
-		for (FConstraintContainerHandle* ConstraintHandle : ConstraintHandles)
-		{
-			ApplySingle(Dt, ConstraintHandle->GetConstraintIndex());
-		}
-
-		// #todo: Return true only if more iteration is needed
-		return true;
+		ConcreteContainer()->GatherInput(Dt, ConstraintIndex, Particle0Level, Particle1Level, SolverData);
 	}
 
+
 	Chaos::FPBDSuspensionConstraints::FConstraintContainerHandle* FPBDSuspensionConstraints::AddConstraint(TGeometryParticleHandle<FReal, 3>* Particle, const FVec3& InSuspensionLocalOffset, const FPBDSuspensionSettings& InConstraintSettings)
 	{
 		int32 NewIndex = ConstrainedParticles.Num();
@@ -50,6 +43,7 @@ namespace Chaos
 		SuspensionLocalOffset.Add(InSuspensionLocalOffset);
 		ConstraintSettings.Add(InConstraintSettings);
 		ConstraintEnabledStates.Add(true); // Note: assumes always enabled on creation
+		ConstraintSolverBodies.Add(nullptr);
 		Handles.Add(HandleAllocator.AllocHandle(this, NewIndex));
 		return Handles[NewIndex];
 	}
@@ -75,6 +69,7 @@ namespace Chaos
 		SuspensionLocalOffset.RemoveAtSwap(ConstraintIndex);
 		ConstraintSettings.RemoveAtSwap(ConstraintIndex);
 		ConstraintEnabledStates.RemoveAtSwap(ConstraintIndex);
+		ConstraintSolverBodies.RemoveAtSwap(ConstraintIndex);
 		Handles.RemoveAtSwap(ConstraintIndex);
 
 		// Update the handle for the constraint that was moved
@@ -83,18 +78,55 @@ namespace Chaos
 			SetConstraintIndex(Handles[ConstraintIndex], ConstraintIndex);
 		}
 	}
+	
+	void FPBDSuspensionConstraints::SetNumIslandConstraints(const int32 NumIslandConstraints, FPBDIslandSolverData& SolverData)
+	{
+		SolverData.GetConstraintIndices(ContainerId).Reset(NumIslandConstraints);
+	}
+
+	void FPBDSuspensionConstraints::GatherInput(const FReal Dt, const int32 ConstraintIndex, const int32 Particle0Level, const int32 Particle1Level, FPBDIslandSolverData& SolverData)
+	{
+		SolverData.GetConstraintIndices(ContainerId).Add(ConstraintIndex);
+
+		ConstraintSolverBodies[ConstraintIndex] = SolverData.GetBodyContainer().FindOrAdd(ConstrainedParticles[ConstraintIndex]);
+	}
+
+	void FPBDSuspensionConstraints::ScatterOutput(FReal Dt, FPBDIslandSolverData& SolverData)
+	{
+		for (int32 ConstraintIndex : SolverData.GetConstraintIndices(ContainerId))
+		{
+			ConstraintSolverBodies[ConstraintIndex] = nullptr;
+		}
+	}
+
+	bool FPBDSuspensionConstraints::ApplyPhase1Serial(const FReal Dt, const int32 It, const int32 NumIts, FPBDIslandSolverData& SolverData)
+	{
+		for (int32 ConstraintIndex : SolverData.GetConstraintIndices(ContainerId))
+		{
+			ApplySingle(Dt, ConstraintIndex);
+		}
+
+		// @todo(chaos): early iteration termination in FPBDSuspensionConstraints
+		return true;
+	}
 
 	void FPBDSuspensionConstraints::ApplySingle(const FReal Dt, int32 ConstraintIndex) const
 	{
-		FGenericParticleHandle Particle = ConstrainedParticles[ConstraintIndex];
+		check(ConstraintSolverBodies[ConstraintIndex] != nullptr);
+		FSolverBody& Body = *ConstraintSolverBodies[ConstraintIndex];
 		const FPBDSuspensionSettings& Setting = ConstraintSettings[ConstraintIndex];
 
-		if (Particle->IsDynamic() && Setting.Enabled)
+		if (Body.IsDynamic() && Setting.Enabled)
 		{
 			const FVec3& T = Setting.Target;
-			const FVec3 WorldSpaceX = Particle->Q().RotateVector(SuspensionLocalOffset[ConstraintIndex]) + Particle->P();
 
-			FVec3 AxisWorld = Particle->Q() * Setting.Axis;
+			// \todo(chaos): we can cache the CoM-relative connector once per frame rather than recalculate per iteration
+			// (we should not be accessing particle state in the solver methods, although this one actually is ok because it only uses frame constrants)
+			const FVec3 SuspensionCoMOffset = FParticleUtilities::ParticleLocalToCoMLocal(FGenericParticleHandle(ConstrainedParticles[ConstraintIndex]), SuspensionLocalOffset[ConstraintIndex]);
+
+			const FVec3 WorldSpaceX = Body.Q().RotateVector(SuspensionCoMOffset) + Body.P();
+
+			FVec3 AxisWorld = Body.Q() * Setting.Axis;
 
 			const float MPHToCmS = 100000.f / 2236.94185f;
 			const float SpeedThreshold = 10.0f * MPHToCmS;
@@ -102,13 +134,13 @@ namespace Chaos
 
 			if (AxisWorld.Z > FortyFiveDegreesThreshold)
 			{
-				if (Particle->V().SquaredLength() < 1.0f)
+				if (Body.V().SquaredLength() < 1.0f)
 				{
 					AxisWorld = FVec3(0.f, 0.f, 1.f);
 				}
 				else
 				{
-					const FReal Speed = FMath::Abs(Particle->V().Length());
+					const FReal Speed = FMath::Abs(Body.V().Length());
 					if (Speed < SpeedThreshold)
 					{
 						AxisWorld = FMath::Lerp(FVec3(0.f, 0.f, 1.f), AxisWorld, Speed / SpeedThreshold);
@@ -127,9 +159,8 @@ namespace Chaos
 			FVec3 DX = FVec3::ZeroVector;
 
 			// Require the velocity at the WorldSpaceX position - not the velocity of the particle origin
-			const FVec3 COM = Chaos::FParticleUtilitiesGT::GetCoMWorldPosition(Particle);
-			const FVec3 Diff = WorldSpaceX - COM;
-			FVec3 ArmVelocity = Particle->V() - FVec3::CrossProduct(Diff, Particle->W());
+			const FVec3 Diff = WorldSpaceX - Body.P();
+			FVec3 ArmVelocity = Body.V() - FVec3::CrossProduct(Diff, Body.W());
 
 			// This constraint is causing considerable harm to the steering effect from the tires, using only the z component for damping
 			// makes this issue go away, rather than using DotProduct against the expected AxisWorld vector
@@ -163,7 +194,7 @@ namespace Chaos
 					FReal VelDt = PointVelocityAlongAxis;
 
 					const bool bAccelerationMode = false;
-					const FReal SpringMassScale = (bAccelerationMode) ? 1.0f / (Particle->InvM()) : 1.0f;
+					const FReal SpringMassScale = (bAccelerationMode) ? FReal(1) / Body.InvM() : FReal(1);
 					const FReal S = SpringMassScale * Setting.SpringStiffness * Dt * Dt;
 					const FReal D = SpringMassScale * Setting.SpringDamping * Dt;
 					DLambda = (S * SpringCompression - D * VelDt);
@@ -171,24 +202,16 @@ namespace Chaos
 				}
 			}
 
-			const FVec3 Arm = WorldSpaceX - Particle->P();
-
-			FRotation3 Q0 = FParticleUtilities::GetCoMWorldRotation(Particle);
-			FVec3 P0 = FParticleUtilities::GetCoMWorldPosition(Particle);
-			const FMatrix33 WorldSpaceInvI = Utilities::ComputeWorldSpaceInertia(Q0, Particle->InvI());
+			const FVec3 Arm = WorldSpaceX - Body.P();
 
-			FVec3 DP = Particle->InvM() * DX;
-			FRotation3 DQ = FRotation3::FromElements(WorldSpaceInvI * FVec3::CrossProduct(Arm, DX), 0.f) * Q0 * FReal(0.5);
-
-			P0 += DP;
-			Q0 += DQ;
-			Q0.Normalize();
-			FParticleUtilities::SetCoMWorldTransform(Particle, P0, Q0);
+			const FVec3 DP = Body.InvM() * DX;
+			const FVec3 DR = Body.InvI() * FVec3::CrossProduct(Arm, DX);
+			Body.ApplyTransformDelta(DP, DR);
+			Body.UpdateRotationDependentState();
 		}
-
 	}
 	
-	template class TContainerConstraintHandle<FPBDSuspensionConstraints>;
+	template class TIndexedContainerConstraintHandle<FPBDSuspensionConstraints>;
 
 
 }
diff --git a/Engine/Source/Runtime/Experimental/Chaos/Private/Chaos/ParticleHandle.cpp b/Engine/Source/Runtime/Experimental/Chaos/Private/Chaos/ParticleHandle.cpp
--- a/Engine/Source/Runtime/Experimental/Chaos/Private/Chaos/ParticleHandle.cpp
+++ b/Engine/Source/Runtime/Experimental/Chaos/Private/Chaos/ParticleHandle.cpp
@@ -243,6 +243,7 @@ namespace Chaos
 	const FRotation3 FGenericParticleHandleHandleImp::IdentityRotation = FRotation3(FQuat::Identity);
 	const FMatrix33 FGenericParticleHandleHandleImp::ZeroMatrix = FMatrix33(0);
 	const TUniquePtr<FBVHParticles> FGenericParticleHandleHandleImp::NullBVHParticles = TUniquePtr<FBVHParticles>();
+	const FKinematicTarget FGenericParticleHandleHandleImp::EmptyKinematicTarget;
 
 	template <>
 	template <>
diff --git a/Engine/Source/Runtime/Experimental/Chaos/Private/Chaos/TriangleMeshImplicitObject.cpp b/Engine/Source/Runtime/Experimental/Chaos/Private/Chaos/TriangleMeshImplicitObject.cpp
--- a/Engine/Source/Runtime/Experimental/Chaos/Private/Chaos/TriangleMeshImplicitObject.cpp
+++ b/Engine/Source/Runtime/Experimental/Chaos/Private/Chaos/TriangleMeshImplicitObject.cpp
@@ -8,6 +8,8 @@
 #include "Chaos/GeometryQueries.h"
 #include "Chaos/Utilities.h"
 
+//PRAGMA_DISABLE_OPTIMIZATION
+
 namespace Chaos
 {
 
@@ -120,6 +122,7 @@ struct FTriangleMeshRaycastVisitor
 		const FVec3& B = Particles.X(Elements[TriIdx][1]);
 		const FVec3& C = Particles.X(Elements[TriIdx][2]);
 
+		// Note: the math here needs to match FTriangleMeshImplicitObject::GetFaceNormal
 		const FVec3 AB = B - A;
 		const FVec3 AC = C - A;
 		FVec3 TriNormal = FVec3::CrossProduct(AB, AC);
@@ -157,6 +160,8 @@ struct FTriangleMeshRaycastVisitor
 				{
 					OutTime = 0;
 					OutFaceIndex = TriIdx;
+					//OutPosition = IntersectionPosition;
+					//OutNormal = RaycastNormal;	//We use the plane normal even when hitting triangle edges. This is to deal with triangles that approximate a single flat surface.
 					return false; //no one will beat Time == 0
 				}
 			}
diff --git a/Engine/Source/Runtime/Experimental/Chaos/Private/EventDefaults.cpp b/Engine/Source/Runtime/Experimental/Chaos/Private/EventDefaults.cpp
--- a/Engine/Source/Runtime/Experimental/Chaos/Private/EventDefaults.cpp
+++ b/Engine/Source/Runtime/Experimental/Chaos/Private/EventDefaults.cpp
@@ -80,13 +80,10 @@ namespace Chaos
 						break;
 					}
 
-					if (ContactHandle->GetType() == FCollisionConstraintBase::FType::SinglePoint ||
-						ContactHandle->GetType() == FCollisionConstraintBase::FType::SinglePointSwept)
+					if (ContactHandle->GetType() != ECollisionConstraintType::None)
 					{
-						const FRigidBodyPointContactConstraint& Constraint = (ContactHandle->GetType() == FCollisionConstraintBase::FType::SinglePoint) ?
-							ContactHandle->GetPointContact() :
-							*ContactHandle->GetSweptPointContact().As<FRigidBodyPointContactConstraint>();
-						
+						const FPBDCollisionConstraint& Constraint = ContactHandle->GetContact();
+
 						if (ensure(!Constraint.AccumulatedImpulse.ContainsNaN() && FMath::IsFinite(Constraint.GetPhi())))
 						{
 							FGeometryParticleHandle* Particle0 = Constraint.Particle[0];
@@ -134,12 +131,9 @@ namespace Chaos
 				{
 					for (int32 IdxCollision = 0; IdxCollision < ValidCollisionHandles.Num(); ++IdxCollision)
 					{
-						if (ValidCollisionHandles[IdxCollision]->GetType() == FCollisionConstraintBase::FType::SinglePoint ||
-							ValidCollisionHandles[IdxCollision]->GetType() == FCollisionConstraintBase::FType::SinglePointSwept)
+						if (ValidCollisionHandles[IdxCollision]->GetType() != ECollisionConstraintType::None)
 						{
-							const FRigidBodyPointContactConstraint& Constraint = (ValidCollisionHandles[IdxCollision]->GetType() == FCollisionConstraintBase::FType::SinglePoint) ? 
-								ValidCollisionHandles[IdxCollision]->GetPointContact() :
-								*ValidCollisionHandles[IdxCollision]->GetSweptPointContact().As<FRigidBodyPointContactConstraint>();
+							const FPBDCollisionConstraint& Constraint = ValidCollisionHandles[IdxCollision]->GetContact();
 
 							FGeometryParticleHandle* Particle0 = Constraint.Particle[0];
 							FGeometryParticleHandle* Particle1 = Constraint.Particle[1];
