Chaos

diff --git a/Engine/Source/Runtime/Experimental/Chaos/Private/PBDRigidsSolver.cpp b/Engine/Source/Runtime/Experimental/Chaos/Private/PBDRigidsSolver.cpp
--- a/Engine/Source/Runtime/Experimental/Chaos/Private/PBDRigidsSolver.cpp
+++ b/Engine/Source/Runtime/Experimental/Chaos/Private/PBDRigidsSolver.cpp
@@ -40,6 +40,7 @@ DECLARE_DWORD_ACCUMULATOR_STAT(TEXT("NumStaticParticles"), STAT_ChaosCounter_Num
 DECLARE_DWORD_ACCUMULATOR_STAT(TEXT("NumGeomCollParticles"), STAT_ChaosCounter_NumGeometryCollectionParticles, STATGROUP_ChaosCounters);
 DECLARE_DWORD_ACCUMULATOR_STAT(TEXT("NumIslands"), STAT_ChaosCounter_NumIslands, STATGROUP_ChaosCounters);
 DECLARE_DWORD_ACCUMULATOR_STAT(TEXT("NumContacts"), STAT_ChaosCounter_NumContacts, STATGROUP_ChaosCounters);
+DECLARE_DWORD_ACCUMULATOR_STAT(TEXT("NumActiveContacts"), STAT_ChaosCounter_NumActiveContacts, STATGROUP_ChaosCounters);
 DECLARE_DWORD_ACCUMULATOR_STAT(TEXT("NumJoints"), STAT_ChaosCounter_NumJoints, STATGROUP_ChaosCounters);
 
 CSV_DEFINE_CATEGORY(ChaosCounters, true);
@@ -59,116 +60,134 @@ DECLARE_DWORD_ACCUMULATOR_STAT(TEXT("NumPushOutJointIterations"), STAT_ChaosIter
 // Must be 0 when checked in...
 #define CHAOS_SOLVER_ENABLE_DEBUG_DRAW 0
 
-int32 ChaosSolverDebugDrawShapes = CHAOS_SOLVER_ENABLE_DEBUG_DRAW;
-int32 ChaosSolverDebugDrawCollisions = CHAOS_SOLVER_ENABLE_DEBUG_DRAW;
-int32 ChaosSolverDebugDrawCollidingShapes = 0;
-int32 ChaosSolverDebugDrawBounds = 0;
-int32 ChaosSolverDrawTransforms = 0;
-int32 ChaosSolverDrawIslands = 0;
-int32 ChaosSolverDrawShapesShowStatic = 1;
-int32 ChaosSolverDrawShapesShowKinematic = 1;
-int32 ChaosSolverDrawShapesShowDynamic = 1;
-int32 ChaosSolverDrawJoints = 0;
-Chaos::DebugDraw::FChaosDebugDrawJointFeatures ChaosSolverDrawJointFeatures = Chaos::DebugDraw::FChaosDebugDrawJointFeatures::MakeDefault();
-FAutoConsoleVariableRef CVarChaosSolverDrawShapes(TEXT("p.Chaos.Solver.DebugDrawShapes"), ChaosSolverDebugDrawShapes, TEXT("Draw Shapes (0 = never; 1 = end of frame)."));
-FAutoConsoleVariableRef CVarChaosSolverDrawCollisions(TEXT("p.Chaos.Solver.DebugDrawCollisions"), ChaosSolverDebugDrawCollisions, TEXT("Draw Collisions (0 = never; 1 = end of frame)."));
-FAutoConsoleVariableRef CVarChaosSolverDrawCollidingShapes(TEXT("p.Chaos.Solver.DebugDrawCollidingShapes"), ChaosSolverDebugDrawCollidingShapes, TEXT("Draw Shapes that have collisions on them (0 = never; 1 = end of frame)."));
-FAutoConsoleVariableRef CVarChaosSolverDrawBounds(TEXT("p.Chaos.Solver.DebugDrawBounds"), ChaosSolverDebugDrawBounds, TEXT("Draw bounding volumes inside the broadphase (0 = never; 1 = end of frame)."));
-FAutoConsoleVariableRef CVarChaosSolverDrawTransforms(TEXT("p.Chaos.Solver.DebugDrawTransforms"), ChaosSolverDrawTransforms, TEXT("Draw particle transforms (0 = never; 1 = end of frame)."));
-FAutoConsoleVariableRef CVarChaosSolverDrawIslands(TEXT("p.Chaos.Solver.DebugDrawIslands"), ChaosSolverDrawIslands, TEXT("Draw solver islands (0 = never; 1 = end of frame)."));
-FAutoConsoleVariableRef CVarChaosSolverDrawShapesShapesStatic(TEXT("p.Chaos.Solver.DebugDraw.ShowStatics"), ChaosSolverDrawShapesShowStatic, TEXT("If DebugDrawShapes is enabled, whether to show static objects"));
-FAutoConsoleVariableRef CVarChaosSolverDrawShapesShapesKinematic(TEXT("p.Chaos.Solver.DebugDraw.ShowKinematics"), ChaosSolverDrawShapesShowKinematic, TEXT("If DebugDrawShapes is enabled, whether to show kinematic objects"));
-FAutoConsoleVariableRef CVarChaosSolverDrawShapesShapesDynamic(TEXT("p.Chaos.Solver.DebugDraw.ShowDynamics"), ChaosSolverDrawShapesShowDynamic, TEXT("If DebugDrawShapes is enabled, whether to show dynamic objects"));
-FAutoConsoleVariableRef CVarChaosSolverDrawJoints(TEXT("p.Chaos.Solver.DebugDrawJoints"), ChaosSolverDrawJoints, TEXT("Draw joints"));
-FAutoConsoleVariableRef CVarChaosSolverDrawJointFeaturesCoMConnector(TEXT("p.Chaos.Solver.DebugDraw.JointFeatures.CoMConnector"), ChaosSolverDrawJointFeatures.bCoMConnector, TEXT("Joint features mask (see FDebugDrawJointFeatures)."));
-FAutoConsoleVariableRef CVarChaosSolverDrawJointFeaturesActorConnector(TEXT("p.Chaos.Solver.DebugDraw.JointFeatures.ActorConnector"), ChaosSolverDrawJointFeatures.bActorConnector, TEXT("Joint features mask (see FDebugDrawJointFeatures)."));
-FAutoConsoleVariableRef CVarChaosSolverDrawJointFeaturesStretch(TEXT("p.Chaos.Solver.DebugDraw.JointFeatures.Stretch"), ChaosSolverDrawJointFeatures.bStretch, TEXT("Joint features mask (see FDebugDrawJointFeatures)."));
-FAutoConsoleVariableRef CVarChaosSolverDrawJointFeaturesAxes(TEXT("p.Chaos.Solver.DebugDraw.JointFeatures.Axes"), ChaosSolverDrawJointFeatures.bAxes, TEXT("Joint features mask (see FDebugDrawJointFeatures)."));
-FAutoConsoleVariableRef CVarChaosSolverDrawJointFeaturesLevel(TEXT("p.Chaos.Solver.DebugDraw.JointFeatures.Level"), ChaosSolverDrawJointFeatures.bLevel, TEXT("Joint features mask (see FDebugDrawJointFeatures)."));
-FAutoConsoleVariableRef CVarChaosSolverDrawJointFeaturesIndex(TEXT("p.Chaos.Solver.DebugDraw.JointFeatures.Index"), ChaosSolverDrawJointFeatures.bIndex, TEXT("Joint features mask (see FDebugDrawJointFeatures)."));
-FAutoConsoleVariableRef CVarChaosSolverDrawJointFeaturesColor(TEXT("p.Chaos.Solver.DebugDraw.JointFeatures.Color"), ChaosSolverDrawJointFeatures.bColor, TEXT("Joint features mask (see FDebugDrawJointFeatures)."));
-FAutoConsoleVariableRef CVarChaosSolverDrawJointFeaturesIsland(TEXT("p.Chaos.Solver.DebugDraw.JointFeatures.Island"), ChaosSolverDrawJointFeatures.bIsland, TEXT("Joint features mask (see FDebugDrawJointFeatures)."));
-
-Chaos::DebugDraw::FChaosDebugDrawSettings ChaosSolverDebugDebugDrawSettings(
-	/* ArrowSize =			*/ 10.0f,
-	/* BodyAxisLen =		*/ 30.0f,
-	/* ContactLen =			*/ 30.0f,
-	/* ContactWidth =		*/ 6.0f,
-	/* ContactPhiWidth =	*/ 0.0f,
-	/* ContactOwnerWidth =	*/ 0.0f,
-	/* ConstraintAxisLen =	*/ 30.0f,
-	/* JointComSize =		*/ 2.0f,
-	/* LineThickness =		*/ 1.0f,
-	/* DrawScale =			*/ 1.0f,
-	/* FontHeight =			*/ 10.0f,
-	/* FontScale =			*/ 1.5f,
-	/* ShapeThicknesScale = */ 1.0f,
-	/* PointSize =			*/ 5.0f,
-	/* VelScale =			*/ 0.0f,
-	/* AngVelScale =		*/ 0.0f,
-	/* ImpulseScale =		*/ 0.0f,
-	/* InertiaScale =				*/ 1.0f,
-	/* DrawPriority =				*/ 10,
-	/* bShowSimple =				*/ true,
-	/* bShowComplex =				*/ false,
-	/* bInShowLevelSetCollision =	*/ true,
-	/* InShapesColorsPerState =     */ Chaos::DebugDraw::GetDefaultShapesColorsByState(),
-	/* InShapesColorsPerShaepType=  */ Chaos::DebugDraw::GetDefaultShapesColorsByShapeType(),
-	/* InBoundsColorsPerState =     */ Chaos::DebugDraw::GetDefaultBoundsColorsByState(),
-	/* InBoundsColorsPerShapeType=  */ Chaos::DebugDraw::GetDefaultBoundsColorsByShapeType()
-);
-
-FAutoConsoleVariableRef CVarChaosSolverArrowSize(TEXT("p.Chaos.Solver.DebugDraw.ArrowSize"), ChaosSolverDebugDebugDrawSettings.ArrowSize, TEXT("ArrowSize."));
-FAutoConsoleVariableRef CVarChaosSolverBodyAxisLen(TEXT("p.Chaos.Solver.DebugDraw.BodyAxisLen"), ChaosSolverDebugDebugDrawSettings.BodyAxisLen, TEXT("BodyAxisLen."));
-FAutoConsoleVariableRef CVarChaosSolverContactLen(TEXT("p.Chaos.Solver.DebugDraw.ContactLen"), ChaosSolverDebugDebugDrawSettings.ContactLen, TEXT("ContactLen."));
-FAutoConsoleVariableRef CVarChaosSolverContactWidth(TEXT("p.Chaos.Solver.DebugDraw.ContactWidth"), ChaosSolverDebugDebugDrawSettings.ContactWidth, TEXT("ContactWidth."));
-FAutoConsoleVariableRef CVarChaosSolverContactPhiWidth(TEXT("p.Chaos.Solver.DebugDraw.ContactPhiWidth"), ChaosSolverDebugDebugDrawSettings.ContactPhiWidth, TEXT("ContactPhiWidth."));
-FAutoConsoleVariableRef CVarChaosSolverContactOwnerWidth(TEXT("p.Chaos.Solver.DebugDraw.ContactOwnerWidth"), ChaosSolverDebugDebugDrawSettings.ContactOwnerWidth, TEXT("ContactOwnerWidth."));
-FAutoConsoleVariableRef CVarChaosSolverConstraintAxisLen(TEXT("p.Chaos.Solver.DebugDraw.ConstraintAxisLen"), ChaosSolverDebugDebugDrawSettings.ConstraintAxisLen, TEXT("ConstraintAxisLen."));
-FAutoConsoleVariableRef CVarChaosSolverLineThickness(TEXT("p.Chaos.Solver.DebugDraw.LineThickness"), ChaosSolverDebugDebugDrawSettings.LineThickness, TEXT("LineThickness."));
-FAutoConsoleVariableRef CVarChaosSolverLineShapeThickness(TEXT("p.Chaos.Solver.DebugDraw.ShapeLineThicknessScale"), ChaosSolverDebugDebugDrawSettings.ShapeThicknesScale, TEXT("Shape lineThickness multiplier."));
-FAutoConsoleVariableRef CVarChaosSolverPointSize(TEXT("p.Chaos.Solver.DebugDraw.PointSize"), ChaosSolverDebugDebugDrawSettings.PointSize, TEXT("Point size."));
-FAutoConsoleVariableRef CVarChaosSolverVelScale(TEXT("p.Chaos.Solver.DebugDraw.VelScale"), ChaosSolverDebugDebugDrawSettings.VelScale, TEXT("If >0 show velocity when drawing particle transforms."));
-FAutoConsoleVariableRef CVarChaosSolverAngVelScale(TEXT("p.Chaos.Solver.DebugDraw.AngVelScale"), ChaosSolverDebugDebugDrawSettings.AngVelScale, TEXT("If >0 show angular velocity when drawing particle transforms."));
-FAutoConsoleVariableRef CVarChaosSolverImpulseScale(TEXT("p.Chaos.Solver.DebugDraw.ImpulseScale"), ChaosSolverDebugDebugDrawSettings.ImpulseScale, TEXT("If >0 show impulses when drawing collisions."));
-FAutoConsoleVariableRef CVarChaosSolverInertiaScale(TEXT("p.Chaos.Solver.DebugDraw.InertiaScale"), ChaosSolverDebugDebugDrawSettings.InertiaScale, TEXT("When DebugDrawTransforms is enabled, show the mass-normalized inertia matrix scaled by this amount."));
-FAutoConsoleVariableRef CVarChaosSolverScale(TEXT("p.Chaos.Solver.DebugDraw.Scale"), ChaosSolverDebugDebugDrawSettings.DrawScale, TEXT("Scale applied to all Chaos Debug Draw line lengths etc."));
-FAutoConsoleVariableRef CVarChaosSolverShowSimple(TEXT("p.Chaos.Solver.DebugDraw.ShowSimple"), ChaosSolverDebugDebugDrawSettings.bShowSimpleCollision, TEXT("Whether to show simple collision is shape drawing is enabled"));
-FAutoConsoleVariableRef CVarChaosSolverShowComplex(TEXT("p.Chaos.Solver.DebugDraw.ShowComplex"), ChaosSolverDebugDebugDrawSettings.bShowComplexCollision, TEXT("Whether to show complex collision is shape drawing is enabled"));
-FAutoConsoleVariableRef CVarChaosSolverShowLevelSet(TEXT("p.Chaos.Solver.DebugDraw.ShowLevelSet"), ChaosSolverDebugDebugDrawSettings.bShowLevelSetCollision, TEXT("Whether to show levelset collision is shape drawing is enabled"));
-
+namespace Chaos
+{
+	namespace CVars
+	{
+		int32 ChaosSolverDebugDrawShapes = CHAOS_SOLVER_ENABLE_DEBUG_DRAW;
+		int32 ChaosSolverDebugDrawCollisions = CHAOS_SOLVER_ENABLE_DEBUG_DRAW;
+		int32 ChaosSolverDebugDrawCollidingShapes = 0;
+		int32 ChaosSolverDebugDrawBounds = 0;
+		int32 ChaosSolverDrawTransforms = 0;
+		int32 ChaosSolverDrawIslands = 0;
+		int32 ChaosSolverDrawShapesShowStatic = 1;
+		int32 ChaosSolverDrawShapesShowKinematic = 1;
+		int32 ChaosSolverDrawShapesShowDynamic = 1;
+		int32 ChaosSolverDrawJoints = 0;
+		int32 ChaosSolverDebugDrawSpatialAccelerationStructure = 0;
+		int32 ChaosSolverDebugDrawSpatialAccelerationStructureShowLeaves = 0;
+		int32 ChaosSolverDebugDrawSpatialAccelerationStructureShowNodes = 0;
+		DebugDraw::FChaosDebugDrawJointFeatures ChaosSolverDrawJointFeatures = DebugDraw::FChaosDebugDrawJointFeatures::MakeDefault();
+		FAutoConsoleVariableRef CVarChaosSolverDrawShapes(TEXT("p.Chaos.Solver.DebugDrawShapes"), ChaosSolverDebugDrawShapes, TEXT("Draw Shapes (0 = never; 1 = end of frame)."));
+		FAutoConsoleVariableRef CVarChaosSolverDrawCollisions(TEXT("p.Chaos.Solver.DebugDrawCollisions"), ChaosSolverDebugDrawCollisions, TEXT("Draw Collisions (0 = never; 1 = end of frame)."));
+		FAutoConsoleVariableRef CVarChaosSolverDrawCollidingShapes(TEXT("p.Chaos.Solver.DebugDrawCollidingShapes"), ChaosSolverDebugDrawCollidingShapes, TEXT("Draw Shapes that have collisions on them (0 = never; 1 = end of frame)."));
+		FAutoConsoleVariableRef CVarChaosSolverDrawBounds(TEXT("p.Chaos.Solver.DebugDrawBounds"), ChaosSolverDebugDrawBounds, TEXT("Draw bounding volumes inside the broadphase (0 = never; 1 = end of frame)."));
+		FAutoConsoleVariableRef CVarChaosSolverDrawTransforms(TEXT("p.Chaos.Solver.DebugDrawTransforms"), ChaosSolverDrawTransforms, TEXT("Draw particle transforms (0 = never; 1 = end of frame)."));
+		FAutoConsoleVariableRef CVarChaosSolverDrawIslands(TEXT("p.Chaos.Solver.DebugDrawIslands"), ChaosSolverDrawIslands, TEXT("Draw solver islands (0 = never; 1 = end of frame)."));
+		FAutoConsoleVariableRef CVarChaosSolverDrawShapesShapesStatic(TEXT("p.Chaos.Solver.DebugDraw.ShowStatics"), ChaosSolverDrawShapesShowStatic, TEXT("If DebugDrawShapes is enabled, whether to show static objects"));
+		FAutoConsoleVariableRef CVarChaosSolverDrawShapesShapesKinematic(TEXT("p.Chaos.Solver.DebugDraw.ShowKinematics"), ChaosSolverDrawShapesShowKinematic, TEXT("If DebugDrawShapes is enabled, whether to show kinematic objects"));
+		FAutoConsoleVariableRef CVarChaosSolverDrawShapesShapesDynamic(TEXT("p.Chaos.Solver.DebugDraw.ShowDynamics"), ChaosSolverDrawShapesShowDynamic, TEXT("If DebugDrawShapes is enabled, whether to show dynamic objects"));
+		FAutoConsoleVariableRef CVarChaosSolverDrawJoints(TEXT("p.Chaos.Solver.DebugDrawJoints"), ChaosSolverDrawJoints, TEXT("Draw joints"));
+		FAutoConsoleVariableRef CVarChaosSolverDebugDrawSpatialAccelerationStructure(TEXT("p.Chaos.Solver.DebugDrawSpatialAccelerationStructure"), ChaosSolverDebugDrawSpatialAccelerationStructure, TEXT("Draw spatial acceleration structure"));
+		FAutoConsoleVariableRef CVarChaosSolverDebugDrawSpatialAccelerationStructureShowLeaves(TEXT("p.Chaos.Solver.DebugDrawSpatialAccelerationStructure.ShowLeaves"), ChaosSolverDebugDrawSpatialAccelerationStructureShowLeaves, TEXT("Show spatial acceleration structure leaves when its debug draw is enabled"));
+		FAutoConsoleVariableRef CVarChaosSolverDebugDrawSpatialAccelerationStructureShowNodes(TEXT("p.Chaos.Solver.DebugDrawSpatialAccelerationStructure.ShowNodes"), ChaosSolverDebugDrawSpatialAccelerationStructureShowNodes, TEXT("Show spatial acceleration structure nodes when its debug draw is enabled"));
+		FAutoConsoleVariableRef CVarChaosSolverDrawJointFeaturesCoMConnector(TEXT("p.Chaos.Solver.DebugDraw.JointFeatures.CoMConnector"), ChaosSolverDrawJointFeatures.bCoMConnector, TEXT("Joint features mask (see FDebugDrawJointFeatures)."));
+		FAutoConsoleVariableRef CVarChaosSolverDrawJointFeaturesActorConnector(TEXT("p.Chaos.Solver.DebugDraw.JointFeatures.ActorConnector"), ChaosSolverDrawJointFeatures.bActorConnector, TEXT("Joint features mask (see FDebugDrawJointFeatures)."));
+		FAutoConsoleVariableRef CVarChaosSolverDrawJointFeaturesStretch(TEXT("p.Chaos.Solver.DebugDraw.JointFeatures.Stretch"), ChaosSolverDrawJointFeatures.bStretch, TEXT("Joint features mask (see FDebugDrawJointFeatures)."));
+		FAutoConsoleVariableRef CVarChaosSolverDrawJointFeaturesAxes(TEXT("p.Chaos.Solver.DebugDraw.JointFeatures.Axes"), ChaosSolverDrawJointFeatures.bAxes, TEXT("Joint features mask (see FDebugDrawJointFeatures)."));
+		FAutoConsoleVariableRef CVarChaosSolverDrawJointFeaturesLevel(TEXT("p.Chaos.Solver.DebugDraw.JointFeatures.Level"), ChaosSolverDrawJointFeatures.bLevel, TEXT("Joint features mask (see FDebugDrawJointFeatures)."));
+		FAutoConsoleVariableRef CVarChaosSolverDrawJointFeaturesIndex(TEXT("p.Chaos.Solver.DebugDraw.JointFeatures.Index"), ChaosSolverDrawJointFeatures.bIndex, TEXT("Joint features mask (see FDebugDrawJointFeatures)."));
+		FAutoConsoleVariableRef CVarChaosSolverDrawJointFeaturesColor(TEXT("p.Chaos.Solver.DebugDraw.JointFeatures.Color"), ChaosSolverDrawJointFeatures.bColor, TEXT("Joint features mask (see FDebugDrawJointFeatures)."));
+		FAutoConsoleVariableRef CVarChaosSolverDrawJointFeaturesIsland(TEXT("p.Chaos.Solver.DebugDraw.JointFeatures.Island"), ChaosSolverDrawJointFeatures.bIsland, TEXT("Joint features mask (see FDebugDrawJointFeatures)."));
+
+
+		DebugDraw::FChaosDebugDrawSettings ChaosSolverDebugDebugDrawSettings(
+			/* ArrowSize =					*/ 10.0f,
+			/* BodyAxisLen =				*/ 30.0f,
+			/* ContactLen =					*/ 30.0f,
+			/* ContactWidth =				*/ 6.0f,
+			/* ContactPhiWidth =			*/ 0.0f,
+			/* ContactOwnerWidth =			*/ 0.0f,
+			/* ConstraintAxisLen =			*/ 30.0f,
+			/* JointComSize =				*/ 2.0f,
+			/* LineThickness =				*/ 1.0f,
+			/* DrawScale =					*/ 1.0f,
+			/* FontHeight =					*/ 10.0f,
+			/* FontScale =					*/ 1.5f,
+			/* ShapeThicknesScale =			*/ 1.0f,
+			/* PointSize =					*/ 5.0f,
+			/* VelScale =					*/ 0.0f,
+			/* AngVelScale =				*/ 0.0f,
+			/* ImpulseScale =				*/ 0.0f,
+			/* PushOutScale =				*/ 0.0f,
+			/* InertiaScale =				*/ 1.0f,
+			/* DrawPriority =				*/ 10,
+			/* bShowSimple =				*/ true,
+			/* bShowComplex =				*/ false,
+			/* bInShowLevelSetCollision =	*/ true,
+			/* InShapesColorsPerState =     */ DebugDraw::GetDefaultShapesColorsByState(),
+			/* InShapesColorsPerShaepType=  */ DebugDraw::GetDefaultShapesColorsByShapeType(),
+			/* InBoundsColorsPerState =     */ DebugDraw::GetDefaultBoundsColorsByState(),
+			/* InBoundsColorsPerShapeType=  */ DebugDraw::GetDefaultBoundsColorsByShapeType()
+		);
+
+		FAutoConsoleVariableRef CVarChaosSolverArrowSize(TEXT("p.Chaos.Solver.DebugDraw.ArrowSize"), ChaosSolverDebugDebugDrawSettings.ArrowSize, TEXT("ArrowSize."));
+		FAutoConsoleVariableRef CVarChaosSolverBodyAxisLen(TEXT("p.Chaos.Solver.DebugDraw.BodyAxisLen"), ChaosSolverDebugDebugDrawSettings.BodyAxisLen, TEXT("BodyAxisLen."));
+		FAutoConsoleVariableRef CVarChaosSolverContactLen(TEXT("p.Chaos.Solver.DebugDraw.ContactLen"), ChaosSolverDebugDebugDrawSettings.ContactLen, TEXT("ContactLen."));
+		FAutoConsoleVariableRef CVarChaosSolverContactWidth(TEXT("p.Chaos.Solver.DebugDraw.ContactWidth"), ChaosSolverDebugDebugDrawSettings.ContactWidth, TEXT("ContactWidth."));
+		FAutoConsoleVariableRef CVarChaosSolverContactPhiWidth(TEXT("p.Chaos.Solver.DebugDraw.ContactPhiWidth"), ChaosSolverDebugDebugDrawSettings.ContactPhiWidth, TEXT("ContactPhiWidth."));
+		FAutoConsoleVariableRef CVarChaosSolverContactOwnerWidth(TEXT("p.Chaos.Solver.DebugDraw.ContactOwnerWidth"), ChaosSolverDebugDebugDrawSettings.ContactOwnerWidth, TEXT("ContactOwnerWidth."));
+		FAutoConsoleVariableRef CVarChaosSolverConstraintAxisLen(TEXT("p.Chaos.Solver.DebugDraw.ConstraintAxisLen"), ChaosSolverDebugDebugDrawSettings.ConstraintAxisLen, TEXT("ConstraintAxisLen."));
+		FAutoConsoleVariableRef CVarChaosSolverLineThickness(TEXT("p.Chaos.Solver.DebugDraw.LineThickness"), ChaosSolverDebugDebugDrawSettings.LineThickness, TEXT("LineThickness."));
+		FAutoConsoleVariableRef CVarChaosSolverLineShapeThickness(TEXT("p.Chaos.Solver.DebugDraw.ShapeLineThicknessScale"), ChaosSolverDebugDebugDrawSettings.ShapeThicknesScale, TEXT("Shape lineThickness multiplier."));
+		FAutoConsoleVariableRef CVarChaosSolverPointSize(TEXT("p.Chaos.Solver.DebugDraw.PointSize"), ChaosSolverDebugDebugDrawSettings.PointSize, TEXT("Point size."));
+		FAutoConsoleVariableRef CVarChaosSolverVelScale(TEXT("p.Chaos.Solver.DebugDraw.VelScale"), ChaosSolverDebugDebugDrawSettings.VelScale, TEXT("If >0 show velocity when drawing particle transforms."));
+		FAutoConsoleVariableRef CVarChaosSolverAngVelScale(TEXT("p.Chaos.Solver.DebugDraw.AngVelScale"), ChaosSolverDebugDebugDrawSettings.AngVelScale, TEXT("If >0 show angular velocity when drawing particle transforms."));
+		FAutoConsoleVariableRef CVarChaosSolverImpulseScale(TEXT("p.Chaos.Solver.DebugDraw.ImpulseScale"), ChaosSolverDebugDebugDrawSettings.ImpulseScale, TEXT("If >0 show impulses when drawing collisions."));
+		FAutoConsoleVariableRef CVarChaosSolverPushOutScale(TEXT("p.Chaos.Solver.DebugDraw.PushOutScale"), ChaosSolverDebugDebugDrawSettings.PushOutScale, TEXT("If >0 show pushouts when drawing collisions."));
+		FAutoConsoleVariableRef CVarChaosSolverInertiaScale(TEXT("p.Chaos.Solver.DebugDraw.InertiaScale"), ChaosSolverDebugDebugDrawSettings.InertiaScale, TEXT("When DebugDrawTransforms is enabled, show the mass-normalized inertia matrix scaled by this amount."));
+		FAutoConsoleVariableRef CVarChaosSolverScale(TEXT("p.Chaos.Solver.DebugDraw.Scale"), ChaosSolverDebugDebugDrawSettings.DrawScale, TEXT("Scale applied to all Chaos Debug Draw line lengths etc."));
+		FAutoConsoleVariableRef CVarChaosSolverShowSimple(TEXT("p.Chaos.Solver.DebugDraw.ShowSimple"), ChaosSolverDebugDebugDrawSettings.bShowSimpleCollision, TEXT("Whether to show simple collision is shape drawing is enabled"));
+		FAutoConsoleVariableRef CVarChaosSolverShowComplex(TEXT("p.Chaos.Solver.DebugDraw.ShowComplex"), ChaosSolverDebugDebugDrawSettings.bShowComplexCollision, TEXT("Whether to show complex collision is shape drawing is enabled"));
+		FAutoConsoleVariableRef CVarChaosSolverShowLevelSet(TEXT("p.Chaos.Solver.DebugDraw.ShowLevelSet"), ChaosSolverDebugDebugDrawSettings.bShowLevelSetCollision, TEXT("Whether to show levelset collision is shape drawing is enabled"));
+	}
+}
 #endif
 
-bool ChaosSolverUseParticlePool = true;
-FAutoConsoleVariableRef CVarChaosSolverUseParticlePool(TEXT("p.Chaos.Solver.UseParticlePool"), ChaosSolverUseParticlePool, TEXT("Whether or not to use dirty particle pool (Optim)"));
+namespace Chaos
+{
+	namespace CVars
+	{
+		bool ChaosSolverUseParticlePool = true;
+		FAutoConsoleVariableRef CVarChaosSolverUseParticlePool(TEXT("p.Chaos.Solver.UseParticlePool"), ChaosSolverUseParticlePool, TEXT("Whether or not to use dirty particle pool (Optim)"));
 
-int32 ChaosSolverParticlePoolNumFrameUntilShrink = 30;
-FAutoConsoleVariableRef CVarChaosSolverParticlePoolNumFrameUntilShrink(TEXT("p.Chaos.Solver.ParticlePoolNumFrameUntilShrink"), ChaosSolverParticlePoolNumFrameUntilShrink, TEXT("Num Frame until we can potentially shrink the pool"));
+		int32 ChaosSolverParticlePoolNumFrameUntilShrink = 30;
+		FAutoConsoleVariableRef CVarChaosSolverParticlePoolNumFrameUntilShrink(TEXT("p.Chaos.Solver.ParticlePoolNumFrameUntilShrink"), ChaosSolverParticlePoolNumFrameUntilShrink, TEXT("Num Frame until we can potentially shrink the pool"));
 
-// Select the solver technique to use until we settle on the final one...
-int32 ChaosSolver_SolverType = (int32)Chaos::EConstraintSolverType::GbfPbd;
-FAutoConsoleVariableRef CVarChaosSolverSolverType(TEXT("p.Chaos.Solver.SolverType"), ChaosSolver_SolverType, TEXT("0 = None; 1 = GbfPbd; 2 = Pbd; 3 = QuasiPbd"));
+		// Select the solver technique to use until we settle on the final one...
+		int32 ChaosSolver_SolverType = (int32)Chaos::EConstraintSolverType::QuasiPbd;
+		FAutoConsoleVariableRef CVarChaosSolverSolverType(TEXT("p.Chaos.Solver.SolverType"), ChaosSolver_SolverType, TEXT("0 = None; 1 = GbfPbd; 2 = Pbd; 3 = QuasiPbd"));
 
-// Iteration count cvars
-// These override the engine config if >= 0
+		// Iteration count cvars
+		// These override the engine config if >= 0
 
-int32 ChaosSolverIterations = -1;
-FAutoConsoleVariableRef CVarChaosSolverIterations(TEXT("p.Chaos.Solver.Iterations"), ChaosSolverIterations, TEXT("Override umber of solver iterations (-1 to use config)"));
+		int32 ChaosSolverIterations = -1;
+		FAutoConsoleVariableRef CVarChaosSolverIterations(TEXT("p.Chaos.Solver.Iterations"), ChaosSolverIterations, TEXT("Override umber of solver iterations (-1 to use config)"));
 
-int32 ChaosSolverCollisionIterations = -1;
-FAutoConsoleVariableRef CVarChaosSolverCollisionIterations(TEXT("p.Chaos.Solver.Collision.Iterations"), ChaosSolverCollisionIterations, TEXT("Override number of collision iterations per solver iteration (-1 to use config)"));
+		int32 ChaosSolverCollisionIterations = -1;
+		FAutoConsoleVariableRef CVarChaosSolverCollisionIterations(TEXT("p.Chaos.Solver.Collision.Iterations"), ChaosSolverCollisionIterations, TEXT("Override number of collision iterations per solver iteration (-1 to use config)"));
 
-int32 ChaosSolverPushOutIterations = -1;
-FAutoConsoleVariableRef CVarChaosSolverPushOutIterations(TEXT("p.Chaos.Solver.PushoutIterations"), ChaosSolverPushOutIterations, TEXT("Override number of solver pushout iterations (-1 to use config)"));
+		int32 ChaosSolverPushOutIterations = -1;
+		FAutoConsoleVariableRef CVarChaosSolverPushOutIterations(TEXT("p.Chaos.Solver.PushoutIterations"), ChaosSolverPushOutIterations, TEXT("Override number of solver pushout iterations (-1 to use config)"));
 
-int32 ChaosSolverCollisionPushOutIterations = -1;
-FAutoConsoleVariableRef CVarChaosSolverCollisionPushOutIterations(TEXT("p.Chaos.Solver.Collision.PushOutIterations"), ChaosSolverCollisionPushOutIterations, TEXT("Override number of collision iterations per solver iteration (-1 to use config)"));
+		int32 ChaosSolverCollisionPushOutIterations = -1;
+		FAutoConsoleVariableRef CVarChaosSolverCollisionPushOutIterations(TEXT("p.Chaos.Solver.Collision.PushOutIterations"), ChaosSolverCollisionPushOutIterations, TEXT("Override number of collision iterations per solver iteration (-1 to use config)"));
 
-int32 ChaosSolverJointPairIterations = -1;
-FAutoConsoleVariableRef CVarChaosSolverJointPairIterations(TEXT("p.Chaos.Solver.Joint.PairIterations"), ChaosSolverJointPairIterations, TEXT("Override number of iterations per joint pair during a solver iteration (-1 to use config)"));
+		int32 ChaosSolverJointPairIterations = -1;
+		FAutoConsoleVariableRef CVarChaosSolverJointPairIterations(TEXT("p.Chaos.Solver.Joint.PairIterations"), ChaosSolverJointPairIterations, TEXT("Override number of iterations per joint pair during a solver iteration (-1 to use config)"));
 
-int32 ChaosSolverJointPushOutPairIterations = -1;
-FAutoConsoleVariableRef CVarChaosSolverJointPushOutPairIterations(TEXT("p.Chaos.Solver.Joint.PushOutPairIterations"), ChaosSolverJointPushOutPairIterations, TEXT("Override number of push out iterations per joint during a solver iteration (-1 to use config)"));
+		int32 ChaosSolverJointPushOutPairIterations = -1;
+		FAutoConsoleVariableRef CVarChaosSolverJointPushOutPairIterations(TEXT("p.Chaos.Solver.Joint.PushOutPairIterations"), ChaosSolverJointPushOutPairIterations, TEXT("Override number of push out iterations per joint during a solver iteration (-1 to use config)"));
 
 // Copied from RBAN
 Chaos::FRealSingle ChaosSolverJointPositionTolerance = 0.025f;
@@ -180,36 +199,39 @@ FAutoConsoleVariableRef CVarChaosSolverJointMinParentMassRatio(TEXT("p.Chaos.Sol
 Chaos::FRealSingle ChaosSolverJointMaxInertiaRatio = 5.0f;
 FAutoConsoleVariableRef CVarChaosSolverJointMaxInertiaRatio(TEXT("p.Chaos.Solver.Joint.MaxInertiaRatio"), ChaosSolverJointMaxInertiaRatio, TEXT("6Dof joint MaxInertiaRatio (if > 0)"));
 
-// Collision detection cvars
-// These override the engine config if >= 0
-Chaos::FRealSingle ChaosSolverCullDistance = -1.0f;
-FAutoConsoleVariableRef CVarChaosSolverCullDistance(TEXT("p.Chaos.Solver.Collision.CullDistance"), ChaosSolverCullDistance, TEXT("Override cull distance (if >= 0)"));
+		// Collision detection cvars
+		// These override the engine config if >= 0
+		Chaos::FRealSingle ChaosSolverCullDistance = -1.0f;
+		FAutoConsoleVariableRef CVarChaosSolverCullDistance(TEXT("p.Chaos.Solver.Collision.CullDistance"), ChaosSolverCullDistance, TEXT("Override cull distance (if >= 0)"));
 
-int32 ChaosSolverCleanupCommandsOnDestruction = 1;
-FAutoConsoleVariableRef CVarChaosSolverCleanupCommandsOnDestruction(TEXT("p.Chaos.Solver.CleanupCommandsOnDestruction"), ChaosSolverCleanupCommandsOnDestruction, TEXT("Whether or not to run internal command queue cleanup on solver destruction (0 = no cleanup, >0 = cleanup all commands)"));
+		int32 ChaosSolverCleanupCommandsOnDestruction = 1;
+		FAutoConsoleVariableRef CVarChaosSolverCleanupCommandsOnDestruction(TEXT("p.Chaos.Solver.CleanupCommandsOnDestruction"), ChaosSolverCleanupCommandsOnDestruction, TEXT("Whether or not to run internal command queue cleanup on solver destruction (0 = no cleanup, >0 = cleanup all commands)"));
 
-int32 ChaosSolverCollisionDeferNarrowPhase = 0;
-FAutoConsoleVariableRef CVarChaosSolverCollisionDeferNarrowPhase(TEXT("p.Chaos.Solver.Collision.DeferNarrowPhase"), ChaosSolverCollisionDeferNarrowPhase, TEXT("Create contacts for all broadphase pairs, perform NarrowPhase later."));
+		int32 ChaosSolverCollisionDeferNarrowPhase = 0;
+		FAutoConsoleVariableRef CVarChaosSolverCollisionDeferNarrowPhase(TEXT("p.Chaos.Solver.Collision.DeferNarrowPhase"), ChaosSolverCollisionDeferNarrowPhase, TEXT("Create contacts for all broadphase pairs, perform NarrowPhase later."));
 
-// Allow one-shot or incremental manifolds where supported (which depends on shape pair types)
-int32 ChaosSolverCollisionUseManifolds = 1;
-FAutoConsoleVariableRef CVarChaosSolverCollisionUseManifolds(TEXT("p.Chaos.Solver.Collision.UseManifolds"), ChaosSolverCollisionUseManifolds, TEXT("Enable/Disable use of manifolds in collision."));
+		// Allow one-shot or incremental manifolds where supported (which depends on shape pair types)
+		int32 ChaosSolverCollisionUseManifolds = 1;
+		FAutoConsoleVariableRef CVarChaosSolverCollisionUseManifolds(TEXT("p.Chaos.Solver.Collision.UseManifolds"), ChaosSolverCollisionUseManifolds, TEXT("Enable/Disable use of manifolds in collision."));
 
 
-// Joint cvars
-float ChaosSolverJointMinSolverStiffness = 1.0f;
-float ChaosSolverJointMaxSolverStiffness = 1.0f;
-int32 ChaosSolverJointNumIterationsAtMaxSolverStiffness = 1;
-FAutoConsoleVariableRef CVarChaosSolverJointMinSolverStiffness(TEXT("p.Chaos.Solver.Joint.MinSolverStiffness"), ChaosSolverJointMinSolverStiffness, TEXT("Solver stiffness on first iteration, increases each iteration toward MaxSolverStiffness."));
-FAutoConsoleVariableRef CVarChaosSolverJointMaxSolverStiffness(TEXT("p.Chaos.Solver.Joint.MaxSolverStiffness"), ChaosSolverJointMaxSolverStiffness, TEXT("Solver stiffness on last iteration, increases each iteration from MinSolverStiffness."));
-FAutoConsoleVariableRef CVarChaosSolverJointNumIterationsAtMaxSolverStiffness(TEXT("p.Chaos.Solver.Joint.NumIterationsAtMaxSolverStiffness"), ChaosSolverJointNumIterationsAtMaxSolverStiffness, TEXT("How many iterations we want at MaxSolverStiffness."));
+		// Joint cvars
+		float ChaosSolverJointMinSolverStiffness = 1.0f;
+		float ChaosSolverJointMaxSolverStiffness = 1.0f;
+		int32 ChaosSolverJointNumIterationsAtMaxSolverStiffness = 1;
+		FAutoConsoleVariableRef CVarChaosSolverJointMinSolverStiffness(TEXT("p.Chaos.Solver.Joint.MinSolverStiffness"), ChaosSolverJointMinSolverStiffness, TEXT("Solver stiffness on first iteration, increases each iteration toward MaxSolverStiffness."));
+		FAutoConsoleVariableRef CVarChaosSolverJointMaxSolverStiffness(TEXT("p.Chaos.Solver.Joint.MaxSolverStiffness"), ChaosSolverJointMaxSolverStiffness, TEXT("Solver stiffness on last iteration, increases each iteration from MinSolverStiffness."));
+		FAutoConsoleVariableRef CVarChaosSolverJointNumIterationsAtMaxSolverStiffness(TEXT("p.Chaos.Solver.Joint.NumIterationsAtMaxSolverStiffness"), ChaosSolverJointNumIterationsAtMaxSolverStiffness, TEXT("How many iterations we want at MaxSolverStiffness."));
 
 
-int32 ChaosVisualDebuggerEnable = 1;
-FAutoConsoleVariableRef CVarChaosVisualDebuggerEnable(TEXT("p.Chaos.VisualDebuggerEnable"), ChaosVisualDebuggerEnable, TEXT("Enable/Disable pushing/saving data to the visual debugger"));
+		int32 ChaosVisualDebuggerEnable = 1;
+		FAutoConsoleVariableRef CVarChaosVisualDebuggerEnable(TEXT("p.Chaos.VisualDebuggerEnable"), ChaosVisualDebuggerEnable, TEXT("Enable/Disable pushing/saving data to the visual debugger"));
+	}
+}
 
 namespace Chaos
 {
+	using namespace CVars;
 
 	class AdvanceOneTimeStepTask : public FNonAbandonableTask
 	{
@@ -272,7 +294,8 @@ namespace Chaos
 				SCOPE_CYCLE_COUNTER(STAT_EvolutionAndKinematicUpdate);
 
 				// clear out the collision constraints as they will be stale from last frame if AdvanceOneTimeStep never gets called due to TimeRemaining being less than MinDeltaTime 
-				MSolver->GetEvolution()->GetCollisionConstraintsRule().ResetConstraints();
+				// @todo(chaos): maybe we can pull data at a better time instead to avoid collision-specific code here for event dispatch
+				MSolver->GetEvolution()->GetCollisionConstraints().BeginFrame();
 			
 				// This outer loop can potentially cause the system to lose energy over integration
 				// in a couple of different cases.
@@ -1334,6 +1357,19 @@ namespace Chaos
 		return GetEvolution()->GetCollisionConstraints().NumConstraints();
 	}
 
+	int32 FPBDRigidsSolver::NumActiveCollisionConstraints() const
+	{
+		int32 NumActiveCollisions =  0;
+		for (const FPBDCollisionConstraintHandle* Collision : GetEvolution()->GetCollisionConstraints().GetConstraints())
+		{
+			if (Collision->GetContact().IsEnabled() && !Collision->GetContact().AccumulatedImpulse.IsNearlyZero())
+			{
+				++NumActiveCollisions;
+			}
+		}
+		return NumActiveCollisions;
+	}
+
 	void FPBDRigidsSolver::UpdateStatCounters() const
 	{
 		// Particle counts
@@ -1350,6 +1386,10 @@ namespace Chaos
 		SET_DWORD_STAT(STAT_ChaosCounter_NumContacts, NumCollisionConstraints());
 		SET_DWORD_STAT(STAT_ChaosCounter_NumJoints, NumJointConstraints());
 
+#if !UE_BUILD_SHIPPING
+		SET_DWORD_STAT(STAT_ChaosCounter_NumActiveContacts, NumActiveCollisionConstraints());
+#endif
+
 		// Iterations
 		SET_DWORD_STAT(STAT_ChaosIterations_NumIterations, GetEvolution()->GetNumIterations());
 		SET_DWORD_STAT(STAT_ChaosIterations_NumCollisionIterations, GetEvolution()->GetCollisionConstraints().GetPairIterations());
@@ -1407,7 +1447,7 @@ namespace Chaos
 		}
 		if (ChaosSolverDrawIslands == 1)
 		{
-			DebugDraw::DrawConstraintGraph(FRigidTransform3(), GetEvolution()->GetCollisionConstraintsRule().GetGraphColor(), &ChaosSolverDebugDebugDrawSettings);
+			DebugDraw::DrawConstraintGraph(FRigidTransform3(), GetEvolution()->GetConstraintGraph(), &ChaosSolverDebugDebugDrawSettings);
 		}
 		if (ChaosSolverDebugDrawCollidingShapes == 1)
 		{
@@ -1417,6 +1457,13 @@ namespace Chaos
 		{
 			DebugDraw::DrawJointConstraints(FRigidTransform3(), MEvolution->GetJointConstraints(), 1.0f, ChaosSolverDrawJointFeatures, &ChaosSolverDebugDebugDrawSettings);
 		}
+		if (ChaosSolverDebugDrawSpatialAccelerationStructure)
+		{
+			if (const auto SpatialAccelerationStructure = GetEvolution()->GetSpatialAcceleration())
+			{
+				DebugDraw::DrawSpatialAccelerationStructure(*SpatialAccelerationStructure, &ChaosSolverDebugDebugDrawSettings);
+			}
+		}
 		if (ChaosSolverDebugDrawCollidingShapes == 1)
 		{
 			DebugDraw::DrawCollidingShapes(FRigidTransform3(), GetEvolution()->GetCollisionConstraints(), 1.f, &ChaosSolverDebugDebugDrawSettings);
diff --git a/Engine/Source/Runtime/Experimental/Chaos/Private/PhysicsProxy/GeometryCollectionPhysicsProxy.cpp b/Engine/Source/Runtime/Experimental/Chaos/Private/PhysicsProxy/GeometryCollectionPhysicsProxy.cpp
--- a/Engine/Source/Runtime/Experimental/Chaos/Private/PhysicsProxy/GeometryCollectionPhysicsProxy.cpp
+++ b/Engine/Source/Runtime/Experimental/Chaos/Private/PhysicsProxy/GeometryCollectionPhysicsProxy.cpp
@@ -1980,7 +1980,7 @@ void FGeometryCollectionPhysicsProxy::InitializeSharedCollisionStructures(
 					InertiaComputationNeeded[GeometryIndex] = true;
 					if(MassProperties.Volume == 0)
 					{
-						FVector Extents = BoundingBox[GeometryIndex].GetExtent();
+						FVector Extents = FReal(2) * BoundingBox[GeometryIndex].GetExtent(); // FBox::GetExtent() returns half the size, but FAABB::Extents() returns total size
 						MassProperties.Volume = Extents.X * Extents.Y * Extents.Z;
 						FReal ExtentsYZ = Extents.Y * Extents.Y + Extents.Z * Extents.Z;
 						FReal ExtentsXZ = Extents.X * Extents.X + Extents.Z * Extents.Z;
diff --git a/Engine/Source/Runtime/Experimental/Chaos/Public/Chaos/AABBTree.h b/Engine/Source/Runtime/Experimental/Chaos/Public/Chaos/AABBTree.h
--- a/Engine/Source/Runtime/Experimental/Chaos/Public/Chaos/AABBTree.h
+++ b/Engine/Source/Runtime/Experimental/Chaos/Public/Chaos/AABBTree.h
@@ -19,6 +19,12 @@ struct CHAOS_API FAABBTreeCVars
 	static int32 UpdateDirtyElementPayloadData;
 	static FAutoConsoleVariableRef CVarUpdateDirtyElementPayloadData;
 
+	static int32 SplitAtAverageCenter;
+	static FAutoConsoleVariableRef CVarSplitAtAverageCenter;
+
+	static int32 SplitOnVarianceAxis;
+	static FAutoConsoleVariableRef CVarSplitOnVarianceAxis;
+
 	static float MaxNonGlobalElementBoundsExtrema; 
 	static FAutoConsoleVariableRef CVarMaxNonGlobalElementBoundsExtrema;
 };
@@ -146,7 +152,6 @@ struct TAABBTreeIntersectionHelper<FQueryFastDataVoid, EAABBQueryType::Overlap>
 template <typename TPayloadType, typename T, bool bComputeBounds>
 struct TBoundsWrapperHelper
 {
-
 };
 
 template <typename TPayloadType, typename T>
@@ -203,6 +208,11 @@ struct TAABBTreeLeafArray : public TBoundsWrapperHelper<TPayloadType, T, bComput
 		return Elems.Num();
 	}
 
+	SIZE_T GetElementCount() const
+	{
+		return Elems.Num();
+	}
+
 	template <typename TSQVisitor, typename TQueryFastData>
 	FORCEINLINE_DEBUGGABLE bool RaycastFast(const TVec3<T>& Start, TQueryFastData& QueryFastData, TSQVisitor& Visitor, const TVec3<T>& Dir, const TVec3<T> InvDir, const bool bParallel[3]) const
 	{
@@ -295,6 +305,25 @@ struct TAABBTreeLeafArray : public TBoundsWrapperHelper<TPayloadType, T, bComput
 		}
 	}
 
+
+#if !UE_BUILD_SHIPPING
+	void DebugDrawLeaf(ISpacialDebugDrawInterface<T>& InInterface, const FLinearColor& InLinearColor, float InThickness) const
+	{
+		const TAABB<T, 3> Bounds = TBoundsWrapperHelper<TPayloadType, T, bComputeBounds>::GetBounds();
+
+		const float Alpha = (float)Elems.Num() / 10.f;
+		const FLinearColor ColorByCount = FLinearColor::Green * (1.f - Alpha) + FLinearColor::Red * Alpha;
+		const FVec3 ColorAsVec = { ColorByCount.R, ColorByCount.G, ColorByCount.B };
+		
+		InInterface.Box(Bounds, ColorAsVec, InThickness);
+		for (const auto& Elem : Elems)
+		{
+			InInterface.Line(Bounds.Center(), Elem.Bounds.Center(), ColorAsVec, InThickness);
+			InInterface.Box(Elem.Bounds, { (T)1.0, (T)0.2, (T)0.2 }, 1.0f);
+		}
+	}
+#endif
+
 	void Serialize(FChaosArchive& Ar)
 	{
 		Ar << Elems;
@@ -322,6 +351,30 @@ struct TAABBTreeNode
 	int32 ChildrenNodes[2] = { 0, 0 };
 	bool bLeaf = false;
 
+#if !UE_BUILD_SHIPPING
+	void DebugDraw(ISpacialDebugDrawInterface<T>& InInterface, const TArray<TAABBTreeNode<T>>& Nodes, const FVec3& InLinearColor, float InThickness) const
+	{
+		constexpr float ColorRatio = 0.75f;
+		constexpr float LineThicknessRatio = 0.75f;
+		if (!bLeaf)
+		{
+			FLinearColor ChildColor = FLinearColor::MakeRandomColor();
+			for (int ChildIndex = 0; ChildIndex < 2; ++ChildIndex)
+			{
+				int32 NodeIndex = ChildrenNodes[ChildIndex];
+				if (NodeIndex > 0 && NodeIndex < Nodes.Num())
+				{
+					Nodes[NodeIndex].DebugDraw(InInterface, Nodes, { ChildColor.R, ChildColor.G, ChildColor.B }, InThickness * LineThicknessRatio);
+				}
+			}
+			for (int ChildIndex = 0; ChildIndex < 2; ++ChildIndex)
+			{
+				InInterface.Box(ChildrenBounds[ChildIndex], InLinearColor, InThickness);
+			}
+		}
+	}
+#endif
+
 	void Serialize(FChaosArchive& Ar)
 	{
 		for (auto& Bounds : ChildrenBounds)
@@ -413,6 +466,7 @@ public:
 		, MaxTreeDepth(DefaultMaxTreeDepth)
 		, MaxPayloadBounds(DefaultMaxPayloadBounds)
 		, MaxNumToProcess(DefaultMaxNumToProcess)
+		, bShouldRebuild(true)
 	{
 		GetCVars();
 	}
@@ -433,6 +487,9 @@ public:
 		WorkStack.Reset();
 		WorkPoolFreeList.Reset();
 		WorkPool.Reset();
+
+		bShouldRebuild = true;
+		this->SetAsyncTimeSlicingComplete(true);
 	}
 
 	virtual void ProgressAsyncTimeSlicing(bool ForceBuildCompletion) override
@@ -459,6 +516,7 @@ public:
 		, MaxTreeDepth(InMaxTreeDepth)
 		, MaxPayloadBounds(InMaxPayloadBounds)
 		, MaxNumToProcess(InMaxNumToProcess)
+		, bShouldRebuild(true)
 
 	{
 		GenerateTree(Particles);
@@ -467,6 +525,7 @@ public:
 	template <typename ParticleView>
 	void Reinitialize(const ParticleView& Particles)
 	{
+		bShouldRebuild = true;
 		GenerateTree(Particles);
 	}
 
@@ -800,6 +859,7 @@ public:
 				}
 
 				PayloadToInfo.Remove(Payload);
+				bShouldRebuild = true;
 			}
 		}
 	}
@@ -814,7 +874,7 @@ public:
 			{
 				if (PayloadInfo->LeafIdx != INDEX_NONE)
 				{
-					//If we are still within the same leaf bounds, do nothing
+					//If we are still within the same leaf bounds, do nothing, don't detect a change either
 					if (bHasBounds) 
 					{
 						const TAABB<T, 3>& LeafBounds = Leaves[PayloadInfo->LeafIdx].GetBounds();
@@ -835,6 +895,8 @@ public:
 				PayloadInfo = &PayloadToInfo.Add(Payload);
 			}
 
+			bShouldRebuild = true;
+
 			bool bTooBig = false;
 			if (bHasBounds)
 			{
@@ -942,6 +1004,93 @@ public:
 		return GlobalPayloads;
 	}
 
+
+	virtual bool ShouldRebuild() override { return bShouldRebuild; }  // Used to find out if something changed since last reset for optimizations
+	// Contract: bShouldRebuild can only ever be cleared by calling the ClearShouldRebuild method, it can be set at will though
+	virtual void ClearShouldRebuild() override { bShouldRebuild = false; }
+
+	virtual void PrepareCopyTimeSliced(const  ISpatialAcceleration<TPayloadType, T, 3>& InFrom) override
+	{
+
+		check(this != &InFrom);
+		check(InFrom.GetType() == ESpatialAcceleration::AABBTree);
+		const TAABBTree<TPayloadType, TLeafType, bMutable, T>& From = static_cast<const TAABBTree<TPayloadType, TLeafType, bMutable, T>&>(InFrom);
+
+		Reset();
+
+		// Copy all the small objects first
+
+		ISpatialAcceleration<TPayloadType, T, 3>::operator=(From);
+
+		DirtyElementGridCellSize = From.DirtyElementGridCellSize;
+		DirtyElementGridCellSizeInv = From.DirtyElementGridCellSizeInv;
+		DirtyElementMaxGridCellQueryCount = From.DirtyElementMaxGridCellQueryCount;
+		DirtyElementMaxPhysicalSizeInCells = From.DirtyElementMaxPhysicalSizeInCells;
+		DirtyElementMaxCellCapacity = From.DirtyElementMaxCellCapacity;
+
+		MaxChildrenInLeaf = From.MaxChildrenInLeaf;
+		MaxTreeDepth = From.MaxTreeDepth;
+		MaxPayloadBounds = From.MaxPayloadBounds;
+		MaxNumToProcess = From.MaxNumToProcess;
+		NumProcessedThisSlice = From.NumProcessedThisSlice;
+		bShouldRebuild = From.bShouldRebuild;
+
+		// Reserve sizes for arrays etc
+
+		Nodes.Reserve(From.Nodes.Num());
+		Leaves.Reserve(From.Leaves.Num());
+		DirtyElements.Reserve(From.DirtyElements.Num());
+		CellHashToFlatArray.Reserve(From.CellHashToFlatArray.Num());
+		FlattenedCellArrayOfDirtyIndices.Reserve(From.FlattenedCellArrayOfDirtyIndices.Num());
+		DirtyElementsGridOverflow.Reserve(From.DirtyElementsGridOverflow.Num());
+		GlobalPayloads.Reserve(From.GlobalPayloads.Num());
+		PayloadToInfo.Reserve(From.PayloadToInfo.Num());
+
+		this->SetAsyncTimeSlicingComplete(false);
+	}
+
+	virtual void ProgressCopyTimeSliced(const  ISpatialAcceleration<TPayloadType, T, 3>& InFrom, int MaximumBytesToCopy) override
+	{
+		check(this != &InFrom);
+		check(InFrom.GetType() == ESpatialAcceleration::AABBTree);
+		const TAABBTree<TPayloadType, TLeafType, bMutable, T>& From = static_cast<const TAABBTree<TPayloadType, TLeafType, bMutable, T>&>(InFrom);
+
+		int32 SizeToCopyLeft = MaximumBytesToCopy;
+		check(From.CellHashToFlatArray.Num() == 0); // Partial Copy of TMAPs not implemented, and this should be empty for our current use cases
+
+		if (!ContinueTimeSliceCopy(From.Nodes, Nodes, SizeToCopyLeft))
+		{
+			return;
+		}
+		if (!ContinueTimeSliceCopy(From.Leaves, Leaves, SizeToCopyLeft))
+		{
+			return;
+		}
+		if (!ContinueTimeSliceCopy(From.DirtyElements, DirtyElements, SizeToCopyLeft))
+		{
+			return;
+		}
+
+		if (!ContinueTimeSliceCopy(From.FlattenedCellArrayOfDirtyIndices, FlattenedCellArrayOfDirtyIndices, SizeToCopyLeft))
+		{
+			return;
+		}
+		if (!ContinueTimeSliceCopy(From.DirtyElementsGridOverflow, DirtyElementsGridOverflow, SizeToCopyLeft))
+		{
+			return;
+		}
+		if (!ContinueTimeSliceCopy(From.GlobalPayloads, GlobalPayloads, SizeToCopyLeft))
+		{
+			return;
+		}
+		if (!ContinueTimeSliceCopy(From.PayloadToInfo, PayloadToInfo, SizeToCopyLeft))
+		{
+			return;
+		}
+
+		this->SetAsyncTimeSlicingComplete(true);
+	}
+
 	// Returns true if bounds appear valid. Returns false if extremely large values, contains NaN, or is empty.
 	FORCEINLINE_DEBUGGABLE bool ValidateBounds(const TAABB<T, 3>& Bounds)
 	{
@@ -1019,6 +1168,7 @@ public:
 			// Disable the Grid until it is rebuilt
 			DirtyElementGridCellSize = 0.0f;
 			DirtyElementGridCellSizeInv = 1.0f;
+			bShouldRebuild = true;
 		}
 	}
 
@@ -1048,8 +1198,9 @@ private:
 			Leaf.GatherElements(AllElements);
 		}
 
-		TAABBTree<TPayloadType,TLeafType,bMutable, T> NewTree(AllElements);
+		TAABBTree<TPayloadType,TLeafType,bMutable,T> NewTree(AllElements);
 		*this = NewTree;
+		bShouldRebuild = true; // No changes since last time tree was built
 	}
 
 	// Returns true if the query should continue
@@ -1386,6 +1537,10 @@ private:
 		{
 			SCOPE_CYCLE_COUNTER(STAT_AABBTreeTimeSliceSetup);
 
+			// Prepare to find the average and scaled variance of particle centers.
+			WorkSnapshot.AverageCenter = FVec3(0);
+			WorkSnapshot.ScaledCenterVariance = FVec3(0);
+
 			int32 Idx = 0;
 
 			//TODO: we need a better way to time-slice this case since there can be a huge number of Particles. Can't do it right now without making full copy
@@ -1414,9 +1569,13 @@ private:
 					}
 					else
 					{
-						WorkSnapshot.Elems.Add(FElement{ Payload, ElemBounds });
+						FReal NumElems = (FReal)(WorkSnapshot.Elems.Add(FElement{ Payload, ElemBounds }) + 1);
 						WorkSnapshot.Bounds.GrowToInclude(ElemBounds);
-						CenterSum += ElemBounds.Center();
+
+						// Include the current particle in the average and scaled variance of the particle centers using Welford's method.
+						FVec3 CenterDelta = ElemBounds.Center() - WorkSnapshot.AverageCenter;
+						WorkSnapshot.AverageCenter += CenterDelta / NumElems;
+						WorkSnapshot.ScaledCenterVariance += (ElemBounds.Center() - WorkSnapshot.AverageCenter) * CenterDelta;
 					}
 				}
 				else
@@ -1515,9 +1674,12 @@ private:
 		eTimeSlicePhase TimeslicePhase;
 
 		TAABB<T, 3> Bounds;
-		TVec3<T> AverageCenter;
 		TArray<FElement> Elems;
 
+		// The average of the element centers and their variance times the number of elements.
+		TVec3<T> AverageCenter;
+		TVec3<T> ScaledCenterVariance;
+
 		int32 NodeLevel;
 		int32 NewNodeIdx;
 
@@ -1551,9 +1713,14 @@ private:
 			if (CHAOS_ENSURE(MinBoxIdx != INDEX_NONE))
 			{
 				FSplitInfo& SplitInfo = CurrentSnapshot.SplitInfos[MinBoxIdx];
-				WorkPool[SplitInfo.WorkSnapshotIdx].Elems.Add(Elem);
-				WorkPool[SplitInfo.WorkSnapshotIdx].AverageCenter += Elem.Bounds.Center();
+				FWorkSnapshot& WorkSnapshot = WorkPool[SplitInfo.WorkSnapshotIdx];
+				FReal NumElems = (FReal)(WorkSnapshot.Elems.Add(Elem) + 1);
 				SplitInfo.RealBounds.GrowToInclude(Elem.Bounds);
+
+				// Include the current particle in the average and scaled variance of the particle centers using Welford's method.
+				FVec3 CenterDelta = Elem.Bounds.Center() - WorkSnapshot.AverageCenter;
+				WorkSnapshot.AverageCenter += CenterDelta / NumElems;
+				WorkSnapshot.ScaledCenterVariance += (Elem.Bounds.Center() - WorkSnapshot.AverageCenter) * CenterDelta;
 			}
 		}
 
@@ -1622,8 +1789,11 @@ private:
 
 			if (WorkPool[CurIdx].TimeslicePhase == eTimeSlicePhase::PreFindBestBounds)
 			{
-				const TVec3<T> Extents = WorkPool[CurIdx].Bounds.Extents();
-				const int32 MaxAxis = WorkPool[CurIdx].Bounds.LargestAxis();	//todo: use variance instead
+				// Determine the axis to split the AABB on based on the SplitOnVarianceAxis console variable. If it is not 1, simply use the largest axis
+				// of the work snapshot bounds; otherwise, select the axis with the greatest center variance. Note that the variance times the number of
+				// elements is actually used but since all that is needed is the axis with the greatest variance the scale factor is irrelevant.
+				const int32 MaxAxis = (FAABBTreeCVars::SplitOnVarianceAxis != 1) ? WorkPool[CurIdx].Bounds.LargestAxis() :
+					(WorkPool[CurIdx].ScaledCenterVariance[0] > WorkPool[CurIdx].ScaledCenterVariance[1] ? (WorkPool[CurIdx].ScaledCenterVariance[0] > WorkPool[CurIdx].ScaledCenterVariance[2] ? 0 : 2) : (WorkPool[CurIdx].ScaledCenterVariance[1] > WorkPool[CurIdx].ScaledCenterVariance[2] ? 1 : 2));
 
 				//Add two children, remember this invalidates any pointers to current snapshot
 				const int32 FirstChildIdx = GetNewWorkSnapshot();
@@ -1637,7 +1807,10 @@ private:
 				WorkPool[CurIdx].SplitInfos[0].SplitBounds = TAABB<T, 3>(WorkPool[CurIdx].Bounds.Min(), WorkPool[CurIdx].Bounds.Min());
 				WorkPool[CurIdx].SplitInfos[1].SplitBounds = TAABB<T, 3>(WorkPool[CurIdx].Bounds.Max(), WorkPool[CurIdx].Bounds.Max());
 
-				const TVec3<T> Center = WorkPool[CurIdx].AverageCenter;
+				// Find the point where the AABB will be split based on the SplitAtAverageCenter console variable. If it is not 1, just use the center
+				// of the AABB; otherwise, use the average of the element centers.
+				const TVec3<T> Center = (FAABBTreeCVars::SplitAtAverageCenter != 1) ? WorkPool[CurIdx].Bounds.Center() : WorkPool[CurIdx].AverageCenter;
+
 				for (FSplitInfo& SplitInfo : WorkPool[CurIdx].SplitInfos)
 				{
 					SplitInfo.RealBounds = TAABB<T, 3>::EmptyAABB();
@@ -1665,8 +1838,11 @@ private:
 					WorkPool[FirstChildIdx].Elems.Reserve(ExpectedNumPerChild);
 					WorkPool[SecondChildIdx].Elems.Reserve(ExpectedNumPerChild);
 
+					// Initialize the the two info's element center average and scaled variance.
 					WorkPool[FirstChildIdx].AverageCenter = TVec3<T>(0);
+					WorkPool[FirstChildIdx].ScaledCenterVariance = TVec3<T>(0);
 					WorkPool[SecondChildIdx].AverageCenter = TVec3<T>(0);
+					WorkPool[SecondChildIdx].ScaledCenterVariance = TVec3<T>(0);
 				}
 			}
 
@@ -1770,6 +1946,84 @@ private:
 		return Results;
 	}
 
+	// Set InOutMaxSize to less than 0 for unlimited
+	// 
+	template<typename ContainerType>
+	static void AddToContainerHelper(const ContainerType& ContainerFrom, ContainerType& ContainerTo, int32 Index)
+	{
+		ContainerTo.Add(ContainerFrom[Index]);
+	}
+
+	template<>
+	static void AddToContainerHelper(const TArrayAsMap<TPayloadType, FAABBTreePayloadInfo>& ContainerFrom, TArrayAsMap<TPayloadType, FAABBTreePayloadInfo>& ContainerTo, int32 Index)
+	{
+		ContainerTo.AddFrom(ContainerFrom, Index);
+	}
+
+	template<typename ContainerType>
+	static int32 ContainerElementSizeHelper(const ContainerType& Container, int32 Index)
+	{
+		return sizeof(typename ContainerType::ElementType);
+	}
+
+	template<>
+	static int32 ContainerElementSizeHelper(const TArray<TAABBTreeLeafArray<TPayloadType, true>>& Container, int32 Index)
+	{
+		return sizeof(typename TArray<TAABBTreeLeafArray<TPayloadType, true>>::ElementType) + sizeof(typename decltype(Container[Index].Elems)::ElementType) * Container[Index].GetElementCount();
+	}
+
+	template<>
+	static int32 ContainerElementSizeHelper(const TArray<TAABBTreeLeafArray<TPayloadType, false>>& Container, int32 Index)
+	{
+		return sizeof(typename TArray<TAABBTreeLeafArray<TPayloadType, false>>::ElementType) + sizeof(typename decltype(Container[Index].Elems)::ElementType) * Container[Index].GetElementCount();
+	}
+
+	template<typename ContainerType>
+	static bool ContinueTimeSliceCopy(const ContainerType& ContainerFrom, ContainerType& ContainerTo, int32& InOutMaxSize)
+	{
+		int32 SizeCopied = 0;
+
+		for (int32 Index = ContainerTo.Num(); Index < ContainerFrom.Num() && (InOutMaxSize < 0 || SizeCopied < InOutMaxSize); Index++)
+		{
+			AddToContainerHelper(ContainerFrom, ContainerTo, Index);
+			SizeCopied += ContainerElementSizeHelper(ContainerFrom, Index);
+		}
+
+		// Update the maximum size left
+		if (InOutMaxSize > 0)
+		{
+			if (SizeCopied > InOutMaxSize)
+			{
+				InOutMaxSize = 0;
+			}
+			else
+			{
+				InOutMaxSize -= SizeCopied;
+			}
+		}
+
+		bool Done = ContainerTo.Num() == ContainerFrom.Num();
+		return Done;
+	}
+
+
+#if !UE_BUILD_SHIPPING
+	virtual void DebugDraw(ISpacialDebugDrawInterface<T>* InInterface) const override
+	{
+		if (InInterface)
+		{
+			if (Nodes.Num() > 0)
+			{
+				Nodes[0].DebugDraw(*InInterface, Nodes, { 1.f, 1.f, 1.f }, 5.f);
+			}
+			for (const TLeafType& Leaf : Leaves)
+			{
+				Leaf.DebugDrawLeaf(*InInterface, FLinearColor::MakeRandomColor(), 10.f);
+			}
+		}
+	}
+#endif
+
 
 	TAABBTree(const TAABBTree<TPayloadType, TLeafType, bMutable, T>& Other)
 		: ISpatialAcceleration<TPayloadType, T, 3>(StaticType)
@@ -1792,12 +2046,21 @@ private:
 		, MaxPayloadBounds(Other.MaxPayloadBounds)
 		, MaxNumToProcess(Other.MaxNumToProcess)
 		, NumProcessedThisSlice(Other.NumProcessedThisSlice)
+		, bShouldRebuild(Other.bShouldRebuild)
 	{
 
 	}
 
+	virtual ISpatialAcceleration<TPayloadType, T, 3>& operator=(const ISpatialAcceleration<TPayloadType, T, 3>& Other) override
+	{
+		
+		check(Other.GetType() == ESpatialAcceleration::AABBTree);
+		return operator=(static_cast<const TAABBTree<TPayloadType, TLeafType, bMutable, T>&>(Other));
+	}
+
 	TAABBTree<TPayloadType,TLeafType, bMutable, T>& operator=(const TAABBTree<TPayloadType,TLeafType,bMutable, T>& Rhs)
 	{
+		ISpatialAcceleration<TPayloadType, T, 3>::operator=(Rhs);
 		ensure(Rhs.WorkStack.Num() == 0);
 		//ensure(Rhs.WorkPool.Num() == 0);
 		//ensure(Rhs.WorkPoolFreeList.Num() == 0);
@@ -1828,6 +2091,7 @@ private:
 			MaxPayloadBounds = Rhs.MaxPayloadBounds;
 			MaxNumToProcess = Rhs.MaxNumToProcess;
 			NumProcessedThisSlice = Rhs.NumProcessedThisSlice;
+			bShouldRebuild = Rhs.bShouldRebuild;
 		}
 
 		return *this;
@@ -1862,6 +2126,8 @@ private:
 	TArray<int32> WorkStack;
 	TArray<int32> WorkPoolFreeList;
 	TArray<FWorkSnapshot> WorkPool;
+
+	bool bShouldRebuild;  // Contract: this can only ever be cleared by calling the ClearShouldRebuild method
 };
 
 template<typename TPayloadType, typename TLeafType, bool bMutable, typename T>
diff --git a/Engine/Source/Runtime/Experimental/Chaos/Public/Chaos/BoundingVolume.h b/Engine/Source/Runtime/Experimental/Chaos/Public/Chaos/BoundingVolume.h
--- a/Engine/Source/Runtime/Experimental/Chaos/Public/Chaos/BoundingVolume.h
+++ b/Engine/Source/Runtime/Experimental/Chaos/Public/Chaos/BoundingVolume.h
@@ -163,7 +163,27 @@ class TBoundingVolume final : public ISpatialAcceleration<InPayloadType, T,d>
 	}
 
 public:
-	TBoundingVolume<TPayloadType, T, d>& operator=(const TBoundingVolume<TPayloadType, T, d>& Other) = delete;
+
+	virtual ISpatialAcceleration<TPayloadType, T, d>& operator=(const ISpatialAcceleration<TPayloadType, T, d>& Other) override
+	{
+
+		check(Other.GetType() == ESpatialAcceleration::BoundingVolume);
+		return operator=(static_cast<const TBoundingVolume<TPayloadType, T, d>&>(Other));
+	}
+
+	TBoundingVolume<TPayloadType, T, d>& operator=(const TBoundingVolume<TPayloadType, T, d>& Other)
+	{
+		ISpatialAcceleration<TPayloadType, FReal, 3>::operator=(Other);
+		MGlobalPayloads = Other.MGlobalPayloads;
+		MGrid = Other.MGrid;
+		MElements = Other.MElements;
+		MDirtyElements = Other.MDirtyElements;
+		MPayloadInfo = Other.MPayloadInfo;
+		MaxPayloadBounds = Other.MaxPayloadBounds;
+		bIsEmpty = Other.bIsEmpty;
+		return *this;
+	}
+
 	TBoundingVolume<TPayloadType, T, d>& operator=(TBoundingVolume<TPayloadType, T, d>&& Other)
 	{
 		MGlobalPayloads = MoveTemp(Other.MGlobalPayloads);
@@ -263,7 +283,7 @@ public:
 	}
 
 	// Begin ISpatialAcceleration interface
-	virtual TArray<TPayloadType> FindAllIntersections(const TAABB<T, d>& Box) const override { return FindAllIntersectionsImp(Box); }
+	virtual TArray<TPayloadType> FindAllIntersections(const FAABB3& Box) const override { return FindAllIntersectionsImp(Box); }
 
 	const TArray<TPayloadBoundsElement<TPayloadType, T>>& GlobalObjects() const
 	{
diff --git a/Engine/Source/Runtime/Experimental/Chaos/Public/Chaos/BoundingVolumeHierarchy.h b/Engine/Source/Runtime/Experimental/Chaos/Public/Chaos/BoundingVolumeHierarchy.h
--- a/Engine/Source/Runtime/Experimental/Chaos/Public/Chaos/BoundingVolumeHierarchy.h
+++ b/Engine/Source/Runtime/Experimental/Chaos/Public/Chaos/BoundingVolumeHierarchy.h
@@ -79,6 +79,25 @@ class TBoundingVolumeHierarchy final : public ISpatialAcceleration<int32, T,d>
 	{
 	}
 
+	CHAOS_API virtual ISpatialAcceleration<int32, T, d>& operator=(const ISpatialAcceleration<int32, T, d>& Other) override
+	{
+
+		check(Other.GetType() == ESpatialAcceleration::AABBTreeBV);
+		return operator=(static_cast<const TBoundingVolumeHierarchy<OBJECT_ARRAY, LEAF_TYPE, T, d>&>(Other));
+	}
+
+	CHAOS_API TBoundingVolumeHierarchy& operator=(const TBoundingVolumeHierarchy<OBJECT_ARRAY, LEAF_TYPE, T, d>& Other)
+	{
+		ISpatialAcceleration<int32, T, d>::operator=(Other);
+		MObjects = Other.MObjects;
+		MGlobalObjects = Other.MGlobalObjects;
+		MWorldSpaceBoxes = Other.MWorldSpaceBoxes;
+		MMaxLevels = Other.MMaxLevels;
+		Elements = Other.Elements;
+		Leafs = Other.Leafs;
+		return *this;
+	}
+
 	CHAOS_API TBoundingVolumeHierarchy& operator=(TBoundingVolumeHierarchy<OBJECT_ARRAY, LEAF_TYPE, T, d>&& Other)
 	{
 		MObjects = Other.MObjects;
@@ -119,7 +138,7 @@ class TBoundingVolumeHierarchy final : public ISpatialAcceleration<int32, T,d>
 	}
 
 	// Begin ISpatialAcceleration interface
-	CHAOS_API TArray<int32> FindAllIntersections(const TAABB<T, d>& Box) const { return FindAllIntersectionsImp(Box); }
+	CHAOS_API TArray<int32> FindAllIntersections(const FAABB3& Box) const { return FindAllIntersectionsImp(Box); }
 	CHAOS_API TArray<int32> FindAllIntersections(const TSpatialRay<T,d>& Ray) const { return FindAllIntersectionsImp(Ray); }
 	CHAOS_API TArray<int32> FindAllIntersections(const TVector<T, d>& Point) const { return FindAllIntersectionsImp(Point); }
 	CHAOS_API TArray<int32> FindAllIntersections(const TGeometryParticles<T, d>& InParticles, const int32 i) const;
diff --git a/Engine/Source/Runtime/Experimental/Chaos/Public/Chaos/ChaosDebugDraw.h b/Engine/Source/Runtime/Experimental/Chaos/Public/Chaos/ChaosDebugDraw.h
--- a/Engine/Source/Runtime/Experimental/Chaos/Public/Chaos/ChaosDebugDraw.h
+++ b/Engine/Source/Runtime/Experimental/Chaos/Public/Chaos/ChaosDebugDraw.h
@@ -9,6 +9,9 @@
 
 namespace Chaos
 {
+	class FAccelerationStructureHandle;
+	template <typename PayloadType, typename T, int d> class ISpatialAcceleration;
+
 	namespace DebugDraw
 	{
 #if CHAOS_DEBUG_DRAW
@@ -53,8 +56,6 @@ namespace Chaos
 		struct CHAOS_API FChaosDebugDrawSettings
 		{
 		public:
-			FChaosDebugDrawSettings();
-
 			FChaosDebugDrawSettings(
 				FRealSingle InArrowSize,
 				FRealSingle InBodyAxisLen,
@@ -73,7 +74,8 @@ namespace Chaos
 				FRealSingle InVelScale,
 				FRealSingle InAngVelScale,
 				FRealSingle InImpulseScale,
-				float InInertiaScale,
+				FRealSingle InPushOutScale,
+				FRealSingle InInertiaScale,
 				uint8 InDrawPriority,
 				bool bInShowSimpleCollision,
 				bool bInShowComplexCollision,
@@ -100,6 +102,7 @@ namespace Chaos
 				, VelScale(InVelScale)
 				, AngVelScale(InAngVelScale)
 				, ImpulseScale(InImpulseScale)
+				, PushOutScale(InPushOutScale)
 				, InertiaScale(InInertiaScale)
 				, DrawPriority(InDrawPriority)
 				, bShowSimpleCollision(bInShowSimpleCollision)
@@ -129,6 +132,7 @@ namespace Chaos
 			FRealSingle VelScale;
 			FRealSingle AngVelScale;
 			FRealSingle ImpulseScale;
+			FRealSingle PushOutScale;
 			FRealSingle InertiaScale;
 			uint8 DrawPriority;
 			bool bShowSimpleCollision;
@@ -202,8 +206,9 @@ namespace Chaos
 		CHAOS_API void DrawJointConstraints(const FRigidTransform3& SpaceTransform, const FPBDJointConstraints& Constraints, FRealSingle ColorScale, const FChaosDebugDrawJointFeatures& FeatureMask = FChaosDebugDrawJointFeatures::MakeDefault(), const FChaosDebugDrawSettings* Settings = nullptr);
 		CHAOS_API void DrawSimulationSpace(const FSimulationSpace& SimSpace, const FChaosDebugDrawSettings* Settings = nullptr);
 		CHAOS_API void DrawShape(const FRigidTransform3& ShapeTransform, const FImplicitObject* Shape, const FColor& Color, const FChaosDebugDrawSettings* Settings = nullptr);
-		CHAOS_API void DrawConstraintGraph(const FRigidTransform3& ShapeTransform, const FPBDConstraintColor& Graph, const FChaosDebugDrawSettings* Settings = nullptr);
+		CHAOS_API void DrawConstraintGraph(const FRigidTransform3& ShapeTransform, const FPBDConstraintGraph& Graph, const FChaosDebugDrawSettings* Settings = nullptr);
 		CHAOS_API void DrawCollidingShapes(const FRigidTransform3& SpaceTransform, const FPBDCollisionConstraints& Collisions, FRealSingle ColorScale, const FChaosDebugDrawSettings* Settings = nullptr);
+		CHAOS_API void DrawSpatialAccelerationStructure(const ISpatialAcceleration<FAccelerationStructureHandle, FReal, 3>& SpatialAccelerationStructure, const FChaosDebugDrawSettings* Settings = nullptr);
 #endif
 	}
 }
diff --git a/Engine/Source/Runtime/Experimental/Chaos/Public/Chaos/Collision/CollisionConstraintAllocator.h b/Engine/Source/Runtime/Experimental/Chaos/Public/Chaos/Collision/CollisionConstraintAllocator.h
new file mode 100644
--- /dev/null
+++ b/Engine/Source/Runtime/Experimental/Chaos/Public/Chaos/Collision/CollisionConstraintAllocator.h
@@ -0,0 +1,657 @@
+// Copyright Epic Games, Inc. All Rights Reserved.
+#pragma once
+
+#include "CoreMinimal.h"
+
+#include "Chaos/Core.h"
+#include "Chaos/CollisionResolutionTypes.h"
+#include "Chaos/Collision/PBDCollisionConstraint.h"
+
+#include "ChaosStats.h"
+
+extern bool bChaos_Collision_EnableManifoldRestore;
+
+namespace Chaos
+{
+	class FParticlePairCollisionsKey
+	{
+	public:
+		FParticlePairCollisionsKey()
+			: Key(0)
+		{
+		}
+
+		FParticlePairCollisionsKey(const FGeometryParticleHandle* Particle0, const FGeometryParticleHandle* Particle1)
+			: Key(GenerateHash(Particle0, Particle1))
+		{
+		}
+
+		uint32 GetKey() const
+		{
+			return Key;
+		}
+
+	private:
+		uint32 GenerateHash(const FGeometryParticleHandle* Particle0, const FGeometryParticleHandle* Particle1)
+		{
+			// @todo(chaos): Particle handles can be reused - this needs to use a global ID or something (not available in this stream)
+			return HashCombine(::GetTypeHash(Particle0), ::GetTypeHash(Particle1));
+		}
+
+		uint32 Key;
+	};
+
+	/**
+	 * @brief All the contacts on a particle pair
+	*/
+	class FParticlePairCollisionConstraints
+	{
+	public:
+		FParticlePairCollisionConstraints(const FParticlePairCollisionsKey& InKey, const FGeometryParticleHandle* InParticle0, const FGeometryParticleHandle* InParticle1, const int32 InEpoch)
+			: Key(InKey)
+			, Particle0(InParticle0)
+			, Particle1(InParticle1)
+			, CreationEpoch(InEpoch)
+			, Epoch(InEpoch)
+		{
+		}
+		
+		FParticlePairCollisionsKey GetKey() const
+		{
+			return Key;
+		}
+
+		/**
+		 * @brief Whether the particle pair was just created
+		*/
+		bool IsNew() const
+		{
+			return CreationEpoch == Epoch;
+		}
+
+		/**
+		 * @brief The transform of Particle0 when the contacts were created
+		*/
+		const FRigidTransform3& GetParticle0Transform() const
+		{
+			return ParticleTransforms[0];
+		}
+
+		/**
+		 * @brief The transform of Particle1 when the contacts were created
+		*/
+		const FRigidTransform3& GetParticle1Transform() const
+		{
+			return ParticleTransforms[1];
+		}
+
+		/**
+		 * @brief Clear the constraint list and update the transforms (subsequently added contacts will have been created at these transforms)
+		*/
+		void Reset(int32 InEpoch, const FRigidTransform3& InParticleTransform0, const FRigidTransform3& InParticleTransform1)
+		{
+			CreationEpoch = Epoch;
+			Epoch = InEpoch;
+			ParticleTransforms[0] =  InParticleTransform0;
+			ParticleTransforms[1] = InParticleTransform1;
+			Constraints.Reset();
+		}
+
+		/**
+		 * @brief Get the epoch counter used for pruning unused pairs
+		*/
+		int32 GetEpoch() const
+		{
+			return Epoch;
+		}
+
+		/**
+		 * @brief Update the epoch counter used for pruning unused pairs
+		*/
+		void SetEpoch(const int32 InEpoch)
+		{
+			Epoch = InEpoch;
+		}
+
+		/**
+		 * @brief Whether we can add mroe contacts for this particle pair
+		 * This is to limit contact creation when two partiicles have a large number of overlapping shape pairs
+		 * @todo(chaos): implement particle pair contact count limit
+		*/
+		bool IsAcceptingContacts() const
+		{
+			return true;
+		}
+
+		/**
+		 * @brief The list of contacts between the particle pair
+		*/
+		TArrayView<FPBDCollisionConstraint* const> GetConstraints() const
+		{
+			return MakeArrayView(Constraints);
+		}
+
+		/**
+		 * @brief Add a constraint to the particle-pair constraint list
+		*/
+		void AddConstraint(FPBDCollisionConstraint* InContact)
+		{
+			Constraints.Add(InContact);
+		}
+
+		/**
+		 * @brief Remove a contact for the particle-pair
+		*/
+		void RemoveConstraint(FPBDCollisionConstraint* InContact)
+		{
+			Constraints.Remove(InContact);
+		}
+
+		/**
+		 * @brief Find the contact for the shape pair
+		 * @note O(N) search
+		*/
+		FPBDCollisionConstraint* FindContact(const FRigidBodyContactKey& ConstraintKey)
+		{
+			for (FPBDCollisionConstraint* Constraint : Constraints)
+			{
+				if (Constraint->GetKey() == ConstraintKey)
+				{
+					return Constraint;
+				}
+			}
+			return nullptr;
+		}
+
+	private:
+		FParticlePairCollisionsKey Key;
+		const FGeometryParticleHandle* Particle0;
+		const FGeometryParticleHandle* Particle1;
+
+		// The transforms of the particles when the contacts were created
+		FRigidTransform3 ParticleTransforms[2];
+
+		// The contacts between the particle pairs
+		TArray<FPBDCollisionConstraint*, TInlineAllocator<1>> Constraints;
+
+		// The tick counter when the particle pair was first created
+		int32 CreationEpoch;
+
+		// The tick counter when the contacts were last used (for pruning)
+		int32 Epoch;
+	};
+
+
+	/**
+	 * @brief An allocator and container of collision constraints that supports reuse of constraints from the previous tick
+	 * 
+	 * This allocator maintains the following structures
+	 * - the set of all allocated constraints
+	 * - the set of all active constraints (created or restored in the current tick)
+	 * - a mapping from shape pair to active constraints
+	 * - a mapping from particle pair to a list of active constraints
+	 * 
+	 * Constraints are allocated during the collision detection phase and retained between ticks. An attempt ot create
+	 * a constraint for the same shape pair as seen on the previous tick will return the existing collision constraint, with
+	 * all of its data intact.
+	 *
+	 * Lists are pruned at the end of each tick so if particles are destroyed or a particle pair is no longer overlapping
+	 * the entries will be removed. However, this means that the allocator may briefly contain contacts from a prior
+	 * tick for particles that have been destroyed. For this reason, no contacts should be accessed unless they
+	 * have been created/restored on the current tick (i.e., the Epoch counter is up to date).
+	 * 
+	 * All constraint pointers returned are persistent in memory until Destroy() is called, or until they are pruned.
+	 * 
+	 * @todo(chaos): optimizations. We can probably remove the shape pair map, if it is cheap enough to find the shape pair
+	 * contact in the particle pair contact list.
+	 * 
+	 * @todo(chaos): handle shape or particle pairs being in the opposite order - should be treated the same here
+	*/
+	class CHAOS_API FCollisionConstraintAllocator
+	{
+	public:
+		FCollisionConstraintAllocator()
+			: ShapePairConstraintMap()
+			, Constraints()
+			, Epoch(0)
+			, bInCollisionDetectionPhase(false)
+		{
+		}
+		
+		
+		/**
+		 * @brief The set of collision constraints for the current tick (created or reinstated)
+		 * 
+		 * @note Some elements may be null (constraints that have been deleted)
+		 * @note This is not thread-safe and should not be used during the collision detection phase (i.e., when the list is being built)
+		*/
+		TArrayView<FPBDCollisionConstraint* const> GetConstraints() const
+		{ 
+			return MakeArrayView(Constraints);
+		}
+
+		/**
+		 * @brief The set of collision constraints for the current tick (created or reinstated)
+		 * 
+		 * @note Some elements may be null (constraints that have been deleted)
+		 * @note This is not thread-safe and should not be used during the collision detection phase (i.e., when the list is being built)
+		 */
+		TArrayView<const FPBDCollisionConstraint* const> GetConstConstraints() const
+		{
+			return MakeArrayView(Constraints);
+		}
+
+		/**
+		 * @brief Destroy all constraints
+		*/
+		void Reset()
+		{
+			Constraints.Reset();
+
+			for (auto& KeyValuePair : ShapePairConstraintMap)
+			{
+				FreeConstraint(KeyValuePair.Value);
+			}
+
+			ShapePairConstraintMap.Reset();
+
+			for (auto& KeyValuePair : ParticlePairConstraintMap)
+			{
+				delete KeyValuePair.Value;
+			}
+
+			ParticlePairConstraintMap.Reset();
+		}
+
+		/**
+		 * @brief Called at the start of the frame to clear the frame's active collision list.
+		 * @todo(chaos): This is only required because of the way events work (see AdvanceOneTimeStepTask::DoWork)
+		*/
+		void BeginFrame()
+		{
+			Constraints.Reset();
+		}
+
+		/**
+		 * @brief Called at the start of the tick to prepare for collision detection.
+		 * Resets the list of active contacts.
+		*/
+		void BeginDetectCollisions()
+		{
+			// Clear the collision list for this tick - we are about to rebuild it
+			Constraints.Reset();
+
+			// Note: we do not reset the shape-pair or particle-pair maps. They may be used to reinstate or reuse
+			// contacts. If they do not (because the pair is no longer overlapping) they will be pruned at the end of the tick.
+
+			bInCollisionDetectionPhase = true;
+		}
+
+		/**
+		 * @brief Called after collision detection to clean up
+		 * Prunes unused contacts
+		*/
+		void EndDetectCollisions()
+		{
+			bInCollisionDetectionPhase = false;
+
+			ProcessNewConstraints();
+
+			PruneShapePairs();
+
+			PruneParticlePairs();
+
+			// Update the tick counter
+			++Epoch;
+		}
+
+		/**
+		 * @brief If we add new constraints after collision detection, do what needs to be done to add them to the system
+		*/
+		void ProcessInjectedConstraints()
+		{
+			ProcessNewConstraints();
+		}
+
+		/**
+		 * @brief Create a constraint for the specified particle+shape pair, or retrieve the one from the last tick if it exists
+		 * This is called by the mid phase, once we have determined that two shapes may collide.
+		*/
+		FPBDCollisionConstraint* FindOrCreateConstraint(FGeometryParticleHandle* Particle0, const FImplicitObject* Implicit0, const FBVHParticles* Simplicial0, FGeometryParticleHandle* Particle1, const FImplicitObject* Implicit1, const FBVHParticles* Simplicial1) const
+		{
+			// See if we already have a contact for these shapes
+			const FRigidBodyContactKey Key = FRigidBodyContactKey(Particle0, Implicit0, Simplicial0, Particle1, Implicit1, Simplicial1);
+			FPBDCollisionConstraint* Constraint = FindConstraint(Key);
+			
+			// If we did not find the contact, create and register a new one
+			if (Constraint == nullptr)
+			{
+				Constraint = CreateConstraint(Key, Particle0, Implicit0, Simplicial0, Particle1, Implicit1, Simplicial1);
+			}
+
+			// Enqueue constraint to add it to maps etc
+			EnqueueNewConstraint(Constraint);
+
+			return Constraint;
+		}
+
+		/**
+		 * @brief Add a set of pre-built constraints and build required internal mapping data
+		 * This is used by the resim cache when restoring constraints after a desync
+		 * @note This is called after the main collision detection phase, so it directly performs all the work
+		 * that would normally be queued up with a call to EnqueueCollision.
+		*/
+		void AddResimConstraints(const TArray<FPBDCollisionConstraint>& InConstraints)
+		{
+			for (const FPBDCollisionConstraint& SourceConstraint : InConstraints)
+			{
+				FPBDCollisionConstraint* Constraint = FindConstraint(SourceConstraint.GetContainerCookie().Key);
+				if (Constraint != nullptr)
+				{
+					// Copy the constraint over the existing one, taking care to leave the cookie intact (which contains the array index)
+					const FPBDCollisionConstraintContainerCookie Cookie = Constraint->GetContainerCookie();
+
+					*Constraint = SourceConstraint;
+
+					Constraint->GetContainerCookie() = Cookie;
+				}
+				else
+				{
+					// Create a new constraint and add it to the array
+					Constraint = CreateConstraint(SourceConstraint);
+
+					if (Constraint != nullptr)
+					{
+						const int32 ConstraintIndex = Constraints.Add(Constraint);
+						Constraint->GetContainerCookie().Update(ConstraintIndex, Epoch);
+					}
+				}
+
+				// Keep track of all contacts between a particle pair (for fast restore)
+				FParticlePairCollisionConstraints* ParticlePairConstraints = FindOrCreateParticlePairConstraints(SourceConstraint.Particle[0], SourceConstraint.Particle[1]);
+				if (ParticlePairConstraints != nullptr)
+				{
+					ParticlePairConstraints->AddConstraint(Constraint);
+				}
+			}
+		}
+
+		/**
+		 * @brief Destroy the specified constraint
+		*/
+		void DestroyConstraint(FPBDCollisionConstraint* Constraint)
+		{
+			DestroyConstraint(Constraint->GetContainerCookie().Key, Constraint->GetContainerCookie().ConstraintIndex, Constraint);
+		}
+
+
+		/**
+		 * @brief Return the set of contact constraints between the specified particle pair, or null if we have not created any yet.
+		*/
+		FParticlePairCollisionConstraints* FindParticlePairConstraints(const FGeometryParticleHandle* Particle0, const FGeometryParticleHandle* Particle1)
+		{
+			const FParticlePairCollisionsKey ParticlePairKey = FParticlePairCollisionsKey(Particle0, Particle1);
+			return FindParticlePairConstraints(ParticlePairKey);
+		}
+
+		/**
+		 * @brief Reinstate the contact constraints for the particle pair
+		 * This will only be called when the particle have not moved relative to each other.
+		 * @note the actual restoration is deferred until after collision detection to prevent lock contention
+		*/
+		void RestoreParticlePairConstraints(FParticlePairCollisionConstraints* ParticlePairConstraints)
+		{
+			if (ParticlePairConstraints != nullptr)
+			{
+				ParticlePairConstraints->SetEpoch(Epoch);
+
+				for (FPBDCollisionConstraint* Constraint : ParticlePairConstraints->GetConstraints())
+				{
+					// Restore the manifold
+					Constraint->RestoreManifold();
+				}
+
+				// Enqueue for adding to active array etc
+				EnqueueNewConstraints(ParticlePairConstraints->GetConstraints());
+			}
+		}
+
+		/**
+		 * @brief Clear the particle pair contact list
+		 * This will be called when the particle pair are still close enough to possible collide but have
+		 * moved beyond some threshold so we will be recreating the contacts this tick
+		*/
+		void ResetParticlePairConstraints(FParticlePairCollisionConstraints* ParticlePairConstraints, const FRigidTransform3& ParticleWorldTransform0, const FRigidTransform3& ParticleWorldTransform1)
+		{
+			if (ParticlePairConstraints != nullptr)
+			{
+				ParticlePairConstraints->Reset(Epoch, ParticleWorldTransform0, ParticleWorldTransform1);
+			}
+		}
+
+	private:
+		FPBDCollisionConstraint* FindConstraint(FRigidBodyContactKey Key) const
+		{
+			FPBDCollisionConstraint* const * ItemPtr = ShapePairConstraintMap.Find(Key.GetKey());
+			if (ItemPtr != nullptr)
+			{
+				return *ItemPtr;
+			}
+			return nullptr;
+		}
+
+		FPBDCollisionConstraint* CreateConstraint(FRigidBodyContactKey Key, FGeometryParticleHandle* Particle0, const FImplicitObject* Implicit0, const FBVHParticles* Simplicial0, FGeometryParticleHandle* Particle1, const FImplicitObject* Implicit1, const FBVHParticles* Simplicial1) const
+		{
+			// TEMP: until we get the block allocator
+			FPBDCollisionConstraint* Constraint = new FPBDCollisionConstraint(Particle0, Implicit0, Simplicial0, Particle1, Implicit1, Simplicial1);
+
+			if (Constraint != nullptr)
+			{
+				// Initalize the constants used by the container to manage the constraint
+				Constraint->GetContainerCookie().Init(Key, Epoch);
+			}
+			return Constraint;
+		}
+
+		FPBDCollisionConstraint* CreateConstraint(const FPBDCollisionConstraint& Other)
+		{
+			// TEMP: until we get the block allocator
+			FPBDCollisionConstraint* Constraint = new FPBDCollisionConstraint(Other);
+			return Constraint;
+		}
+
+		void DestroyConstraint(const FRigidBodyContactKey& Key, const int32 Index, FPBDCollisionConstraint* Constraint)
+		{
+			FParticlePairCollisionConstraints** ParticlePairConstraintsPtr = ParticlePairConstraintMap.Find(FParticlePairCollisionsKey(Constraint->Particle[0], Constraint->Particle[1]).GetKey());
+			if (ParticlePairConstraintsPtr != nullptr)
+			{
+				FParticlePairCollisionConstraints* ParticlePairConstraints = *ParticlePairConstraintsPtr;
+				ParticlePairConstraints->RemoveConstraint(Constraint);
+			}
+
+			ShapePairConstraintMap.Remove(Key.GetKey());
+			Constraints[Index] = nullptr;
+
+			FreeConstraint(Constraint);
+		}
+
+		void FreeConstraint(FPBDCollisionConstraint* Constraint)
+		{
+			if (Constraint != nullptr)
+			{
+				delete Constraint;
+			}
+		}
+
+		FParticlePairCollisionConstraints* FindParticlePairConstraints(const FParticlePairCollisionsKey& ParticlePairKey)
+		{
+			if (!bChaos_Collision_EnableManifoldRestore)
+			{
+				return nullptr;
+			}
+
+			FParticlePairCollisionConstraints** ParticlePairContactsPtr = ParticlePairConstraintMap.Find(ParticlePairKey.GetKey());
+			if (ParticlePairContactsPtr != nullptr)
+			{
+				return *ParticlePairContactsPtr;
+			}
+
+			return nullptr;
+		}
+
+		FParticlePairCollisionConstraints* FindOrCreateParticlePairConstraints(const FGeometryParticleHandle* Particle0, const FGeometryParticleHandle* Particle1)
+		{
+			if (!bChaos_Collision_EnableManifoldRestore)
+			{
+				return nullptr;
+			}
+
+			FParticlePairCollisionsKey ParticlePairKey = FParticlePairCollisionsKey(Particle0, Particle1);
+			FParticlePairCollisionConstraints* ParticlePairConstraints = FindParticlePairConstraints(ParticlePairKey);
+			if (ParticlePairConstraints == nullptr)
+			{
+				ParticlePairConstraints = new FParticlePairCollisionConstraints(ParticlePairKey, Particle0, Particle1, Epoch);
+
+				ParticlePairConstraintMap.Add(ParticlePairConstraints->GetKey().GetKey(), ParticlePairConstraints);
+			}
+
+			return ParticlePairConstraints;
+		}
+
+		void ProcessNewConstraints()
+		{
+			while (FPBDCollisionConstraint* NewConstraint = DequeueNewConstraint())
+			{
+				const int32 ConstraintIndex = Constraints.Add(NewConstraint);
+
+				NewConstraint->GetContainerCookie().Update(ConstraintIndex, Epoch);
+
+				// Contacts that were fully restored are already in the maps (see RestoreParticlePairConstraints()) and cannot be swept constraints
+				if (!NewConstraint->WasManifoldRestored())
+				{
+					ShapePairConstraintMap.Add(NewConstraint->GetKey().GetKey(), NewConstraint);
+
+					// @todo(chaos): only store the particle pair contact list when we might want to retore it next frame (small movement)
+					FParticlePairCollisionConstraints* ParticlePairConstraints = FindOrCreateParticlePairConstraints(NewConstraint->Particle[0], NewConstraint->Particle[1]);
+					if (ParticlePairConstraints != nullptr)
+					{
+						ParticlePairConstraints->AddConstraint(NewConstraint);
+					}
+				}
+			}
+
+			// We need to sort constraints for solver stability
+			// We have to use StableSort so that constraints of the same pair stay in the same order
+			// Otherwise the order within each pair can change due to where they start out in the array
+			// @todo(chaos): we should label each contact (and shape) for things like warm starting GJK
+			// and so we could use that label as part of the key
+			// and then we could use regular Sort (which is faster)			
+			// @todo(chaos): this can be moved to the island and therefoe done in parallel
+			Constraints.StableSort(ContactConstraintSortPredicate);
+		}
+
+		void PruneShapePairs()
+		{
+			// @todo(chaos): this should be optimizable but for now just check every contact to see if it is stale
+			// Also, we should be able to clean some of this up now that we have persistent islands. Do we even
+			// need a persistent global constraint list now? Maybe the pruning should be in the islands now.
+
+			// Destroy contacts that were not used this tick, as long as they are not held in a sleeping island
+			TArray<uint32> PrunedKeyList;
+			PrunedKeyList.Reserve(ShapePairConstraintMap.Num());
+			for (auto& KeyValuePair : ShapePairConstraintMap)
+			{
+				FPBDCollisionConstraint* Constraint = KeyValuePair.Value;
+				if (!Constraint->IsSleeping() && (Constraint->GetContainerCookie().LastUsedEpoch < Epoch))
+				{
+					FreeConstraint(KeyValuePair.Value);
+
+					PrunedKeyList.Add(KeyValuePair.Key);
+				}
+			}
+
+			// Remove destroyed contacts from the map
+			for (uint32 Key : PrunedKeyList)
+			{
+				ShapePairConstraintMap.Remove(Key);
+			}
+		}
+
+		void PruneParticlePairs()
+		{
+			if (!bChaos_Collision_EnableManifoldRestore)
+			{
+				ParticlePairConstraintMap.Reset();
+				return;
+			}
+
+			TArray<uint32> PrunedKeyList;
+			PrunedKeyList.Reserve(ParticlePairConstraintMap.Num());
+
+			for (auto& KeyValuePair : ParticlePairConstraintMap)
+			{
+				if (KeyValuePair.Value->GetEpoch() < Epoch)
+				{
+					delete KeyValuePair.Value;
+
+					PrunedKeyList.Add(KeyValuePair.Key);
+				}
+			}
+
+			for (uint32 Key : PrunedKeyList)
+			{
+				ParticlePairConstraintMap.Remove(Key);
+			}
+		}
+
+		void EnqueueNewConstraint(FPBDCollisionConstraint* InConstraint) const
+		{
+			if (InConstraint != nullptr)
+			{
+				NewConstraints.Push(InConstraint);
+			}
+		}
+
+		void EnqueueNewConstraints(TArrayView<FPBDCollisionConstraint* const> InConstraints) const
+		{
+			for (FPBDCollisionConstraint* Constraint : InConstraints)
+			{
+				if (Constraint != nullptr)
+				{
+					NewConstraints.Push(Constraint);
+				}
+			}
+		}
+
+		FPBDCollisionConstraint* DequeueNewConstraint()
+		{
+			return NewConstraints.Pop();
+		}
+
+		using FShapePairConstraintMap = TMap<uint32, FPBDCollisionConstraint*>;
+		using FParticlePairConstraintsMap = TMap<uint32, FParticlePairCollisionConstraints*>;
+
+		// All the constraints indexed by particle pair
+		// Note: partciles can be deleted so this map may briefly contain collisions for invalid particles
+		// however they will be pruned in PruneParticlePairs()
+		FParticlePairConstraintsMap ParticlePairConstraintMap;
+
+		// All of the allocated constraints, indexed by shape pair
+		// @todo(chaos): consider removing this and just having ParticlePairConstraintMap
+		FShapePairConstraintMap ShapePairConstraintMap;
+
+		// The active constraints (added or recovered this tick)
+		TArray<FPBDCollisionConstraint*> Constraints;
+
+		// The current epoch used to track out-of-date contacts. A constraint whose Epoch is
+		// older than the current Epoch at the end of the tick was not refreshed this tick.
+		int32 Epoch;
+
+		bool bInCollisionDetectionPhase;
+
+		// The set of constraints created or restored this tick during collision detection
+		// @todo(chaos): we can eliminate this lock-free structure by having the broadphase supply an array for use per thread
+		mutable TLockFreePointerListLIFO<FPBDCollisionConstraint> NewConstraints;
+	};
+
+}
\ No newline at end of file
diff --git a/Engine/Source/Runtime/Experimental/Chaos/Public/Chaos/Collision/CollisionContext.h b/Engine/Source/Runtime/Experimental/Chaos/Public/Chaos/Collision/CollisionContext.h
--- a/Engine/Source/Runtime/Experimental/Chaos/Public/Chaos/Collision/CollisionContext.h
+++ b/Engine/Source/Runtime/Experimental/Chaos/Public/Chaos/Collision/CollisionContext.h
@@ -5,6 +5,8 @@
 
 namespace Chaos
 {
+	class FCollisionConstraintAllocator;
+
 	/**
 	 * Data passed down into the collision detection functions.
 	 */
@@ -16,6 +18,7 @@ namespace Chaos
 			, bDeferUpdate(true)
 			, bAllowManifolds(false)
 			, bForceDisableCCD(false)
+			, CollisionAllocator(nullptr)
 		{
 		}
 
@@ -32,5 +35,7 @@ namespace Chaos
 
 		// Force disable CCD
 		bool bForceDisableCCD;
+
+		FCollisionConstraintAllocator* CollisionAllocator;
 	};
 }
\ No newline at end of file
diff --git a/Engine/Source/Runtime/Experimental/Chaos/Public/Chaos/Collision/CollisionDetector.h b/Engine/Source/Runtime/Experimental/Chaos/Public/Chaos/Collision/CollisionDetector.h
--- a/Engine/Source/Runtime/Experimental/Chaos/Public/Chaos/Collision/CollisionDetector.h
+++ b/Engine/Source/Runtime/Experimental/Chaos/Public/Chaos/Collision/CollisionDetector.h
@@ -26,13 +26,7 @@ namespace Chaos
 		FPBDCollisionConstraints& GetCollisionContainer() { return CollisionContainer; }
 		FNarrowPhase& GetNarrowPhase() { return NarrowPhase; }
 
-		virtual void DetectCollisionsWithStats(const FReal Dt, CollisionStats::FStatData& StatData, FEvolutionResimCache* ResimCache) = 0;
-
-		void DetectCollisions(const FReal Dt)
-		{
-			CollisionStats::FStatData StatData(false);
-			DetectCollisionsWithStats(Dt, StatData, nullptr);
-		}
+		virtual void DetectCollisions(const FReal Dt, FEvolutionResimCache* ResimCache) = 0;
 
 	protected:
 		FNarrowPhase& NarrowPhase;
diff --git a/Engine/Source/Runtime/Experimental/Chaos/Public/Chaos/Collision/ContactPoint.h b/Engine/Source/Runtime/Experimental/Chaos/Public/Chaos/Collision/ContactPoint.h
--- a/Engine/Source/Runtime/Experimental/Chaos/Public/Chaos/Collision/ContactPoint.h
+++ b/Engine/Source/Runtime/Experimental/Chaos/Public/Chaos/Collision/ContactPoint.h
@@ -17,7 +17,7 @@ namespace Chaos
 		FVec3 ShapeContactPoints[2];	// Shape-space contact points on the two bodies, without the margin added (i.e., contact point on the core shape)
 		FVec3 ShapeContactNormal;		// Shape-space contact normal relative to the NormalOwner, but with direction that always goes from body 1 to body 0
 		FReal ShapeMargins[2];			// Margins used in collision detection
-		int32 ContactNormalOwnerIndex;	// The shape which owns the contact normal (usually the second body, but not always for manifolds)
+		int32 ContactNormalOwnerIndex;	// The shape which owns the contact normal (usually the second body, but not always for manifolds) // Remove this when bChaos_Collision_Manifold_FixNormalsInWorldSpace is set to true permanetly
 
 		// @todo(chaos): these do not need to be stored here (they can be derived from above)
 		FVec3 Location;					// World-space contact location
@@ -42,7 +42,7 @@ namespace Chaos
 				Swap(ShapeContactPoints[0], ShapeContactPoints[1]);
 				Swap(ShapeMargins[0], ShapeMargins[1]);
 				ShapeContactNormal = -ShapeContactNormal;
-				ContactNormalOwnerIndex = (ContactNormalOwnerIndex == 0) ? 1 : 0;
+				ContactNormalOwnerIndex = (ContactNormalOwnerIndex == 0) ? 1 : 0; // ContactNormalOwnerIndex is ignored when  bChaos_Collision_Manifold_FixNormalsInWorldSpace is set
 				Normal = -Normal;
 			}
 			return *this;
diff --git a/Engine/Source/Runtime/Experimental/Chaos/Public/Chaos/Collision/NarrowPhase.h b/Engine/Source/Runtime/Experimental/Chaos/Public/Chaos/Collision/NarrowPhase.h
--- a/Engine/Source/Runtime/Experimental/Chaos/Public/Chaos/Collision/NarrowPhase.h
+++ b/Engine/Source/Runtime/Experimental/Chaos/Public/Chaos/Collision/NarrowPhase.h
@@ -6,6 +6,7 @@
 #include "Chaos/Collision/StatsData.h"
 #include "Chaos/CollisionResolution.h"
 #include "Chaos/ParticleHandle.h"
+#include "Chaos/Particle/ParticleUtilities.h"
 #include "Chaos/PBDCollisionConstraints.h"
 #include "Chaos/PBDRigidsSOAs.h"
 #include "ChaosStats.h"
@@ -24,8 +25,6 @@ namespace Chaos
 
 	/**
 	 * Generate contact manifolds for particle pairs.
-	 *
-	 * /see FAsyncCollisionReceiver, FSyncCollisionReceiver.
 	 */
 	class CHAOS_API FNarrowPhase
 	{
@@ -36,31 +35,55 @@ namespace Chaos
 
 		FCollisionContext& GetContext() { return Context; }
 
-		// @todo(chaos): COLLISION Transient Handle version
 		/**
 		 * /param CullDistance The contact separation at which we ignore the constraint
 		 */
-		void GenerateCollisions(FCollisionConstraintsArray& NewConstraints, FReal Dt, TGeometryParticleHandle<FReal, 3>* Particle0, TGeometryParticleHandle<FReal, 3>* Particle1, const FReal CullDistance, const bool bForceDisableCCD)
+		void GenerateCollisions(FReal Dt, TGeometryParticleHandle<FReal, 3>* Particle0, TGeometryParticleHandle<FReal, 3>* Particle1, const FReal CullDistance, const bool bForceDisableCCD)
 		{
 			SCOPE_CYCLE_COUNTER_NAROWPHASE();
 			if (ensure(Particle0 && Particle1))
 			{
-				//
-				// @todo(chaos) : Collision Constraints
-				//   This is not efficient. The constraint has to go through a construction 
-				//   process, only to be deleted later if it already existed. This should 
-				//   determine if the constraint is already defined, and then opt out of 
-				//   the creation process. 
-				//
+				FConstGenericParticleHandle P0 = Particle0;
+				FConstGenericParticleHandle P1 = Particle1;
+				const FRigidTransform3 ParticleTransform0 = FParticleUtilitiesPQ::GetActorWorldTransform(P0);
+				const FRigidTransform3 ParticleTransform1 = FParticleUtilitiesPQ::GetActorWorldTransform(P1);
 
-				// @todo(mlentine): Collision particles should exist optionally on geometry particles not rigid dynamic particles
-				TPBDRigidParticleHandle<FReal, 3>* RigidParticle0 = Particle0->CastToRigidParticle();
-				TPBDRigidParticleHandle<FReal, 3>* RigidParticle1 = Particle1->CastToRigidParticle();
 				Context.bForceDisableCCD = bForceDisableCCD;
-				Collisions::ConstructConstraints(Particle0, Particle1, Particle0->Geometry().Get(), RigidParticle0 ? RigidParticle0->CollisionParticles().Get() : nullptr, Particle1->Geometry().Get(), RigidParticle1 ? RigidParticle1->CollisionParticles().Get() : nullptr, FRigidTransform3(), FRigidTransform3(), CullDistance, Dt, Context, NewConstraints);
+
+				// @todo(mlentine): Collision particles should exist optionally on geometry particles not rigid dynamic particles
+				Collisions::ConstructConstraints(
+					Particle0, 
+					Particle1, 
+					Particle0->Geometry().Get(), 
+					P0->CollisionParticles().Get(), 
+					Particle1->Geometry().Get(), 
+					P1->CollisionParticles().Get(), 
+					ParticleTransform0, 
+					FRigidTransform3(), 
+					ParticleTransform1, 
+					FRigidTransform3(), 
+					CullDistance, 
+					Dt, 
+					Context);
 			}
 		}
 
+		bool TryRestoreCollisions(FReal Dt, TGeometryParticleHandle<FReal, 3>* Particle0, TGeometryParticleHandle<FReal, 3>* Particle1)
+		{
+			FConstGenericParticleHandle P0 = Particle0;
+			FConstGenericParticleHandle P1 = Particle1;
+			const FRigidTransform3 ParticleTransform0 = FParticleUtilitiesPQ::GetActorWorldTransform(P0);
+			const FRigidTransform3 ParticleTransform1 = FParticleUtilitiesPQ::GetActorWorldTransform(P1);
+
+			return Collisions::TryRestoreConstraints(
+				P0,
+				P1,
+				ParticleTransform0,
+				ParticleTransform1,
+				Dt,
+				Context);
+		}
+
 	private:
 		FCollisionContext Context;
 	};
