Chaos

diff --git a/Engine/Source/Runtime/Experimental/Chaos/Public/Chaos/Collision/PBDCollisionConstraint.h b/Engine/Source/Runtime/Experimental/Chaos/Public/Chaos/Collision/PBDCollisionConstraint.h
--- a/Engine/Source/Runtime/Experimental/Chaos/Public/Chaos/Collision/PBDCollisionConstraint.h
+++ b/Engine/Source/Runtime/Experimental/Chaos/Public/Chaos/Collision/PBDCollisionConstraint.h
@@ -6,23 +6,85 @@
 #include "Chaos/Core.h"
 #include "Chaos/CollisionResolutionTypes.h"
 #include "Chaos/Collision/ContactPoint.h"
+#include "Chaos/Collision/PBDCollisionConstraintHandle.h"
+#include "Chaos/GJK.h"
 #include "Chaos/ParticleHandleFwd.h"
 #include "Chaos/Vector.h"
 #include "Chaos/BVHParticles.h"
 
 namespace Chaos
 {
-	class FImplicitObject;
-	class FPBDCollisionConstraintHandle;
+	class FCollisionConstraintAllocator;
 	class FConstGenericParticleHandle;
+	class FImplicitObject;
+	class FParticlePairCollisionConstraints;
+	class FPBDCollisionConstraint;
+	class FPBDCollisionConstraints;
+	class FPBDCollisionSolver;
+	class FSolverBody;
+	class FSolverBodyContainer;
+	class FPBDCollisionSolverContainer;
+
+	UE_DEPRECATED(4.27, "Use FPBDCollisionConstraint instead")
+	typedef FPBDCollisionConstraint FRigidBodyPointContactConstraint;
+
+	CHAOS_API bool ContactConstraintSortPredicate(const FPBDCollisionConstraint& L, const FPBDCollisionConstraint& R);
+
+	class FRigidBodyContactKey
+	{
+	public:
+		FRigidBodyContactKey()
+			: Key(0)
+		{
+		}
+
+		FRigidBodyContactKey(const FGeometryParticleHandle* Particle0, const FImplicitObject* Implicit0, const FBVHParticles* Simplicial0, const FGeometryParticleHandle* Particle1, const FImplicitObject* Implicit1, const FBVHParticles* Simplicial1)
+			: Key(GenerateHash(Particle0, Implicit0, Simplicial0, Particle1, Implicit1, Simplicial1))
+		{
+		}
+
+		uint32 GetKey() const
+		{
+			return Key;
+		}
+
+		friend bool operator==(const FRigidBodyContactKey& L, const FRigidBodyContactKey& R)
+		{
+			return L.Key == R.Key;
+		}
+
+		friend bool operator!=(const FRigidBodyContactKey& L, const FRigidBodyContactKey& R)
+		{
+			return !(L == R);
+		}
+
+		friend bool operator<(const FRigidBodyContactKey& L, const FRigidBodyContactKey& R)
+		{
+			return L.Key < R.Key;
+		}
+
+	private:
+		uint32 GenerateHash(const FGeometryParticleHandle* Particle0, const FImplicitObject* Implicit0, const FBVHParticles* Simplicial0, const FGeometryParticleHandle* Particle1, const FImplicitObject* Implicit1, const FBVHParticles* Simplicial1)
+		{
+			// @todo(chaos): Particle handles can be reused - this needs to use a global ID or something (not available in this stream)
+			const uint32 ParticlesHash = HashCombine(::GetTypeHash(Particle0), ::GetTypeHash(Particle1));
+			const uint32 ImplicitHash = HashCombine(::GetTypeHash(Implicit0), ::GetTypeHash(Implicit1));
+			const uint32 SimplicialHash = HashCombine(::GetTypeHash(Simplicial0), ::GetTypeHash(Simplicial1));
+			const uint32 ParticleImplicitHash = HashCombine(ParticlesHash, ImplicitHash);
+			return HashCombine(ParticleImplicitHash, SimplicialHash);
+		}
+
+		uint32 Key;
+	};
+
 
 	class CHAOS_API FManifoldPoint
 	{
 	public:
 		FManifoldPoint() 
 			: CoMContactPoints{ FVec3(0), FVec3(0) }
-			, CoMContactNormal(0)
-			, CoMContactTangent(0)
+			, ManifoldContactTangents{ FVec3(0), FVec3(0) }
+			, ManifoldContactNormal(0)
 			, NetImpulse(0)
 			, NetPushOut(0)
 			, NetPushOutImpulseNormal(0)
@@ -37,13 +99,14 @@ namespace Chaos
 			, bInsideStaticFrictionCone(false)
 			, bRestitutionEnabled(false)
 			, bActive(false)
-			{}
+			, StaticFrictionMax(0)
+		{}
 
 		FManifoldPoint(const FContactPoint& InContactPoint) 
 			: ContactPoint(InContactPoint)
 			, CoMContactPoints{ FVec3(0), FVec3(0) }
-			, CoMContactNormal(0)
-			, CoMContactTangent(0)
+			, ManifoldContactTangents{ FVec3(0), FVec3(0) }
+			, ManifoldContactNormal(0)
 			, NetImpulse(0)
 			, NetPushOut(0)
 			, NetPushOutImpulseNormal(0)
@@ -58,13 +121,15 @@ namespace Chaos
 			, bInsideStaticFrictionCone(false)
 			, bRestitutionEnabled(false)
 			, bActive(false)
+			, StaticFrictionMax(0)
 		{}
 
 		// @todo(chaos): Normal and plane owner should be per manifold, not per manifold-point when we are not using incremental manifolds any more
 		FContactPoint ContactPoint;			// Shape-space data from low-level collision detection
 		FVec3 CoMContactPoints[2];			// CoM-space contact points on the two bodies core shapes (not including margin)
-		FVec3 CoMContactNormal;				// CoM-space contact normal relative to ContactNormalOwner body	
-		FVec3 CoMContactTangent;			// CoM-space contact tangent for friction
+		FVec3 ManifoldContactTangents[2];		// CoM-space or world space (depending on CVAR bChaos_Collision_Manifold_FixNormalsInWorldSpace) contact tangents for friction
+		FVec3 ManifoldContactNormal;				// CoM-space or  world space  (depending on CVAR bChaos_Collision_Manifold_FixNormalsInWorldSpace) contact normal relative to ContactNormalOwner body
+		FVec3 CoMAnchorPoints[2];			// CoM-space contact points on the two bodies core shapes used for static friction
 		FVec3 NetImpulse;					// Total impulse applied by this contact point
 		FVec3 NetPushOut;					// Total pushout applied at this contact point
 		FReal NetPushOutImpulseNormal;		// Total pushout impulse along normal (for final velocity correction) applied at this contact point
@@ -75,10 +140,12 @@ namespace Chaos
 		FReal NetImpulseTangent;
 		FReal InitialContactVelocity;		// Contact velocity at start of frame (used for restitution)
 		FReal InitialPhi;					// Contact separation at first contact (used for pushout restitution)
-		bool bPotentialRestingContact;		// Whether this may be a resting contact (used for static fricton)
+		bool bPotentialRestingContact;		// Whether this may be a resting contact (used for static friction)
 		bool bInsideStaticFrictionCone;		// Whether we are inside the static friction cone (used in PushOut)
 		bool bRestitutionEnabled;			// Whether restitution was added in the apply step (used in PushOut)
 		bool bActive;						// Whether this contact applied an impulse
+
+		FReal StaticFrictionMax;			// A proxy for the normal impulse used to limit static friction correction. Used for smoothing static friction limits
 	};
 
 	/*
@@ -97,6 +164,7 @@ namespace Chaos
 			, AngularFriction(0)
 			, Restitution(0)
 			, RestitutionPadding(0)
+			, RestitutionThreshold(0)
 			, InvInertiaScale0(1.f)
 			, InvInertiaScale1(1.f)
 			, ContactMoveSQRDistance(0)
@@ -119,11 +187,18 @@ namespace Chaos
 		FReal AngularFriction;
 		FReal Restitution;
 		FReal RestitutionPadding; // For PBD implementation of resitution, we pad constraints on initial contact to enforce outward velocity
+		FReal RestitutionThreshold;
 		FReal InvInertiaScale0;
 		FReal InvInertiaScale1;
 		FReal ContactMoveSQRDistance; // How much the contact position moved after the last update
 		EContactShapesType ShapesType;
 
+		void Reset()
+		{
+			bDisabled = false;
+			Phi = FLT_MAX;
+			RestitutionPadding = 0;
+		}
 
 		FString ToString() const
 		{
@@ -134,71 +209,154 @@ namespace Chaos
 		const FBVHParticles* Simplicial[2]; // {Of Particle[0], Of Particle[1]}
 	};
 
-	/*
-	*
+
+	/**
+	 * @brief Information used by the constraint allocator to determine which parts of a previous tick's contact can be reused (if any)
 	*/
-	class CHAOS_API FCollisionConstraintBase
+	class FPBDCollisionConstraintContainerCookie
 	{
 	public:
-		using FGeometryParticleHandle = TGeometryParticleHandle<FReal, 3>;
+		FPBDCollisionConstraintContainerCookie()
+			: Key()
+			, CreationEpoch(INDEX_NONE)
+			, LastUsedEpoch(INDEX_NONE)
+			, ConstraintIndex(INDEX_NONE)
+			, bIsSleeping(false)
+		{
+		}
+
+		/**
+		 * @brief Used to clear the container data when copying constraints out of the container (see resim cache)
+		 * The constraint index will no be valid when the constraint is copied out of the container, but everything
+		 * else is ok and should be restorable.
+		*/
+		void ClearContainerData()
+		{
+			ConstraintIndex = INDEX_NONE;			
+		}
+
+		/**
+		 * @brief Called once on creation by the container to initialize the lifetime contants
+		*/
+		void Init(const FRigidBodyContactKey& InKey, const int32 InEpoch)
+		{
+			Key = InKey;
+			CreationEpoch = InEpoch;
+			LastUsedEpoch = InEpoch;
+			bIsSleeping = false;
+		}
 
-		enum class FType
+		/**
+		 * @brief Called each tick by the container
+		*/
+		void Update(const int32 InIndex, const int32 InEpoch)
 		{
-			None = 0,			// default value also indicates a invalid constraint
-			SinglePoint,		// TRigidBodyPointContactConstraint
-			SinglePointSwept,	// TRigidBodySweptPointContactConstraint
-		};
-
-		FCollisionConstraintBase(FType InType = FType::None)
-			: AccumulatedImpulse(0)
-			, Timestamp(-INT_MAX)
-			, ConstraintHandle(nullptr)
-			, Type(InType)
-			, Stiffness(1.0)
+			ConstraintIndex = InIndex;
+			LastUsedEpoch = InEpoch;
+			bIsSleeping = false;
+		}
+
+		/**
+		 * @brief Change the sleeping state (used by the Island Manager)
+		*/
+		void SetIsSleeping(const bool bInIsSleeping)
 		{
-			ImplicitTransform[0] = FRigidTransform3::Identity; ImplicitTransform[1] = FRigidTransform3::Identity;
-			Manifold.Implicit[0] = nullptr; Manifold.Implicit[1] = nullptr;
-			Manifold.Simplicial[0] = nullptr; Manifold.Simplicial[1] = nullptr;
-			Particle[0] = nullptr; Particle[1] = nullptr;
+			bIsSleeping = bInIsSleeping;
 		}
 
-		FCollisionConstraintBase(
+		// A hash of the colliding object pair to uniquely identify a contact and allow recovery next tick
+		FRigidBodyContactKey Key;
+
+		// The conatiner Epoch when then constraint was initially created
+		int32 CreationEpoch;
+
+		// The container Epoch when the constraint was last used
+		int32 LastUsedEpoch;
+
+		// The index in the container - this changes every tick
+		int32 ConstraintIndex;
+
+		// Whether the constraint is in a sleeping island (this is used to prevent culling because Epoch is not updated for sleepers)
+		bool bIsSleeping;
+	};
+
+
+	/**
+	 * @brief A contact constraint
+	 * 
+	 * A contact constraint represents the non-penetration, friction, and restitution constraints for a single
+	 * shape pair on a particle pair. I.e., a specific particle-pair may have multiple contact constraints
+	 * between them if one or boht has multuple collision shapes that overlap the shape(s) of the other body.
+	 * 
+	 * Each contact constraint contains a Manifold, which is a set of contact points that approximate the
+	 * contact patch between the two shapes.
+	 * 
+	 * Contact constraints are allocated on the heap and have permanent addresses. They use intrusive handles
+	 * to reduce unnecessary indirection.
+	 * 
+	*/
+	class CHAOS_API FPBDCollisionConstraint : public FPBDCollisionConstraintHandle
+	{
+		friend class FCollisionConstraintAllocator;
+		friend class FPBDCollisionConstraints;
+		friend CHAOS_API bool ContactConstraintSortPredicate(const FPBDCollisionConstraint& L, const FPBDCollisionConstraint& R);
+
+	public:
+		using FConstraintContainerHandle = TIntrusiveConstraintHandle<FPBDCollisionConstraint>;
+
+		/**
+		 * @brief Create a standard (non-swept) contact
+		 * Allocates a constraint on the heap, with a permanent address.
+		 * May return null if we hit the contact limit for the scene.
+		*/
+		static FPBDCollisionConstraint* Make(
 			FGeometryParticleHandle* Particle0,
 			const FImplicitObject* Implicit0,
 			const FBVHParticles* Simplicial0,
-			const FRigidTransform3& Transform0,
+			const FRigidTransform3& ParticleWorldTransform0,
+			const FRigidTransform3& ImplicitLocalTransform0,
 			FGeometryParticleHandle* Particle1,
 			const FImplicitObject* Implicit1,
 			const FBVHParticles* Simplicial1,
-			const FRigidTransform3& Transform1,
-			const FType InType, 
-			const EContactShapesType ShapesType, 
-			const int32 InTimestamp = -INT_MAX)
-			: AccumulatedImpulse(0)
-			, Timestamp(InTimestamp)
-			, ConstraintHandle(nullptr)
-			, Type(InType)
-			, Stiffness(1.0)
-		{
-			ImplicitTransform[0] = Transform0; 
-			ImplicitTransform[1] = Transform1;
-
-			Manifold.Implicit[0] = Implicit0; 
-			Manifold.Implicit[1] = Implicit1;
-
-			Manifold.Simplicial[0] = Simplicial0; 
-			Manifold.Simplicial[1] = Simplicial1;
+			const FRigidTransform3& ParticleWorldTransform1,
+			const FRigidTransform3& ImplicitLocalTransform1,
+			const FReal InCullDistance,
+			const EContactShapesType ShapesType,
+			const bool bInUseManifold,
+			FCollisionConstraintAllocator& Allocator);
+
+		/**
+		 * @brief Create a swept contact
+		 * Allocates a constraint on the heap, with a permanent address.
+		 * May return null if we hit the contact limit for the scene.
+		*/
+		static FPBDCollisionConstraint* MakeSwept(
+			FGeometryParticleHandle* Particle0,
+			const FImplicitObject* Implicit0,
+			const FBVHParticles* Simplicial0,
+			const FRigidTransform3& ParticleWorldTransform0,
+			const FRigidTransform3& ImplicitLocalTransform0,
+			FGeometryParticleHandle* Particle1,
+			const FImplicitObject* Implicit1,
+			const FBVHParticles* Simplicial1,
+			const FRigidTransform3& ParticleWorldTransform1,
+			const FRigidTransform3& ImplicitLocalTransform1,
+			const FReal InCullDistance,
+			EContactShapesType ShapesType,
+			FCollisionConstraintAllocator& Allocator);
 
-			Manifold.ShapesType = ShapesType;
+		/**
+		 * @brief Return a constraint copied from the Source constraint, for use in the Resim Cache
+		 * This copies everything needed to rehydrate the contact after a rewind
+		 * @note Unlike the other factory methods, this version returns a constraint by value for emplacing into an array (ther others are by pointer)
+		*/
+		static FPBDCollisionConstraint MakeResimCache(
+			const FPBDCollisionConstraint& Source);
 
-			Particle[0] = Particle0; 
-			Particle[1] = Particle1;
-		}
 
-		FType GetType() const { return Type; }
+		FPBDCollisionConstraint();
 
-		template<class AS_T> AS_T* As() { return static_cast<AS_T*>(this); }
-		template<class AS_T> const AS_T* As() const { return static_cast<const AS_T*>(this); }
+		ECollisionConstraintType GetType() const { return Type; }
 
 		bool ContainsManifold(const FImplicitObject* A, const FBVHParticles* AS, const FImplicitObject* B, const FBVHParticles* BS) const
 		{
@@ -211,12 +369,14 @@ namespace Chaos
 			Manifold.Simplicial[0] = AS; Manifold.Simplicial[1] = BS;
 		}
 
-		const FCollisionContact& GetManifold() const { return Manifold; }
-
 		//
 		// API
 		//
 
+		// @todo(chaos): half of this API is wrong for the new multi-point manifold constraints. Remove it
+
+		const FCollisionContact& GetManifold() const { return Manifold; }
+
 		void ResetPhi(FReal InPhi) { SetPhi(InPhi); }
 
 		void SetPhi(FReal InPhi) { Manifold.Phi = InPhi; }
@@ -225,6 +385,9 @@ namespace Chaos
 		void SetDisabled(bool bInDisabled) { Manifold.bDisabled = bInDisabled; }
 		bool GetDisabled() const { return Manifold.bDisabled; }
 
+		void SetIsSleeping(const bool bInIsSleeping) { ContainerCookie.SetIsSleeping(bInIsSleeping); }
+		bool GetIsSleeping() const { return ContainerCookie.bIsSleeping; }
+
 		void SetNormal(const FVec3& InNormal) { Manifold.Normal = InNormal; }
 		FVec3 GetNormal() const { return Manifold.Normal; }
 
@@ -242,66 +405,18 @@ namespace Chaos
 
 		FString ToString() const;
 
-		FRigidTransform3 ImplicitTransform[2];		// Local-space transforms of the shape (relative to particle)
-		FGeometryParticleHandle* Particle[2];
-		FVec3 AccumulatedImpulse;					// @todo(chaos): we need to accumulate angular impulse separately
-		FCollisionContact Manifold					;// @todo(chaos): rename
-		int32 Timestamp;
-		
-		void SetConstraintHandle(FPBDCollisionConstraintHandle* InHandle) { ConstraintHandle = InHandle; }
-		FPBDCollisionConstraintHandle* GetConstraintHandle() const { return ConstraintHandle; }
-
-		bool operator<(const FCollisionConstraintBase& Other) const;
-
-	private:
-		FPBDCollisionConstraintHandle* ConstraintHandle;
-		FType Type;
-		FReal Stiffness;
-	};
-
-
-	/*
-	*
-	*/
-	class CHAOS_API FRigidBodyPointContactConstraint : public FCollisionConstraintBase
-	{
-	public:
-		using Base = FCollisionConstraintBase;
-		using FGeometryParticleHandle = TGeometryParticleHandle<FReal, 3>;
-
-		FRigidBodyPointContactConstraint() 
-			: Base(Base::FType::SinglePoint)
-			, CullDistance(FLT_MAX)
-			, bUseManifold(false)
-		{}
-
-		FRigidBodyPointContactConstraint(
-			FGeometryParticleHandle* Particle0, 
-			const FImplicitObject* Implicit0, 
-			const FBVHParticles* Simplicial0, 
-			const FRigidTransform3& Transform0,
-			FGeometryParticleHandle* Particle1, 
-			const FImplicitObject* Implicit1, 
-			const FBVHParticles* Simplicial1, 
-			const FRigidTransform3& Transform1,
-			const FReal InCullDistance,
-			const EContactShapesType ShapesType,
-			const bool bInUseManifold)
-			: Base(Particle0, Implicit0, Simplicial0, Transform0, Particle1, Implicit1, Simplicial0, Transform1, Base::FType::SinglePoint, ShapesType)
-			, CullDistance(InCullDistance)
-			, bUseManifold(false)
-		{
-			bUseManifold = bInUseManifold && CanUseManifold(Particle0, Particle1);
-		}
-
-		static typename Base::FType StaticType() { return Base::FType::SinglePoint; };
-
 		FReal GetCullDistance() const { return CullDistance; }
 		void SetCullDistance(FReal InCullDistance) { CullDistance = InCullDistance; }
 
 		bool GetUseManifold() const { return bUseManifold; }
+		
+		bool GetUseIncrementalCollisionDetection() const { return !bUseManifold || bUseIncrementalManifold; }
+
+		/**
+		 * @brief Clear the current and previous manifolds
+		*/
+		void ResetManifold();
 
-		void SetManifoldPoints(const TArray<FManifoldPoint>& InManifoldPoints) { ManifoldPoints=InManifoldPoints; }
 		TArrayView<FManifoldPoint> GetManifoldPoints() { return MakeArrayView(ManifoldPoints); }
 		TArrayView<const FManifoldPoint> GetManifoldPoints() const { return MakeArrayView(ManifoldPoints); }
 		FManifoldPoint& SetActiveManifoldPoint(
@@ -311,8 +426,8 @@ namespace Chaos
 			const FVec3& P1,
 			const FRotation3& Q1);
 		void CalculatePrevCoMContactPoints(
-			const FConstGenericParticleHandle Particle0,
-			const FConstGenericParticleHandle Particle1,
+			const FSolverBody& Body0,
+			const FSolverBody& Body1,
 			FManifoldPoint& ManifoldPoint,
 			FReal Dt,
 			FVec3& OutPrevCoMContactPoint0,
@@ -320,32 +435,81 @@ namespace Chaos
 
 		void AddIncrementalManifoldContact(const FContactPoint& ContactPoint, const FReal Dt);
 		void AddOneshotManifoldContact(const FContactPoint& ContactPoint, const FReal Dt);
-		void UpdateManifoldContacts(FReal Dt);
+		void UpdateManifoldContacts();
 		void ClearManifold();
 
-		//@todo(make protected) : Was exposed to support contact copying. 
+		//@ todo(chaos): These are for the collision forwarding system - this should use the collision modifier system (which should be extended to support adding collisions)
+		void SetManifoldPoints(const TArray<FManifoldPoint>& InManifoldPoints) { ManifoldPoints = InManifoldPoints; }
 		void UpdateManifoldPointFromContact(FManifoldPoint& ManifoldPoint);
 
+		// The GJK warm-start data. This is updated directly in the narrow phase
+		FGJKSimplexData& GetGJKWarmStartData() { return GJKWarmStartData; }
+
+		FSolverBody* GetSolverBody0() { return SolverBodies[0]; }
+		const FSolverBody* GetSolverBody0() const { return SolverBodies[0]; }
+		FSolverBody* GetSolverBody1() { return SolverBodies[1]; }
+		const FSolverBody* GetSolverBody1() const { return SolverBodies[1]; }
+
+		void SetSolverBodies(FSolverBody* InSolverBody0, FSolverBody* InSolverBody1)
+		{
+			SolverBodies[0] = InSolverBody0;
+			SolverBodies[1] = InSolverBody1;
+		}
+
+		/**
+		 * @brief Whether this constraint was newly created this tick (as opposed to restored from a previous tick)
+		 * @see IsRestored()
+		*/
+		bool IsNew() const { return ContainerCookie.CreationEpoch == ContainerCookie.LastUsedEpoch; }
+
+		/**
+		 * @brief Whether this constraint was fully restored from a previous tick, and the manifold should be reused as-is
+		*/
+		bool WasManifoldRestored() const { return bWasManifoldRestored; }
+
+		/**
+		 * @brief Restore the contact manifold (assumes relative motion of the two bodies is small)
+		 * @see IsWithinManifoldRestorationThreshold
+		*/
+		void RestoreManifold();
+
+		/**
+		 * Determine the constraint direction based on Normal and Phi.
+		 * This function assumes that the constraint is update-to-date.
+		 */
+		ECollisionConstraintDirection GetConstraintDirection(const FReal Dt) const;
+
+		/**
+		 * @brief The container key for this contact (for internal use only)
+		*/
+		const FRigidBodyContactKey& GetKey() const { return GetContainerCookie().Key; }
+
+	public:
+		const FPBDCollisionConstraintHandle* GetConstraintHandle() const { return this; }
+		FPBDCollisionConstraintHandle* GetConstraintHandle() { return this; }
+
 	protected:
-		// For use by derived types that can be used as point constraints in Update
-		FRigidBodyPointContactConstraint(typename Base::FType InType) : Base(InType) {}
-
-		FRigidBodyPointContactConstraint(
-			FGeometryParticleHandle* Particle0, 
-			const FImplicitObject* Implicit0, 
-			const FBVHParticles* Simplicial0, 
-			const FRigidTransform3& Transform0,
-			FGeometryParticleHandle* Particle1, 
-			const FImplicitObject* Implicit1, 
-			const FBVHParticles* Simplicial1,
-			const FRigidTransform3& Transform1,
+
+		FPBDCollisionConstraint(
+			FGeometryParticleHandle* Particle0,
+			const FImplicitObject* Implicit0,
+			const FBVHParticles* Simplicial0,
+			FGeometryParticleHandle* Particle1,
+			const FImplicitObject* Implicit1,
+			const FBVHParticles* Simplicial1);
+
+		// Set all the data not initialized in the constructor
+		void Setup(
+			const ECollisionConstraintType InType,
+			const EContactShapesType InShapesType,
+			const FRigidTransform3& InImplicitTransform0,
+			const FRigidTransform3& InImplicitTransform1,
 			const FReal InCullDistance,
-			typename Base::FType InType,
-			const EContactShapesType ShapesType)
-			: Base(Particle0, Implicit0, Simplicial0, Transform0, Particle1, Implicit1, Simplicial1, Transform1, InType, ShapesType)
-			, CullDistance(InCullDistance)
-			, bUseManifold(false)
-		{}
+			const bool bInUseManifold);
+
+		// Access to the data used by the container
+		const FPBDCollisionConstraintContainerCookie& GetContainerCookie() const { return ContainerCookie; }
+		FPBDCollisionConstraintContainerCookie& GetContainerCookie() { return ContainerCookie; }
 
 		// Whether we can use manifolds for the given partices. Manifolds do not work well with Joints and PBD
 		// because the bodies may be moved (and especially rotated) a lot in the solver and this can make the manifold extremely inaccurate
@@ -359,78 +523,55 @@ namespace Chaos
 		void SetActiveContactPoint(const FContactPoint& ContactPoint);
 		void GetWorldSpaceManifoldPoint(const FManifoldPoint& ManifoldPoint, const FVec3& P0, const FRotation3& Q0, const FVec3& P1, const FRotation3& Q1, FVec3& OutContactLocation, FVec3& OutContactNormal, FReal& OutContactPhi);
 
-		// @todo(chaos): Switching this to inline allocator does not help, but maybe a physics scratchpad would
-		//TArray<FManifoldPoint, TInlineAllocator<4>> ManifoldPoints;
-		TArray<FManifoldPoint> ManifoldPoints;
-		FReal CullDistance;
-		bool bUseManifold;
-	};
+		/**
+		 * @brief Move the current manifold to the previous manifold, and reset the current manifold ready for it to be rebuilt
+		*/
+		void SaveManifold();
 
+		/**
+		 * @brief Restore data used for static friction from the previous point(s) to the current
+		*/
+		bool TryRestoreFrictionData(FManifoldPoint& ManifoldPoint);
 
-	/*
-	*
-	*/
-	class CHAOS_API FRigidBodySweptPointContactConstraint : public FRigidBodyPointContactConstraint
-	{
 	public:
-		using Base = FRigidBodyPointContactConstraint;
-		using FGeometryParticleHandle = TGeometryParticleHandle<FReal, 3>;
-
-		FRigidBodySweptPointContactConstraint() : Base(Base::FType::SinglePointSwept) {}
-		FRigidBodySweptPointContactConstraint(
-			FGeometryParticleHandle* Particle0, 
-			const FImplicitObject* Implicit0, 
-			const FBVHParticles* Simplicial0, 
-			const FRigidTransform3& Transform0,
-			FGeometryParticleHandle* Particle1, 
-			const FImplicitObject* Implicit1, 
-			const FBVHParticles* Simplicial1, 
-			const FRigidTransform3& Transform1,
-			const FReal InCullDistance,
-			EContactShapesType ShapesType)
-			: Base(Particle0, Implicit0, Simplicial0, Transform0, Particle1, Implicit1, Simplicial1, Transform1, InCullDistance, Base::FType::SinglePointSwept, ShapesType)
-			, TimeOfImpact(0)
-		{}
+		//@todo(chaos): make this stuff private
+		FRigidTransform3 ImplicitTransform[2];		// Local-space transforms of the shape (relative to particle)
+		FGeometryParticleHandle* Particle[2];
+		FVec3 AccumulatedImpulse;					// @todo(chaos): we need to accumulate angular impulse separately
+		FCollisionContact Manifold;// @todo(chaos): rename
 
 		// Value in range [0,1] used to interpolate P between [X,P] that we will rollback to when solving at time of impact.
 		FReal TimeOfImpact;
-		static typename Base::FType StaticType() { return Base::FType::SinglePointSwept; };
-	};
 
+	private:
+		FPBDCollisionConstraintContainerCookie ContainerCookie;
+		ECollisionConstraintType Type;
+		FReal Stiffness;
 
-	struct CHAOS_API FCollisionConstraintsArray
-	{
-		// Previous arrays using TInlineAllocator<8> changed to normal arrays. Profiling showed faster without and
-		// we are able to use a per particle async reciever cache and just MoveTemp all the constraints for a 
-		// speed improvement to collision detection/gather.
-		TArray<FRigidBodyPointContactConstraint> SinglePointConstraints;
-		TArray<FRigidBodySweptPointContactConstraint> SinglePointSweptConstraints;
+		// @todo(chaos): Switching this to inline allocator does not help, but maybe a physics scratchpad would
+		//TArray<FManifoldPoint, TInlineAllocator<4>> ManifoldPoints;
+		TArray<FManifoldPoint> ManifoldPoints;
+		FReal CullDistance;
+		bool bUseManifold;
+		bool bUseIncrementalManifold;
 
-		int32 Num() const { return SinglePointConstraints.Num() + SinglePointSweptConstraints.Num(); }
+		// These are only needed here while we still have thge lgacy solvers (not QuasiPBD)
+		FSolverBody* SolverBodies[2];
 
-		void Reset()
-		{
-			SinglePointConstraints.Reset();
-			SinglePointSweptConstraints.Reset();
-		}
+		// Simplex data from the last call to GJK, used to warm-start GJK
+		FGJKSimplexData GJKWarmStartData;
 
-		void Empty()
-		{
-			SinglePointConstraints.Empty();
-			SinglePointSweptConstraints.Empty();
-		}
+		// The manifold points from the previous tick when we don't reuse the manifold. Used by static friction.
+		// @todo(chaos): this could be a subset of the manifold data
+		TArray<FManifoldPoint> PrevManifoldPoints;
 
+		bool bWasManifoldRestored;
+	};
 
-		FRigidBodyPointContactConstraint* Add(const FRigidBodyPointContactConstraint& C)
-		{
-			int32 ConstraintIndex = SinglePointConstraints.Add(C);
-			return &SinglePointConstraints[ConstraintIndex];
-		}
 
-		FRigidBodySweptPointContactConstraint* Add(const FRigidBodySweptPointContactConstraint& C)
-		{
-			int32 ConstraintIndex = SinglePointSweptConstraints.Add(C);
-			return &SinglePointSweptConstraints[ConstraintIndex];
-		}
-	};
+	UE_DEPRECATED(4.27, "FCollisionConstraintBase has been removed and folded into FPBDCollisionConstraint. Use FPBDCollisionConstraint")
+	typedef FPBDCollisionConstraint FCollisionConstraintBase;
+
+	UE_DEPRECATED(4.27, "FRigidBodySweptPointContactConstraint has been removed and folded into FPBDCollisionConstraint. Use FPBDCollisionConstraint")
+	typedef FPBDCollisionConstraint FRigidBodySweptPointContactConstraint;
 }
diff --git a/Engine/Source/Runtime/Experimental/Chaos/Public/Chaos/Collision/PBDCollisionConstraintHandle.h b/Engine/Source/Runtime/Experimental/Chaos/Public/Chaos/Collision/PBDCollisionConstraintHandle.h
--- a/Engine/Source/Runtime/Experimental/Chaos/Public/Chaos/Collision/PBDCollisionConstraintHandle.h
+++ b/Engine/Source/Runtime/Experimental/Chaos/Public/Chaos/Collision/PBDCollisionConstraintHandle.h
@@ -4,92 +4,106 @@
 #include "CoreMinimal.h"
 
 #include "Chaos/CollisionResolutionTypes.h"
-#include "Chaos/Collision/PBDCollisionConstraint.h"
 #include "Chaos/ConstraintHandle.h"
 #include "Chaos/PBDConstraintContainer.h"
 
 namespace Chaos
 {
 	class FPBDCollisionConstraints;
+	class FPBDCollisionConstraint;
+	class FPBDIslandSolverData;
+	class FPBDCollisionSolverContainer;
+
+	/**
+	 * @brief The type of contact (swept or not)
+	*/
+	enum class ECollisionConstraintType
+	{
+		// Constraint is not initialized
+		None = 0,
+
+		// Standard contact constraint
+		Standard,
+
+		// Swept contact constraint
+		Swept,
+	};
 
-	class CHAOS_API FPBDCollisionConstraintHandle : public TContainerConstraintHandle<FPBDCollisionConstraints>
+	/**
+	 * @brief The resting directionality of a contact constraint for use in constraint solver ordering
+	*/
+	enum ECollisionConstraintDirection
+	{
+		// Particle 1 is on top
+		Particle0ToParticle1,
+
+		// Particle 0 is on top
+		Particle1ToParticle0,
+
+		// Neither particle is on top
+		NoRestingDependency
+	};
+
+	/**
+	 * @brief A handle to a contact constraint.
+	 * @note This is an intrusive handle, so you can use a contact pointer as a handle.
+	*/
+	class CHAOS_API FPBDCollisionConstraintHandle : public TIntrusiveConstraintHandle<FPBDCollisionConstraint>
 	{
 	public:
-		using Base = TContainerConstraintHandle<FPBDCollisionConstraints>;
+		using Base = TIntrusiveConstraintHandle<FPBDCollisionConstraint>;
 		using FImplicitPair = TPair<const FImplicitObject*, const FImplicitObject*>;
 		using FGeometryPair = TPair<const TGeometryParticleHandle<FReal, 3>*, const TGeometryParticleHandle<FReal, 3>*>;
 		using FHandleKey = TPair<FImplicitPair, FGeometryPair>;
 
-
 		FPBDCollisionConstraintHandle()
-			: ConstraintType(FCollisionConstraintBase::FType::None)
-		{}
-
-		FPBDCollisionConstraintHandle(FPBDCollisionConstraints* InConstraintContainer, int32 InConstraintIndex, typename FCollisionConstraintBase::FType InType)
-			: TContainerConstraintHandle<FPBDCollisionConstraints>(InConstraintContainer, InConstraintIndex)
-			, ConstraintType(InType)
+			: TIntrusiveConstraintHandle<FPBDCollisionConstraint>()
 		{
 		}
-		static EConstraintContainerType StaticType() { return EConstraintContainerType::Collision; }
 
+		const FPBDCollisionConstraint& GetContact() const;
+		FPBDCollisionConstraint& GetContact();
 
-		FHandleKey GetKey();
+		UE_DEPRECATED(4.27, "Use GetContact()")
+		const FPBDCollisionConstraint& GetPointContact() const { return GetContact(); }
 
-		static FHandleKey MakeKey(const TGeometryParticleHandle<FReal, 3>* Particle0, const TGeometryParticleHandle<FReal, 3>* Particle1,
-			const FImplicitObject* Implicit0, const FImplicitObject* Implicit1)
-		{
-			return FHandleKey(FImplicitPair(Implicit0, Implicit1), FGeometryPair(Particle0, Particle1));
-		}
+		UE_DEPRECATED(4.27, "Use GetContact()")
+		FPBDCollisionConstraint& GetPointContact() { return GetContact(); }
 
-		static FHandleKey MakeKey(const FCollisionConstraintBase* Base)
-		{
-			return FHandleKey(FImplicitPair(Base->Manifold.Implicit[0], Base->Manifold.Implicit[1]), FGeometryPair(Base->Particle[0], Base->Particle[1]));
-		}
+		UE_DEPRECATED(4.27, "Use GetContact()")
+		const FPBDCollisionConstraint& GetSweptPointContact() const { return GetContact(); }
 
+		UE_DEPRECATED(4.27, "Use GetContact()")
+		FPBDCollisionConstraint& GetSweptPointContact() { return GetContact(); }
 
-		const FCollisionConstraintBase& GetContact() const;
-		FCollisionConstraintBase& GetContact();
 
-		const FRigidBodyPointContactConstraint& GetPointContact() const;
-		FRigidBodyPointContactConstraint& GetPointContact();
+		ECollisionConstraintType GetType() const;
 
-		const FRigidBodySweptPointContactConstraint& GetSweptPointContact() const;
-		FRigidBodySweptPointContactConstraint& GetSweptPointContact();
+		virtual void SetEnabled(bool InEnabled) override;
 
-		typename FCollisionConstraintBase::FType GetType() const { return ConstraintType; }
+		virtual bool IsEnabled() const override;
 
-		void SetConstraintIndex(int32 IndexIn, typename FCollisionConstraintBase::FType InType)
-		{
-			ConstraintIndex = IndexIn;
-			ConstraintType = InType;
-		}
+		FVec3 GetContactLocation() const;
 
-		FVec3 GetContactLocation() const
-		{
-			return GetContact().GetLocation();
-		}
-
-		FVec3 GetAccumulatedImpulse() const
-		{
-			return GetContact().AccumulatedImpulse;
-		}
+		FVec3 GetAccumulatedImpulse() const;
 
-		TVector<const TGeometryParticleHandle<FReal, 3>*, 2> GetConstrainedParticles() const
-		{
-			return { GetContact().Particle[0], GetContact().Particle[1] };
-		}
+		TVector<const TGeometryParticleHandle<FReal, 3>*, 2> GetConstrainedParticles() const;
 
-		TVector<TGeometryParticleHandle<FReal, 3>*, 2> GetConstrainedParticles()
-		{
-			return { GetContact().Particle[0], GetContact().Particle[1] };
-		}
+		TVector<TGeometryParticleHandle<FReal, 3>*, 2> GetConstrainedParticles();
 
+		void GatherInput(FReal Dt, const int32 Particle0Level, const int32 Particle1Level, FPBDIslandSolverData& SolverData);
 
-	protected:
-		typename FCollisionConstraintBase::FType ConstraintType;
-		using Base::ConstraintIndex;
-		using Base::ConcreteContainer;
+		FSolverBody* GetSolverBody0();
+		FSolverBody* GetSolverBody1();
 
+		const FPBDCollisionConstraints* ConcreteContainer() const;
+		FPBDCollisionConstraints* ConcreteContainer();
 
+		static const FConstraintHandleTypeID& StaticType()
+		{
+			static FConstraintHandleTypeID STypeID(TEXT("FCollisionConstraintHandle"), &FIntrusiveConstraintHandle::StaticType());
+			return STypeID;
+		}
 	};
+
 }
\ No newline at end of file
diff --git a/Engine/Source/Runtime/Experimental/Chaos/Public/Chaos/Collision/ParticlePairBroadPhase.h b/Engine/Source/Runtime/Experimental/Chaos/Public/Chaos/Collision/ParticlePairBroadPhase.h
--- a/Engine/Source/Runtime/Experimental/Chaos/Public/Chaos/Collision/ParticlePairBroadPhase.h
+++ b/Engine/Source/Runtime/Experimental/Chaos/Public/Chaos/Collision/ParticlePairBroadPhase.h
@@ -46,10 +46,7 @@ namespace Chaos
 		/**
 		 *
 		 */
-		void ProduceOverlaps(FReal Dt,
-			FCollisionConstraintsArray& ConstraintsArray,
-			FNarrowPhase& NarrowPhase,
-			CollisionStats::FStatData& StatData)
+		void ProduceOverlaps(FReal Dt, FNarrowPhase& NarrowPhase)
 		{
 			SCOPE_CYCLE_COUNTER(STAT_Collisions_BroadPhase);
 
@@ -63,7 +60,7 @@ namespace Chaos
 
 					if ((ParticleA != nullptr) && (ParticleB != nullptr))
 					{
-						ProduceOverlaps(Dt, ConstraintsArray, NarrowPhase, ParticleA, ParticleB, StatData);
+						ProduceOverlaps(Dt, NarrowPhase, ParticleA, ParticleB);
 					}
 				}
 			}
@@ -78,7 +75,7 @@ namespace Chaos
 						{
 							if (ParticleB != nullptr)
 							{
-								ProduceOverlaps(Dt, ConstraintsArray, NarrowPhase, ParticleA, ParticleB, StatData);
+								ProduceOverlaps(Dt, NarrowPhase, ParticleA, ParticleB);
 							}
 						}
 					}
@@ -89,21 +86,16 @@ namespace Chaos
 	private:
 		inline void ProduceOverlaps(
 			FReal Dt,
-			FCollisionConstraintsArray& ConstraintsArray,
 			FNarrowPhase& NarrowPhase,
 			FParticleHandle* ParticleA,
-			FParticleHandle* ParticleB,
-			CollisionStats::FStatData& StatData)
+			FParticleHandle* ParticleB)
 		{
 			const FAABB3& Box0 = ParticleA->WorldSpaceInflatedBounds();
 			const FAABB3& Box1 = ParticleB->WorldSpaceInflatedBounds();
 			if (Box0.Intersects(Box1))
 			{
-				NarrowPhase.GenerateCollisions(ConstraintsArray, Dt, ParticleA, ParticleB, CullDistance, true);
+				NarrowPhase.GenerateCollisions(Dt, ParticleA, ParticleB, CullDistance, true);
 			}
-
-			CHAOS_COLLISION_STAT(if (ConstraintsArray.Num()) { StatData.IncrementCountNP(ConstraintsArray.Num()); });
-			CHAOS_COLLISION_STAT(if (!ConstraintsArray.Num()) { StatData.IncrementRejectedNP(); });
 		}
 
 		const TArray<FParticlePair>* ParticlePairs;
diff --git a/Engine/Source/Runtime/Experimental/Chaos/Public/Chaos/Collision/ParticlePairCollisionDetector.h b/Engine/Source/Runtime/Experimental/Chaos/Public/Chaos/Collision/ParticlePairCollisionDetector.h
--- a/Engine/Source/Runtime/Experimental/Chaos/Public/Chaos/Collision/ParticlePairCollisionDetector.h
+++ b/Engine/Source/Runtime/Experimental/Chaos/Public/Chaos/Collision/ParticlePairCollisionDetector.h
@@ -19,7 +19,7 @@ namespace Chaos
 
 		FParticlePairBroadPhase& GetBroadPhase() { return BroadPhase; }
 
-		virtual void DetectCollisionsWithStats(const FReal Dt, CollisionStats::FStatData& StatData, FEvolutionResimCache*) override
+		virtual void DetectCollisions(const FReal Dt, FEvolutionResimCache* Unused) override
 		{
 			SCOPE_CYCLE_COUNTER(STAT_Collisions_Detect);
 			CHAOS_SCOPED_TIMER(DetectCollisions);
@@ -30,10 +30,12 @@ namespace Chaos
 				return;
 			}
 
-			CollisionContainer.UpdateConstraints(Dt);
+			CollisionContainer.BeginDetectCollisions();
 
 			// Collision detection pipeline: BroadPhase -> NarrowPhase -> Container
-			BroadPhase.ProduceOverlaps(Dt, CollisionContainer.GetConstraintsArray(), NarrowPhase, StatData);
+			BroadPhase.ProduceOverlaps(Dt, NarrowPhase);
+
+			CollisionContainer.EndDetectCollisions();
 		}
 
 	private:
diff --git a/Engine/Source/Runtime/Experimental/Chaos/Public/Chaos/Collision/SolverCollisionContainer.h b/Engine/Source/Runtime/Experimental/Chaos/Public/Chaos/Collision/SolverCollisionContainer.h
new file mode 100644
--- /dev/null
+++ b/Engine/Source/Runtime/Experimental/Chaos/Public/Chaos/Collision/SolverCollisionContainer.h
@@ -0,0 +1,59 @@
+// Copyright Epic Games, Inc. All Rights Reserved.
+#pragma once
+
+#include "CoreMinimal.h"
+#include "Chaos/CollisionResolutionTypes.h"
+#include "Chaos/Collision/CollisionApplyType.h"
+#include "Chaos/Evolution/SolverConstraintContainer.h"
+
+namespace Chaos
+{
+	namespace Collisions
+	{
+		struct FContactParticleParameters;
+		struct FContactIterationParameters;
+	}
+
+	class FPBDCollisionConstraint;
+	class FPBDCollisionSolver;
+	class FPBDCollisionSolverAdapter;
+	class FSolverBodyContainer;
+
+	/**
+	 * @brief A container of low-level data used to solve collision constraints
+	*/
+	class FPBDCollisionSolverContainer : public FConstraintSolverContainer
+	{
+	public:
+		FPBDCollisionSolverContainer();
+		~FPBDCollisionSolverContainer();
+
+		int32 NumSolvers() const { return CollisionSolvers.Num(); }
+
+		virtual void Reset(const int32 InMaxCollisions) override;
+
+		// Add a solver constraint for the specified constraint and gather the required solver data
+		FPBDCollisionSolver* AddConstraintSolver(FReal Dt, FPBDCollisionConstraint& Constraint, const int32 Particle0Level, const int32 Particle1Level, FSolverBodyContainer& SolverBodyContainer);
+
+		void SolveSwept(const FReal Dt);
+
+		bool SolvePositionSerial(const FReal Dt, const int32 It, const int32 NumIts, const int32 BeginIndex, const int32 EndIndex);
+		bool SolvePositionParallel(const FReal Dt, const int32 It, const int32 NumIts, const int32 BeginIndex, const int32 EndIndex);
+
+		bool SolveVelocitySerial(const FReal Dt, const int32 It, const int32 NumIts, const int32 BeginIndex, const int32 EndIndex);
+		bool SolveVelocityParallel(const FReal Dt, const int32 It, const int32 NumIts, const int32 BeginIndex, const int32 EndIndex);
+
+		void ScatterOutput(const FReal Dt, const int32 BeginIndex, const int32 EndIndex);
+
+	private:
+		void UpdatePositionShockPropagation(const FReal Dt, const int32 It, const int32 NumIts, const int32 BeginIndex, const int32 EndIndex);
+		void UpdateVelocityShockPropagation(const FReal Dt, const int32 It, const int32 NumIts, const int32 BeginIndex, const int32 EndIndex);
+		void SolveSwept(const FReal Dt, FPBDCollisionConstraint& Constraint, FPBDCollisionSolver& CollisionSolver);
+		bool SolvePositionImpl(const FReal Dt, const int32 It, const int32 NumIts, const int32 BeginIndex, const int32 EndIndex, const bool bParallel);
+		bool SolveVelocityImpl(const FReal Dt, const int32 It, const int32 NumIts, const int32 BeginIndex, const int32 EndIndex, const bool bParallel);
+		void ScatterOutputImpl(const FReal Dt, const int32 BeginIndex, const int32 EndIndex, const bool bParallel);
+
+		TArray<FPBDCollisionSolverAdapter> CollisionSolvers;
+		TArray<int32> SweptCollisionSolvers;
+	};
+}
diff --git a/Engine/Source/Runtime/Experimental/Chaos/Public/Chaos/Collision/SpatialAccelerationBroadPhase.h b/Engine/Source/Runtime/Experimental/Chaos/Public/Chaos/Collision/SpatialAccelerationBroadPhase.h
--- a/Engine/Source/Runtime/Experimental/Chaos/Public/Chaos/Collision/SpatialAccelerationBroadPhase.h
+++ b/Engine/Source/Runtime/Experimental/Chaos/Public/Chaos/Collision/SpatialAccelerationBroadPhase.h
@@ -3,7 +3,6 @@
 
 #include "Chaos/Collision/BroadPhase.h"
 #include "Chaos/Collision/CollisionConstraintFlags.h"
-#include "Chaos/Collision/CollisionReceiver.h"
 #include "Chaos/Collision/StatsData.h"
 #include "Chaos/Collision/NarrowPhase.h"
 #include "Chaos/ISpatialAccelerationCollection.h"
@@ -18,12 +17,12 @@
 namespace Chaos
 {
 	DECLARE_CYCLE_STAT_EXTERN(TEXT("Collisions::BroadPhase"), STAT_Collisions_SpatialBroadPhase, STATGROUP_ChaosCollision, CHAOS_API);
+	DECLARE_CYCLE_STAT_EXTERN(TEXT("Collisions::AABBTree"), STAT_Collisions_AABBTree, STATGROUP_ChaosCollision, CHAOS_API);
 	DECLARE_CYCLE_STAT_EXTERN(TEXT("Collisions::Filtering"), STAT_Collisions_Filtering, STATGROUP_ChaosCollision, CHAOS_API);
+	DECLARE_CYCLE_STAT_EXTERN(TEXT("Collisions::Restore"), STAT_Collisions_Restore, STATGROUP_ChaosCollision, CHAOS_API);
+	DECLARE_CYCLE_STAT_EXTERN(TEXT("Collisions::ComputeBoundsThickness"), STAT_Collisions_ComputeBoundsThickness, STATGROUP_ChaosCollision, CHAOS_API);
 	DECLARE_CYCLE_STAT_EXTERN(TEXT("Collisions::GenerateCollisions"), STAT_Collisions_GenerateCollisions, STATGROUP_ChaosCollision, CHAOS_API);
-	DECLARE_CYCLE_STAT_EXTERN(TEXT("Collisions::ReceiveCollisions"), STAT_Collisions_ReceiveCollisions, STATGROUP_ChaosCollision, CHAOS_API);
-
-	extern bool ChaosPerfHackIgnoreSleepingContacts;
-
+	
 	template <typename TPayloadType, typename T, int d>
 	class ISpatialAcceleration;
 
@@ -108,11 +107,11 @@ namespace Chaos
 		void ProduceOverlaps(
 			FReal Dt, 
 			FNarrowPhase& NarrowPhase, 
-			FAsyncCollisionReceiver& Receiver,
-			CollisionStats::FStatData& StatData,
 			IResimCacheBase* ResimCache
 			)
 		{
+			SCOPE_CYCLE_COUNTER(STAT_Collisions_SpatialBroadPhase);
+
 			if (!ensure(SpatialAcceleration))
 			{
 				// Must call SetSpatialAcceleration
@@ -121,19 +120,19 @@ namespace Chaos
 
 			if (const auto AABBTree = SpatialAcceleration->template As<TAABBTree<FAccelerationStructureHandle, TAABBTreeLeafArray<FAccelerationStructureHandle>>>())
 			{
-				ProduceOverlaps(Dt, *AABBTree, NarrowPhase, Receiver, StatData, ResimCache);
+				ProduceOverlaps(Dt, *AABBTree, NarrowPhase, ResimCache);
 			}
 			else if (const auto BV = SpatialAcceleration->template As<TBoundingVolume<FAccelerationStructureHandle>>())
 			{
-				ProduceOverlaps(Dt, *BV, NarrowPhase, Receiver, StatData, ResimCache);
+				ProduceOverlaps(Dt, *BV, NarrowPhase, ResimCache);
 			}
 			else if (const auto AABBTreeBV = SpatialAcceleration->template As<TAABBTree<FAccelerationStructureHandle, TBoundingVolume<FAccelerationStructureHandle>>>())
 			{
-				ProduceOverlaps(Dt, *AABBTreeBV, NarrowPhase, Receiver, StatData, ResimCache);
+				ProduceOverlaps(Dt, *AABBTreeBV, NarrowPhase, ResimCache);
 			}
 			else if (const auto Collection = SpatialAcceleration->template As<ISpatialAccelerationCollection<FAccelerationStructureHandle, FReal, 3>>())
 			{
-				Collection->PBDComputeConstraintsLowLevel(Dt, *this, NarrowPhase, Receiver, StatData, ResimCache);
+				Collection->PBDComputeConstraintsLowLevel(Dt, *this, NarrowPhase, ResimCache);
 			}
 			else
 			{
@@ -146,12 +145,10 @@ namespace Chaos
 			FReal Dt, 
 			const T_SPATIALACCELERATION& InSpatialAcceleration, 
 			FNarrowPhase& NarrowPhase, 
-			FAsyncCollisionReceiver& Receiver,
-			CollisionStats::FStatData& StatData,
 			IResimCacheBase* ResimCache
 			)
 		{
-			const bool bDisableParallelFor = StatData.IsEnabled() || bDisableCollisionParallelFor;
+			const bool bDisableParallelFor = bDisableCollisionParallelFor;
 			auto* EvolutionResimCache = static_cast<FEvolutionResimCache*>(ResimCache);
 			const bool bResimSkipCollision = ResimCache && ResimCache->IsResimming();
 			if(!bResimSkipCollision)
@@ -165,11 +162,9 @@ namespace Chaos
 					EntryCount += View.SOAViews[Index].Size();
 				}
 
-				Receiver.Prepare(EntryCount);
-
 				View.ParallelFor([&](auto& Particle1,int32 ActiveIdxIdx)
 				{
-					ProduceParticleOverlaps</*bResimming=*/false>(Dt,Particle1,InSpatialAcceleration,NarrowPhase,Receiver,StatData, ActiveIdxIdx);
+					ProduceParticleOverlaps</*bResimming=*/false>(Dt,Particle1,InSpatialAcceleration,NarrowPhase,ActiveIdxIdx);
 				},bDisableParallelFor);
 			}
 			else
@@ -183,14 +178,12 @@ namespace Chaos
 					EntryCount += View.SOAViews[Index].Size();
 				}
 
-				Receiver.Prepare(EntryCount);
-
 				View.ParallelFor(
 					[&](auto& Particle1,int32 ActiveIdxIdx)
 				{
 					//TODO: use transient handle
 					FGenericParticleHandleHandleImp GenericHandle(Particle1.Handle());
-					ProduceParticleOverlaps</*bResimming=*/true>(Dt,GenericHandle,InSpatialAcceleration,NarrowPhase,Receiver,StatData, ActiveIdxIdx);
+					ProduceParticleOverlaps</*bResimming=*/true>(Dt,GenericHandle,InSpatialAcceleration,NarrowPhase,ActiveIdxIdx);
 				},bDisableParallelFor);
 			}
 		}
@@ -204,25 +197,18 @@ namespace Chaos
 		    THandle& Particle1,
 		    const T_SPATIALACCELERATION& InSpatialAcceleration,
 		    FNarrowPhase& NarrowPhase,
-		    FAsyncCollisionReceiver& Receiver,
-		    CollisionStats::FStatData& StatData,
 			int32 EntryIndex)
 		{
-			CHAOS_COLLISION_STAT(StatData.IncrementSimulatedParticles());
-
 			//Durning non-resim we must pass rigid particles
 			static_assert(bIsResimming || THandle::StaticType() == EParticleType::Rigid, "During non-resim we expect rigid particles");
 
 			TArray<FAccelerationStructureHandle> PotentialIntersections;
 
-			//If we are in a resim, this particle must be diverged and so it better be marked as enabled during resim
-			ensure(!bIsResimming || Particle1.EnabledDuringResim());
-
 			if (bIsResimming || (Particle1.ObjectState() == EObjectStateType::Dynamic || Particle1.ObjectState() == EObjectStateType::Sleeping))
 			{
 				const bool bBody1Bounded = HasBoundingBox(Particle1);
 				{
-					SCOPE_CYCLE_COUNTER(STAT_Collisions_SpatialBroadPhase);
+					SCOPE_CYCLE_COUNTER(STAT_Collisions_AABBTree);
 					if (bBody1Bounded)
 					{
 						FCollisionFilterData ParticleSimData;
@@ -231,8 +217,6 @@ namespace Chaos
 						const FReal Box1Thickness = ComputeBoundsThickness(Particle1, Dt, BoundsThickness, BoundsThicknessVelocityInflation).Size();
 						const FAABB3 Box1 = ComputeWorldSpaceBoundingBox<FReal>(Particle1).ThickenSymmetrically(FVec3(Box1Thickness));
 
-						CHAOS_COLLISION_STAT(StatData.RecordBoundsData(Box1));
-
 						FSimOverlapVisitor OverlapVisitor(Particle1.UniqueIdx(), ParticleSimData, PotentialIntersections);
 						InSpatialAcceleration.Overlap(Box1, OverlapVisitor);
 					}
@@ -246,13 +230,10 @@ namespace Chaos
 							PotentialIntersections.Add(Elem.Payload);
 						}
 					}
-
-					CHAOS_COLLISION_STAT(StatData.RecordBroadphasePotentials(PotentialIntersections.Num()))
 				}
 
 				SCOPE_CYCLE_COUNTER(STAT_Collisions_Filtering);
 				const int32 NumPotentials = PotentialIntersections.Num();
-				FCollisionConstraintsArray NewConstraints;
 				for (int32 i = 0; i < NumPotentials; ++i)
 				{
 					auto& Particle2 = *PotentialIntersections[i].GetGeometryParticleHandle_PhysicsThread();
@@ -293,6 +274,11 @@ namespace Chaos
 					{
 						continue;
 					}
+					
+					if ((Particle1.ObjectState() == EObjectStateType::Sleeping) && (Particle2.ObjectState() == EObjectStateType::Sleeping))
+					{
+						continue;
+					}
 
 					if (Particle1.Handle() == Particle2.Handle())
 					{
@@ -305,43 +291,68 @@ namespace Chaos
 						continue;
 					}
 
-					const bool bSecondParticleWillHaveAnswer = !bIsResimming || Particle2.SyncState() == ESyncState::HardDesync;
-					// Sleeping vs dynamic gets picked up by the other direction.
-					const bool bIsParticle2Dynamic = Particle2.CastToRigidParticle() && Particle2.ObjectState() == EObjectStateType::Dynamic;
-
-					bool bSkipSleepingContact = bSecondParticleWillHaveAnswer;
-					if (ChaosPerfHackIgnoreSleepingContacts)
+					// Note:
+					// 
+					// Normally (not resimming) we iterate over dynamic particles, so:
+					// - Particle1 is always dynamic, but may be asleep
+					// - Particle2 may be static, kinematic or dynamic, but may be asleep
+					// 
+					// When resimming we iterate over "desynced" particles which may be kinematic so:
+					// - Particle1 is always desynced
+					// - Particle2 may also be desynced, in which case we will also visit the opposite ordering regardless of dynamic/kinematic status
+					// - Particle1 may be static, kinematic or dynamic, but may be asleep
+					// - Particle2 may be static, kinematic or dynamic, but may be asleep
+					// 
+					// Even though Particle1 may be kinematic when resimming, we want to create the contacts in the original order (i.e., dynamic first)
+					//
+					const bool bIsParticle2Dynamic = (Particle2.ObjectState() == EObjectStateType::Dynamic) || (Particle2.ObjectState() == EObjectStateType::Sleeping);
+
+					// Used to determine a winner in cases where we will visit particle pairs in both orders
+					const bool bIsParticle1Preferred = (Particle2.ParticleID() < Particle1.ParticleID());
+
+					bool bAcceptParticlePair = false;
+					if (!bIsResimming)
 					{
-						const bool bIsParticle2Kinematic = Particle2.CastToKinematicParticle() &&
-							(Particle2.ObjectState() == EObjectStateType::Kinematic &&
-								(Particle2.CastToKinematicParticle()->V().SizeSquared() > 1e-4 ||
-									Particle2.Geometry()->GetType() == FCapsule::StaticType()));
-
-						bSkipSleepingContact &= !bIsParticle2Kinematic;
+						// When not resimming, accept pair if any of the following are true
+						// (1) Both particles are dynamic|asleep and Particle1 is "preferred" (to prevent duplicate entries when we visit the pair in the opposite order)
+						// (2) Particle1 is dynamic|asleep and Particle2 is kinematic (we won't visit the opposite order)
+						// Then both conditions can be combined and reduced to:
+						bAcceptParticlePair = (bIsParticle1Preferred || !bIsParticle2Dynamic);
 					}
 					else
 					{
-						bSkipSleepingContact &= bIsParticle2Dynamic;
+						// When resimming, accept pair if any of the following are true
+						// (1) Both particles are dynamic|asleep and desynced and Particle1 is preferred (we will visit both pair orderings)
+						// (2) Both particles are dynamic|asleep but Particle2 is not desynced (we won't visit the opposite order)
+						// (3) Particle1 is dynamic|sleeping and Particle2 is kinematic (regardless of desync state, see (4))
+						// (4) Particle1 is kinematic and Particle2 is dynamic but not desynced (we won't visit the opposite order)
+						// We should skip non-dynamic pairs should they show up (implicit in the logic below)
+						const bool bIsParticle1Dynamic = (Particle1.ObjectState() == EObjectStateType::Dynamic) || (Particle1.ObjectState() == EObjectStateType::Sleeping);
+						const bool bIsParticle2Desynced = (Particle2.SyncState() == ESyncState::HardDesync);
+						bAcceptParticlePair = 
+							(bIsParticle1Dynamic && bIsParticle2Dynamic && (bIsParticle1Preferred || !bIsParticle2Desynced))		// Case (1) and (2)
+							|| (bIsParticle1Dynamic && !bIsParticle2Dynamic)														// Case (3)
+							|| (!bIsParticle1Dynamic && bIsParticle2Dynamic && !bIsParticle2Desynced);								// Case (4)
+
+						// @todo(chaos): for Case (4) above, we should be switching the order of the particles before continuing
+						// if we want this to be properly deterministic - we always put dynamics first in the non-resim path...
 					}
 
-					if (Particle1.ObjectState() == EObjectStateType::Sleeping && bSkipSleepingContact)
+					if (!bAcceptParticlePair)
 					{
-						//question: if !bSecondParticleWillHaveAnswer do we need to reorder constraint?
 						continue;
 					}
 
-					// Make sure we don't add a second set of constaint for the same body pair (with the body order flipped)
-					const bool bBody2Bounded = HasBoundingBox(Particle2);
-					if (bBody1Bounded == bBody2Bounded && bIsParticle2Dynamic)
+					// Try to restore all the contacts for this pair. This will only succeed if they have not moved (within some threshold)
 					{
-						if (Particle1.ParticleID() < Particle2.ParticleID() && bSecondParticleWillHaveAnswer)
+						SCOPE_CYCLE_COUNTER(STAT_Collisions_Restore);
+						if (NarrowPhase.TryRestoreCollisions(Dt, Particle1.Handle(), Particle2.Handle()))
 						{
-							//question: if !bSecondParticleWillHaveAnswer do we need to reorder constraint?
 							continue;
 						}
 					}
-				
 					
+					// If we get here, we need to run the narrow phase to possibly generate new contacts, or refresh existing ones
 					{
 						SCOPE_CYCLE_COUNTER(STAT_Collisions_GenerateCollisions);
 
@@ -358,23 +369,10 @@ namespace Chaos
 
 						// Generate constraints for the potentially overlapping shape pairs. Also run collision detection to generate
 						// the contact position and normal (for contacts within CullDistance) for use in collision callbacks.
-						NarrowPhase.GenerateCollisions(NewConstraints, Dt, Particle1.Handle(), Particle2.Handle(), NetCullDistance, false);
+						NarrowPhase.GenerateCollisions(Dt, Particle1.Handle(), Particle2.Handle(), NetCullDistance, false);
 					}
 				}
-
-				{
-					SCOPE_CYCLE_COUNTER(STAT_Collisions_ReceiveCollisions);
-
-					CHAOS_COLLISION_STAT(if (NewConstraints.Num()) { StatData.IncrementCountNP(NewConstraints.Num()); });
-					CHAOS_COLLISION_STAT(if (!NewConstraints.Num()) { StatData.IncrementRejectedNP(); });
-
-					// We are probably running in a parallel task here. The Receiver collects the contacts from all the tasks 
-					// and passes them to theconstraint container in serial.
-					Receiver.ReceiveCollisions(MoveTemp(NewConstraints), EntryIndex);
-				}
 			}
-
-			CHAOS_COLLISION_STAT(StatData.FinalizeData());
 		}
 
 		const FPBDRigidsSOAs& Particles;
diff --git a/Engine/Source/Runtime/Experimental/Chaos/Public/Chaos/Collision/SpatialAccelerationCollisionDetector.h b/Engine/Source/Runtime/Experimental/Chaos/Public/Chaos/Collision/SpatialAccelerationCollisionDetector.h
--- a/Engine/Source/Runtime/Experimental/Chaos/Public/Chaos/Collision/SpatialAccelerationCollisionDetector.h
+++ b/Engine/Source/Runtime/Experimental/Chaos/Public/Chaos/Collision/SpatialAccelerationCollisionDetector.h
@@ -3,7 +3,6 @@
 
 #include "Chaos/Collision/CollisionDetector.h"
 #include "Chaos/Collision/CollisionContext.h"
-#include "Chaos/Collision/CollisionReceiver.h"
 #include "Chaos/Collision/NarrowPhase.h"
 #include "Chaos/Collision/SpatialAccelerationBroadPhase.h"
 #include "Chaos/EvolutionResimCache.h"
@@ -22,7 +21,7 @@ namespace Chaos
 
 		FSpatialAccelerationBroadPhase& GetBroadPhase() { return BroadPhase; }
 
-		virtual void DetectCollisionsWithStats(const FReal Dt, CollisionStats::FStatData& StatData, FEvolutionResimCache* ResimCache) override
+		virtual void DetectCollisions(const FReal Dt, FEvolutionResimCache* ResimCache) override
 		{
 			SCOPE_CYCLE_COUNTER(STAT_Collisions_Detect);
 			CHAOS_SCOPED_TIMER(DetectCollisions);
@@ -33,24 +32,20 @@ namespace Chaos
 				return;
 			}
 
-			CollisionContainer.UpdateConstraints(Dt);
+			CollisionContainer.BeginDetectCollisions();
 
 			// Collision detection pipeline: BroadPhase -[parallel]-> NarrowPhase -[parallel]-> Receiver -[serial]-> Container
-			FAsyncCollisionReceiver Receiver(CollisionContainer, ResimCache);
-			BroadPhase.ProduceOverlaps(Dt, NarrowPhase, Receiver, StatData, ResimCache);
+			BroadPhase.ProduceOverlaps(Dt, NarrowPhase, ResimCache);
+
+			CollisionContainer.EndDetectCollisions();
+
+			// If we have a resim cache restore and save contacts
 			if(ResimCache)
 			{
-				// Push the resim constraints into slot zero with the first particle. Doesn't really matter where they go at this
-				// point as long as it is consistent so we don't need to sort the constraints in ProcessCollisions
-				if(Receiver.CacheNum() == 0)
-				{
-					// In case we have zero particles but somehow have constraints
-					Receiver.Prepare(1);
-				}
-
-				Receiver.AppendCollisions(ResimCache->GetAndSanitizeConstraints(), 0);
+				CollisionContainer.GetConstraintAllocator().AddResimConstraints(ResimCache->GetAndSanitizeConstraints());
+
+				ResimCache->SaveConstraints(CollisionContainer.GetConstraints());
 			}
-			Receiver.ProcessCollisions();
 		}
 
 	private:
diff --git a/Engine/Source/Runtime/Experimental/Chaos/Public/Chaos/CollisionResolution.h b/Engine/Source/Runtime/Experimental/Chaos/Public/Chaos/CollisionResolution.h
--- a/Engine/Source/Runtime/Experimental/Chaos/Public/Chaos/CollisionResolution.h
+++ b/Engine/Source/Runtime/Experimental/Chaos/Public/Chaos/CollisionResolution.h
@@ -20,10 +20,9 @@ namespace Chaos
 	class TPlane;
 
 	class FCollisionContext;
-	class FCollisionConstraintBase;
 	class FHeightField;
 	class FImplicitObject;
-	class FRigidBodyPointContactConstraint;
+	class FPBDCollisionConstraint;
 	class FTriangleMeshImplicitObject;
 
 
@@ -35,17 +34,23 @@ namespace Chaos
 		//
 
 		// Update the constraint by re-running collision detection on the shape pair.
+		template<ECollisionUpdateType UpdateType>
+		void CHAOS_API UpdateConstraintFromGeometry(FPBDCollisionConstraint& Constraint, const FRigidTransform3& ParticleTransform0, const FRigidTransform3& ParticleTransform1, const FReal Dt);
 
-		template<ECollisionUpdateType UpdateType, typename RigidBodyContactConstraint>
-		void CHAOS_API UpdateConstraintFromGeometry(RigidBodyContactConstraint& Constraint, const FRigidTransform3& ParticleTransform0, const FRigidTransform3& ParticleTransform1, const FReal Dt);
+		// Update the constraint by re-running collision detection on the shape pair.
+		template<ECollisionUpdateType UpdateType>
+		void CHAOS_API UpdateConstraintFromGeometrySwept(FPBDCollisionConstraint& Constraint, const FRigidTransform3& ParticleTransform0, const FRigidTransform3& ParticleTransform1, const FReal Dt);
 
 		// Create constraints for the particle pair. This could create multiple constraints: one for each potentially colliding shape pair in multi-shape particles.
-		void CHAOS_API ConstructConstraints(TGeometryParticleHandle<FReal, 3>* Particle0, TGeometryParticleHandle<FReal, 3>* Particle1, const FImplicitObject* Implicit0, const FBVHParticles* Simplicial0, const FImplicitObject* Implicit1, const FBVHParticles* Simplicial1, const FRigidTransform3& Transform0, const FRigidTransform3& Transform1, const FReal CullDistance, const FReal dT,const FCollisionContext& Context, FCollisionConstraintsArray& NewConstraints);
+		void CHAOS_API ConstructConstraints(TGeometryParticleHandle<FReal, 3>* Particle0, TGeometryParticleHandle<FReal, 3>* Particle1, const FImplicitObject* Implicit0, const FBVHParticles* Simplicial0, const FImplicitObject* Implicit1, const FBVHParticles* Simplicial1, const FRigidTransform3& ParticleWorldTransform0, const FRigidTransform3& Transform0, const FRigidTransform3& ParticleWorldTransform1, const FRigidTransform3& Transform1, const FReal CullDistance, const FReal Dt,const FCollisionContext& Context);
+
+		// Attempt to restore all constraints for the particle pair. This will fail is the pair were not previously colliding, or if their relative transforms have changed. Return true if constraints were restored.
+		bool CHAOS_API TryRestoreConstraints(FConstGenericParticleHandle Particle0, FConstGenericParticleHandle Particle1, const FRigidTransform3& ParticleWorldTransform0, const FRigidTransform3& ParticleWorldTransform1, const FReal Dt, const FCollisionContext& Context);
 
 
 		// @todo(chaos): this is only called in tests - should it really be exposed?
 		template<ECollisionUpdateType UpdateType>
-		void UpdateLevelsetLevelsetConstraint(const FRigidTransform3& WorldTransform0, const FRigidTransform3& WorldTransform1, const FReal Dt, FRigidBodyPointContactConstraint& Constraint);
+		void UpdateLevelsetLevelsetConstraint(const FRigidTransform3& WorldTransform0, const FRigidTransform3& WorldTransform1, const FReal Dt, FPBDCollisionConstraint& Constraint);
 
 		// Reset per-frame collision stat counters
 		void CHAOS_API ResetChaosCollisionCounters();
diff --git a/Engine/Source/Runtime/Experimental/Chaos/Public/Chaos/CollisionResolutionUtil.h b/Engine/Source/Runtime/Experimental/Chaos/Public/Chaos/CollisionResolutionUtil.h
--- a/Engine/Source/Runtime/Experimental/Chaos/Public/Chaos/CollisionResolutionUtil.h
+++ b/Engine/Source/Runtime/Experimental/Chaos/Public/Chaos/CollisionResolutionUtil.h
@@ -10,7 +10,7 @@ namespace Chaos
 	class FBVHParticles;
 	class FContactPoint;
 	class FImplicitObject;
-	class FRigidBodyPointContactConstraint;
+	class FPBDCollisionConstraint;
 
 	namespace Collisions
 	{
diff --git a/Engine/Source/Runtime/Experimental/Chaos/Public/Chaos/ConstraintHandle.h b/Engine/Source/Runtime/Experimental/Chaos/Public/Chaos/ConstraintHandle.h
--- a/Engine/Source/Runtime/Experimental/Chaos/Public/Chaos/ConstraintHandle.h
+++ b/Engine/Source/Runtime/Experimental/Chaos/Public/Chaos/ConstraintHandle.h
@@ -8,31 +8,90 @@
 namespace Chaos
 {
 	class FPBDConstraintContainer;
+	class FPBDIndexedConstraintContainer;
+
+	/**
+	 * @brief A type id for constraint handles to support safe up/down casting (including intermediate classes in the hierrachy)
+	 *
+	 * Every constraint handle must provide a StaticType() member which gives the constraint type
+	 * name and base class chain.
+	 * 
+	 * Every constraint container must provide a GetConstraintHandleType() method to get the constraint 
+	 * type for handles that reference the container.
+	*/
+	class FConstraintHandleTypeID
+	{
+	public:
+		FConstraintHandleTypeID(const FName& InName, const FConstraintHandleTypeID* InBaseType = nullptr)
+			: TypeName(InName)
+			, BaseType(InBaseType)
+		{
+		}
+
+		/**
+		 * @brief An invalid constraint handle type for initialization and invalidation
+		*/
+		static const FConstraintHandleTypeID InvalidTypeID()
+		{
+			return FConstraintHandleTypeID(NAME_None);
+		}
+
+		/**
+		 * @brief Whether this type can be cast to the specified type
+		*/
+		bool IsA(const FConstraintHandleTypeID& TypeID) const
+		{
+			if (TypeID.TypeName == TypeName)
+			{
+				return true;
+			}
+			if (TypeID.BaseType != nullptr)
+			{
+				return IsA(*TypeID.BaseType);
+			}
+			return false;
+		}
+
+	private:
+		FName TypeName;
+		const FConstraintHandleTypeID* BaseType;
+	};
 
 
 	/**
-	 * Base class for constraint handles.
-	 */
+	 * @brief Base class for constraint handles.
+	 * 
+	 * Constraints are referenced by handle in the constraint graph. 
+	 * Constraint handles allow us to support different allocation and storage policies for constraints.
+	 * E.g., heap-allocated constraints, array-based constraints etc.
+	 * 
+	 * @see FIndexedConstraintHandle, FIntrusiveConstraintHandle
+	*/
 	class CHAOS_API FConstraintHandle
 	{
 	public:
 		using FGeometryParticleHandle = TGeometryParticleHandle<FReal, 3>;
 
-		FConstraintHandle() : ConstraintContainer(nullptr), ConstraintIndex(INDEX_NONE), ConstraintGraphIndex(INDEX_NONE) { }
-
-		FConstraintHandle(FPBDConstraintContainer* InContainer, int32 InConstraintIndex): ConstraintContainer(InContainer), ConstraintIndex(InConstraintIndex), ConstraintGraphIndex(INDEX_NONE) {}
+		FConstraintHandle() 
+			: ConstraintContainer(nullptr)
+			, ConstraintGraphIndex(INDEX_NONE) 
+		{
+		}
 
-		virtual ~FConstraintHandle() {}
+		FConstraintHandle(FPBDConstraintContainer* InContainer)
+			: ConstraintContainer(InContainer)
+			, ConstraintGraphIndex(INDEX_NONE)
+		{
+		}
 
-		bool IsValid() const
+		virtual ~FConstraintHandle()
 		{
-			// @todo(chaos): why does IsValid() also check IsEnabled()?
-			return (ConstraintContainer != nullptr) && (ConstraintIndex != INDEX_NONE) && IsEnabled();
 		}
 
-		int32 GetConstraintIndex() const
+		virtual bool IsValid() const
 		{
-			return ConstraintIndex;
+			// @todo(chaos): why does IsValid() also check IsEnabled()?
+			return (ConstraintContainer != nullptr) && IsEnabled();
 		}
 
 		int32 GetConstraintGraphIndex() const
@@ -50,26 +109,142 @@ namespace Chaos
 			return (ConstraintGraphIndex != INDEX_NONE);
 		}
 
+		virtual void SetEnabled(bool InEnabled) = 0;
 
-		// Implemented in ConstraintContainer.h
-		int32 GetContainerId() const;
+		virtual bool IsEnabled() const = 0;
 
-		// Implemented in ConstraintContainer.h
-		void SetEnabled(bool InEnabled);
+		virtual bool IsSleeping() const { return false; }
+		virtual void SetIsSleeping(const bool bInIsSleeping) {}
 
 		// Implemented in ConstraintContainer.h
-		bool IsEnabled() const;
+		int32 GetContainerId() const;
 
 		// Implemented in ConstraintContainer.h
 		template<typename T>  T* As();
 		template<typename T>  const T* As() const;
 
+		static const FConstraintHandleTypeID& StaticType()
+		{
+			static FConstraintHandleTypeID STypeID(TEXT("FConstraintHandle"), nullptr);
+			return STypeID;
+		}
+
 	protected:
 		friend class FPBDConstraintContainer;
 
 		FPBDConstraintContainer* ConstraintContainer;
+		
+		// @todo(chaos): move constraint graph index to base constraint container
+		int32 ConstraintGraphIndex;
+	};
+
+
+	/**
+	 * @brief Base class for constraints that are allocated at permanent memory addresses and inherit the handle.
+	 *
+	 * Intended for use by constraint types that are allocated on the heap or in a block allocator and therefore have a persistent
+	 * address (as opposed to array-based containers where the array could relocate). The constraint class should inherit
+	 * this handle class. This effectively eliminates the handle, reducing cache misses and allocations.
+	*/
+	class CHAOS_API FIntrusiveConstraintHandle : public FConstraintHandle
+	{
+	public:
+		FIntrusiveConstraintHandle()
+			: FConstraintHandle()
+		{
+		}
+
+		void SetContainer(FPBDConstraintContainer* InContainer)
+		{
+			ConstraintContainer = InContainer;
+		}
+
+		static const FConstraintHandleTypeID& StaticType()
+		{
+			static FConstraintHandleTypeID STypeID(TEXT("FIntrusiveConstraintHandle"), &FConstraintHandle::StaticType());
+			return STypeID;
+		}
+	};
+
+
+	/**
+	 * @brief Base class for constraints that are allocated at permanent memory addresses and inherit the handle.
+	 * 
+	 * @see FIntrusiveConstraintHandle
+	 * 
+	 * @tparam T_CONSTRAINT The constraint type
+	*/
+	template<typename T_CONSTRAINT>
+	class CHAOS_API TIntrusiveConstraintHandle : public FIntrusiveConstraintHandle
+	{
+	public:
+		using FConstraint = T_CONSTRAINT;
+
+		TIntrusiveConstraintHandle()
+			: FIntrusiveConstraintHandle()
+		{
+		}
+
+		void SetContainer(FPBDConstraintContainer* InContainer)
+		{
+			ConstraintContainer = InContainer;
+		}
+
+		FConstraint* GetConstraint()
+		{
+			return static_cast<FConstraint*>(this);
+		}
+
+		const FConstraint* GetConstraint() const
+		{
+			return static_cast<const FConstraint*>(this);
+		}
+	};
+
+	/**
+	 * Base class for handles to constraints in an index-based container
+	 */
+	class CHAOS_API FIndexedConstraintHandle : public FConstraintHandle
+	{
+	public:
+		using FGeometryParticleHandle = TGeometryParticleHandle<FReal, 3>;
+
+		FIndexedConstraintHandle() 
+			: FConstraintHandle()
+			, ConstraintIndex(INDEX_NONE)
+		{
+		}
+
+		FIndexedConstraintHandle(FPBDConstraintContainer* InContainer, int32 InConstraintIndex)
+			: FConstraintHandle(InContainer)
+			, ConstraintIndex(InConstraintIndex)
+		{
+		}
+
+		virtual ~FIndexedConstraintHandle()
+		{
+		}
+
+		virtual bool IsValid() const override
+		{
+			return (ConstraintIndex != INDEX_NONE) && FConstraintHandle::IsValid();
+		}
+
+		int32 GetConstraintIndex() const
+		{
+			return ConstraintIndex;
+		}
+
+		static const FConstraintHandleTypeID& StaticType()
+		{
+			static FConstraintHandleTypeID STypeID(TEXT("FIndexedConstraintHandle"), &FConstraintHandle::StaticType());
+			return STypeID;
+		}
+
+	protected:
+		friend class FPBDIndexedConstraintContainer;
+
 		int32 ConstraintIndex;
-		int32 ConstraintGraphIndex; // @todo(chaos): move constraint graph index to base constraint container
 	};
 
 
@@ -77,20 +252,38 @@ namespace Chaos
 	 * Utility base class for ConstraintHandles. Provides basic functionality common to most constraint containers.
 	 */
 	template<typename T_CONTAINER>
-	class CHAOS_API TContainerConstraintHandle : public FConstraintHandle
+	class CHAOS_API TIndexedContainerConstraintHandle : public FIndexedConstraintHandle
 	{
 	public:
-		using Base = FConstraintHandle;
+		using Base = FIndexedConstraintHandle;
 		using FGeometryParticleHandle = typename Base::FGeometryParticleHandle;
 		using FConstraintContainer = T_CONTAINER;
 
-		TContainerConstraintHandle()
+		TIndexedContainerConstraintHandle()
+			: FIndexedConstraintHandle()
 		{
 		}
 		
-		TContainerConstraintHandle(FConstraintContainer* InConstraintContainer, int32 InConstraintIndex)
-			: FConstraintHandle(InConstraintContainer, InConstraintIndex)
+		TIndexedContainerConstraintHandle(FConstraintContainer* InConstraintContainer, int32 InConstraintIndex)
+			: FIndexedConstraintHandle(InConstraintContainer, InConstraintIndex)
+		{
+		}
+
+		inline virtual void SetEnabled(bool bInEnabled) override
+		{
+			if (ConcreteContainer() != nullptr)
+			{
+				ConcreteContainer()->SetConstraintEnabled(ConstraintIndex, bInEnabled);
+			}
+		}
+
+		inline virtual bool IsEnabled() const override
 		{
+			if (ConcreteContainer() != nullptr)
+			{
+				return ConcreteContainer()->IsConstraintEnabled(ConstraintIndex);
+			}
+			return false;
 		}
 
 		// @todo(chaos): Make this a virtual on FConstraintContainer and move to base class
@@ -127,8 +320,7 @@ namespace Chaos
 		using FConstraintContainerHandle = typename FConstraintContainer::FConstraintContainerHandle;
 
 		FConstraintContainerHandle* AllocHandle(FConstraintContainer* ConstraintContainer, int32 ConstraintIndex) { return new FConstraintContainerHandle(ConstraintContainer, ConstraintIndex); }
-		template<class TYPE>
-		FConstraintContainerHandle* AllocHandle(FConstraintContainer* ConstraintContainer, int32 ConstraintIndex) { return new FConstraintContainerHandle(ConstraintContainer, ConstraintIndex, TYPE::StaticType()); }
+
 		void FreeHandle(FConstraintContainerHandle* Handle) { delete Handle; }
 	};
 }
diff --git a/Engine/Source/Runtime/Experimental/Chaos/Public/Chaos/Evolution/PBDMinEvolution.h b/Engine/Source/Runtime/Experimental/Chaos/Public/Chaos/Evolution/PBDMinEvolution.h
--- a/Engine/Source/Runtime/Experimental/Chaos/Public/Chaos/Evolution/PBDMinEvolution.h
+++ b/Engine/Source/Runtime/Experimental/Chaos/Public/Chaos/Evolution/PBDMinEvolution.h
@@ -4,6 +4,7 @@
 #include "Chaos/Core.h"
 #include "Chaos/ArrayCollectionArray.h"
 #include "Chaos/Evolution/SimulationSpace.h"
+#include "Chaos/Evolution/SolverDatas.h"
 #include "Chaos/ParticleHandleFwd.h"
 
 
@@ -108,11 +109,11 @@ namespace Chaos
 		void IntegrateImplISPC(FReal Dt);
 		void ApplyKinematicTargets(FReal Dt, FReal StepFraction);
 		void DetectCollisions(FReal Dt);
-		void PrepareIteration(FReal Dt);
-		void UnprepareIteration(FReal Dt);
-		void ApplyConstraints(FReal Dt);
+		void GatherInput(FReal Dt);
+		void ScatterOutput(FReal Dt);
+		void ApplyConstraintsPhase1(FReal Dt);
 		void UpdateVelocities(FReal Dt);
-		void ApplyPushOutConstraints(FReal Dt);
+		void ApplyConstraintsPhase2(FReal Dt);
 		void UpdatePositions(FReal Dt);
 
 		FRigidParticleSOAs& Particles;
@@ -123,6 +124,7 @@ namespace Chaos
 
 		TArray<FSimpleConstraintRule*> ConstraintRules;
 		TArray<FSimpleConstraintRule*> PrioritizedConstraintRules;
+		FPBDIslandSolverData SolverData;
 
 		int32 NumApplyIterations;
 		int32 NumApplyPushOutIterations;
diff --git a/Engine/Source/Runtime/Experimental/Chaos/Public/Chaos/Evolution/SolverBody.h b/Engine/Source/Runtime/Experimental/Chaos/Public/Chaos/Evolution/SolverBody.h
new file mode 100644
--- /dev/null
+++ b/Engine/Source/Runtime/Experimental/Chaos/Public/Chaos/Evolution/SolverBody.h
@@ -0,0 +1,426 @@
+// Copyright Epic Games, Inc. All Rights Reserved.
+#pragma once
+
+#include "Chaos/Core.h"
+
+namespace Chaos
+{
+	class FSolverBody;
+	class FSolverBodyContainer;
+
+	// A pair of pointers so solver bodies
+	// @note Pointers are only valid for the Constraint Solving phase of the tick
+	using FSolverBodyPtrPair = TVector<FSolverBody*, 2>;
+
+	/**
+	 * Used by the constraint solver loop to cache all state for a particle and accumulate solver results.
+	 * Uses a gather/scatter mechanism to read/write data to the particle SOAs at the beginning/end of the constraint solve.
+	 * Constraint solver algorithms, and collision Update functions are implemented to use FSolverBody, and do not 
+	 * directly read/write to the particle handles. Constraint Solvers will modify P(), Q(), V() and W() via 
+	 * ApplyTransformDelta() and other methods.
+	 * 
+	 * There is one SolverBody for each particle in an island. Most constraint solvers will actually wrap the
+	 * FSolverBody in FConstraintSolverBody, which allows us to apply per-constraint modifiers to the Solver Body.
+	 *
+	 * @note the X(), P(), R(), Q() accessors on FSolverBody return the Center of Mass positions and rotations, in contrast
+	 * to the Particle methods which gives Actor positions and rotations. This is because the Constraint Solvers all calculate
+	 * impulses and position corrections relative to the center of mass.
+	 * 
+	 * @todo(chaos): layout for cache
+	 */
+	class FSolverBody
+	{
+	public:
+
+		/**
+		 * @brief Create an empty solver body
+		 * @note This is only used by unit tests
+		*/
+		FSolverBody();
+
+		/**
+		 * @brief Calculate and set the velocity and angular velocity from the net transform delta
+		*/
+		void SetImplicitVelocity(FReal Dt)
+		{
+			if (IsDynamic())
+			{
+				State.V = FVec3::CalculateVelocity(State.X, State.P, Dt);
+				State.W = FRotation3::CalculateAngularVelocity(State.R, State.Q, Dt);
+			}
+		}
+
+		/**
+		 * @brief Get the inverse mass
+		*/
+		FReal InvM() const { return State.InvM; }
+
+		/**
+		 * @brief Set the inverse mass
+		*/
+		void SetInvM(FReal InInvM) { State.InvM = InInvM; }
+
+		/**
+		 * @brief Get the world-space inverse inertia
+		*/
+		const FMatrix33& InvI() const { return State.InvI; }
+
+		/**
+		 * @brief Set the world-space inverse inertia
+		*/
+		void SetInvI(const FMatrix33& InInvI) { State.InvI = InInvI; }
+
+		/**
+		 * @brief Get the local-space inverse inertia (diagonal elements)
+		*/
+		const FVec3& InvILocal() const { return State.InvILocal; }
+
+		/**
+		 * @brief Set the local-space inverse inertia (diagonal elements)
+		*/
+		void SetInvILocal(const FVec3& InInvILocal)
+		{ 
+			State.InvILocal = InInvILocal; 
+			UpdateRotationDependentState();
+		}
+
+		/**
+		 * @brief The current CoM transform
+		*/
+		FRigidTransform3 CoMTransform() const { return FRigidTransform3(P(), Q()); }
+
+		/**
+		 * @brief Pre-integration world-space center of mass position
+		*/
+		const FVec3& X() const { return State.X; }
+		void SetX(const FVec3& InX) { State.X = InX; }
+
+		/**
+		 * @brief Pre-integration world-space center of mass rotation
+		*/
+		const FRotation3& R() const { return State.R; }
+		void SetR(const FRotation3& InR) { State.R = InR; }
+
+		/**
+		 * @brief World-space center of mass position
+		*/
+		const FVec3& P() const { return State.P; }
+		void SetP(const FVec3& InP) { State.P = InP; }
+
+		/**
+		 * @brief World-space center of mass rotation
+		*/
+		const FRotation3& Q() const { return State.Q; }
+		void SetQ(const FRotation3& InQ) { State.Q = InQ; }
+
+		/**
+		 * @brief World-space center of mass velocity
+		*/
+		const FVec3& V() const { return State.V; }
+		void SetV(const FVec3& InV) { State.V = InV; }
+
+		/**
+		 * @brief World-space center of mass angular velocity
+		*/
+		const FVec3& W() const { return State.W; }
+		void SetW(const FVec3& InW) { State.W = InW; }
+
+		const FVec3& CoM() const { return State.CoM; }
+		void SetCoM(const FVec3& InCoM) { State.CoM = InCoM; }
+
+		const FRotation3& RoM() const { return State.RoM; }
+		void SetRoM(const FRotation3& InRoM) { State.RoM = InRoM; }
+
+		/**
+		 * @brief Get the current world-space Actor position 
+		 * @note This is recalculated from the current CoM transform.
+		*/
+		FVec3 ActorP() const { return P() - ActorQ().RotateVector(CoM()); }
+
+		/**
+		 * @brief Get the current world-space Actor rotation
+		 * @note This is recalculated from the current CoM transform.
+		*/
+		FRotation3 ActorQ() const { return Q() * RoM().Inverse(); }
+
+		/**
+		 * @brief Contact graph level. This is used in shock propagation to determine which of two bodies should have its inverse mass scaled
+		*/
+		int32 Level() const { return State.Level; }
+		void SetLevel(int32 InLevel) { State.Level = InLevel; }
+
+		/**
+		 * @brief Whether there were any active collision constraints on this body
+		*/
+		bool HasActiveCollision() const { return State.bHasActiveCollision; }
+		void SetHasActiveCollision(bool bInHasCollision) { State.bHasActiveCollision = bInHasCollision; }
+
+		/**
+		 * @brief Whether the body has a finite mass
+		 * @note This is based on the current inverse mass, so a "dynamic" particle with 0 inverse mass will return true here.
+		*/
+		bool IsDynamic() const { return (State.InvM > SMALL_NUMBER); }
+
+		/**
+		 * @brief Apply a world-space position and rotation delta to the body center of mass, and update inverse mass
+		*/
+		void ApplyTransformDelta(const FVec3 DP, const FVec3 DR)
+		{
+			ApplyPositionDelta(DP);
+			ApplyRotationDelta(DR);
+		}
+
+		/**
+		 * @brief Apply a world-space position delta to the solver body center of mass
+		*/
+		void ApplyPositionDelta(const FVec3 DP)
+		{
+			State.P += DP;
+			++State.LastChangeEpoch;
+		}
+
+		/**
+		 * @brief Apply a world-space rotation delta to the solver body and update the inverse mass
+		*/
+		void ApplyRotationDelta(const FVec3 DR)
+		{
+			State.Q += FRotation3::FromElements(DR, 0.f) * State.Q * FReal(0.5);
+			State.Q.Normalize();
+			++State.LastChangeEpoch;
+
+			//if (bUpdateRotationDependentState)
+			//{
+			//	UpdateRotationDependentState();
+			//}
+		}
+
+		/**
+		 * @brief Apply a world-space velocity delta to the solver body
+		*/
+		void ApplyVelocityDelta(const FVec3& DV, const FVec3& DW)
+		{
+			ApplyLinearVelocityDelta(DV);
+			ApplyAngularVelocityDelta(DW);
+		}
+
+		/**
+		 * @brief Apply a world-space linear velocity delta to the solver body
+		*/
+		void ApplyLinearVelocityDelta(const FVec3& DV)
+		{
+			State.V += DV;
+			++State.LastChangeEpoch;
+		}
+
+		/**
+		 * @brief Apply an world-space angular velocity delta to the solver body
+		*/
+		void ApplyAngularVelocityDelta(const FVec3& DW)
+		{
+			State.W += DW;
+			++State.LastChangeEpoch;
+		}
+
+		/**
+		 * @brief Update the rotation to be in the same hemisphere as the provided quaternion.
+		 * This is used by joints with angular constraint/drives
+		*/
+		void EnforceShortestRotationTo(const FRotation3& InQ)
+		{
+			State.Q.EnforceShortestArcWith(InQ);
+		}
+
+		int32 LastChangeEpoch() const
+		{
+			return State.LastChangeEpoch;
+		}
+
+		/**
+		 * @brief Update cached state that depends on rotation (i.e., world space inertia)
+		*/
+		void UpdateRotationDependentState();
+
+	private:
+
+		// The struct exists only so that we can use the variable names
+		// as accessor names without violation the variable naming convention
+		struct FState
+		{
+			FState()
+				: InvILocal(0)
+				, InvM(0)
+				, InvI(0)
+				, X(0)
+				, R(FRotation3::Identity)
+				, P(0)
+				, Q(FRotation3::Identity)
+				, V(0)
+				, W(0)
+				, CoM(0)
+				, RoM(FRotation3::Identity)
+				, Level(0)
+				, LastChangeEpoch(0)
+				, bHasActiveCollision(false)
+			{}
+
+			// Local-space inverse inertia (diagonal, so only 3 elements)
+			FVec3 InvILocal;
+
+			// Inverse mass
+			FReal InvM;
+
+			// World-space inverse inertia
+			// @todo(chaos): do we need this, or should we force all systems to use the FConstraintSolverBody decorator?
+			FMatrix33 InvI;
+
+			// World-space center of mass state at start of sub step
+			FVec3 X;
+
+			// World-space rotation of mass
+			FRotation3 R;
+
+			// World-space center of mass position
+			FVec3 P;
+
+			// World-space center of mass rotation
+			FRotation3 Q;
+
+			// World-space center of mass velocity
+			FVec3 V;
+
+			// World-space center of mass angular velocity
+			FVec3 W;
+
+			// Actor-space center of mass location
+			FVec3 CoM;
+
+			// Actor-space center of mass rotation
+			FRotation3 RoM;
+
+			// Distance to a kinmatic body (through the contact graph). Used by collision shock propagation
+			int32 Level;
+
+			// A counter incremented every time the state changes. 
+			// Used by constraints to determine if some other constraint has modified the body for early exit logic
+			int32 LastChangeEpoch;
+
+			// Whether we had any active contacts this substep
+			// @todo(chaos): maybe make this a counter?
+			bool bHasActiveCollision;
+		};
+
+		//FGenericParticleHandle Particle;
+		FState State;
+	};
+
+	/**
+	 * An FSolverBody decorator for adding mass modifiers to a SolverBody. This will scale the
+	 * inverse mass and inverse inertia using the supplied scale. It also updates IsDynamic() to
+	 * return false if the scaled inverse mass is zero.
+	 * 
+	 * See FSolverBody for comments on methods.
+	 * 
+	 * @note This functionality cannot be in FSolverBody because two constraints referencing
+	 * the same body may be applying different mass modifiers (e.g., Joints support "bParentDominates"
+	 * which is a per-constraint property, not a per-body property.
+	 */
+	class FConstraintSolverBody
+	{
+	public:
+		FConstraintSolverBody()
+			: Body(nullptr)
+		{
+		}
+
+		FConstraintSolverBody(FSolverBody& InBody)
+			: Body(&InBody)
+		{
+		}
+
+		FConstraintSolverBody(FSolverBody& InBody, FReal InInvMassScale)
+			: Body(&InBody)
+		{
+			SetInvMassScale(InInvMassScale);
+		}
+
+		/**
+		 * @brief True if we have been set up to decorate a SolverBody
+		*/
+		bool IsValid() const { return Body != nullptr; }
+
+		/**
+		 * @brief Invalidate the solver body reference
+		*/
+		void Reset() { Body = nullptr; }
+
+		/**
+		 * @brief The decorated SolverBody
+		*/
+		FSolverBody& SolverBody() { check(IsValid()); return *Body; }
+		const FSolverBody& SolverBody() const { check(IsValid()); return *Body; }
+
+		/**
+		 * @brief A scale applied to both inverse mass and inverse inertia
+		*/
+		FReal InvMassScale() const { return State.InvMassScale; }
+		void SetInvMassScale(FReal InInvMassScale) { State.InvMassScale = InInvMassScale; }
+
+		/**
+		 * @brief The scaled inverse mass
+		*/
+		FReal InvM() const { return State.InvMassScale * Body->InvM(); }
+
+		/**
+		 * @brief The scaled inverse inertia
+		*/
+		FMatrix33 InvI() const { return State.InvMassScale * Body->InvI(); }
+
+		/**
+		 * @brief Whether the body is dynamic (i.e., has a finite mass) after InvMassScale is applied
+		*/
+		bool IsDynamic() const { return (InvM() > SMALL_NUMBER); }
+
+		//
+		// From here all methods just forward to the FSolverBody
+		//
+
+		void SetImplicitVelocity(FReal Dt) { Body->SetImplicitVelocity(Dt); }
+		FRigidTransform3 CoMTransform() const { return Body->CoMTransform(); }
+		const FVec3& X() const { return Body->X(); }
+		const FRotation3& R() const { return Body->R(); }
+		const FVec3& P() const { return Body->P(); }
+		const FRotation3& Q() const { return Body->Q(); }
+		const FVec3 ActorP() const { return Body->ActorP(); }
+		const FRotation3 ActorQ() const { return Body->ActorQ(); }
+		const FVec3& V() const { return Body->V(); }
+		const FVec3& W() const { return Body->W(); }
+		int32 Level() const { return Body->Level(); }
+
+		void ApplyTransformDelta(const FVec3 DP, const FVec3 DR) { Body->ApplyTransformDelta(DP, DR); }
+		void ApplyPositionDelta(const FVec3 DP) { Body->ApplyPositionDelta(DP); }
+		void ApplyRotationDelta(const FVec3 DR) { Body->ApplyRotationDelta(DR); }
+		void ApplyVelocityDelta(const FVec3& DV, const FVec3& DW) { Body->ApplyVelocityDelta(DV, DW); }
+		void ApplyLinearVelocityDelta(const FVec3& DV) { Body->ApplyLinearVelocityDelta(DV); }
+		void ApplyAngularVelocityDelta(const FVec3& DW) { Body->ApplyAngularVelocityDelta(DW); }
+		void EnforceShortestRotationTo(const FRotation3& InQ) { Body->EnforceShortestRotationTo(InQ); }
+		void UpdateRotationDependentState() { Body->UpdateRotationDependentState(); }
+
+		int32 LastChangeEpoch() const { return Body->LastChangeEpoch(); }
+
+	private:
+		// Struct is only so that we can use the same var names as function names
+		struct FState
+		{
+			FState() 
+				: InvMassScale(FReal(1)) 
+			{}
+			FReal InvMassScale;
+		};
+
+		// The body we decorate
+		FSolverBody* Body;
+
+		// The body modifiers
+		FState State;
+	};
+}
diff --git a/Engine/Source/Runtime/Experimental/Chaos/Public/Chaos/Evolution/SolverBodyContainer.h b/Engine/Source/Runtime/Experimental/Chaos/Public/Chaos/Evolution/SolverBodyContainer.h
new file mode 100644
--- /dev/null
+++ b/Engine/Source/Runtime/Experimental/Chaos/Public/Chaos/Evolution/SolverBodyContainer.h
@@ -0,0 +1,151 @@
+// Copyright Epic Games, Inc. All Rights Reserved.
+#pragma once
+
+#include "Chaos/Evolution/SolverBody.h"
+#include "Chaos/Framework/ContainerItemHandle.h"
+#include "Chaos/ParticleHandle.h"
+
+#include "Containers/Array.h"
+#include "Containers/Map.h"
+
+namespace Chaos
+{
+	class FSolverBody;
+	class FSolverBodyContainer;
+
+	/**
+	 * @brief An FSolverBody wrapper that binds the solver to a particle and provides gather/scatter methods from/to the particle
+	 * Solver Bodies are used by the core constraint solver
+	 * @see FSolverBody
+	*/
+	class FSolverBodyAdapter
+	{
+	public:
+		FSolverBodyAdapter(const FGenericParticleHandle& InParticle)
+			: Particle(InParticle)
+		{
+			GatherInput();
+		}
+		
+		FSolverBody& GetSolverBody() { return SolverBody; }
+		const FGenericParticleHandle& GetParticle() { return Particle; }
+
+		void GatherInput();
+		void ScatterOutput();
+
+	private:
+		FSolverBody SolverBody;
+		FGenericParticleHandle Particle;
+	};
+
+	/**
+	 * The SolverBodies for a set of particles.
+	 * 
+	 * Each Island owns a SolverBodyContainer containing the data required for solving the constraints in the island.
+	 * Constraints in the island will hold pointers to SolverBodies in the Island's SolverBodyContainer.
+	 * ConstraintSolvers read and write to their SolverBodies and do not access particles directly.
+	 * 
+	 * SolverBodies are created at the start of the constraint solving phase, and destroyed at the end. They are
+	 * stored in a contiguous array in the order that they are accessed, for cache efficiency. The array
+	 * is guaranteed not to resize while accumulating bodies for the solving phase (it will assert if it does)
+	 * so pointers to elements are valid for the duration of the constraint solving phase, but not longer.
+	 * 
+	 * \note This container holds all the state of bodies in an island. Dynamics will only appear in one island, 
+	 * but kinematics are in multiple and so their state will be duplicated in each island. This is ok - we do 
+	 * not update the state of any kinematic bodies in the constraint solvers. We also assume that the number
+	 * of kinematics is small compared to the number of dynamics. If this is commonly untrue we may want to 
+	 * consider having a separate (global) container of kinematic solver bodies.
+	 */
+	class FSolverBodyContainer
+	{
+	public:
+
+		// Clear the bodies array and allocate enough space for at least MaxBodies bodies.
+		// @param MaxBodies The number of bodies that will get added to the container. The container asserts if we attempt to add more than this.
+		inline void Reset(int MaxBodies)
+		{
+			SolverBodies.Reset(MaxBodies);
+			ParticleToIndexMap.Reset();
+		}
+
+		// The number of bodies in the container
+		inline int NumItems() const
+		{
+			return SolverBodies.Num();
+		}
+
+		// The maximum number of bodies the container can hold (until Reset() is called again)
+		inline int MaxItems() const
+		{
+			return SolverBodies.Max();
+		}
+
+		// Get a pointer to the item at the specified index, or null for INDEX_NONE
+		// @param Index The index of the body to return
+		inline FSolverBodyAdapter* TryGetItem(int Index)
+		{
+			if (IsValid(Index))
+			{
+				return &SolverBodies[Index];
+			}
+			return nullptr;
+		}
+
+		// Get a pointer to the item at the specified index, or null for INDEX_NONE
+		// @param Index The index of the body to return
+		inline const FSolverBodyAdapter* TryGetItem(int Index) const
+		{
+			if (IsValid(Index))
+			{
+				return &SolverBodies[Index];
+			}
+			return nullptr;
+		}
+
+		// Get a reference to the item at the specified index. Asserts on invalid index.
+		// @param Index The index of the body to return
+		inline FSolverBodyAdapter& GetItem(int Index)
+		{
+			return SolverBodies[Index];
+		}
+
+		// Get a reference to the item at the specified index. Asserts on invalid index.
+		// @param Index The index of the body to return
+		inline const FSolverBodyAdapter& GetItem(int Index) const
+		{
+			return SolverBodies[Index];
+		}
+
+		// Whether the specified index is valid.
+		// @param Index The index of the body to return
+		inline bool IsValid(int Index) const
+		{
+			return SolverBodies.IsValidIndex(Index);
+		}
+
+		// Add a solver body to represent the solver state of the particle
+		// This should ideally be called in the order in which the bodies will be accessed
+		// (or as close as we can get, given most constraints access 2 bodies so there is no perfect order)
+		FSolverBody* FindOrAdd(FGenericParticleHandle InParticle);
+
+		// Collect all the data we need from the particles represented by our SolverBodies
+		//void GatherInput();
+
+		// Scatter the solver results back to the particles represented by our SolverBodies
+		void ScatterOutput();
+
+		// Recalculate the velocities of all bodies based on their transform deltas
+		void SetImplicitVelocities(FReal Dt);
+
+	private:
+		int32 AddParticle(FGenericParticleHandle InParticle);
+
+		// Solver bodies, usually collected in the order in which they are accessed
+		TArray<FSolverBodyAdapter> SolverBodies;
+
+		// Used to determine if a (non-dynamic) body is already present in the container. 
+		// Dynamic bodies are not in this map - the particle stores its index in this case.
+		TMap<FGenericParticleHandle, int32> ParticleToIndexMap;
+	};
+
+}
diff --git a/Engine/Source/Runtime/Experimental/Chaos/Public/Chaos/Evolution/SolverConstraintContainer.h b/Engine/Source/Runtime/Experimental/Chaos/Public/Chaos/Evolution/SolverConstraintContainer.h
new file mode 100644
--- /dev/null
+++ b/Engine/Source/Runtime/Experimental/Chaos/Public/Chaos/Evolution/SolverConstraintContainer.h
@@ -0,0 +1,21 @@
+// Copyright Epic Games, Inc. All Rights Reserved.
+#pragma once
+
+namespace Chaos
+{
+
+/** Parent base class for all the solver constraint containers */
+class CHAOS_API FConstraintSolverContainer
+{
+public:
+
+	/** Virtual destructor */
+	virtual ~FConstraintSolverContainer() {}
+		
+	/** Reset the container to have zero constraints
+	* @param MaxConstraints Maximum number of constraints that could be in the container
+	*/
+	virtual void Reset(const int32 MaxConstraints) {}
+};
+	
+}
diff --git a/Engine/Source/Runtime/Experimental/Chaos/Public/Chaos/Evolution/SolverDatas.h b/Engine/Source/Runtime/Experimental/Chaos/Public/Chaos/Evolution/SolverDatas.h
new file mode 100644
--- /dev/null
+++ b/Engine/Source/Runtime/Experimental/Chaos/Public/Chaos/Evolution/SolverDatas.h
@@ -0,0 +1,143 @@
+// Copyright Epic Games, Inc. All Rights Reserved.
+#pragma once
+
+#include "Containers/SparseArray.h"
+#include "Chaos/Evolution/SolverBodyContainer.h"
+#include "Chaos/Evolution/SolverConstraintContainer.h"
+
+namespace Chaos
+{
+	
+/** Forward Declaration */
+class FPBDIslandManager;
+
+/**
+* List of bodies and constraints that wil be stored per island.
+*/
+class CHAOS_API FPBDIslandSolverData
+{
+public :
+	/**
+	* Init the solver datas 
+	*/
+	FPBDIslandSolverData(const int32 IslandIndex = 0) : BodyContainer(MakeUnique<FSolverBodyContainer>()), ConstraintDatas(), IslandIndex(IslandIndex)
+	{}
+
+	/** Accessors of the island index */
+	FORCEINLINE const int32& GetIslandIndex() const { return IslandIndex; };
+	FORCEINLINE int32& GetIslandIndex() { return IslandIndex; };
+	
+	/** Accessors of the body container */
+	FORCEINLINE const FSolverBodyContainer& GetBodyContainer() const { return *BodyContainer.Get(); };
+	FORCEINLINE FSolverBodyContainer& GetBodyContainer() { return *BodyContainer.Get(); };
+
+	/**
+	 * @brief The number of constraint containers registsred
+	*/
+	int32 NumConstraintContainerIds() const { return ConstraintDatas.Num(); }
+
+	/** Accessors of the constraint container */
+	template<typename ContainerType> const ContainerType& GetConstraintContainer(const int32 ContainerId) const;
+	template<typename ContainerType> ContainerType& GetConstraintContainer(const int32 ContainerId);
+
+	/** Accessors of the constraint indices */
+	const TArray<int32>& GetConstraintIndices(const int32 ContainerId) const;
+	TArray<int32>& GetConstraintIndices(const int32 ContainerId);
+	
+	/** Accessors of the constraint handles */
+	const TArray<FConstraintHandle*>& GetConstraintHandles(const int32 ContainerId) const;
+	TArray<FConstraintHandle*>& GetConstraintHandles(const int32 ContainerId);
+
+	/** Accessor of one constraint handle given a container id and a constraint index */
+	template<typename ConstraintType> const ConstraintType* GetConstraintHandle(const int32 ContainerId, const int32 ConstraintIndex) const;
+	template<typename ConstraintType> ConstraintType* GetConstraintHandle(const int32 ContainerId, const int32 ConstraintIndex);
+
+	/** Add a constraint container to the solver island given a container id
+	* @param ContainerId Constraints container id from which the solver constraint datas is being built
+	*/
+	template<typename ConstraintType> void AddConstraintDatas(const int32 ContainerId);
+
+protected:
+	/**
+	* Constraint datas that will be stored per container (joints, collisions...)
+	*/
+	struct FConstraintDatas
+	{
+		/** Constraint indices that will be used for legacy solver */
+		TArray<int32> ConstraintIndices;
+
+		/** Constraint handles that will be used for legacy solver */
+		TArray<FConstraintHandle*> ConstraintHandles;
+	};
+
+	/** Solver body container of that datas */
+	TUniquePtr<FSolverBodyContainer> BodyContainer;
+	
+	/** List of constraint containers (collision, joints...) that will be used to solve constraints */
+	TSparseArray<TUniquePtr<FConstraintSolverContainer>> ConstraintContainers;
+	
+	/** List of constraint datas (collision, joints...) that will be used to solve constraints */
+	TSparseArray<FConstraintDatas> ConstraintDatas;
+
+	/** Island index in case these datas belong to an island */
+	int32 IslandIndex;
+};
+
+template<typename ContainerType>
+FORCEINLINE const ContainerType& FPBDIslandSolverData::GetConstraintContainer(const int32 ContainerId) const
+{
+	return *static_cast<const ContainerType*>(ConstraintContainers[ContainerId].Get());
+};
+
+template<typename ContainerType>
+FORCEINLINE ContainerType& FPBDIslandSolverData::GetConstraintContainer(const int32 ContainerId)
+{
+	return *static_cast<ContainerType*>(ConstraintContainers[ContainerId].Get());
+};
+	
+FORCEINLINE const TArray<int32>& FPBDIslandSolverData::GetConstraintIndices(const int32 ContainerId) const
+{
+	return ConstraintDatas[ContainerId].ConstraintIndices;
+};
+	
+FORCEINLINE TArray<int32>& FPBDIslandSolverData::GetConstraintIndices(const int32 ContainerId)
+{
+	return ConstraintDatas[ContainerId].ConstraintIndices;
+};
+	
+FORCEINLINE const TArray<FConstraintHandle*>& FPBDIslandSolverData::GetConstraintHandles(const int32 ContainerId) const
+{
+	return ConstraintDatas[ContainerId].ConstraintHandles;
+};
+
+FORCEINLINE TArray<FConstraintHandle*>& FPBDIslandSolverData::GetConstraintHandles(const int32 ContainerId)
+{
+	return ConstraintDatas[ContainerId].ConstraintHandles;
+};
+
+template<typename ConstraintType>
+FORCEINLINE const ConstraintType* FPBDIslandSolverData::GetConstraintHandle(const int32 ContainerId, const int32 ConstraintIndex) const
+{
+	return static_cast<const ConstraintType*>(ConstraintDatas[ContainerId].ConstraintHandles[ConstraintIndex]);
+};
+
+template<typename ConstraintType>
+FORCEINLINE ConstraintType* FPBDIslandSolverData::GetConstraintHandle(const int32 ContainerId, const int32 ConstraintIndex)
+{
+	return static_cast<ConstraintType*>(ConstraintDatas[ContainerId].ConstraintHandles[ConstraintIndex]);
+};
+	
+template<typename ConstraintType>
+inline void FPBDIslandSolverData::AddConstraintDatas(const int32 ContainerId)
+{
+	if((ContainerId >= 0) && !ConstraintDatas.IsValidIndex(ContainerId) && !ConstraintContainers.IsValidIndex(ContainerId) )
+	{
+		ConstraintDatas.Reserve(ContainerId+1);
+		ConstraintDatas.EmplaceAt(ContainerId, FConstraintDatas());
+		
+		ConstraintContainers.Reserve(ContainerId+1);
+		ConstraintContainers.EmplaceAt(ContainerId, MakeUnique<typename ConstraintType::FConstraintSolverContainerType>());
+	}
+}
+	
+}
\ No newline at end of file
