Chaos

diff --git a/Engine/Source/Runtime/Experimental/Chaos/Public/Chaos/EvolutionResimCache.h b/Engine/Source/Runtime/Experimental/Chaos/Public/Chaos/EvolutionResimCache.h
--- a/Engine/Source/Runtime/Experimental/Chaos/Public/Chaos/EvolutionResimCache.h
+++ b/Engine/Source/Runtime/Experimental/Chaos/Public/Chaos/EvolutionResimCache.h
@@ -19,7 +19,6 @@ namespace Chaos
 		{
 			SavedConstraints.Reset();
 			WeakSinglePointConstraints.Reset();
-			WeakSinglePointSweptConstraints.Reset();
 			ParticleToCachedSolve.Reset();
 		}
 
@@ -45,15 +44,20 @@ namespace Chaos
 			}
 		}
 
-		void SaveConstraints(const FCollisionConstraintsArray& CollisionsArray)
+		void SaveConstraints(TArrayView<const FPBDCollisionConstraint* const> CollisionsArray)
 		{
-			SavedConstraints = CollisionsArray;
+			// Copy the constraints into an array
+			SavedConstraints.Reset(CollisionsArray.Num());
+			for (const FPBDCollisionConstraint* Collision : CollisionsArray)
+			{
+				SavedConstraints.Emplace(FPBDCollisionConstraint::MakeResimCache(*Collision));
+			}
 
 			//Create weak handles so we can make sure everything is alive later
 			auto SaveArrayHelper = [](auto& Constraints,auto& WeakPairs)
 			{
 				WeakPairs.Empty(Constraints.Num());
-				for(FCollisionConstraintBase& Constraint : Constraints)
+				for(FPBDCollisionConstraint& Constraint : Constraints)
 				{
 					WeakPairs.Add(FWeakConstraintPair{Constraint.Particle[0]->WeakParticleHandle(),Constraint.Particle[1]->WeakParticleHandle()});
 
@@ -66,19 +70,18 @@ namespace Chaos
 				check(Constraints.Num() == WeakPairs.Num());
 			};
 
-			SaveArrayHelper(SavedConstraints.SinglePointConstraints,WeakSinglePointConstraints);
-			SaveArrayHelper(SavedConstraints.SinglePointSweptConstraints,WeakSinglePointSweptConstraints);
+			SaveArrayHelper(SavedConstraints,WeakSinglePointConstraints);
 		}
 
 		//Returns all constraints that are still valid (resim can invalidate constraints by either deleting particles, moving particles, etc...)
-		const FCollisionConstraintsArray& GetAndSanitizeConstraints()
+		const TArray<FPBDCollisionConstraint>& GetAndSanitizeConstraints()
 		{
 			auto CleanupArrayHelper = [](auto& Constraints,auto& WeakPairs)
 			{
 				check(Constraints.Num() == WeakPairs.Num());
 				for(int32 Idx = Constraints.Num() - 1; Idx >= 0; --Idx)
 				{
-					FCollisionConstraintBase& Constraint = Constraints[Idx];
+					FPBDCollisionConstraint& Constraint = Constraints[Idx];
 					FWeakConstraintPair& WeakPair = WeakPairs[Idx];
 
 					TGeometryParticleHandle<FReal,3>* A = WeakPair.A.GetHandleUnsafe();
@@ -99,8 +102,7 @@ namespace Chaos
 				}
 			};
 
-			CleanupArrayHelper(SavedConstraints.SinglePointConstraints,WeakSinglePointConstraints);
-			CleanupArrayHelper(SavedConstraints.SinglePointSweptConstraints,WeakSinglePointSweptConstraints);
+			CleanupArrayHelper(SavedConstraints,WeakSinglePointConstraints);
 
 			return SavedConstraints;
 		}
@@ -130,7 +132,7 @@ namespace Chaos
 		}
 
 		//NOTE: You must sanitize this before using. This can contain dangling pointers or invalid constraints
-		FCollisionConstraintsArray SavedConstraints;
+		TArray<FPBDCollisionConstraint> SavedConstraints;
 
 		struct FWeakConstraintPair
 		{
@@ -148,7 +150,6 @@ namespace Chaos
 
 		//TODO: better way to handle this?
 		TArray<FWeakConstraintPair> WeakSinglePointConstraints;
-		TArray<FWeakConstraintPair> WeakSinglePointSweptConstraints;
 		TMap<FUniqueIdx, FPBDSolveCache> ParticleToCachedSolve;
 	};
 
diff --git a/Engine/Source/Runtime/Experimental/Chaos/Public/Chaos/GJK.h b/Engine/Source/Runtime/Experimental/Chaos/Public/Chaos/GJK.h
--- a/Engine/Source/Runtime/Experimental/Chaos/Public/Chaos/GJK.h
+++ b/Engine/Source/Runtime/Experimental/Chaos/Public/Chaos/GJK.h
@@ -90,6 +90,280 @@ namespace Chaos
 		
 	}
 
+	/**
+	 * @brief Internal simplex data for GJK that can also be stored for warm-starting subsequent calls
+	 * @tparam T the number type (float or double)
+	 *
+	 * @see GJKPenetrationWarmStartable
+	*/
+	template<typename T>
+	class TGJKSimplexData
+	{
+	public:
+		TGJKSimplexData()
+			: NumVerts(0)
+		{}
+
+		// Clear the data - used to start a GJK search from the default search direction
+		void Reset()
+		{
+			NumVerts = 0;
+		}
+
+		// Save any data that was not directly updated while iterating in GJK
+		void Save(const FSimplex InSimplexIDs)
+		{
+			// We don't need to store the simplex vertex order because the indices are always
+			// sorted at the end of each iteration. We just need to know how many vertices we have.
+			NumVerts = InSimplexIDs.NumVerts;
+		}
+
+		// Recompute the Simplex and separating vector from the stored data at the current relative transform
+		// This aborts if we have no simplex data to restore or the origin is inside the simplex. Outputs must already 
+		// have reasonable default values for running GJK without a warm-start.
+		void Restore(const TRigidTransform<T, 3>& BToATM, FSimplex& OutSimplexIDs, TVec3<T> OutSimplex[], TVec3<T>& OutV, T& OutDistance, const T Epsilon)
+		{
+			if (NumVerts > 0)
+			{
+				OutSimplexIDs.NumVerts = NumVerts;
+
+				for (int32 VertIndex = 0; VertIndex < NumVerts; ++VertIndex)
+				{
+					OutSimplexIDs.Idxs[VertIndex] = VertIndex;
+					OutSimplex[VertIndex] = As[VertIndex] - BToATM.TransformPositionNoScale(Bs[VertIndex]);
+				}
+
+				const TVec3<T> V = SimplexFindClosestToOrigin(OutSimplex, OutSimplexIDs, Barycentric, As, Bs);
+				const T Distance = V.Size();
+
+				// If the origin is inside the simplex at the new transform, we need to abort the restore
+				// This is necessary to cover the very-small separation case where we use the normal
+				// calculated in the previous iteration in GJKL, but we have no way to restore that.
+				// Note: we have already written to the simplex but that's ok because we reset the vert count.
+				if (Distance > Epsilon)
+				{
+					OutV = V / Distance;
+					OutDistance = Distance;
+				}
+				else
+				{
+					OutSimplexIDs.NumVerts = 0;
+				}
+			}
+		}
+
+		// Maximum number of vertices that a GJK simplex can have
+		static const int32 MaxSimplexVerts = 4;
+
+		// Simplex vertices on shape A, in A-local space
+		TVec3<T> As[MaxSimplexVerts];
+
+		// Simplex vertices on shape B, in B-local space
+		TVec3<T> Bs[MaxSimplexVerts];
+
+		// Barycentric coordinates of closest point to origin on the simplex
+		T Barycentric[MaxSimplexVerts];
+
+		// Number of vertices in the simplex. Up to 4.
+		int32 NumVerts;
+	};
+
+	// GJK warm-start data at default numeric precision
+	using FGJKSimplexData = TGJKSimplexData<FReal>;
+
+	/**
+	 * @brief Calculate the penetration data for two shapes using GJK and a warm-start buffer.
+	 *
+	 * @tparam T Number type (float or double)
+	 * @tparam TGeometryA The first shape type
+	 * @tparam TGeometryB The second shape type
+	 * @param A The first shape
+	 * @param B The second shape
+	 * @param BToATM A transform from B-local space to A-local space
+	 * @param OutPenetration The overlap distance (+ve for overlap, -ve for separation)
+	 * @param OutClosestA The closest point on A, in A-local space
+	 * @param OutClosestB The closest point on B, in B-local space
+	 * @param OutNormalA The contact normal pointing away from A in A-local space
+	 * @param OutNormalB The contact normal pointing away from A in B-local space
+	 * @param SimplexData In/out simplex data used to initialize and update GJK. Can be stored to improve convergence on subsequent calls for "small" changes in relative rotation.
+	 * @param Epsilon The separation distance below which GJK aborts and switches to EPA
+	 * @return true if the results are populated (always for this implementation, but @see EPA())
+	 *
+	 * The WarmStartData is an input and output parameter. If the function is called with
+	 * a small change in BToATM there's it will converge much faster, usually in 1 iteration
+	 * for polygonal shapes.
+	 * 
+	 * @note This version returns OutClosestB in B's local space, compated to GJKPenetration where all output is in the space of A
+	 * 
+	 * @todo(chaos): convert GJKPenetration() to use this version (but see note above)
+	 *
+	*/
+	template <bool bFindSeparationDistance = false, typename T, typename TGeometryA, typename TGeometryB>
+	bool GJKPenetrationWarmStartable(const TGeometryA& A, const TGeometryB& B, const TRigidTransform<T, 3>& BToATM, const T InThicknessA, const T InThicknessB, T& OutPenetration, TVec3<T>& OutClosestA, TVec3<T>& OutClosestB, TVec3<T>& OutNormalA, TVec3<T>& OutNormalB, TGJKSimplexData<T>& InOutSimplexData, const T Epsilon = T(1.e-3))
+	{
+		// Return the support vertex in A-local space for a vector V in A-local space
+		auto SupportAFunc = [&A](const TVec3<T>& V)
+		{
+			return A.SupportCore(V, A.GetMargin());
+		};
+
+		const TRotation<T, 3> AToBRotation = BToATM.GetRotation().Inverse();
+
+		// Return the support point on B, in B-local space, for a vector V in A-local space
+		auto SupportBFunc = [&B, &AToBRotation](const TVec3<T>& V)
+		{
+			const TVec3<T> VInB = AToBRotation * V;
+			return B.SupportCore(VInB, B.GetMargin());
+		};
+
+		// V and Simplex are in A-local space
+		TVec3<T> V = TVec3<T>(-1, 0, 0);
+		TVec3<T> Simplex[4];
+		FSimplex SimplexIDs;
+		T Distance = FLT_MAX;
+
+		// If we have warm-start data, rebuild the simplex from the stored data
+		InOutSimplexData.Restore(BToATM, SimplexIDs, Simplex, V, Distance, Epsilon);
+
+		TVec3<T> Normal = -V;					// Remember the last good normal (i.e. don't update it if separation goes less than Epsilon and we can no longer normalize)
+		bool bIsDegenerate = false;				// True if GJK cannot make any more progress
+		bool bIsContact = false;				// True if shapes are within Epsilon or overlapping - GJK cannot provide a solution
+		int32 NumIterations = 0;
+		const T ThicknessA = InThicknessA + A.GetMargin();
+		const T ThicknessB = InThicknessB + B.GetMargin();
+		const T SeparatedDistance = ThicknessA + ThicknessB + Epsilon;
+		while (!bIsContact && !bIsDegenerate)
+		{
+			if (!ensure(NumIterations++ < 32))	//todo: take this out
+			{
+				break;	//if taking too long just stop. This should never happen
+			}
+
+			const TVec3<T> NegV = -V;
+			const TVec3<T> SupportA = SupportAFunc(NegV);
+			const TVec3<T> SupportB = SupportBFunc(V);
+			const TVec3<T> SupportBInA = BToATM.TransformPositionNoScale(SupportB);
+			const TVec3<T> W = SupportA - SupportBInA;
+
+			// If we didn't move to at least ConvergedDistance or closer, assume we have reached a minimum
+			const T ConvergenceTolerance = 1.e-4f;
+			const T ConvergedDistance = (1.0f - ConvergenceTolerance) * Distance;
+			const T VW = TVec3<T>::DotProduct(V, W);
+			if (!bFindSeparationDistance && (VW > SeparatedDistance))
+			{
+				// We are separated and don't care about the separation distance
+				break;
+			}
+
+			InOutSimplexData.As[SimplexIDs.NumVerts] = SupportA;
+			InOutSimplexData.Bs[SimplexIDs.NumVerts] = SupportB;
+			SimplexIDs[SimplexIDs.NumVerts] = SimplexIDs.NumVerts;
+			Simplex[SimplexIDs.NumVerts++] = W;
+
+			V = SimplexFindClosestToOrigin(Simplex, SimplexIDs, InOutSimplexData.Barycentric, InOutSimplexData.As, InOutSimplexData.Bs);
+			T NewDistance = V.Size();
+
+			// Are we overlapping or too close for GJK to get a good result?
+			bIsContact = (NewDistance < Epsilon);
+
+			// If we did not get closer in this iteration, we are in a degenerate situation
+			bIsDegenerate = (NewDistance >= Distance);
+
+			// If we are not too close, update the separating vector and keep iterating
+			// If we are too close we drop through to EPA, but if EPA determines the shapes are actually separated 
+			// after all, we will end up using the normal from the previous loop
+			if (!bIsContact)
+			{
+				V /= NewDistance;
+				Normal = -V;
+			}
+			Distance = NewDistance;
+		}
+
+		// Save the warm-start data (not much to store since we updated most of it in the loop above)
+		InOutSimplexData.Save(SimplexIDs);
+
+		if (bIsContact)
+		{
+			// We did not converge or we detected an overlap situation, so run EPA to get contact data
+			// Rebuild the simplex into a variable sized array - EPA can generate any mumber of vertices
+			TArray<TVec3<T>> VertsA;
+			TArray<TVec3<T>> VertsB;
+			VertsA.Reserve(8);
+			VertsB.Reserve(8);
+			for (int i = 0; i < SimplexIDs.NumVerts; ++i)
+			{
+				VertsA.Add(InOutSimplexData.As[i]);
+				VertsB.Add(BToATM.TransformPositionNoScale(InOutSimplexData.Bs[i]));
+			}
+
+			auto SupportBInAFunc = [&B, &BToATM, &AToBRotation](const TVec3<T>& V)
+			{
+				const TVec3<T> VInB = AToBRotation * V;
+				const TVec3<T> SupportBLocal = B.SupportCore(VInB, B.GetMargin());
+				return BToATM.TransformPositionNoScale(SupportBLocal);
+			};
+
+			T Penetration;
+			TVec3<T> MTD, ClosestA, ClosestBInA;
+			const EEPAResult EPAResult = EPA(VertsA, VertsB, SupportAFunc, SupportBInAFunc, Penetration, MTD, ClosestA, ClosestBInA);
+
+			switch (EPAResult)
+			{
+			case EEPAResult::MaxIterations:
+				// Possibly a solution but with unknown error. Just return the last EPA state. 
+				// Fall through...
+			case EEPAResult::Ok:
+				// EPA has a solution - return it now
+				OutNormalA = MTD;
+				OutNormalB = BToATM.InverseTransformVectorNoScale(MTD);
+				OutPenetration = Penetration + ThicknessA + ThicknessB;
+				OutClosestA = ClosestA + MTD * ThicknessA;
+				OutClosestB = ClosestBInA - MTD * ThicknessB;
+				return true;
+			case EEPAResult::BadInitialSimplex:
+				// The origin is outside the simplex. Must be a touching contact and EPA setup will have calculated the normal
+				// and penetration but we keep the position generated by GJK.
+				Normal = MTD;
+				Distance = -Penetration;
+				break;
+			case EEPAResult::Degenerate:
+				// We hit a degenerate simplex condition and could not reach a solution so use whatever near touching point GJK came up with.
+				// The result from EPA under these circumstances is not usable.
+				//UE_LOG(LogChaos, Warning, TEXT("EPA Degenerate Case"));
+				break;
+			}
+		}
+
+		// @todo(chaos): handle the case where EPA hits a degenerate triangle in the simplex.
+		// Currently we return a touching contact with the last position and normal from GJK. We could run SAT instead.
+
+		// GJK converged or we have a touching contact
+		{
+			TVec3<T> ClosestA(0);
+			TVec3<T> ClosestB(0);
+			for (int i = 0; i < SimplexIDs.NumVerts; ++i)
+			{
+				ClosestA += InOutSimplexData.As[i] * InOutSimplexData.Barycentric[i];
+				ClosestB += InOutSimplexData.Bs[i] * InOutSimplexData.Barycentric[i];
+			}
+
+			OutNormalA = Normal;
+			OutNormalB = BToATM.InverseTransformVectorNoScale(Normal);
+
+			T Penetration = ThicknessA + ThicknessB - Distance;
+			OutPenetration = Penetration;
+			OutClosestA = ClosestA + OutNormalA * ThicknessA;
+			OutClosestB = ClosestB - OutNormalB * ThicknessB;
+
+			// If we don't care about separation distance/normal, the return value is true if we are overlapping, false otherwise.
+			// If we do care about seperation distance/normal, the return value is true if we found a solution.
+			// @todo(chaos): we should pass back failure for the degenerate case so we can decide how to handle it externally.
+			return (bFindSeparationDistance || (Penetration >= 0.0f));
+		}
+	}
+
+
 	// Calculate the penetration depth (or separating distance) of two geometries.
 	//
 	// Set bNegativePenetrationAllowed to false (default) if you do not care about the normal and distance when the shapes are separated. The return value
@@ -279,6 +553,7 @@ namespace Chaos
 	}
 
 
+
 	/** Sweeps one geometry against the other
 	 @A The first geometry
 	 @B The second geometry
diff --git a/Engine/Source/Runtime/Experimental/Chaos/Public/Chaos/GeometryParticlesfwd.h b/Engine/Source/Runtime/Experimental/Chaos/Public/Chaos/GeometryParticlesfwd.h
--- a/Engine/Source/Runtime/Experimental/Chaos/Public/Chaos/GeometryParticlesfwd.h
+++ b/Engine/Source/Runtime/Experimental/Chaos/Public/Chaos/GeometryParticlesfwd.h
@@ -63,6 +63,9 @@ namespace Chaos
 		static constexpr uint16 MaxBucketEntries = 1 << 13;
 		static constexpr uint16 MaxBuckets = 1 << 3;
 
+		FSpatialAccelerationIdx() : Bucket(0), InnerIdx(0) {};
+		FSpatialAccelerationIdx(uint16 inBucket, uint16 inInnerIdx) : Bucket(inBucket), InnerIdx(inInnerIdx) {};
+
 		bool operator==(const FSpatialAccelerationIdx& Rhs) const
 		{
 			return ((const uint16&)*this) == ((const uint16&)Rhs);
diff --git a/Engine/Source/Runtime/Experimental/Chaos/Public/Chaos/ISpatialAcceleration.h b/Engine/Source/Runtime/Experimental/Chaos/Public/Chaos/ISpatialAcceleration.h
--- a/Engine/Source/Runtime/Experimental/Chaos/Public/Chaos/ISpatialAcceleration.h
+++ b/Engine/Source/Runtime/Experimental/Chaos/Public/Chaos/ISpatialAcceleration.h
@@ -166,6 +166,8 @@ public:
 
 };
 
+using ISpatialDebugDrawInterface = ISpacialDebugDrawInterface<FReal>;
+
 using SpatialAccelerationType = uint8;	//see ESpatialAcceleration. Projects can add their own custom types by using enum values higher than ESpatialAcceleration::Unknown
 enum class ESpatialAcceleration : SpatialAccelerationType
 {
@@ -261,6 +263,10 @@ public:
 
 	virtual bool IsAsyncTimeSlicingComplete() { return AsyncTimeSlicingComplete; }
 	virtual void ProgressAsyncTimeSlicing(bool ForceBuildCompletion = false) {}
+	virtual bool ShouldRebuild() { return true; }  // Used to find out if something changed since last reset for optimizations
+	virtual void ClearShouldRebuild() {}
+	virtual void PrepareCopyTimeSliced(const  ISpatialAcceleration<TPayloadType, T, 3>& InFrom) { check(false); }
+	virtual void ProgressCopyTimeSliced(const  ISpatialAcceleration<TPayloadType, T, 3>& InFrom, int MaximumBytesToCopy) { check(false); }
 
 	// IMPORTANT : (LWC) this API should be typed on Freal not T, as we want the query API to be using the highest precision while maintaining arbitrary internal precision for the acceleration structure ( based on T )
 	virtual TArray<TPayloadType> FindAllIntersections(const FAABB3& Box) const { check(false); return TArray<TPayloadType>(); }
@@ -299,8 +305,17 @@ public:
 		return nullptr;
 	}
 
+	virtual ISpatialAcceleration<TPayloadType, T, d>& operator=(const ISpatialAcceleration<TPayloadType, T, d>& Other)
+	{
+		Type = Other.Type;
+		SyncTimestamp = Other.SyncTimestamp;
+		AsyncTimeSlicingComplete = Other.AsyncTimeSlicingComplete;
+		return *this;
+	}
+
 #if !UE_BUILD_SHIPPING
 	virtual void DebugDraw(ISpacialDebugDrawInterface<T>* InInterface) const {}
+	virtual void DebugDrawLeaf(ISpacialDebugDrawInterface<T>& InInterface, const FLinearColor& InLinearColor, float InThickness) const {}
 	virtual void DumpStats() const {}
 #endif
 
@@ -428,6 +443,31 @@ template <typename TKey, typename TValue>
 class TArrayAsMap
 {
 public:
+	// @todo(chaos): rename with "F"
+	using ElementType = TValue;
+
+	// @todo(chaos): rename with "F"
+	struct Element
+	{
+#if CHAOS_SERIALIZE_OUT
+		TKey KeyToSerializeOut;
+#endif
+		TValue Entry;
+	};
+
+	int32 Num() const
+	{
+		return Entries.Num();
+	}
+
+	void Reserve(int32 Size)
+	{
+		Entries.Reserve(Size);
+#if CHAOS_SERIALIZE_OUT
+		KeysToSerializeOut.Reserve(Size);
+#endif
+	}
+
 	TValue* Find(const TKey& Key)
 	{
 		const int32 Idx = GetUniqueIdx(Key).Idx;
@@ -553,6 +593,14 @@ public:
 		}
 	}
 
+	void AddFrom(const TArrayAsMap<TKey, TValue>& Source, int32 SourceIndex)
+	{
+		Entries.Add(Source.Entries[SourceIndex]);
+#if CHAOS_SERIALIZE_OUT
+		KeysToSerializeOut.Add(Source.KeysToSerializeOut[SourceIndex]);
+#endif
+	}
+
 private:
 
 	struct FEntry
diff --git a/Engine/Source/Runtime/Experimental/Chaos/Public/Chaos/ISpatialAccelerationCollection.h b/Engine/Source/Runtime/Experimental/Chaos/Public/Chaos/ISpatialAccelerationCollection.h
--- a/Engine/Source/Runtime/Experimental/Chaos/Public/Chaos/ISpatialAccelerationCollection.h
+++ b/Engine/Source/Runtime/Experimental/Chaos/Public/Chaos/ISpatialAccelerationCollection.h
@@ -24,12 +24,12 @@ public:
 	, AllAsyncTasksComplete(true)
 	{}
 	static constexpr ESpatialAcceleration StaticType = ESpatialAcceleration::Collection;
-	virtual FSpatialAccelerationIdx AddSubstructure(TUniquePtr<ISpatialAcceleration<TPayloadType, T, d>>&& Substructure, uint16 Bucket) = 0;
+	virtual FSpatialAccelerationIdx AddSubstructure(TUniquePtr<ISpatialAcceleration<TPayloadType, T, d>>&& Substructure, uint16 Bucket, uint16 BucketInnerIdx) = 0;
 	virtual TUniquePtr<ISpatialAcceleration<TPayloadType, T, d>> RemoveSubstructure(FSpatialAccelerationIdx Idx) = 0;
 	virtual ISpatialAcceleration<TPayloadType, T, d>* GetSubstructure(FSpatialAccelerationIdx Idx) = 0;
 
 	/** This is kind of a hack to avoid virtuals. We simply route calls into templated functions */
-	virtual void PBDComputeConstraintsLowLevel(T Dt, FSpatialAccelerationBroadPhase& BroadPhase, FNarrowPhase& NarrowPhase, FAsyncCollisionReceiver& Receiver, CollisionStats::FStatData& StatData, IResimCacheBase* ResimCache) const = 0;
+	virtual void PBDComputeConstraintsLowLevel(T Dt, FSpatialAccelerationBroadPhase& BroadPhase, FNarrowPhase& NarrowPhase, IResimCacheBase* ResimCache) const = 0;
 	virtual TArray<FSpatialAccelerationIdx> GetAllSpatialIndices() const = 0;
 
 	bool IsBucketActive(uint8 BucketIdx) const
@@ -40,6 +40,24 @@ public:
 	bool IsAllAsyncTasksComplete() const { return AllAsyncTasksComplete; }
 	void SetAllAsyncTasksComplete(bool State) { AllAsyncTasksComplete = State; }
 
+	ISpatialAccelerationCollection<TPayloadType, FReal, 3>& operator=(const ISpatialAccelerationCollection<TPayloadType, FReal, 3>& Other)
+	{
+		ISpatialAcceleration<TPayloadType, FReal, 3>::operator=(Other);
+		check(ActiveBucketsMask == Other.ActiveBucketsMask);
+		AllAsyncTasksComplete = Other.AllAsyncTasksComplete;
+		return *this;
+	}
+
+	virtual ISpatialAcceleration<TPayloadType, FReal, 3>& operator=(const ISpatialAcceleration<TPayloadType, FReal, 3>& Other) override
+	{
+		check(false);	//not implemented
+		return *this;
+	}
+
+#if !UE_BUILD_SHIPPING
+	virtual void DebugDraw(ISpacialDebugDrawInterface<T>* InInterface) const = 0;
+#endif
+
 protected:
 	uint8 ActiveBucketsMask;
 	bool AllAsyncTasksComplete;
diff --git a/Engine/Source/Runtime/Experimental/Chaos/Public/Chaos/Island/IslandGraph.h b/Engine/Source/Runtime/Experimental/Chaos/Public/Chaos/Island/IslandGraph.h
new file mode 100644
--- /dev/null
+++ b/Engine/Source/Runtime/Experimental/Chaos/Public/Chaos/Island/IslandGraph.h
@@ -0,0 +1,247 @@
+// Copyright Epic Games, Inc. All Rights Reserved.
+#pragma once
+
+#include "Containers/Map.h"
+#include "Containers/SparseArray.h"
+
+namespace Chaos
+{
+
+/**
+ * Island Graph.
+ *
+ * This template implements a graph that will be stored on the island manager
+ * The goal here is to minimize memory allocation while doing graph operations. 
+ *
+ * @param NodeType The node type of the item that will be stored in each nodes
+ * @param EdgeType The edge type of the item that will be stored in each edges
+ */
+template<typename NodeType, typename EdgeType, typename IslandType>
+class CHAOS_API FIslandGraph
+{
+public:
+
+	static constexpr const int32 MaxCount = 100000;
+
+	/**
+	 * Update current graph to merge connected islands 
+	 */
+	void UpdateGraph();
+
+	/**
+	 * Reserve a number of nodes in memory before adding them later
+	 * @param NumNodes Number of nodes to be reserved in memory
+	 */
+	void ReserveNodes(const int32 NumNodes);
+
+	/**
+	 * Clear the graph nodes without deallocating memory
+	 */
+	void ResetNodes();
+
+	/**
+	 * Given a node item, add a node to the graph nodes
+	 * @param NodeItem Node Item to be stored in the graph node
+	 * @param bValidNode Check if a node should be considered for graph partitioning
+	 * @param IslandIndex Potential island index we want the node to belong to
+	 * @param bDiscardNode Node to check if a node will be dicarded ot not for island splitting
+	 * @return Node index that has just been added
+	 */
+	int32 AddNode(const NodeType& NodeItem, const bool bValidNode = true, const int32 IslandIndex = INDEX_NONE, const bool bDiscardNode = false);
+
+	/**
+	 * Remove a node from the graph nodes list
+	 * @param NodeItem Item to be removed from the nodes list
+	 */
+	void RemoveNode(const NodeType& NodeItem);
+
+	/**
+	 * Reserve a number of edges in memory before adding them later
+	 * @param NumEdges Number of edges to be reserved in memory
+	 */
+	void ReserveEdges(const int32 NumEdges);
+
+	/**
+	 * Reset the graph edges without deallocating memory
+	 */
+	void ResetEdges();
+
+	/**
+	 * Given an edge item, add an edge to the graph edges
+	 * @param EdgeItem Node Item to be stored in the graph edge
+	 * @param ItemContainer Item container id that will be stored on the graph edge
+	 * @param FirstNode First node index of the graph edge
+	 * @param SecondNode Second node index of the graph edge
+	 * @return Edge index that has just been added
+	 */
+	int32 AddEdge(const EdgeType& EdgeItem, const int32 ItemContainer, const int32 FirstNode,  const int32 SecondNode);
+
+	/**
+	 * Remove an edge from the graph edges list
+	 * @param EdgeItem Item to be removed from the edges list
+	 */
+	void RemoveEdge(const EdgeType& EdgeItem);
+
+	/**
+	* Remove all the edges from the graph
+	*/
+	void ResetIslands();
+
+	/**
+	 * Get the number of edges that are inside a graph
+	 * @return Number of graph edges
+	 */
+	FORCEINLINE int32 NumEdges() const { return GraphEdges.GetMaxIndex(); }
+
+	/**
+	 * Get the number of nodes that are inside a graph
+	 * @return Number of graph nodes
+	 */
+	FORCEINLINE int32 NumNodes() const { return GraphNodes.GetMaxIndex(); }
+
+	/**
+	 * Get the number of islands that are inside a graph
+	 * @return Number of graph islands
+	 */
+	FORCEINLINE int32 NumIslands() const { return GraphIslands.GetMaxIndex(); }
+
+	/** Graph node structure */
+	struct FGraphNode 
+	{
+		/** List of edges connected to the node */
+		TSparseArray<int32> NodeEdges;
+
+		/** Node Island Index (for static/kinematic particles could belong to several islands) */
+		int32 IslandIndex = INDEX_NONE;
+
+		/** Check if a node is valid (checked for graph partitionning) */
+		bool bValidNode = true;
+
+		/** Check if a node will be dicarded for graph splitting */
+		bool bDiscardNode = true;
+
+		/** Node counter to filter nodes already processed */
+		int32 NodeCounter = 0;
+
+		/** Node item that is stored per node */
+		NodeType NodeItem;
+	};
+
+	/** Graph edge structure */
+	struct FGraphEdge 
+	{
+		/** First node od the edge */
+		int32 FirstNode = INDEX_NONE;
+
+		/** Second node of the edge*/
+		int32 SecondNode = INDEX_NONE;
+
+		/** Current edge index in the list of edges of the first node */
+		int32 FirstEdge = INDEX_NONE;
+
+		/** Current edge index in the list of edges of the second node  */
+		int32 SecondEdge = INDEX_NONE;
+
+		/** Unique egde island index */
+		int32 IslandIndex = INDEX_NONE;
+
+		/** Edge counter to filter edges already processed */
+		int32 EdgeCounter = 0;
+
+		/** Edge item that is stored per node */
+		EdgeType EdgeItem;
+
+		/** Item Container Id  */
+		int32 ItemContainer = 0;
+	};
+
+	/** Graph island structure */
+	struct FGraphIsland
+	{
+		/** Edge root of the graph island (to be replaced by a list of roots in case the splitting is not done every frame)*/
+		int32 RootIndex = INDEX_NONE;
+		
+		/** Number of edges per islands*/
+		int32 NumEdges = 0;
+
+		/** Number of nodes per islands*/
+		int32 NumNodes = 0;
+
+		/** Island counter to filter islands already processed */
+		int32 IslandCounter = 0;
+		
+		/** Boolean to check if an island is persistent or not */
+		bool bIsPersistent = true;
+
+		/** Boolean to check if an island is sleeping or not */
+		bool bIsSleeping = true;
+
+		/** List of children islands  to be merged */
+		TSet<int32> ChildrenIslands;
+		
+		/** Island Item that is stored per island*/
+		IslandType IslandItem;
+	};
+
+	/**
+	 * Loop over all the nodes within an island
+	 * @param IslandIndex Island in which we will loop over all the nodes
+	 * @param ApplyFunction Function to be applied on every nodes that are linked to the root node 
+	 * @param bUpdateCounter Boolean to check if we need to update the graph counter or not 
+	 */
+	void ForEachNodes(const int32 IslandIndex, TFunctionRef<void(FGraphNode&)> ApplyFunction, const bool bUpdateCounter = true);
+
+	/**
+	 * Loop over all the edges within an island
+	 * @param IslandIndex Island in which we will loop over all the edges
+	 * @param ApplyFunction Function to be applied on every nodes that are linked to the root node
+	 * @param bUpdateCounter Boolean to check if we need to update the graph counter or not
+	 */
+	void ForEachEdges(const int32 IslandIndex, TFunctionRef<void(FGraphEdge&)> ApplyFunction, const bool bUpdateCounter = true);
+
+	/**
+	 * Link two islands to each other for future island traversal
+	 * @param FirstIsland First island index to be linked
+	 * @param SecondIsland Second island index to be linked
+	 */
+	void ParentIslands(const int32 FirstIsland, const int32 SecondIsland);
+
+	/**
+	 * Given an edge index we try to attach islands. 
+	 * @param EdgeIndex Index of the edge used to attach the islands
+	 */
+	void AttachIslands(const int32 EdgeIndex);
+
+	/**
+	 * Merge 2 islands
+	 * @param ParentIndex Index of the parent island
+	 * @param ChildIndex Index of the child island
+	 */
+	void MergeIslands(const int32 ParentIndex, const int32 ChildIndex);
+
+	/**
+	* Split an island
+	* @param IslandIndex Island index to be splitted
+	*/
+	void SplitIsland(const int32 IslandIndex);
+
+	/** List of graph nodes */
+	TSparseArray<FGraphNode> GraphNodes;
+
+	/** List of graph edges */
+	TSparseArray<FGraphEdge> GraphEdges;
+
+	/** List of graph islands */
+	TSparseArray<FGraphIsland> GraphIslands;
+
+	/** Reverse list to find given an item the matching node */
+	TMap<NodeType,int32> ItemNodes;
+
+	/** Reverse list to find given an item the matching edge */
+	TMap<EdgeType,int32> ItemEdges;
+
+	/** Graph counter used for graph traversal to check if an edge/node/island has already been processed */
+	int32 GraphCounter = 0;
+};
+
+}
\ No newline at end of file
diff --git a/Engine/Source/Runtime/Experimental/Chaos/Public/Chaos/Island/IslandManager.h b/Engine/Source/Runtime/Experimental/Chaos/Public/Chaos/Island/IslandManager.h
new file mode 100644
--- /dev/null
+++ b/Engine/Source/Runtime/Experimental/Chaos/Public/Chaos/Island/IslandManager.h
@@ -0,0 +1,270 @@
+// Copyright Epic Games, Inc. All Rights Reserved.
+#pragma once
+
+#include "Containers/SparseArray.h"
+
+#include "Chaos/Real.h"
+#include "Chaos/Vector.h"
+#include "Chaos/Framework/Handles.h"
+#include "Chaos/ParticleHandleFwd.h"
+#include "Chaos/GeometryParticlesfwd.h"
+#include "Chaos/Island/IslandGraph.h"
+#include "Chaos/Island/SolverIsland.h"
+
+namespace Chaos
+{
+
+/** Forward declaration */
+template<typename T, int D>
+class TPBDRigidParticles;
+template<class T>
+class TChaosPhysicsMaterial;
+template <typename T>
+class TSerializablePtr;
+template<typename T>
+class TArrayCollectionArray;
+template <typename T>
+class TParticleView;
+class FPBDRigidsSOAs;
+class FConstraintHandle;
+class FChaosPhysicsMaterial;
+
+using FPBDRigidParticles = TPBDRigidParticles<FReal, 3>;
+	
+/** Island manager responsible to create the list of solver islands that will be persistent over time */
+class CHAOS_API FPBDIslandManager
+{
+public:
+	using GraphType = FIslandGraph<FGeometryParticleHandle*, FConstraintHandle*, FPBDIslandSolver*>;
+	using FGraphNode = GraphType::FGraphNode;
+	using FGraphEdge = GraphType::FGraphEdge;
+	
+	/**
+	  * Default island manager constructor
+	  */
+	FPBDIslandManager();
+
+	/**
+	  * Default island manager constructor
+	  * @param Particles List of particles to be used to fill the graph nodes
+	  */
+	FPBDIslandManager(const TParticleView<FGeometryParticles>& Particles);
+
+	/**
+	* Clear the manager and set up the particle-to-graph-node mapping for the specified particles
+	* Should be called before AddConstraint.
+	* @param Particles List of particles to be used to fill the graph nodes
+	*/
+	void InitializeGraph(const TParticleView<TGeometryParticles<FReal, 3>>& Particles);
+
+	/**
+	  * Reserve space in the graph for NumConstraints additional constraints.
+	  * @param NumConstraints Number of constraints to be used to reserved memory space
+	  */
+	void ReserveConstraints(const int32 NumConstraints);
+
+	/**
+	  * Add a constraint to the graph for each constraint in the container.
+	  * @param ContainerId Contaner id the constraint is belonging to
+	  * @param ConstraintHandle Constraint Handle that will be used for the edge item
+	  * @param ConstrainedParticles List of 2 particles handles that are used within the constraint
+	  * @return Edge index within the sparse graph edges list
+	  */
+	int32 AddConstraint(const uint32 ContainerId, FConstraintHandle* ConstraintHandle, const TVector<FGeometryParticleHandle*, 2>& ConstrainedParticles);
+
+	/**
+	  * Remove a constraint from the graph
+	  * @param ContainerId Contaner id the constraint is belonging to
+	  * @param ConstraintHandle Constraint Handle that will be used for the edge item
+	  * @param ConstrainedParticles List of 2 particles handles that are used within the constraint
+	  */
+	void RemoveConstraint(const uint32 ContainerId, FConstraintHandle* ConstraintHandle, const TVector<FGeometryParticleHandle*, 2>& ConstrainedParticles);
+	
+	/**
+	  * Preallocate buffers for \p Num particles.
+	  * @param NumParticles Number of particles to be used to reserve memory
+	  * @return Number of new slots created.
+	  */
+	int32 ReserveParticles(const int32 NumParticles);
+
+	/**
+	  * Add a particle to the sialnd graph
+	  * @param AddedParticle Particle to be added to the graph
+	  * @param IslandIndex Potential island index in which the particle will be added
+	  * @return Sparse graph node index
+	  */
+	int32 AddParticle(FGeometryParticleHandle* AddedParticle, const int32 IslandIndex = INDEX_NONE);
+
+	/**
+	  * Remove a particle from the graph
+	  * @param Particle Particle to be removed from the islands
+	  */
+	void RemoveParticle(FGeometryParticleHandle* Particle);
+
+	/**
+	  * Adds \p ChildParticle to the constraint graph.  Copies the sleeping
+	  * state and island of \p ParentParticle to \p ChildParticle.  Does nothing
+	  * if \p ParentParticle is not supplied.
+	  * @param ChildParticle Particle oni which the sleeping state and island index will be copied
+	  * @param ParentParticle Particle from which the sleeping state and island index will be extracted
+	  */
+	void EnableParticle(FGeometryParticleHandle* ChildParticle, const FGeometryParticleHandle* ParentParticle);
+
+	/**
+	  * Disable a particle and remove it from its island.
+	  * @note: this does not remove any attached constraints - constraints need to be re-added and islands will need to be rebuilt.
+	  * @param Particle Particle to be disabled
+	  */
+	void DisableParticle(FGeometryParticleHandle* Particle);
+
+	/**
+	  * Remove all the graph edges
+	  */
+	void ResetIslands(const TParticleView<FPBDRigidParticles>& PBDRigids);
+
+	/**
+	  * Generate the simulation islands of connected particles (AddConstraints must have already been called). There are no constraints connecting particles in different islands.
+	  * @param PBDRigids Rigid particles that will be used to fill the solver islands
+	  * @param Particles SOA used in the evolution
+	  */
+	void UpdateIslands(const TParticleView<FPBDRigidParticles>& PBDRigids, FPBDRigidsSOAs& Particles);
+
+	/**
+	  * @brief Put all particles and constraints in an island to sleep
+	  * @param Particles used in the islands
+	  * @param IslandIndex Island index to be woken up
+	*/
+	void SleepIsland(FPBDRigidsSOAs& Particles, const int32 IslandIndex);
+
+	/**
+	  * Wake all particles and constraints in an Island.
+	  * @param Particles used in the islands
+	  * @param IslandIndex Island index to be woken up
+	  */
+	void WakeIsland(FPBDRigidsSOAs& Particles, const int32 IslandIndex);
+
+	/**
+	  * Put particles in inactive islands to sleep.
+	  * @param IslandIndex Island index
+	  * @param PerParticleMaterialAttributes Material attributes that has been set on each particles
+	  * @param SolverPhysicsMaterials Solver materials that could be used to find the sleeping thresholds
+	  * @return return true if the island is going to sleep
+	  */
+	bool SleepInactive(const int32 IslandIndex, const TArrayCollectionArray<TSerializablePtr<FChaosPhysicsMaterial>>& PerParticleMaterialAttributes, const THandleArray<FChaosPhysicsMaterial>& SolverPhysicsMaterials);
+	
+	/**
+	  * Get the island particles stored in the solver island
+	  * @param IslandIndex Island index
+	  * @return List of particles handles in the island
+	  */
+	const TArray<FGeometryParticleHandle*>& GetIslandParticles(const int32 IslandIndex) const;
+
+	/**
+	  * Get the island constraints stored in the solver island
+	  * @param IslandIndex Island index
+	  * @return List of constraints handles in the island
+	  */
+	const TArray<FConstraintHandle*>& GetIslandConstraints(const int32 IslandIndex) const;
+
+	/**
+	 * When resim is used, tells us whether we need to resolve island
+	 * @param IslandIndex Island index
+	 * @return If true the island need to be resim
+	 */
+	bool IslandNeedsResim(const int32 IslandIndex) const;
+
+	/**
+	  * The number of islands in the graph.
+	  * @return the number of islands
+	  */
+	FORCEINLINE int32 NumIslands() const
+	{
+		return IslandIndexing.Num();
+	}
+	
+	/**
+	 * The number of islands in the graph.
+	 * @return Get the island graph
+	 */
+	FORCEINLINE const GraphType* GetIslandGraph() const
+	{
+		return IslandGraph.Get();
+	}
+
+	/**
+	 * Get Max UniqueIdx of the particles.
+	 * Used to create arrays large enough to use Particle.UniqueIdx().Idx for indexing.
+	 * @return the max index of the particles in the islands
+	 */
+	inline int32 GetMaxParticleUniqueIdx() const { return MaxParticleIndex; }
+
+
+	/**
+	 * Get the sparse island graph index given a dense index between 0 and NumIslands
+	 * @param IslandIndex Island index
+	 * @return the max index of the particles in the islands
+	 */
+	int32 GetGraphIndex(const int32 IslandIndex) const {
+		return IslandIndexing.IsValidIndex(IslandIndex) ? IslandIndexing[IslandIndex] : 0;
+	}
+	
+	/** Add a constraint container to all the solver island given a container id
+	* @param ContainerId Constraints container id from which the solver constraint containers are being built
+	*/
+	template<typename ConstraintType>
+	void AddConstraintDatas(const int32 ContainerId)
+	{
+		for(auto& IslandSolver : IslandSolvers)
+		{
+			IslandSolver->AddConstraintDatas<ConstraintType>(ContainerId);
+		}
+	}
+
+	// /** Add a constraint container to all the solver island given a container id
+	// * @param ContainerId Constraints container id from which the solver constraint containers are being built
+	// */
+	// template<typename ConstraintType>
+	// void GatherSolverInput(const FReal Dt, int32 IslandIndex, const int32 ContainerId)
+	// {
+	// 	if(IslandIndexing.IsValidIndex(IslandIndex))
+	// 	{
+	// 		IslandSolvers[IslandIndexing[IslandIndex]]->GatherSolverInput<ConstraintType>(Dt, IslandIndex, ContainerId);
+	// 	}
+	// }
+
+	/** Accessors of the graph island */
+	const FPBDIslandSolver* GetSolverIsland(const int32 IslandIndex) const {return IslandIndexing.IsValidIndex(IslandIndex) ? IslandSolvers[IslandIndexing[IslandIndex]].Get() : nullptr; }
+	FPBDIslandSolver* GetSolverIsland(const int32 IslandIndex) {return IslandIndexing.IsValidIndex(IslandIndex) ? IslandSolvers[IslandIndexing[IslandIndex]].Get() : nullptr; }
+
+	/** Get the graph nodes */
+	const TSparseArray<GraphType::FGraphNode>& GetGraphNodes() const { return IslandGraph->GraphNodes; }
+	
+	/** Get the graph edges */
+	const TSparseArray<GraphType::FGraphEdge>& GetGraphEdges() const { return IslandGraph->GraphEdges; }
+
+	/** Get the particle nodes */
+	const TMap<FGeometryParticleHandle*, int32>& GetParticleNodes() const { return IslandGraph->ItemNodes; }
+	
+	/** Get the constraints edges */
+	const TMap<FConstraintHandle*, int32>& GetConstraintEdges() const { return IslandGraph->ItemEdges; }
+	
+protected:
+
+	/**
+	* Sync the IslandSolvers with the IslandGraph
+	*/
+	void SyncIslands(FPBDRigidsSOAs& Particles);
+
+	/** List of solver islands in the manager */
+	TSparseArray<TUniquePtr<FPBDIslandSolver>> IslandSolvers;
+
+	/** Island map to make the redirect an index between 0...NumIslands to the persistent island index  */
+	TArray<int32>							IslandIndexing;
+
+	/** Global graph to deal with merging and splitting of islands */
+	TUniquePtr<GraphType>					IslandGraph;
+
+	/** Maximum particle index that we have in the graph*/
+	int32									MaxParticleIndex = INDEX_NONE;
+};
+}
\ No newline at end of file
diff --git a/Engine/Source/Runtime/Experimental/Chaos/Public/Chaos/Island/SolverIsland.h b/Engine/Source/Runtime/Experimental/Chaos/Public/Chaos/Island/SolverIsland.h
new file mode 100644
--- /dev/null
+++ b/Engine/Source/Runtime/Experimental/Chaos/Public/Chaos/Island/SolverIsland.h
@@ -0,0 +1,163 @@
+// Copyright Epic Games, Inc. All Rights Reserved.
+#pragma once
+
+#include "Containers/Array.h"
+#include "Chaos/ParticleHandle.h"
+#include "Chaos/ConstraintHandle.h"
+#include "Chaos/Evolution/SolverDatas.h"
+
+namespace Chaos
+{
+	
+/** Forward Declaration */
+class FPBDIslandManager;
+	
+/**
+* List of per island datas that will created bu the island manager
+*/
+class CHAOS_API FPBDIslandSolver : public FPBDIslandSolverData
+{
+	public:
+
+	/**
+	* Init the solver island with an island manager and an index
+	* @param IslandManager Island manager that is storing the solver islands
+	* @param IslandIndex Index of the solver island within that list
+	*/
+	FPBDIslandSolver(const FPBDIslandManager* IslandManager, const int32 IslandIndex);
+
+	/**
+	* Clear all the particles from the island
+	*/
+	void ClearParticles();
+
+	/**
+	* Add particle to the island
+	* @param ParticleHandle ParticleHandle to be added
+	*/
+	void AddParticle(FGenericParticleHandle ParticleHandle);
+
+	/**
+	* Reset the particles list and reserve a number of particles in memory
+	* @param NumParticles Number of particles to be reserved
+	*/
+	void ReserveParticles(const int32 NumParticles);
+
+	/**
+	* Remove particle from the island
+	* @param ParticleHandle ParticleHandle to be removed
+	*/
+	void RemoveParticle(FGenericParticleHandle ParticleHandle);
+	
+	/**
+	* Update the particles island index to match the graph index
+	*/
+	void UpdateParticles();
+	
+	/**
+	* Reset the constraints list and reserve a number of constraints in memory 
+	* @param NumConstraints number of constraints to be reserved
+	*/
+	void ReserveConstraints(const int32 NumConstraints);
+
+	/**
+	* Remove all the constraints from the solver island
+	*/
+	void ClearConstraints();
+
+	/**
+	* Add constraint to the island
+	* @param ConstraintHandle ConstraintHandle to be added
+	*/
+	void AddConstraint(FConstraintHandle* ConstraintHandle);
+
+	/**
+	* Remove constraint from the island
+	* @param ConstraintHandle ConstraintHandle to be removed
+	*/
+	void RemoveConstraint(FConstraintHandle* ConstraintHandle);
+
+	/**
+	* Return the list of particles within the solver island
+	*/
+	FORCEINLINE const TArray<FGeometryParticleHandle*>& GetParticles() const { return IslandParticles; }
+
+	/**
+	* Return the list of constraints within the island
+	*/
+	FORCEINLINE const TArray<FConstraintHandle*>& GetConstraints() const { return IslandConstraints; }
+	
+	/**
+	* Get the number of particles within the island
+	*/
+	FORCEINLINE int32 NumParticles() const {return IslandParticles.Num(); }
+
+	/**
+	* Get the number of constraints within the island
+	*/
+	FORCEINLINE int32 NumConstraints() const {return IslandConstraints.Num(); }
+	
+	/**
+	* Members accessors
+	*/
+	FORCEINLINE bool GetIsSleeping() const {return bIsSleeping;}
+	FORCEINLINE void SetIsSleeping(const bool bIsSleepingIn ) { bIsSleeping = bIsSleepingIn; }
+	FORCEINLINE bool GetIsPersistent() const { return bIsPersistent; }
+	FORCEINLINE void SetIsPersistent(const bool bIsPersistentIn) { bIsPersistent = bIsPersistentIn; }
+	FORCEINLINE bool GetNeedsResim() const { return bNeedsResim; }
+	FORCEINLINE void SetNeedsResim(const bool bNeedsResimIn) { bNeedsResim = bNeedsResimIn; }
+	FORCEINLINE int32 GetSleepCounter() const { return SleepCounter; }
+	FORCEINLINE void SetSleepCounter(const int32 SleepCounterIn) { SleepCounter = SleepCounterIn; }
+	
+	// template<typename ConstraintType>
+	// void GatherSolverInput(const FReal Dt, const int32 IslandIndex, const int32 ContainerId);
+	
+	private:
+
+	/** Island manager that is storing the IslandSolver */
+	const FPBDIslandManager* IslandManager = nullptr;
+
+	/** Flag to check if an island is awake or sleeping */
+	bool bIsSleeping = false;
+
+	/** Flag to check if an island need to be re-simulated or not */
+	bool bNeedsResim = false;
+
+	/** Flag to check if an island is persistent over time */
+	bool bIsPersistent = true;
+
+	/** Sleep counter to trigger island sleeping */
+	int32 SleepCounter = 0;
+
+	/** List of all the island particles handles */
+	TArray<FGeometryParticleHandle*> IslandParticles;
+
+	/** List of all the island constraints handles */
+	TArray<FConstraintHandle*> IslandConstraints;
+};
+	
+// template<typename ConstraintType>
+// FORCEINLINE void FPBDIslandSolver::GatherSolverInput(const FReal Dt, const int32 IslandIndex, const int32 ContainerId)
+// {
+// 	using FConstraintContainerHandle = typename ConstraintType::FConstraintContainerHandle;
+// 	if(ConstraintContainers.IsValidIndex(ContainerId))
+// 	{
+// 		for (FConstraintHandle* ConstraintHandle : IslandConstraints)
+// 		{
+// 			if (ConstraintHandle->GetContainerId() == ContainerId)
+// 			{
+// 				FConstraintContainerHandle* Constraint = ConstraintHandle->As<FConstraintContainerHandle>();
+//
+// 				// Note we are building the SolverBodies as we go, in the order that we visit them. Each constraint
+// 				// references two bodies, so we won't strictly be accessing only in cache order, but it's about as good as it can be.
+// 				if (Constraint->IsEnabled())
+// 				{
+// 					// @todo(chaos): we should provide Particle Levels in the island rule as well (see TPBDConstraintColorRule)
+// 					Constraint->GatherInput(Dt, GraphIndex, INDEX_NONE, INDEX_NONE, *this);
+// 				}
+// 			}
+// 		}
+// 	}
+// }
+	
+}
\ No newline at end of file
diff --git a/Engine/Source/Runtime/Experimental/Chaos/Public/Chaos/Joint/JointSolverConstraints.h b/Engine/Source/Runtime/Experimental/Chaos/Public/Chaos/Joint/JointSolverConstraints.h
--- a/Engine/Source/Runtime/Experimental/Chaos/Public/Chaos/Joint/JointSolverConstraints.h
+++ b/Engine/Source/Runtime/Experimental/Chaos/Public/Chaos/Joint/JointSolverConstraints.h
@@ -9,7 +9,7 @@
 
 namespace Chaos
 {
-	class FJointSolverGaussSeidel;
+	class FPBDJointSolver;
 
 	// The method used to calculate the axis and angle for one of the rotation constraints in the joint
 	enum class EJointSolverConstraintUpdateType
@@ -181,7 +181,7 @@ namespace Chaos
 	class FJointSolverConstraints
 	{
 	public:
-		friend class FJointSolverGaussSeidel;
+		friend class FPBDJointSolver;
 
 		FJointSolverConstraints();
 
diff --git a/Engine/Source/Runtime/Experimental/Chaos/Public/Chaos/Joint/PBDJointSolverGaussSeidel.h b/Engine/Source/Runtime/Experimental/Chaos/Public/Chaos/Joint/PBDJointSolverGaussSeidel.h
--- a/Engine/Source/Runtime/Experimental/Chaos/Public/Chaos/Joint/PBDJointSolverGaussSeidel.h
+++ b/Engine/Source/Runtime/Experimental/Chaos/Public/Chaos/Joint/PBDJointSolverGaussSeidel.h
@@ -4,6 +4,7 @@
 #include "CoreMinimal.h"
 
 #include "Chaos/DenseMatrix.h"
+#include "Chaos/Evolution/SolverBody.h"
 #include "Chaos/ParticleHandleFwd.h"
 #include "Chaos/PBDJointConstraintTypes.h"
 #include "Chaos/PBDJointConstraintUtilities.h"
@@ -11,6 +12,11 @@
 
 namespace Chaos
 {
+	class FPBDJointSolver;
+
+	UE_DEPRECATED(4.27, "Use FPBDJointSolver")
+	typedef FPBDJointSolver FJointSolverGaussSeidel;
+
 	/**
 	 * Calculate new positions and rotations for a pair of bodies connected by a joint.
 	 *
@@ -19,134 +25,149 @@ namespace Chaos
 	 *
 	 * \see FJointSolverCholesky
 	 */
-	class FJointSolverGaussSeidel
+	class FPBDJointSolver
 	{
 	public:
 		static const int32 MaxConstrainedBodies = 2;
 
-		FORCEINLINE const FVec3& GetP(const int32 Index) const
+		// Access to the SolverBodies that the constraint is operating on.
+		// \note SolverBodies are transient and exist only as long as we are in the Island's constraint solver loop
+		inline FConstraintSolverBody& Body(int32 BodyIndex)
+		{
+			check((BodyIndex >= 0) && (BodyIndex < 2));
+			check(SolverBodies[BodyIndex].IsValid());
+
+			return SolverBodies[BodyIndex];
+		}
+
+		inline const FConstraintSolverBody& Body(int32 BodyIndex) const
+		{
+			check((BodyIndex >= 0) && (BodyIndex < 2));
+			check(SolverBodies[BodyIndex].IsValid());
+
+			return SolverBodies[BodyIndex];
+		}
+
+		inline const FVec3& X(int BodyIndex) const
+		{
+			return Body(BodyIndex).X();
+		}
+
+		inline const FRotation3& R(int BodyIndex) const
+		{
+			return Body(BodyIndex).R();
+		}
+
+		inline const FVec3& P(int BodyIndex) const
 		{
-			checkSlow(Index < MaxConstrainedBodies);
-			return Ps[Index];
+			return Body(BodyIndex).P();
 		}
 
-		FORCEINLINE const FRotation3& GetQ(const int32 Index) const
+		inline const FRotation3& Q(int BodyIndex) const
 		{
-			checkSlow(Index < MaxConstrainedBodies);
-			return Qs[Index];
+			return Body(BodyIndex).Q();
 		}
 
-		FORCEINLINE const FVec3& GetInitP(const int32 Index) const
+		inline const FVec3& V(int BodyIndex) const
 		{
-			checkSlow(Index < MaxConstrainedBodies);
-			return InitPs[Index];
+			return Body(BodyIndex).V();
 		}
 
-		FORCEINLINE const FRotation3& GetInitQ(const int32 Index) const
+		inline const FVec3& W(int BodyIndex) const
 		{
-			checkSlow(Index < MaxConstrainedBodies);
-			return InitQs[Index];
+			return Body(BodyIndex).W();
 		}
 
-		FORCEINLINE const FVec3& GetV(const int32 Index) const
+		inline FReal InvM(int32 BodyIndex) const
 		{
-			checkSlow(Index < MaxConstrainedBodies);
-			return Vs[Index];
+			return Body(BodyIndex).InvM();
 		}
 
-		FORCEINLINE const FVec3& GetW(const int32 Index) const
+		inline FMatrix33 InvI(int32 BodyIndex) const
 		{
-			checkSlow(Index < MaxConstrainedBodies);
-			return Ws[Index];
+			return Body(BodyIndex).InvI();
 		}
 
-		FORCEINLINE const FVec3& GetNetLinearImpulse() const
+		inline const FVec3& GetNetLinearImpulse() const
 		{
 			return NetLinearImpulse;
 		}
 
-		FORCEINLINE const FVec3& GetNetAngularImpulse() const
+		inline const FVec3& GetNetAngularImpulse() const
 		{
 			return NetAngularImpulse;
 		}
 
-		FORCEINLINE int32 GetNumActiveConstraints() const
+		inline int32 GetNumActiveConstraints() const
 		{
 			// We use -1 as unitialized, but that should not be exposed outside the solver
 			return FMath::Max(NumActiveConstraints, 0);
 		}
 
-		FORCEINLINE bool GetIsActive() const
+		inline bool GetIsActive() const
 		{
 			return bIsActive;
 		}
 
-		FORCEINLINE FReal InvM(int32 index) const
-		{
-			return InvMs[index];
-		}
 
+		FPBDJointSolver();
 
-		FJointSolverGaussSeidel();
-
+		// Called once per frame to initialze the joint solver from the joint settings
 		void Init(
 			const FReal Dt,
+			const FSolverBodyPtrPair& SolverBodyPair,
 			const FPBDJointSolverSettings& SolverSettings,
 			const FPBDJointSettings& JointSettings,
-			const FVec3& PrevP0,
-			const FVec3& PrevP1,
-			const FRotation3& PrevQ0,
-			const FRotation3& PrevQ1,
-			const FReal InvM0,
-			const FVec3& InvIL0,
-			const FReal InvM1,
-			const FVec3& InvIL1,
 			const FRigidTransform3& XL0,
 			const FRigidTransform3& XL1);
 
+		void Deinit();
+
+		// @todo(chaos): this doesn't do much now we have SolverBodies - remove it
 		void Update(
 			const FReal Dt,
-			const FReal InSolverStiffness,
 			const FPBDJointSolverSettings& SolverSettings,
-			const FPBDJointSettings& JointSettings,
-			const FVec3& P0,
-			const FRotation3& Q0,
-			const FVec3& V0,
-			const FVec3& W0,
-			const FVec3& P1,
-			const FRotation3& Q1,
-			const FVec3& V1,
-			const FVec3& W1);
+			const FPBDJointSettings& JointSettings);
 
+		// Run the position solve for the constraints
 		void ApplyConstraints(
 			const FReal Dt,
+			const FReal InSolverStiffness,
 			const FPBDJointSolverSettings& SolverSettings,
 			const FPBDJointSettings& JointSettings);
 
+		// Run the velocity solve for the constraints
 		void ApplyVelocityConstraints(
 			const FReal Dt,
+			const FReal InSolverStiffness,
 			const FPBDJointSolverSettings& SolverSettings,
 			const FPBDJointSettings& JointSettings);
 
+		// Apply projection (a position solve where the parent has infinite mass) for the constraints
+		// @todo(chaos): this can be build into the Apply phase when the whole solver is PBD
 		void ApplyProjections(
 			const FReal Dt,
+			const FReal InSolverStiffness,
 			const FPBDJointSolverSettings& SolverSettings,
 			const FPBDJointSettings& JointSettings);
 
 	private:
 
+		// Initialize state dependent on particle positions
 		void InitDerivedState();
 
+		// Update state dependent on particle positions (after moving one or both of them)
 		void UpdateDerivedState(const int32 BodyIndex);
-
 		void UpdateDerivedState();
 
+		// Check to see if this constraint still needs further solved
+		// @todo(chaos): the term "active" is used inconsistently with the meaning elsewhere. Active should
+		// mean "contributed impusles". Should use "solved" rather than "active"
 		bool UpdateIsActive(
 			const FReal Dt,
 			const FPBDJointSolverSettings& SolverSettings,
 			const FPBDJointSettings& JointSettings);
 
-
 		void ApplyPositionConstraints(
 			const FReal Dt,
 			const FPBDJointSolverSettings& SolverSettings,
@@ -217,6 +238,7 @@ namespace Chaos
 			const FReal JointStiffness,
 			const FVec3& Axis,
 			const FReal Delta,
+			const FVec3& Connector0Correction = FVec3(0),
 			const int32 LinearHardLambdaIndex = -1);
 
 		void ApplyPositionConstraintSoft(
@@ -470,7 +492,9 @@ namespace Chaos
 
 		void ApplyLinearVelocityConstraint(
 			const FReal Stiffness,
-			const FVec3& Axis);
+			const FVec3& Axis,
+			const FVec3& Connector0Correction = FVec3(0),
+			const FReal TargetVel = 0);
 
 		void ApplyPointVelocityConstraint(
 			const FReal Dt,
@@ -535,7 +559,8 @@ namespace Chaos
 
 		void ApplyAngularVelocityConstraint(
 			const FReal Stiffness,
-			const FVec3& Axis);
+			const FVec3& Axis,
+			const FReal TargetVel = 0.0f);
 
 		void ApplyLockedRotationVelocityConstraints(
 			const FReal Dt,
@@ -593,65 +618,68 @@ namespace Chaos
 			AngularConstraintPadding[(int32)ConstraintIndex] = Padding;
 		}
 
+		// Calculate linear velocities along constraint axes based on the constraint type in JointSettings. CV is the relative linear velocity of joint connectors.
+		void CalculateConstraintAxisLinearVelocities(
+			const FPBDJointSettings& JointSettings,
+			FVec3& ConstraintAxisLinearVelocities) const;
 
-		using FDenseMatrix66 = TDenseMatrix<6 * 6>;
-		using FDenseMatrix61 = TDenseMatrix<6 * 1>;
+		void CalculateSphericalConstraintAxisLinearVelocities(
+			const FPBDJointSettings& JointSettings,
+			FVec3& ConstraintAxisLinearVelocities
+			) const;
 
-		int32 AddLinear(
-			const FReal Dt,
-			const FVec3& Axis,
-			const FVec3& ConnectorOffset0,
-			const FVec3& ConnectorOffset1,
-			const FReal Error,
-			const FReal VelTarget,
-			const FReal Stiffness,
-			const FReal Damping,
-			const bool bSoft,
-			const bool bAccelerationMode,
-			FDenseMatrix66& J0,
-			FDenseMatrix66& J1,
-			FDenseMatrix61& C,
-			FDenseMatrix61& V,
-			FDenseMatrix61& S,
-			FDenseMatrix61& D);
+		void CalculateCylindricalConstraintAxisLinearVelocities(
+			const int32 AxisIndex,
+			const EJointMotionType AxialMotion,
+			const EJointMotionType RadialMotion,
+			const FPBDJointSettings& JointSettings,
+			FVec3& ConstraintAxisLinearVelocities) const;
 
-		int32 AddAngular(
-			const FReal Dt,
-			const FVec3& Axis,
-			const FReal Error,
-			const FReal VelTarget,
-			const FReal Stiffness,
-			const FReal Damping,
-			const bool bSoft,
-			const bool bAccelerationMode,
-			FDenseMatrix66& J0,
-			FDenseMatrix66& J1,
-			FDenseMatrix61& C,
-			FDenseMatrix61& V,
-			FDenseMatrix61& S,
-			FDenseMatrix61& D);
+		void CalculatePlanarConstraintAxisLinearVelocities(
+			const int32 AxisIndex,
+			const EJointMotionType AxialMotion,
+			const FPBDJointSettings& JointSettings,
+			FVec3& ConstraintAxisLinearVelocities) const;
 
-		// Local-space constraint settings
-		FRigidTransform3 XLs[MaxConstrainedBodies];	// Local-space joint connector transforms
-		FVec3 InvILs[MaxConstrainedBodies];			// Local-space inverse inertias
-		FReal InvMs[MaxConstrainedBodies];			// Inverse masses
+		// Calculate angular velocities along constraint axes based on the constraint type in JointSettings. CW is the relative angular velocity of joint connectors.
+		void CalculateConstraintAxisAngularVelocities(
+			const FPBDJointSolverSettings& SolverSettings,
+			const FPBDJointSettings& JointSettings,
+			FVec3& ConstraintAxisAngularVelocities
+			) const;
+
+		void CalculateTwistConstraintAxisAngularVelocities(
+			const FPBDJointSettings& JointSettings,
+			FVec3& ConstraintAxisAngularVelocities) const;
 
-		// World-space constraint state
-		FVec3 Xs[MaxConstrainedBodies];				// World-space joint connector positions
-		FRotation3 Rs[MaxConstrainedBodies];		// World-space joint connector rotations
+		void CalculateConeConstraintAxisAngularVelocities(
+			const FPBDJointSettings& JointSettings,
+			FVec3& ConstraintAxisAngularVelocities) const;
+
+		void CalculateDualConeSwingConstraintAxisAngularVelocities(
+			const FPBDJointSettings& JointSettings,
+			const EJointAngularConstraintIndex SwingConstraintIndex,
+			FVec3& ConstraintAxisAngularVelocities) const;
+
+		void CalculateSwingConstraintAxisAngularVelocities(
+			const FPBDJointSolverSettings& SolverSettings,
+			const FPBDJointSettings& JointSettings,
+			const EJointAngularConstraintIndex SwingConstraintIndex,
+			FVec3& ConstraintAxisAngularVelocities) const;
+
+		// The cached body state on which the joint operates
+		FConstraintSolverBody SolverBodies[MaxConstrainedBodies];
+
+		// Local-space constraint settings
+		FRigidTransform3 LocalConnectorXs[MaxConstrainedBodies];	// Local(CoM)-space joint connector transforms
 
-		// World-space body state
-		FVec3 Ps[MaxConstrainedBodies];				// World-space particle CoM positions
-		FRotation3 Qs[MaxConstrainedBodies];		// World-space particle CoM rotations
-		FVec3 Vs[MaxConstrainedBodies];				// World-space particle CoM velocities
-		FVec3 Ws[MaxConstrainedBodies];				// World-space particle CoM angular velocities
-		FMatrix33 InvIs[MaxConstrainedBodies];		// World-space inverse inertias
+		// World-space constraint settings
+		FVec3 ConnectorXs[MaxConstrainedBodies];			// World-space joint connector positions
+		FRotation3 ConnectorRs[MaxConstrainedBodies];		// World-space joint connector rotations
 
 		// XPBD Initial iteration world-space body state
-		FVec3 InitPs[MaxConstrainedBodies];			// World-space particle positions
-		FRotation3 InitQs[MaxConstrainedBodies];	// World-space particle rotations
-		FVec3 InitXs[MaxConstrainedBodies];			// World-space joint connector positions
-		FRotation3 InitRs[MaxConstrainedBodies];	// World-space joint connector rotations
+		FVec3 InitConnectorXs[MaxConstrainedBodies];		// World-space joint connector positions
+		FRotation3 InitConnectorRs[MaxConstrainedBodies];	// World-space joint connector rotations
 
 		// Accumulated Impulse and AngularImpulse (Impulse * Dt since they are mass multiplied position corrections)
 		FVec3 NetLinearImpulse;
@@ -672,7 +700,7 @@ namespace Chaos
 		FVec3 RotationDriveLambdas;
 
 		// Solver stiffness - increased over iterations for stability
-		// @todo(chaos): remove Stiffness from SolverSettings (because it is not a solver constant)
+		// \todo(chaos): remove Stiffness from SolverSettings (because it is not a solver constant)
 		FReal SolverStiffness;
 
 		// Constraint padding which can act something like a velocity constraint (for reslockfreetitution)
@@ -686,6 +714,8 @@ namespace Chaos
 		// Tracking whether the solver is resolved
 		FVec3 LastPs[MaxConstrainedBodies];			// Positions at the beginning of the iteration
 		FRotation3 LastQs[MaxConstrainedBodies];	// Rotations at the beginning of the iteration
+		FVec3 InitConstraintAxisLinearVelocities; // Linear velocities along the constraint axes at the begining of the frame, used by restitution
+		FVec3 InitConstraintAxisAngularVelocities; // Angular velocities along the constraint axes at the begining of the frame, used by restitution
 		int32 NumActiveConstraints;					// The number of active constraints and drives in the last iteration (-1 initial value)
 		bool bIsActive;								// Whether any constraints actually moved any bodies in last iteration
 	};
diff --git a/Engine/Source/Runtime/Experimental/Chaos/Public/Chaos/KinematicTargets.h b/Engine/Source/Runtime/Experimental/Chaos/Public/Chaos/KinematicTargets.h
--- a/Engine/Source/Runtime/Experimental/Chaos/Public/Chaos/KinematicTargets.h
+++ b/Engine/Source/Runtime/Experimental/Chaos/Public/Chaos/KinematicTargets.h
@@ -26,7 +26,9 @@ namespace Chaos
 	{
 	public:
 		TKinematicTarget()
-			: Mode(EKinematicTargetMode::None)
+			: Position(0)
+			, Rotation(TRotation<T, d>::FromIdentity())
+			, Mode(EKinematicTargetMode::None)
 		{
 		}
 
diff --git a/Engine/Source/Runtime/Experimental/Chaos/Public/Chaos/PBDChainConstraints.h b/Engine/Source/Runtime/Experimental/Chaos/Public/Chaos/PBDChainConstraints.h
--- a/Engine/Source/Runtime/Experimental/Chaos/Public/Chaos/PBDChainConstraints.h
+++ b/Engine/Source/Runtime/Experimental/Chaos/Public/Chaos/PBDChainConstraints.h
@@ -9,7 +9,7 @@
 
 namespace Chaos
 {
-class FPBDChainConstraints : public FPBDConstraintContainer
+class FPBDChainConstraints : public FPBDIndexedConstraintContainer
 {
 public:
 	FPBDChainConstraints(const FDynamicParticles& InParticles, TArray<TArray<int32>>&& Constraints, const FReal Coefficient = (FReal)1.)
diff --git a/Engine/Source/Runtime/Experimental/Chaos/Public/Chaos/PBDCollisionConstraints.h b/Engine/Source/Runtime/Experimental/Chaos/Public/Chaos/PBDCollisionConstraints.h
--- a/Engine/Source/Runtime/Experimental/Chaos/Public/Chaos/PBDCollisionConstraints.h
+++ b/Engine/Source/Runtime/Experimental/Chaos/Public/Chaos/PBDCollisionConstraints.h
@@ -3,7 +3,9 @@
 
 #include "Chaos/CollisionResolutionTypes.h"
 #include "Chaos/Collision/CollisionApplyType.h"
+#include "Chaos/Collision/CollisionConstraintAllocator.h"
 #include "Chaos/Collision/PBDCollisionConstraintHandle.h"
+#include "Chaos/Collision/SolverCollisionContainer.h"
 #include "Chaos/PBDConstraintContainer.h"
 #include "Framework/BufferedData.h"
 
@@ -13,16 +15,13 @@
 #include "BoundingVolume.h"
 #include "AABBTree.h"
 
-// @todo(chaos): optimize and re-enable persistent constraints if we want it
-#define CHAOS_COLLISION_PERSISTENCE_ENABLED 0
-
 namespace Chaos
 {
-class FCollisionConstraintBase;
 class FImplicitObject;
 class FPBDCollisionConstraints;
-class FRigidBodyPointContactConstraint;
 class FPBDRigidsSOAs;
+class FPBDCollisionConstraint;
+class FPBDIslandSolverData;
 
 using FRigidBodyContactConstraintsPostComputeCallback = TFunction<void()>;
 using FRigidBodyContactConstraintsPostApplyCallback = TFunction<void(const FReal Dt, const TArray<FPBDCollisionConstraintHandle*>&)>;
@@ -34,22 +33,32 @@ namespace Collisions
 	struct FContactIterationParameters;
 }
 
+DECLARE_CYCLE_STAT_EXTERN(TEXT("Collisions::Gather"), STAT_Collisions_Gather, STATGROUP_ChaosCollision, CHAOS_API);
+DECLARE_CYCLE_STAT_EXTERN(TEXT("Collisions::Scatter"), STAT_Collisions_Scatter, STATGROUP_ChaosCollision, CHAOS_API);
+DECLARE_CYCLE_STAT_EXTERN(TEXT("Collisions::Apply"), STAT_Collisions_Apply, STATGROUP_ChaosCollision, CHAOS_API);
+DECLARE_CYCLE_STAT_EXTERN(TEXT("Collisions::ApplyPushOut"), STAT_Collisions_ApplyPushOut, STATGROUP_ChaosCollision, CHAOS_API);
+
+
 /**
  * A container and solver for collision constraints.
+ * 
+ * @todo(chaos): remove handles array
  */
 class CHAOS_API FPBDCollisionConstraints : public FPBDConstraintContainer
 {
 public:
 	friend class FPBDCollisionConstraintHandle;
 
-	using Base = FPBDConstraintContainer;
-	using FHandles = TArray<FPBDCollisionConstraintHandle*>;
-	using FConstraintHandleAllocator = TConstraintHandleAllocator<FPBDCollisionConstraints>;
-	using FConstraintContainerHandleKey = typename FPBDCollisionConstraintHandle::FHandleKey;
+	using Base = FPBDIndexedConstraintContainer;
+
+	// Collision constraints have intrusive pointers. An array of constraint pointers can be uased as an array of handle pointers
+	using FHandles = TArrayView<FPBDCollisionConstraint* const>;
+	using FConstHandles = TArrayView<const FPBDCollisionConstraint* const>;
 
 	// For use by dependent types
-	using FPointContactConstraint = FRigidBodyPointContactConstraint;
-	using FConstraintContainerHandle = FPBDCollisionConstraintHandle;
+	using FPointContactConstraint = FPBDCollisionConstraint;				// Used by clustering (remove this?)
+	using FConstraintContainerHandle = FPBDCollisionConstraintHandle;		// Used by constraint rules
+	using FConstraintSolverContainerType = FPBDCollisionSolverContainer;	// Used by constraint rules
 	using FParticleHandle = TGeometryParticleHandle<FReal, 3>;
 
 	FPBDCollisionConstraints(const FPBDRigidsSOAs& InParticles, 
@@ -60,7 +69,7 @@ public:
 		const int32 ApplyPushOutPairIterations = 1, 
 		const FReal RestitutionThreshold = 2000.0f);
 
-	virtual ~FPBDCollisionConstraints() {}
+	virtual ~FPBDCollisionConstraints();
 
 	/**
 	 * Whether this container provides constraint handles (simple solvers do not need them)
@@ -81,74 +90,26 @@ public:
 	}
 
 	/**
-	 * Helper object for efficiently appending constraints into the constraint container
-	 * in a scope. The encapsulates the separation of appending the constraints into the
-	 * owning container and building the handles required for them. Previously this was
-	 * done one at a time, this helper lets us batch the operation to make it faster
-	 *
-	 * It's important not to mutate the owning container while this helper is alive
-	 * otherwise it will not be able to append correctly to it.
-	 */
-	struct FConstraintAppendScope
-	{
-		FConstraintAppendScope() = delete;
-		FConstraintAppendScope(const FConstraintAppendScope&) = delete;
-		FConstraintAppendScope& operator=(const FConstraintAppendScope&) = delete;
-
-		FConstraintAppendScope(FConstraintAppendScope&&) = default;
-		FConstraintAppendScope& operator=(FConstraintAppendScope&&) = default;
-
-		FConstraintAppendScope(FPBDCollisionConstraints* InOwner);
-		~FConstraintAppendScope();
-
-		// Reserves space for NumToAdd constraints in the internal container
-		void ReserveSingle(int32 NumToAdd);
-		void ReserveSingleSwept(int32 NumToAdd);
-
-		// Append constraint lists to the internal container.
-		// note this will move the container, it will no longer be valid after a call to Append.
-		void Append(TArray<FRigidBodyPointContactConstraint>&& InConstraints);
-		void Append(TArray<FRigidBodySweptPointContactConstraint>&& InConstraints);
-
-	private:
-		FPBDCollisionConstraints* Owner = nullptr;
-		FCollisionConstraintsArray* Constraints = nullptr;
-
-		// Tracking for how many constraints the container began with and how many
-		// the helper added so we can build the new handles on scope exit
-		int32 NumBeginSingle = 0;
-		int32 NumBeginSingleSwept = 0;
-		int32 NumAddedSingle = 0;
-		int32 NumAddedSingleSwept = 0;
-	};
-	
-	/** Begin an append operation, recieving a helper object for bulk operations on the constraint container */
-	FConstraintAppendScope BeginAppendScope();
-
-private:
+	 *  Clears the list of active constraints.
+	 * @todo(chaos): This is only required because of the way events work (see AdvanceOneTimeStepTask::DoWork)
+	*/
+	void BeginFrame();
 
-	// Set whenever an append scope is constructed, and unset when destructed
-	// and used to assert the container isn't mutated during appending.
-	bool bInAppendOperation;
+	/**
+	*  Destroy all constraints 
+	*/
+	void Reset();
 
-public:
 
 	/**
-	*  Add the constraint to the container. 
-	*
-	*  @todo(chaos) : Collision Constraints 
-	*  Update to use a custom allocator. 
-	*  The InConstraint should be a point to unmanaged, raw memory. 
-	*  This function will make a deep copy of the constraint and 
-	*  then delete the InConstraint. 
+	 * @brief Called before collision detection to reset contacts
 	*/
-	void AddConstraint(const FRigidBodyPointContactConstraint& InConstraint);
-	void AddConstraint(const FRigidBodySweptPointContactConstraint& InConstraint);
+	void BeginDetectCollisions();
 
 	/**
-	*  Reset the constraint frame. 
+	 * @brief Called after collision detection to finalize the contacts
 	*/
-	void Reset();
+	void EndDetectCollisions();
 
 	/**
 	 * Apply modifiers to the constraints and specify which constraints should be disabled.
@@ -176,21 +137,6 @@ public:
 	void RemoveConstraint(FPBDCollisionConstraintHandle* ConstraintHandle);
 
 
-	/**
-	 * Update all constraint values within the set
-	 */
-	void UpdateConstraints(FReal Dt, const TSet<TGeometryParticleHandle<FReal, 3>*>& AddedParticles);
-
-	/**
-	 * Update all constraint values
-	 */
-
-	 /**
-	 * Update all constraint values
-	 */
-	void UpdateConstraints(FReal Dt);
-
-
 	//
 	// General Rule API
 	//
@@ -199,10 +145,6 @@ public:
 
 	void UnprepareTick() {}
 
-	void PrepareIteration(FReal Dt);
-
-	void UnprepareIteration(FReal Dt) {}
-
 	/**
 	 * Generate all contact constraints.
 	 */
@@ -212,47 +154,39 @@ public:
 	// Simple Rule API
 	//
 
-	bool Apply(const FReal Dt, const int32 It, const int32 NumIts);
-	bool ApplyPushOut(const FReal Dt, const int32 It, const int32 NumIts);
+	void GatherInput(const FReal Dt, FPBDIslandSolverData& SolverData);
+	void ScatterOutput(const FReal Dt, FPBDIslandSolverData& SolverData);
+
+	bool ApplyPhase1(const FReal Dt, const int32 It, const int32 NumIts, FPBDIslandSolverData& SolverData);
+	bool ApplyPhase2(const FReal Dt, const int32 It, const int32 NumIts, FPBDIslandSolverData& SolverData);
 
 	//
 	// Island Rule API
 	//
-	// @todo(ccaulfield): this runs wide. The serial/parallel decision should be in the ConstraintRule
 
-	bool Apply(const FReal Dt, const TArray<FPBDCollisionConstraintHandle*>& InConstraintHandles, const int32 It, const int32 NumIts);
-	bool ApplyPushOut(const FReal Dt, const TArray<FPBDCollisionConstraintHandle*>& InConstraintHandles, 
-		const TSet<const TGeometryParticleHandle<FReal, 3>*>& IsTemporarilyStatic, int32 Iteration, int32 NumIterations);
+	void SetNumIslandConstraints(const int32 NumIslandConstraints, FPBDIslandSolverData& SolverData);
+	void GatherInput(const FReal Dt, FPBDCollisionConstraint& Constraint, const int32 Particle0Level, const int32 Particle1Level, FPBDIslandSolverData& SolverData);
 
+	void ApplySwept(const FReal Dt, FPBDIslandSolverData& SolverData);
+	bool ApplyPhase1Serial(const FReal Dt, const int32 It, const int32 NumIts, FPBDIslandSolverData& SolverData);
+	bool ApplyPhase2Serial(const FReal Dt, const int32 It, const int32 NumIts, FPBDIslandSolverData& SolverData);
 
-	/**
-	 *  Callbacks
-	 */
-	void SetPostApplyCallback(const FRigidBodyContactConstraintsPostApplyCallback& Callback);
-	void ClearPostApplyCallback();
+	//
+	// Color Rule API
+	//
 
-	void SetPostApplyPushOutCallback(const FRigidBodyContactConstraintsPostApplyPushOutCallback& Callback);
-	void ClearPostApplyPushOutCallback();
+	void ScatterOutput(const FReal Dt, const int32 BeginIndex, const int32 EndIndex, FPBDIslandSolverData& SolverData);
 
+	bool ApplyPhase1Serial(const FReal Dt, const int32 It, const int32 NumIts, const int32 BeginIndex, const int32 EndIndex, FPBDIslandSolverData& SolverData);
+	bool ApplyPhase1Parallel(const FReal Dt, const int32 It, const int32 NumIts, const int32 BeginIndex, const int32 EndIndex, FPBDIslandSolverData& SolverData);
+	
+	bool ApplyPhase2Serial(const FReal Dt, const int32 It, const int32 NumIts, const int32 BeginIndex, const int32 EndIndex, FPBDIslandSolverData& SolverData);
+	bool ApplyPhase2Parallel(const FReal Dt, const int32 It, const int32 NumIts, const int32 BeginIndex, const int32 EndIndex, FPBDIslandSolverData& SolverData);
 
 	//
 	// Member Access
 	//
 
-	const TArray<FPBDCollisionConstraintHandle*>& GetAllConstraintHandles() const 
-	{ 
-		return Handles; 
-	}
-
-	bool Contains(const FCollisionConstraintBase* Base) const
-	{
-#if CHAOS_COLLISION_PERSISTENCE_ENABLED
-		return Manifolds.Contains(FPBDCollisionConstraintHandle::MakeKey(Base));
-#else
-		return false;
-#endif
-	}
-
 	void SetCanDisableContacts(bool bInCanDisableContacts)
 	{
 		bCanDisableContacts = bInCanDisableContacts;
@@ -315,58 +249,68 @@ public:
 
 	void SetGravity(const FVec3& InGravity)
 	{
-		GravityDir = InGravity.GetSafeNormal();
+		GravityDirection = InGravity;
+		GravitySize = GravityDirection.SafeNormalize();
 	}
 
-	int32 NumConstraints() const
+	FVec3 GetGravityDirection() const
+	{
+		return GravityDirection;
+	}
+
+	FReal GetGravitySize() const
 	{
-		return Constraints.SinglePointConstraints.Num() + Constraints.SinglePointSweptConstraints.Num();
+		return GravitySize;
 	}
 
-	FHandles& GetConstraintHandles()
+	int32 NumConstraints() const
 	{
-		return Handles;
+		return GetConstraints().Num();
 	}
 
-	const FHandles& GetConstConstraintHandles() const
+	TArrayView<FPBDCollisionConstraint* const> GetConstraints() const
 	{
-		return Handles;
+		return ConstraintAllocator.GetConstraints();
 	}
 
-	const FCollisionConstraintBase& GetConstraint(int32 Index) const;
+	FHandles GetConstraintHandles() const;
+	FConstHandles GetConstConstraintHandles() const;
 
-	FCollisionConstraintsArray& GetConstraintsArray() { return Constraints; }
+	const FPBDCollisionConstraint& GetConstraint(int32 Index) const;
 
-	//Sort constraints based on particle indices so that we have a deterministic solve order
-	void SortConstraints();
+	FCollisionConstraintAllocator& GetConstraintAllocator() { return ConstraintAllocator; }
 
 	TMap<FParticleHandle*, TArray<FPBDCollisionConstraintHandle*> >& GetParticleCollisionsMap() { return ParticleCollisionsMap; }
 	const TMap<FParticleHandle*, TArray<FPBDCollisionConstraintHandle*> >& GetParticleCollisionsMap() const { return ParticleCollisionsMap; }
 
 protected:
-	using Base::GetConstraintIndex;
-	using Base::SetConstraintIndex;
+	FPBDCollisionConstraint& GetConstraint(int32 Index);
+	FPBDCollisionSolverContainer& GetConstraintSolverContainer(FPBDIslandSolverData& SolverData);
 
-	void UpdateConstraintMaterialProperties(FCollisionConstraintBase& Contact);
+	void UpdateConstraintMaterialProperties(FPBDCollisionConstraint& Contact);
 
 	Collisions::FContactParticleParameters GetContactParticleParameters(const FReal Dt);
 	Collisions::FContactIterationParameters GetContactIterationParameters(const FReal Dt, const int32 Iteration, const int32 NumIterations, const int32 NumPairIterations, bool& bNeedsAnotherIteration);
 
+	// The "Legacy" functions handle the older solver types (GbfPbd and StandardPbd)
+	// @todo(chaos): remove legacy methods when the new solver is fully operational and used everywhere (RBAN)
+	void LegacyGatherInput(const FReal Dt, FPBDCollisionConstraint& Constraint, const int32 Particle0Level, const int32 Particle1Level, FPBDIslandSolverData& SolverData);
+	void LegacyScatterOutput(const FReal Dt, const int32 BeginIndex, const int32 EndIndex, FPBDIslandSolverData& SolverData);
+	bool LegacyApplyPhase1Serial(const FReal Dt, const int32 Iterations, const int32 NumIterations, const int32 BeginIndex, const int32 EndIndex, FPBDIslandSolverData& SolverData);
+	bool LegacyApplyPhase1Parallel(const FReal Dt, const int32 Iterations, const int32 NumIterations, const int32 BeginIndex, const int32 EndIndex, FPBDIslandSolverData& SolverData);
+	bool LegacyApplyPhase2Serial(const FReal Dt, const int32 Iterations, const int32 NumIterations, const int32 BeginIndex, const int32 EndIndex, FPBDIslandSolverData& SolverData);
+	bool LegacyApplyPhase2Parallel(const FReal Dt, const int32 Iterations, const int32 NumIterations, const int32 BeginIndex, const int32 EndIndex, FPBDIslandSolverData& SolverData);
+
+	void AddToParticleCollisionMap(FPBDCollisionConstraint* Constraint);
+
 private:
 
-	friend FConstraintAppendScope;
 	const FPBDRigidsSOAs& Particles;
 
-	FCollisionConstraintsArray Constraints;
+	FCollisionConstraintAllocator ConstraintAllocator;
 	int32 NumActivePointConstraints;
-	int32 NumActiveSweptPointConstraints;
-
-#if CHAOS_COLLISION_PERSISTENCE_ENABLED
-	TMap< FConstraintContainerHandleKey, FPBDCollisionConstraintHandle* > Manifolds;
-#endif
-	TArray<FPBDCollisionConstraintHandle*> Handles;
-	FConstraintHandleAllocator HandleAllocator;
 
+	// @todo(chaos): this functionality should be moved to the PBDCollisionConstraintAlloctor now
 	TMap<FParticleHandle*,TArray<FPBDCollisionConstraintHandle*> > ParticleCollisionsMap;
 
 
@@ -376,7 +320,6 @@ private:
 	int32 MApplyPairIterations;
 	int32 MApplyPushOutPairIterations;
 	FReal RestitutionThreshold;
-	bool bUseCCD;
 	bool bEnableCollisions;
 	bool bEnableRestitution;
 	bool bHandlesEnabled;
@@ -386,14 +329,10 @@ private:
 	// This improves performance, but can decrease stability if contacts are culled prematurely.
 	bool bCanDisableContacts;
 
-	// Used by PushOut to decide on priority when two bodies are at same shock propagation level
-	FVec3 GravityDir;
+	// Used to determine constraint directions
+	FVec3 GravityDirection;
+	FReal GravitySize;
 
 	EConstraintSolverType SolverType;
-
-	int32 LifespanCounter;
-
-	FRigidBodyContactConstraintsPostApplyCallback PostApplyCallback;
-	FRigidBodyContactConstraintsPostApplyPushOutCallback PostApplyPushOutCallback;
 };
 }
diff --git a/Engine/Source/Runtime/Experimental/Chaos/Public/Chaos/PBDCollisionConstraintsContact.h b/Engine/Source/Runtime/Experimental/Chaos/Public/Chaos/PBDCollisionConstraintsContact.h
--- a/Engine/Source/Runtime/Experimental/Chaos/Public/Chaos/PBDCollisionConstraintsContact.h
+++ b/Engine/Source/Runtime/Experimental/Chaos/Public/Chaos/PBDCollisionConstraintsContact.h
@@ -12,6 +12,7 @@ namespace Chaos
 
 	namespace Collisions
 	{
+		// @todo(chaos): remove this
 		struct FContactParticleParameters 
 		{
 			FContactParticleParameters()
@@ -70,14 +71,12 @@ namespace Chaos
 		};
 
 		// Update the constraint (re-runs collision detection for this contact)
-		extern void Update(FRigidBodyPointContactConstraint& Constraint, const FReal Dt);
-
-		extern void Apply(FCollisionConstraintBase& Constraint, const FContactIterationParameters& IterationParameters, const FContactParticleParameters& ParticleParameters);
-
-		extern void ApplyPushOut(FCollisionConstraintBase& Constraint, const TSet<const TGeometryParticleHandle<FReal, 3>*>& IsTemporarilyStatic,
-			const FContactIterationParameters& IterationParameters, const FContactParticleParameters& ParticleParameters, const FVec3& GravityDir = FVec3::DownVector);
+		extern void Update(FPBDCollisionConstraint& Constraint, const FReal Dt);
+		extern void UpdateSwept(FPBDCollisionConstraint& Constraint, const FReal Dt);
 
+		extern void Apply(FPBDCollisionConstraint& Constraint, const FContactIterationParameters& IterationParameters, const FContactParticleParameters& ParticleParameters);
 
+		extern void ApplyPushOut(FPBDCollisionConstraint& Constraint, const FContactIterationParameters& IterationParameters, const FContactParticleParameters& ParticleParameters);
 	}
 
 }
diff --git a/Engine/Source/Runtime/Experimental/Chaos/Public/Chaos/PBDCollisionConstraintsUtil.h b/Engine/Source/Runtime/Experimental/Chaos/Public/Chaos/PBDCollisionConstraintsUtil.h
--- a/Engine/Source/Runtime/Experimental/Chaos/Public/Chaos/PBDCollisionConstraintsUtil.h
+++ b/Engine/Source/Runtime/Experimental/Chaos/Public/Chaos/PBDCollisionConstraintsUtil.h
@@ -7,7 +7,7 @@
 
 namespace Chaos
 {
-	void CHAOS_API ComputeHashTable(const TArray<Chaos::FRigidBodyPointContactConstraint>& ConstraintsArray,
+	void CHAOS_API ComputeHashTable(const TArray<Chaos::FPBDCollisionConstraint>& ConstraintsArray,
 									const FBox& BoundingBox, TMultiMap<int32, int32>& HashTableMap, const Chaos::FReal SpatialHashRadius);
 
 	void CHAOS_API ComputeHashTable(const TArray<FCollidingData>& CollisionsArray,
diff --git a/Engine/Source/Runtime/Experimental/Chaos/Public/Chaos/PBDConstraintColor.h b/Engine/Source/Runtime/Experimental/Chaos/Public/Chaos/PBDConstraintColor.h
--- a/Engine/Source/Runtime/Experimental/Chaos/Public/Chaos/PBDConstraintColor.h
+++ b/Engine/Source/Runtime/Experimental/Chaos/Public/Chaos/PBDConstraintColor.h
@@ -29,6 +29,11 @@ namespace Chaos
 			return IslandData.Num();
 		}
 
+		int32 NumIslandEdges(const int32 Island) const
+		{
+			return IslandData[Island].NumEdges;
+		}
+
 		/**
 		 * Initialize the color structures based on the connectivity graph (i.e., reset all color-related node, edge and island data).
 		 */
@@ -37,7 +42,7 @@ namespace Chaos
 		/**
 		 * Calculate the color information for the specified island.
 		 */
-		void ComputeColor(const int32 Island, const FPBDConstraintGraph& ConstraintGraph, uint32 ContainerId);
+		void ComputeColor(const FReal Dt, const int32 Island, const FPBDConstraintGraph& ConstraintGraph, uint32 ContainerId);
 
 		/**
 		 * Get the Level-Color-ConstraintList map for the specified island.
@@ -64,10 +69,98 @@ namespace Chaos
 		bool DebugCheckGraph(const FPBDConstraintGraph& ConstraintGraph) const;
 #endif
 
+		/**
+		 * Get the level of a particle after particle levels are computed using ComputeContactGraphGBF 
+		 */
+		int32 GetParticleLevel(const FGeometryParticleHandle* ParticleHandle) const;
+
 	private:
 		void ComputeContactGraph(const int32 Island, const FPBDConstraintGraph& ConstraintGraph, uint32 ContainerId);
 		void ComputeIslandColoring(const int32 Island, const FPBDConstraintGraph& ConstraintGraph, uint32 ContainerId);
 
+		/**
+		 * Compute contact graph using GBF paper Sec 8.1
+		 */
+		void ComputeContactGraphGBF(const FReal Dt, const int32 Island, const FPBDConstraintGraph& ConstraintGraph, const uint32 ContainerId);
+
+		/**
+		 * Compute the mapping between constraint graph node index to island graph node indices. 
+		 * The mapping of dynamic particles are stored in TArray DynamicGlobalToIslandLocalNodeIndices. 
+		 * The mapping of kinematic particles are stored in TMap KinematicGlobalToIslandLocalNodeIndices.
+		 */
+		void ComputeGlobalToIslandLocalNodeMapping(const int32 Island, const FPBDConstraintGraph& ConstraintGraph, TArray<int32>& DynamicGlobalToIslandLocalNodeIndices, TMap<int32, int32>& KinematicGlobalToIslandLocalNodeIndices) const;
+
+		/**
+		 * Compute the graph node index in the island using the mapping data above.
+		 */
+		int32 GetIslandLocalNodeIdx(const int32 GlobalNodeIndex, const TArray<int32>& DynamicGlobalToIslandLocalNodeIndices, const TMap<int32, int32>& KinematicGlobalToIslandLocalNodeIndices, const bool IsNodeDynamic) const;
+
+		struct FDirectedEdge
+		{
+			int32 FirstNode; // Or source node of the directed edge.
+			int32 SecondNode; // Or target node of the directed edge. 
+			FDirectedEdge(){};
+			FDirectedEdge(const int32 FirstInput, const int32 SecondInput) : FirstNode(FirstInput), SecondNode(SecondInput) {};
+		};
+
+		/**
+		 * Collect directed edges from the island.
+		 * All collected edges will point from lower particles to upper particles.
+		 * The set of collected directed edges is a subset of collision constraint edges.
+		 * Horizontal edges will not be collected.
+		 */
+		void CollectIslandDirectedEdges(const FReal Dt, const int32 Island, const FPBDConstraintGraph& ConstraintGraph, const TArray<int32>& DynamicGlobalToIslandLocalNodeIndices, const TMap<int32, int32>& KinematicGlobalToIslandLocalNodeIndices, TArray<FDirectedEdge>& DirectedEdges) const;
+
+		/**
+		 * CSR: Compressed sparse row
+		 * Build CSR representation of the graph from input edges
+		 */
+
+		void BuildGraphCSR(const TArray<FDirectedEdge>& GraphEdges, const int32 NumNodes, TArray<int32>& CSRIndices, TArray<int32>& CSRValues, TArray<int32>& CurrentIndices) const;
+
+		/**
+		 * A node is a root if it has no parents.
+		 * Result is stored in IsRootNode
+		 */
+		void ComputeIsRootNode(const TArray<FDirectedEdge>& GraphEdges, const int32 NumNodes, TArray<bool> &IsRootNode) const;
+
+		/**
+		 * Digraph: directed graph (according to wikipedia: https://en.wikipedia.org/wiki/Directed_graph) that might contain loops
+		 * DAG: directed acyclic graph
+		 * Collapse directed graph to directed acyclic graph. Recursively collapse the loops in the directed graph to single nodes until there are no loops any more.
+		 * For each digraph node, the corresponding DAG node index is stored in DigraphToDAGIndices
+		 * For each DAG node, the corresponding list of digraph nodes is stored in DAGToDigraphIndices
+		 */
+		void CollapseDigraphToDAG(const TArray<int32>& ChildrenCSRIndices, const TArray<int32>& ChildrenCSRValues, const TArray<bool>& IsRootNode, TArray<int32>& DigraphToDAGIndices, TArray<TArray<int32>>& DAGToDigraphIndices, TArray<bool>& Visited, TArray<int32>& TraversalStack, TArray<int32>& NumOnStack) const;
+
+		/**
+		 * Compute DAG Edges from Digraph edges
+		 */
+		void ComputeDAGEdges(const TArray<FDirectedEdge>& DigraphEdges, const TArray<int32>& DigraphToDAGIndices, TArray<FDirectedEdge> &DAGEdges) const;
+
+		/**
+		 * Topological sort DAG nodes
+		 * DAGToDigraphIndices is used to determine whether DAGI is a valid DAG node (point to non-zero Digraph nodes)
+		 */
+		void TopologicalSortDAG(const TArray<int32>& DAGChildrenCSRIndices, const TArray<int32>& DAGChildrenCSRValues, const TArray<TArray<int32>>& DAGToDigraphIndices, TArray<int32>& SortedDAGNodes, TArray<bool>& Visited) const;
+
+		/**
+		 * Assign levels to DAG nodes.
+		 * We want all parents of DAGI have smaller levels than DAGI. The level of DAGI is the max distance to root nodes. 
+		 */
+		void AssignDAGNodeLevels(const TArray<int32>& DAGChildrenCSRIndices, const TArray<int32>& DAGChildrenCSRValues, const TArray<int32>& SortedDAGNodes, TArray<int32> &DAGNodeLevels) const;
+
+		/**
+		 * Assign edge levels based on DAG node levels.
+		 * Edge level is the max of its end node levels.
+		 */
+		void AssignEdgeLevels(const int32 Island, const FPBDConstraintGraph& ConstraintGraph, const TArray<int32>& DynamicGlobalToIslandLocalNodeIndices, const TMap<int32, int32>& KinematicGlobalToIslandLocalNodeIndices, const TArray<int32>& DigraphToDAGIndices, const TArray<int32> DAGNodeLevels);
+
+		/**
+		 * Update particle levels based on DAGNodeLevels.
+		 */
+		void UpdateParticleToLevel(const int32 Island, const FPBDConstraintGraph& ConstraintGraph, const TArray<int32>& DynamicGlobalToIslandLocalNodeIndices, const TMap<int32, int32>& KinematicGlobalToIslandLocalNodeIndices, const TArray<int32>& DigraphToDAGIndices, const TArray<int32>& DAGNodeLevels);
+
 		struct FGraphNodeColor
 		{
 			FGraphNodeColor()
@@ -94,18 +187,47 @@ namespace Chaos
 			FIslandColor()
 				: MaxColor(0)
 				, MaxLevel(0)
+				, NumEdges(0)
 			{
 			}
 			int32 MaxColor;
 			int32 MaxLevel;
+			int32 NumEdges;
 			FLevelToColorToConstraintListMap LevelToColorToConstraintListMap;
 		};
 
+		struct FGBFContactGraphData
+		{
+			TMap<int32, int32> KinematicGlobalToIslandLocalNodeIndices;
+			TArray<FDirectedEdge> DigraphEdges;
+			TArray<int32> ChildrenCSRIndices, ChildrenCSRValues;
+
+			// Buffer array for BuildGraphCSR function. Not used outside BuildGraphCSR.
+			TArray<int32> CSRCurrentIndices; 
+
+			TArray<bool> IsRootNode;
+			TArray<int32> DigraphToDAGIndices;
+			TArray<TArray<int32>> DAGToDigraphIndices;
+
+			// Helper arrays for CollapseDigraphToDAG and TopologicalSortDAG. Not used elsewhere
+			TArray<bool> Visited; 
+			TArray<int32> TraversalStack;
+			TArray<int32> NumOnStack;
+
+			TArray<FDirectedEdge> DAGEdges;
+			TArray<int32> DAGChildrenCSRIndices, DAGChildrenCSRValues;
+			TArray<int32> SortedDAGNodes;
+			TArray<int32> DAGNodeLevels;
+		};
+
 		TArray<FGraphNodeColor> Nodes;
 		TArray<FGraphEdgeColor> Edges;
 		TArray<FIslandColor> IslandData;
 		FLevelToColorToConstraintListMap EmptyLevelToColorToConstraintListMap;
-		TArray<int32> UpdatedNodes;
+		TArray<int32> DynamicGlobalToIslandLocalNodeIndicesArray; 
+		TArray<int32> ParticleToLevel; 
+		TArray<const TGeometryParticleHandle<FReal, 3>*> NodeToParticle;
+		TArray<FGBFContactGraphData> IslandToGBFContactGraphData;
 		bool bUseContactGraph;
 	};
 
diff --git a/Engine/Source/Runtime/Experimental/Chaos/Public/Chaos/PBDConstraintContainer.h b/Engine/Source/Runtime/Experimental/Chaos/Public/Chaos/PBDConstraintContainer.h
--- a/Engine/Source/Runtime/Experimental/Chaos/Public/Chaos/PBDConstraintContainer.h
+++ b/Engine/Source/Runtime/Experimental/Chaos/Public/Chaos/PBDConstraintContainer.h
@@ -3,21 +3,10 @@
 
 #include "Chaos/ConstraintHandle.h"
 #include "Chaos/Transform.h"
+#include "Chaos/Evolution/SolverBody.h"
 
 namespace Chaos
 {
-	// @todo(chaos): get rid of this - we should be able to extend constraint types without having to add to an enum
-	enum class EConstraintContainerType : uint8
-	{
-		Invalid = 0,
-		Collision,
-		RigidSpring,
-		DynamicSpring,
-		Position,
-		Joint,
-		Suspension
-	};
-
 	/**
 	 * Base class for containers of constraints.
 	 * A Constraint Container holds an array of constraints and provides methods to allocate and deallocate constraints
@@ -26,7 +15,7 @@ namespace Chaos
 	class CHAOS_API FPBDConstraintContainer
 	{
 	public:
-		FPBDConstraintContainer(EConstraintContainerType InType);
+		FPBDConstraintContainer(FConstraintHandleTypeID InConstraintHandleType);
 
 		virtual ~FPBDConstraintContainer();
 
@@ -41,24 +30,34 @@ namespace Chaos
 			ContainerId = InContainerId;
 		}
 
-		EConstraintContainerType GetConstraintType() const
+		const FConstraintHandleTypeID& GetConstraintHandleType() const
 		{
-			return ConstraintContainerType;
+			return ConstraintHandleType;
 		}
 
-		virtual void SetConstraintEnabled(int32 ConstraintIndex, bool bEnabled) { }
-		virtual bool IsConstraintEnabled(int32 ConstraintIndex) const { return true; }
 		virtual void DisconnectConstraints(const TSet<TGeometryParticleHandle<FReal, 3>*>&) {}
 
 	protected:
-		int32 GetConstraintIndex(const FConstraintHandle* ConstraintHandle) const;
-		void SetConstraintIndex(FConstraintHandle* ConstraintHandle, int32 ConstraintIndex) const;
-
-	private:
-		EConstraintContainerType ConstraintContainerType;
+		FConstraintHandleTypeID ConstraintHandleType;
 		int32 ContainerId;
 	};
 
+	class CHAOS_API FPBDIndexedConstraintContainer : public FPBDConstraintContainer
+	{
+	public:
+		FPBDIndexedConstraintContainer(FConstraintHandleTypeID InType)
+			: FPBDConstraintContainer(InType)
+		{
+		}
+
+		virtual void SetConstraintEnabled(int32 ConstraintIndex, bool bEnabled) { }
+		virtual bool IsConstraintEnabled(int32 ConstraintIndex) const { return true; }
+
+	protected:
+		int32 GetConstraintIndex(const FIndexedConstraintHandle* ConstraintHandle) const;
+		void SetConstraintIndex(FIndexedConstraintHandle* ConstraintHandle, int32 ConstraintIndex) const;
+	};
+
 
 	//
 	// From ConstraintHandle.h
@@ -73,33 +72,17 @@ namespace Chaos
 		return INDEX_NONE;
 	}
 
-	inline void FConstraintHandle::SetEnabled(bool bInEnabled)
-	{
-		if (ConstraintContainer != nullptr)
-		{
-			ConstraintContainer->SetConstraintEnabled(ConstraintIndex, bInEnabled);
-		}
-	}
-
-	inline bool FConstraintHandle::IsEnabled() const
-	{
-		if (ConstraintContainer != nullptr)
-		{
-			return ConstraintContainer->IsConstraintEnabled(ConstraintIndex);
-		}
-		return false;
-	}
-
 	template<typename T>
 	inline T* FConstraintHandle::As()
 	{
-		return ((ConstraintContainer != nullptr) && (T::StaticType() == ConstraintContainer->GetConstraintType())) ? static_cast<T*>(this) : nullptr;
+		// @todo(chaos): we need a safe cast that allows for casting to intermediate base classes (e.g., FIndexedConstraintHandle)
+		return ((ConstraintContainer != nullptr) && ConstraintContainer->GetConstraintHandleType().IsA(T::StaticType())) ? static_cast<T*>(this) : nullptr;
 	}
 
 	template<typename T>
 	inline const T* FConstraintHandle::As() const
 	{
-		return ((ConstraintContainer != nullptr) && (T::StaticType() == ConstraintContainer->GetConstraintType())) ? static_cast<const T*>(this) : nullptr;
+		return ((ConstraintContainer != nullptr) && ConstraintContainer->GetConstraintHandleType().IsA(T::StaticType())) ? static_cast<const T*>(this) : nullptr;
 	}
 
 }
\ No newline at end of file
