Chaos

diff --git a/Engine/Source/Runtime/Experimental/Chaos/Public/Chaos/PBDConstraintGraph.h b/Engine/Source/Runtime/Experimental/Chaos/Public/Chaos/PBDConstraintGraph.h
--- a/Engine/Source/Runtime/Experimental/Chaos/Public/Chaos/PBDConstraintGraph.h
+++ b/Engine/Source/Runtime/Experimental/Chaos/Public/Chaos/PBDConstraintGraph.h
@@ -1,268 +1,10 @@
 // Copyright Epic Games, Inc. All Rights Reserved.
 #pragma once
 
-#include "ParticleHandle.h"
-
-namespace ChaosTest
-{
-	class ConstraintGraphValidation;
-}
-
+#include "Chaos/Island/IslandManager.h"
 namespace Chaos
 {
-	template<typename T, int d>
-	class TPBDRigidParticles;
-
-	template<typename T, int d>
-	class TPBDConstraintColor;
-
-	class FConstraintHandle;
-
-	template<class T>
-	class TChaosPhysicsMaterial;
-
-	template <typename T>
-	class TSerializablePtr;
-
-	template<typename T>
-	class TArrayCollectionArray;
-
-	class FPBDRigidsSOAs;
-
-	/**
-	 * Build a graph of connected particles, and then a set of independent islands.
-	 * Particles/constraints in different islands do not interact, so islands can be updated in parallel.
-	 * This is also where particle sleeping is controlled.
-	 * Note that at the moment this graph is rebuilt every frame. When this changes some of the underlying data will need to be changed. I.e. move from indices to pointers to handles
-	 */
-	 // @todo(ccaulfield): Break out the Island and Sleep stuff from the particle graph
-	 // @todo(ccaulfield): The level system should be part of the Island system, and not embedded with the color system
-	class CHAOS_API FPBDConstraintGraph
-	{
-	public:
-		friend class FPBDConstraintColor;
-
-		//Make friend with unit test code so we can verify some behavior
-		friend class ChaosTest::ConstraintGraphValidation;
-
-		/** Information required to map a graph edge back to its constraint */
-		struct FConstraintData
-		{
-		public:
-			FConstraintData() : ContainerId(0), ConstraintHandle(nullptr) {}
-			FConstraintData(uint32 InContainerId, FConstraintHandle* InConstraintHandle) : ContainerId(InContainerId), ConstraintHandle(InConstraintHandle) {}
-
-			uint32 GetContainerId() const { return ContainerId; }
-			FConstraintHandle* GetConstraintHandle() const { return ConstraintHandle; }
-
-		private:
-			uint32 ContainerId;
-			FConstraintHandle* ConstraintHandle;
-		};
-
-		FPBDConstraintGraph();
-		FPBDConstraintGraph(const TParticleView<FGeometryParticles>& Particles);
-		virtual ~FPBDConstraintGraph() {}
-
-		/**
-		 * Clear the graph and set up the particle-to-graph-node mapping for the specified particles
-		 * Should be called before AddConstraint.
-		 */
-		void InitializeGraph(const TParticleView<TGeometryParticles<FReal, 3>>& Particles);
-
-		/**
-		 * Reserve space in the graph for NumConstraints additional constraints.
-		 */
-		void ReserveConstraints(const int32 NumConstraints);
-
-		/**
-		 * Add a constraint to the graph for each constraint in the container.
-		 */
-		void AddConstraint(const uint32 InContainerId, FConstraintHandle* InConstraintHandle, const TVector<FGeometryParticleHandle*, 2>& InConstrainedParticles);
-
-		/**
-		 * Remove a constraint from the graph
-		 */
-		void RemoveConstraint(const uint32 InContainerId, FConstraintHandle* InConstraintHandle, const TVector<FGeometryParticleHandle*, 2>& InConstrainedParticles);
-
-		/**
-		 * Add particles/constraints to their particle's already-assigned islands (if applicable).
-		 */
-		 // @todo(ccaulfield): InitializeIslands and ResetIslands are a bit confusing. Try to come up with better names.
-		void ResetIslands(const TParticleView<FPBDRigidParticles>& PBDRigids);
-
-		/**
-		 * Generate the simulation islands of connected particles (AddConstraints must have already been called). There are no constraints connecting particles in different islands.
-		 */
-		void UpdateIslands(const TParticleView<FPBDRigidParticles>& confusing, FPBDRigidsSOAs& Particles);
-
-		/**
-		 * Put particles in inactive islands to sleep.
-		 */
-		bool SleepInactive(const int32 Island, const TArrayCollectionArray<TSerializablePtr<FChaosPhysicsMaterial>>& PerParticleMaterialAttributes, const THandleArray<FChaosPhysicsMaterial>& SolverPhysicsMaterials);
-
-		/**
-		 * Wake all particles in an Island.
-		 */
-		void WakeIsland(FPBDRigidsSOAs& Particles, const int32 Island);
-
-		/**
-		 * Ensure that the particles in each island have consistent sleep states - if any are awake, wake all.
-		 */
-		 // @todo(ccaulfield): Do we really need this? It implies some behind-the-scenes state manipulation.
-		//void ReconcileIslands();
-
-		/**
-		 * Get the list of ConstraintsData indices associated with the specified island. NOTE: ConstraintDataIndex is an internal index and not related to 
-		 * a constraint's index in its owning container. Indices are into the ConstraintData array in the ConstraintGraph to get to the Constraint Index and Container Id.
-		 * @see GetConstraintsData().
-		 */
-		const TArray<int32>& GetIslandConstraintData(int32 Island) const
-		{
-			return IslandToConstraints[Island];
-		}
-
-		/**
-		 * Get the list of Particle indices associated with the specified island.
-		 */
-		const TArray<FGeometryParticleHandle*>& GetIslandParticles(int32 Island) const
-		{
-			return IslandToParticles[Island];
-		}
-
-		/**
-		 * Get the list of Particle indices associated with the specified island.
-		 */
-		int32 GetIslandSleepCount(int32 Island) const
-		{
-			return IslandToSleepCount[Island];
-		}
-
-		/**
-		 * The number of islands in the graph.
-		 */
-		int32 NumIslands() const
-		{
-			return IslandToParticles.Num();
-		}
-
-		/**
-		 * When resim is used, tells us whether we need to resolve island
-		 */
-		bool IslandNeedsResim(const int32 Island) const
-		{
-			return IslandToData[Island].bNeedsResim;
-		}
-
-		/**
-		 * Information mapping the edge back to a rule's constraint
-		 */
-		const FConstraintData& GetConstraintData(int32 ConstraintDataIndex) const;
-
-		/**
-		 * Adds \p ChildParticle to the constraint graph.  Copies the sleeping 
-		 * state and island of \p ParentParticle to \p ChildParticle.  Does nothing
-		 * if \p ParentParticle is not supplied.
-		 */
-		void EnableParticle(FGeometryParticleHandle* ChildParticle, const FGeometryParticleHandle* ParentParticle);
-
-		/**
-		 * Disable a particle and remove it from its island.
-		 * @note: this does not remove any attached constraints - constraints need to be re-added and islands will need to be rebuilt.
-		 */
-		void DisableParticle(FGeometryParticleHandle* Particle);
-
-		/**
-		 * Preallocate buffers for \p Num particles.
-		 *
-		 * Returns the number of new slots created.
-		 */
-		int32 ReserveParticles(const int32 Num);
-
-		void AddParticle(FGeometryParticleHandle* AddedParticle)
-		{
-			ParticleAdd(AddedParticle);
-		}
-
-		//Remove particle from constraint, maybe rethink some of these names
-		void RemoveParticle(FGeometryParticleHandle* Particle)
-		{
-			//we know it just removes, probably rename later
-			DisableParticle(Particle);
-		}
-
-		/**
-		 * Disable a set of particles and remove them from their island.
-		 * @note: this does not remove any attached constraints - constraints need to be re-added and islands will need to be rebuilt.
-		 */
-		void DisableParticles(const TSet<FGeometryParticleHandle*>& Particles);
-
-#if !UE_BUILD_SHIPPING
-		// For debugging onlue: run some (slow) validation checks on the graph and log any errors. Return false if there were errors.
-		bool DebugCheckGraph() const;
-#endif
-
-	private:
-		struct FGraphNode
-		{
-			FGraphNode()
-				: Particle(nullptr)
-				, Island(INDEX_NONE)
-			{
-			}
-
-			TArray<int32> Edges;
-			FGeometryParticleHandle* Particle;
-			int32 Island;
-		};
-
-		struct FGraphEdge
-		{
-			FGraphEdge()
-			    : FirstNode(INDEX_NONE)
-			    , SecondNode(INDEX_NONE)
-			{
-			}
-
-			int32 FirstNode;
-			int32 SecondNode;
-			FConstraintData Data;
-		};
-
-		struct FIslandData
-		{
-			FIslandData()
-			    : bIsIslandPersistant(false)
-				, bNeedsResim(true)
-			{
-			}
-
-			bool bIsIslandPersistant;
-			bool bNeedsResim;
-		};
-
-		void ComputeIslands(const TParticleView<FPBDRigidParticles>& PBDRigids, FPBDRigidsSOAs& Particles);
-		bool ComputeIsland(const int32 Node, const int32 Island, TSet<FGeometryParticleHandle*>& ParticlesInIsland);
-		
-		void ParticleAdd(FGeometryParticleHandle* AddedParticle);
-		void ParticleRemove(FGeometryParticleHandle* RemovedParticle);
-		int32 GetNextNodeIndex();
-		const TArray<int32>& GetUpdatedNodes() const { return UpdatedNodes; }
-		int32 GetConstraintIndexFromNodeAndConstraintHandle(const FGraphNode& Node, const FConstraintHandle* ConstraintHandle);
-
-		TArray<FGraphNode> Nodes;
-		TArray<FGraphEdge> Edges;
-		TArray<FIslandData> IslandToData;
-		TMap<FGeometryParticleHandle*, int32> ParticleToNodeIndex;
-
-		TArray<TArray<FGeometryParticleHandle*>> IslandToParticles;
-		TArray<TArray<int32>> IslandToConstraints;
-		TArray<int32> IslandToSleepCount;
-
-		TArray<int32> FreeIndexList;
-		TArray<int32> UpdatedNodes;
-		TArray<int32> Visited;
-		uint8 VisitToken;
-	};
-
+	class CHAOS_API FPBDConstraintGraph : public FPBDIslandManager
+	{};
 }
+
diff --git a/Engine/Source/Runtime/Experimental/Chaos/Public/Chaos/PBDConstraintRule.h b/Engine/Source/Runtime/Experimental/Chaos/Public/Chaos/PBDConstraintRule.h
--- a/Engine/Source/Runtime/Experimental/Chaos/Public/Chaos/PBDConstraintRule.h
+++ b/Engine/Source/Runtime/Experimental/Chaos/Public/Chaos/PBDConstraintRule.h
@@ -41,12 +41,6 @@ namespace Chaos
 		/** Called once per frame. Should undo whatever is done in PrepareTick (can also free any other transient buffers created after) */
 		virtual void UnprepareTick() {}
 
-		/** Called once per iteration before Apply. Can be used to prepare caches etc. */
-		virtual void PrepareIteration(FReal Dt) {}
-
-		/** Called once per iteration after Apply. Should be used to release any transient stores created in PrepareConstraints. */
-		virtual void UnprepareIteration(FReal Dt) {}
-
 
 	protected:
 		int32 Priority;
@@ -58,59 +52,47 @@ namespace Chaos
 	class CHAOS_API FSimpleConstraintRule : public FConstraintRule
 	{
 	public:
-		FSimpleConstraintRule(int32 InPriority) : FConstraintRule(InPriority) {}
+		FSimpleConstraintRule(int32 InPriority) : FConstraintRule(InPriority), SolverData(nullptr) {}
+
+		/** Bind the solver datas to the one in the evolution */
+		virtual void BindToDatas(FPBDIslandSolverData& InSolverDatas, const uint32 InContainerId) {}
 
 		virtual void UpdatePositionBasedState(const FReal Dt) {}
+		virtual void GatherSolverInput(const FReal Dt) {}
+		virtual void ScatterSolverOutput(const FReal Dt) {}
 		virtual bool ApplyConstraints(const FReal Dt, const int32 It, const int32 NumIts) { return false; }
 		virtual bool ApplyPushOut(const FReal Dt, const int32 It, const int32 NumIts) { return false; }
+	protected:
+		/** Solver datas that are coming from the evolution */
+		FPBDIslandSolverData* SolverData = nullptr;
 	};
 
-	template<typename T_CONSTRAINTS>
+	template<typename ConstraintType>
 	class CHAOS_API TSimpleConstraintRule : public FSimpleConstraintRule
 	{
 	public:
-		using FConstraints = T_CONSTRAINTS;
+		using FConstraints = ConstraintType;
 
-		TSimpleConstraintRule(int32 InPriority, FConstraints& InConstraints)
-			: FSimpleConstraintRule(InPriority)
-			, Constraints(InConstraints)
-		{
-		}
+		TSimpleConstraintRule(int32 InPriority, FConstraints& InConstraints);
 
-		virtual void PrepareTick() override
-		{
-			Constraints.PrepareTick();
-		}
+		virtual ~TSimpleConstraintRule();
 
-		virtual void UnprepareTick() override
-		{
-			Constraints.UnprepareTick();
-		}
+		/** Bind the solver datas to the one in the evolution */
+		virtual void BindToDatas(FPBDIslandSolverData& InSolverDatas, const uint32 InContainerId) override;
 
-		virtual void PrepareIteration(FReal Dt) override
-		{
-			Constraints.PrepareIteration(Dt);
-		}
+		virtual void PrepareTick() override;
 
-		virtual void UnprepareIteration(FReal Dt) override
-		{
-			Constraints.UnprepareIteration(Dt);
-		}
+		virtual void UnprepareTick() override;
 
-		virtual void UpdatePositionBasedState(const FReal Dt) override
-		{
-			return Constraints.UpdatePositionBasedState(Dt);
-		}
+		virtual void UpdatePositionBasedState(const FReal Dt) override;
 
-		virtual bool ApplyConstraints(const FReal Dt, const int32 It, const int32 NumIts) override
-		{
-			return Constraints.Apply(Dt, It, NumIts);
-		}
+		virtual void GatherSolverInput(const FReal Dt) override;
 
-		virtual bool ApplyPushOut(const FReal Dt, const int32 It, const int32 NumIts) override
-		{ 
-			return Constraints.ApplyPushOut(Dt, It, NumIts);
-		}
+		virtual void ScatterSolverOutput(const FReal Dt) override;
+
+		virtual bool ApplyConstraints(const FReal Dt, const int32 It, const int32 NumIts) override;
+
+		virtual bool ApplyPushOut(const FReal Dt, const int32 It, const int32 NumIts) override;
 
 	private:
 		FConstraints& Constraints;
@@ -132,9 +114,18 @@ namespace Chaos
 
 		virtual void BindToGraph(FPBDConstraintGraph& InContactGraph, uint32 InContainerId) {}
 
+		// Collect all the data required to solve the constraints in the specified island. This also fills the SolverBodies
+		virtual void GatherSolverInput(const FReal Dt, int32 Island) {}
+
+		// Scatter the results of the islands constraint solver(s) out to the appropriate places (e.g., impulses, break flags, etc)
+		virtual void ScatterSolverOutput(const FReal Dt, int32 Island) {}
+
 		/** Called once per tick to allow constraint containers to create/alter their constraints based on particle position */
 		virtual void UpdatePositionBasedState(const FReal Dt) {}
 
+		/** Called once per tick (usually in parallel per island) to handle solving any constraints that must be solved before the main solver loop. Used for CCD. */
+		virtual void ApplySwept(const FReal Dt, int32 Island) {}
+
 		/** Apply all corrections for constraints in the specified island. Return true if more iterations are needed. */
 		virtual bool ApplyConstraints(const FReal Dt, int32 Island, const int32 It, const int32 NumIts) { return false; }
 
@@ -148,7 +139,7 @@ namespace Chaos
 		virtual void InitializeAccelerationStructures() {}
 
 		/** Set up the perf-acceleration structures for the specified island. May be called in parallel for islands */
-		virtual void UpdateAccelerationStructures(const int32 Island) {}
+		virtual void UpdateAccelerationStructures(const FReal Dt, const int32 Island) {}
 
 		virtual void SetUseContactGraph(const bool InUseContactGraph) {}
 
@@ -181,11 +172,11 @@ namespace Chaos
 	/**
 	 * ConstraintGraphRule helper base class - templatized on Constraint Container.
 	 */
-	template<typename T_CONSTRAINTS>
+	template<typename ConstraintType>
 	class CHAOS_API TPBDConstraintGraphRuleImpl : public FPBDConstraintGraphRule
 	{
 	public:
-		typedef T_CONSTRAINTS FConstraints;
+		typedef ConstraintType FConstraints;
 
 		TPBDConstraintGraphRuleImpl(FConstraints& InConstraints, int32 InPriority);
 
@@ -199,16 +190,6 @@ namespace Chaos
 			Constraints.UnprepareTick();
 		}
 
-		virtual void PrepareIteration(FReal Dt) override
-		{
-			Constraints.PrepareIteration(Dt);
-		}
-
-		virtual void UnprepareIteration(FReal Dt) override
-		{
-			Constraints.UnprepareIteration(Dt);
-		}
-
 		virtual void BindToGraph(FPBDConstraintGraph& InContactGraph, uint32 InContainerId) override;
 
 		virtual void UpdatePositionBasedState(const FReal Dt) override;
@@ -236,86 +217,37 @@ namespace Chaos
 	/**
 	 * Island-based constraint rule. All constraints in an island are updated in single-threaded a loop. Islands may be updated in parallel.
 	 */
-	template<typename T_CONSTRAINTS>
-	class CHAOS_API TPBDConstraintIslandRule : public TPBDConstraintGraphRuleImpl<T_CONSTRAINTS>
+	template<typename ConstraintType>
+	class CHAOS_API TPBDConstraintIslandRule : public TPBDConstraintGraphRuleImpl<ConstraintType>
 	{
-		typedef TPBDConstraintGraphRuleImpl<T_CONSTRAINTS> Base;
+		typedef TPBDConstraintGraphRuleImpl<ConstraintType> Base;
 
 	public:
-		using FConstraints = T_CONSTRAINTS;
+		using FConstraints = ConstraintType;
 		using FConstraintContainerHandle = typename FConstraints::FConstraintContainerHandle;
 		using FConstraintList = TArray<FConstraintContainerHandle*>;
 
 		using Base::GetContainerId;
 
-		TPBDConstraintIslandRule(FConstraints& InConstraints, int32 InPriority = 0)
-			: TPBDConstraintGraphRuleImpl<T_CONSTRAINTS>(InConstraints, InPriority)
-		{
-		}
+		TPBDConstraintIslandRule(FConstraints& InConstraints, int32 InPriority = 0);
 
-		virtual bool ApplyConstraints(const FReal Dt, int32 Island, const int32 It, const int32 NumIts) override
-		{
-			if (IslandConstraintLists[Island].Num())
-			{
-				return Constraints.Apply(Dt, GetIslandConstraints(Island), It, NumIts);
-			}
-			return false;
-		}
+		virtual ~TPBDConstraintIslandRule();
 
-		virtual bool ApplyPushOut(const FReal Dt, int32 Island, const int32 It, const int32 NumIts) override
-		{
-			if (IslandConstraintLists[Island].Num())
-			{
-				return Constraints.ApplyPushOut(Dt, GetIslandConstraints(Island), It, NumIts);
-			}
-			return false;
-		}
+		virtual void GatherSolverInput(const FReal Dt, int32 Island) override;
 
-		virtual void InitializeAccelerationStructures() override
-		{
-			IslandConstraintLists.SetNum(ConstraintGraph->NumIslands());
-			for (FConstraintList& IslandConstraintList : IslandConstraintLists)
-			{
-				IslandConstraintList.Reset();
-			}
-		}
-		
-		virtual void UpdateAccelerationStructures(const int32 Island) override
-		{
-			const TArray<int32>& ConstraintDataIndices = ConstraintGraph->GetIslandConstraintData(Island);
-			FConstraintList& IslandConstraintList = IslandConstraintLists[Island];
-			IslandConstraintList.Reset();
-			IslandConstraintList.Reserve(ConstraintDataIndices.Num());
-			for (int32 ConstraintDataIndex : ConstraintDataIndices)
-			{
-				const typename FPBDConstraintGraph::FConstraintData& ConstraintData = ConstraintGraph->GetConstraintData(ConstraintDataIndex);
-				if (ConstraintData.GetContainerId() == GetContainerId())
-				{
-					IslandConstraintList.Add(ConstraintData.GetConstraintHandle()->As<FConstraintContainerHandle>());
-				}
-			}
-		}
+		virtual void ScatterSolverOutput(const FReal Dt, int32 Island) override;
 
-		template<typename TVisitor>
-		void VisitIslandConstraints(const int32 Island, const TVisitor& Visitor) const
-		{
-			Visitor(GetIslandConstraints(Island), Island);
-		}
+		virtual bool ApplyConstraints(const FReal Dt, int32 Island, const int32 It, const int32 NumIts) override;
+
+		virtual bool ApplyPushOut(const FReal Dt, int32 Island, const int32 It, const int32 NumIts) override;
+
+		virtual void InitializeAccelerationStructures() override;
+		
+		virtual void UpdateAccelerationStructures(const FReal Dt, const int32 Island) override;
 
 	private:
 		using Base::Constraints;
 		using Base::ConstraintGraph;
-
-		const TArray<FConstraintContainerHandle*>& GetIslandConstraints(int32 Island) const
-		{
-			// Constraint rules are bound to a single type, but the FPBDConstraintGraph works with many types. We have
-			// already pre-filtered the constraint lists based on type, so this case is safe.
-			return reinterpret_cast<const TArray<FConstraintContainerHandle*>&>(IslandConstraintLists[Island]);
-		}
-
-		// @todo(ccaulfield): optimize: eliminate the need for this index list - it is a subset of EdgeData. 
-		// If EdgeData were sorted by ContainerId we could use TArray<TArrayView<int32>>...
-		TArray<FConstraintList> IslandConstraintLists;
 	};
 
 	/**
@@ -324,45 +256,31 @@ namespace Chaos
 	 * The level is used to implement shock propagation: constraints of lower levels are frozen in 
 	 * place as far as higher-level constraints are concerned.
 	 */
-	template<typename T_CONSTRAINTS>
-	class CHAOS_API TPBDConstraintColorRule : public TPBDConstraintGraphRuleImpl<T_CONSTRAINTS>
+	template<typename ConstraintType>
+	class CHAOS_API TPBDConstraintColorRule : public TPBDConstraintGraphRuleImpl<ConstraintType>
 	{
-		typedef TPBDConstraintGraphRuleImpl<T_CONSTRAINTS> Base;
+		typedef TPBDConstraintGraphRuleImpl<ConstraintType> Base;
 
 	public:
-		using FConstraints = T_CONSTRAINTS;
+		using FConstraints = ConstraintType;
+		using FConstraintContainerHandle = typename FConstraints::FConstraintContainerHandle;
 
-		TPBDConstraintColorRule(FConstraints& InConstraints, int32 InPriority = 0)
-			: TPBDConstraintGraphRuleImpl<T_CONSTRAINTS>(InConstraints, InPriority)
-		{
-		}
+		TPBDConstraintColorRule(FConstraints& InConstraints, int32 InPriority = 0);
 
+		virtual ~TPBDConstraintColorRule();
 
 		virtual void UpdatePositionBasedState(const FReal Dt) override
 		{
 			Constraints.UpdatePositionBasedState(Dt);
 		}
 
-		virtual bool ApplyConstraints(const FReal Dt, int32 Island, const int32 It, const int32 NumIts) override
-		{
-			const typename FPBDConstraintColor::FLevelToColorToConstraintListMap& LevelToColorToConstraintListMap = GraphColor.GetIslandLevelToColorToConstraintListMap(Island);
-			int32 MaxColor = GraphColor.GetIslandMaxColor(Island);
-			int32 MaxLevel = GraphColor.GetIslandMaxLevel(Island);
-			bool bNeedsMoreIterations = false;
-			for (int32 Level = 0; Level <= MaxLevel; ++Level)
-			{
-				for (int32 Color = 0; Color <= MaxColor; ++Color)
-				{
-					if (LevelToColorToConstraintListMap[Level].Contains(Color) && LevelToColorToConstraintListMap[Level][Color].Num())
-					{
-						const TArray<typename FConstraints::FConstraintContainerHandle*>& ConstraintHandles = GetLevelColorConstraints(LevelToColorToConstraintListMap, Level, Color);
-						bNeedsMoreIterations |= Constraints.Apply(Dt, ConstraintHandles, It, NumIts);
-					}
-				}
-			}
+		virtual void GatherSolverInput(const FReal Dt, int32 Island) override;
 
-			return bNeedsMoreIterations;
-		}
+		virtual void ScatterSolverOutput(const FReal Dt, int32 Island) override;
+
+		virtual void ApplySwept(const FReal Dt, int32 Island) override;
+
+		virtual bool ApplyConstraints(const FReal Dt, int32 Island, const int32 It, const int32 NumIts) override;
 
 		virtual void RemoveConstraints(const TSet<TGeometryParticleHandle<FReal, 3>*>& InConstraints)
 		{
@@ -374,96 +292,13 @@ namespace Chaos
 			Constraints.Reset();
 		}
 
-		virtual bool ApplyPushOut(const FReal Dt, int32 Island, const int32 It, const int32 NumIts) override
-		{
-			const typename FPBDConstraintColor::FLevelToColorToConstraintListMap& LevelToColorToConstraintListMap = GraphColor.GetIslandLevelToColorToConstraintListMap(Island);
-			int32 MaxColor = GraphColor.GetIslandMaxColor(Island);
-			int32 MaxLevel = GraphColor.GetIslandMaxLevel(Island);
-
-			TSet<const TGeometryParticleHandle<FReal, 3>*> IsTemporarilyStatic;
-			bool bNeedsAnotherIteration = false;
-			for (int32 Level = 0; Level <= MaxLevel; ++Level)
-			{
-				for (int32 Color = 0; Color <= MaxColor; ++Color)
-				{
-					if (LevelToColorToConstraintListMap[Level].Contains(Color) && LevelToColorToConstraintListMap[Level][Color].Num())
-					{
-						const TArray<typename FConstraints::FConstraintContainerHandle*>& ConstraintHandles = GetLevelColorConstraints(LevelToColorToConstraintListMap, Level, Color);
-						if (Constraints.ApplyPushOut(Dt, ConstraintHandles, IsTemporarilyStatic, It, NumIts))
-						{
-							bNeedsAnotherIteration = true;
-						}
-					}
-				}
-
-				// @todo(ccaulfield): Move shock propagation out of color rule
-#if USE_SHOCK_PROPOGATION
-				for (int32 Color = 0; Color <= MaxColor; ++Color)
-				{
-					if (LevelToColorToConstraintListMap[Level].Contains(Color))
-					{
-						const TArray<typename FConstraints::FConstraintContainerHandle*>& ConstraintHandles = GetLevelColorConstraints(LevelToColorToConstraintListMap, Level, Color);
-
-						for (int32 Edge = 0; Edge < ConstraintHandles.Num(); ++Edge)
-						{
-							const typename FConstraints::FConstraintContainerHandle* Handle = ConstraintHandles[Edge];
-							TVector<const TGeometryParticleHandle<FReal, 3>*, 2> Particles = Handle->GetConstrainedParticles();
-							if (It == NumIts - 1)
-							{
-								const bool bIsParticleDynamic0 = Particles[0]->CastToRigidParticle() && Particles[0]->ObjectState() == EObjectStateType::Dynamic;
-								const bool bIsParticleDynamic1 = Particles[1]->CastToRigidParticle() && Particles[1]->ObjectState() == EObjectStateType::Dynamic;
-
-								if (bIsParticleDynamic0 == false || IsTemporarilyStatic.Contains(Particles[0]))
-								{
-									IsTemporarilyStatic.Add(Particles[1]);
-								}
-								else if (bIsParticleDynamic1 == false || IsTemporarilyStatic.Contains(Particles[1]))
-								{
-									IsTemporarilyStatic.Add(Particles[0]);
-								}
-							}
-						}
-					}
-				}
-#endif
-			}
-
-			return bNeedsAnotherIteration;
-		}
-
-		virtual void InitializeAccelerationStructures() override
-		{
-			GraphColor.InitializeColor(*ConstraintGraph);
-		}
+		virtual bool ApplyPushOut(const FReal Dt, int32 Island, const int32 It, const int32 NumIts) override;
 
-		virtual void UpdateAccelerationStructures(const int32 Island) override
-		{
-			GraphColor.ComputeColor(Island, *ConstraintGraph, Constraints.GetContainerId());
-		}
+		virtual void InitializeAccelerationStructures() override;
 
-		virtual void SetUseContactGraph(const bool bInUseContactGraph) override
-		{
-			GraphColor.SetUseContactGraph(bInUseContactGraph);
-		}
+		virtual void UpdateAccelerationStructures(const FReal Dt, const int32 Island) override;
 
-		template<typename TVisitor>
-		void VisitIslandConstraints(const int32 Island, const TVisitor& Visitor) const
-		{
-			const typename FPBDConstraintColor::FLevelToColorToConstraintListMap& LevelToColorToConstraintListMap = GraphColor.GetIslandLevelToColorToConstraintListMap(Island);
-			int32 MaxColor = GraphColor.GetIslandMaxColor(Island);
-			int32 MaxLevel = GraphColor.GetIslandMaxLevel(Island);
-			for (int32 Level = 0; Level <= MaxLevel; ++Level)
-			{
-				for (int32 Color = 0; Color <= MaxColor; ++Color)
-				{
-					if (LevelToColorToConstraintListMap[Level].Contains(Color) && LevelToColorToConstraintListMap[Level][Color].Num())
-					{
-						const TArray<typename FConstraints::FConstraintContainerHandle*>& ConstraintHandles = GetLevelColorConstraints(LevelToColorToConstraintListMap, Level, Color);
-						Visitor(ConstraintHandles, Island, Level, Color);
-					}
-				}
-			}
-		}
+		virtual void SetUseContactGraph(const bool bInUseContactGraph) override;
 
 		const FPBDConstraintColor GetGraphColor() const
 		{
@@ -474,13 +309,20 @@ namespace Chaos
 		using Base::Constraints;
 		using Base::ConstraintGraph;
 
-		const TArray<typename FConstraints::FConstraintContainerHandle*>& GetLevelColorConstraints(const typename FPBDConstraintColor::FLevelToColorToConstraintListMap& LevelToColorToConstraintListMap, int32 Level, int32 Color) const
+		const TArray<FConstraintContainerHandle*>& GetLevelColorConstraints(const typename FPBDConstraintColor::FLevelToColorToConstraintListMap& LevelToColorToConstraintListMap, int32 Level, int32 Color) const
 		{
 			// FPBDConstraintColor works with any constraint type (in principle - currently only used with Collisions), but the rule is bound to a single type and so this cast is ok
-			return reinterpret_cast<const TArray<typename FConstraints::FConstraintContainerHandle*>&>(LevelToColorToConstraintListMap[Level][Color]);
+			return reinterpret_cast<const TArray<FConstraintContainerHandle*>&>(LevelToColorToConstraintListMap[Level][Color]);
 		}
 
 		FPBDConstraintColor GraphColor;
+
+		// Each array entry contains the [begin, end) index of a set of independent constraints
+		// that can be solved in parallel. The sets are ordered by level and must be solved sequentially.
+		TArray<TArray<TPair<int32, int32>>> ConstraintSets;
+
+		// Whether to enable graph coloring and parallel processing of same-color constraints
+		bool bUseColor;
 	};
 
 }
diff --git a/Engine/Source/Runtime/Experimental/Chaos/Public/Chaos/PBDJointConstraints.h b/Engine/Source/Runtime/Experimental/Chaos/Public/Chaos/PBDJointConstraints.h
--- a/Engine/Source/Runtime/Experimental/Chaos/Public/Chaos/PBDJointConstraints.h
+++ b/Engine/Source/Runtime/Experimental/Chaos/Public/Chaos/PBDJointConstraints.h
@@ -10,6 +10,7 @@
 
 #include "Chaos/ConstraintHandle.h"
 #include "Chaos/Collision/CollisionApplyType.h"
+#include "Chaos/Evolution/SolverDatas.h"
 #include "Chaos/Joint/JointSolverConstraints.h"
 #include "Chaos/Joint/PBDJointSolverGaussSeidel.h"
 #include "Chaos/ParticleHandleFwd.h"
@@ -21,17 +22,17 @@
 namespace Chaos
 {
 	class FJointSolverConstraints;
-	class FJointSolverGaussSeidel;
+	class FPBDJointSolver;
+	class FPBDIslandSolverData;
 
-	class CHAOS_API FPBDJointConstraintHandle : public TContainerConstraintHandle<FPBDJointConstraints>
+	class CHAOS_API FPBDJointConstraintHandle : public TIndexedContainerConstraintHandle<FPBDJointConstraints>
 	{
 	public:
-		using Base = TContainerConstraintHandle<FPBDJointConstraints>;
+		using Base = TIndexedContainerConstraintHandle<FPBDJointConstraints>;
 		using FConstraintContainer = FPBDJointConstraints;
 
 		FPBDJointConstraintHandle();
 		FPBDJointConstraintHandle(FConstraintContainer* InConstraintContainer, int32 InConstraintIndex);
-		static EConstraintContainerType StaticType() { return EConstraintContainerType::Joint; }
 
 		void SetConstraintEnabled(bool bEnabled);
 
@@ -54,6 +55,14 @@ namespace Chaos
 		void SetSettings(const FPBDJointSettings& Settings);
 		TVec2<FGeometryParticleHandle*> GetConstrainedParticles() const;
 
+		void GatherInput(const FReal Dt, const int32 Particle0Level, const int32 Particle1Level, FPBDIslandSolverData& SolverData);
+
+		static const FConstraintHandleTypeID& StaticType()
+		{
+			static FConstraintHandleTypeID STypeID(TEXT("FJointConstraintHandle"), &FIndexedConstraintHandle::StaticType());
+			return STypeID;
+		}
+
 		ESyncState SyncState() const;
 		void SetSyncState(ESyncState SyncState);
 
@@ -87,10 +96,10 @@ namespace Chaos
 	/**
 	 * A joint restricting up to 6 degrees of freedom, with linear and angular limits.
 	 */
-	class CHAOS_API FPBDJointConstraints : public FPBDConstraintContainer
+	class CHAOS_API FPBDJointConstraints : public FPBDIndexedConstraintContainer
 	{
 	public:
-		using Base = FPBDConstraintContainer;
+		using Base = FPBDIndexedConstraintContainer;
 
 		using FConstraintContainerHandle = FPBDJointConstraintHandle;
 		using FConstraintHandleAllocator = TConstraintHandleAllocator<FPBDJointConstraints>;
@@ -98,6 +107,7 @@ namespace Chaos
 		using FVectorPair = TVector<FVec3, 2>;
 		using FTransformPair = TVector<FRigidTransform3, 2>;
 		using FHandles = TArray<FConstraintContainerHandle*>;
+		using FConstraintSolverContainerType = FConstraintSolverContainer;	// @todo(chaos): Add island solver for this constraint type
 
 		FPBDJointConstraints(const FPBDJointSolverSettings& InSettings = FPBDJointSolverSettings());
 
@@ -192,20 +202,6 @@ namespace Chaos
 		 */
 		void FixConstraints(int32 ConstraintIndex);
 
-		/*
-		* Enable or disable velocity update in apply constraints
-		*/
-		void SetUpdateVelocityInApplyConstraints(bool bEnabled) { bUpdateVelocityInApplyConstraints = bEnabled; }
-
-		void SetPreApplyCallback(const FJointPostApplyCallback& Callback);
-		void ClearPreApplyCallback();
-
-		void SetPostApplyCallback(const FJointPostApplyCallback& Callback);
-		void ClearPostApplyCallback();
-
-		void SetPostProjectCallback(const FJointPostApplyCallback& Callback);
-		void ClearPostProjectCallback();
-
 		void SetBreakCallback(const FJointBreakCallback& Callback);
 		void ClearBreakCallback();
 
@@ -255,25 +251,26 @@ namespace Chaos
 
 		void UnprepareTick();
 
-		void PrepareIteration(FReal Dt);
-
-		void UnprepareIteration(FReal Dt);
-
 		void UpdatePositionBasedState(const FReal Dt);
 
 		//
 		// Simple Rule API
 		//
 
-		bool Apply(const FReal Dt, const int32 It, const int32 NumIts);
-		bool ApplyPushOut(const FReal Dt, const int32 It, const int32 NumIts);
+		void GatherInput(const FReal Dt, FPBDIslandSolverData& SolverData);
+		void ScatterOutput(const FReal Dt, FPBDIslandSolverData& SolverData);
+
+		bool ApplyPhase1(const FReal Dt, const int32 It, const int32 NumIts, FPBDIslandSolverData& SolverData);
+		bool ApplyPhase2(const FReal Dt, const int32 It, const int32 NumIts, FPBDIslandSolverData& SolverData);
 
 		//
 		// Island Rule API
 		//
 
-		bool Apply(const FReal Dt, const TArray<FConstraintContainerHandle*>& InConstraintHandles, const int32 It, const int32 NumIts);
-		bool ApplyPushOut(const FReal Dt, const TArray<FConstraintContainerHandle*>& InConstraintHandles, const int32 It, const int32 NumIts);
+		void SetNumIslandConstraints(const int32 NumIslandConstraints, FPBDIslandSolverData& SolverData);
+		void GatherInput(const FReal Dt, const int32 ConstraintIndex, const int32 Particle0Level, const int32 Particle1Level, FPBDIslandSolverData& SolverData);
+		bool ApplyPhase1Serial(const FReal Dt, const int32 It, const int32 NumIts, FPBDIslandSolverData& SolverData);
+		bool ApplyPhase2Serial(const FReal Dt, const int32 It, const int32 NumIts, FPBDIslandSolverData& SolverData);
 
 		/**
 		 * Set the solver method to use
@@ -281,18 +278,6 @@ namespace Chaos
 		void SetSolverType(EConstraintSolverType InSolverType)
 		{
 			SolverType = InSolverType;
-			if (InSolverType == EConstraintSolverType::GbfPbd)
-			{
-				SetUpdateVelocityInApplyConstraints(true);
-			}
-			else if (InSolverType == EConstraintSolverType::StandardPbd)
-			{
-				SetUpdateVelocityInApplyConstraints(false);
-			}
-			else if (InSolverType == EConstraintSolverType::QuasiPbd)
-			{
-				SetUpdateVelocityInApplyConstraints(false);
-			}
 		}
 
 	protected:
@@ -306,16 +291,14 @@ namespace Chaos
 
 		void GetConstrainedParticleIndices(const int32 ConstraintIndex, int32& Index0, int32& Index1) const;
 		void CalculateConstraintSpace(int32 ConstraintIndex, FVec3& OutX0, FMatrix33& OutR0, FVec3& OutX1, FMatrix33& OutR1) const;
-		void UpdateParticleState(TPBDRigidParticleHandle<FReal, 3>* Rigid, const FReal Dt, const FVec3& PrevP, const FRotation3& PrevQ, const FVec3& P, const FRotation3& Q, const bool bUpdateVelocity = true);
-		void UpdateParticleStateExplicit(TPBDRigidParticleHandle<FReal, 3>* Rigid, const FReal Dt, const FVec3& P, const FRotation3& Q, const FVec3& V, const FVec3& W);
 		
 		void ColorConstraints();
 		void SortConstraints();
 
 		bool CanEvaluate(const int32 ConstraintIndex) const;
 
-		bool ApplySingle(const FReal Dt, const int32 ConstraintIndex, const int32 NumPairIts, const int32 It, const int32 NumIts);
-		bool ApplyPushOutSingle(const FReal Dt, const int32 ConstraintIndex, const int32 NumPairIts, const int32 It, const int32 NumIts);
+		bool ApplyPhase1Single(const FReal Dt, const int32 ConstraintIndex, const int32 NumPairIts, const int32 It, const int32 NumIts);
+		bool ApplyPhase2Single(const FReal Dt, const int32 ConstraintIndex, const int32 NumPairIts, const int32 It, const int32 NumIts);
 		void ApplyBreakThreshold(const FReal Dt, int32 ConstraintIndex, const FVec3& LinearImpulse, const FVec3& AngularImpulse);
 		void ApplyPlasticityLimits(const int32 ConstraintIndex);
 
@@ -328,15 +311,11 @@ namespace Chaos
 		FHandles Handles;
 		FConstraintHandleAllocator HandleAllocator;
 		bool bJointsDirty;
-		bool bUpdateVelocityInApplyConstraints;
 
-		FJointPreApplyCallback PreApplyCallback;
-		FJointPostApplyCallback PostApplyCallback;
-		FJointPostApplyCallback PostProjectCallback;
 		FJointBreakCallback BreakCallback;
 
 		// @todo(ccaulfield): optimize storage for joint solver
-		TArray<FJointSolverGaussSeidel> ConstraintSolvers;
+		TArray<FPBDJointSolver> ConstraintSolvers;
 
 		EConstraintSolverType SolverType;
 	};
diff --git a/Engine/Source/Runtime/Experimental/Chaos/Public/Chaos/PBDPositionConstraints.h b/Engine/Source/Runtime/Experimental/Chaos/Public/Chaos/PBDPositionConstraints.h
--- a/Engine/Source/Runtime/Experimental/Chaos/Public/Chaos/PBDPositionConstraints.h
+++ b/Engine/Source/Runtime/Experimental/Chaos/Public/Chaos/PBDPositionConstraints.h
@@ -3,48 +3,59 @@
 
 #include "Chaos/Array.h"
 #include "Chaos/ConstraintHandle.h"
+#include "Chaos/Evolution/SolverDatas.h"
 #include "Chaos/ParticleHandle.h"
 #include "Chaos/PBDConstraintContainer.h"
 
 namespace Chaos
 {
 	class FPBDPositionConstraints;
+	class FPBDIslandSolverData;
 
-	class FPBDPositionConstraintHandle : public TContainerConstraintHandle<FPBDPositionConstraints>
+	class FPBDPositionConstraintHandle : public TIndexedContainerConstraintHandle<FPBDPositionConstraints>
 	{
 	public:
-		using Base = TContainerConstraintHandle<FPBDPositionConstraints>;
+		using Base = TIndexedContainerConstraintHandle<FPBDPositionConstraints>;
 		using FConstraintContainer = FPBDPositionConstraints;
 		using FGeometryParticleHandle = FGeometryParticleHandle;
 
 		FPBDPositionConstraintHandle() {}
 		FPBDPositionConstraintHandle(FConstraintContainer* InConstraintContainer, int32 InConstraintIndex) 
-			: TContainerConstraintHandle<FPBDPositionConstraints>(InConstraintContainer, InConstraintIndex) {}
-		static EConstraintContainerType StaticType() { return EConstraintContainerType::Position; }
+			: TIndexedContainerConstraintHandle<FPBDPositionConstraints>(InConstraintContainer, InConstraintIndex) {}
 		TVector<FGeometryParticleHandle*, 2> GetConstrainedParticles() const;
 
+		void GatherInput(const FReal Dt, const int32 Particle0Level, const int32 Particle1Level, FPBDIslandSolverData& SolverData);
+
+		static const FConstraintHandleTypeID& StaticType()
+		{
+			static FConstraintHandleTypeID STypeID(TEXT("FPositionConstraintHandle"), &FIndexedConstraintHandle::StaticType());
+			return STypeID;
+		}
+
 	protected:
 		using Base::ConstraintIndex;
 		using Base::ConcreteContainer;
 	};
 
-	class FPBDPositionConstraints : public FPBDConstraintContainer
+	//! Constraint a single particle to a world-space position
+	class FPBDPositionConstraints : public FPBDIndexedConstraintContainer
 	{
 	public:
-		using Base = FPBDConstraintContainer;
+		using Base = FPBDIndexedConstraintContainer;
 		//using FReal = T;
 		//static const int Dimensions = 3;
 		using FConstraintContainerHandle = FPBDPositionConstraintHandle;
 		using FConstraintHandleAllocator = TConstraintHandleAllocator<FPBDPositionConstraints>;
 		using FHandles = TArray<FConstraintContainerHandle*>;
+		using FConstraintSolverContainerType = FConstraintSolverContainer;	// @todo(chaos): Add island solver for this constraint type
 
 		FPBDPositionConstraints(const FReal InStiffness = (FReal)1.)
-			: FPBDConstraintContainer(EConstraintContainerType::Position)
+			: FPBDIndexedConstraintContainer(FConstraintContainerHandle::StaticType())
 			, Stiffness(InStiffness)
 		{}
 
 		FPBDPositionConstraints(TArray<FVec3>&& Locations, TArray<FPBDRigidParticleHandle*>&& InConstrainedParticles, const FReal InStiffness = (FReal)1.)
-			: FPBDConstraintContainer(EConstraintContainerType::Position)
+			: FPBDIndexedConstraintContainer(FConstraintContainerHandle::StaticType())
 			, Targets(MoveTemp(Locations)), ConstrainedParticles(MoveTemp(InConstrainedParticles)), Stiffness(InStiffness)
 		{
 			if (ConstrainedParticles.Num() > 0)
@@ -54,6 +65,7 @@ namespace Chaos
 				{
 					Handles.Add(HandleAllocator.AllocHandle(this, ConstraintIndex));
 				}
+				ConstraintSolverBodies.SetNumZeroed(ConstrainedParticles.Num());
 			}
 		}
 
@@ -80,6 +92,7 @@ namespace Chaos
 			int32 NewIndex = Targets.Num();
 			Targets.Add(Position);
 			ConstrainedParticles.Add(Particle);
+			ConstraintSolverBodies.Add(nullptr);
 			Handles.Add(HandleAllocator.AllocHandle(this, NewIndex));
 			return Handles[NewIndex];
 		}
@@ -100,6 +113,7 @@ namespace Chaos
 			// Swap the last constraint into the gap to keep the array packed
 			Targets.RemoveAtSwap(ConstraintIndex);
 			ConstrainedParticles.RemoveAtSwap(ConstraintIndex);
+			ConstraintSolverBodies.RemoveAtSwap(ConstraintIndex);
 			Handles.RemoveAtSwap(ConstraintIndex);
 
 			// Update the handle for the constraint that was moved
@@ -163,27 +177,20 @@ namespace Chaos
 			Targets[ConstraintIndex] = Position;
 		}
 
-
 		//
 		// Island Rule API
 		//
 
 		void PrepareTick() {}
-
 		void UnprepareTick() {}
-
-		void PrepareIteration(FReal Dt) {}
-
-		void UnprepareIteration(FReal Dt) {}
-
 		void UpdatePositionBasedState(const FReal Dt) {}
 
-		bool Apply(const FReal Dt, const TArray<FConstraintContainerHandle*>& ConstraintHandles, const int32 It, const int32 NumIts) const;
+		void SetNumIslandConstraints(const int32 NumIslandConstraints, FPBDIslandSolverData& SolverData);
+		void GatherInput(const FReal Dt, const int32 ConstraintIndex, const int32 Particle0Level, const int32 Particle1Level, FPBDIslandSolverData& SolverData);
+		void ScatterOutput(FReal Dt, FPBDIslandSolverData& SolverData);
 
-		bool ApplyPushOut(const FReal Dt, const TArray<FConstraintContainerHandle*>& InConstraintIndices, const int32 It, const int32 NumIts) const
-		{
-			return false;
-		}
+		bool ApplyPhase1Serial(const FReal Dt, const int32 It, const int32 NumIts, FPBDIslandSolverData& SolverData);
+		bool ApplyPhase2Serial(const FReal Dt, const int32 It, const int32 NumIts, FPBDIslandSolverData& SolverData) { return false; }
 
 	protected:
 		using Base::GetConstraintIndex;
@@ -192,12 +199,13 @@ namespace Chaos
 	private:
 		void ApplySingle(const FReal Dt, int32 ConstraintIndex) const
 		{
-			if (FPBDRigidParticleHandle* PBDRigid = ConstrainedParticles[ConstraintIndex])
+			FSolverBody* Body = ConstraintSolverBodies[ConstraintIndex];
+			if (Body != nullptr)
 			{
-				const FVec3& P1 = PBDRigid->P();
+				const FVec3& P1 = Body->P();
 				const FVec3& P2 = Targets[ConstraintIndex];
-				FVec3 Difference = P1 - P2;
-				PBDRigid->P() -= Stiffness * Difference;
+				const FVec3 Difference = P1 - P2;
+				Body->ApplyPositionDelta(-Stiffness * Difference);
 			}
 		}
 
@@ -205,6 +213,8 @@ namespace Chaos
 		TArray<FPBDRigidParticleHandle*> ConstrainedParticles;
 		FReal Stiffness;
 
+		TArray<FSolverBody*> ConstraintSolverBodies;
+
 		FHandles Handles;
 		FConstraintHandleAllocator HandleAllocator;
 	};
diff --git a/Engine/Source/Runtime/Experimental/Chaos/Public/Chaos/PBDRigidDynamicSpringConstraints.h b/Engine/Source/Runtime/Experimental/Chaos/Public/Chaos/PBDRigidDynamicSpringConstraints.h
--- a/Engine/Source/Runtime/Experimental/Chaos/Public/Chaos/PBDRigidDynamicSpringConstraints.h
+++ b/Engine/Source/Runtime/Experimental/Chaos/Public/Chaos/PBDRigidDynamicSpringConstraints.h
@@ -5,15 +5,17 @@
 #include "Chaos/ConstraintHandle.h"
 #include "Chaos/ParticleHandle.h"
 #include "Chaos/PBDConstraintContainer.h"
+#include "Chaos/Evolution/SolverConstraintContainer.h"
 
 namespace Chaos
 {
 	class FPBDRigidDynamicSpringConstraints;
+	class FPBDIslandSolverData;
 
-	class CHAOS_API FPBDRigidDynamicSpringConstraintHandle : public TContainerConstraintHandle<FPBDRigidDynamicSpringConstraints>
+	class CHAOS_API FPBDRigidDynamicSpringConstraintHandle : public TIndexedContainerConstraintHandle<FPBDRigidDynamicSpringConstraints>
 	{
 	public:
-		using Base = TContainerConstraintHandle<FPBDRigidDynamicSpringConstraints>;
+		using Base = TIndexedContainerConstraintHandle<FPBDRigidDynamicSpringConstraints>;
 		using FConstraintContainer = FPBDRigidDynamicSpringConstraints;
 
 		FPBDRigidDynamicSpringConstraintHandle() 
@@ -21,37 +23,45 @@ namespace Chaos
 		}
 
 		FPBDRigidDynamicSpringConstraintHandle(FConstraintContainer* InConstraintContainer, int32 InConstraintIndex) 
-			: TContainerConstraintHandle<FPBDRigidDynamicSpringConstraints>(InConstraintContainer, InConstraintIndex) 
+			: TIndexedContainerConstraintHandle<FPBDRigidDynamicSpringConstraints>(InConstraintContainer, InConstraintIndex)
 		{
 		}
 
-		static EConstraintContainerType StaticType() { return EConstraintContainerType::DynamicSpring; }
 		TVec2<FGeometryParticleHandle*> GetConstrainedParticles() const;
 
+		void GatherInput(const FReal Dt, const int32 Particle0Level, const int32 Particle1Level, FPBDIslandSolverData& SolverData);
+
+		static const FConstraintHandleTypeID& StaticType()
+		{
+			static FConstraintHandleTypeID STypeID(TEXT("FRigidDynamicSpringConstraintHandle"), &FIndexedConstraintHandle::StaticType());
+			return STypeID;
+		}
+
 	protected:
 		using Base::ConstraintIndex;
 		using Base::ConcreteContainer;
 	};
 
-	class CHAOS_API FPBDRigidDynamicSpringConstraints : public FPBDConstraintContainer
+	class CHAOS_API FPBDRigidDynamicSpringConstraints : public FPBDIndexedConstraintContainer
 	{
 	public:
-		using Base = FPBDConstraintContainer;
+		using Base = FPBDIndexedConstraintContainer;
 		using FConstrainedParticlePair = TVec2<FGeometryParticleHandle*>;
 		//static const int Dimensions = 3;
 		using FConstraintContainerHandle = FPBDRigidDynamicSpringConstraintHandle;
 		using FConstraintHandleAllocator = TConstraintHandleAllocator<FPBDRigidDynamicSpringConstraints>;
 		using FHandles = TArray<FConstraintContainerHandle*>;
+		using FConstraintSolverContainerType = FConstraintSolverContainer;	// @todo(chaos): Add island solver for this constraint type
 
 		FPBDRigidDynamicSpringConstraints(const FReal InStiffness = (FReal)1.)
-			: FPBDConstraintContainer(EConstraintContainerType::DynamicSpring)
+			: FPBDIndexedConstraintContainer(FConstraintContainerHandle::StaticType())
 			, CreationThreshold(1)
 			, MaxSprings(1)
 			, Stiffness(InStiffness) 
 		{}
 
 		FPBDRigidDynamicSpringConstraints(TArray<FConstrainedParticlePair>&& InConstraints, const FReal InCreationThreshold = (FReal)1., const int32 InMaxSprings = 1, const FReal InStiffness = (FReal)1.)
-			: FPBDConstraintContainer(EConstraintContainerType::DynamicSpring)
+			: FPBDIndexedConstraintContainer(FConstraintContainerHandle::StaticType())
 			, Constraints(MoveTemp(InConstraints))
 			, CreationThreshold(InCreationThreshold)
 			, MaxSprings(InMaxSprings)
@@ -62,11 +72,13 @@ namespace Chaos
 				Handles.Reserve(Constraints.Num());
 				Distances.Reserve(Constraints.Num());
 				SpringDistances.Reserve(Constraints.Num());
+				ConstraintSolverBodies.Reserve(Constraints.Num());
 				for (int32 ConstraintIndex = 0; ConstraintIndex < Constraints.Num(); ++ConstraintIndex)
 				{
 					Handles.Add(HandleAllocator.AllocHandle(this, ConstraintIndex));
 					Distances.Add({});
 					SpringDistances.Add({});
+					ConstraintSolverBodies.Add({ nullptr, nullptr });
 				}
 			}
 		}
@@ -95,6 +107,7 @@ namespace Chaos
 			Constraints.Add(InConstrainedParticles);
 			Distances.Add({});
 			SpringDistances.Add({});
+			ConstraintSolverBodies.Add({ nullptr, nullptr });
 			return Handles.Last();
 		}
 
@@ -115,6 +128,7 @@ namespace Chaos
 			Constraints.RemoveAtSwap(ConstraintIndex);
 			Distances.RemoveAtSwap(ConstraintIndex);
 			SpringDistances.RemoveAtSwap(ConstraintIndex);
+			ConstraintSolverBodies.RemoveAtSwap(ConstraintIndex);
 			Handles.RemoveAtSwap(ConstraintIndex);
 
 			// Update the handle for the constraint that was moved
@@ -189,30 +203,15 @@ namespace Chaos
 		//
 
 		void PrepareTick() {}
-
 		void UnprepareTick() {}
-
-		void PrepareIteration(FReal Dt) {}
-
-		void UnprepareIteration(FReal Dt) {}
-
 		void UpdatePositionBasedState(const FReal Dt);
 
-		bool Apply(const FReal Dt, const TArray<FConstraintContainerHandle*>& InConstraintHandles, const int32 It, const int32 NumIts)
-		{
-			for (FConstraintContainerHandle* ConstraintHandle : InConstraintHandles)
-			{
-				ApplySingle(Dt, ConstraintHandle->GetConstraintIndex());
-			}
-
-			// TODO: Return true only if more iteration are needed
-			return true;
-		}
+		void SetNumIslandConstraints(const int32 NumIslandConstraints, FPBDIslandSolverData& SolverData);
+		void GatherInput(const FReal Dt, const int32 ConstraintIndex, const int32 Particle0Level, const int32 Particle1Level, FPBDIslandSolverData& SolverData);
+		void ScatterOutput(FReal Dt, FPBDIslandSolverData& SolverData);
 
-		bool ApplyPushOut(const FReal Dt, const TArray<FConstraintContainerHandle*>& InConstraintHandles, const int32 It, const int32 NumIts)
-		{
-			return false;
-		}
+		bool ApplyPhase1Serial(const FReal Dt, const int32 It, const int32 NumIts, FPBDIslandSolverData& SolverData);
+		bool ApplyPhase2Serial(const FReal Dt, const int32 It, const int32 NumIts, FPBDIslandSolverData& SolverData) { return false; }
 
 	protected:
 		using Base::GetConstraintIndex;
@@ -230,6 +229,8 @@ namespace Chaos
 		int32 MaxSprings;
 		FReal Stiffness;
 
+		TArray<FSolverBodyPtrPair> ConstraintSolverBodies;
+
 		FHandles Handles;
 		FConstraintHandleAllocator HandleAllocator;
 	};
diff --git a/Engine/Source/Runtime/Experimental/Chaos/Public/Chaos/PBDRigidParticles.h b/Engine/Source/Runtime/Experimental/Chaos/Public/Chaos/PBDRigidParticles.h
--- a/Engine/Source/Runtime/Experimental/Chaos/Public/Chaos/PBDRigidParticles.h
+++ b/Engine/Source/Runtime/Experimental/Chaos/Public/Chaos/PBDRigidParticles.h
@@ -33,6 +33,7 @@ class TPBDRigidParticles : public TRigidParticles<T, d>
 		TArrayCollection::AddArray(&MQ);
 		TArrayCollection::AddArray(&MPreV);
 		TArrayCollection::AddArray(&MPreW);
+		TArrayCollection::AddArray(&MSolverBodyIndex);
 	}
 	TPBDRigidParticles(const TPBDRigidParticles<T, d>& Other) = delete;
 	CHAOS_API TPBDRigidParticles(TPBDRigidParticles<T, d>&& Other)
@@ -41,12 +42,14 @@ class TPBDRigidParticles : public TRigidParticles<T, d>
 		, MQ(MoveTemp(Other.MQ))
 		, MPreV(MoveTemp(Other.MPreV))
 		, MPreW(MoveTemp(Other.MPreW))
+		, MSolverBodyIndex(MoveTemp(Other.MSolverBodyIndex))
 	{
 		this->MParticleType = EParticleType::Rigid;
 		TArrayCollection::AddArray(&MP);
 		TArrayCollection::AddArray(&MQ);
 		TArrayCollection::AddArray(&MPreV);
 		TArrayCollection::AddArray(&MPreW);
+		TArrayCollection::AddArray(&MSolverBodyIndex);
 	}
 
 	CHAOS_API virtual ~TPBDRigidParticles()
@@ -64,6 +67,11 @@ class TPBDRigidParticles : public TRigidParticles<T, d>
 	CHAOS_API const TVector<T, d>& PreW(const int32 index) const { return MPreW[index]; }
 	CHAOS_API TVector<T, d>& PreW(const int32 index) { return MPreW[index]; }
 
+	// The index into an FSolverBodyContainer (for dynamic particles only), or INDEX_NONE.
+	// \see FSolverBodyContainer
+	CHAOS_API int32 SolverBodyIndex(const int32 index) const { return MSolverBodyIndex[index]; }
+	CHAOS_API void SetSolverBodyIndex(const int32 index, const int32 InSolverBodyIndex) { MSolverBodyIndex[index] = InSolverBodyIndex; }
+
     // Must be reinterpret cast instead of static_cast as it's a forward declare
 	typedef TPBDRigidParticleHandle<T, d> THandleType;
 	const THandleType* Handle(int32 Index) const { return reinterpret_cast<const THandleType*>(TGeometryParticles<T,d>::Handle(Index)); }
@@ -208,6 +216,7 @@ class TPBDRigidParticles : public TRigidParticles<T, d>
 	TArrayCollectionArray<TRotation<T, d>> MQ;
 	TArrayCollectionArray<TVector<T, d>> MPreV;
 	TArrayCollectionArray<TVector<T, d>> MPreW;
+	TArrayCollectionArray<int32> MSolverBodyIndex;	// Transient for use in constraint solver
 };
 
 #if PLATFORM_MAC || PLATFORM_LINUX
diff --git a/Engine/Source/Runtime/Experimental/Chaos/Public/Chaos/PBDRigidSpringConstraints.h b/Engine/Source/Runtime/Experimental/Chaos/Public/Chaos/PBDRigidSpringConstraints.h
--- a/Engine/Source/Runtime/Experimental/Chaos/Public/Chaos/PBDRigidSpringConstraints.h
+++ b/Engine/Source/Runtime/Experimental/Chaos/Public/Chaos/PBDRigidSpringConstraints.h
@@ -5,15 +5,17 @@
 #include "Chaos/ConstraintHandle.h"
 #include "Chaos/ParticleHandle.h"
 #include "Chaos/PBDConstraintContainer.h"
+#include "Chaos/Evolution/SolverConstraintContainer.h"
 
 namespace Chaos
 {
 	class FPBDRigidSpringConstraints;
+	class FPBDIslandSolverData;
 
-	class CHAOS_API FPBDRigidSpringConstraintHandle : public TContainerConstraintHandle<FPBDRigidSpringConstraints>
+	class CHAOS_API FPBDRigidSpringConstraintHandle : public TIndexedContainerConstraintHandle<FPBDRigidSpringConstraints>
 	{
 	public:
-		using Base = TContainerConstraintHandle<FPBDRigidSpringConstraints>;
+		using Base = TIndexedContainerConstraintHandle<FPBDRigidSpringConstraints>;
 		using FConstraintContainer = FPBDRigidSpringConstraints;
 		using FGeometryParticleHandle = TGeometryParticleHandle<FReal, 3>;
 
@@ -22,12 +24,10 @@ namespace Chaos
 		}
 		
 		FPBDRigidSpringConstraintHandle(FConstraintContainer* InConstraintContainer, int32 InConstraintIndex) 
-		: TContainerConstraintHandle<FPBDRigidSpringConstraints>(InConstraintContainer, InConstraintIndex) 
+		: TIndexedContainerConstraintHandle<FPBDRigidSpringConstraints>(InConstraintContainer, InConstraintIndex)
 		{
 		}
 
-		static EConstraintContainerType StaticType() { return EConstraintContainerType::RigidSpring; }
-
 		const TVector<FVec3, 2>& GetConstraintPositions() const;
 		void SetConstraintPositions(const TVector<FVec3, 2>& ConstraintPositions);
 		
@@ -37,19 +37,27 @@ namespace Chaos
 		FReal GetRestLength() const;
 		void SetRestLength(const FReal SpringLength);
 
+		void GatherInput(const FReal Dt, const int32 Particle0Level, const int32 Particle1Level, FPBDIslandSolverData& SolverData);
+
+		static const FConstraintHandleTypeID& StaticType()
+		{
+			static FConstraintHandleTypeID STypeID(TEXT("FRigidSpringConstraintHandle"), &FIndexedConstraintHandle::StaticType());
+			return STypeID;
+		}
 	};
 
 
-	class FPBDRigidSpringConstraints : public FPBDConstraintContainer
+	class FPBDRigidSpringConstraints : public FPBDIndexedConstraintContainer
 	{
 	public:
 		// @todo(ccaulfield): an alternative AddConstraint which takes the constrain settings rather than assuming everything is in world-space rest pose
 
-		using Base = FPBDConstraintContainer;
+		using Base = FPBDIndexedConstraintContainer;
 		using FConstraintContainerHandle = FPBDRigidSpringConstraintHandle;
 		using FConstraintHandleAllocator = TConstraintHandleAllocator<FPBDRigidSpringConstraints>;
 		using FConstrainedParticlePair = TVector<TGeometryParticleHandle<FReal, 3>*, 2>;
 		using FHandles = TArray<FConstraintContainerHandle*>;
+		using FConstraintSolverContainerType = FConstraintSolverContainer;	// @todo(chaos): Add island solver for this constraint type
 
 		FPBDRigidSpringConstraints();
 		virtual ~FPBDRigidSpringConstraints();
@@ -150,43 +158,38 @@ namespace Chaos
 			SpringSettings[ConstraintIndex].RestLength = SpringLength;
 		}
 
+		//
+		// Simple Rule API
+		//
+
+		void GatherInput(const FReal Dt, FPBDIslandSolverData& SolverData);
+		void ScatterOutput(const FReal Dt, FPBDIslandSolverData& SolverData);
+
+		bool ApplyPhase1(const FReal Dt, const int32 It, const int32 NumIts, FPBDIslandSolverData& SolverData);
+		bool ApplyPhase2(const FReal Dt, const int32 It, const int32 NumIts, FPBDIslandSolverData& SolverData) { return false; }
 
 		//
 		// Island Rule API
 		//
 
 		void PrepareTick() {}
-
 		void UnprepareTick() {}
-
-		void PrepareIteration(FReal Dt) {}
-
-		void UnprepareIteration(FReal Dt) {}
-
 		void UpdatePositionBasedState(const FReal Dt) {}
 
-		bool Apply(const FReal Dt, const int32 It, const int32 NumIts);
-
-		bool ApplyPushOut(const FReal Dt, const int32 It, const int32 NumIts)
-		{
-			return false;
-		}
+		void SetNumIslandConstraints(const int32 NumIslandConstraints, FPBDIslandSolverData& SolverData);
+		void GatherInput(const FReal Dt, const int32 ConstraintIndex, const int32 Particle0Level, const int32 Particle1Level, FPBDIslandSolverData& SolverData);
 
-		bool Apply(const FReal Dt, const TArray<FConstraintContainerHandle*>& InConstraintHandles, const int32 It, const int32 NumIts);
-
-		bool ApplyPushOut(const FReal Dt, const TArray<FConstraintContainerHandle*>& InConstraintHandles, const int32 It, const int32 NumIts)
-		{
-			return false;
-		}
+		bool ApplyPhase1Serial(const FReal Dt, const int32 It, const int32 NumIts, FPBDIslandSolverData& SolverData);
+		bool ApplyPhase2Serial(const FReal Dt, const int32 It, const int32 NumIts, FPBDIslandSolverData& SolverData) { return false; }
 
 	protected:
 		using Base::GetConstraintIndex;
 		using Base::SetConstraintIndex;
 
 	private:
-		void ApplySingle(const FReal Dt, int32 ConstraintIndex) const;
+		void ApplyPhase1Single(const FReal Dt, int32 ConstraintIndex) const;
 
-		void UpdateDistance(int32 ConstraintIndex, const FVec3& Location0, const FVec3& Location1);
+		void InitDistance(int32 ConstraintIndex, const FVec3& Location0, const FVec3& Location1);
 
 		FVec3 GetDelta(int32 ConstraintIndex, const FVec3& WorldSpaceX1, const FVec3& WorldSpaceX2) const;
 
@@ -201,6 +204,8 @@ namespace Chaos
 		TArray<TVector<FVec3, 2>> Distances;
 		TArray<FSpringSettings> SpringSettings;
 
+		TArray<FSolverBodyPtrPair> ConstraintSolverBodies;
+
 		TArray<FConstraintContainerHandle*> Handles;
 		FConstraintHandleAllocator HandleAllocator;
 	};
diff --git a/Engine/Source/Runtime/Experimental/Chaos/Public/Chaos/PBDRigidsEvolution.h b/Engine/Source/Runtime/Experimental/Chaos/Public/Chaos/PBDRigidsEvolution.h
--- a/Engine/Source/Runtime/Experimental/Chaos/Public/Chaos/PBDRigidsEvolution.h
+++ b/Engine/Source/Runtime/Experimental/Chaos/Public/Chaos/PBDRigidsEvolution.h
@@ -20,7 +20,6 @@
 #include "ProfilingDebugging/CsvProfiler.h"
 #include "RewindData.h"
 
-
 extern int32 ChaosRigidsEvolutionApplyAllowEarlyOutCVar;
 extern int32 ChaosRigidsEvolutionApplyPushoutAllowEarlyOutCVar;
 extern int32 ChaosNumPushOutIterationsOverride;
@@ -297,7 +296,6 @@ public:
 
 	CHAOS_API void AddForceFunction(FForceRule ForceFunction) { ForceRules.Add(ForceFunction); }
 	CHAOS_API void AddImpulseFunction(FForceRule ImpulseFunction) { ImpulseRules.Add(ImpulseFunction); }
-	CHAOS_API void SetParticleUpdateVelocityFunction(FUpdateVelocityRule ParticleUpdate) { ParticleUpdateVelocity = ParticleUpdate; }
 	CHAOS_API void SetParticleUpdatePositionFunction(FUpdatePositionRule ParticleUpdate) { ParticleUpdatePosition = ParticleUpdate; }
 	CHAOS_API void SetCaptureRewindDataFunction(FCaptureRewindRule Rule){ CaptureRewindData = Rule; }
 
@@ -339,8 +337,6 @@ public:
 		Particles.EnableParticle(Particle);
 		EnableConstraints({ Particle });
 		ConstraintGraph.EnableParticle(Particle, ParentParticle);
-		// enable constraint after disabling the particle in the graph because of dependencies 
-		AddConstraintsToConstraintGraph(Particle->ParticleConstraints());
 		DirtyParticle(*Particle);
 	}
 
@@ -351,6 +347,7 @@ public:
 		// disable constraint before disabling the particle in the graph because of dependencies 
 		DisableConstraints({ Particle });
 		ConstraintGraph.DisableParticle(Particle);
+		DisableConstraints(TSet<FGeometryParticleHandle*>({ Particle }));
 	}
 
 	CHAOS_API void FlushExternalAccelerationQueue(FAccelerationStructure& Acceleration,FPendingSpatialDataQueue& ExternalQueue);
@@ -363,7 +360,6 @@ public:
 		}
 	}
 
-
 	template <bool bPersistent>
 	FORCEINLINE_DEBUGGABLE void DirtyParticle(TGeometryParticleHandleImp<FReal, 3, bPersistent>& Particle)
 	{
@@ -413,7 +409,7 @@ public:
 
 	CHAOS_API void DestroyParticle(FGeometryParticleHandle* Particle)
 	{
-		if(MRewindData)
+		if (MRewindData)
 		{
 			MRewindData->RemoveObject(Particle);
 		}
@@ -612,54 +608,13 @@ public:
 		}
 	}
 
-	void PrepareIteration(const FReal Dt)
-	{
-		for (FPBDConstraintGraphRule* ConstraintRule : ConstraintRules)
-		{
-			ConstraintRule->PrepareIteration(Dt);
-		}
-	}
-
-	void UnprepareIteration(const FReal Dt)
-	{
-		for (FPBDConstraintGraphRule* ConstraintRule : ConstraintRules)
-		{
-			ConstraintRule->UnprepareIteration(Dt);
-		}
-	}
-
-	void UpdateAccelerationStructures(int32 Island)
+	void UpdateAccelerationStructures(const FReal Dt, const int32 Island)
 	{
 		CSV_SCOPED_TIMING_STAT(Chaos, UpdateAccelerationStructures);
 
 		for (FPBDConstraintGraphRule* ConstraintRule : ConstraintRules)
 		{
-			ConstraintRule->UpdateAccelerationStructures(Island);
-		}
-	}
-
-	void ApplyConstraints(const FReal Dt, int32 Island)
-	{
-		UpdateAccelerationStructures(Island);
-
-		{
-			CSV_SCOPED_TIMING_STAT(Chaos, ApplyConstraints);
-
-			int32 LocalNumIterations = ChaosNumContactIterationsOverride >= 0 ? ChaosNumContactIterationsOverride : NumIterations;
-			// @todo(ccaulfield): track whether we are sufficiently solved and can early-out
-			for (int i = 0; i < LocalNumIterations; ++i)
-			{
-				bool bNeedsAnotherIteration = false;
-				for (FPBDConstraintGraphRule* ConstraintRule : PrioritizedConstraintRules)
-				{
-					bNeedsAnotherIteration |= ConstraintRule->ApplyConstraints(Dt, Island, i, LocalNumIterations);
-				}
-
-				if (ChaosRigidsEvolutionApplyAllowEarlyOutCVar && !bNeedsAnotherIteration)
-				{
-					break;
-				}
-			}
+			ConstraintRule->UpdateAccelerationStructures(Dt, Island);
 		}
 	}
 
@@ -747,17 +702,16 @@ public:
 				Rigid->PreW() = Rigid->W();
 
 				// Update the world bounds
-			/*	if (Rigid->HasBounds())
+				if (Rigid->HasBounds())
 				{
 					const FAABB3& LocalBounds = Rigid->LocalBounds();
 					FAABB3 WorldSpaceBounds = LocalBounds.TransformedAABB(FRigidTransform3(Rigid->X(), Rigid->R()));
-					if (Rigid->CCDEnabled())
-					{
-						WorldSpaceBounds.ThickenSymmetrically(Rigid->V() * Dt);
-					}
+					//if (Rigid->CCDEnabled())
+					//{
+					//	WorldSpaceBounds.ThickenSymmetrically(Rigid->V() * Dt);
+					//}
 					Rigid->SetWorldSpaceInflatedBounds(WorldSpaceBounds);
-				}*/
-				
+				}
 			}
 		}
 	}
@@ -887,37 +841,13 @@ protected:
 		}
 	}
 	
-	void UpdateVelocities(const FReal Dt, int32 Island)
-	{
-		ParticleUpdateVelocity(ConstraintGraph.GetIslandParticles(Island), Dt);
-	}
-
-	void ApplyPushOut(const FReal Dt, int32 Island)
-	{
-		int32 LocalNumPushOutIterations = ChaosNumPushOutIterationsOverride >= 0 ? ChaosNumPushOutIterationsOverride : NumPushOutIterations;
-		bool bNeedsAnotherIteration = true;
-		for (int32 It = 0; It < LocalNumPushOutIterations; ++It)
-		{
-			bNeedsAnotherIteration = false;
-			for (FPBDConstraintGraphRule* ConstraintRule : PrioritizedConstraintRules)
-			{
-				bNeedsAnotherIteration |= ConstraintRule->ApplyPushOut(Dt, Island, It, LocalNumPushOutIterations);
-			}
-
-			if (ChaosRigidsEvolutionApplyPushoutAllowEarlyOutCVar && !bNeedsAnotherIteration)
-			{
-				break;
-			}
-		}
-	}
-
 	void FlushInternalAccelerationQueue();
 	void FlushAsyncAccelerationQueue();
 	void WaitOnAccelerationStructure();
+	static void CopyPristineAccelerationStructures(const TMap<FSpatialAccelerationIdx, TUniquePtr<FSpatialAccelerationCache>>& SpatialAccelerationCache, FAccelerationStructure* FromStructure, FAccelerationStructure* ToStructure, bool CheckPristine);
 
 	TArray<FForceRule> ForceRules;
 	TArray<FForceRule> ImpulseRules;
-	FUpdateVelocityRule ParticleUpdateVelocity;
 	FUpdatePositionRule ParticleUpdatePosition;
 	FKinematicUpdateRule KinematicUpdate;
 	FCaptureRewindRule CaptureRewindData;
@@ -1013,7 +943,7 @@ protected:
 		bool bNeedsReset;
 
 	private:
-		void UpdateStructure(FAccelerationStructure* AccelerationStructure);
+		void UpdateStructure(FAccelerationStructure* AccelerationStructure, FAccelerationStructure* CopyToAccelerationStructure = nullptr);
 	};
 	FGraphEventRef AccelerationStructureTaskComplete;
 
@@ -1028,7 +958,7 @@ protected:
 	void ReleasePendingIndices();
 
 	TArray<FUniqueIdx> PendingReleaseIndices;	//for now just assume a one frame delay, but may need something more general
-	bool bIsResim = false;
+	bool bIsResim = false; 
 };
 
 
diff --git a/Engine/Source/Runtime/Experimental/Chaos/Public/Chaos/PBDRigidsEvolutionGBF.h b/Engine/Source/Runtime/Experimental/Chaos/Public/Chaos/PBDRigidsEvolutionGBF.h
--- a/Engine/Source/Runtime/Experimental/Chaos/Public/Chaos/PBDRigidsEvolutionGBF.h
+++ b/Engine/Source/Runtime/Experimental/Chaos/Public/Chaos/PBDRigidsEvolutionGBF.h
@@ -5,6 +5,7 @@
 #include "Chaos/Collision/NarrowPhase.h"
 #include "Chaos/Collision/SpatialAccelerationBroadPhase.h"
 #include "Chaos/Collision/SpatialAccelerationCollisionDetector.h"
+#include "Chaos/Evolution/SolverBodyContainer.h"
 #include "Chaos/PBDCollisionConstraints.h"
 #include "Chaos/PBDRigidsEvolution.h"
 #include "Chaos/PerParticleAddImpulses.h"
@@ -18,15 +19,16 @@
 
 namespace Chaos
 {
+	class FCollisionConstraintAllocator;
 	class FChaosArchive;
 	class IResimCacheBase;
 	class FEvolutionResimCache;
 
-	CHAOS_API extern FRealSingle HackMaxAngularVelocity;
-	CHAOS_API extern FRealSingle HackMaxVelocity;
-
-	CHAOS_API extern FRealSingle HackLinearDrag;
-	CHAOS_API extern FRealSingle HackAngularDrag;
+	namespace CVars
+	{
+		CHAOS_API extern FRealSingle HackMaxAngularVelocity;
+		CHAOS_API extern FRealSingle HackMaxVelocity;
+	}
 
 	using FPBDRigidsEvolutionCallback = TFunction<void()>;
 
@@ -64,18 +66,13 @@ namespace Chaos
 		using Base::NumIslands;
 		using Base::GetNonDisabledClusteredView;
 		using Base::DisableParticle;
-		using Base::PrepareIteration;
 		using Base::GetConstraintGraph;
-		using Base::ApplyConstraints;
-		using Base::UpdateVelocities;
 		using Base::PhysicsMaterials;
 		using Base::PerParticlePhysicsMaterials;
 		using Base::ParticleDisableCount;
 		using Base::SolverPhysicsMaterials;
-		using Base::UnprepareIteration;
 		using Base::CaptureRewindData;
 		using Base::Collided;
-		using Base::SetParticleUpdateVelocityFunction;
 		using Base::SetParticleUpdatePositionFunction;
 		using Base::AddForceFunction;
 		using Base::AddConstraintRule;
@@ -97,19 +94,16 @@ namespace Chaos
 		static constexpr int32 DefaultNumIterations = 8;
 		static constexpr int32 DefaultNumCollisionPairIterations = 1;
 		static constexpr int32 DefaultNumPushOutIterations = 1;
-		static constexpr int32 DefaultNumCollisionPushOutPairIterations = 3;
+		static constexpr int32 DefaultNumCollisionPushOutPairIterations = 1;
 		static constexpr FRealSingle DefaultCollisionMarginFraction = 0.1f;
 		static constexpr FRealSingle DefaultCollisionMarginMax = 100.0f;
 		static constexpr FRealSingle DefaultCollisionCullDistance = 5.0f;
-		static constexpr int32 DefaultNumJointPairIterations = 3;
-		static constexpr int32 DefaultNumJointPushOutPairIterations = 0;
+		static constexpr int32 DefaultNumJointPairIterations = 1;
+		static constexpr int32 DefaultNumJointPushOutPairIterations = 1;
 		static constexpr int32 DefaultRestitutionThreshold = 1000;
 
-		// @todo(chaos): Required by clustering - clean up
-		using Base::ApplyPushOut;
-
 		CHAOS_API FPBDRigidsEvolutionGBF(FPBDRigidsSOAs& InParticles, THandleArray<FChaosPhysicsMaterial>& SolverPhysicsMaterials, const TArray<ISimCallbackObject*>* InCollisionModifiers = nullptr, bool InIsSingleThreaded = false);
-		CHAOS_API ~FPBDRigidsEvolutionGBF() {}
+		CHAOS_API ~FPBDRigidsEvolutionGBF();
 
 		FORCEINLINE void SetPostIntegrateCallback(const FPBDRigidsEvolutionCallback& Cb)
 		{
@@ -191,8 +185,8 @@ namespace Chaos
 			FPerParticleEtherDrag EtherDragRule;
 			FPerParticlePBDEulerStep EulerStepRule;
 
-			const FReal MaxAngularSpeedSq = HackMaxAngularVelocity * HackMaxAngularVelocity;
-			const FReal MaxSpeedSq = HackMaxVelocity * HackMaxVelocity;
+			const FReal MaxAngularSpeedSq = CVars::HackMaxAngularVelocity * CVars::HackMaxAngularVelocity;
+			const FReal MaxSpeedSq = CVars::HackMaxVelocity * CVars::HackMaxVelocity;
 			InParticles.ParallelFor([&](auto& GeomParticle, int32 Index) {
 				//question: can we enforce this at the API layer? Right now islands contain non dynamic which makes this hard
 				auto PBDParticle = GeomParticle.CastToRigidParticle();
@@ -212,21 +206,21 @@ namespace Chaos
 					AddImpulsesRule.Apply(Particle, Dt);
 					EtherDragRule.Apply(Particle, Dt);
 
-					if (HackMaxAngularVelocity >= 0.f)
+					if (CVars::HackMaxAngularVelocity >= 0.f)
 					{
 						const FReal AngularSpeedSq = Particle.W().SizeSquared();
 						if (AngularSpeedSq > MaxAngularSpeedSq)
 						{
-							Particle.W() = Particle.W() * (HackMaxAngularVelocity / FMath::Sqrt(AngularSpeedSq));
+							Particle.W() = Particle.W() * (CVars::HackMaxAngularVelocity / FMath::Sqrt(AngularSpeedSq));
 						}
 					}
 
-					if (HackMaxVelocity >= 0.f)
+					if (CVars::HackMaxVelocity >= 0.f)
 					{
 						const FReal SpeedSq = Particle.V().SizeSquared();
 						if (SpeedSq > MaxSpeedSq)
 						{
-							Particle.V() = Particle.V() * (HackMaxVelocity / FMath::Sqrt(SpeedSq));
+							Particle.V() = Particle.V() * (CVars::HackMaxVelocity / FMath::Sqrt(SpeedSq));
 						}
 					}
 
@@ -251,6 +245,23 @@ namespace Chaos
 			}
 		}
 
+		// Pre-solver phase for handling special constraints (used for CCD)
+		void ApplyConstraintsPhase0(const FReal Dt, int32 Island);
+
+		// First phase of constraint solver
+		// For GBF this is the velocity solve phase
+		// For PBD/QuasiPBD this is the position solve phase
+		void ApplyConstraintsPhase1(const FReal Dt, int32 Island);
+
+		// Calculate the implicit velocites based on the change in position from ApplyConstraintsPhase1
+		void SetImplicitVelocities(const FReal Dt, int32 Island);
+		
+		// Second phase of constraint solver (after implicit velocity calculation following results of phase 1)
+		// For GBF this is the pushout phase
+		// For QuasiPBD this is the velocity solve phase
+		void ApplyConstraintsPhase2(const FReal Dt, int32 Island);
+
+
 		CHAOS_API void Serialize(FChaosArchive& Ar);
 
 		CHAOS_API TUniquePtr<IResimCacheBase> CreateExternalResimCache() const;
@@ -263,6 +274,9 @@ namespace Chaos
 	protected:
 
 		CHAOS_API void AdvanceOneTimeStepImpl(const FReal dt, const FSubStepInfo& SubStepInfo);
+
+		void GatherSolverInput(FReal Dt, int32 Island);
+		void ScatterSolverOutput(FReal Dt, int32 Island);
 		
 		FEvolutionResimCache* GetCurrentStepResimCache()
 		{
@@ -283,6 +297,9 @@ namespace Chaos
 		FNarrowPhase NarrowPhase;
 		FSpatialAccelerationCollisionDetector CollisionDetector;
 
+		// @todo(chaos): these are transient and the arrays should belong in the Island when we have one...
+		TArray<FSolverBodyContainer> IslandSolverBodies;
+
 		FPBDRigidsEvolutionCallback PostIntegrateCallback;
 		FPBDRigidsEvolutionCallback PostDetectCollisionsCallback;
 		FPBDRigidsEvolutionCallback PreApplyCallback;
diff --git a/Engine/Source/Runtime/Experimental/Chaos/Public/Chaos/PBDSuspensionConstraints.h b/Engine/Source/Runtime/Experimental/Chaos/Public/Chaos/PBDSuspensionConstraints.h
--- a/Engine/Source/Runtime/Experimental/Chaos/Public/Chaos/PBDSuspensionConstraints.h
+++ b/Engine/Source/Runtime/Experimental/Chaos/Public/Chaos/PBDSuspensionConstraints.h
@@ -3,6 +3,7 @@
 
 #include "Chaos/Array.h"
 #include "Chaos/ConstraintHandle.h"
+#include "Chaos/Evolution/SolverDatas.h"
 #include "Chaos/ParticleHandle.h"
 #include "Chaos/PBDConstraintContainer.h"
 #include "Chaos/PBDSuspensionConstraintTypes.h"
@@ -15,16 +16,16 @@
 namespace Chaos
 {
 	class FPBDSuspensionConstraints;
+	class FPBDIslandSolverData;
 
-	class CHAOS_API FPBDSuspensionConstraintHandle : public TContainerConstraintHandle<FPBDSuspensionConstraints>
+	class CHAOS_API FPBDSuspensionConstraintHandle : public TIndexedContainerConstraintHandle<FPBDSuspensionConstraints>
 	{
 	public:
-		using Base = TContainerConstraintHandle<FPBDSuspensionConstraints>;
+		using Base = TIndexedContainerConstraintHandle<FPBDSuspensionConstraints>;
 		using FConstraintContainer = FPBDSuspensionConstraints;
 
 		FPBDSuspensionConstraintHandle() {}
 		FPBDSuspensionConstraintHandle(FConstraintContainer* InConstraintContainer, int32 InConstraintIndex);
-		static EConstraintContainerType StaticType() { return EConstraintContainerType::Suspension; }
 
 		FPBDSuspensionSettings& GetSettings();
 		const FPBDSuspensionSettings& GetSettings() const;
@@ -33,26 +34,34 @@ namespace Chaos
 
 		TVec2<FGeometryParticleHandle*> GetConstrainedParticles() const;
 
+		void GatherInput(const FReal Dt, const int32 Particle0Level, const int32 Particle1Level, FPBDIslandSolverData& SolverData);
+
+		static const FConstraintHandleTypeID& StaticType()
+		{
+			static FConstraintHandleTypeID STypeID(TEXT("FSuspensionConstraintHandle"), &FIndexedConstraintHandle::StaticType());
+			return STypeID;
+		}
 	protected:
 		using Base::ConstraintIndex;
 		using Base::ConcreteContainer;
 	};
 
-	class CHAOS_API FPBDSuspensionConstraints : public FPBDConstraintContainer
+	class CHAOS_API FPBDSuspensionConstraints : public FPBDIndexedConstraintContainer
 	{
 	public:
-		using Base = FPBDConstraintContainer;
+		using Base = FPBDIndexedConstraintContainer;
 		using FConstraintContainerHandle = FPBDSuspensionConstraintHandle;
 		using FConstraintHandleAllocator = TConstraintHandleAllocator<FPBDSuspensionConstraints>;
 		using FHandles = TArray<FConstraintContainerHandle*>;
+		using FConstraintSolverContainerType = FConstraintSolverContainer;	// @todo(chaos): Add island solver for this constraint type
 
 		FPBDSuspensionConstraints(const FPBDSuspensionSolverSettings& InSolverSettings = FPBDSuspensionSolverSettings())
-			: FPBDConstraintContainer(EConstraintContainerType::Suspension)
+			: FPBDIndexedConstraintContainer(FConstraintContainerHandle::StaticType())
 			, SolverSettings(InSolverSettings)
 		{}
 
 		FPBDSuspensionConstraints(TArray<FVec3>&& Locations, TArray<TGeometryParticleHandle<FReal,3>*>&& InConstrainedParticles, TArray<FVec3>&& InLocalOffset, TArray<FPBDSuspensionSettings>&& InConstraintSettings)
-			: FPBDConstraintContainer(EConstraintContainerType::Suspension)
+			: FPBDIndexedConstraintContainer(FConstraintContainerHandle::StaticType())
 			, ConstrainedParticles(MoveTemp(InConstrainedParticles)), SuspensionLocalOffset(MoveTemp(InLocalOffset)), ConstraintSettings(MoveTemp(InConstraintSettings))
 		{
 			if (ConstrainedParticles.Num() > 0)
@@ -104,11 +113,11 @@ namespace Chaos
 			{
 				for (FConstraintHandle* ConstraintHandle : RemovedParticle->ParticleConstraints())
 				{
-					if (ConstraintHandle->As<FPBDSuspensionConstraintHandle>())
+					if (FPBDSuspensionConstraintHandle* SuspensionHandle = ConstraintHandle->As<FPBDSuspensionConstraintHandle>())
 					{
-						ConstraintHandle->SetEnabled(false); // constraint lifespan is managed by the proxy
+						SuspensionHandle->SetEnabled(false); // constraint lifespan is managed by the proxy
 
-						int ConstraintIndex = ConstraintHandle->GetConstraintIndex();
+						int ConstraintIndex = SuspensionHandle->GetConstraintIndex();
 						if (ConstraintIndex != INDEX_NONE)
 						{
 							if (ConstrainedParticles[ConstraintIndex] == RemovedParticle)
@@ -209,27 +218,20 @@ namespace Chaos
 			SuspensionLocalOffset[ConstraintIndex] = Position;
 		}
 
-
 		//
 		// Island Rule API
 		//
 
 		void PrepareTick() {}
-
 		void UnprepareTick() {}
-
-		void PrepareIteration(FReal Dt) {}
-
-		void UnprepareIteration(FReal Dt) {}
-
 		void UpdatePositionBasedState(const FReal Dt) {}
 
-		bool Apply(const FReal Dt, const TArray<FConstraintContainerHandle*>& ConstraintHandles, const int32 It, const int32 NumIts) const;
+		void SetNumIslandConstraints(const int32 NumIslandConstraints, FPBDIslandSolverData& SolverData);
+		void GatherInput(const FReal Dt, const int32 ConstraintIndex, const int32 Particle0Level, const int32 Particle1Level, FPBDIslandSolverData& SolverData);
+		void ScatterOutput(FReal Dt, FPBDIslandSolverData& SolverData);
 
-		bool ApplyPushOut(const FReal Dt, const TArray<FConstraintContainerHandle*>& InConstraintIndices, const int32 It, const int32 NumIts) const
-		{
-			return false;
-		}
+		bool ApplyPhase1Serial(const FReal Dt, const int32 It, const int32 NumIts, FPBDIslandSolverData& SolverData);
+		bool ApplyPhase2Serial(const FReal Dt, const int32 It, const int32 NumIts, FPBDIslandSolverData& SolverData) { return false; }
 
 	protected:
 		using Base::GetConstraintIndex;
@@ -247,6 +249,8 @@ namespace Chaos
 		TArray<FPBDSuspensionSettings> ConstraintSettings;
 		TArray<bool> ConstraintEnabledStates;
 
+		TArray<FSolverBody*> ConstraintSolverBodies;
+
 		FHandles Handles;
 		FConstraintHandleAllocator HandleAllocator;
 	};
diff --git a/Engine/Source/Runtime/Experimental/Chaos/Public/Chaos/Particle/ParticleUtilities.h b/Engine/Source/Runtime/Experimental/Chaos/Public/Chaos/Particle/ParticleUtilities.h
--- a/Engine/Source/Runtime/Experimental/Chaos/Public/Chaos/Particle/ParticleUtilities.h
+++ b/Engine/Source/Runtime/Experimental/Chaos/Public/Chaos/Particle/ParticleUtilities.h
@@ -51,7 +51,7 @@ namespace Chaos
 		template<typename T_PARTICLEHANDLE>
 		static inline FRigidTransform3 GetActorWorldTransform(T_PARTICLEHANDLE Particle)
 		{
-			return FRigidTransform3(Particle->X(), Particle->R());
+			return FRigidTransform3(TSpatialAccessor::GetPosition(Particle), TSpatialAccessor::GetRotation(Particle));
 		}
 
 		template<typename T_PARTICLEHANDLE>
diff --git a/Engine/Source/Runtime/Experimental/Chaos/Public/Chaos/ParticleHandle.h b/Engine/Source/Runtime/Experimental/Chaos/Public/Chaos/ParticleHandle.h
--- a/Engine/Source/Runtime/Experimental/Chaos/Public/Chaos/ParticleHandle.h
+++ b/Engine/Source/Runtime/Experimental/Chaos/Public/Chaos/ParticleHandle.h
@@ -68,11 +68,21 @@ void GeometryParticleDefaultConstruct(FConcrete& Concrete, const FGeometryPartic
 	Concrete.SetEnabledDuringResim(true);
 }
 
+
+extern CHAOS_API int32 AccelerationStructureSplitStaticAndDynamic;
 template <typename T, int d, typename FConcrete>
 void KinematicGeometryParticleDefaultConstruct(FConcrete& Concrete, const FKinematicGeometryParticleParameters& Params)
 {
 	Concrete.SetV(TVector<T, d>(0));
 	Concrete.SetW(TVector<T, d>(0));
+	if (AccelerationStructureSplitStaticAndDynamic == 1)
+	{
+		Concrete.SetSpatialIdx(FSpatialAccelerationIdx{ 0,1 });
+	}
+	else
+	{
+		Concrete.SetSpatialIdx(FSpatialAccelerationIdx{ 0,0 });
+	}
 }
 template <typename T, int d, typename FConcrete>
 void PBDRigidParticleDefaultConstruct(FConcrete& Concrete, const FPBDRigidParticleParameters& Params)
@@ -562,6 +572,9 @@ public:
 
 	const TGeometryParticleHandle<T, d>* Handle() const { return GetHandleHelper(this);}
 	TGeometryParticleHandle<T, d>* Handle() { return GetHandleHelper(this); }
+	
+	// Useful for logging to indicate particle (this is locally unique among all particles)
+	int32 GetHandleIdx() const { return HandleIdx; }
 
 	bool Sleeping() const { return GeometryParticleSleeping(*this); }
 
@@ -794,6 +807,7 @@ protected:
 		SetDisabled(Params.bDisabled);
 		SetPreV(this->V());
 		SetPreW(this->W());
+		SetSolverBodyIndex(INDEX_NONE);
 		SetP(this->X());
 		SetQ(this->R());
 		SetVSmooth(this->V());
@@ -853,6 +867,9 @@ public:
 	TVector<T, d>& PreW() { return PBDRigidParticles->PreW(ParticleIdx); }
 	void SetPreW(const TVector<T, d>& InPreW) { PBDRigidParticles->PreW(ParticleIdx) = InPreW; }
 
+	int32 SolverBodyIndex() const { return PBDRigidParticles->SolverBodyIndex(ParticleIdx); }
+	void SetSolverBodyIndex(const int32 InSolverBodyIndex) { PBDRigidParticles->SetSolverBodyIndex(ParticleIdx, InSolverBodyIndex); }
+
 	const TVector<T, d>& P() const { return PBDRigidParticles->P(ParticleIdx); }
 	TVector<T, d>& P() { return PBDRigidParticles->P(ParticleIdx); }
 	void SetP(const TVector<T, d>& InP) { PBDRigidParticles->P(ParticleIdx) = InP; }
@@ -1245,7 +1262,7 @@ public:
 	using FDynamicParticleHandleType = FPBDRigidParticleHandle;
 	using FKinematicParticleHandleType = FKinematicGeometryParticleHandle;
 
-	FGenericParticleHandleHandleImp(FGeometryParticleHandle* InHandle) { MHandle = InHandle; }
+	FGenericParticleHandleHandleImp(FGeometryParticleHandle* InHandle) : MHandle(InHandle) {}
 
 	// Check for the exact type of particle (see also AsKinematic etc, which will work on derived types)
 	bool IsStatic() const { return (MHandle->ObjectState() == EObjectStateType::Static); }
@@ -1261,6 +1278,7 @@ public:
 	//Needed for templated code to be the same
 	const FGeometryParticleHandle* Handle() const { return MHandle; }
 	FGeometryParticleHandle* Handle() { return MHandle; }
+	int32 GetHandleIdx() const { return MHandle->GetHandleIdx(); }
 
 	// Static Particles
 	FVec3& X() { return MHandle->X(); }
@@ -1286,6 +1304,8 @@ public:
 	void SetV(const FVec3& InV) { if (MHandle->CastToKinematicParticle()) { MHandle->CastToKinematicParticle()->V() = InV; } }
 	void SetW(const FVec3& InW) { if (MHandle->CastToKinematicParticle()) { MHandle->CastToKinematicParticle()->W() = InW; } }
 
+	const FKinematicTarget& KinematicTarget() const { return (MHandle->CastToKinematicParticle())? MHandle->CastToKinematicParticle()->KinematicTarget() : EmptyKinematicTarget; }
+
 	// Dynamic Particles
 
 	// TODO: Make all of these check ObjectState to maintain current functionality
@@ -1360,6 +1380,23 @@ public:
 		return ZeroVector;
 	}
 
+	int32 SolverBodyIndex() const
+	{
+		if (MHandle->CastToRigidParticle())
+		{
+			return MHandle->CastToRigidParticle()->SolverBodyIndex();
+		}
+		return INDEX_NONE;
+	}
+
+	void SetSolverBodyIndex(const int32 InSolverBodyIndex)
+	{
+		if (MHandle->CastToRigidParticle())
+		{
+			return MHandle->CastToRigidParticle()->SetSolverBodyIndex(InSolverBodyIndex);
+		}
+	}
+
 	FVec3& P()
 	{
 		if (MHandle->CastToRigidParticle() && MHandle->ObjectState() == EObjectStateType::Dynamic)
@@ -1562,16 +1599,33 @@ public:
 	}
 #endif
 
-	int32 Island() const
+	/** Get the island index from the particle handle */
+	FORCEINLINE int32 Island() const
 	{
-		if ((MHandle->CastToRigidParticle() != nullptr) && IsDynamic())
+		if (auto RigidHandle = MHandle->CastToRigidParticle())
 		{
-			return MHandle->CastToRigidParticle()->Island();
+			if (IsDynamic())
+			{
+				return RigidHandle->Island();
+			}
 		}
-
 		return INDEX_NONE;
 	}
 
+	/** Set the island index onto the particle handle 
+	 * @param IslandIndex Island Index to be set onto the particle 
+	 */
+	FORCEINLINE void SetIsland( const int32 IslandIndex) 
+	{
+		if (auto RigidHandle = MHandle->CastToRigidParticle())
+		{
+			if (IsDynamic())
+			{
+				RigidHandle->Island() = IslandIndex;
+			}
+		}
+	}
+
 	bool ToBeRemovedOnFracture() const 
 	{
 		if (MHandle->CastToRigidParticle() && MHandle->ObjectState() == EObjectStateType::Dynamic)
@@ -1596,6 +1650,7 @@ private:
 	static const FRotation3 IdentityRotation;
 	static const FMatrix33 ZeroMatrix;
 	static const TUniquePtr<FBVHParticles> NullBVHParticles;
+	static const FKinematicTarget EmptyKinematicTarget;
 };
 
 template <typename T, int d>
@@ -1612,11 +1667,29 @@ using TGenericParticleHandleHandleImp UE_DEPRECATED(4.27, "Deprecated. this clas
 class CHAOS_API FGenericParticleHandle
 {
 public:
+	FGenericParticleHandle() : Imp(nullptr) {}
 	FGenericParticleHandle(FGeometryParticleHandle* InHandle) : Imp(InHandle) {}
 
 	FGenericParticleHandleHandleImp* operator->() const { return const_cast<FGenericParticleHandleHandleImp*>(&Imp); }
 	FGenericParticleHandleHandleImp* Get() const { return const_cast<FGenericParticleHandleHandleImp*>(&Imp); }
 
+	bool IsValid() const { return Imp.Handle() != nullptr; }
+
+	friend uint32 GetTypeHash(const FGenericParticleHandle& H)
+	{
+		return GetTypeHash(H->ParticleID());
+	}
+
+	friend bool operator==(const FGenericParticleHandle& L, const FGenericParticleHandle& R)
+	{
+		return L->ParticleID() == R->ParticleID();
+	}
+
+	friend bool operator<(const FGenericParticleHandle& L, const FGenericParticleHandle& R)
+	{
+		return L->ParticleID() < R->ParticleID();
+	}
+
 private:
 	FGenericParticleHandleHandleImp Imp;
 };
@@ -1627,14 +1700,32 @@ using TGenericParticleHandle UE_DEPRECATED(4.27, "Deprecated. this class is to b
 class CHAOS_API FConstGenericParticleHandle
 {
 public:
+	FConstGenericParticleHandle() : Imp(nullptr) {}
 	FConstGenericParticleHandle(const FGeometryParticleHandle* InHandle) : Imp(const_cast<FGeometryParticleHandle*>(InHandle)) {}
 	FConstGenericParticleHandle(const FGenericParticleHandle InHandle) : Imp(InHandle->Handle()) {}
 
 	const FGenericParticleHandleHandleImp* operator->() const { return &Imp; }
 	const FGenericParticleHandleHandleImp* Get() const { return &Imp; }
 
+	bool IsValid() const { return Imp.Handle() != nullptr; }
+
+	friend uint32 GetTypeHash(const FConstGenericParticleHandle& H)
+	{
+		return GetTypeHash(H->ParticleID());
+	}
+
+	friend bool operator==(const FConstGenericParticleHandle& L, const FConstGenericParticleHandle& R)
+	{
+		return L->ParticleID() == R->ParticleID();
+	}
+
+	friend bool operator<(const FConstGenericParticleHandle& L, const FConstGenericParticleHandle& R)
+	{
+		return L->ParticleID() < R->ParticleID();
+	}
+
 private:
-	const FGenericParticleHandleHandleImp Imp;
+	FGenericParticleHandleHandleImp Imp;
 };
 
 template <typename T, int d>
diff --git a/Engine/Source/Runtime/Experimental/Chaos/Public/Chaos/ParticleIterator.h b/Engine/Source/Runtime/Experimental/Chaos/Public/Chaos/ParticleIterator.h
--- a/Engine/Source/Runtime/Experimental/Chaos/Public/Chaos/ParticleIterator.h
+++ b/Engine/Source/Runtime/Experimental/Chaos/Public/Chaos/ParticleIterator.h
@@ -292,6 +292,12 @@ template <typename TSOA>
 struct TSOAView
 {
 	using THandle = typename TSOA::THandleType;
+
+	TSOAView()
+		: SOA(nullptr)
+		, HandlesArray(nullptr)
+	{}
+
 	TSOAView(TSOA* InSOA)
 		: SOA(InSOA)
 		, HandlesArray(nullptr)
@@ -556,6 +562,12 @@ public:
 		}
 	}
 
+	TConstParticleView(TSOAView<TSOA>&& InSOAView)
+		: Size(InSOAView.Size())
+	{
+		SOAViews.Add(MoveTemp(InSOAView));
+	}
+
 	int32 Num() const
 	{
 		return Size;
@@ -621,6 +633,12 @@ TConstParticleView<TSOA> MakeConstParticleView(TArray<TSOAView<TSOA>>&& SOAViews
 	return TConstParticleView<TSOA>(MoveTemp(SOAViews));
 }
 
+template <typename TSOA>
+TConstParticleView<TSOA> MakeConstParticleView(TSOAView<TSOA>&& SOAView)
+{
+	return TConstParticleView<TSOA>(MoveTemp(SOAView));
+}
+
 template <typename TSOA>
 TParticleView<TSOA> MakeParticleView(TArray<TSOAView<TSOA>>&& SOAViews)
 {
diff --git a/Engine/Source/Runtime/Experimental/Chaos/Public/Chaos/Plane.h b/Engine/Source/Runtime/Experimental/Chaos/Public/Chaos/Plane.h
--- a/Engine/Source/Runtime/Experimental/Chaos/Public/Chaos/Plane.h
+++ b/Engine/Source/Runtime/Experimental/Chaos/Public/Chaos/Plane.h
@@ -130,6 +130,9 @@ public:
 		if (FMath::Abs(SignedDist) < Thickness)
 		{
 			//initial overlap so stop
+			//const FReal DirDotNormal = FVec3::DotProduct(Dir, (FVec3)MNormal);
+			//OutPosition = StartPoint;
+			//OutNormal = DirDotNormal < 0 ? MNormal : -MNormal;
 			OutTime = 0;
 			return true;
 		}
diff --git a/Engine/Source/Runtime/Experimental/Chaos/Public/Chaos/SpatialAccelerationCollection.h b/Engine/Source/Runtime/Experimental/Chaos/Public/Chaos/SpatialAccelerationCollection.h
--- a/Engine/Source/Runtime/Experimental/Chaos/Public/Chaos/SpatialAccelerationCollection.h
+++ b/Engine/Source/Runtime/Experimental/Chaos/Public/Chaos/SpatialAccelerationCollection.h
@@ -22,7 +22,7 @@ template <typename TPayloadType, typename T, int d>
 struct CHAOS_API TSpatialAccelerationBucketEntry
 {
 	TUniquePtr<ISpatialAcceleration<TPayloadType, T, d>> Acceleration;
-	uint16 TypeInnerIdx;
+	uint16 TypeInnerIdx; // Index in bucket
 
 	void CopyFrom(TSpatialAccelerationBucketEntry<TPayloadType, T, d>& Src)
 	{
@@ -96,6 +96,25 @@ struct CHAOS_API TSpatialCollectionBucket
 		return Idx;
 	}
 
+   	void UpdateOrAddAt(uint16 Idx, TObj&& Obj)
+	{
+		uint16 ObjectNum = static_cast<uint16>(Objects.Num());
+		if (ObjectNum <= Idx)
+		{
+			Objects.SetNum(Idx + 1);
+			for (uint16 FreeIndex = ObjectNum; FreeIndex < Idx - 1; FreeIndex++)
+			{
+				FreeIndices.Add(FreeIndex);
+			}
+		}
+		else
+		{
+			FreeIndices.Remove(Idx);
+		}
+
+		Objects[Idx] = MoveTemp(Obj);
+	}
+
 	void Remove(uint16 Idx)
 	{
 		if (Objects.Num() == Idx + 1)
@@ -166,6 +185,8 @@ struct CHAOS_API TSpatialTypeTupleGetter<0, First, Rest...>
 	static const auto& Get(const TSpatialTypeTuple<First, Rest...>& Types) { return Types.First; }
 };
 
+
+// Gets a bucket of acceleration structure pointers
 template <int Idx, typename First, typename... Rest>
 auto& GetAccelerationsPerType(TSpatialTypeTuple<First, Rest...>& Types)
 {
@@ -315,13 +336,13 @@ struct TSpatialAccelerationCollectionHelper
 };
 
 template <typename SpatialAccelerationCollection>
-typename TEnableIf<TIsSame<typename SpatialAccelerationCollection::TPayloadType, FAccelerationStructureHandle>::Value, void>::Type PBDComputeConstraintsLowLevel_Helper(FReal Dt, const SpatialAccelerationCollection& Accel, FSpatialAccelerationBroadPhase& BroadPhase, FNarrowPhase& NarrowPhase, FAsyncCollisionReceiver& Receiver, CollisionStats::FStatData& StatData, IResimCacheBase* ResimCache)
+typename TEnableIf<TIsSame<typename SpatialAccelerationCollection::TPayloadType, FAccelerationStructureHandle>::Value, void>::Type PBDComputeConstraintsLowLevel_Helper(FReal Dt, const SpatialAccelerationCollection& Accel, FSpatialAccelerationBroadPhase& BroadPhase, FNarrowPhase& NarrowPhase, IResimCacheBase* ResimCache)
 {
-	BroadPhase.ProduceOverlaps(Dt, Accel, NarrowPhase, Receiver, StatData, ResimCache);
+	BroadPhase.ProduceOverlaps(Dt, Accel, NarrowPhase, ResimCache);
 }
 
 template <typename SpatialAccelerationCollection>
-typename TEnableIf<!TIsSame<typename SpatialAccelerationCollection::TPayloadType, FAccelerationStructureHandle>::Value, void>::Type PBDComputeConstraintsLowLevel_Helper(FReal Dt, const SpatialAccelerationCollection& Accel, FSpatialAccelerationBroadPhase& BroadPhase, FNarrowPhase& NarrowPhase, FAsyncCollisionReceiver& Receiver, CollisionStats::FStatData& StatData, IResimCacheBase* ResimCache)
+typename TEnableIf<!TIsSame<typename SpatialAccelerationCollection::TPayloadType, FAccelerationStructureHandle>::Value, void>::Type PBDComputeConstraintsLowLevel_Helper(FReal Dt, const SpatialAccelerationCollection& Accel, FSpatialAccelerationBroadPhase& BroadPhase, FNarrowPhase& NarrowPhase, IResimCacheBase* ResimCache)
 {
 }
 
@@ -343,7 +364,7 @@ public:
 	{
 	}
 
-	virtual FSpatialAccelerationIdx AddSubstructure(TUniquePtr<ISpatialAcceleration<TPayloadType, T, d>>&& Substructure, uint16 BucketIdx) override
+	virtual FSpatialAccelerationIdx AddSubstructure(TUniquePtr<ISpatialAcceleration<TPayloadType, T, d>>&& Substructure, uint16 BucketIdx, uint16 BucketInnerIdx) override
 	{
 		check(BucketIdx < MaxBuckets);
 		FSpatialAccelerationIdx Result;
@@ -367,8 +388,10 @@ public:
 		}
 
 		this->ActiveBucketsMask |= (1 << BucketIdx);
-
-		Result.InnerIdx = static_cast<uint16>(Buckets[BucketIdx].Add(MoveTemp(BucketEntry)));
+		
+		Result.InnerIdx = BucketInnerIdx;
+		Buckets[BucketIdx].UpdateOrAddAt(BucketInnerIdx, MoveTemp(BucketEntry));
+		
 		return Result;
 	}
 
@@ -498,11 +521,27 @@ public:
 		return TUniquePtr<ISpatialAcceleration<TPayloadType, T, d>>(new TSpatialAccelerationCollection<TSpatialAccelerationTypes...>(*this));
 	}
 
-	virtual void PBDComputeConstraintsLowLevel(T Dt, FSpatialAccelerationBroadPhase& BroadPhase, FNarrowPhase& NarrowPhase, FAsyncCollisionReceiver& Receiver, CollisionStats::FStatData& StatData, IResimCacheBase* ResimCache) const override
+	virtual void PBDComputeConstraintsLowLevel(T Dt, FSpatialAccelerationBroadPhase& BroadPhase, FNarrowPhase& NarrowPhase, IResimCacheBase* ResimCache) const override
 	{
-		PBDComputeConstraintsLowLevel_Helper(Dt, *this, BroadPhase, NarrowPhase, Receiver, StatData, ResimCache);
+		PBDComputeConstraintsLowLevel_Helper(Dt, *this, BroadPhase, NarrowPhase, ResimCache);
 	}
 
+#if !UE_BUILD_SHIPPING
+	virtual void DebugDraw(ISpacialDebugDrawInterface<T>* InInterface) const override
+	{
+		for (const auto& Bucket : Buckets)
+		{
+			for (const auto& Entry : Bucket.Objects)
+			{
+				if (Entry.Acceleration)
+				{
+					Entry.Acceleration->DebugDraw(InInterface);
+				}
+			}
+		}
+	}
+#endif
+
 	virtual void Serialize(FChaosArchive& Ar)
 	{
 		//todo: let user serialize out bucket params
@@ -564,6 +603,13 @@ private:
 			}
 		}
 	}
+
+	TSpatialAccelerationCollection<TSpatialAccelerationTypes...>& operator=(const TSpatialAccelerationCollection<TSpatialAccelerationTypes...>& Other) = delete;
+	virtual ISpatialAcceleration<TPayloadType, FReal, 3>& operator=(const ISpatialAcceleration<TPayloadType, FReal, 3>& Other) override
+	{
+		check(false);	//not implemented
+		return *this;
+	}
 	
 	static constexpr uint32 ClampedIdx(uint32 Idx)
 	{
@@ -618,7 +664,7 @@ private:
 
 	static constexpr uint16 MaxBuckets = 8;
 	TSpatialCollectionBucket<TSpatialAccelerationBucketEntry<TPayloadType, T, d>> Buckets[MaxBuckets];
-	TSpatialTypeTuple< TSpatialAccelerationTypes...> Types;
+	TSpatialTypeTuple< TSpatialAccelerationTypes...> Types; // Have buckets of acceleration structure pointers
 	static constexpr uint32 NumTypes = sizeof...(TSpatialAccelerationTypes);
 };
 
diff --git a/Engine/Source/Runtime/Experimental/Chaos/Public/PBDRigidsSolver.h b/Engine/Source/Runtime/Experimental/Chaos/Public/PBDRigidsSolver.h
--- a/Engine/Source/Runtime/Experimental/Chaos/Public/PBDRigidsSolver.h
+++ b/Engine/Source/Runtime/Experimental/Chaos/Public/PBDRigidsSolver.h
@@ -151,6 +151,7 @@ namespace Chaos
 
 		int32 NumJointConstraints() const;
 		int32 NumCollisionConstraints() const;
+		int32 NumActiveCollisionConstraints() const;
 
 		//Make friend with unit test code so we can verify some behavior
 		template <typename TSolver>
diff --git a/Engine/Source/Runtime/Experimental/ChaosCore/Public/Chaos/ArrayND.h b/Engine/Source/Runtime/Experimental/ChaosCore/Public/Chaos/ArrayND.h
--- a/Engine/Source/Runtime/Experimental/ChaosCore/Public/Chaos/ArrayND.h
+++ b/Engine/Source/Runtime/Experimental/ChaosCore/Public/Chaos/ArrayND.h
@@ -130,7 +130,13 @@ class TArrayNDBase
 		Ar << MArray;
 	}
 
-	FORCEINLINE TArrayNDBase<T_DERIVED, T, d>& operator=(const TArrayNDBase<T_DERIVED, T, d>& Other) = delete;
+	FORCEINLINE TArrayNDBase<T_DERIVED, T, d>& operator=(const TArrayNDBase<T_DERIVED, T, d>& Other)
+	{
+		MCounts = Other.MCounts;
+		MArray = Other.MArray;
+		return *this;
+	}
+
 	FORCEINLINE TArrayNDBase<T_DERIVED, T, d>& operator=(TArrayNDBase<T_DERIVED, T, d>&& Other)
 	{
 		MCounts = Other.MCounts;
@@ -193,7 +199,11 @@ class TArrayND : public TArrayNDBase<TArrayND<T, d>, T, d>
 	    : Base(MoveTemp(Other)) {}
 	FORCEINLINE TArrayND(std::istream& Stream)
 	    : Base(Stream) {}
-	FORCEINLINE TArrayND<T, d>& operator=(const TArrayND<T, d>& Other) = delete;
+	FORCEINLINE TArrayND<T, d>& operator=(const TArrayND<T, d>& Other)
+	{
+		Base::operator=(Other);
+		return *this;
+	}
 	FORCEINLINE TArrayND<T, d>& operator=(TArrayND<T, d>&& Other)
 	{
 		Base::operator=(MoveTemp(Other));
@@ -242,7 +252,11 @@ class TArrayND<T, 3> : public TArrayNDBase<TArrayND<T, 3>, T, 3>
 	    : Base(MoveTemp(Other)) {}
 	FORCEINLINE TArrayND(std::istream& Stream)
 	    : Base(Stream) {}
-	FORCEINLINE TArrayND<T, 3>& operator=(const TArrayND<T, 3>& Other) = delete;
+	FORCEINLINE TArrayND<T, 3>& operator=(const TArrayND<T, 3>& Other)
+	{
+		Base::operator=(Other);
+		return *this;
+	}
 	FORCEINLINE TArrayND<T, 3>& operator=(TArrayND<T, 3>&& Other)
 	{
 		Base::operator=(MoveTemp(Other));
