Fix CIS - implicit conversion from int to float in PBDRigidsEvolutionGBF.cpp

diff --git a/Engine/Source/Runtime/Experimental/Chaos/Private/Chaos/PBDRigidsEvolutionGBF.cpp b/Engine/Source/Runtime/Experimental/Chaos/Private/Chaos/PBDRigidsEvolutionGBF.cpp
--- a/Engine/Source/Runtime/Experimental/Chaos/Private/Chaos/PBDRigidsEvolutionGBF.cpp
+++ b/Engine/Source/Runtime/Experimental/Chaos/Private/Chaos/PBDRigidsEvolutionGBF.cpp
@@ -156,7 +156,7 @@ void FPBDRigidsEvolutionGBF::Advance(const FReal Dt,const FReal MaxStepDt,const
 		{
 			// StepFraction: how much of the remaining time this step represents, used to interpolate kinematic targets
 			// E.g., for 4 steps this will be: 1/4, 1/2, 3/4, 1
-			const FReal StepFraction = (FReal)(Step + 1) / (NumSteps);
+			const FReal StepFraction = (FReal)(Step + 1) / (FReal)(NumSteps);
 		
 			UE_LOG(LogChaos, Verbose, TEXT("Advance dt = %f [%d/%d]"), StepDt, Step + 1, NumSteps);
 
