Fix CIS - declaration hiding class member

diff --git a/Engine/Source/Runtime/Experimental/Chaos/Public/Chaos/AABBTree.h b/Engine/Source/Runtime/Experimental/Chaos/Public/Chaos/AABBTree.h
--- a/Engine/Source/Runtime/Experimental/Chaos/Public/Chaos/AABBTree.h
+++ b/Engine/Source/Runtime/Experimental/Chaos/Public/Chaos/AABBTree.h
@@ -309,16 +309,16 @@ struct TAABBTreeLeafArray : public TBoundsWrapperHelper<TPayloadType, T, bComput
 #if !UE_BUILD_SHIPPING
 	void DebugDrawLeaf(ISpacialDebugDrawInterface<T>& InInterface, const FLinearColor& InLinearColor, float InThickness) const
 	{
-		const TAABB<T, 3> Bounds = TBoundsWrapperHelper<TPayloadType, T, bComputeBounds>::GetBounds();
+		const TAABB<T, 3> LeafBounds = TBoundsWrapperHelper<TPayloadType, T, bComputeBounds>::GetBounds();
 
 		const float Alpha = (float)Elems.Num() / 10.f;
 		const FLinearColor ColorByCount = FLinearColor::Green * (1.f - Alpha) + FLinearColor::Red * Alpha;
 		const FVec3 ColorAsVec = { ColorByCount.R, ColorByCount.G, ColorByCount.B };
 		
-		InInterface.Box(Bounds, ColorAsVec, InThickness);
+		InInterface.Box(LeafBounds, ColorAsVec, InThickness);
 		for (const auto& Elem : Elems)
 		{
-			InInterface.Line(Bounds.Center(), Elem.Bounds.Center(), ColorAsVec, InThickness);
+			InInterface.Line(LeafBounds.Center(), Elem.Bounds.Center(), ColorAsVec, InThickness);
 			InInterface.Box(Elem.Bounds, { (T)1.0, (T)0.2, (T)0.2 }, 1.0f);
 		}
 	}
