Chaos - Fix post-EPA transform for second body

diff --git a/Engine/Source/Runtime/Experimental/Chaos/Public/Chaos/GJK.h b/Engine/Source/Runtime/Experimental/Chaos/Public/Chaos/GJK.h
--- a/Engine/Source/Runtime/Experimental/Chaos/Public/Chaos/GJK.h
+++ b/Engine/Source/Runtime/Experimental/Chaos/Public/Chaos/GJK.h
@@ -319,7 +319,7 @@ namespace Chaos
 				OutNormalB = BToATM.InverseTransformVectorNoScale(MTD);
 				OutPenetration = Penetration + ThicknessA + ThicknessB;
 				OutClosestA = ClosestA + MTD * ThicknessA;
-				OutClosestB = ClosestBInA - MTD * ThicknessB;
+				OutClosestB = BToATM.InverseTransformPositionNoScale(ClosestBInA) - MTD * ThicknessB;
 				return true;
 			case EEPAResult::BadInitialSimplex:
 				// The origin is outside the simplex. Must be a touching contact and EPA setup will have calculated the normal
