Chaos - fix zero friction contacts (was not updating anchor positions in this case)

diff --git a/Engine/Source/Runtime/Experimental/Chaos/Private/Chaos/Collision/PBDCollisionSolver.h b/Engine/Source/Runtime/Experimental/Chaos/Private/Chaos/Collision/PBDCollisionSolver.h
--- a/Engine/Source/Runtime/Experimental/Chaos/Private/Chaos/Collision/PBDCollisionSolver.h
+++ b/Engine/Source/Runtime/Experimental/Chaos/Private/Chaos/Collision/PBDCollisionSolver.h
@@ -119,6 +119,9 @@ namespace Chaos
 
 		FPBDCollisionSolver();
 
+		FReal StaticFriction() const { return State.StaticFriction; }
+		FReal DynamicFriction() const { return State.DynamicFriction; }
+
 		void SetFriction(const FReal InStaticFriction, const FReal InDynamicFriction)
 		{
 			State.StaticFriction = InStaticFriction;
diff --git a/Engine/Source/Runtime/Experimental/Chaos/Private/Chaos/Collision/SolverCollisionContainer.cpp b/Engine/Source/Runtime/Experimental/Chaos/Private/Chaos/Collision/SolverCollisionContainer.cpp
--- a/Engine/Source/Runtime/Experimental/Chaos/Private/Chaos/Collision/SolverCollisionContainer.cpp
+++ b/Engine/Source/Runtime/Experimental/Chaos/Private/Chaos/Collision/SolverCollisionContainer.cpp
@@ -129,7 +129,7 @@ namespace Chaos
 				Solver.InitMaterial(
 					ManifoldPointIndex,
 					WorldContactVelocityTargetNormal,
-					ManifoldPoint.bPotentialRestingContact,
+					ManifoldPoint.bPotentialRestingContact && (Solver.StaticFriction() > 0),
 					ManifoldPoint.StaticFrictionMax);
 			}
 		}
