add query for global matrices used by niagara
also fixes horde

diff --git a/Engine/Source/Runtime/Experimental/GeometryCollectionEngine/Public/GeometryCollection/GeometryCollectionComponent.h b/Engine/Source/Runtime/Experimental/GeometryCollectionEngine/Public/GeometryCollection/GeometryCollectionComponent.h
--- a/Engine/Source/Runtime/Experimental/GeometryCollectionEngine/Public/GeometryCollection/GeometryCollectionComponent.h
+++ b/Engine/Source/Runtime/Experimental/GeometryCollectionEngine/Public/GeometryCollection/GeometryCollectionComponent.h
@@ -707,6 +707,9 @@ public:
 	UPROPERTY(EditAnywhere, BlueprintReadOnly, Category = "ChaosPhysics|General")
 	bool bNotifyRemovals;
 
+	// Used by Niagara DI to query global matrices rather than recalculating them again
+	const TArray<FMatrix>& GetGlobalMatrices() { return GlobalMatrices; }
+
 protected:
 	/** Display Bone Colors instead of assigned materials */
 	UPROPERTY(EditAnywhere, BlueprintReadWrite, Category = "ChaosPhysics|General")
@@ -735,6 +738,8 @@ protected:
 	void UpdateBreakEventRegistration();
 	void UpdateRemovalEventRegistration();
 
+
+
 	/* Per-instance override to enable/disable replication for the geometry collection */
 	UPROPERTY(EditInstanceOnly, BlueprintReadOnly, Category=Network)
 	bool bEnableReplication;
