fix cis

diff --git a/Engine/Plugins/Experimental/ChaosNiagara/Source/ChaosNiagara/Public/NiagaraDataInterfaceGeometryCollection.h b/Engine/Plugins/Experimental/ChaosNiagara/Source/ChaosNiagara/Public/NiagaraDataInterfaceGeometryCollection.h
--- a/Engine/Plugins/Experimental/ChaosNiagara/Source/ChaosNiagara/Public/NiagaraDataInterfaceGeometryCollection.h
+++ b/Engine/Plugins/Experimental/ChaosNiagara/Source/ChaosNiagara/Public/NiagaraDataInterfaceGeometryCollection.h
@@ -125,8 +125,7 @@ public:
 	DECLARE_NIAGARA_DI_PARAMETER();
 
 	UPROPERTY(EditAnywhere, Category = "Geometry Collection")
-	AGeometryCollectionActor *GeometryCollectionActor;
-
+	TObjectPtr<AGeometryCollectionActor> GeometryCollectionActor;
 
 	/** UObject Interface */
 	virtual void PostInitProperties() override;
diff --git a/Engine/Source/Runtime/Experimental/GeometryCollectionEngine/Public/GeometryCollection/GeometryCollectionComponent.h b/Engine/Source/Runtime/Experimental/GeometryCollectionEngine/Public/GeometryCollection/GeometryCollectionComponent.h
--- a/Engine/Source/Runtime/Experimental/GeometryCollectionEngine/Public/GeometryCollection/GeometryCollectionComponent.h
+++ b/Engine/Source/Runtime/Experimental/GeometryCollectionEngine/Public/GeometryCollection/GeometryCollectionComponent.h
@@ -681,6 +681,9 @@ public:
 	/** Modify DynamicCollection transform hierarchy to effect cluster breaks releasing the specified indices. */
 	void SetInitialClusterBreaks(const TArray<int32>& ReleaseIndices);
 
+	/** Used by Niagara DI to query global matrices rather than recalculating them again */
+	const TArray<FMatrix>& GetGlobalMatrices() { return GlobalMatrices; }
+
 public:
 	UPROPERTY(BlueprintAssignable, Category = "Collision")
 	FOnChaosPhysicsCollision OnChaosPhysicsCollision;
@@ -707,9 +710,6 @@ public:
 	UPROPERTY(EditAnywhere, BlueprintReadOnly, Category = "ChaosPhysics|General")
 	bool bNotifyRemovals;
 
-	// Used by Niagara DI to query global matrices rather than recalculating them again
-	const TArray<FMatrix>& GetGlobalMatrices() { return GlobalMatrices; }
-
 protected:
 	/** Display Bone Colors instead of assigned materials */
 	UPROPERTY(EditAnywhere, BlueprintReadWrite, Category = "ChaosPhysics|General")
