Chaos : Fields
- Added pass through of Chaos_ExternalClusterStrain for the geometry collection component.

diff --git a/Engine/Source/Runtime/Experimental/Chaos/Public/GeometryCollection/GeometryCollectionSimulationTypes.h b/Engine/Source/Runtime/Experimental/Chaos/Public/GeometryCollection/GeometryCollectionSimulationTypes.h
--- a/Engine/Source/Runtime/Experimental/Chaos/Public/GeometryCollection/GeometryCollectionSimulationTypes.h
+++ b/Engine/Source/Runtime/Experimental/Chaos/Public/GeometryCollection/GeometryCollectionSimulationTypes.h
@@ -54,15 +54,19 @@ enum class EGeometryCollectionPhysicsTypeEnum : uint8
 	Chaos_AngularTorque            UMETA(DisplayName = "Angular Torque", ToolTip = "Add a vector field to the particles angular torque."),
 	Chaos_DisableThreshold         UMETA(DisplayName = "Disable Threshold", ToolTip = "Disable the particles if their linear and angular velocity are less than the threshold."),
 	Chaos_SleepingThreshold        UMETA(DisplayName = "Sleeping Threshold", ToolTip = "Set particles in sleeping mode if their linear and angular velocity are less than the threshold."),
+	Chaos_ExternalClusterStrain    UMETA(DisplayName = "External Strain", ToolTip = "Apply an external strain over the particles. If this strain is over the internal one, the cluster will break."),
+	Chaos_InternalClusterStrain    UMETA(DisplayName = "Internal Strain", ToolTip = "Add a strain field to the particles internal one."),
 	//
 	Chaos_Max						UMETA(Hidden)
 };
 
+
 inline CHAOS_API EFieldPhysicsType GetGeometryCollectionPhysicsType(const EGeometryCollectionPhysicsTypeEnum GeoCollectionType)
 {
 	static const TArray<EFieldPhysicsType> PhysicsTypes = { EFieldPhysicsType::Field_AngularVelociy , EFieldPhysicsType::Field_DynamicState, EFieldPhysicsType::Field_LinearVelocity, 
 		EFieldPhysicsType::Field_InitialAngularVelocity, EFieldPhysicsType::Field_InitialLinearVelocity, EFieldPhysicsType::Field_CollisionGroup, 
-		EFieldPhysicsType::Field_LinearForce, EFieldPhysicsType::Field_AngularTorque, EFieldPhysicsType::Field_DisableThreshold, EFieldPhysicsType::Field_SleepingThreshold, EFieldPhysicsType::Field_PhysicsType_Max };
+		EFieldPhysicsType::Field_LinearForce, EFieldPhysicsType::Field_AngularTorque, EFieldPhysicsType::Field_DisableThreshold, EFieldPhysicsType::Field_SleepingThreshold, 
+		EFieldPhysicsType::Field_ExternalClusterStrain, EFieldPhysicsType::Field_InternalClusterStrain,EFieldPhysicsType::Field_PhysicsType_Max };
 
 	return PhysicsTypes[(uint8)GeoCollectionType];
 }
