Chaos - Add max pushout velocity for collisions
- Add physics config setting to control max pushout velocity

This fixes initial penetration explosions when there aren't enough velocity iterations to correct the velocity added by the PBD phase.

diff --git a/Engine/Source/Runtime/Experimental/Chaos/Private/Chaos/Collision/PBDCollisionSolver.cpp b/Engine/Source/Runtime/Experimental/Chaos/Private/Chaos/Collision/PBDCollisionSolver.cpp
--- a/Engine/Source/Runtime/Experimental/Chaos/Private/Chaos/Collision/PBDCollisionSolver.cpp
+++ b/Engine/Source/Runtime/Experimental/Chaos/Private/Chaos/Collision/PBDCollisionSolver.cpp
@@ -412,10 +412,19 @@ namespace Chaos
 		return ContactVelocity0 - ContactVelocity1;
 	}
 
-	inline void FPBDCollisionSolver::FSolverManifoldPoint::CalculateContactPositionError(const FConstraintSolverBody& Body0, const FConstraintSolverBody& Body1, FVec3& OutContactDelta, FReal& OutContactDeltaNormal) const
+	inline void FPBDCollisionSolver::FSolverManifoldPoint::CalculateContactPositionError(const FConstraintSolverBody& Body0, const FConstraintSolverBody& Body1, const FReal MaxPushOut, FVec3& OutContactDelta, FReal& OutContactDeltaNormal) const
 	{
 		OutContactDelta = (Body0.P() + Body0.Q() * LocalRelativeAnchorPoint0) - (Body1.P() + Body1.Q() * LocalRelativeAnchorPoint1);
 		OutContactDeltaNormal = FVec3::DotProduct(OutContactDelta, WorldContactNormal);
+
+		// NOTE: OutContactDeltaNormal is negative for penetration
+		// NOTE: MaxPushOut == 0 disables the pushout limits
+		if ((MaxPushOut > 0) && (OutContactDeltaNormal < -MaxPushOut))
+		{
+			const FReal ClampedContactDeltaNormal = -MaxPushOut;
+			OutContactDelta += (ClampedContactDeltaNormal - OutContactDeltaNormal) * WorldContactNormal;
+			OutContactDeltaNormal = ClampedContactDeltaNormal;
+		}
 	}
 
 	inline void FPBDCollisionSolver::FSolverManifoldPoint::CalculateContactVelocityError(const FConstraintSolverBody& Body0, const FConstraintSolverBody& Body1, const FReal DynamicFriction, const FReal Dt, FVec3& OutContactVelocityDelta, FReal& OutContactVelocityDeltaNormal) const
@@ -534,7 +543,7 @@ namespace Chaos
 		State.ManifoldPoints[ManifoldPoiontIndex].UpdateContact(State.SolverBodies[0], State.SolverBodies[1], InCoMAnchorPoint0, InCoMAnchorPoint1, InWorldContactNormal);
 	}
 
-	bool FPBDCollisionSolver::SolvePosition(const FReal Dt, const bool bApplyStaticFriction)
+	bool FPBDCollisionSolver::SolvePosition(const FReal Dt, const FReal MaxPushOut, const bool bApplyStaticFriction)
 	{
 		if (!bChaos_PBDCollisionSolver_Position_SolveEnabled)
 		{
@@ -572,7 +581,7 @@ namespace Chaos
 
 			FVec3 ContactDelta;
 			FReal ContactDeltaNormal;
-			SolverManifoldPoint.CalculateContactPositionError(Body0.SolverBody(), Body1.SolverBody(), ContactDelta, ContactDeltaNormal);
+			SolverManifoldPoint.CalculateContactPositionError(Body0.SolverBody(), Body1.SolverBody(), MaxPushOut, ContactDelta, ContactDeltaNormal);
 
 			const bool bProcessManifoldPoint = (ContactDeltaNormal < Chaos_PBDCollisionSolver_Position_NormalTolerance) || !SolverManifoldPoint.NetPushOut.IsNearlyZero();
 			if (bProcessManifoldPoint)
diff --git a/Engine/Source/Runtime/Experimental/Chaos/Private/Chaos/Collision/PBDCollisionSolver.h b/Engine/Source/Runtime/Experimental/Chaos/Private/Chaos/Collision/PBDCollisionSolver.h
--- a/Engine/Source/Runtime/Experimental/Chaos/Private/Chaos/Collision/PBDCollisionSolver.h
+++ b/Engine/Source/Runtime/Experimental/Chaos/Private/Chaos/Collision/PBDCollisionSolver.h
@@ -76,8 +76,9 @@ namespace Chaos
 
 			/**
 			 * @brief Calculate the position error at the current transforms
+			 * @param MaxPushOut a limit on the position error for this iteration to prevent initial-penetration explosion (a common PBD problem)
 			*/
-			void CalculateContactPositionError(const FConstraintSolverBody& Body0, const FConstraintSolverBody& Body1, FVec3& OutContactDelta, FReal& OutContactDeltaNormal) const;
+			void CalculateContactPositionError(const FConstraintSolverBody& Body0, const FConstraintSolverBody& Body1, const FReal MaxPushOut, FVec3& OutContactDelta, FReal& OutContactDeltaNormal) const;
 
 			/**
 			 * @brief Calculate the velocity error at the current transforms
@@ -212,7 +213,7 @@ namespace Chaos
 		 * @brief Calculate and apply the position correction for this iteration
 		 * @return true if we need to run more iterations, false if we did not apply any correction
 		*/
-		bool SolvePosition(const FReal Dt, const bool bApplyStaticFriction);
+		bool SolvePosition(const FReal Dt, const FReal MaxPushOut, const bool bApplyStaticFriction);
 
 		/**
 		 * @brief Calculate and apply the velocity correction for this iteration
diff --git a/Engine/Source/Runtime/Experimental/Chaos/Private/Chaos/Collision/SolverCollisionContainer.cpp b/Engine/Source/Runtime/Experimental/Chaos/Private/Chaos/Collision/SolverCollisionContainer.cpp
--- a/Engine/Source/Runtime/Experimental/Chaos/Private/Chaos/Collision/SolverCollisionContainer.cpp
+++ b/Engine/Source/Runtime/Experimental/Chaos/Private/Chaos/Collision/SolverCollisionContainer.cpp
@@ -182,6 +182,7 @@ namespace Chaos
 
 	FPBDCollisionSolverContainer::FPBDCollisionSolverContainer()
 		: FConstraintSolverContainer()
+		, MaxPushOutVelocity(0)
 	{
 	}
 
@@ -299,6 +300,7 @@ namespace Chaos
 
 		const int32 NumZeroFrictionIterations = Chaos_PBDCollisionSolver_Position_ZeroFrictionIterations;
 		const bool bApplyStaticFriction = (It >= NumZeroFrictionIterations);
+		const FReal MaxPushOut = (MaxPushOutVelocity > 0) ? (MaxPushOutVelocity * Dt) / FReal(NumIts) : 0;
 
 		// Apply the position correction
 		// @todo(chaos): parallel version of SolvePosition
@@ -314,7 +316,7 @@ namespace Chaos
 				CollisionSolver.GatherManifoldPoints(Dt);
 			}
 
-			bNeedsAnotherIteration |= CollisionSolvers[SolverIndex].GetSolver().SolvePosition(Dt, bApplyStaticFriction);
+			bNeedsAnotherIteration |= CollisionSolvers[SolverIndex].GetSolver().SolvePosition(Dt, MaxPushOut, bApplyStaticFriction);
 		}
 
 		return bNeedsAnotherIteration;
diff --git a/Engine/Source/Runtime/Experimental/Chaos/Private/Chaos/PBDCollisionConstraints.cpp b/Engine/Source/Runtime/Experimental/Chaos/Private/Chaos/PBDCollisionConstraints.cpp
--- a/Engine/Source/Runtime/Experimental/Chaos/Private/Chaos/PBDCollisionConstraints.cpp
+++ b/Engine/Source/Runtime/Experimental/Chaos/Private/Chaos/PBDCollisionConstraints.cpp
@@ -31,6 +31,8 @@
 #include "Algo/Sort.h"
 #include "Algo/StableSort.h"
 
+// Private includes
+#include "Collision/PBDCollisionSolver.h"
 
 #if INTEL_ISPC
 #include "PBDCollisionConstraints.ispc.generated.h"
@@ -418,7 +420,13 @@ namespace Chaos
 		{
 			// We shouldn't be adding disabled constraints to the solver list. The check needs to be at caller site or we should return success/fail - see TPBDConstraintColorRule::GatherSolverInput
 			check(Constraint.IsEnabled());
-			GetConstraintSolverContainer(SolverData).AddConstraintSolver(Dt, Constraint, Particle0Level, Particle1Level, SolverData.GetBodyContainer());
+
+			FPBDCollisionSolverContainer& SolverContainer = GetConstraintSolverContainer(SolverData);
+			SolverContainer.AddConstraintSolver(Dt, Constraint, Particle0Level, Particle1Level, SolverData.GetBodyContainer());
+
+			// @todo(chaos): we shouldn't have to pass settings down to the solvers every frame. Ideally the solver containers would be created by the constraint containers and
+			// the settings would be passed on then. Currently the solver containers are created here FPBDIslandSolverData::AddConstraintDatas
+			SolverContainer.SetMaxPushOutVelocity(MaxPushOutVelocity);
 		}
 		else
 		{
diff --git a/Engine/Source/Runtime/Experimental/Chaos/Private/PBDRigidsSolver.cpp b/Engine/Source/Runtime/Experimental/Chaos/Private/PBDRigidsSolver.cpp
--- a/Engine/Source/Runtime/Experimental/Chaos/Private/PBDRigidsSolver.cpp
+++ b/Engine/Source/Runtime/Experimental/Chaos/Private/PBDRigidsSolver.cpp
@@ -189,21 +189,24 @@ namespace Chaos
 		int32 ChaosSolverJointPushOutPairIterations = -1;
 		FAutoConsoleVariableRef CVarChaosSolverJointPushOutPairIterations(TEXT("p.Chaos.Solver.Joint.PushOutPairIterations"), ChaosSolverJointPushOutPairIterations, TEXT("Override number of push out iterations per joint during a solver iteration (-1 to use config)"));
 
-// Copied from RBAN
-Chaos::FRealSingle ChaosSolverJointPositionTolerance = 0.025f;
-FAutoConsoleVariableRef CVarChaosSolverJointPositionTolerance(TEXT("p.Chaos.Solver.Joint.PositionTolerance"), ChaosSolverJointPositionTolerance, TEXT("PositionTolerance."));
-Chaos::FRealSingle ChaosSolverJointAngleTolerance = 0.001f;
-FAutoConsoleVariableRef CVarChaosSolverJointAngleTolerance(TEXT("p.Chaos.Solver.Joint.AngleTolerance"), ChaosSolverJointAngleTolerance, TEXT("AngleTolerance."));
-Chaos::FRealSingle ChaosSolverJointMinParentMassRatio = 0.2f;
-FAutoConsoleVariableRef CVarChaosSolverJointMinParentMassRatio(TEXT("p.Chaos.Solver.Joint.MinParentMassRatio"), ChaosSolverJointMinParentMassRatio, TEXT("6Dof joint MinParentMassRatio (if > 0)"));
-Chaos::FRealSingle ChaosSolverJointMaxInertiaRatio = 5.0f;
-FAutoConsoleVariableRef CVarChaosSolverJointMaxInertiaRatio(TEXT("p.Chaos.Solver.Joint.MaxInertiaRatio"), ChaosSolverJointMaxInertiaRatio, TEXT("6Dof joint MaxInertiaRatio (if > 0)"));
+		// Copied from RBAN
+		Chaos::FRealSingle ChaosSolverJointPositionTolerance = 0.025f;
+		FAutoConsoleVariableRef CVarChaosSolverJointPositionTolerance(TEXT("p.Chaos.Solver.Joint.PositionTolerance"), ChaosSolverJointPositionTolerance, TEXT("PositionTolerance."));
+		Chaos::FRealSingle ChaosSolverJointAngleTolerance = 0.001f;
+		FAutoConsoleVariableRef CVarChaosSolverJointAngleTolerance(TEXT("p.Chaos.Solver.Joint.AngleTolerance"), ChaosSolverJointAngleTolerance, TEXT("AngleTolerance."));
+		Chaos::FRealSingle ChaosSolverJointMinParentMassRatio = 0.2f;
+		FAutoConsoleVariableRef CVarChaosSolverJointMinParentMassRatio(TEXT("p.Chaos.Solver.Joint.MinParentMassRatio"), ChaosSolverJointMinParentMassRatio, TEXT("6Dof joint MinParentMassRatio (if > 0)"));
+		Chaos::FRealSingle ChaosSolverJointMaxInertiaRatio = 5.0f;
+		FAutoConsoleVariableRef CVarChaosSolverJointMaxInertiaRatio(TEXT("p.Chaos.Solver.Joint.MaxInertiaRatio"), ChaosSolverJointMaxInertiaRatio, TEXT("6Dof joint MaxInertiaRatio (if > 0)"));
 
 		// Collision detection cvars
 		// These override the engine config if >= 0
 		Chaos::FRealSingle ChaosSolverCullDistance = -1.0f;
 		FAutoConsoleVariableRef CVarChaosSolverCullDistance(TEXT("p.Chaos.Solver.Collision.CullDistance"), ChaosSolverCullDistance, TEXT("Override cull distance (if >= 0)"));
 
+		Chaos::FRealSingle ChaosSolverMaxPushOutVelocity = -1.0f;
+		FAutoConsoleVariableRef CVarChaosSolverMaxPushOutVelocity(TEXT("p.Chaos.Solver.Collision.MaxPushOutVelocity"), ChaosSolverMaxPushOutVelocity, TEXT("Override max pushout velocity (if >= 0)"));
+
 		int32 ChaosSolverCleanupCommandsOnDestruction = 1;
 		FAutoConsoleVariableRef CVarChaosSolverCleanupCommandsOnDestruction(TEXT("p.Chaos.Solver.CleanupCommandsOnDestruction"), ChaosSolverCleanupCommandsOnDestruction, TEXT("Whether or not to run internal command queue cleanup on solver destruction (0 = no cleanup, >0 = cleanup all commands)"));
 
@@ -822,6 +825,10 @@ namespace Chaos
 			{
 				SetCollisionCullDistance(ChaosSolverCullDistance);
 			}
+			if (ChaosSolverMaxPushOutVelocity >= 0.0f)
+			{
+				SetCollisionMaxPushOutVelocity(ChaosSolverMaxPushOutVelocity);
+			}
 		}
 
 		UE_LOG(LogPBDRigidsSolver, Verbose, TEXT("PBDRigidsSolver::Tick(%3.5f)"), DeltaTime);
@@ -1585,6 +1592,7 @@ namespace Chaos
 		SetJointPairIterations(InConfig.JointPairIterations);
 		SetJointPushOutPairIterations(InConfig.JointPushOutPairIterations);
 		SetCollisionCullDistance(InConfig.CollisionCullDistance);
+		SetCollisionMaxPushOutVelocity(InConfig.CollisionMaxPushOutVelocity);
 		SetGenerateCollisionData(InConfig.bGenerateCollisionData);
 		SetGenerateBreakingData(InConfig.bGenerateBreakData);
 		SetGenerateTrailingData(InConfig.bGenerateTrailingData);
diff --git a/Engine/Source/Runtime/Experimental/Chaos/Public/Chaos/Collision/SolverCollisionContainer.h b/Engine/Source/Runtime/Experimental/Chaos/Public/Chaos/Collision/SolverCollisionContainer.h
--- a/Engine/Source/Runtime/Experimental/Chaos/Public/Chaos/Collision/SolverCollisionContainer.h
+++ b/Engine/Source/Runtime/Experimental/Chaos/Public/Chaos/Collision/SolverCollisionContainer.h
@@ -30,6 +30,8 @@ namespace Chaos
 
 		int32 NumSolvers() const { return CollisionSolvers.Num(); }
 
+		void SetMaxPushOutVelocity(const FReal InMaxPushOutVelocity) { MaxPushOutVelocity = InMaxPushOutVelocity; }
+
 		virtual void Reset(const int32 InMaxCollisions) override;
 
 		// Add a solver constraint for the specified constraint and gather the required solver data
@@ -55,5 +57,6 @@ namespace Chaos
 
 		TArray<FPBDCollisionSolverAdapter> CollisionSolvers;
 		TArray<int32> SweptCollisionSolvers;
+		FReal MaxPushOutVelocity;
 	};
 }
diff --git a/Engine/Source/Runtime/Experimental/Chaos/Public/Chaos/PBDCollisionConstraints.h b/Engine/Source/Runtime/Experimental/Chaos/Public/Chaos/PBDCollisionConstraints.h
--- a/Engine/Source/Runtime/Experimental/Chaos/Public/Chaos/PBDCollisionConstraints.h
+++ b/Engine/Source/Runtime/Experimental/Chaos/Public/Chaos/PBDCollisionConstraints.h
@@ -263,6 +263,11 @@ public:
 		return GravitySize;
 	}
 
+	void SetMaxPushOutVelocity(const FReal InMaxPushOutVelocity)
+	{
+		MaxPushOutVelocity = InMaxPushOutVelocity;
+	}
+
 	int32 NumConstraints() const
 	{
 		return GetConstraints().Num();
@@ -333,6 +338,8 @@ private:
 	FVec3 GravityDirection;
 	FReal GravitySize;
 
+	FReal MaxPushOutVelocity;
+
 	EConstraintSolverType SolverType;
 };
 }
diff --git a/Engine/Source/Runtime/Experimental/Chaos/Public/Chaos/PBDRigidsEvolutionGBF.h b/Engine/Source/Runtime/Experimental/Chaos/Public/Chaos/PBDRigidsEvolutionGBF.h
--- a/Engine/Source/Runtime/Experimental/Chaos/Public/Chaos/PBDRigidsEvolutionGBF.h
+++ b/Engine/Source/Runtime/Experimental/Chaos/Public/Chaos/PBDRigidsEvolutionGBF.h
@@ -98,6 +98,7 @@ namespace Chaos
 		static constexpr FRealSingle DefaultCollisionMarginFraction = 0.1f;
 		static constexpr FRealSingle DefaultCollisionMarginMax = 100.0f;
 		static constexpr FRealSingle DefaultCollisionCullDistance = 5.0f;
+		static constexpr FRealSingle DefaultCollisionMaxPushOutVelocity = 1000.0f;
 		static constexpr int32 DefaultNumJointPairIterations = 1;
 		static constexpr int32 DefaultNumJointPushOutPairIterations = 1;
 		static constexpr int32 DefaultRestitutionThreshold = 1000;
diff --git a/Engine/Source/Runtime/Experimental/Chaos/Public/ChaosSolverConfiguration.h b/Engine/Source/Runtime/Experimental/Chaos/Public/ChaosSolverConfiguration.h
--- a/Engine/Source/Runtime/Experimental/Chaos/Public/ChaosSolverConfiguration.h
+++ b/Engine/Source/Runtime/Experimental/Chaos/Public/ChaosSolverConfiguration.h
@@ -30,6 +30,7 @@ struct FChaosSolverConfiguration
 		, CollisionMarginFraction(FEvolution::DefaultCollisionMarginFraction)
 		, CollisionMarginMax(FEvolution::DefaultCollisionMarginMax)
 		, CollisionCullDistance(FEvolution::DefaultCollisionCullDistance)
+		, CollisionMaxPushOutVelocity(FEvolution::DefaultCollisionMaxPushOutVelocity)
 		, JointPairIterations(FEvolution::DefaultNumJointPairIterations)
 		, JointPushOutPairIterations(FEvolution::DefaultNumJointPushOutPairIterations)
 		, ClusterConnectionFactor(1.0f)
@@ -75,6 +76,14 @@ struct FChaosSolverConfiguration
 	UPROPERTY(EditAnywhere, Category = "SolverConfiguration|Collision")
 	float CollisionCullDistance;
 
+	// The maximum speed at which two bodies can be extracted from each other when they start a frame inter-penetrating. This can
+	// happen because they spawned on top of each other, or the solver failed to fully reolve collisions last frame. A value of
+	// zero means "no limit". A non-zero value can be used to prevent explosive behaviour when bodies start deeply penetrating. 
+	// An alternative to using this approach is to increase the number of Velocity Iterations, which is more expensive but will 
+	// ensure the bdoies are depenetrated in a single frame without explosive behaviour.
+	UPROPERTY(EditAnywhere, Category = "SolverConfiguration|Collision")
+	float CollisionMaxPushOutVelocity;
+
 	// The number of iterations to run on each constraint during the constraint solver step
 	UPROPERTY(EditAnywhere, Category = "SolverConfiguration|Iterations")
 	int32 JointPairIterations;
diff --git a/Engine/Source/Runtime/Experimental/Chaos/Public/PBDRigidsSolver.h b/Engine/Source/Runtime/Experimental/Chaos/Public/PBDRigidsSolver.h
--- a/Engine/Source/Runtime/Experimental/Chaos/Public/PBDRigidsSolver.h
+++ b/Engine/Source/Runtime/Experimental/Chaos/Public/PBDRigidsSolver.h
@@ -190,6 +190,7 @@ namespace Chaos
 		void SetJointPairIterations(const int32 InNumIterations) { GetJointConstraints().SetNumPairIterations(InNumIterations); }
 		void SetJointPushOutPairIterations(const int32 InNumIterations) {GetJointConstraints().SetNumPushOutPairIterations(InNumIterations); }
 		void SetCollisionCullDistance(const FReal InCullDistance) { GetEvolution()->GetBroadPhase().SetCullDistance(InCullDistance); }
+		void SetCollisionMaxPushOutVelocity(const FReal InMaxPushOutVelocity) { GetEvolution()->GetCollisionConstraints().SetMaxPushOutVelocity(InMaxPushOutVelocity); }
 		void SetUseContactGraph(const bool bInUseContactGraph) { GetEvolution()->GetCollisionConstraintsRule().SetUseContactGraph(bInUseContactGraph); }
 
 		/**/
