Chaos - fix uninitialized variabl
- fixes unit tests

diff --git a/Engine/Source/Runtime/Experimental/Chaos/Private/Chaos/PBDCollisionConstraints.cpp b/Engine/Source/Runtime/Experimental/Chaos/Private/Chaos/PBDCollisionConstraints.cpp
--- a/Engine/Source/Runtime/Experimental/Chaos/Private/Chaos/PBDCollisionConstraints.cpp
+++ b/Engine/Source/Runtime/Experimental/Chaos/Private/Chaos/PBDCollisionConstraints.cpp
@@ -123,6 +123,9 @@ namespace Chaos
 		, bEnableRestitution(true)
 		, bHandlesEnabled(true)
 		, bCanDisableContacts(true)
+		, GravityDirection(FVec3(0,0,-1))
+		, GravitySize(980)
+		, MaxPushOutVelocity(0)
 		, SolverType(EConstraintSolverType::GbfPbd)
 	{
 #if INTEL_ISPC
