Fix not all acceleration substructures being reset in rebuild.

diff --git a/Engine/Source/Runtime/Experimental/Chaos/Private/Chaos/PBDRigidsEvolution.cpp b/Engine/Source/Runtime/Experimental/Chaos/Private/Chaos/PBDRigidsEvolution.cpp
--- a/Engine/Source/Runtime/Experimental/Chaos/Private/Chaos/PBDRigidsEvolution.cpp
+++ b/Engine/Source/Runtime/Experimental/Chaos/Private/Chaos/PBDRigidsEvolution.cpp
@@ -333,6 +333,13 @@ namespace Chaos
 		bool IsTimeSlicingProgressing = false;
 		bool CreatingNewTimeSlicedStructures = bNeedsReset;
 
+		if (CreatingNewTimeSlicedStructures)
+		{
+			// Cannot assume all substructures that exist are in cache!!
+			// Make sure AccelerationStructureCopy is completely reset and does not have dirty substructures that will be ignored below.
+			AccelerationStructureCopy->Reset();
+		}
+
 		for(const auto& Itr : SpatialAccelerationCache)
 		{
 			const FSpatialAccelerationIdx SpatialIdx = Itr.Key;
