Chaos - fix moving kinematic in sleeping island crash

diff --git a/Engine/Source/Runtime/Experimental/Chaos/Private/Chaos/CollisionResolution.cpp b/Engine/Source/Runtime/Experimental/Chaos/Private/Chaos/CollisionResolution.cpp
--- a/Engine/Source/Runtime/Experimental/Chaos/Private/Chaos/CollisionResolution.cpp
+++ b/Engine/Source/Runtime/Experimental/Chaos/Private/Chaos/CollisionResolution.cpp
@@ -2862,12 +2862,8 @@ namespace Chaos
 			FParticlePairCollisionConstraints* PairContacts = Context.CollisionAllocator->FindParticlePairConstraints(Particle0->Handle(), Particle1->Handle());
 			if (PairContacts != nullptr)
 			{
-				// @todo(chaos): eventually we shouldn't even be calling this function for sleeping pairs
-				const bool bIsAsleep0 = (Particle0->IsDynamic() && Particle0->Sleeping()) || (Particle0->IsKinematic() && !Particle0->KinematicTarget().IsSet()) || Particle0->IsStatic();
-				const bool bIsAsleep1 = (Particle1->IsDynamic() && Particle1->Sleeping()) || (Particle1->IsKinematic() && !Particle1->KinematicTarget().IsSet()) || Particle1->IsStatic();
-
 				// See if the particles have the same relative transform and can have their constraint fully restored
-				const bool bCanRestore = (bIsAsleep0 && bIsAsleep1) || IsWithinManifoldRestorationThreshold(PairContacts->GetParticle0Transform(), ParticleWorldTransform0, PairContacts->GetParticle1Transform(), ParticleWorldTransform1);
+				const bool bCanRestore = IsWithinManifoldRestorationThreshold(PairContacts->GetParticle0Transform(), ParticleWorldTransform0, PairContacts->GetParticle1Transform(), ParticleWorldTransform1);
 
 				if (bCanRestore)
 				{
diff --git a/Engine/Source/Runtime/Experimental/Chaos/Public/Chaos/Collision/CollisionConstraintAllocator.h b/Engine/Source/Runtime/Experimental/Chaos/Public/Chaos/Collision/CollisionConstraintAllocator.h
--- a/Engine/Source/Runtime/Experimental/Chaos/Public/Chaos/Collision/CollisionConstraintAllocator.h
+++ b/Engine/Source/Runtime/Experimental/Chaos/Public/Chaos/Collision/CollisionConstraintAllocator.h
@@ -585,14 +585,17 @@ namespace Chaos
 			for (auto& KeyValuePair : ShapePairConstraintMap)
 			{
 				FPBDCollisionConstraint* Constraint = KeyValuePair.Value;
-				
-				if (Constraint->GetContainerCookie().LastUsedEpoch < Epoch)
+
+				if (!Constraint->IsSleeping())
 				{
-					if(!Constraint->IsSleeping())
+					if (Constraint->GetContainerCookie().LastUsedEpoch < Epoch)
 					{
-						FreeConstraint(KeyValuePair.Value);
+						Constraint->GetContainerCookie().bIsInShapePairMap = false;
+
+						FreeConstraint(Constraint);
+
+						PrunedKeyList.Add(KeyValuePair.Key);
 					}
-					PrunedKeyList.Add(KeyValuePair.Key);
 				}
 			}
 
diff --git a/Engine/Source/Runtime/Experimental/Chaos/Public/Chaos/Collision/PBDCollisionConstraint.h b/Engine/Source/Runtime/Experimental/Chaos/Public/Chaos/Collision/PBDCollisionConstraint.h
--- a/Engine/Source/Runtime/Experimental/Chaos/Public/Chaos/Collision/PBDCollisionConstraint.h
+++ b/Engine/Source/Runtime/Experimental/Chaos/Public/Chaos/Collision/PBDCollisionConstraint.h
@@ -253,7 +253,6 @@ namespace Chaos
 		{
 			ConstraintIndex = InIndex;
 			LastUsedEpoch = InEpoch;
-			bIsSleeping = false;
 		}
 
 		/**
@@ -285,6 +284,8 @@ namespace Chaos
 		int32 ConstraintIndex;
 
 		// Whether the constraint is in a sleeping island (this is used to prevent culling because Epoch is not updated for sleepers)
+		// NOTE: this variable is managed by the IslandManager. It should probably not be in the Container Cookie (but currently
+		// it is used by the container to prevent pruning).
 		bool bIsSleeping;
 
 		// Whether the constraint has just been awaken in an island
