Move all interchange code files under experimental folder.

diff --git a/Engine/Source/Runtime/Experimental/Interchange/Engine/Private/Tasks/InterchangeTaskCreateAsset.cpp b/Engine/Source/Runtime/Experimental/Interchange/Engine/Private/Tasks/InterchangeTaskCreateAsset.cpp
new file mode 100644
--- /dev/null
+++ b/Engine/Source/Runtime/Experimental/Interchange/Engine/Private/Tasks/InterchangeTaskCreateAsset.cpp
@@ -0,0 +1,311 @@
+// Copyright Epic Games, Inc. All Rights Reserved.
+#include "InterchangeTaskCreateAsset.h"
+
+#include "Async/TaskGraphInterfaces.h"
+#include "CoreMinimal.h"
+#include "InterchangeFactoryBase.h"
+#include "InterchangeEngineLogPrivate.h"
+#include "InterchangeManager.h"
+#include "InterchangeResult.h"
+#include "InterchangeSourceData.h"
+#include "InterchangeTranslatorBase.h"
+#include "Misc/Paths.h"
+#include "PackageUtils/PackageUtils.h"
+#include "Stats/Stats.h"
+#include "Templates/SharedPointer.h"
+#include "UObject/Object.h"
+#include "UObject/UObjectGlobals.h"
+#include "UObject/ObjectMacros.h"
+#include "UObject/WeakObjectPtrTemplates.h"
+#include "Nodes/InterchangeBaseNode.h"
+
+namespace UE
+{
+	namespace Interchange
+	{
+		namespace Private
+		{
+			void InternalGetPackageName(const UE::Interchange::FImportAsyncHelper& AsyncHelper, const int32 SourceIndex, const FString& PackageBasePath, const UInterchangeBaseNode* Node, FString& OutPackageName, FString& OutAssetName)
+			{
+				const UInterchangeSourceData* SourceData = AsyncHelper.SourceDatas[SourceIndex];
+				check(SourceData);
+				FString NodeDisplayName = Node->GetAssetName();
+
+				// Set the asset name and the package name
+				OutAssetName = NodeDisplayName;
+				SanitizeObjectName(OutAssetName);
+
+				FString SanitazedPackageBasePath = PackageBasePath;
+				SanitizeObjectPath(SanitazedPackageBasePath);
+
+				OutPackageName = FPaths::Combine(*SanitazedPackageBasePath, *OutAssetName);
+			}
+		}//ns Private
+	}//ns Interchange
+}//ns UE
+
+void UE::Interchange::FTaskCreatePackage::DoTask(ENamedThreads::Type CurrentThread, const FGraphEventRef& MyCompletionGraphEvent)
+{
+#if INTERCHANGE_TRACE_ASYNCHRONOUS_TASK_ENABLED
+	INTERCHANGE_TRACE_ASYNCHRONOUS_TASK(CreatePackage)
+#endif
+	TSharedPtr<UE::Interchange::FImportAsyncHelper, ESPMode::ThreadSafe> AsyncHelper = WeakAsyncHelper.Pin();
+	check(AsyncHelper.IsValid());
+
+	//Verify if the task was cancel
+	if (AsyncHelper->bCancel)
+	{
+		return;
+	}
+
+	//The create package thread must always execute on the game thread
+	check(IsInGameThread());
+
+	// Create factory
+	UInterchangeFactoryBase* Factory = NewObject<UInterchangeFactoryBase>(GetTransientPackage(), FactoryClass);
+	Factory->SetResultsContainer(AsyncHelper->AssetImportResult->GetResults());
+
+	{
+		FScopeLock Lock(&AsyncHelper->CreatedFactoriesLock);
+		AsyncHelper->CreatedFactories.Add(Node->GetUniqueID(), Factory);
+	}
+
+	UPackage* Pkg = nullptr;
+	FString PackageName;
+	FString AssetName;
+	//If we do a reimport no need to create a package
+	if (AsyncHelper->TaskData.ReimportObject)
+	{
+		Private::InternalGetPackageName(*AsyncHelper, SourceIndex, PackageBasePath, Node, PackageName, AssetName);
+		UPackage* ResultFindPackage = FindPackage(nullptr, *PackageName);
+		UObject* FindOuter = (ResultFindPackage == nullptr ? ANY_PACKAGE : ResultFindPackage);
+		UObject* ExistingObject = FindObject<UObject>(FindOuter, *AssetName);
+		if (ExistingObject)
+		{
+			Pkg = ExistingObject->GetPackage();
+			PackageName = Pkg->GetPathName();
+		}
+		else
+		{
+			UInterchangeResultError_Generic* Message = Factory->AddMessage<UInterchangeResultError_Generic>();
+			Message->SourceAssetName = AsyncHelper->SourceDatas[SourceIndex]->GetFilename();
+			Message->DestinationAssetName = AssetName;
+			Message->AssetType = Node->GetObjectClass();
+			Message->Text = NSLOCTEXT("Interchange", "CannotFindPackageDuringReimport", "Cannot find an existing package.");
+
+			//Skip this asset
+			return;
+		}
+	}
+	else
+	{
+		Private::InternalGetPackageName(*AsyncHelper, SourceIndex, PackageBasePath, Node, PackageName, AssetName);
+		// We can not create assets that share the name of a map file in the same location
+		if (UE::Interchange::FPackageUtils::IsMapPackageAsset(PackageName))
+		{
+			UInterchangeResultError_Generic* Message = Factory->AddMessage<UInterchangeResultError_Generic>();
+			Message->SourceAssetName = AsyncHelper->SourceDatas[SourceIndex]->GetFilename();
+			Message->DestinationAssetName = AssetName;
+			Message->AssetType = Node->GetObjectClass();
+			Message->Text = NSLOCTEXT("Interchange", "MapExistsWithSameName", "You cannot create an asset with this name, as there is already a map file with the same name in this folder.");
+
+			//Skip this asset
+			return;
+		}
+
+		Pkg = CreatePackage(*PackageName);
+		if (Pkg == nullptr)
+		{
+			UInterchangeResultError_Generic* Message = Factory->AddMessage<UInterchangeResultError_Generic>();
+			Message->SourceAssetName = AsyncHelper->SourceDatas[SourceIndex]->GetFilename();
+			Message->DestinationAssetName = AssetName;
+			Message->AssetType = Node->GetObjectClass();
+			Message->Text = FText::Format(NSLOCTEXT("Interchange", "CouldntCreatePackage", "It was not possible to create a package named '{0}'; the asset will not be imported."), FText::FromString(PackageName));
+
+			//Skip this asset
+			return;
+		}
+
+		//Import Asset describe by the node
+		UInterchangeFactoryBase::FCreateAssetParams CreateAssetParams;
+		CreateAssetParams.AssetName = AssetName;
+		CreateAssetParams.AssetNode = Node;
+		CreateAssetParams.Parent = Pkg;
+		CreateAssetParams.SourceData = AsyncHelper->SourceDatas[SourceIndex];
+		CreateAssetParams.Translator = nullptr;
+		if (AsyncHelper->BaseNodeContainers.IsValidIndex(SourceIndex))
+		{
+			CreateAssetParams.NodeContainer = AsyncHelper->BaseNodeContainers[SourceIndex].Get();
+		}
+		CreateAssetParams.ReimportObject = AsyncHelper->TaskData.ReimportObject;
+		//Make sure the asset UObject is created with the correct type on the main thread
+		UObject* NodeAsset = Factory->CreateEmptyAsset(CreateAssetParams);
+		if (NodeAsset)
+		{
+			if (!NodeAsset->HasAnyInternalFlags(EInternalObjectFlags::Async))
+			{
+				//Since the async flag is not set we must be in the game thread
+				ensure(IsInGameThread());
+				NodeAsset->SetInternalFlags(EInternalObjectFlags::Async);
+			}
+			FScopeLock Lock(&AsyncHelper->ImportedAssetsPerSourceIndexLock);
+			TArray<UE::Interchange::FImportAsyncHelper::FImportedObjectInfo>& ImportedInfos = AsyncHelper->ImportedAssetsPerSourceIndex.FindOrAdd(SourceIndex);
+			UE::Interchange::FImportAsyncHelper::FImportedObjectInfo& AssetInfo = ImportedInfos.AddDefaulted_GetRef();
+			AssetInfo.ImportedObject = NodeAsset;
+			AssetInfo.Factory = Factory;
+			AssetInfo.FactoryNode = Node;
+			AssetInfo.bIsReimport = bool(AsyncHelper->TaskData.ReimportObject);
+			Node->ReferenceObject = FSoftObjectPath(NodeAsset);
+		}
+	}
+
+	// Make sure the destination package is loaded
+	Pkg->FullyLoad();
+	
+	{
+		FScopeLock Lock(&AsyncHelper->CreatedPackagesLock);
+		AsyncHelper->CreatedPackages.Add(PackageName, Pkg);
+	}
+}
+
+void UE::Interchange::FTaskCreateAsset::DoTask(ENamedThreads::Type CurrentThread, const FGraphEventRef& MyCompletionGraphEvent)
+{
+#if INTERCHANGE_TRACE_ASYNCHRONOUS_TASK_ENABLED
+	INTERCHANGE_TRACE_ASYNCHRONOUS_TASK(CreateAsset)
+#endif
+	TSharedPtr<UE::Interchange::FImportAsyncHelper, ESPMode::ThreadSafe> AsyncHelper = WeakAsyncHelper.Pin();
+	check(AsyncHelper.IsValid());
+
+	//Verify if the task was cancel
+	if (AsyncHelper->bCancel)
+	{
+		return;
+	}
+
+	UInterchangeFactoryBase* Factory = nullptr;
+	{
+		FScopeLock Lock(&AsyncHelper->CreatedFactoriesLock);
+		Factory = AsyncHelper->CreatedFactories.FindChecked(Node->GetUniqueID());
+	}
+
+	UPackage* Pkg = nullptr;
+	FString PackageName;
+	FString AssetName;
+	Private::InternalGetPackageName(*AsyncHelper, SourceIndex, PackageBasePath, Node, PackageName, AssetName);
+	bool bSkipAsset = false;
+	UObject* ExistingObject = nullptr;
+	if (AsyncHelper->TaskData.ReimportObject)
+	{
+		//When we re-import one particular asset, if the source file contain other assets, we want to set the node= reference UObject for those asset to the existing asset
+		//The way to discover this case is to compare the re-import asset with the node asset.
+		UPackage* ResultFindPackage = FindPackage(nullptr, *PackageName);
+		UObject* FindOuter = (ResultFindPackage == nullptr ? ANY_PACKAGE : ResultFindPackage);
+		ExistingObject = FindObject<UObject>(FindOuter, *AssetName);
+		bSkipAsset = !ExistingObject || ExistingObject != AsyncHelper->TaskData.ReimportObject;
+		if (!bSkipAsset)
+		{
+			Pkg = AsyncHelper->TaskData.ReimportObject->GetPackage();
+			PackageName = Pkg->GetPathName();
+		}
+		else if(ExistingObject)
+		{
+			Pkg = ExistingObject->GetPackage();
+			PackageName = Pkg->GetPathName();
+		}
+	}
+	else
+	{
+		FScopeLock Lock(&AsyncHelper->CreatedPackagesLock);
+		UPackage** PkgPtr = AsyncHelper->CreatedPackages.Find(PackageName);
+
+		if (!PkgPtr || !(*PkgPtr))
+		{
+			UInterchangeResultError_Generic* Message = Factory->AddMessage<UInterchangeResultError_Generic>();
+			Message->SourceAssetName = AsyncHelper->SourceDatas[SourceIndex]->GetFilename();
+			Message->DestinationAssetName = AssetName;
+			Message->AssetType = Node->GetObjectClass();
+			Message->Text = NSLOCTEXT("Interchange", "BadPackage", "It was not possible to create the asset as its package was not created correctly.");
+
+			return;
+		}
+
+		if (!AsyncHelper->SourceDatas.IsValidIndex(SourceIndex) || !AsyncHelper->Translators.IsValidIndex(SourceIndex))
+		{
+			UInterchangeResultError_Generic* Message = Factory->AddMessage<UInterchangeResultError_Generic>();
+			Message->DestinationAssetName = AssetName;
+			Message->AssetType = Node->GetObjectClass();
+			Message->Text = NSLOCTEXT("Interchange", "SourceDataOrTranslatorInvalid", "It was not possible to create the asset as its translator was not created correctly.");
+
+			return;
+		}
+
+		Pkg = *PkgPtr;
+	}
+
+	UObject* NodeAsset = nullptr;
+	if (bSkipAsset)
+	{
+		NodeAsset = ExistingObject;
+	}
+	else
+	{
+		UInterchangeTranslatorBase* Translator = AsyncHelper->Translators[SourceIndex];
+		//Import Asset describe by the node
+		UInterchangeFactoryBase::FCreateAssetParams CreateAssetParams;
+		CreateAssetParams.AssetName = AssetName;
+		CreateAssetParams.AssetNode = Node;
+		CreateAssetParams.Parent = Pkg;
+		CreateAssetParams.SourceData = AsyncHelper->SourceDatas[SourceIndex];
+		CreateAssetParams.Translator = Translator;
+		if (AsyncHelper->BaseNodeContainers.IsValidIndex(SourceIndex))
+		{
+			CreateAssetParams.NodeContainer = AsyncHelper->BaseNodeContainers[SourceIndex].Get();
+		}
+		CreateAssetParams.ReimportObject = AsyncHelper->TaskData.ReimportObject;
+
+		CreateAssetParams.ReimportStrategyFlags = EReimportStrategyFlags::ApplyNoProperties;
+		//CreateAssetParams.ReimportStrategyFlags = EReimportStrategyFlags::ApplyPipelineProperties;
+		//CreateAssetParams.ReimportStrategyFlags = EReimportStrategyFlags::ApplyEditorChangedProperties;
+
+		NodeAsset = Factory->CreateAsset(CreateAssetParams);
+	}
+	if (NodeAsset)
+	{
+		if (!bSkipAsset)
+		{
+			FScopeLock Lock(&AsyncHelper->ImportedAssetsPerSourceIndexLock);
+			TArray<UE::Interchange::FImportAsyncHelper::FImportedObjectInfo>& ImportedInfos = AsyncHelper->ImportedAssetsPerSourceIndex.FindOrAdd(SourceIndex);
+			UE::Interchange::FImportAsyncHelper::FImportedObjectInfo* AssetInfoPtr = ImportedInfos.FindByPredicate([NodeAsset](const UE::Interchange::FImportAsyncHelper::FImportedObjectInfo& CurInfo)
+			{
+				return CurInfo.ImportedObject == NodeAsset;
+			});
+
+			if (!AssetInfoPtr)
+			{
+				UE::Interchange::FImportAsyncHelper::FImportedObjectInfo& AssetInfo = ImportedInfos.AddDefaulted_GetRef();
+				AssetInfo.ImportedObject = NodeAsset;
+				AssetInfo.Factory = Factory;
+				AssetInfo.FactoryNode = Node;
+				AssetInfo.bIsReimport = bool(AsyncHelper->TaskData.ReimportObject);
+			}
+
+			// Fill in destination asset and type in any results which have been added previously by a translator or pipeline, now that we have a corresponding factory.
+			UInterchangeResultsContainer* Results = AsyncHelper->AssetImportResult->GetResults();
+			for (UInterchangeResult* Result : Results->GetResults())
+			{
+				if (!Result->InterchangeKey.IsEmpty() && (Result->DestinationAssetName.IsEmpty() || Result->AssetType == nullptr))
+				{
+					TArray<FString> TargetAssets;
+					Node->GetTargetNodeUids(TargetAssets);
+					if (TargetAssets.Contains(Result->InterchangeKey))
+					{
+						Result->DestinationAssetName = NodeAsset->GetPathName();
+						Result->AssetType = NodeAsset->GetClass();
+					}
+				}
+			}
+		}
+
+		Node->ReferenceObject = FSoftObjectPath(NodeAsset);
+	}
+}
\ No newline at end of file
diff --git a/Engine/Source/Runtime/Experimental/Interchange/Engine/Private/Tasks/InterchangeTaskCreateAsset.h b/Engine/Source/Runtime/Experimental/Interchange/Engine/Private/Tasks/InterchangeTaskCreateAsset.h
new file mode 100644
--- /dev/null
+++ b/Engine/Source/Runtime/Experimental/Interchange/Engine/Private/Tasks/InterchangeTaskCreateAsset.h
@@ -0,0 +1,106 @@
+// Copyright Epic Games, Inc. All Rights Reserved.
+
+#pragma once
+
+#include "Async/TaskGraphInterfaces.h"
+#include "CoreMinimal.h"
+#include "InterchangeFactoryBase.h"
+#include "InterchangeManager.h"
+#include "Stats/Stats.h"
+#include "UObject/WeakObjectPtrTemplates.h"
+#include "Nodes/InterchangeBaseNode.h"
+
+namespace UE
+{
+	namespace Interchange
+	{
+
+		/**
+		 * This task create package, Cook::PackageTracker::NotifyUObjectCreated is not thread safe, so we need to create the packages on the main thread
+		 */
+		class FTaskCreatePackage
+		{
+		private:
+			FString PackageBasePath;
+			int32 SourceIndex;
+			TWeakPtr<FImportAsyncHelper, ESPMode::ThreadSafe> WeakAsyncHelper;
+			UInterchangeBaseNode* Node;
+			const UClass* FactoryClass;
+
+		public:
+			FTaskCreatePackage(const FString& InPackageBasePath, const int32 InSourceIndex, TWeakPtr<FImportAsyncHelper, ESPMode::ThreadSafe> InAsyncHelper, UInterchangeBaseNode* InNode, const UClass* InFactoryClass)
+				: PackageBasePath(InPackageBasePath)
+				, SourceIndex(InSourceIndex)
+				, WeakAsyncHelper(InAsyncHelper)
+				, Node(InNode)
+				, FactoryClass(InFactoryClass)
+			{
+				check(Node);
+				check(FactoryClass);
+			}
+
+			ENamedThreads::Type GetDesiredThread() const
+			{
+				// This is no longer true now that we have to construct a factory here
+				//if (WeakAsyncHelper.IsValid() && WeakAsyncHelper.Pin()->TaskData.ReimportObject)
+				//{
+				//	//When doing a reimport the package already exist, so we can get it outside of the main thread
+				//	return ENamedThreads::AnyBackgroundThreadNormalTask;
+				//}
+				return ENamedThreads::GameThread;
+			}
+
+			static ESubsequentsMode::Type GetSubsequentsMode()
+			{
+				return ESubsequentsMode::TrackSubsequents;
+			}
+
+			TStatId GetStatId() const
+			{
+				RETURN_QUICK_DECLARE_CYCLE_STAT(FTaskCreatePackage, STATGROUP_TaskGraphTasks);
+			}
+
+			void DoTask(ENamedThreads::Type CurrentThread, const FGraphEventRef& MyCompletionGraphEvent);
+		};
+
+		class FTaskCreateAsset
+		{
+		private:
+			FString PackageBasePath;
+			int32 SourceIndex;
+			TWeakPtr<FImportAsyncHelper, ESPMode::ThreadSafe> WeakAsyncHelper;
+			UInterchangeBaseNode* Node;
+			bool bCanRunOnAnyThread;
+
+		public:
+			FTaskCreateAsset(const FString& InPackageBasePath, const int32 InSourceIndex, TWeakPtr<FImportAsyncHelper, ESPMode::ThreadSafe> InAsyncHelper, UInterchangeBaseNode* InNode, bool bInCanRunOnAnyThread)
+				: PackageBasePath(InPackageBasePath)
+				, SourceIndex(InSourceIndex)
+				, WeakAsyncHelper(InAsyncHelper)
+				, Node(InNode)
+				, bCanRunOnAnyThread(bInCanRunOnAnyThread)
+			{
+				check(Node);
+			}
+
+			ENamedThreads::Type GetDesiredThread() const
+			{
+				return bCanRunOnAnyThread ? ENamedThreads::AnyBackgroundThreadNormalTask : ENamedThreads::GameThread;
+			}
+
+			static ESubsequentsMode::Type GetSubsequentsMode()
+			{
+				return ESubsequentsMode::TrackSubsequents;
+			}
+
+			TStatId GetStatId() const
+			{
+				RETURN_QUICK_DECLARE_CYCLE_STAT(FTaskCreateAsset, STATGROUP_TaskGraphTasks);
+			}
+
+			void DoTask(ENamedThreads::Type CurrentThread, const FGraphEventRef& MyCompletionGraphEvent);
+		};
+
+
+	} //ns Interchange
+}//ns UE
