Integrate 17966523:

Fix contact modification removing disabled constraints and defereferncing null handles.

diff --git a/Engine/Source/Runtime/Experimental/Chaos/Private/Chaos/PBDCollisionConstraints.cpp b/Engine/Source/Runtime/Experimental/Chaos/Private/Chaos/PBDCollisionConstraints.cpp
--- a/Engine/Source/Runtime/Experimental/Chaos/Private/Chaos/PBDCollisionConstraints.cpp
+++ b/Engine/Source/Runtime/Experimental/Chaos/Private/Chaos/PBDCollisionConstraints.cpp
@@ -338,18 +338,13 @@ namespace Chaos
 				ModificationResults.Reserve(GetConstraints().Num());
 				for (FPBDCollisionConstraintHandle* Handle : GetConstraintHandles())
 				{
-					ModificationResults.Emplace(Handle);
-				}
-
-				Modifier->ContactModification_Internal(TArrayView<FPBDCollisionConstraintHandleModification>(ModificationResults.GetData(), ModificationResults.Num()));
-				
-				for (const FPBDCollisionConstraintHandleModification& Modification : ModificationResults)
-				{
-					if (Modification.GetResult() == ECollisionModifierResult::Disabled)
+					if (Handle)
 					{
-						RemoveConstraint(Modification.GetHandle());
+						ModificationResults.Emplace(Handle);
 					}
 				}
+
+				Modifier->ContactModification_Internal(TArrayView<FPBDCollisionConstraintHandleModification>(ModificationResults.GetData(), ModificationResults.Num()));
 			}
 		}
 	}
