add simplification option to geometry collection convex generation/tool

diff --git a/Engine/Plugins/Experimental/ChaosEditor/Source/FractureEditor/Private/FractureToolConvex.cpp b/Engine/Plugins/Experimental/ChaosEditor/Source/FractureEditor/Private/FractureToolConvex.cpp
--- a/Engine/Plugins/Experimental/ChaosEditor/Source/FractureEditor/Private/FractureToolConvex.cpp
+++ b/Engine/Plugins/Experimental/ChaosEditor/Source/FractureEditor/Private/FractureToolConvex.cpp
@@ -153,7 +153,7 @@ int32 UFractureToolConvex::ExecuteFracture(const FFractureToolContext& FractureC
 		FGeometryCollection& Collection = *FractureContext.GetGeometryCollection();
 		FGeometryCollectionProximityUtility ProximityUtility(&Collection);
 		ProximityUtility.UpdateProximity();
-		FGeometryCollectionConvexUtility::CreateNonOverlappingConvexHullData(&Collection, ConvexSettings->FractionAllowRemove);
+		FGeometryCollectionConvexUtility::CreateNonOverlappingConvexHullData(&Collection, ConvexSettings->FractionAllowRemove, ConvexSettings->SimplificationDistanceThreshold);
 	}
 
 	return INDEX_NONE;
diff --git a/Engine/Plugins/Experimental/ChaosEditor/Source/FractureEditor/Private/FractureToolConvex.h b/Engine/Plugins/Experimental/ChaosEditor/Source/FractureEditor/Private/FractureToolConvex.h
--- a/Engine/Plugins/Experimental/ChaosEditor/Source/FractureEditor/Private/FractureToolConvex.h
+++ b/Engine/Plugins/Experimental/ChaosEditor/Source/FractureEditor/Private/FractureToolConvex.h
@@ -23,6 +23,10 @@ public:
 	/** Fraction of the convex hulls for a transform that we can remove before instead using the hulls of the children */
 	UPROPERTY(EditAnywhere, Category = MapSettings, meta = (ClampMin = ".01", ClampMax = "1"))
 	double FractionAllowRemove = .5;
+
+	/** We simplify the convex shape to keep points spaced at least this far apart (except to keep the hull from collapsing to zero volume) */
+	UPROPERTY(EditAnywhere, Category = MapSettings, meta = (ClampMin = "0"))
+	double SimplificationDistanceThreshold = 0;
 };
 
 
diff --git a/Engine/Source/Runtime/Experimental/Chaos/Private/GeometryCollection/GeometryCollectionConvexUtility.cpp b/Engine/Source/Runtime/Experimental/Chaos/Private/GeometryCollection/GeometryCollectionConvexUtility.cpp
--- a/Engine/Source/Runtime/Experimental/Chaos/Private/GeometryCollection/GeometryCollectionConvexUtility.cpp
+++ b/Engine/Source/Runtime/Experimental/Chaos/Private/GeometryCollection/GeometryCollectionConvexUtility.cpp
@@ -5,6 +5,9 @@
 #include "Chaos/GJK.h"
 #include "GeometryCollection/GeometryCollection.h"
 #include "GeometryCollection/GeometryCollectionAlgo.h"
+#include "CompGeom/ConvexHull3.h"
+#include "Templates/Sorting.h"
+#include "Spatial/PointHashGrid3.h"
 
 
 FGeometryCollectionConvexUtility::FGeometryCollectionConvexData FGeometryCollectionConvexUtility::GetValidConvexHullData(FGeometryCollection* GeometryCollection)
@@ -66,6 +69,107 @@ namespace
 
 typedef Chaos::TPlaneConcrete<Chaos::FReal, 3> FChaosPlane;
 
+// filter points s.t. they are spaced at least more than SimplificationDistanceThreshold apart (after starting with the 4 'extreme' points to ensure we cover a volume)
+void FilterHullPoints(const TArray<Chaos::FVec3>& InPts, TArray<Chaos::FVec3>& OutPts, double SimplificationDistanceThreshold)
+{
+	if (SimplificationDistanceThreshold > 0)
+	{
+		OutPts.Reset();
+
+		int32 NumPts = InPts.Num();
+		TArray<Chaos::FReal> DistSq;
+		DistSq.SetNumUninitialized(NumPts);
+		TArray<int32> PointOrder;
+		PointOrder.SetNumUninitialized(NumPts);
+		UE::Geometry::TPointHashGrid3<int32, Chaos::FReal> Spatial((Chaos::FReal)SimplificationDistanceThreshold, INDEX_NONE);
+		Chaos::FAABB3 Bounds;
+		for (int32 VIdx = 0; VIdx < NumPts; VIdx++)
+		{
+			Bounds.GrowToInclude(InPts[VIdx]);
+			PointOrder[VIdx] = VIdx;
+		}
+
+		// Rank points by squared distance from center
+		Chaos::FVec3 Center = Bounds.Center();
+		for (int i = 0; i < NumPts; i++)
+		{
+			DistSq[i] = (InPts[i] - Center).SizeSquared();
+		}
+
+		// Start by picking the 'extreme' points to ensure we cover the volume reasonably (otherwise it's too easy to end up with a degenerate hull pieces)
+		UE::Geometry::TExtremePoints3<Chaos::FReal> ExtremePoints(NumPts,
+			[&InPts](int32 Idx)->UE::Math::TVector<Chaos::FReal> 
+			{ 
+				return (UE::Math::TVector<Chaos::FReal>)InPts[Idx];
+			});
+		for (int32 ExtremeIdx = 0; ExtremeIdx < ExtremePoints.Dimension + 1; ExtremeIdx++)
+		{
+			int32 ExtremePtIdx = ExtremePoints.Extreme[ExtremeIdx];
+			Chaos::FVec3 ExtremePt = InPts[ExtremePtIdx];
+			OutPts.Add(ExtremePt);
+			Spatial.InsertPointUnsafe(ExtremePtIdx, ExtremePt);
+			DistSq[ExtremePtIdx] = -1; // remove these points from the distance ranking
+		}
+
+		// Sort in descending order
+		Sort(&PointOrder[0], NumPts, [&DistSq](int32 i, int32 j)
+			{
+				return DistSq[i] > DistSq[j];
+			});
+
+		// Filter to include only w/ no other too-close points, prioritizing the farthest points
+		for (int32 OrdIdx = 0; OrdIdx < NumPts; OrdIdx++)
+		{
+			int32 PtIdx = PointOrder[OrdIdx];
+			if (DistSq[PtIdx] < 0)
+			{
+				break; // we've reached the extreme points that were already selected
+			}
+			Chaos::FVec3 Pt = InPts[PtIdx];
+			TPair<int32, Chaos::FReal> NearestIdx = Spatial.FindNearestInRadius(Pt, (Chaos::FReal)SimplificationDistanceThreshold,
+				[&InPts, &Pt](int32 Idx)
+				{
+					return (Pt - InPts[Idx]).SizeSquared();
+				});
+			if (NearestIdx.Key == INDEX_NONE)
+			{
+				Spatial.InsertPointUnsafe(PtIdx, Pt);
+				OutPts.Add(Pt);
+			}
+		}
+	}
+	else
+	{
+		// No filtering requested -- in practice, we don't call this function in these cases below
+		OutPts = InPts;
+	}
+}
+
+void FilterHullPoints(TArray<Chaos::FVec3>& InOutPts, double SimplificationDistanceThreshold)
+{
+	if (SimplificationDistanceThreshold > 0)
+	{
+		TArray<Chaos::FVec3> FilteredPts;
+		FilterHullPoints(InOutPts, FilteredPts, SimplificationDistanceThreshold);
+		InOutPts = MoveTemp(FilteredPts);
+	}
+}
+
+Chaos::FConvex MakeHull(const TArray<Chaos::FVec3>& Pts, double SimplificationDistanceThreshold)
+{
+	if (SimplificationDistanceThreshold > 0)
+	{
+		TArray<Chaos::FVec3> FilteredPts;
+		FilterHullPoints(Pts, FilteredPts, SimplificationDistanceThreshold);
+		
+		return Chaos::FConvex(FilteredPts, KINDA_SMALL_NUMBER);
+	}
+	else
+	{
+		return Chaos::FConvex(Pts, KINDA_SMALL_NUMBER);
+	}
+}
+
 /// Cut hull with plane, generating the point set of a new hull
 /// @return false if plane did not cut any points on the hull
 bool CutHull(const Chaos::FConvex& HullIn, FChaosPlane Plane, bool KeepSide, TArray<Chaos::FVec3>& HullPtsOut)
@@ -179,7 +283,8 @@ void CreateNonoverlappingConvexHulls(
 	const TManagedArray<int32>& Parents,
 	const TManagedArray<TSet<int32>>* GeoProximity,
 	const TManagedArray<int32>& GeometryToTransformIndex,
-	double FracAllowRemove
+	double FracAllowRemove,
+	double SimplificationDistanceThreshold
 )
 {
 	int32 NumBones = TransformToConvexIndices.Num();
@@ -350,7 +455,7 @@ void CreateNonoverlappingConvexHulls(
 		return true;
 	};
 
-	auto FixCollisionWithCut = [&Convexes, &FindCutPlane](int32 ConvexA, int32 ConvexB)
+	auto FixCollisionWithCut = [&Convexes, &FindCutPlane, &SimplificationDistanceThreshold](int32 ConvexA, int32 ConvexB)
 	{
 		if (!ensure(Convexes[ConvexA]->NumVertices() > 0 && Convexes[ConvexB]->NumVertices() > 0))
 		{
@@ -374,12 +479,12 @@ void CreateNonoverlappingConvexHulls(
 			TArray<Chaos::FVec3> CutHullPts;
 			if (CutHull(*Convexes[ConvexA], CutPlane, true, CutHullPts))
 			{
-				*Convexes[ConvexA] = Chaos::FConvex(CutHullPts, KINDA_SMALL_NUMBER);
+				*Convexes[ConvexA] = MakeHull(CutHullPts, SimplificationDistanceThreshold);
 			}
 			CutHullPts.Reset();
 			if (CutHull(*Convexes[ConvexB], CutPlane, false, CutHullPts))
 			{
-				*Convexes[ConvexB] = Chaos::FConvex(CutHullPts, KINDA_SMALL_NUMBER);
+				*Convexes[ConvexB] = MakeHull(CutHullPts, SimplificationDistanceThreshold);
 			}
 		}
 		return bCollide;
@@ -489,6 +594,7 @@ void CreateNonoverlappingConvexHulls(
 					}
 					if (JoinedHullPts.Num() > 0)
 					{
+						FilterHullPoints(JoinedHullPts, SimplificationDistanceThreshold);
 						int32 ConvexIdx = Convexes.Add(MakeUnique<Chaos::FConvex>(JoinedHullPts, KINDA_SMALL_NUMBER));
 						TransformToConvexIndices[Bone].Add(ConvexIdx);
 						ToProcess.Add(Bone);
@@ -565,14 +671,14 @@ void CreateNonoverlappingConvexHulls(
 				int32 ConvexIdx = OnlyConvex(Bone);
 				if (ConvexIdx)
 				{
-					*Convexes[ConvexIdx] = Chaos::FConvex(JoinedHullPts, KINDA_SMALL_NUMBER);
+					*Convexes[ConvexIdx] = MakeHull(JoinedHullPts, SimplificationDistanceThreshold);
 					NonLeafVolumes[ConvexIdx] = Convexes[ConvexIdx]->GetVolume();
 				}
 			}
 		}
 	}
 
-	auto CutIfOk = [&Convexes, &OnlyConvex, &NonLeafVolumes, &FindCutPlane, &FracAllowRemove](bool bOneSidedCut, int32 ConvexA, int32 ConvexB) -> bool
+	auto CutIfOk = [&Convexes, &OnlyConvex, &NonLeafVolumes, &FindCutPlane, &FracAllowRemove, &SimplificationDistanceThreshold](bool bOneSidedCut, int32 ConvexA, int32 ConvexB) -> bool
 	{
 		Chaos::FReal Depth;
 		Chaos::FVec3 CloseA, CloseB, Normal;
@@ -595,7 +701,7 @@ void CreateNonoverlappingConvexHulls(
 				{
 					return false;
 				}
-				CutHullA = Chaos::FConvex(CutHullPts, KINDA_SMALL_NUMBER);
+				CutHullA = MakeHull(CutHullPts, SimplificationDistanceThreshold);
 				bCreatedA = true;
 			}
 			if (!bOneSidedCut)
@@ -603,7 +709,7 @@ void CreateNonoverlappingConvexHulls(
 				CutHullPts.Reset();
 				if (CutHull(*Convexes[ConvexB], CutPlane, false, CutHullPts))
 				{
-					CutHullB = Chaos::FConvex(CutHullPts, KINDA_SMALL_NUMBER);
+					CutHullB = MakeHull(CutHullPts, SimplificationDistanceThreshold);
 					bCreatedB = true;
 				}
 			}
@@ -718,7 +824,8 @@ void HullsFromGeometry(
 	TArray<TUniquePtr<Chaos::FConvex>>& Convexes,
 	TArray<TSet<int32>>& TransformToConvexIndices,
 	const TManagedArray<int32>& SimulationType,
-	int32 RigidType
+	int32 RigidType,
+	double SimplificationDistanceThreshold
 )
 {
 	TArray<FVector> GlobalVertices;
@@ -750,12 +857,13 @@ void HullsFromGeometry(
 				HullPts.Add(GlobalVertices[VIdx]);
 			}
 			ensure(HullPts.Num() > 0);
+			FilterHullPoints(HullPts, SimplificationDistanceThreshold);
 			int32 ConvexIdx = Convexes.Add(MakeUnique<Chaos::FConvex>(HullPts, KINDA_SMALL_NUMBER));
 			if (Convexes[ConvexIdx]->NumVertices() == 0 && HullPts.Num() > 0)
 			{
 				// if we've failed to make a convex hull, add a tiny bounding box just to ensure every geometry has a hull
 				Chaos::FAABB3 AABB = Convexes[ConvexIdx]->BoundingBox();
-				AABB.Thicken((Chaos::FReal).0001);
+				AABB.Thicken((Chaos::FReal).001);
 				Chaos::FVec3 Min = AABB.Min();
 				Chaos::FVec3 Max = AABB.Max();
 				HullPts.Add(Min);
@@ -766,6 +874,7 @@ void HullsFromGeometry(
 				HullPts.Add(Chaos::FVec3(Max.X, Max.Y, Min.Z));
 				HullPts.Add(Chaos::FVec3(Max.X, Min.Y, Min.Z));
 				HullPts.Add(Chaos::FVec3(Min.X, Max.Y, Min.Z));
+				// note: Do not use SimplificationDistanceThreshold for this fixed tiny hull
 				*Convexes[ConvexIdx] = Chaos::FConvex(HullPts, KINDA_SMALL_NUMBER);
 			}
 			ensure(Convexes[ConvexIdx]->NumVertices() > 0);
@@ -791,6 +900,7 @@ void TransformHullsToLocal(
 			{
 				HullPts.Add(Transform.InverseTransformPosition(P));
 			}
+			// Do not simplify hulls when we're just trying to transform them
 			*Convexes[ConvexIdx] = Chaos::FConvex(HullPts, KINDA_SMALL_NUMBER);
 		}
 	}
@@ -800,7 +910,7 @@ void TransformHullsToLocal(
 }
 
 
-FGeometryCollectionConvexUtility::FGeometryCollectionConvexData FGeometryCollectionConvexUtility::CreateNonOverlappingConvexHullData(FGeometryCollection* GeometryCollection, double FracAllowRemove)
+FGeometryCollectionConvexUtility::FGeometryCollectionConvexData FGeometryCollectionConvexUtility::CreateNonOverlappingConvexHullData(FGeometryCollection* GeometryCollection, double FracAllowRemove, double SimplificationDistanceThreshold)
 {
 	check(GeometryCollection);
 
@@ -810,9 +920,9 @@ FGeometryCollectionConvexUtility::FGeometryCollectionConvexData FGeometryCollect
 	TArray<TSet<int32>> TransformToConvexIndexArr;
 	TArray<FTransform> GlobalTransformArray;
 	GeometryCollectionAlgo::GlobalMatrices(GeometryCollection->Transform, GeometryCollection->Parent, GlobalTransformArray);
-	HullsFromGeometry(*GeometryCollection, GlobalTransformArray, Convexes, TransformToConvexIndexArr, GeometryCollection->SimulationType, FGeometryCollection::ESimulationTypes::FST_Rigid);
+	HullsFromGeometry(*GeometryCollection, GlobalTransformArray, Convexes, TransformToConvexIndexArr, GeometryCollection->SimulationType, FGeometryCollection::ESimulationTypes::FST_Rigid, SimplificationDistanceThreshold);
 
-	CreateNonoverlappingConvexHulls(Convexes, TransformToConvexIndexArr, GeometryCollection->SimulationType, FGeometryCollection::ESimulationTypes::FST_Rigid, FGeometryCollection::ESimulationTypes::FST_None, GeometryCollection->Parent, GCProximity, GeometryCollection->TransformIndex, FracAllowRemove);
+	CreateNonoverlappingConvexHulls(Convexes, TransformToConvexIndexArr, GeometryCollection->SimulationType, FGeometryCollection::ESimulationTypes::FST_Rigid, FGeometryCollection::ESimulationTypes::FST_None, GeometryCollection->Parent, GCProximity, GeometryCollection->TransformIndex, FracAllowRemove, SimplificationDistanceThreshold);
 
 	TransformHullsToLocal(GlobalTransformArray, Convexes, TransformToConvexIndexArr);
 
diff --git a/Engine/Source/Runtime/Experimental/Chaos/Public/GeometryCollection/GeometryCollectionConvexUtility.h b/Engine/Source/Runtime/Experimental/Chaos/Public/GeometryCollection/GeometryCollectionConvexUtility.h
--- a/Engine/Source/Runtime/Experimental/Chaos/Public/GeometryCollection/GeometryCollectionConvexUtility.h
+++ b/Engine/Source/Runtime/Experimental/Chaos/Public/GeometryCollection/GeometryCollectionConvexUtility.h
@@ -25,8 +25,14 @@ public:
 	/** Ensure that convex hull data exists for the Geometry Collection and construct it if not (or if some data is missing. */
 	static FGeometryCollectionConvexData GetValidConvexHullData(FGeometryCollection* GeometryCollection);
 
-	/** Create non-overlapping convex hull data for all transforms in the geometry collection (except transforms where it would be better to just use the hulls of the children) */
-	static FGeometryCollectionConvexData CreateNonOverlappingConvexHullData(FGeometryCollection* GeometryCollection, double FractionAllowRemove = .3);
+	/**
+	 Create non-overlapping convex hull data for all transforms in the geometry collection (except transforms where it would be better to just use the hulls of the children) 
+
+	 @param GeometryCollection					The collection to add convex hulls to
+	 @param FractionAllowRemove					The fraction of a convex body we can cut away to remove overlaps with neighbors, before we fall back to using the hulls of the children directly.  (Does not affect leaves of hierarchy)
+	 @param SimplificationDistanceThreshold		Approximate minimum distance between vertices, below which we remove vertices to generate a simpler convex shape.  If 0.0, no simplification will occur.
+	 */
+	static FGeometryCollectionConvexData CreateNonOverlappingConvexHullData(FGeometryCollection* GeometryCollection, double FractionAllowRemove = .3, double SimplificationDistanceThreshold = 0.0);
 
 	/** Returns the convex hull of the vertices contained in the specified geometry. */
 	static TUniquePtr<Chaos::FConvex> FindConvexHull(const FGeometryCollection* GeometryCollection, int32 GeometryIndex);
diff --git a/Engine/Source/Runtime/GeometryCore/Private/CompGeom/ConvexHull3.cpp b/Engine/Source/Runtime/GeometryCore/Private/CompGeom/ConvexHull3.cpp
--- a/Engine/Source/Runtime/GeometryCore/Private/CompGeom/ConvexHull3.cpp
+++ b/Engine/Source/Runtime/GeometryCore/Private/CompGeom/ConvexHull3.cpp
@@ -1,7 +1,9 @@
 // Copyright Epic Games, Inc. All Rights Reserved.
 
 // Adaptation/Port of GTEngine's ConvexHull3 algorithm;
-// ref: Engine\Plugins\Experimental\GeometryProcessing\Source\GeometryAlgorithms\Private\ThirdParty\GTEngine\Mathematics\GteConvexHull3.h
+// ref: Engine\Plugins\Runtime\GeometryProcessing\Source\GeometryAlgorithms\Private\ThirdParty\GTEngine\Mathematics\GteConvexHull3.h
+// ExtremePoints::Init adapted from GteVector3.h's IntrinsicsVector3
+// ref: Engine\Plugins\Runtime\GeometryProcessing\Source\GeometryAlgorithms\Private\ThirdParty\GTEngine\Mathematics\GteVector3.h
 
 #include "CompGeom/ConvexHull3.h"
 #include "CompGeom/ExactPredicates.h"
@@ -13,166 +15,148 @@ namespace UE
 namespace Geometry
 {
 
-/**
- * Helper class to find the dimensions spanned by a point cloud
- * and (if it spans 3 dimensions) the indices of four 'extreme' points
- * forming a (non-degenerate, volume > 0) tetrahedron
- *
- * The extreme points are chosen to be far apart, and are used as the starting point for
- * incremental convex hull construction.
- * Construction method is adapted from GteVector3.h's IntrinsicsVector3
- */
+
 template<typename RealType>
-struct TExtremePoints3
+void TExtremePoints3<RealType>::Init(int32 NumPoints, TFunctionRef<TVector<RealType>(int32)> GetPointFunc, TFunctionRef<bool(int32)> FilterFunc, double Epsilon)
 {
-	int Dimension = 0;
-	int Extreme[4]{ 0, 0, 0, 0 };
-
-	// Coordinate frame spanned by input points
-	TVector<RealType> Origin{ 0,0,0 };
-	TVector<RealType> Basis[3]{ {0,0,0}, {0,0,0}, {0,0,0} };
-
-	TExtremePoints3(int32 NumPoints, TFunctionRef<TVector<RealType>(int32)> GetPointFunc, TFunctionRef<bool(int32)> FilterFunc, double Epsilon = 0)
+	TVector<RealType> FirstPoint;
+	int FirstPtIdx = -1;
+	for (FirstPtIdx = 0; FirstPtIdx < NumPoints; FirstPtIdx++)
 	{
-		TVector<RealType> FirstPoint;
-		int FirstPtIdx = -1;
-		for (FirstPtIdx = 0; FirstPtIdx < NumPoints; FirstPtIdx++)
+		if (FilterFunc(FirstPtIdx))
 		{
-			if (FilterFunc(FirstPtIdx))
-			{
-				FirstPoint = GetPointFunc(FirstPtIdx);
-				break;
-			}
+			FirstPoint = GetPointFunc(FirstPtIdx);
+			break;
 		}
-		if (FirstPtIdx == -1)
+	}
+	if (FirstPtIdx == -1)
+	{
+		// no points passed filter
+		Dimension = 0;
+		return;
+	}
+
+	TVector<RealType> Min = GetPointFunc(FirstPtIdx), Max = GetPointFunc(FirstPtIdx);
+	FIndex3i IndexMin(FirstPtIdx, FirstPtIdx, FirstPtIdx), IndexMax(FirstPtIdx, FirstPtIdx, FirstPtIdx);
+	for (int Idx = FirstPtIdx + 1; Idx < NumPoints; Idx++)
+	{
+		if (!FilterFunc(Idx))
 		{
-			// no points passed filter
-			Dimension = 0;
-			return;
+			continue;
 		}
-
-		TVector<RealType> Min = GetPointFunc(FirstPtIdx), Max = GetPointFunc(FirstPtIdx);
-		FIndex3i IndexMin(FirstPtIdx, FirstPtIdx, FirstPtIdx), IndexMax(FirstPtIdx, FirstPtIdx, FirstPtIdx);
-		for (int Idx = FirstPtIdx + 1; Idx < NumPoints; Idx++)
+		for (int Dim = 0; Dim < 3; Dim++)
 		{
-			if (!FilterFunc(Idx))
+			RealType Val = GetPointFunc(Idx)[Dim];
+			if (Val < Min[Dim])
 			{
-				continue;
+				Min[Dim] = Val;
+				IndexMin[Dim] = Idx;
 			}
-			for (int Dim = 0; Dim < 3; Dim++)
+			else if (Val > Max[Dim])
 			{
-				RealType Val = GetPointFunc(Idx)[Dim];
-				if (Val < Min[Dim])
-				{
-					Min[Dim] = Val;
-					IndexMin[Dim] = Idx;
-				}
-				else if (Val > Max[Dim])
-				{
-					Max[Dim] = Val;
-					IndexMax[Dim] = Idx;
-				}
+				Max[Dim] = Val;
+				IndexMax[Dim] = Idx;
 			}
 		}
+	}
 
-		RealType MaxRange = Max[0] - Min[0];
-		int MaxRangeDim = 0;
-		for (int Dim = 1; Dim < 3; Dim++)
+	RealType MaxRange = Max[0] - Min[0];
+	int MaxRangeDim = 0;
+	for (int Dim = 1; Dim < 3; Dim++)
+	{
+		RealType Range = Max[Dim] - Min[Dim];
+		if (Range > MaxRange)
 		{
-			RealType Range = Max[Dim] - Min[Dim];
-			if (Range > MaxRange)
-			{
-				MaxRange = Range;
-				MaxRangeDim = Dim;
-			}
+			MaxRange = Range;
+			MaxRangeDim = Dim;
 		}
-		Extreme[0] = IndexMin[MaxRangeDim];
-		Extreme[1] = IndexMax[MaxRangeDim];
+	}
+	Extreme[0] = IndexMin[MaxRangeDim];
+	Extreme[1] = IndexMax[MaxRangeDim];
 
-		// all points within box of Epsilon extent; Dimension must be 0
-		if (MaxRange <= Epsilon)
-		{
-			Dimension = 0;
-			Extreme[3] = Extreme[2] = Extreme[1] = Extreme[0];
-			return;
-		}
+	// all points within box of Epsilon extent; Dimension must be 0
+	if (MaxRange <= Epsilon)
+	{
+		Dimension = 0;
+		Extreme[3] = Extreme[2] = Extreme[1] = Extreme[0];
+		return;
+	}
 
-		Origin = GetPointFunc(Extreme[0]);
-		Basis[0] = GetPointFunc(Extreme[1]) - Origin;
-		Normalize(Basis[0]);
+	Origin = GetPointFunc(Extreme[0]);
+	Basis[0] = GetPointFunc(Extreme[1]) - Origin;
+	Normalize(Basis[0]);
 
-		// find point furthest from the line formed by the first two extreme points
+	// find point furthest from the line formed by the first two extreme points
+	{
+		TLine3<RealType> Basis0Line(Origin, Basis[0]);
+		RealType MaxDistSq = 0;
+		for (int Idx = FirstPtIdx; Idx < NumPoints; Idx++)
 		{
-			TLine3<RealType> Basis0Line(Origin, Basis[0]);
-			RealType MaxDistSq = 0;
-			for (int Idx = FirstPtIdx; Idx < NumPoints; Idx++)
+			if (!FilterFunc(Idx))
 			{
-				if (!FilterFunc(Idx))
-				{
-					continue;
-				}
-				RealType DistSq = Basis0Line.DistanceSquared(GetPointFunc(Idx));
-				if (DistSq > MaxDistSq)
-				{
-					MaxDistSq = DistSq;
-					Extreme[2] = Idx;
-				}
+				continue;
 			}
-
-			// Nearly collinear points
-			if (TMathUtil<RealType>::Sqrt(MaxDistSq) <= Epsilon * MaxRange)
+			RealType DistSq = Basis0Line.DistanceSquared(GetPointFunc(Idx));
+			if (DistSq > MaxDistSq)
 			{
-				Dimension = 1;
-				Extreme[3] = Extreme[2] = Extreme[1];
-				return;
+				MaxDistSq = DistSq;
+				Extreme[2] = Idx;
 			}
 		}
 
+		// Nearly collinear points
+		if (TMathUtil<RealType>::Sqrt(MaxDistSq) <= Epsilon * MaxRange)
+		{
+			Dimension = 1;
+			Extreme[3] = Extreme[2] = Extreme[1];
+			return;
+		}
+	}
+
 
-		Basis[1] = GetPointFunc(Extreme[2]) - Origin;
-		// project Basis[1] to be orthogonal to Basis[0]
-		Basis[1] -= (Basis[0].Dot(Basis[1])) * Basis[0];
-		Normalize(Basis[1]);
-		Basis[2] = Basis[0].Cross(Basis[1]);
+	Basis[1] = GetPointFunc(Extreme[2]) - Origin;
+	// project Basis[1] to be orthogonal to Basis[0]
+	Basis[1] -= (Basis[0].Dot(Basis[1])) * Basis[0];
+	Normalize(Basis[1]);
+	Basis[2] = Basis[0].Cross(Basis[1]);
 
+	{
+		TPlane3<RealType> Plane(Basis[2], Origin);
+		RealType MaxDist = 0, MaxSign = 0;
+		for (int Idx = FirstPtIdx; Idx < NumPoints; Idx++)
 		{
-			TPlane3<RealType> Plane(Basis[2], Origin);
-			RealType MaxDist = 0, MaxSign = 0;
-			for (int Idx = FirstPtIdx; Idx < NumPoints; Idx++)
+			if (!FilterFunc(Idx))
 			{
-				if (!FilterFunc(Idx))
-				{
-					continue;
-				}
-				RealType DistSigned = Plane.DistanceTo(GetPointFunc(Idx));
-				RealType Dist = TMathUtil<RealType>::Abs(DistSigned);
-				if (Dist > MaxDist)
-				{
-					MaxDist = Dist;
-					MaxSign = TMathUtil<RealType>::Sign(DistSigned);
-					Extreme[3] = Idx;
-				}
+				continue;
 			}
-
-			// Nearly coplanar points
-			if (MaxDist <= Epsilon * MaxRange)
+			RealType DistSigned = Plane.DistanceTo(GetPointFunc(Idx));
+			RealType Dist = TMathUtil<RealType>::Abs(DistSigned);
+			if (Dist > MaxDist)
 			{
-				Dimension = 2;
-				Extreme[3] = Extreme[2];
-				return;
+				MaxDist = Dist;
+				MaxSign = TMathUtil<RealType>::Sign(DistSigned);
+				Extreme[3] = Idx;
 			}
+		}
 
-			// make sure the tetrahedron is CW-oriented
-			if (MaxSign > 0)
-			{
-				Swap(Extreme[3], Extreme[2]);
-			}
+		// Nearly coplanar points
+		if (MaxDist <= Epsilon * MaxRange)
+		{
+			Dimension = 2;
+			Extreme[3] = Extreme[2];
+			return;
 		}
 
-		Dimension = 3;
+		// make sure the tetrahedron is CW-oriented
+		if (MaxSign > 0)
+		{
+			Swap(Extreme[3], Extreme[2]);
+		}
 	}
 
-};
+	Dimension = 3;
+}
+
 
 template<typename RealType>
 struct FHullConnectivity
diff --git a/Engine/Source/Runtime/GeometryCore/Public/CompGeom/ConvexHull3.h b/Engine/Source/Runtime/GeometryCore/Public/CompGeom/ConvexHull3.h
--- a/Engine/Source/Runtime/GeometryCore/Public/CompGeom/ConvexHull3.h
+++ b/Engine/Source/Runtime/GeometryCore/Public/CompGeom/ConvexHull3.h
@@ -17,6 +17,41 @@ using namespace UE::Math;
 
 template <typename RealType> struct TConvexHull3Internal;
 
+/**
+ * Helper class to find the dimensions spanned by a point cloud
+ * and (if it spans 3 dimensions) the indices of four 'extreme' points
+ * forming a (non-degenerate, volume > 0) tetrahedron
+ *
+ * The extreme points are chosen to be far apart, and are used as the starting point for
+ * incremental convex hull construction.
+ */
+template<typename RealType>
+struct GEOMETRYCORE_API TExtremePoints3
+{
+	int Dimension = 0;
+	int Extreme[4]{ 0, 0, 0, 0 };
+
+	// Coordinate frame spanned by input points
+	TVector<RealType> Origin{ 0,0,0 };
+	TVector<RealType> Basis[3]{ {0,0,0}, {0,0,0}, {0,0,0} };
+
+	TExtremePoints3(int32 NumPoints, TFunctionRef<TVector<RealType>(int32)> GetPointFunc, TFunctionRef<bool(int32)> FilterFunc = [](int32 Idx) {return true; }, double Epsilon = 0)
+	{
+		Init(NumPoints, GetPointFunc, FilterFunc, Epsilon);
+	}
+
+	TExtremePoints3(TArrayView<TVector<RealType>> Points, TFunctionRef<bool(int32)> FilterFunc = [](int32 Idx) {return true; })
+	{
+		Init(Points.Num(), [&Points](int32 Idx)
+			{
+				return Points[Idx];
+			}, FilterFunc);
+	}
+
+private:
+	void Init(int32 NumPoints, TFunctionRef<TVector<RealType>(int32)> GetPointFunc, TFunctionRef<bool(int32)> FilterFunc, double Epsilon = 0);
+};
+
 /**
  * Calculate the Convex Hull of a 3D point set as a Triangle Mesh
  */
@@ -186,6 +221,8 @@ protected:
 	
 };
 
+typedef TExtremePoints3<float> FExtremePoints3f;
+typedef TExtremePoints3<double> FExtremePoints3d;
 typedef TConvexHull3<float> FConvexHull3f;
 typedef TConvexHull3<double> FConvexHull3d;
 
