Fix last Mac symbol visibility issues

Thanks to [REDACTED_PERSON]

diff --git a/Engine/Source/Runtime/Experimental/Chaos/Public/Chaos/BoundingVolume.h b/Engine/Source/Runtime/Experimental/Chaos/Public/Chaos/BoundingVolume.h
--- a/Engine/Source/Runtime/Experimental/Chaos/Public/Chaos/BoundingVolume.h
+++ b/Engine/Source/Runtime/Experimental/Chaos/Public/Chaos/BoundingVolume.h
@@ -1324,5 +1324,12 @@ FArchive& operator<<(FChaosArchive& Ar, TBoundingVolume<TPayloadType, T, d>& Bou
 	return Ar;
 }
 
-}
+#if PLATFORM_MAC || PLATFORM_LINUX
+extern template class CHAOS_API Chaos::TBoundingVolume<int32, Chaos::FReal, 3>;
+extern template class CHAOS_API Chaos::TBoundingVolume<Chaos::FAccelerationStructureHandle, Chaos::FReal, 3>;
+#else
+extern template class TBoundingVolume<int32, FReal, 3>;
+extern template class TBoundingVolume<FAccelerationStructureHandle, FReal, 3>;
+#endif
 
+}
\ No newline at end of file
