Fix Mac symbol visibility

diff --git a/Engine/Source/Runtime/Experimental/Chaos/Private/Chaos/UniformGrid.cpp b/Engine/Source/Runtime/Experimental/Chaos/Private/Chaos/UniformGrid.cpp
--- a/Engine/Source/Runtime/Experimental/Chaos/Private/Chaos/UniformGrid.cpp
+++ b/Engine/Source/Runtime/Experimental/Chaos/Private/Chaos/UniformGrid.cpp
@@ -257,9 +257,15 @@ bool Chaos::TUniformGrid<T, 3>::IsValid(const TVector<int32, 3>& X) const
 	return X == ClampIndex(X);
 }
 
+#if PLATFORM_MAC || PLATFORM_LINUX
+template class CHAOS_API Chaos::TUniformGridBase<Chaos::FReal, 3>;
+template class CHAOS_API Chaos::TUniformGrid<Chaos::FReal, 3>;
+template class CHAOS_API Chaos::TUniformGrid<Chaos::FReal, 2>;
+#else
 template class Chaos::TUniformGridBase<Chaos::FReal, 3>;
 template class Chaos::TUniformGrid<Chaos::FReal, 3>;
 template class Chaos::TUniformGrid<Chaos::FReal, 2>;
+#endif
 
 template FVec3 Chaos::TUniformGridBase<Chaos::FReal, 3>::LinearlyInterpolate<FVec3>(const TArrayND<FVec3, 3>&, const FVec3&) const;
-template Chaos::FReal Chaos::TUniformGridBase<Chaos::FReal, 3>::LinearlyInterpolate<Chaos::FReal>(const TArrayND<Chaos::FReal, 3>&, const FVec3&) const;
+template CHAOS_API Chaos::FReal Chaos::TUniformGridBase<Chaos::FReal, 3>::LinearlyInterpolate<Chaos::FReal>(const TArrayND<Chaos::FReal, 3>&, const FVec3&) const;
diff --git a/Engine/Source/Runtime/Experimental/Chaos/Public/Chaos/UniformGrid.h b/Engine/Source/Runtime/Experimental/Chaos/Public/Chaos/UniformGrid.h
--- a/Engine/Source/Runtime/Experimental/Chaos/Public/Chaos/UniformGrid.h
+++ b/Engine/Source/Runtime/Experimental/Chaos/Public/Chaos/UniformGrid.h
@@ -248,4 +248,10 @@ FArchive& operator<<(FArchive& Ar, TUniformGridBase<T, d>& Value)
 	return Ar;
 }
 
+#if PLATFORM_MAC || PLATFORM_LINUX
+extern template class CHAOS_API Chaos::TUniformGridBase<Chaos::FReal, 3>;
+extern template class CHAOS_API Chaos::TUniformGrid<Chaos::FReal, 3>;
+extern template class CHAOS_API Chaos::TUniformGrid<Chaos::FReal, 2>;
+#endif
+
 }
