add a serialize-upgrade step to geometry collection to remove fully-invisible geometries and invalid exemplars; this fixes artifacts common in older fractures -- especially ones that were re-clustered.

#rnx

diff --git a/Engine/Source/Runtime/Experimental/Chaos/Private/GeometryCollection/GeometryCollection.cpp b/Engine/Source/Runtime/Experimental/Chaos/Private/GeometryCollection/GeometryCollection.cpp
--- a/Engine/Source/Runtime/Experimental/Chaos/Private/GeometryCollection/GeometryCollection.cpp
+++ b/Engine/Source/Runtime/Experimental/Chaos/Private/GeometryCollection/GeometryCollection.cpp
@@ -1071,6 +1071,85 @@ void FGeometryCollection::Serialize(Chaos::FChaosArchive& Ar)
 			// Structure is conditioned, now considered up to date.
 			Version = 8;
 		}
+
+		// Version 9 fixed fully-invisible geometry and invalid exemplars left in the hierarchy (artifacts from old fracture + the Version 5 change)
+		if (Version < 9)
+		{
+			auto HasVisibleFaces = [this](int32 TransformGroupIndex) -> bool
+			{
+				int32 GeometryIndex = TransformToGeometryIndex[TransformGroupIndex];
+				if (GeometryIndex == INDEX_NONE)
+				{
+					return false;
+				}
+				int32 Start = FaceStart[GeometryIndex], Count = FaceCount[GeometryIndex];
+				for (int32 FaceIndex = Start; FaceIndex < Start + Count; FaceIndex++)
+				{
+					if (Visible[FaceIndex])
+					{
+						return true;
+					}
+				}
+				return false;
+			};
+		
+			TArray<int32> InvalidTransforms, InvalidGeometry;
+			TArray<bool> RigidChildren; RigidChildren.Init(false, NumElements(FTransformCollection::TransformGroup));
+			const TArray<int32> RecursiveOrder = GeometryCollectionAlgo::ComputeRecursiveOrder(*this);
+			for (const int32 TransformGroupIndex : RecursiveOrder)
+			{
+				bool bHasExemplar = ExemplarIndex[TransformGroupIndex] > INDEX_NONE;
+				bool bHasGeometry = HasVisibleFaces(TransformGroupIndex);
+				if (!bHasGeometry && TransformToGeometryIndex[TransformGroupIndex] != INDEX_NONE)
+				{
+					InvalidGeometry.Add(TransformToGeometryIndex[TransformGroupIndex]);
+				}
+				bool bHasRigidChildren = RigidChildren[TransformGroupIndex];
+				bool bKeep = true;
+				if (SimulationType[TransformGroupIndex] == ESimulationTypes::FST_None && !bHasExemplar) // handle exemplars with no exemplar (remove or convert to rigid)
+				{
+					if (bHasGeometry)
+					{
+						SimulationType[TransformGroupIndex] = ESimulationTypes::FST_Rigid;
+					}
+					else
+					{
+						InvalidTransforms.Add(TransformGroupIndex);
+						bKeep = false;
+					}
+				}
+				else if (SimulationType[TransformGroupIndex] == ESimulationTypes::FST_Rigid && !bHasGeometry) // handle rigids with no geometry (remove or convert to clustered)
+				{
+					if (bHasRigidChildren)
+					{
+						SimulationType[TransformGroupIndex] = ESimulationTypes::FST_Clustered;
+					}
+					else
+					{
+						InvalidTransforms.Add(TransformGroupIndex);
+						bKeep = false;
+					}
+				}
+				if (bKeep && SimulationType[TransformGroupIndex] != ESimulationTypes::FST_None && Parent[TransformGroupIndex] != INDEX_NONE)
+				{
+					RigidChildren[Parent[TransformGroupIndex]] = true;
+				}
+			}
+			if (InvalidGeometry.Num() > 0)
+			{
+				UE_LOG(FGeometryCollectionLogging, Warning, TEXT("Removing %d invalid, fully-invisible geometries from geometry collection."), InvalidGeometry.Num());
+				InvalidGeometry.Sort();
+				RemoveElements(GeometryGroup, InvalidGeometry);
+			}
+			if (InvalidTransforms.Num() > 0)
+			{
+				UE_LOG(FGeometryCollectionLogging, Warning, TEXT("Removing %d invalid, empty transforms from geometry collection."), InvalidTransforms.Num());
+				InvalidTransforms.Sort();
+				RemoveElements(TransformGroup, InvalidTransforms);
+			}
+			
+			Version = 9;
+		}
 	}
 }
 
diff --git a/Engine/Source/Runtime/Experimental/Chaos/Private/GeometryCollection/ManagedArrayCollection.cpp b/Engine/Source/Runtime/Experimental/Chaos/Private/GeometryCollection/ManagedArrayCollection.cpp
--- a/Engine/Source/Runtime/Experimental/Chaos/Private/GeometryCollection/ManagedArrayCollection.cpp
+++ b/Engine/Source/Runtime/Experimental/Chaos/Private/GeometryCollection/ManagedArrayCollection.cpp
@@ -12,7 +12,7 @@ int8 FManagedArrayCollection::Invalid = INDEX_NONE;
 
 FManagedArrayCollection::FManagedArrayCollection()
 {
-	Version = 8;
+	Version = 9;
 }
 
 void FManagedArrayCollection::AddGroup(FName Group)
@@ -374,7 +374,7 @@ static const FName GuidName("GUID");
 
 void FManagedArrayCollection::Serialize(Chaos::FChaosArchive& Ar)
 {
-	if (Ar.IsSaving()) Version = 6;
+	if (Ar.IsSaving()) Version = 9;
 	Ar << Version;
 
 	if (Ar.IsLoading())
