Chaos - more manifold overhead reduction

diff --git a/Engine/Source/Runtime/Experimental/Chaos/Private/Chaos/Collision/CollisionSolver.cpp b/Engine/Source/Runtime/Experimental/Chaos/Private/Chaos/Collision/CollisionSolver.cpp
--- a/Engine/Source/Runtime/Experimental/Chaos/Private/Chaos/Collision/CollisionSolver.cpp
+++ b/Engine/Source/Runtime/Experimental/Chaos/Private/Chaos/Collision/CollisionSolver.cpp
@@ -532,6 +532,7 @@ namespace Chaos
 			const FContactIterationParameters& IterationParameters,
 			const FContactParticleParameters& ParticleParameters)
 		{
+#if 0
 			FSolverBody& Body0 = *Constraint.GetSolverBody0();
 			FSolverBody& Body1 = *Constraint.GetSolverBody1();
 
@@ -604,6 +605,7 @@ namespace Chaos
 				Body1.ApplyTransformDelta(DX1, DR1);
 				Body1.UpdateRotationDependentState();
 			}
+#endif
 		}
 
 
diff --git a/Engine/Source/Runtime/Experimental/Chaos/Private/Chaos/Collision/PBDCollisionConstraint.cpp b/Engine/Source/Runtime/Experimental/Chaos/Private/Chaos/Collision/PBDCollisionConstraint.cpp
--- a/Engine/Source/Runtime/Experimental/Chaos/Private/Chaos/Collision/PBDCollisionConstraint.cpp
+++ b/Engine/Source/Runtime/Experimental/Chaos/Private/Chaos/Collision/PBDCollisionConstraint.cpp
@@ -364,7 +364,7 @@ namespace Chaos
 				ManifoldPoints[0].ContactPoint = ContactPoint;
 			}
 
-			InitManifoldPoint(ManifoldPoints[0], Dt);
+			InitManifoldPoint(0, Dt);
 
 			SetActiveContactPoint(ManifoldPoints[0].ContactPoint);
 		}
@@ -377,8 +377,10 @@ namespace Chaos
 		ManifoldPoints.Reset();
 	}
 
-	void FPBDCollisionConstraint::InitManifoldPoint(FManifoldPoint& ManifoldPoint, FReal Dt)
+	void FPBDCollisionConstraint::InitManifoldPoint(const int32 ManifoldPointIndex, FReal Dt)
 	{
+		FManifoldPoint& ManifoldPoint = ManifoldPoints[ManifoldPointIndex];
+
 		FConstGenericParticleHandle Particle0 = Particle[0];
 		FConstGenericParticleHandle Particle1 = Particle[1];
 		if (!Particle0.IsValid() || !Particle1.IsValid())
@@ -388,54 +390,18 @@ namespace Chaos
 			return;
 		}
 
-		// @todo(chaos): support static friction position correction for non-manifold contacts (spheres, point clouds, etc)
-		ManifoldPoint.bInsideStaticFrictionCone = bUseManifold;
-
 		// Update the derived contact state (CoM relative data)
-		UpdateManifoldPointFromContact(ManifoldPoint);
+		UpdateManifoldPointFromContact(ManifoldPointIndex);
 
 		// Initialize the previous contact transforms if the data is available, otherwise reset them to current
-		TryRestoreFrictionData(ManifoldPoint);
-
-		// @todo(chaos): Remove SolverBody use from InitManifoldPoint when we don't need Incremental Manifolds any more
-		// This is a bit ugly. While we still support incremental manifolds, this function can be called
-		// in CreateConstraint (OK) and UpdateConstraint (not OK). In UpdateConstraint we should be using
-		// the SolverBodies for the state, but in InitConstraint we access the Particles directly.
-		FVec3 PCoM0, PCoM1;
-		FRotation3 QCoM0, QCoM1;
-		if ((GetSolverBody0() != nullptr) && (GetSolverBody1() != nullptr))
-		{
-			// This will be called from UpdateConstraint
-			PCoM0 = GetSolverBody0()->P();
-			QCoM0 = GetSolverBody0()->Q();
-			PCoM1 = GetSolverBody1()->P();
-			QCoM1 = GetSolverBody1()->Q();
-		}
-		else
-		{
-			// This will be called from CreateConstraint
-			PCoM0 = FParticleUtilities::GetCoMWorldPosition(Particle0);
-			QCoM0 = FParticleUtilities::GetCoMWorldRotation(Particle0);
-			PCoM1 = FParticleUtilities::GetCoMWorldPosition(Particle1);
-			QCoM1 = FParticleUtilities::GetCoMWorldRotation(Particle1);
-		}
-
-		// We can used the particle here even in UpdateConstraints because we are accessing tick constants 
-		// and nothing that is changed by the constraint solvers
-		const FVec3& PreV0 = Particle0->PreV();
-		const FVec3& PreW0 = Particle0->PreW();
-		const FVec3& PreV1 = Particle1->PreV();
-		const FVec3& PreW1 = Particle1->PreW();
-
-		// World-space contact state used below
-		GetWorldSpaceManifoldPoint(ManifoldPoint, PCoM0, QCoM0, PCoM1, QCoM1, ManifoldPoint.ContactPoint.Location, ManifoldPoint.ContactPoint.Phi);
+		TryRestoreFrictionData(ManifoldPointIndex);
 	}
 
 	int32 FPBDCollisionConstraint::AddManifoldPoint(const FContactPoint& ContactPoint, const FReal Dt)
 	{
 		int32 ManifoldPointIndex = ManifoldPoints.Add(ContactPoint);
 
-		InitManifoldPoint(ManifoldPoints[ManifoldPointIndex], Dt);
+		InitManifoldPoint(ManifoldPointIndex, Dt);
 
 		return ManifoldPointIndex;
 	}
@@ -446,28 +412,29 @@ namespace Chaos
 		// otherwise the PrevLocalContactPoint1 we calculated is not longer for the correct point.
 		FManifoldPoint& ManifoldPoint = ManifoldPoints[ManifoldPointIndex];
 		ManifoldPoint.ContactPoint = ContactPoint;
-		UpdateManifoldPointFromContact(ManifoldPoint);
+		UpdateManifoldPointFromContact(ManifoldPointIndex);
 	}
 
 	// Update the derived contact state (CoM relative data)
-	// @todo(chaos): we shouldn't have to recalculate local space positions - they were available in collision update
-	void FPBDCollisionConstraint::UpdateManifoldPointFromContact(FManifoldPoint& ManifoldPoint)
+	void FPBDCollisionConstraint::UpdateManifoldPointFromContact(const int32 ManifoldPointIndex)
 	{
+		FManifoldPoint& ManifoldPoint = ManifoldPoints[ManifoldPointIndex];
 		FConstGenericParticleHandle Particle0 = Particle[0];
 		FConstGenericParticleHandle Particle1 = Particle[1];
 
-		const int32 PlaneOwner = 1;
-		const FRigidTransform3& PlaneTransform = ImplicitTransform[PlaneOwner];
 		const FVec3 LocalContactPoint0 = ImplicitTransform[0].TransformPositionNoScale(ManifoldPoint.ContactPoint.ShapeContactPoints[0]);	// Particle Space on body 0
 		const FVec3 LocalContactPoint1 = ImplicitTransform[1].TransformPositionNoScale(ManifoldPoint.ContactPoint.ShapeContactPoints[1]);	// Particle Space on body 1
 
 		const FVec3 CoMContactPoint0 = Particle0->RotationOfMass().Inverse() * (LocalContactPoint0 - Particle0->CenterOfMass());	// CoM Space on Body 0
 		const FVec3 CoMContactPoint1 = Particle1->RotationOfMass().Inverse() * (LocalContactPoint1 - Particle1->CenterOfMass());	// CoM Space on Body 1
-
-		const FRotation3&  Particle1Rotation = Particle1->R();// Plane is attached to particle1
-		
 		ManifoldPoint.CoMContactPoints[0] = CoMContactPoint0;
 		ManifoldPoint.CoMContactPoints[1] = CoMContactPoint1;
+
+		// We now assume that the low-level collision detection functions initialize the world-space contact in the way we want, which is as below...
+		//const FVec3 WorldContactPoint0 = Particle0->P() + Particle0->Q().RotateVector(LocalContactPoint0);
+		//const FVec3 WorldContactPoint1 = Particle1->P() + Particle1->Q().RotateVector(LocalContactPoint1);
+		//ManifoldPoint.ContactPoint.Location = FReal(0.5) * (WorldContactPoint0 + WorldContactPoint1);
+		//ManifoldPoint.ContactPoint.Phi = FVec3::DotProduct(WorldContactPoint0 - WorldContactPoint1, ManifoldPoint.ContactPoint.Normal);
 	}
 
 	void FPBDCollisionConstraint::SetActiveContactPoint(const FContactPoint& ContactPoint)
@@ -524,24 +491,6 @@ namespace Chaos
 		OutContactPhi = FVec3::DotProduct(ContactPos0 - ContactPos1, ManifoldPoint.ContactPoint.Normal);
 	}
 
-	FManifoldPoint& FPBDCollisionConstraint::SetActiveManifoldPoint(
-		int32 ManifoldPointIndex,
-		const FVec3& P0,
-		const FRotation3& Q0,
-		const FVec3& P1,
-		const FRotation3& Q1)
-	{
-		FManifoldPoint& ManifoldPoint = ManifoldPoints[ManifoldPointIndex];
-
-		// Update the world-space state in the manifold point
-		GetWorldSpaceManifoldPoint(ManifoldPoint, P0, Q0, P1, Q1, ManifoldPoint.ContactPoint.Location, ManifoldPoint.ContactPoint.Phi);
-	
-		// Copy the world-space state into the active contact
-		SetActiveContactPoint(ManifoldPoint.ContactPoint);
-
-		return ManifoldPoint;
-	}
-
 	bool FPBDCollisionConstraint::CanUseManifold(FGeometryParticleHandle* Particle0, FGeometryParticleHandle* Particle1) const
 	{
 		// Do not use manifolds when a body is connected by a joint to another. Manifolds do not work when the bodies may be moved
@@ -578,45 +527,51 @@ namespace Chaos
 		bWasManifoldRestored = true;
 	}
 
-	bool FPBDCollisionConstraint::TryRestoreFrictionData(FManifoldPoint& ManifoldPoint)
+	bool FPBDCollisionConstraint::TryRestoreFrictionData(const int32 ManifoldPointIndex)
 	{
 		// Find the previous manifold point that matches
-		const FManifoldPoint* MatchedManifoldPoint = nullptr;
+		FManifoldPoint& ManifoldPoint = ManifoldPoints[ManifoldPointIndex];
+		
+		// Assume we have no matching point from the previous tick, but assume we can retain friction from now on
+		// Not supported for non-manifolds yet (hopefully we don't need to)
+		ManifoldPoint.bInsideStaticFrictionCone = bUseManifold;
+		ManifoldPoint.StaticFrictionMax = FReal(0);
+
 		if (bChaos_Manifold_EnableFrictionRestore)
 		{
-			// @todo(chaos): ManifoldPoints and PrevManifoldPoints are usually in the same order, so this loop could normally terminate in 1 iteration
-			for (int32 PrevManifoldPointIndex = 0; PrevManifoldPointIndex < PrevManifoldPoints.Num(); ++PrevManifoldPointIndex)
+			// ManifoldPoints and PrevManifoldPoints are usually in the same order, so this loop should normally terminate in first iteration
+			int32 PrevManifoldPointIndex = ManifoldPointIndex;
+			for (int32 PrevManifoldPointOffset = 0; PrevManifoldPointOffset < PrevManifoldPoints.Num(); ++PrevManifoldPointOffset, ++PrevManifoldPointIndex)
 			{
+				if (PrevManifoldPointIndex >= PrevManifoldPoints.Num())
+				{
+					PrevManifoldPointIndex = 0;
+				}
+
 				const FManifoldPoint& PrevManifoldPoint = PrevManifoldPoints[PrevManifoldPointIndex];
-				const FVec3 DP0 = ManifoldPoint.CoMContactPoints[0] - PrevManifoldPoint.CoMContactPoints[0];
-				const FVec3 DP1 = ManifoldPoint.CoMContactPoints[1] - PrevManifoldPoint.CoMContactPoints[1];
-
-				// If the contact point is in the same spot on one of the bodies, assume it is the same contact
-				// @todo(chaos): more robust same-point test. E.g., this won't work for very small or very large objects,
-				// so at least make the tolerance size-dependent
-				const FReal DistanceToleranceSq = FMath::Square(Chaos_Manifold_FrictionPositionTolerance);
-				if ((DP0.SizeSquared() < DistanceToleranceSq) || (DP1.SizeSquared() < DistanceToleranceSq))
+				if (PrevManifoldPoint.bInsideStaticFrictionCone)
 				{
-					MatchedManifoldPoint = &PrevManifoldPoint;
-					break;
+					const FVec3 DP0 = ManifoldPoint.CoMContactPoints[0] - PrevManifoldPoint.CoMContactPoints[0];
+					const FVec3 DP1 = ManifoldPoint.CoMContactPoints[1] - PrevManifoldPoint.CoMContactPoints[1];
+
+					// If the contact point is in the same spot on one of the bodies, assume it is the same contact
+					// @todo(chaos): more robust same-point test. E.g., this won't work for very small or very large objects,
+					// so at least make the tolerance size-dependent
+					const FReal DistanceToleranceSq = FMath::Square(Chaos_Manifold_FrictionPositionTolerance);
+					if ((DP0.SizeSquared() < DistanceToleranceSq) || (DP1.SizeSquared() < DistanceToleranceSq))
+					{
+						// If we have a previous point, use it to set the previous-state data required for 
+						// static friction otherwise reset static friction limits
+						ManifoldPoint.CoMContactPoints[0] = PrevManifoldPoint.CoMContactPoints[0];
+						ManifoldPoint.CoMContactPoints[1] = PrevManifoldPoint.CoMContactPoints[1];
+						ManifoldPoint.StaticFrictionMax = PrevManifoldPoint.StaticFrictionMax;
+						return true;
+					}
 				}
 			}
 		}
 
-		// If we have a previous point, use it to set the previous-state data required for 
-		// static friction otherwise reset static friction limits
-		if ((MatchedManifoldPoint != nullptr) && MatchedManifoldPoint->bInsideStaticFrictionCone)
-		{
-			ManifoldPoint.CoMContactPoints[0] = MatchedManifoldPoint->CoMContactPoints[0];
-			ManifoldPoint.CoMContactPoints[1] = MatchedManifoldPoint->CoMContactPoints[1];
-			ManifoldPoint.StaticFrictionMax = MatchedManifoldPoint->StaticFrictionMax;
-			return true;
-		}
-		else
-		{
-			ManifoldPoint.StaticFrictionMax = FReal(0);
-			return false;
-		}
+		return false;
 	}
 
 	ECollisionConstraintDirection FPBDCollisionConstraint::GetConstraintDirection(const FReal Dt) const
diff --git a/Engine/Source/Runtime/Experimental/Chaos/Private/Chaos/CollisionOneShotManifolds.cpp b/Engine/Source/Runtime/Experimental/Chaos/Private/Chaos/CollisionOneShotManifolds.cpp
--- a/Engine/Source/Runtime/Experimental/Chaos/Private/Chaos/CollisionOneShotManifolds.cpp
+++ b/Engine/Source/Runtime/Experimental/Chaos/Private/Chaos/CollisionOneShotManifolds.cpp
@@ -619,7 +619,7 @@ namespace Chaos
 		// Use GJK to find the closest points (or shallowest penetrating points) on two convex shapes usingthe specified margin
 		// @todo(chaos): dedupe from GJKContactPoint in CollisionResolution.cpp
 		template <typename GeometryA, typename GeometryB>
-		FContactPoint GJKContactPointMargin(const GeometryA& A, const GeometryB& B, const FRigidTransform3& ATM, const FRigidTransform3& BToATM, FReal MarginA, FReal MarginB, FGJKSimplexData& InOutGjkWarmStartData, FReal& OutMaxMarginDelta)
+		FContactPoint GJKContactPointMargin(const GeometryA& A, const GeometryB& B, const FRigidTransform3& ATM, const FRigidTransform3& BTM, FReal MarginA, FReal MarginB, FGJKSimplexData& InOutGjkWarmStartData, FReal& OutMaxMarginDelta)
 		{
 			SCOPE_CYCLE_COUNTER_MANIFOLD_GJK();
 
@@ -633,13 +633,16 @@ namespace Chaos
 
 			const TGJKCoreShape<GeometryA> AWithMargin(A, MarginA);
 			const TGJKCoreShape<GeometryB> BWithMargin(B, MarginB);
+			const FRigidTransform3 BToATM = BTM.GetRelativeTransformNoScale(ATM);
 
 			if (GJKPenetrationWarmStartable<true>(AWithMargin, BWithMargin, BToATM, FReal(0), FReal(0), Penetration, ClosestA, ClosestB, NormalA, NormalB, InOutGjkWarmStartData, OutMaxMarginDelta, Epsilon))
 			{
 				Contact.ShapeContactPoints[0] = ClosestA;
 				Contact.ShapeContactPoints[1] = ClosestB;
 				Contact.ShapeContactNormal = -NormalB;	// We want normal pointing from B to A
-				Contact.Location = ATM.TransformPositionNoScale(ClosestA);
+				const FVec3 WorldLocationA = ATM.TransformPositionNoScale(ClosestA);
+				const FVec3 WorldLocationB = BTM.TransformPositionNoScale(ClosestB);
+				Contact.Location = FReal(0.5) * (WorldLocationA + WorldLocationB);
 				Contact.Normal = -ATM.TransformVectorNoScale(NormalA);
 				Contact.Phi = -Penetration;
 			}
@@ -726,8 +729,6 @@ namespace Chaos
 			ensure(Convex1Transform.GetScale3D() == FVec3(1.0f, 1.0f, 1.0f));
 			ensure(Convex2Transform.GetScale3D() == FVec3(1.0f, 1.0f, 1.0f));
 
-			const FRigidTransform3 Convex2ToConvex1Transform = Convex2Transform.GetRelativeTransformNoScale(Convex1Transform);
-
 			// Get the adjusted margins for each convex
 			FReal Margin1 = Convex1.GetMargin();
 			FReal Margin2 = Convex2.GetMargin();
@@ -750,7 +751,7 @@ namespace Chaos
 			// Find the deepest penetration. This is used to determine the planes and points to use for the manifold
 			// MaxMarginDelta is an upper bound on the distance from the contact on the rounded core shape to the actual shape surface. 
 			FReal MaxMarginDelta = FReal(0);
-			FContactPoint GJKContactPoint = GJKContactPointMargin(Convex1, Convex2, Convex1Transform, Convex2ToConvex1Transform, Margin1, Margin2, Constraint.GetGJKWarmStartData(), MaxMarginDelta);
+			FContactPoint GJKContactPoint = GJKContactPointMargin(Convex1, Convex2, Convex1Transform, Convex2Transform, Margin1, Margin2, Constraint.GetGJKWarmStartData(), MaxMarginDelta);
 
 			// GJK is using margins and rounded corner, so if we have a corner-to-corner contact it will under-report the actual distance by an amount that depends on how
 			// "pointy" the edge/corner is - this error is bounded by MaxMarginDelta.
@@ -807,7 +808,8 @@ namespace Chaos
 				GJKContactPoint.ShapeContactPoints[0] = ShapeEdgePos1;
 				GJKContactPoint.ShapeContactPoints[1] = ShapeEdgePos2;
 				GJKContactPoint.Phi = EdgePhi;
-				GJKContactPoint.Location = 0.5f * (EdgePos1 + EdgePos2);
+				GJKContactPoint.Location = FReal(0.5) * (EdgePos1 + EdgePos2);
+				// Normal unchanged from GJK result
 
 				Constraint.AddOneshotManifoldContact(GJKContactPoint, Dt);
 				return;
diff --git a/Engine/Source/Runtime/Experimental/Chaos/Public/Chaos/Collision/ContactPoint.h b/Engine/Source/Runtime/Experimental/Chaos/Public/Chaos/Collision/ContactPoint.h
--- a/Engine/Source/Runtime/Experimental/Chaos/Public/Chaos/Collision/ContactPoint.h
+++ b/Engine/Source/Runtime/Experimental/Chaos/Public/Chaos/Collision/ContactPoint.h
@@ -10,23 +10,38 @@
 
 namespace Chaos
 {
-	// Data returned by the low-level collision functions
+	// 
+	// In principle, the low-level collision detection functions should not need to calculate the world-space contact data, but some of them
+	// do some of their work in world space and already have the results we need. To avoid duplicate work, all collision detection functions
+	// should also fill in the world-space data
+
+	/**
+	 * @brief Data returned by the low-level collision functions
+	 * 
+	 * @note In principle, the low-level collision detection functions should not need to calculate the world-space contact data, but some of them
+	 * do some of their work in world space and already have the results we need. To avoid duplicate work, all collision detection functions
+	 * should also fill in the world-space data
+	*/
 	class CHAOS_API FContactPoint
 	{
 	public:
-		FVec3 ShapeContactPoints[2];	// Shape-space contact points on the two bodies, without the margin added (i.e., contact point on the core shape)
-		FVec3 ShapeContactNormal;		// Shape-space contact normal relative to the NormalOwner, but with direction that always goes from body 1 to body 0
+		// Shape-space contact points on the two bodies, without the margin added (i.e., contact point on the core shape)
+		FVec3 ShapeContactPoints[2];
 
-		// @todo(chaos): Collision detection output should only produce shape-space results so Location and Normal should go
-		// E.g., ShapeContactNormal is used in the trianglemesh contact culling
-		FVec3 Location;					// World-space contact location
-		FVec3 Normal;					// World-space contact normal with direction that always goes from body 1 to body 0
+		// Shape-space contact normal on the second shape with direction that points away from shape 1
+		FVec3 ShapeContactNormal;
 
-		FReal Phi;						// Contact separation (negative for overlap)
+		// World-space contact location
+		FVec3 Location;
+
+		// World-space contact normal with direction that points away from shape 1
+		FVec3 Normal;
+
+		// Contact separation (negative for overlap)
+		FReal Phi;
 
 		FContactPoint()
-			: Normal(1, 0, 0)
-			, Phi(TNumericLimits<FReal>::Max())
+			: Phi(TNumericLimits<FReal>::Max())
 		{
 		}
 
diff --git a/Engine/Source/Runtime/Experimental/Chaos/Public/Chaos/Collision/PBDCollisionConstraint.h b/Engine/Source/Runtime/Experimental/Chaos/Public/Chaos/Collision/PBDCollisionConstraint.h
--- a/Engine/Source/Runtime/Experimental/Chaos/Public/Chaos/Collision/PBDCollisionConstraint.h
+++ b/Engine/Source/Runtime/Experimental/Chaos/Public/Chaos/Collision/PBDCollisionConstraint.h
@@ -322,12 +322,6 @@ namespace Chaos
 
 		TArrayView<FManifoldPoint> GetManifoldPoints() { return MakeArrayView(ManifoldPoints); }
 		TArrayView<const FManifoldPoint> GetManifoldPoints() const { return MakeArrayView(ManifoldPoints); }
-		FManifoldPoint& SetActiveManifoldPoint(
-			int32 ManifoldPointIndex,
-			const FVec3& P0,
-			const FRotation3& Q0,
-			const FVec3& P1,
-			const FRotation3& Q1);
 
 		void AddIncrementalManifoldContact(const FContactPoint& ContactPoint, const FReal Dt);
 		void AddOneshotManifoldContact(const FContactPoint& ContactPoint, const FReal Dt);
@@ -336,7 +330,7 @@ namespace Chaos
 
 		//@ todo(chaos): These are for the collision forwarding system - this should use the collision modifier system (which should be extended to support adding collisions)
 		void SetManifoldPoints(const TArray<FManifoldPoint>& InManifoldPoints) { ManifoldPoints = InManifoldPoints; }
-		void UpdateManifoldPointFromContact(FManifoldPoint& ManifoldPoint);
+		void UpdateManifoldPointFromContact(const int32 ManifoldPointIndex);
 
 		// Helpers for interacting with constraint from world space
 		static void GetWorldSpaceContactPositions(const FManifoldPoint& ManifoldPoint, const FVec3& PCoM0, const FRotation3& QCoM0, const FVec3& PCoM1, const FRotation3& QCoM1, FVec3& OutWorldPosition0, FVec3& OutWorldPosition1);
@@ -419,7 +413,7 @@ namespace Chaos
 		bool AreMatchingContactPoints(const FContactPoint& A, const FContactPoint& B, FReal& OutScore) const;
 		int32 FindManifoldPoint(const FContactPoint& ContactPoint) const;
 		int32 AddManifoldPoint(const FContactPoint& ContactPoint, const FReal Dt);
-		void InitManifoldPoint(FManifoldPoint& ManifoldPoint, FReal Dt);
+		void InitManifoldPoint(const int32 ManifoldPointIndex, FReal Dt);
 		void UpdateManifoldPoint(int32 ManifoldPointIndex, const FContactPoint& ContactPoint, const FReal Dt);
 		void SetActiveContactPoint(const FContactPoint& ContactPoint);
 		void GetWorldSpaceManifoldPoint(const FManifoldPoint& ManifoldPoint, const FVec3& P0, const FRotation3& Q0, const FVec3& P1, const FRotation3& Q1, FVec3& OutContactLocation, FReal& OutContactPhi);
@@ -432,7 +426,7 @@ namespace Chaos
 		/**
 		 * @brief Restore data used for static friction from the previous point(s) to the current
 		*/
-		bool TryRestoreFrictionData(FManifoldPoint& ManifoldPoint);
+		bool TryRestoreFrictionData(const int32 ManifoldPointIndex);
 
 	public:
 		//@todo(chaos): make this stuff private
