Levels+Colors computation optims + cvar to enable/disable them

diff --git a/Engine/Source/Programs/HeadlessChaos/Private/HeadlessChaosTestGraph.cpp b/Engine/Source/Programs/HeadlessChaos/Private/HeadlessChaosTestGraph.cpp
--- a/Engine/Source/Programs/HeadlessChaos/Private/HeadlessChaosTestGraph.cpp
+++ b/Engine/Source/Programs/HeadlessChaos/Private/HeadlessChaosTestGraph.cpp
@@ -524,6 +524,12 @@ namespace ChaosTest {
 			// Check the number of islands
 			EXPECT_EQ(IterData.ExpectedIslandParticleIndices.Num(), Graph.NumIslands());
 
+			// Compute levels and colors
+			Graph.GetIslandGraph()->InitSorting();
+
+			Graph.GetIslandGraph()->ComputeLevels(ContainerId);
+			Graph.GetIslandGraph()->ComputeColors(ContainerId, INT32_MAX);
+
 			// Assign color to constraints
 			GraphColor.InitializeColor(Graph);
 			for (int32 IslandIndex = 0; IslandIndex < Graph.NumIslands(); ++IslandIndex)
@@ -575,7 +581,10 @@ namespace ChaosTest {
 
 				EXPECT_EQ(IterData.MaxLevel[Index], MaxLevel);
 				EXPECT_EQ(IterData.MaxColor[Index], MaxColor);
-
+				
+				EXPECT_EQ(FMath::Max(0,IterData.MaxLevel[Index]), FMath::Max(0,Graph.GetIslandGraph()->GraphIslands[Graph.GetGraphIndex(Island)].MaxLevels));
+				EXPECT_EQ(FMath::Max(0,IterData.MaxColor[Index]), FMath::Max(0,Graph.GetIslandGraph()->GraphIslands[Graph.GetGraphIndex(Island)].MaxColors));
+				
 				++Index;
 			}
 		}
diff --git a/Engine/Source/Runtime/Experimental/Chaos/Private/Chaos/Island/IslandGraph.cpp b/Engine/Source/Runtime/Experimental/Chaos/Private/Chaos/Island/IslandGraph.cpp
--- a/Engine/Source/Runtime/Experimental/Chaos/Private/Chaos/Island/IslandGraph.cpp
+++ b/Engine/Source/Runtime/Experimental/Chaos/Private/Chaos/Island/IslandGraph.cpp
@@ -374,6 +374,187 @@ void FIslandGraph<NodeType, EdgeType, IslandType>::MergeIslands()
 	}
 }
 
+template<typename NodeType, typename EdgeType, typename IslandType>
+void FIslandGraph<NodeType, EdgeType, IslandType>::InitSorting()
+{
+	// Reset nodes levels and colors
+	for(auto& GraphNode : GraphNodes)
+	{
+		GraphNode.LevelIndex = INDEX_NONE;
+		GraphNode.ColorIndices.Reset();
+	}
+	// Reset edges levels and colors
+	for(auto& GraphEdge : GraphEdges)
+	{
+		GraphEdge.LevelIndex = INDEX_NONE;
+		GraphEdge.ColorIndex = INDEX_NONE;
+	}
+	// Reset island max number of levels and colors
+	for(auto& GraphIsland : GraphIslands)
+	{
+		GraphIsland.MaxLevels = INDEX_NONE;
+		GraphIsland.MaxColors = INDEX_NONE;
+	}
+}
+	
+template<typename NodeType, typename EdgeType, typename IslandType>
+void FIslandGraph<NodeType, EdgeType, IslandType>::UpdateLevels(const int32 NodeIndex, const int32 ContainerId, TQueue<int32>& NodeQueue)
+{
+	if(GraphNodes.IsValidIndex(NodeIndex))
+	{
+		FGraphNode& GraphNode = GraphNodes[NodeIndex];
+		for (int32 EdgeIndex : GraphNode.NodeEdges)
+		{
+			FGraphEdge& GraphEdge = GraphEdges[EdgeIndex];
+
+			// Do nothing if the edge is not coming from the same container or if the island is sleeping 
+			if (GraphEdge.ItemContainer == ContainerId && GraphEdge.LevelIndex == INDEX_NONE && !GraphIslands[GraphEdge.IslandIndex].bIsSleeping)
+			{
+				const int32 OtherIndex = (NodeIndex == GraphEdge.FirstNode) ?
+							GraphEdge.SecondNode : GraphEdge.FirstNode;
+
+				GraphEdge.LevelIndex = GraphNode.LevelIndex;
+					
+				GraphIslands[GraphEdge.IslandIndex].MaxLevels = FGenericPlatformMath::Max(
+					GraphIslands[GraphEdge.IslandIndex].MaxLevels, GraphEdge.LevelIndex);
+
+				// If we have another node, append it to our queue on the next level
+				if (GraphNodes.IsValidIndex(OtherIndex) && GraphNodes[OtherIndex].bValidNode && GraphNodes[OtherIndex].LevelIndex == INDEX_NONE)
+				{
+					GraphNodes[OtherIndex].LevelIndex = GraphEdge.LevelIndex+1;
+					NodeQueue.Enqueue(OtherIndex);
+				}
+			}
+		}
+	}
+}
+
+template<typename NodeType, typename EdgeType, typename IslandType>
+void FIslandGraph<NodeType, EdgeType, IslandType>::ComputeLevels(const int32 ContainerId)
+{
+	// First enqueue all the static/kinematic nodes for levels 0
+	TQueue<int32> NodeQueue;
+	for (int32 NodeIndex = 0, NumNodes = GraphNodes.GetMaxIndex(); NodeIndex < NumNodes; ++NodeIndex)
+	{
+		if(GraphNodes.IsValidIndex(NodeIndex))
+		{
+			FGraphNode& GraphNode = GraphNodes[NodeIndex];
+			if(!GraphNode.bValidNode)
+			{
+				GraphNode.LevelIndex = 0;
+				UpdateLevels(NodeIndex, ContainerId, NodeQueue);
+			}
+		}
+	}
+
+	// Then iteratively loop over these root nodes and propagate the levels through connectivity
+	int32 NodeIndex = INDEX_NONE;
+	while(!NodeQueue.IsEmpty())
+	{
+		NodeQueue.Dequeue(NodeIndex);
+		UpdateLevels(NodeIndex, ContainerId, NodeQueue);
+	}
+
+	// An isolated island that is only dynamics will not have been processed above, put everything without a level into level zero
+	for(auto& GraphEdge : GraphEdges)
+	{
+		if(GraphEdge.ItemContainer == ContainerId)
+		{
+			GraphEdge.LevelIndex = FGenericPlatformMath::Max(GraphEdge.LevelIndex, 0);
+		}
+	}
+}
+	
+template<typename NodeType, typename EdgeType, typename IslandType>
+int32 FIslandGraph<NodeType, EdgeType, IslandType>::PickColor(const FGraphNode& GraphNode, const int32 OtherIndex, TQueue<int32>& NodeQueue)
+{
+	int32 ColorToUse = 0;
+	if (GraphNodes.IsValidIndex(OtherIndex) && GraphNodes[OtherIndex].bValidNode)
+	{
+		FGraphNode& OtherNode = GraphNodes[OtherIndex];
+		// Pick the first color not used by the 2 edge nodes
+		while (OtherNode.ColorIndices.Contains(ColorToUse) || GraphNode.ColorIndices.Contains(ColorToUse))
+		{
+			ColorToUse++;
+		}
+		// The color will be added to the graph node in the UpdateColors function
+		OtherNode.ColorIndices.Add(ColorToUse);
+		if (OtherNode.NodeCounter != GraphCounter)
+		{
+			NodeQueue.Enqueue(OtherIndex);
+		}
+	}
+	else
+	{
+		// If only one node, only iterate over that node available color
+		while (GraphNode.ColorIndices.Contains(ColorToUse))
+		{
+			ColorToUse++;
+		}
+	}
+	return ColorToUse;
+}
+
+template<typename NodeType, typename EdgeType, typename IslandType>
+void FIslandGraph<NodeType, EdgeType, IslandType>::UpdateColors(const int32 NodeIndex, const int32 ContainerId, TQueue<int32>& NodeQueue, const int32 MinEdges)
+{
+	if (GraphNodes.IsValidIndex(NodeIndex))
+	{
+		FGraphNode& GraphNode = GraphNodes[NodeIndex];
+		GraphNode.NodeCounter = GraphCounter;
+		
+		for (int32 EdgeIndex : GraphNode.NodeEdges)
+		{
+			FGraphEdge& GraphEdge = GraphEdges[EdgeIndex];
+
+			// Do nothing if the edge is not coming from the same container or if the island is sleeping 
+			if (GraphEdge.ItemContainer == ContainerId && GraphEdge.ColorIndex == INDEX_NONE &&
+				!GraphIslands[GraphEdge.IslandIndex].bIsSleeping && (GraphIslands[GraphEdge.IslandIndex].NumEdges > MinEdges) )
+			{
+				// Get the opposite node index for the given edge
+				const int32 OtherIndex = (NodeIndex == GraphEdge.FirstNode) ?
+					GraphEdge.SecondNode : GraphEdge.FirstNode;
+
+				// Get the first available color to be used by the edge
+				const int32 ColorToUse = PickColor(GraphNode, OtherIndex, NodeQueue);
+				
+				GraphNode.ColorIndices.Add(ColorToUse);
+				GraphEdge.ColorIndex = ColorToUse;
+				
+				GraphIslands[GraphEdge.IslandIndex].MaxColors = FGenericPlatformMath::Max(
+					GraphIslands[GraphEdge.IslandIndex].MaxColors, ColorToUse);
+			}
+		}
+	}
+}
+
+template<typename NodeType, typename EdgeType, typename IslandType>
+void FIslandGraph<NodeType, EdgeType, IslandType>::ComputeColors(const int32 ContainerId, const int32 MinEdges)
+{
+	GraphCounter = (GraphCounter + 1) % MaxCount;
+
+	TQueue<int32> NodeQueue;
+	int32 NodeIndex = INDEX_NONE;
+
+	// We first loop over all the nodes that have not been processed and valid (dynamic/sleeping)
+	for (int32 RootIndex = 0, NumNodes = GraphNodes.GetMaxIndex(); RootIndex < NumNodes; ++RootIndex)
+	{
+		if (GraphNodes.IsValidIndex(RootIndex))
+		{
+			FGraphNode& GraphNode = GraphNodes[RootIndex];
+			if (GraphNode.NodeCounter != GraphCounter && GraphNode.bValidNode)
+			{
+				NodeQueue.Enqueue(RootIndex);
+				while (!NodeQueue.IsEmpty())
+				{
+					NodeQueue.Dequeue(NodeIndex);
+					UpdateColors(NodeIndex, ContainerId, NodeQueue, MinEdges);
+				}
+			}
+		}
+	}
+}
+
 template<typename NodeType, typename EdgeType, typename IslandType>
 void FIslandGraph<NodeType, EdgeType, IslandType>::SplitIsland(TQueue<int32>& NodeQueue, const int32 RootIndex, const int32 IslandIndex)
 {
@@ -412,7 +593,7 @@ void FIslandGraph<NodeType, EdgeType, IslandType>::SplitIsland(TQueue<int32>& No
 		}
 	}
 }
-	
+
 template<typename NodeType, typename EdgeType, typename IslandType>
 void FIslandGraph<NodeType, EdgeType, IslandType>::SplitIslands()
 {
diff --git a/Engine/Source/Runtime/Experimental/Chaos/Private/Chaos/PBDConstraintRule.cpp b/Engine/Source/Runtime/Experimental/Chaos/Private/Chaos/PBDConstraintRule.cpp
--- a/Engine/Source/Runtime/Experimental/Chaos/Private/Chaos/PBDConstraintRule.cpp
+++ b/Engine/Source/Runtime/Experimental/Chaos/Private/Chaos/PBDConstraintRule.cpp
@@ -20,6 +20,15 @@ namespace Chaos
 	int32 ChaosCollisionColorMinParticles = 2000;
 	FAutoConsoleVariableRef CVarChaosCollisionColorMinParticles(TEXT("p.Chaos.Collision.Color.MinParticles"), ChaosCollisionColorMinParticles, TEXT(""));
 
+	/** Console variable to disable the levels computation on each island */
+	int32 ChaosDisableIslandLevels = false;
+	FAutoConsoleVariableRef CVarChaosDisableIslandLevels(TEXT("p.Chaos.Islands.DisableLevels"), ChaosDisableIslandLevels, TEXT(""));
+
+	/** console variable to disable the colors computation on each island */
+	int32 ChaosDisableIslandColors = true;
+	FAutoConsoleVariableRef CVarChaosDisableIslandColors(TEXT("p.Chaos.Islands.DisableColors"), ChaosDisableIslandColors, TEXT(""));
+
+
 	template<class ConstraintType>
 	TSimpleConstraintRule<ConstraintType>::TSimpleConstraintRule(int32 InPriority, FConstraints& InConstraints)
 		: FSimpleConstraintRule(InPriority)
@@ -143,7 +152,7 @@ namespace Chaos
 		{
 			const TArray<FConstraintHandleHolder>& IslandConstraints = ConstraintGraph->GetIslandConstraints(Island);
 
-			// This will reset the number od constraints inside the solver datas. For now we keep this function since according
+			// This will reset the number of constraints inside the solver datas. For now we keep this function since according
 			// to the constraints we can use the handles, the indices or the container. when only the container will be used
 			// we can replace this call by :
 			// IslandSolver->template GetConstraintContainer<typename ConstraintType::FSolverConstraintContainerType>(ContainerId)->Reset();
@@ -156,7 +165,7 @@ namespace Chaos
 					FConstraintContainerHandle* Constraint = ConstraintHandle->As<FConstraintContainerHandle>();
 
 					// Note we are building the SolverBodies as we go, in the order that we visit them. Each constraint
-					// references two bodies, so we won't strictly be acessing only in cache order, but it's about as good as it can be.
+					// references two bodies, so we won't strictly be accessing only in cache order, but it's about as good as it can be.
 					if (Constraint->IsEnabled())
 					{
 						// @todo(chaos): we should provide Particle Levels in the island rule as well (see TPBDConstraintColorRule)
@@ -209,9 +218,8 @@ namespace Chaos
 
 	template<class ConstraintType>
 	TPBDConstraintColorRule<ConstraintType>::TPBDConstraintColorRule(FConstraints& InConstraints, int32 InPriority)
-		: TPBDConstraintGraphRuleImpl<ConstraintType>(InConstraints, InPriority)
+		: TPBDConstraintIslandRule<ConstraintType>(InConstraints, InPriority)
 		, ConstraintSets()
-		, bUseColor(false)
 	{
 	}
 
@@ -223,166 +231,322 @@ namespace Chaos
 	template<class ConstraintType>
 	void TPBDConstraintColorRule<ConstraintType>::GatherSolverInput(const FReal Dt, int32 Island)
 	{
-		if(FPBDIslandSolver* IslandSolver = ConstraintGraph->GetSolverIsland(Island))
+		FPBDConstraintGraph* LocalGraph = ConstraintGraph;
+		auto GetParticleLevel = [LocalGraph](TGeometryParticleHandle<FReal, 3>* ConstrainedParticle) -> int32
 		{
-			const typename FPBDConstraintColor::FLevelToColorToConstraintListMap& LevelToColorToConstraintListMap = GraphColor.GetIslandLevelToColorToConstraintListMap(Island);
-			int32 MaxColor = GraphColor.GetIslandMaxColor(Island);
-			int32 MaxLevel = GraphColor.GetIslandMaxLevel(Island);
-
-			TArray<TPair<int32, int32>>& IslandConstraintSets = ConstraintSets[Island];
-			IslandConstraintSets.Reset(MaxLevel * MaxColor);	// Pessimistic array size - we could store the actual required size in coloring algorithm
-			int32 ConstraintSetEnd = 0;
-
-			Constraints.SetNumIslandConstraints(GraphColor.NumIslandEdges(Island), *IslandSolver);
-
-			for (int32 Level = 0; Level <= MaxLevel; ++Level)
+			int32 ParticleLevel = INDEX_NONE;
+			if(FPBDRigidParticleHandle* PBDRigid = ConstrainedParticle->CastToRigidParticle())
 			{
-				const int32 OldConstraintSetEnd = ConstraintSetEnd;
-
-				for (int32 Color = 0; Color <= MaxColor; ++Color)
+				if(LocalGraph->GetIslandGraph()->GraphNodes.IsValidIndex(PBDRigid->ConstraintGraphIndex()))
+				{
+					ParticleLevel = FMath::Max( LocalGraph->GetIslandGraph()->GraphNodes[PBDRigid->ConstraintGraphIndex()].LevelIndex, 0);
+				}
+			}
+			return ParticleLevel;
+		};
+		if (IsSortingEnabled())
+		{
+			if (FPBDIslandSolver* IslandSolver = ConstraintGraph->GetSolverIsland(Island))
+			{
+				const int32 GraphIndex = ConstraintGraph->GetGraphIndex(Island);
+				if(ConstraintGraph->GetIslandGraph()->GraphIslands.IsValidIndex(GraphIndex))
 				{
-					if (LevelToColorToConstraintListMap[Level].Contains(Color) && LevelToColorToConstraintListMap[Level][Color].Num())
+					const int32 MaxColor = IsSortingUsingColors() ? FMath::Max(1,ConstraintGraph->GetIslandGraph()->GraphIslands[GraphIndex].MaxColors+1) : 1;
+					const int32 MaxLevel = IsSortingUsingLevels() ? FMath::Max(1,ConstraintGraph->GetIslandGraph()->GraphIslands[GraphIndex].MaxLevels+1) : 1;
+					
+					TArray<TPair<int32, int32>>& IslandConstraintSets = ConstraintSets[Island];
+					IslandConstraintSets.Reset(MaxLevel * MaxColor);	// Pessimistic array size - we could store the actual required size in coloring algorithm
+					int32 ConstraintSetEnd = 0;
+					
+					Constraints.SetNumIslandConstraints(ConstraintGraph->GetIslandConstraints(Island).Num(), *IslandSolver);
+					
+					for (int32 Level = 0; Level < MaxLevel; ++Level)
 					{
-						// Calculate the range of indices for this color as a set of independent contacts
-						TPair<int32, int32> ColorConstrainSet(ConstraintSetEnd, ConstraintSetEnd);
-
-						const TArray<FConstraintContainerHandle*>& ConstraintHandles = GetLevelColorConstraints(LevelToColorToConstraintListMap, Level, Color);
-						for (FConstraintContainerHandle* Constraint : ConstraintHandles)
-						{
-							if (Constraint->IsEnabled())
-							{
-								// Levels that should be assigned to the bodies for shock propagation
-								// @todo(chaos): optimize the lookup
-								const TVector<TGeometryParticleHandle<FReal, 3>*, 2> ConstrainedParticles = Constraint->GetConstrainedParticles();
-								const int32 Particle0Level = GraphColor.GetParticleLevel(ConstrainedParticles[0]);
-								const int32 Particle1Level = GraphColor.GetParticleLevel(ConstrainedParticles[1]);
-
-								// Note we are building the SolverBodies as we go, in the order that we visit them. Each constraint
-								// references two bodies, so we won't strictly be acessing only in cache order, but it's about as good as it can be.
-								Constraint->GatherInput(Dt, Particle0Level, Particle1Level, *IslandSolver);
-
-								// Update the current constraint set of this color
-								ColorConstrainSet.Value = ++ConstraintSetEnd;
-							}
-						}
-
-						// Remember the set of constraints of this color
-						if (bUseColor)
+						for (int32 Color = 0; Color < MaxColor; ++Color)
 						{
-							IslandConstraintSets.Add(ColorConstrainSet);
+							const int32 OffsetIndex = IslandOffsets[Island] + Level * MaxColor + Color;
+							const int32 OffsetBegin = (OffsetIndex == 0) ? 0 : ConstraintOffsets[OffsetIndex-1];
+							const int32 OffsetEnd = ConstraintOffsets[OffsetIndex];
+							
+							 if (OffsetEnd != OffsetBegin)
+							 {
+							 	// Calculate the range of indices for this color as a set of independent contacts
+							 	TPair<int32, int32> ColorConstrainSet(ConstraintSetEnd, ConstraintSetEnd);
+							 	for(int32 ConstraintIndex = OffsetBegin; ConstraintIndex < OffsetEnd; ++ConstraintIndex)
+							 	{
+							 		FConstraintContainerHandle* Constraint = SortedConstraints[ConstraintIndex];
+							 		if (Constraint->IsEnabled())
+							 		{
+							 			// Levels that should be assigned to the bodies for shock propagation
+							 			// @todo(chaos): optimize the lookup
+							 			const TVector<TGeometryParticleHandle<FReal, 3>*, 2> ConstrainedParticles = Constraint->GetConstrainedParticles();
+							 			const int32 Particle0Level = GetParticleLevel(ConstrainedParticles[0]);
+							 			const int32 Particle1Level = GetParticleLevel(ConstrainedParticles[1]);
+							
+							 			// Note we are building the SolverBodies as we go, in the order that we visit them. Each constraint
+							 			// references two bodies, so we won't strictly be accessing only in cache order, but it's about as good as it can be.
+							 			Constraint->GatherInput(Dt, Particle0Level, Particle1Level, *IslandSolver);
+							
+							 			// Update the current constraint set of this color
+							 			ColorConstrainSet.Value = ++ConstraintSetEnd;
+							 		}
+							 	}
+							 	// Remember the set of constraints of this color
+							 	if (IsSortingUsingColors()) 
+							 	{
+							 		IslandConstraintSets.Add(ColorConstrainSet);
+							 	}
+							 }
 						}
 					}
+					// If we aren't coloring, we have a single group of all constraints (they have been created in level order above)
+					if (!IsSortingUsingColors())
+					{
+						const TPair<int32, int32> ConstrainSet(0, ConstraintSetEnd);
+						IslandConstraintSets.Add(ConstrainSet);
+					}
 				}
-
-				const int32 NumAdded = ConstraintSetEnd - OldConstraintSetEnd;
-				if (NumAdded > 1)
-				{
-					static int32 sHere = 0;
-					++sHere;
-				}
-			}
-			// If we aren't coloring, we have a single group of all constraints (they have been created in level order above)
-			if (!bUseColor)
-			{
-				const TPair<int32, int32> ConstrainSet(0, ConstraintSetEnd);
-				IslandConstraintSets.Add(ConstrainSet);
 			}
 		}
+		else
+		{
+			Base::GatherSolverInput(Dt, Island);
+		}
 	}
 
 	template<class ConstraintType>
 	void TPBDConstraintColorRule<ConstraintType>::ScatterSolverOutput(const FReal Dt, int32 Island)
 	{
-		if(FPBDIslandSolver* IslandSolver = ConstraintGraph->GetSolverIsland(Island))
+		if (IsSortingEnabled())
 		{
-			const TArray<TPair<int32, int32>>& IslandConstraintSets = ConstraintSets[Island];
-			for (const TPair<int32, int32>& ConstraintSet : IslandConstraintSets)
+			if (FPBDIslandSolver* IslandSolver = ConstraintGraph->GetSolverIsland(Island))
 			{
-				Constraints.ScatterOutput(Dt, ConstraintSet.Key, ConstraintSet.Value, *IslandSolver);
+				const TArray<TPair<int32, int32>>& IslandConstraintSets = ConstraintSets[Island];
+				for (const TPair<int32, int32>& ConstraintSet : IslandConstraintSets)
+				{
+					Constraints.ScatterOutput(Dt, ConstraintSet.Key, ConstraintSet.Value, *IslandSolver);
+				}
 			}
 		}
+		else
+		{
+			return Base::ScatterSolverOutput(Dt, Island);
+		}
 	}
 
 	template<class ConstraintType>
 	bool TPBDConstraintColorRule<ConstraintType>::ApplyConstraints(const FReal Dt, int32 Island, const int32 It, const int32 NumIts)
 	{
-		bool bNeedsAnotherIteration = false;
-		if(FPBDIslandSolver* IslandSolver = ConstraintGraph->GetSolverIsland(Island))
+		if (IsSortingEnabled())
 		{
-			const TArray<TPair<int32, int32>>& IslandConstraintSets = ConstraintSets[Island];
-			if (!bUseColor)
+			bool bNeedsAnotherIteration = false;
+			if (FPBDIslandSolver* IslandSolver = ConstraintGraph->GetSolverIsland(Island))
 			{
-				for (const TPair<int32, int32>& ConstraintSet : IslandConstraintSets)
+				const TArray<TPair<int32, int32>>& IslandConstraintSets = ConstraintSets[Island];
+				if (!IsSortingUsingColors())
 				{
-					bNeedsAnotherIteration |= Constraints.ApplyPhase1Serial(Dt, It, NumIts, ConstraintSet.Key, ConstraintSet.Value, *IslandSolver);
+					for (const TPair<int32, int32>& ConstraintSet : IslandConstraintSets)
+					{
+						bNeedsAnotherIteration |= Constraints.ApplyPhase1Serial(Dt, It, NumIts, ConstraintSet.Key, ConstraintSet.Value, *IslandSolver);
+					}
 				}
-			}
-			else
-			{
-				for (const TPair<int32, int32>& ConstraintSet : IslandConstraintSets)
+				else
 				{
-					bNeedsAnotherIteration |= Constraints.ApplyPhase1Parallel(Dt, It, NumIts, ConstraintSet.Key, ConstraintSet.Value, *IslandSolver);
+					for (const TPair<int32, int32>& ConstraintSet : IslandConstraintSets)
+					{
+						bNeedsAnotherIteration |= Constraints.ApplyPhase1Parallel(Dt, It, NumIts, ConstraintSet.Key, ConstraintSet.Value, *IslandSolver);
+					}
 				}
 			}
+			return bNeedsAnotherIteration;
+		}
+		else
+		{
+			return Base::ApplyConstraints(Dt, Island, It, NumIts);
 		}
-		return bNeedsAnotherIteration;
 	}
 
 	template<class ConstraintType>
 	bool TPBDConstraintColorRule<ConstraintType>::ApplyPushOut(const FReal Dt, int32 Island, const int32 It, const int32 NumIts)
 	{
-		bool bNeedsAnotherIteration = false;
-		if(FPBDIslandSolver* IslandSolver = ConstraintGraph->GetSolverIsland(Island))
+		if (IsSortingEnabled())
 		{
-			const TArray<TPair<int32, int32>>& IslandConstraintSets = ConstraintSets[Island];
-			if (!bUseColor)
+			bool bNeedsAnotherIteration = false;
+			if (FPBDIslandSolver* IslandSolver = ConstraintGraph->GetSolverIsland(Island))
 			{
-				for (const TPair<int32, int32>& ConstraintSet : IslandConstraintSets)
+				const TArray<TPair<int32, int32>>& IslandConstraintSets = ConstraintSets[Island];
+				if (!IsSortingUsingColors())
 				{
-					bNeedsAnotherIteration |= Constraints.ApplyPhase2Serial(Dt, It, NumIts, ConstraintSet.Key, ConstraintSet.Value, *IslandSolver);
+					for (const TPair<int32, int32>& ConstraintSet : IslandConstraintSets)
+					{
+						bNeedsAnotherIteration |= Constraints.ApplyPhase2Serial(Dt, It, NumIts, ConstraintSet.Key, ConstraintSet.Value, *IslandSolver);
+					}
 				}
-			}
-			else
-			{
-				for (const TPair<int32, int32>& ConstraintSet : IslandConstraintSets)
+				else
 				{
-					bNeedsAnotherIteration |= Constraints.ApplyPhase2Parallel(Dt, It, NumIts, ConstraintSet.Key, ConstraintSet.Value, *IslandSolver);
+					for (const TPair<int32, int32>& ConstraintSet : IslandConstraintSets)
+					{
+						bNeedsAnotherIteration |= Constraints.ApplyPhase2Parallel(Dt, It, NumIts, ConstraintSet.Key, ConstraintSet.Value, *IslandSolver);
+					}
 				}
 			}
+			return bNeedsAnotherIteration;
+		}
+		else
+		{
+			return Base::ApplyPushOut(Dt, Island, It, NumIts);
 		}
-		return bNeedsAnotherIteration;
 	}
 
 	template<class ConstraintType>
 	void TPBDConstraintColorRule<ConstraintType>::InitializeAccelerationStructures()
 	{
-		GraphColor.InitializeColor(*ConstraintGraph);
-
-		ConstraintSets.SetNum(GraphColor.NumIslands());
-		ConstraintGraph->template AddConstraintDatas<ConstraintType>(Constraints.GetContainerId());
+		if (IsSortingEnabled())
+		{
+			ConstraintSets.SetNum(ConstraintGraph->NumIslands());
+			ConstraintGraph->template AddConstraintDatas<ConstraintType>(Constraints.GetContainerId());
+		}
+		else
+		{
+			Base::InitializeAccelerationStructures();
+		}
 	}
 
 	template<class ConstraintType>
 	void TPBDConstraintColorRule<ConstraintType>::UpdateAccelerationStructures(const FReal Dt, const int32 Island)
+	{}
+	
+	template<class ConstraintType>
+	void TPBDConstraintColorRule<ConstraintType>::SortConstraints()
 	{
-		if(FPBDIslandSolver* IslandSolver = ConstraintGraph->GetSolverIsland(Island))
+		// Compute levels for each constraints
+		ComputeLevels();
+
+		// Compute colors for each constraints
+		ComputeColors();
+
+		// Populate the sorted constraints based on the island/level/color
+		PopulateConstraints();
+	}
+	
+	template<class ConstraintType>
+    bool TPBDConstraintColorRule<ConstraintType>::IsSortingEnabled() const 
+    {
+        return IsSortingUsingColors() || IsSortingUsingLevels();
+    }
+	
+	template<class ConstraintType>
+    bool TPBDConstraintColorRule<ConstraintType>::IsSortingUsingLevels() const 
+    {
+    	return !ChaosDisableIslandLevels;
+    }
+    
+    template<class ConstraintType>
+    bool TPBDConstraintColorRule<ConstraintType>::IsSortingUsingColors() const
+    {
+    	return !ChaosDisableIslandColors;
+    }
+	
+	template<class ConstraintType>
+	void TPBDConstraintColorRule<ConstraintType>::ComputeLevels()
+	{
+		if (IsSortingUsingLevels())
 		{
-			// Decide if we want to attempt to solve the constraints in parallel - this requires coloring the graph which
-			// is expensive, so we need to get a good number of contacts per color (per level) for it to be worthwhile.
-			// @todo(chaos): if enabled here, it should get disabled again if the level/coloring doesn't produce enough parallelism
-			bUseColor = (ConstraintGraph->GetIslandParticles(Island).Num() > ChaosCollisionColorMinParticles);
+			ConstraintGraph->GetIslandGraph()->ComputeLevels(Constraints.GetContainerId());
+		}
+	}
 
-			// @todo(chaos): disable coloring, but still generate levels and ordering, when we don't need parallelism
-			GraphColor.ComputeColor(Dt, Island, *ConstraintGraph, Constraints.GetContainerId());
+	template<class ConstraintType>
+	void TPBDConstraintColorRule<ConstraintType>::ComputeColors()
+	{
+		if (IsSortingUsingColors())
+		{
+			ConstraintGraph->GetIslandGraph()->ComputeColors(Constraints.GetContainerId(), ChaosCollisionColorMinParticles);
 		}
 	}
 
 	template<class ConstraintType>
-	void TPBDConstraintColorRule<ConstraintType>::SetUseContactGraph(const bool bInUseContactGraph)
+	void TPBDConstraintColorRule<ConstraintType>::ForEachEdges(TFunctionRef<void(const int32, const FPBDConstraintGraph::GraphType::FGraphEdge&)> InFunction)
 	{
-		GraphColor.SetUseContactGraph(bInUseContactGraph);
+		for(auto& GraphEdge : ConstraintGraph->GetIslandGraph()->GraphEdges)
+		{
+			auto& GraphIsland = ConstraintGraph->GetIslandGraph()->GraphIslands[GraphEdge.IslandIndex];
+			if(GraphEdge.ItemContainer == Constraints.GetContainerId() && !GraphIsland.IslandItem->IsSleeping())
+			{
+				const int32 EdgeColor = IsSortingUsingColors() ? FMath::Max(0,GraphEdge.ColorIndex) : 0;
+				const int32 EdgeLevel = IsSortingUsingLevels() ? FMath::Max(0,GraphEdge.LevelIndex) : 0;
+				const int32 MaxColors = IsSortingUsingColors() ? FMath::Max(1,GraphIsland.MaxColors+1) : 1;
+
+				const int32 IslandOffset = IslandOffsets[GraphIsland.IslandItem->GetIslandIndex()];
+				const int32 OffsetIndex = IslandOffset + EdgeLevel * MaxColors + EdgeColor;
+
+				InFunction(OffsetIndex,GraphEdge);
+			}
+		}
+	}
+
+	template<class ConstraintType>
+	void TPBDConstraintColorRule<ConstraintType>::PopulateConstraints()
+	{
+		if(IsSortingEnabled())
+		{
+			SortedConstraints.SetNum(ConstraintGraph->GetIslandGraph()->GraphEdges.Num(), false);
+			IslandOffsets.SetNum(ConstraintGraph->NumIslands(), false);
+
+			// We first fill the island offsets that will be used in the gather
+			int32 IslandOffset = 0;
+			for(auto& GraphIsland : ConstraintGraph->GetIslandGraph()->GraphIslands)
+			{
+				if(!GraphIsland.IslandItem->IsSleeping())
+				{
+					IslandOffsets[GraphIsland.IslandItem->GetIslandIndex()] = IslandOffset;
+					
+					const int32 MaxColors = IsSortingUsingColors() ? FMath::Max(1,GraphIsland.MaxColors+1) : 1;
+					const int32 MaxLevels = IsSortingUsingLevels() ? FMath::Max(1,GraphIsland.MaxLevels+1) : 1;
+
+					IslandOffset += MaxLevels * MaxColors;
+				}
+			}
+			// Initialization of the constraint offsets and the offsets counters
+			ConstraintOffsets.SetNum(IslandOffset, false);
+			OffsetCounters.SetNum(IslandOffset, false);
+			
+			FMemory::Memzero(ConstraintOffsets.GetData(), ConstraintOffsets.Num() * sizeof(int32));
+			FMemory::Memzero(OffsetCounters.GetData(), ConstraintOffsets.Num() * sizeof(int32));
+			
+			auto& LocalConstraintOffsets = ConstraintOffsets;
+			auto& LocalOffsetCounters = OffsetCounters;
+			auto& LocalSortedConstraints = SortedConstraints;
+
+			// Build of the constraint offsets that will be used to know where the sorted constraints will
+			// be inserted into the flat array  
+			ForEachEdges( [&LocalConstraintOffsets](
+				const int32 OffsetIndex, const FPBDConstraintGraph::GraphType::FGraphEdge& GraphEdge)
+			{
+				++LocalConstraintOffsets[OffsetIndex];
+			});
+
+			for(int32 OffsetIndex = 0, OffsetEnd = ConstraintOffsets.Num()-1; OffsetIndex < OffsetEnd; ++OffsetIndex)
+			{
+				ConstraintOffsets[OffsetIndex+1] += ConstraintOffsets[OffsetIndex];
+			}
+
+			// Insert the constraint handles in the right order : Island|Level|Color to be processed by the Gather
+			ForEachEdges( [&LocalConstraintOffsets, &LocalOffsetCounters, &LocalSortedConstraints](
+				const int32 OffsetIndex, const FPBDConstraintGraph::GraphType::FGraphEdge& GraphEdge)
+			{
+				const int32 OffsetValue = (OffsetIndex == 0) ? 0 : LocalConstraintOffsets[OffsetIndex-1];
+				LocalSortedConstraints[OffsetValue+LocalOffsetCounters[OffsetIndex]] = GraphEdge.EdgeItem->As<FConstraintContainerHandle>();
+				
+				++LocalOffsetCounters[OffsetIndex];
+			});
+		}
 	}
 
+	template<class ConstraintType>
+	void TPBDConstraintColorRule<ConstraintType>::SetUseContactGraph(const bool bInUseContactGraph)
+	{}
+
 
 	template class TSimpleConstraintRule<FPBDCollisionConstraints>;
 	template class TSimpleConstraintRule<FPBDJointConstraints>;
diff --git a/Engine/Source/Runtime/Experimental/Chaos/Private/Chaos/PBDRigidsEvolutionGBF.cpp b/Engine/Source/Runtime/Experimental/Chaos/Private/Chaos/PBDRigidsEvolutionGBF.cpp
--- a/Engine/Source/Runtime/Experimental/Chaos/Private/Chaos/PBDRigidsEvolutionGBF.cpp
+++ b/Engine/Source/Runtime/Experimental/Chaos/Private/Chaos/PBDRigidsEvolutionGBF.cpp
@@ -100,6 +100,7 @@ namespace Chaos
 		DECLARE_CYCLE_STAT(TEXT("FPBDRigidsEvolutionGBF::ComputeIntermediateSpatialAcceleration"), STAT_Evolution_ComputeIntermediateSpatialAcceleration, STATGROUP_Chaos);
 		DECLARE_CYCLE_STAT(TEXT("FPBDRigidsEvolutionGBF::CreateConstraintGraph"), STAT_Evolution_CreateConstraintGraph, STATGROUP_Chaos);
 		DECLARE_CYCLE_STAT(TEXT("FPBDRigidsEvolutionGBF::CreateIslands"), STAT_Evolution_CreateIslands, STATGROUP_Chaos);
+		DECLARE_CYCLE_STAT(TEXT("FPBDRigidsEvolutionGBF::UpdateAccelerationStructures"), STAT_Evolution_UpdateAccelerationStructures, STATGROUP_Chaos);
 		DECLARE_CYCLE_STAT(TEXT("FPBDRigidsEvolutionGBF::AddSleepingContacts"), STAT_Evolution_AddSleepingContacts, STATGROUP_Chaos);
 		DECLARE_CYCLE_STAT(TEXT("FPBDRigidsEvolutionGBF::PreApplyCallback"), STAT_Evolution_PreApplyCallback, STATGROUP_Chaos);
 		DECLARE_CYCLE_STAT(TEXT("FPBDRigidsEvolutionGBF::ParallelSolve"), STAT_Evolution_ParallelSolve, STATGROUP_Chaos);
@@ -282,12 +283,18 @@ void FPBDRigidsEvolutionGBF::AdvanceOneTimeStepImpl(const FReal Dt, const FSubSt
 		CSV_SCOPED_TIMING_STAT(PhysicsVerbose, StepSolver_CreateConstraintGraph);
 		CreateConstraintGraph();
 	}
-	//CollisionDetector.GetCollisionContainer().GetConstraintAllocator().SortConstraintsHandles();
 	{
 		SCOPE_CYCLE_COUNTER(STAT_Evolution_CreateIslands);
 		CSV_SCOPED_TIMING_STAT(PhysicsVerbose, StepSolver_CreateIslands);
 		CreateIslands();
 	}
+	{
+		// We keep the graph color stat name for now to compare with previous implementation
+		// @todo : change the name to sort constraints
+		SCOPE_CYCLE_COUNTER(STAT_Evolution_GraphColor);
+		CSV_SCOPED_TIMING_STAT(PhysicsVerbose, StepSolver_GraphColor);
+		SortConstraints();
+	}
 	if (PreApplyCallback != nullptr)
 	{
 		SCOPE_CYCLE_COUNTER(STAT_Evolution_PreApplyCallback);
@@ -333,8 +340,8 @@ void FPBDRigidsEvolutionGBF::AdvanceOneTimeStepImpl(const FReal Dt, const FSubSt
 			{
 				// Update constraint graphs, coloring etc as required by the different constraint types in this island
 				{
-					SCOPE_CYCLE_COUNTER(STAT_Evolution_GraphColor);
-					CSV_SCOPED_TIMING_STAT(PhysicsVerbose, StepSolver_GraphColor);
+					SCOPE_CYCLE_COUNTER(STAT_Evolution_UpdateAccelerationStructures);
+					CSV_SCOPED_TIMING_STAT(PhysicsVerbose, STAT_Evolution_UpdateAccelerationStructures);
 					UpdateAccelerationStructures(Dt, Island);
 				}
 
diff --git a/Engine/Source/Runtime/Experimental/Chaos/Public/Chaos/Evolution/SolverDatas.h b/Engine/Source/Runtime/Experimental/Chaos/Public/Chaos/Evolution/SolverDatas.h
--- a/Engine/Source/Runtime/Experimental/Chaos/Public/Chaos/Evolution/SolverDatas.h
+++ b/Engine/Source/Runtime/Experimental/Chaos/Public/Chaos/Evolution/SolverDatas.h
@@ -32,7 +32,7 @@ public :
 	FORCEINLINE FSolverBodyContainer& GetBodyContainer() { return *BodyContainer.Get(); };
 
 	/**
-	 * @brief The number of constraint containers registsred
+	 * @brief The number of constraint containers registered
 	*/
 	int32 NumConstraintContainerIds() const { return ConstraintDatas.Num(); }
 
diff --git a/Engine/Source/Runtime/Experimental/Chaos/Public/Chaos/Island/IslandGraph.h b/Engine/Source/Runtime/Experimental/Chaos/Public/Chaos/Island/IslandGraph.h
--- a/Engine/Source/Runtime/Experimental/Chaos/Public/Chaos/Island/IslandGraph.h
+++ b/Engine/Source/Runtime/Experimental/Chaos/Public/Chaos/Island/IslandGraph.h
@@ -142,6 +142,12 @@ public:
 
 		/** Node item that is stored per node */
 		NodeType NodeItem;
+		
+		/** Node level index */
+		int32 LevelIndex = INDEX_NONE;
+
+		/** Set of used colors */
+		TSet<int32> ColorIndices;
 	};
 
 	/** Graph edge structure */
@@ -170,6 +176,12 @@ public:
 
 		/** Item Container Id  */
 		int32 ItemContainer = 0;
+
+		/** Edge level index */
+		int32 LevelIndex = INDEX_NONE;
+
+		/** Edge Color index */
+		int32 ColorIndex = INDEX_NONE;
 	};
 
 	/** Graph island structure */
@@ -198,6 +210,12 @@ public:
 		
 		/** Island Item that is stored per island*/
 		IslandType IslandItem;
+
+		/** Max levels per island */
+		int32 MaxLevels = INDEX_NONE;
+
+		/** Max color per island */
+		int32 MaxColors = INDEX_NONE;
 	};
 
 	/**
@@ -242,8 +260,53 @@ public:
 	* Reassign the updated island index from the merging phase to the nodes/edges
 	*/
 	void ReassignIslands();
+
+	/**
+	* Init all the islands/edges/nodes levels and colors for sorting
+	*/
+	void InitSorting();
+
+	/**
+	* Given a constraint containerId compute all the islands/edges/nodes levels
+	* @param ContainerId Constraint container id that will be compared to the item container
+	*/
+	void ComputeLevels(const int32 ContainerId);
 	
+	/**
+	* Given a constraint containerId and a node index update all the islands/edges/nodes levels and push the connected nodes
+	* into the node queue
+	* @param NodeIndex Node index that we are currently iterating over
+	* @param ContainerId Constraint container id that will be compared to the item container
+	* @param NodeQueue Node queue in which all the connected nodes not processed yet will be added
+	*/
+	void UpdateLevels(const int32 NodeIndex, const int32 ContainerId, TQueue<int32>& NodeQueue);
+
+	/**
+	* Given a constraint containerId compute all the islands/edges/nodes colors
+	* @param ContainerId Constraint container id that will be compared to the item container
+	* @param MinEdges Minimum number of edges to compute coloring
+	*/
+	void ComputeColors(const int32 ContainerId, const int32 MinEdges);
 
+	/**
+	* Given a constraint containerId and a node index update all the islands/edges/nodes colors and push the connected nodes
+	* into the node queue
+	* @param NodeIndex Node index that we are currently iterating over
+	* @param ContainerId Constraint container id that will be compared to the item container
+	* @param NodeQueue Node queue in which all the connected nodes not processed yet will be added
+	* @param MinEdges Minimum number of edges to compute coloring
+	*/
+	void UpdateColors(const int32 NodeIndex, const int32 ContainerId, TQueue<int32>& NodeQueue, const int32 MinEdges);
+
+	/**
+	* Pick the first available color that is not used yet by the current graph node and the edge opposite one
+	* @param GraphNode Graph node that we currently iterate over
+	* @param OtherIndex Index of the edge opposite node
+	* @param NodeQueue Node queue in which all the connected nodes not processed yet will be added
+	* @return First available color 
+	*/
+	int32 PickColor(const FGraphNode& GraphNode, const int32 OtherIndex, TQueue<int32>& NodeQueue);
+	
 	/** List of graph nodes */
 	TSparseArray<FGraphNode> GraphNodes;
 
diff --git a/Engine/Source/Runtime/Experimental/Chaos/Public/Chaos/Island/IslandManager.h b/Engine/Source/Runtime/Experimental/Chaos/Public/Chaos/Island/IslandManager.h
--- a/Engine/Source/Runtime/Experimental/Chaos/Public/Chaos/Island/IslandManager.h
+++ b/Engine/Source/Runtime/Experimental/Chaos/Public/Chaos/Island/IslandManager.h
@@ -197,6 +197,14 @@ public:
 	{
 		return IslandGraph.Get();
 	}
+	/**
+	* Get the island graph
+	* @return Island graph of the manager
+	*/
+	FORCEINLINE GraphType* GetIslandGraph()
+	{
+		return IslandGraph.Get();
+	}
 
 	/**
 	 * Get Max UniqueIdx of the particles.
diff --git a/Engine/Source/Runtime/Experimental/Chaos/Public/Chaos/PBDConstraintRule.h b/Engine/Source/Runtime/Experimental/Chaos/Public/Chaos/PBDConstraintRule.h
--- a/Engine/Source/Runtime/Experimental/Chaos/Public/Chaos/PBDConstraintRule.h
+++ b/Engine/Source/Runtime/Experimental/Chaos/Public/Chaos/PBDConstraintRule.h
@@ -138,6 +138,12 @@ namespace Chaos
 		/** Set up the perf-acceleration structures for the specified island. May be called in parallel for islands */
 		virtual void UpdateAccelerationStructures(const FReal Dt, const int32 Island) {}
 
+		/** Sort constraints if necessary  */
+		virtual void SortConstraints() {}
+
+		/** Boolean to check if we need to sort the constraints */
+		virtual bool IsSortingEnabled() const {return false;}
+
 		virtual void SetUseContactGraph(const bool InUseContactGraph) {}
 
 		/** Change enabled state on all constraints associated with the specified particles */
@@ -217,6 +223,7 @@ namespace Chaos
 	template<typename ConstraintType>
 	class CHAOS_API TPBDConstraintIslandRule : public TPBDConstraintGraphRuleImpl<ConstraintType>
 	{
+	protected:
 		typedef TPBDConstraintGraphRuleImpl<ConstraintType> Base;
 
 	public:
@@ -242,7 +249,7 @@ namespace Chaos
 		
 		virtual void UpdateAccelerationStructures(const FReal Dt, const int32 Island) override;
 
-	private:
+	protected:
 		using Base::Constraints;
 		using Base::ConstraintGraph;
 	};
@@ -254,9 +261,9 @@ namespace Chaos
 	 * place as far as higher-level constraints are concerned.
 	 */
 	template<typename ConstraintType>
-	class CHAOS_API TPBDConstraintColorRule : public TPBDConstraintGraphRuleImpl<ConstraintType>
+	class CHAOS_API TPBDConstraintColorRule : public TPBDConstraintIslandRule<ConstraintType>
 	{
-		typedef TPBDConstraintGraphRuleImpl<ConstraintType> Base;
+		typedef TPBDConstraintIslandRule<ConstraintType> Base;
 
 	public:
 		using FConstraints = ConstraintType;
@@ -293,31 +300,51 @@ namespace Chaos
 
 		virtual void UpdateAccelerationStructures(const FReal Dt, const int32 Island) override;
 
-		virtual void SetUseContactGraph(const bool bInUseContactGraph) override;
+		/** Sort constraints according to island/level/color*/
+		virtual void SortConstraints() override;
+		
+		/** Boolean to check if we need to sort the constraints*/
+		virtual bool IsSortingEnabled() const override;
 
-		const FPBDConstraintColor GetGraphColor() const
-		{
-			return GraphColor;
-		}
+		virtual void SetUseContactGraph(const bool bInUseContactGraph) override;
 
 	private:
-		using Base::Constraints;
-		using Base::ConstraintGraph;
+		using Base::Base::Constraints;
+		using Base::Base::ConstraintGraph;
 
-		const TArray<FConstraintContainerHandle*>& GetLevelColorConstraints(const typename FPBDConstraintColor::FLevelToColorToConstraintListMap& LevelToColorToConstraintListMap, int32 Level, int32 Color) const
-		{
-			// FPBDConstraintColor works with any constraint type (in principle - currently only used with Collisions), but the rule is bound to a single type and so this cast is ok
-			return reinterpret_cast<const TArray<FConstraintContainerHandle*>&>(LevelToColorToConstraintListMap[Level][Color]);
-		}
+		/** Check if sorting is using levels*/
+		bool IsSortingUsingColors() const;
+
+		/** Check if sorting is using colors */
+		bool IsSortingUsingLevels() const;
+
+		/** Compute island levels if necessary */
+		void ComputeLevels();
 
-		FPBDConstraintColor GraphColor;
+		/** Compute island colors if necessary */
+		void ComputeColors();
 
+		/** Populate the sorted constraints list based on island/level/color */
+		void PopulateConstraints();
+		
+		/** Loop over all the edges and apply a function */
+		void ForEachEdges(TFunctionRef<void(const int32, const FPBDConstraintGraph::GraphType::FGraphEdge&)> InFunction);
+		
 		// Each array entry contains the [begin, end) index of a set of independent constraints
 		// that can be solved in parallel. The sets are ordered by level and must be solved sequentially.
 		TArray<TArray<TPair<int32, int32>>> ConstraintSets;
 
-		// Whether to enable graph coloring and parallel processing of same-color constraints
-		bool bUseColor;
+		// Sorted constraint handles
+		TArray<FConstraintContainerHandle*> SortedConstraints;
+
+		// Constraint offsets into the sorted handles list for a given tuple island/level/color
+		TArray<int32> ConstraintOffsets;
+
+		// Island offsets into the constraint offsets 
+		TArray<int32> IslandOffsets;
+
+		// Counters to know at which position after the island/level/color constraint offset the handle will be inserted
+		TArray<int32> OffsetCounters;
 	};
 
 }
diff --git a/Engine/Source/Runtime/Experimental/Chaos/Public/Chaos/PBDRigidsEvolution.h b/Engine/Source/Runtime/Experimental/Chaos/Public/Chaos/PBDRigidsEvolution.h
--- a/Engine/Source/Runtime/Experimental/Chaos/Public/Chaos/PBDRigidsEvolution.h
+++ b/Engine/Source/Runtime/Experimental/Chaos/Public/Chaos/PBDRigidsEvolution.h
@@ -888,6 +888,24 @@ protected:
 		}
 	}
 	
+	/** Sort constraints if necessary */
+	void SortConstraints()
+	{
+		bool bNeedsSorting = false;
+		for (FPBDConstraintGraphRule* ConstraintRule : ConstraintRules)
+		{
+			bNeedsSorting |= ConstraintRule->IsSortingEnabled();
+		}
+		if(bNeedsSorting)
+		{
+			ConstraintGraph.GetIslandGraph()->InitSorting();
+			for (FPBDConstraintGraphRule* ConstraintRule : ConstraintRules)
+			{
+				ConstraintRule->SortConstraints();
+			}
+		}
+	}
+	
 	void FlushInternalAccelerationQueue();
 	void FlushAsyncAccelerationQueue();
 	void WaitOnAccelerationStructure();
