Added one shot manifold support for Box/Heightfield and Box/Trimesh via the existing convex paths
Fixed manifold generation against uniform scaled trimesh.

diff --git a/Engine/Source/Runtime/Experimental/Chaos/Private/Chaos/Collision/ContactPointsMiscShapes.cpp b/Engine/Source/Runtime/Experimental/Chaos/Private/Chaos/Collision/ContactPointsMiscShapes.cpp
--- a/Engine/Source/Runtime/Experimental/Chaos/Private/Chaos/Collision/ContactPointsMiscShapes.cpp
+++ b/Engine/Source/Runtime/Experimental/Chaos/Private/Chaos/Collision/ContactPointsMiscShapes.cpp
@@ -456,4 +456,9 @@ namespace Chaos
 
 	template void GJKImplicitManifold<FCapsule, TImplicitObjectScaled<class Chaos::FTriangleMeshImplicitObject, 1>>(const FImplicitObject& A, const FRigidTransform3& ATransform, const TImplicitObjectScaled<class Chaos::FTriangleMeshImplicitObject, 1>& B, const FRigidTransform3& BTransform, const FReal CullDistance, const FReal ShapePadding, TArray<FContactPoint>& ContactPoints);
 	template void GJKImplicitManifold<FCapsule, FTriangleMeshImplicitObject>(const FImplicitObject& A, const FRigidTransform3& ATransform, const FTriangleMeshImplicitObject& B, const FRigidTransform3& BTransform, const FReal CullDistance, const FReal ShapePadding, TArray<FContactPoint>& ContactPoints);
+
+	template void GJKImplicitManifold<FImplicitBox3, FHeightField>(const FImplicitObject& A, const FRigidTransform3& ATransform, const FHeightField& B, const FRigidTransform3& BTransform, const FReal CullDistance, const FReal ShapePadding, TArray<FContactPoint>& ContactPoints);
+
+	template void GJKImplicitManifold<FImplicitBox3, TImplicitObjectScaled<class Chaos::FTriangleMeshImplicitObject, 1>>(const FImplicitObject& A, const FRigidTransform3& ATransform, const TImplicitObjectScaled<class Chaos::FTriangleMeshImplicitObject, 1>& B, const FRigidTransform3& BTransform, const FReal CullDistance, const FReal ShapePadding, TArray<FContactPoint>& ContactPoints);
+	template void GJKImplicitManifold<FImplicitBox3, FTriangleMeshImplicitObject>(const FImplicitObject& A, const FRigidTransform3& ATransform, const FTriangleMeshImplicitObject& B, const FRigidTransform3& BTransform, const FReal CullDistance, const FReal ShapePadding, TArray<FContactPoint>& ContactPoints);
 }
diff --git a/Engine/Source/Runtime/Experimental/Chaos/Private/Chaos/CollisionOneShotManifoldsMiscShapes.cpp b/Engine/Source/Runtime/Experimental/Chaos/Private/Chaos/CollisionOneShotManifoldsMiscShapes.cpp
--- a/Engine/Source/Runtime/Experimental/Chaos/Private/Chaos/CollisionOneShotManifoldsMiscShapes.cpp
+++ b/Engine/Source/Runtime/Experimental/Chaos/Private/Chaos/CollisionOneShotManifoldsMiscShapes.cpp
@@ -257,6 +257,7 @@ namespace Chaos
 			}
 		}
 
+		template<typename ConvexType>
 		void ConstructConvexHeightFieldOneShotManifold(const FImplicitObject& Convex, const FRigidTransform3& ConvexTransform, const FHeightField& HeightField, const FRigidTransform3& HeightFieldTransform, const FReal Dt, FPBDCollisionConstraint& Constraint)
 		{
 			// We only build one shot manifolds once
@@ -268,14 +269,14 @@ namespace Chaos
 			//FContactPoint ContactPoint = ConvexHeightFieldContactPoint(Convex, ConvexTransform, HeightField, HeightFieldTransform, Constraint.GetCullDistance(), 0.0f);
 
 			TArray<FContactPoint> ContactPoints;
-			GJKImplicitManifold<FConvex>(Convex, ConvexTransform, HeightField, HeightFieldTransform, Constraint.GetCullDistance(), 0.0f, ContactPoints);
+			GJKImplicitManifold<ConvexType>(Convex, ConvexTransform, HeightField, HeightFieldTransform, Constraint.GetCullDistance(), 0.0f, ContactPoints);
 			for (FContactPoint& ContactPoint : ContactPoints)
 			{
 				Constraint.AddOneshotManifoldContact(ContactPoint, Dt);
 			}
 		}
 
-		template <typename TriMeshType>
+		template <typename ConvexType, typename TriMeshType>
 		void ConstructConvexTriMeshOneShotManifold(const FImplicitObject& Convex, const FRigidTransform3& ConvexTransform, const TriMeshType& TriangleMesh, const FRigidTransform3& TriMeshTransform, const FReal Dt, FPBDCollisionConstraint& Constraint)
 		{
 			// We only build one shot manifolds once
@@ -285,7 +286,7 @@ namespace Chaos
 			ensure(TriMeshTransform.GetScale3D() == FVec3(1.0f, 1.0f, 1.0f));
 
 			TArray<FContactPoint> ContactPoints;
-			GJKImplicitManifold<FConvex, TriMeshType>(Convex, ConvexTransform, TriangleMesh, TriMeshTransform, Constraint.GetCullDistance(), 0.0f, ContactPoints);
+			GJKImplicitManifold<ConvexType, TriMeshType>(Convex, ConvexTransform, TriangleMesh, TriMeshTransform, Constraint.GetCullDistance(), 0.0f, ContactPoints);
 			for (FContactPoint& ContactPoint : ContactPoints)
 			{
 				Constraint.AddOneshotManifoldContact(ContactPoint, Dt);
@@ -298,8 +299,13 @@ namespace Chaos
 		template void ConstructCapsuleTriMeshOneShotManifold<TImplicitObjectScaled<class Chaos::FTriangleMeshImplicitObject, 1>>(const FCapsule& Capsule, const FRigidTransform3& CapsuleWorldTransform, const TImplicitObjectScaled<class Chaos::FTriangleMeshImplicitObject, 1>& TriangleMesh, const FRigidTransform3& TriMeshWorldTransform, const FReal Dt, FPBDCollisionConstraint& Constraint);
 		template void ConstructCapsuleTriMeshOneShotManifold<FTriangleMeshImplicitObject>(const FCapsule& Capsule, const FRigidTransform3& CapsuleWorldTransform, const FTriangleMeshImplicitObject& TriangleMesh, const FRigidTransform3& TriMeshWorldTransform, const FReal Dt, FPBDCollisionConstraint& Constraint);
 
-		template void ConstructConvexTriMeshOneShotManifold<TImplicitObjectScaled<class Chaos::FTriangleMeshImplicitObject, 1>>(const FImplicitObject& Convex, const FRigidTransform3& ConvexTransform, const TImplicitObjectScaled<class Chaos::FTriangleMeshImplicitObject, 1>& TriangleMesh, const FRigidTransform3& TriMeshTransform, const FReal Dt, FPBDCollisionConstraint& Constraint);
-		template void ConstructConvexTriMeshOneShotManifold<FTriangleMeshImplicitObject>(const FImplicitObject& Convex, const FRigidTransform3& ConvexTransform, const FTriangleMeshImplicitObject& TriangleMesh, const FRigidTransform3& TriMeshTransform, const FReal Dt, FPBDCollisionConstraint& Constraint);
+		template void ConstructConvexTriMeshOneShotManifold<FConvex, TImplicitObjectScaled<class Chaos::FTriangleMeshImplicitObject, 1>>(const FImplicitObject& Convex, const FRigidTransform3& ConvexTransform, const TImplicitObjectScaled<class Chaos::FTriangleMeshImplicitObject, 1>& TriangleMesh, const FRigidTransform3& TriMeshTransform, const FReal Dt, FPBDCollisionConstraint& Constraint);
+		template void ConstructConvexTriMeshOneShotManifold<FConvex, FTriangleMeshImplicitObject>(const FImplicitObject& Convex, const FRigidTransform3& ConvexTransform, const FTriangleMeshImplicitObject& TriangleMesh, const FRigidTransform3& TriMeshTransform, const FReal Dt, FPBDCollisionConstraint& Constraint);
+		template void ConstructConvexTriMeshOneShotManifold<FImplicitBox3, TImplicitObjectScaled<class Chaos::FTriangleMeshImplicitObject, 1>>(const FImplicitObject& Convex, const FRigidTransform3& ConvexTransform, const TImplicitObjectScaled<class Chaos::FTriangleMeshImplicitObject, 1>& TriangleMesh, const FRigidTransform3& TriMeshTransform, const FReal Dt, FPBDCollisionConstraint& Constraint);
+		template void ConstructConvexTriMeshOneShotManifold<FImplicitBox3, FTriangleMeshImplicitObject>(const FImplicitObject& Convex, const FRigidTransform3& ConvexTransform, const FTriangleMeshImplicitObject& TriangleMesh, const FRigidTransform3& TriMeshTransform, const FReal Dt, FPBDCollisionConstraint& Constraint);
+
+		template void ConstructConvexHeightFieldOneShotManifold<FConvex>(const FImplicitObject& Convex, const FRigidTransform3& ConvexTransform, const FHeightField& HeightField, const FRigidTransform3& HeightFieldTransform, const FReal Dt, FPBDCollisionConstraint& Constraint);
+		template void ConstructConvexHeightFieldOneShotManifold<FImplicitBox3>(const FImplicitObject& Convex, const FRigidTransform3& ConvexTransform, const FHeightField& HeightField, const FRigidTransform3& HeightFieldTransform, const FReal Dt, FPBDCollisionConstraint& Constraint);
 	}
 }
 
diff --git a/Engine/Source/Runtime/Experimental/Chaos/Private/Chaos/CollisionOneShotManifoldsMiscShapes.h b/Engine/Source/Runtime/Experimental/Chaos/Private/Chaos/CollisionOneShotManifoldsMiscShapes.h
--- a/Engine/Source/Runtime/Experimental/Chaos/Private/Chaos/CollisionOneShotManifoldsMiscShapes.h
+++ b/Engine/Source/Runtime/Experimental/Chaos/Private/Chaos/CollisionOneShotManifoldsMiscShapes.h
@@ -94,6 +94,7 @@ namespace Chaos
 			const FReal Dt, 
 			FPBDCollisionConstraint& Constraint);
 
+		template<typename ConvexType = FConvex>
 		void ConstructConvexHeightFieldOneShotManifold(
 			const FImplicitObject& Convex,
 			const FRigidTransform3& ConvexTransform,
@@ -102,7 +103,7 @@ namespace Chaos
 			const FReal Dt,
 			FPBDCollisionConstraint& Constraint);
 
-		template <typename TriMeshType>
+		template <typename ConvexType = FConvex, typename TriMeshType>
 		void ConstructConvexTriMeshOneShotManifold(
 			const FImplicitObject& Convex, 
 			const FRigidTransform3& ConvexTransform, 
diff --git a/Engine/Source/Runtime/Experimental/Chaos/Private/Chaos/CollisionResolution.cpp b/Engine/Source/Runtime/Experimental/Chaos/Private/Chaos/CollisionResolution.cpp
--- a/Engine/Source/Runtime/Experimental/Chaos/Private/Chaos/CollisionResolution.cpp
+++ b/Engine/Source/Runtime/Experimental/Chaos/Private/Chaos/CollisionResolution.cpp
@@ -567,8 +567,22 @@ namespace Chaos
 
 		void UpdateBoxHeightFieldConstraint(const FImplicitBox3& A, const FRigidTransform3& ATransform, const FHeightField& B, const FRigidTransform3& BTransform, const FReal Dt, FPBDCollisionConstraint& Constraint)
 		{
-			// @todo(chaos): restitutionpadding
-			UpdateContactPoint(Constraint, BoxHeightFieldContactPoint(A, ATransform, B, BTransform, Constraint.GetCullDistance(), 0.0f), Dt);
+			if(Constraint.GetUseManifold())
+			{
+				if(Constraint.GetManifoldPoints().Num() == 0)
+				{
+					ConstructConvexHeightFieldOneShotManifold<FImplicitBox3>(A, ATransform, B, BTransform, Dt, Constraint);
+				}
+				else
+				{
+					Constraint.UpdateManifoldContacts();
+				}
+			}
+			else
+			{
+				// @todo(chaos): restitutionpadding
+				UpdateContactPoint(Constraint, BoxHeightFieldContactPoint(A, ATransform, B, BTransform, Constraint.GetCullDistance(), 0.0f), Dt);
+			}
 		}
 
 		template<typename T_TRAITS>
@@ -691,8 +705,22 @@ namespace Chaos
 		template <typename TriMeshType>
 		void UpdateBoxTriangleMeshConstraint(const FImplicitBox3& Box0, const FRigidTransform3& WorldTransform0, const TriMeshType& TriangleMesh1, const FRigidTransform3& WorldTransform1, const FReal Dt, FPBDCollisionConstraint& Constraint)
 		{
-			// @toto(chaos): restitutionpadding
-			UpdateContactPoint(Constraint, BoxTriangleMeshContactPoint(Box0, WorldTransform0, TriangleMesh1, WorldTransform1, Constraint.GetCullDistance(), 0.0f), Dt);
+			if(Constraint.GetUseManifold())
+			{
+				if(Constraint.GetManifoldPoints().Num() == 0)
+				{
+					ConstructConvexTriMeshOneShotManifold<FImplicitBox3, TriMeshType>(Box0, WorldTransform0, TriangleMesh1, WorldTransform1, Dt, Constraint);
+				}
+				else
+				{
+					Constraint.UpdateManifoldContacts();
+				}
+			}
+			else
+			{
+				// @toto(chaos): restitutionpadding
+				UpdateContactPoint(Constraint, BoxTriangleMeshContactPoint(Box0, WorldTransform0, TriangleMesh1, WorldTransform1, Constraint.GetCullDistance(), 0.0f), Dt);
+			}
 		}
 
 		template<typename T_TRAITS>
diff --git a/Engine/Source/Runtime/Experimental/Chaos/Private/Chaos/TriangleMeshImplicitObject.cpp b/Engine/Source/Runtime/Experimental/Chaos/Private/Chaos/TriangleMeshImplicitObject.cpp
--- a/Engine/Source/Runtime/Experimental/Chaos/Private/Chaos/TriangleMeshImplicitObject.cpp
+++ b/Engine/Source/Runtime/Experimental/Chaos/Private/Chaos/TriangleMeshImplicitObject.cpp
@@ -373,13 +373,8 @@ bool FTriangleMeshImplicitObject::ContactManifoldImp(const GeomType& QueryGeom,
 	auto OverlapTriangle = [&](const FVec3& A, const FVec3& B, const FVec3& C,
 		FPBDCollisionConstraint& Constraint)
 	{
-		const FVec3 AB = B - A;
-		const FVec3 AC = C - A;
-
-		const FVec3 Offset = FVec3::CrossProduct(AB, AC);
-
 		FTriangle TriangleConvex(A, B, C);
-		Collisions::ConstructConvexConvexOneShotManifold(QueryGeom, QueryTM, TriangleConvex, FRigidTransform3::Identity, 0, Constraint);
+		Collisions::ConstructConvexConvexOneShotManifold(WorldScaleGeom, WorldScaleQueryTM, TriangleConvex, FRigidTransform3::Identity, 0, Constraint);
 	};
 
 	auto LambdaHelper = [&](const auto& Elements)
diff --git a/Engine/Source/Runtime/Experimental/Chaos/Public/Chaos/ImplicitObjectScaled.h b/Engine/Source/Runtime/Experimental/Chaos/Public/Chaos/ImplicitObjectScaled.h
--- a/Engine/Source/Runtime/Experimental/Chaos/Public/Chaos/ImplicitObjectScaled.h
+++ b/Engine/Source/Runtime/Experimental/Chaos/Public/Chaos/ImplicitObjectScaled.h
@@ -657,7 +657,7 @@ public:
 		const FReal UnscaledThickness = Thickness;
 
 		auto ScaledA = MakeScaledHelper(A, MInvScale);
-		return MObject->ContactManifold(ScaledA, AToBTMNoScale, UnscaledThickness, ContactPoints);
+		return MObject->ContactManifold(ScaledA, AToBTMNoScale, UnscaledThickness, ContactPoints, MScale);
 	}
 
 	/** This is a low level function and assumes the internal object has a OverlapGeom function. Should not be called directly. See GeometryQueries.h : OverlapQuery */
