clear the explodedvectors from geometry collections on beginplay, so they don't persist through simulation

#rnx

diff --git a/Engine/Source/Runtime/Experimental/GeometryCollectionEngine/Private/GeometryCollection/GeometryCollectionComponent.cpp b/Engine/Source/Runtime/Experimental/GeometryCollectionEngine/Private/GeometryCollection/GeometryCollectionComponent.cpp
--- a/Engine/Source/Runtime/Experimental/GeometryCollectionEngine/Private/GeometryCollection/GeometryCollectionComponent.cpp
+++ b/Engine/Source/Runtime/Experimental/GeometryCollectionEngine/Private/GeometryCollection/GeometryCollectionComponent.cpp
@@ -311,6 +311,13 @@ void UGeometryCollectionComponent::BeginPlay()
 	HackRegisterGeomAccelerator(*this);
 #endif
 
+#if WITH_EDITOR
+	if (RestCollection->GetGeometryCollection()->HasAttribute("ExplodedVector", FGeometryCollection::TransformGroup))
+	{
+		RestCollection->GetGeometryCollection()->RemoveAttribute("ExplodedVector", FGeometryCollection::TransformGroup);
+	}
+#endif
+
 	//////////////////////////////////////////////////////////////////////////
 	// Commenting out these callbacks for now due to the threading model. The callbacks here
 	// expect the rest collection to be mutable which is not the case when running in multiple
