Implemented material support for collision events. Simple and complex shapes now report materials to subscribed components and listeners.

diff --git a/Engine/Source/Runtime/Engine/Private/PhysicsEngine/Experimental/PhysScene_Chaos.cpp b/Engine/Source/Runtime/Engine/Private/PhysicsEngine/Experimental/PhysScene_Chaos.cpp
--- a/Engine/Source/Runtime/Engine/Private/PhysicsEngine/Experimental/PhysScene_Chaos.cpp
+++ b/Engine/Source/Runtime/Engine/Private/PhysicsEngine/Experimental/PhysScene_Chaos.cpp
@@ -47,7 +47,7 @@
 #include "EventManager.h"
 #include "RewindData.h"
 #include "Chaos/PhysicsSolverBaseImpl.h"
-
+#include "Chaos/Defines.h"
 
 #if !UE_BUILD_SHIPPING
 #include "Engine/World.h"
@@ -741,7 +741,12 @@ void FPhysScene_Chaos::HandleCollisionEvents(const Chaos::FCollisionEventData& E
 								NewContact.ContactPosition = CollisionDataItem.Location;
 								NewContact.ContactPenetration = CollisionDataItem.PenetrationDepth;
 								NotifyInfo.RigidCollisionData.bIsVelocityDeltaUnderThreshold = CollisionDataItem.DeltaVelocity1.IsNearlyZero(MinDeltaVelocityThreshold) && CollisionDataItem.DeltaVelocity2.IsNearlyZero(MinDeltaVelocityThreshold);
-								// NewContact.PhysMaterial[1] UPhysicalMaterial required here
+								
+								Chaos::FChaosPhysicsMaterial* InternalMat1 = CollisionDataItem.Mat1.Get();
+								Chaos::FChaosPhysicsMaterial* InternalMat2 = CollisionDataItem.Mat2.Get();
+								
+								NewContact.PhysMaterial[0] = InternalMat1 ? FPhysicsUserData::Get<UPhysicalMaterial>(InternalMat1->UserData) : nullptr;
+								NewContact.PhysMaterial[1] = InternalMat2 ? FPhysicsUserData::Get<UPhysicalMaterial>(InternalMat2->UserData) : nullptr;
 							}
 						}
 					}
diff --git a/Engine/Source/Runtime/Experimental/Chaos/Private/Chaos/HeightField.cpp b/Engine/Source/Runtime/Experimental/Chaos/Private/Chaos/HeightField.cpp
--- a/Engine/Source/Runtime/Experimental/Chaos/Private/Chaos/HeightField.cpp
+++ b/Engine/Source/Runtime/Experimental/Chaos/Private/Chaos/HeightField.cpp
@@ -1351,8 +1351,9 @@ namespace Chaos
 				{
 					FPBDCollisionConstraint Constraint;
 					OverlapTriangle(Points[0], Points[1], Points[3], Constraint);
-					for (const FManifoldPoint& ManifoldPoint : Constraint.GetManifoldPoints())
+					for (FManifoldPoint& ManifoldPoint : Constraint.GetManifoldPoints())
 					{
+						ManifoldPoint.ContactPoint.FaceIndex = CellIndex * 2;
 						InsertSorted(ManifoldPoint.ContactPoint);
 					}
 				}
@@ -1360,8 +1361,9 @@ namespace Chaos
 				{
 					FPBDCollisionConstraint Constraint;
 					OverlapTriangle(Points[0], Points[3], Points[2], Constraint);
-					for (const FManifoldPoint& ManifoldPoint : Constraint.GetManifoldPoints())
+					for (FManifoldPoint& ManifoldPoint : Constraint.GetManifoldPoints())
 					{
+						ManifoldPoint.ContactPoint.FaceIndex = CellIndex * 2 + 1;
 						InsertSorted(ManifoldPoint.ContactPoint);
 					}
 				}
diff --git a/Engine/Source/Runtime/Experimental/Chaos/Private/Chaos/TriangleMeshImplicitObject.cpp b/Engine/Source/Runtime/Experimental/Chaos/Private/Chaos/TriangleMeshImplicitObject.cpp
--- a/Engine/Source/Runtime/Experimental/Chaos/Private/Chaos/TriangleMeshImplicitObject.cpp
+++ b/Engine/Source/Runtime/Experimental/Chaos/Private/Chaos/TriangleMeshImplicitObject.cpp
@@ -397,8 +397,9 @@ bool FTriangleMeshImplicitObject::ContactManifoldImp(const GeomType& QueryGeom,
 			TransformVertsHelper(TriMeshScale, TriIdx, MParticles, Elements, A, B, C);
 			FPBDCollisionConstraint Constraint;
 			OverlapTriangle(A, B, C, Constraint);
-			for (const FManifoldPoint& ManifoldPoint : Constraint.GetManifoldPoints())
+			for(FManifoldPoint& ManifoldPoint : Constraint.GetManifoldPoints())
 			{
+				ManifoldPoint.ContactPoint.FaceIndex = TriIdx;
 				InsertSorted(ManifoldPoint.ContactPoint);
 			}
 
diff --git a/Engine/Source/Runtime/Experimental/Chaos/Private/EventDefaults.cpp b/Engine/Source/Runtime/Experimental/Chaos/Private/EventDefaults.cpp
--- a/Engine/Source/Runtime/Experimental/Chaos/Private/EventDefaults.cpp
+++ b/Engine/Source/Runtime/Experimental/Chaos/Private/EventDefaults.cpp
@@ -17,6 +17,46 @@
 
 namespace Chaos
 {
+	/**
+	 * Resolves a material from a shape and a collision.
+	 * Either both shapes are simple primitives, or one is simple and the other is a heightfield or trimesh.
+	 * In the case of the trimesh/heightfield the contact points should have the face index of the hit face
+	 * which the geometry can translate into a material index.
+	 * @param InShape - The shape to resolve a material for
+	 * @param InConstraint - The collision that the shape is a part of
+	 * @see FTriangleMeshImplicitObject::ContactManifoldImp
+	 * @see FHeightField::ContactManifoldImp
+	 */
+	FMaterialHandle ResolveMaterial(const FPerShapeData* InShape, const FPBDCollisionConstraint& InConstraint)
+	{
+		if(InShape)
+		{
+			const TArray<FMaterialHandle>& MaterialArray = InShape->GetMaterials();
+
+			// Simple case, one material. All primitives (Box, convex, sphere etc.) should only have one material.
+			// Heightfield and Trimesh can have one or more and will require data from the contacts to resolve.
+			if(MaterialArray.Num() == 1)
+			{
+				return MaterialArray[0];
+			}
+			else
+			{
+				// Check the manifold points, return the first valid material based on the face that the contact hit.
+				for(const FManifoldPoint& Point : InConstraint.GetManifoldPoints())
+				{
+					const int32 PotentialMaterialIndex = InShape->GetGeometry()->GetMaterialIndex(Point.ContactPoint.FaceIndex);
+					if(MaterialArray.IsValidIndex(PotentialMaterialIndex))
+					{
+						return MaterialArray[PotentialMaterialIndex];
+					}
+				}
+			}
+		}
+
+		// No valid material
+		return {};
+	}
+
 	void FEventDefaults::RegisterSystemEvents(FEventManager& EventManager)
 	{
 		RegisterCollisionEvent(EventManager);
@@ -133,6 +173,9 @@ namespace Chaos
 						FGeometryParticleHandle* Particle0 = Constraint.Particle[0];
 						FGeometryParticleHandle* Particle1 = Constraint.Particle[1];
 
+						const FPerShapeData* Shape0 = Particle0->GetImplicitShape(Constraint.GetImplicit0());
+						const FPerShapeData* Shape1 = Particle1->GetImplicitShape(Constraint.GetImplicit1());
+
 						FCollidingData Data;
 						Data.Location = Constraint.GetLocation();
 						Data.AccumulatedImpulse = Constraint.AccumulatedImpulse;
@@ -142,6 +185,9 @@ namespace Chaos
 						Data.Proxy1 = Particle0 ? Particle0->PhysicsProxy() : nullptr;
 						Data.Proxy2 = Particle1 ? Particle1->PhysicsProxy() : nullptr;
 
+						Data.Mat1 = ResolveMaterial(Shape0, Constraint);
+						Data.Mat2 = ResolveMaterial(Shape1, Constraint);
+
 						// Collision constraints require both proxies are valid. If either is not, we needn't record the collision event.
 						if (Data.Proxy1 == nullptr || Data.Proxy2 == nullptr)
 						{
diff --git a/Engine/Source/Runtime/Experimental/Chaos/Public/Chaos/Collision/ContactPoint.h b/Engine/Source/Runtime/Experimental/Chaos/Public/Chaos/Collision/ContactPoint.h
--- a/Engine/Source/Runtime/Experimental/Chaos/Public/Chaos/Collision/ContactPoint.h
+++ b/Engine/Source/Runtime/Experimental/Chaos/Public/Chaos/Collision/ContactPoint.h
@@ -40,8 +40,12 @@ namespace Chaos
 		// Contact separation (negative for overlap)
 		FReal Phi;
 
+		// Face index of the shape we hit. Only valid for Heightfield and Trimesh contact points, otherwise INDEX_NONE
+		int32 FaceIndex;
+
 		FContactPoint()
 			: Phi(TNumericLimits<FReal>::Max())
+			, FaceIndex(INDEX_NONE)
 		{
 		}
 
diff --git a/Engine/Source/Runtime/Experimental/Chaos/Public/Chaos/ExternalCollisionData.h b/Engine/Source/Runtime/Experimental/Chaos/Public/Chaos/ExternalCollisionData.h
--- a/Engine/Source/Runtime/Experimental/Chaos/Public/Chaos/ExternalCollisionData.h
+++ b/Engine/Source/Runtime/Experimental/Chaos/Public/Chaos/ExternalCollisionData.h
@@ -5,6 +5,7 @@
 #include "Chaos/PBDConstraintContainer.h"
 #include "Chaos/ParticleHandleFwd.h"
 #include "Chaos/Vector.h"
+#include "Chaos/PhysicalMaterials.h"
 
 
 class UPhysicalMaterial;
@@ -54,11 +55,17 @@ namespace Chaos
 		FVec3 Location;
 		FVec3 AccumulatedImpulse;
 		FVec3 Normal;
-		FVec3 Velocity1, Velocity2;
-		FVec3 DeltaVelocity1, DeltaVelocity2;
-		FVec3 AngularVelocity1, AngularVelocity2;
-		FReal Mass1, Mass2;
+		FVec3 Velocity1;
+		FVec3 Velocity2;
+		FVec3 DeltaVelocity1;
+		FVec3 DeltaVelocity2;
+		FVec3 AngularVelocity1;
+		FVec3 AngularVelocity2;
+		FReal Mass1;
+		FReal Mass2;
 		FReal PenetrationDepth;
+		FMaterialHandle Mat1;
+		FMaterialHandle Mat2;
 
 		// The pointers to the proxies should be used with caution on the Game Thread.
 		// Ideally we only ever use these as table keys when acquiring related structures.
