Chaos - Heightfield raycast perf improvement
- Make sure VisitRaycast returns relevant value if a triangle has been hit
- Shorten the raylength to the exit time from the test against the heightfield bounds
- measured improvement of 25% for raycast against heightfield

diff --git a/Engine/Source/Runtime/Experimental/Chaos/Private/Chaos/HeightField.cpp b/Engine/Source/Runtime/Experimental/Chaos/Private/Chaos/HeightField.cpp
--- a/Engine/Source/Runtime/Experimental/Chaos/Private/Chaos/HeightField.cpp
+++ b/Engine/Source/Runtime/Experimental/Chaos/Private/Chaos/HeightField.cpp
@@ -27,7 +27,6 @@ namespace Chaos
 	class FHeightfieldRaycastVisitor
 	{
 	public:
-
 		FHeightfieldRaycastVisitor(const typename FHeightField::FDataType* InData, const FVec3& InStart, const FVec3& InDir, const FReal InThickness)
 			: OutTime(TNumericLimits<FReal>::Max())
 			, OutFaceIndex(INDEX_NONE)
@@ -61,6 +60,7 @@ namespace Chaos
 			const FReal Radius2 = Radius * Radius;
 			bool bIntersection = false;
 
+			// return if the triangle was hit or not 
 			auto TestTriangle = [&](int32 FaceIndex, const FVec3& A, const FVec3& B, const FVec3& C) -> bool
 			{
 				const FVec3 AB = B - A;
@@ -72,7 +72,7 @@ namespace Chaos
 				if(!ensure(Len2 > SMALL_NUMBER))
 				{
 					// Bad triangle, co-linear points or very thin
-					return true;
+					return false;
 				}
 
 				const TPlane<FReal, 3> TrianglePlane(A, Normal);
@@ -95,7 +95,7 @@ namespace Chaos
 							OutPosition = ClosestPtOnTri;
 							OutNormal = Normal;
 							OutFaceIndex = FaceIndex;
-							return false;
+							return true;
 						}
 					}
 					else
@@ -175,11 +175,12 @@ namespace Chaos
 							OutTime = Time;
 							OutFaceIndex = FaceIndex;
 							CurrentLength = Time;
+							return true;
 						}
 					}
 				}
 
-				return true;
+				return false;
 			};
 
 			FVec3 Points[4];
@@ -191,13 +192,15 @@ namespace Chaos
 			//todo: can do it without raycast
 			FReal TOI;
 			FVec3 HitPoint;
+			bool bHit = false;
 			if (CellBounds.RaycastFast(Start, Dir, InvDir, bParallel, CurrentLength, 1.0f / CurrentLength, TOI, HitPoint))
 			{
 				// Test both triangles that are in this cell, as we could hit both in any order
-				TestTriangle(Payload * 2, Points[0], Points[1], Points[3]);
-				TestTriangle(Payload * 2 + 1, Points[0], Points[3], Points[2]);
+				bHit |= TestTriangle(Payload * 2, Points[0], Points[1], Points[3]);
+				bHit |= TestTriangle(Payload * 2 + 1, Points[0], Points[3], Points[2]);
 			}
-			return OutTime > 0;
+			const bool bShouldContinueVisiting = !bHit;
+			return bShouldContinueVisiting;
 		}
 
 		bool VisitRaycast(int32 Payload, FReal& CurLength)
@@ -820,16 +823,19 @@ namespace Chaos
 			InvDir[Axis] = bParallel[Axis] ? 0 : 1 / Dir[Axis];
 		}
 
-		FReal TOI;
 		const FBounds2D FlatBounds = GetFlatBounds();
 		FAABB3 Bounds(
 			FVec3(FlatBounds.Min[0],FlatBounds.Min[1],GeomData.GetMinHeight() * GeomData.Scale[2]),
 			FVec3(FlatBounds.Max[0],FlatBounds.Max[1],GeomData.GetMaxHeight() * GeomData.Scale[2])
 			);
-		FVec3 NextStart;
 
-		if(Bounds.RaycastFast(StartPoint, Dir, InvDir, bParallel, Length, InvCurrentLength, TOI, NextStart))
+		FReal RayEntryTime;
+		FReal RayExitTime;
+		if(Bounds.RaycastFast(StartPoint, Dir, InvDir, bParallel, Length, InvCurrentLength, RayEntryTime, RayExitTime))
 		{
+			CurrentLength = RayExitTime + SMALL_NUMBER; // to account for precision errors 
+			FVec3 NextStart = StartPoint + (Dir * RayEntryTime);
+
 			const FVec2 Scale2D(GeomData.Scale[0],GeomData.Scale[1]);
 			TVec2<int32> CellIdx = FlatGrid.Cell(TVec2<int32>(static_cast<int32>(NextStart[0] / Scale2D[0]), static_cast<int32>(NextStart[1] / Scale2D[1])));
 			const FReal ZDx = Bounds.Extents()[2];
@@ -843,6 +849,7 @@ namespace Chaos
 			{
 				if (FlatGrid.IsValid(CellIdx))
 				{
+					PHYSICS_CSV_CUSTOM_VERY_EXPENSIVE(PhysicsCounters, NumRayHeightfieldCellVisited, 1, ECsvCustomStatOp::Accumulate);
 					// Test for the cell bounding box is done in the visitor at the same time as fetching the points for the triangles
 					// this avoid fetching the points twice ( here and in the visitor ) 
 					bool bContinue = Visitor.VisitRaycast(CellIdx[1] * (GeomData.NumCols - 1) + CellIdx[0], CurrentLength);
diff --git a/Engine/Source/Runtime/Experimental/Chaos/Public/Chaos/AABB.h b/Engine/Source/Runtime/Experimental/Chaos/Public/Chaos/AABB.h
--- a/Engine/Source/Runtime/Experimental/Chaos/Public/Chaos/AABB.h
+++ b/Engine/Source/Runtime/Experimental/Chaos/Public/Chaos/AABB.h
@@ -190,8 +190,7 @@ namespace Chaos
 
 		bool CHAOS_API Raycast(const TVector<FReal, d>& StartPoint, const TVector<FReal, d>& Dir, const FReal Length, const FReal Thickness, FReal& OutTime, TVector<FReal, d>& OutPosition, TVector<FReal, d>& OutNormal, int32& OutFaceIndex) const;
 
-
-		FORCEINLINE bool RaycastFast(const TVector<FReal, d>& StartPoint, const TVector<FReal, d>& Dir, const TVector<FReal, d>& InvDir, const bool* bParallel, const FReal Length, const FReal InvLength, FReal& OutTime, TVector<FReal, d>& OutPosition) const
+		FORCEINLINE bool RaycastFast(const TVector<FReal, d>& StartPoint, const TVector<FReal, d>& Dir, const TVector<FReal, d>& InvDir, const bool* bParallel, const FReal Length, const FReal InvLength, FReal& OutEntryTime, FReal& OutExitTime) const
 		{
 			const TVector<FReal, d> StartToMin = TVector<FReal, d>(MMin) - StartPoint;
 			const TVector<FReal, d> StartToMax = TVector<FReal, d>(MMax) - StartPoint;
@@ -243,11 +242,24 @@ namespace Chaos
 				return false;
 			}
 
-			OutTime = LatestStartTime;
-			OutPosition = StartPoint + LatestStartTime * Dir;
+			OutEntryTime = LatestStartTime;
+			OutExitTime = EarliestEndTime;
 			return true;
 		}
 
+		FORCEINLINE bool RaycastFast(const TVector<FReal, d>& StartPoint, const TVector<FReal, d>& Dir, const TVector<FReal, d>& InvDir, const bool* bParallel, const FReal Length, const FReal InvLength, FReal& OutTime, TVector<FReal, d>& OutPosition) const
+		{
+			FReal RayEntryTime;
+			FReal RayExitTime;
+			if (RaycastFast(StartPoint, Dir, InvDir, bParallel, Length, InvLength, RayEntryTime, RayExitTime))
+			{
+				OutTime = RayEntryTime;
+				OutPosition = StartPoint + RayEntryTime * Dir;
+				return true;
+			}
+			return false;
+		}
+
 		CHAOS_API TVector<T, d> FindClosestPoint(const TVector<T, d>& StartPoint, const T Thickness = (T)0) const;
 
 		CHAOS_API Pair<TVector<FReal, d>, bool> FindClosestIntersectionImp(const TVector<FReal, d>& StartPoint, const TVector<FReal, d>& EndPoint, const FReal Thickness) const;
