Fix missing backface culling on complex sweeps. Improves perf and is consistent with raytrace behavior

diff --git a/Engine/Source/Runtime/Experimental/Chaos/Private/Chaos/TriangleMeshImplicitObject.cpp b/Engine/Source/Runtime/Experimental/Chaos/Private/Chaos/TriangleMeshImplicitObject.cpp
--- a/Engine/Source/Runtime/Experimental/Chaos/Private/Chaos/TriangleMeshImplicitObject.cpp
+++ b/Engine/Source/Runtime/Experimental/Chaos/Private/Chaos/TriangleMeshImplicitObject.cpp
@@ -810,7 +810,7 @@ template <typename QueryGeomType, typename IdxType>
 struct FTriangleMeshSweepVisitor
 {
 	FTriangleMeshSweepVisitor(const FTriangleMeshImplicitObject& InTriMesh, const TArray<TVec3<IdxType>>& InElements, const QueryGeomType& InQueryGeom, const TRigidTransform<FReal,3>& InStartTM, const FVec3& InDir,
-		const FVec3& InScaledDirNormalized, const FReal InLengthScale, const FRigidTransform3& InScaledStartTM, const FReal InThickness, const bool InComputeMTD, FVec3 InTriMeshScale)
+		const FVec3& InScaledDirNormalized, const FReal InLengthScale, const FRigidTransform3& InScaledStartTM, const FReal InThickness, const bool InComputeMTD, FVec3 InTriMeshScale, bool InCullsBackFaceSweeps)
 	: TriMesh(InTriMesh)
 	, Elements(InElements)
 	, StartTM(InStartTM)
@@ -818,6 +818,7 @@ struct FTriangleMeshSweepVisitor
 	, Dir(InDir)
 	, Thickness(InThickness)
 	, bComputeMTD(InComputeMTD)
+	, bCullsBackFaceSweeps(InCullsBackFaceSweeps)
 	, ScaledDirNormalized(InScaledDirNormalized)
 	, LengthScale(InLengthScale)
 	, ScaledStartTM(InScaledStartTM)
@@ -853,6 +854,16 @@ struct FTriangleMeshSweepVisitor
 		TransformVertsHelper(TriMeshScale,TriIdx,TriMesh.MParticles,Elements,A,B,C);
 		FTriangle Tri(A, B, C);
 
+		if(bCullsBackFaceSweeps)
+		{
+			const FVec3 TriNormal = FVec3::CrossProduct(B - A, C - A);
+			if (FVec3::DotProduct(TriNormal, ScaledDirNormalized) > 0)
+			{
+				return true;
+			}
+		}
+		
+
 		const auto& WorldScaleQueryGeom = ScaleGeomIntoWorldHelper(QueryGeom, TriMeshScale);
 
 		if(GJKRaycast2<FReal>(Tri, WorldScaleQueryGeom, ScaledStartTM, ScaledDirNormalized, LengthScale * CurData.CurrentLength, Time, HitPosition, HitNormal, Thickness, bComputeMTD))
@@ -886,6 +897,7 @@ struct FTriangleMeshSweepVisitor
 	const FVec3& Dir;
 	const FReal Thickness;
 	const bool bComputeMTD;
+	const bool bCullsBackFaceSweeps;
 
 	// Cache these values for Scaled Triangle Mesh, as they are needed for transformation when sweeping against triangles.
 	FVec3 ScaledDirNormalized;
@@ -931,7 +943,7 @@ bool FTriangleMeshImplicitObject::SweepGeomImp(const QueryGeomType& QueryGeom, c
 	auto LambdaHelper = [&](const auto& Elements)
 	{
 		using VisitorType = FTriangleMeshSweepVisitor<QueryGeomType,decltype(Elements[0][0])>;
-		VisitorType SQVisitor(*this,Elements, QueryGeom,StartTM,Dir,ScaledDirNormalized,LengthScale,ScaledStartTM,Thickness,bComputeMTD, TriMeshScale);
+		VisitorType SQVisitor(*this,Elements, QueryGeom,StartTM,Dir,ScaledDirNormalized,LengthScale,ScaledStartTM,Thickness,bComputeMTD, TriMeshScale, bCullsBackFaceRaycast);
 
 
 		const FAABB3 QueryBounds = QueryGeom.BoundingBox().TransformedAABB(FRigidTransform3(FVec3::ZeroVector,StartTM.GetRotation()));
