Chaos - fix GeometryCollectionComponent with null RestCollection

diff --git a/Engine/Source/Runtime/Experimental/GeometryCollectionEngine/Private/GeometryCollection/GeometryCollectionComponent.cpp b/Engine/Source/Runtime/Experimental/GeometryCollectionEngine/Private/GeometryCollection/GeometryCollectionComponent.cpp
--- a/Engine/Source/Runtime/Experimental/GeometryCollectionEngine/Private/GeometryCollection/GeometryCollectionComponent.cpp
+++ b/Engine/Source/Runtime/Experimental/GeometryCollectionEngine/Private/GeometryCollection/GeometryCollectionComponent.cpp
@@ -312,9 +312,12 @@ void UGeometryCollectionComponent::BeginPlay()
 #endif
 
 #if WITH_EDITOR
-	if (RestCollection->GetGeometryCollection()->HasAttribute("ExplodedVector", FGeometryCollection::TransformGroup))
+	if (RestCollection != nullptr)
 	{
-		RestCollection->GetGeometryCollection()->RemoveAttribute("ExplodedVector", FGeometryCollection::TransformGroup);
+		if (RestCollection->GetGeometryCollection()->HasAttribute("ExplodedVector", FGeometryCollection::TransformGroup))
+		{
+			RestCollection->GetGeometryCollection()->RemoveAttribute("ExplodedVector", FGeometryCollection::TransformGroup);
+		}
 	}
 #endif
 
