Add AABBTreeExpensiveStats at EndFrame

diff --git a/Engine/Source/Runtime/Experimental/Chaos/Public/Chaos/AABBTree.h b/Engine/Source/Runtime/Experimental/Chaos/Public/Chaos/AABBTree.h
--- a/Engine/Source/Runtime/Experimental/Chaos/Public/Chaos/AABBTree.h
+++ b/Engine/Source/Runtime/Experimental/Chaos/Public/Chaos/AABBTree.h
@@ -65,7 +65,7 @@ struct AABBTreeStatistics
 		StatNumGridOverflowElements = 0;
 	}
 
-	AABBTreeStatistics& operator+=(const AABBTreeStatistics& Rhs)
+	AABBTreeStatistics& MergeStatistics(const AABBTreeStatistics& Rhs)
 	{
 		StatNumNonEmptyCellsInGrid += Rhs.StatNumNonEmptyCellsInGrid;
 		StatNumElementsTooLargeForGrid += Rhs.StatNumElementsTooLargeForGrid;
@@ -80,6 +80,33 @@ struct AABBTreeStatistics
 	int32 StatNumGridOverflowElements = 0;
 };
 
+struct AABBTreeExpensiveStatistics
+{
+	void Reset()
+	{
+		StatMaxNumLeaves = 0;
+		StatMaxDirtyElements = 0;
+		StatMaxLeafSize = 0;
+		StatMaxTreeDepth = 0;
+		StatGlobalPayloadsSize = 0;
+	}
+
+	AABBTreeExpensiveStatistics& MergeStatistics(const AABBTreeExpensiveStatistics& Rhs)
+	{
+		StatMaxNumLeaves = FMath::Max(StatMaxNumLeaves, Rhs.StatMaxNumLeaves);
+		StatMaxDirtyElements = FMath::Max(StatMaxDirtyElements, Rhs.StatMaxDirtyElements);
+		StatMaxLeafSize = FMath::Max(StatMaxLeafSize, Rhs.StatMaxLeafSize);
+		StatMaxTreeDepth = FMath::Max(StatMaxTreeDepth, Rhs.StatMaxTreeDepth);
+		StatGlobalPayloadsSize += Rhs.StatGlobalPayloadsSize;
+		return *this;
+	}
+	int32 StatMaxNumLeaves = 0;
+	int32 StatMaxDirtyElements = 0;
+	int32 StatMaxLeafSize = 0;
+	int32 StatMaxTreeDepth = 0;
+	int32 StatGlobalPayloadsSize = 0;
+};
+
 DECLARE_CYCLE_STAT(TEXT("AABBTreeGenerateTree"), STAT_AABBTreeGenerateTree, STATGROUP_Chaos);
 DECLARE_CYCLE_STAT(TEXT("AABBTreeTimeSliceSetup"), STAT_AABBTreeTimeSliceSetup, STATGROUP_Chaos);
 DECLARE_CYCLE_STAT(TEXT("AABBTreeInitialTimeSlice"), STAT_AABBTreeInitialTimeSlice, STATGROUP_Chaos);
@@ -484,6 +511,7 @@ public:
 		FlattenedCellArrayOfDirtyIndices.Reset();
 		DirtyElementsGridOverflow.Reset();
 		TreeStats.Reset();
+		TreeExpensiveStats.Reset();
 		GlobalPayloads.Reset();
 		PayloadToInfo.Reset();
 
@@ -996,13 +1024,43 @@ public:
 	}
 
 	// Some useful statistics
-	const AABBTreeStatistics& GetAABBTreeStatistics() {
+	const AABBTreeStatistics& GetAABBTreeStatistics()
+	{
 		// Update the stats that needs it first
 		TreeStats.StatNumDirtyElements = DirtyElements.Num();
 		TreeStats.StatNumGridOverflowElements = DirtyElementsGridOverflow.Num();
 		return TreeStats;
 	}
 
+	const AABBTreeExpensiveStatistics& GetAABBTreeExpensiveStatistics()
+	{
+		TreeExpensiveStats.StatMaxDirtyElements = DirtyElements.Num();
+		TreeExpensiveStats.StatMaxNumLeaves = Leaves.Num();
+		int32 StatMaxLeafSize = 0;
+		for(const TLeafType& Leaf : Leaves)
+		{
+			StatMaxLeafSize = FMath::Max(StatMaxLeafSize, (int32)Leaf.GetElementCount());
+		}
+		TreeExpensiveStats.StatMaxLeafSize = StatMaxLeafSize;
+		TreeExpensiveStats.StatMaxTreeDepth = GetSubtreeDepth(0);
+		TreeExpensiveStats.StatGlobalPayloadsSize = GlobalPayloads.Num();
+
+		return TreeExpensiveStats;
+	}
+
+	const int32 GetSubtreeDepth(const int32 NodeIdx)
+	{
+		const FNode& Node = Nodes[NodeIdx];
+		if (Node.bLeaf)
+		{
+			return 1;
+		}
+		else
+		{
+			return FMath::Max(GetSubtreeDepth(Node.ChildrenNodes[0]), GetSubtreeDepth(Node.ChildrenNodes[1])) + 1;
+		}
+	}
+
 	const TArray<TPayloadBoundsElement<TPayloadType, T>>& GlobalObjects() const
 	{
 		return GlobalPayloads;
@@ -1548,6 +1606,7 @@ private:
 		FlattenedCellArrayOfDirtyIndices.Reset();
 		DirtyElementsGridOverflow.Reset();
 		TreeStats.Reset();
+		TreeExpensiveStats.Reset();
 		PayloadToInfo.Reset();
 		NumProcessedThisSlice = 0;
 		GetCVars();  // Safe to copy CVARS here
@@ -1758,7 +1817,8 @@ private:
 			auto& LeavesRef = Leaves;
 			auto& NodesRef = Nodes;
 			auto& WorkPoolRef = WorkPool;
-			auto MakeLeaf = [NewNodeIdx, &PayloadToInfoRef, &WorkPoolRef, CurIdx, &LeavesRef, &NodesRef]()
+			auto& TreeExpensiveStatsRef = TreeExpensiveStats;
+			auto MakeLeaf = [NewNodeIdx, &PayloadToInfoRef, &WorkPoolRef, CurIdx, &LeavesRef, &NodesRef, &TreeExpensiveStatsRef]()
 			{
 				if (bMutable)
 				{
@@ -1772,7 +1832,6 @@ private:
 
 				NodesRef[NewNodeIdx].bLeaf = true;
 				NodesRef[NewNodeIdx].ChildrenNodes[0] = LeavesRef.Add(TLeafType{ WorkPoolRef[CurIdx].Elems }); //todo: avoid copy?
-
 			};
 
 			if (WorkPool[CurIdx].Elems.Num() <= MaxChildrenInLeaf || WorkPool[CurIdx].NodeLevel >= MaxTreeDepth)
@@ -2017,6 +2076,7 @@ private:
 		, DirtyElementMaxPhysicalSizeInCells(Other.DirtyElementMaxPhysicalSizeInCells)
 		, DirtyElementMaxCellCapacity(Other.DirtyElementMaxCellCapacity)
 		, TreeStats(Other.TreeStats)
+		, TreeExpensiveStats(Other.TreeExpensiveStats)
 		, GlobalPayloads(Other.GlobalPayloads)
 		, PayloadToInfo(Other.PayloadToInfo)
 		, MaxChildrenInLeaf(Other.MaxChildrenInLeaf)
@@ -2055,6 +2115,7 @@ private:
 			FlattenedCellArrayOfDirtyIndices = Rhs.FlattenedCellArrayOfDirtyIndices;
 			DirtyElementsGridOverflow = Rhs.DirtyElementsGridOverflow;
 			TreeStats = Rhs.TreeStats;
+			TreeExpensiveStats = Rhs.TreeExpensiveStats;
 			
 			DirtyElementGridCellSize = Rhs.DirtyElementGridCellSize;
 			DirtyElementGridCellSizeInv = Rhs.DirtyElementGridCellSizeInv;
@@ -2091,6 +2152,7 @@ private:
 	int32 DirtyElementMaxCellCapacity;
 	// Some useful statistics
 	AABBTreeStatistics TreeStats;
+	AABBTreeExpensiveStatistics TreeExpensiveStats;
 
 	TArray<FElement> GlobalPayloads;
 	TArrayAsMap<TPayloadType, FAABBTreePayloadInfo> PayloadToInfo;
diff --git a/Engine/Source/Runtime/PhysicsCore/Private/ChaosScene.cpp b/Engine/Source/Runtime/PhysicsCore/Private/ChaosScene.cpp
--- a/Engine/Source/Runtime/PhysicsCore/Private/ChaosScene.cpp
+++ b/Engine/Source/Runtime/PhysicsCore/Private/ChaosScene.cpp
@@ -35,10 +35,10 @@
 #include "PhysicsSettingsCore.h"
 #include "Chaos/PhysicsSolverBaseImpl.h"
 
-#include "ProfilingDebugging/CsvProfiler.h"
 
 DECLARE_CYCLE_STAT(TEXT("Update Kinematics On Deferred SkelMeshes"),STAT_UpdateKinematicsOnDeferredSkelMeshesChaos,STATGROUP_Physics);
 CSV_DEFINE_CATEGORY(ChaosPhysics,true);
+CSV_DEFINE_CATEGORY(AABBTreeExpensiveStats, true);
 
 // Stat Counters
 DECLARE_DWORD_ACCUMULATOR_STAT(TEXT("NumDirtyAABBTreeElements"), STAT_ChaosCounter_NumDirtyAABBTreeElements, STATGROUP_ChaosCounters);
@@ -442,7 +442,7 @@ void FChaosScene::SyncBodies(TSolver* Solver)
 }
 
 // Accumulate all the AABBTree stats
-void GetAABBTreeStats(Chaos::ISpatialAccelerationCollection<Chaos::FAccelerationStructureHandle, Chaos::FReal, 3>& Collection, Chaos::AABBTreeStatistics& OutAABBTreeStatistics)
+void GetAABBTreeStats(Chaos::ISpatialAccelerationCollection<Chaos::FAccelerationStructureHandle, Chaos::FReal, 3>& Collection, Chaos::AABBTreeStatistics& OutAABBTreeStatistics, Chaos::AABBTreeExpensiveStatistics& OutAABBTreeExpensiveStatistics)
 {
 	using namespace Chaos;
 	OutAABBTreeStatistics.Reset();
@@ -452,11 +452,17 @@ void GetAABBTreeStats(Chaos::ISpatialAccelerationCollection<Chaos::FAcceleration
 		auto SubStructure = Collection.GetSubstructure(SpatialIndex);
 		if (const auto AABBTree = SubStructure->template As<TAABBTree<FAccelerationStructureHandle, TAABBTreeLeafArray<FAccelerationStructureHandle>>>())
 		{
-			OutAABBTreeStatistics += AABBTree->GetAABBTreeStatistics();
+			OutAABBTreeStatistics.MergeStatistics(AABBTree->GetAABBTreeStatistics());
+#if CSV_PROFILER
+			if (FCsvProfiler::Get()->IsCapturing() && FCsvProfiler::Get()->IsCategoryEnabled(CSV_CATEGORY_INDEX(AABBTreeExpensiveStats)))
+			{
+				OutAABBTreeExpensiveStatistics.MergeStatistics(AABBTree->GetAABBTreeExpensiveStatistics());
+			}
+#endif
 		}
 		else if (const auto AABBTreeBV = SubStructure->template As<TAABBTree<FAccelerationStructureHandle, TBoundingVolume<FAccelerationStructureHandle>>>())
 		{
-			OutAABBTreeStatistics += AABBTreeBV->GetAABBTreeStatistics();
+			OutAABBTreeStatistics.MergeStatistics(AABBTreeBV->GetAABBTreeStatistics());
 		}
 	}
 }
@@ -475,7 +481,8 @@ void FChaosScene::EndFrame()
 	}
 
 	Chaos::AABBTreeStatistics TreeStats;
-	GetAABBTreeStats(GetSpacialAcceleration()->AsChecked<SpatialAccelerationCollection>(), TreeStats);
+	Chaos::AABBTreeExpensiveStatistics TreeExpensiveStats;
+	GetAABBTreeStats(GetSpacialAcceleration()->AsChecked<SpatialAccelerationCollection>(), TreeStats, TreeExpensiveStats);
 
 	CSV_CUSTOM_STAT(ChaosPhysics, AABBTreeDirtyElementCount, TreeStats.StatNumDirtyElements, ECsvCustomStatOp::Set);
 	SET_DWORD_STAT(STAT_ChaosCounter_NumDirtyAABBTreeElements, TreeStats.StatNumDirtyElements);
@@ -489,6 +496,17 @@ void FChaosScene::EndFrame()
 	CSV_CUSTOM_STAT(ChaosPhysics, AABBTreeDirtyElementNonEmptyCellCount, TreeStats.StatNumNonEmptyCellsInGrid, ECsvCustomStatOp::Set);
 	SET_DWORD_STAT(STAT_ChaosCounter_NumDirtyNonEmptyCellsInGrid, TreeStats.StatNumNonEmptyCellsInGrid);
 
+#if CSV_PROFILER
+	if (FCsvProfiler::Get()->IsCapturing() && FCsvProfiler::Get()->IsCategoryEnabled(CSV_CATEGORY_INDEX(AABBTreeExpensiveStats)))
+	{
+		CSV_CUSTOM_STAT(AABBTreeExpensiveStats, AABBTreeMaxNumLeaves, TreeExpensiveStats.StatMaxNumLeaves, ECsvCustomStatOp::Set);
+		CSV_CUSTOM_STAT(AABBTreeExpensiveStats, AABBTreeMaxDirtyElements, TreeExpensiveStats.StatMaxDirtyElements, ECsvCustomStatOp::Set);
+		CSV_CUSTOM_STAT(AABBTreeExpensiveStats, AABBTreeMaxTreeDepth, TreeExpensiveStats.StatMaxTreeDepth, ECsvCustomStatOp::Set);
+		CSV_CUSTOM_STAT(AABBTreeExpensiveStats, AABBTreeMaxLeafSize, TreeExpensiveStats.StatMaxLeafSize, ECsvCustomStatOp::Set);
+		CSV_CUSTOM_STAT(AABBTreeExpensiveStats, AABBTreeGlobalPayloadsSize, TreeExpensiveStats.StatGlobalPayloadsSize, ECsvCustomStatOp::Set);
+	}
+#endif
+
 	check(IsCompletionEventComplete())
 	//check(PhysicsTickTask->IsComplete());
 	CompletionEvents.Reset();
