Chaos: Fix CSV stat so that it works with perf test

diff --git a/Engine/Source/Runtime/Experimental/Chaos/Private/Chaos/PBDRigidsEvolutionGBF.cpp b/Engine/Source/Runtime/Experimental/Chaos/Private/Chaos/PBDRigidsEvolutionGBF.cpp
--- a/Engine/Source/Runtime/Experimental/Chaos/Private/Chaos/PBDRigidsEvolutionGBF.cpp
+++ b/Engine/Source/Runtime/Experimental/Chaos/Private/Chaos/PBDRigidsEvolutionGBF.cpp
@@ -341,7 +341,7 @@ void FPBDRigidsEvolutionGBF::AdvanceOneTimeStepImpl(const FReal Dt, const FSubSt
 				// Update constraint graphs, coloring etc as required by the different constraint types in this island
 				{
 					SCOPE_CYCLE_COUNTER(STAT_Evolution_UpdateAccelerationStructures);
-					CSV_SCOPED_TIMING_STAT(PhysicsVerbose, STAT_Evolution_UpdateAccelerationStructures);
+					CSV_SCOPED_TIMING_STAT(PhysicsVerbose, StepSolver_UpdateAccelerationStructures);
 					UpdateAccelerationStructures(Dt, Island);
 				}
 
