Fix AABB tree crash

diff --git a/Engine/Source/Runtime/Experimental/Chaos/Public/Chaos/AABBTree.h b/Engine/Source/Runtime/Experimental/Chaos/Public/Chaos/AABBTree.h
--- a/Engine/Source/Runtime/Experimental/Chaos/Public/Chaos/AABBTree.h
+++ b/Engine/Source/Runtime/Experimental/Chaos/Public/Chaos/AABBTree.h
@@ -1042,7 +1042,7 @@ public:
 			StatMaxLeafSize = FMath::Max(StatMaxLeafSize, (int32)Leaf.GetElementCount());
 		}
 		TreeExpensiveStats.StatMaxLeafSize = StatMaxLeafSize;
-		TreeExpensiveStats.StatMaxTreeDepth = GetSubtreeDepth(0);
+		TreeExpensiveStats.StatMaxTreeDepth = (Nodes.Num() == 0) ? 0 : GetSubtreeDepth(0);
 		TreeExpensiveStats.StatGlobalPayloadsSize = GlobalPayloads.Num();
 
 		return TreeExpensiveStats;
