Change max children in leaf default back to 5

diff --git a/Engine/Source/Runtime/Experimental/Chaos/Private/Chaos/PBDRigidsEvolution.cpp b/Engine/Source/Runtime/Experimental/Chaos/Private/Chaos/PBDRigidsEvolution.cpp
--- a/Engine/Source/Runtime/Experimental/Chaos/Private/Chaos/PBDRigidsEvolution.cpp
+++ b/Engine/Source/Runtime/Experimental/Chaos/Private/Chaos/PBDRigidsEvolution.cpp
@@ -51,7 +51,7 @@ namespace Chaos
 		{
 			BroadphaseType = 3;
 			BVNumCells = 35;
-			MaxChildrenInLeaf = 80;
+			MaxChildrenInLeaf = 5;
 			MaxTreeDepth = 200;
 			AABBMaxChildrenInLeaf = 500;
 			AABBMaxTreeDepth = 200;
