Loop over all planes if the number of planes per vertex is higher than 3 + pass vertex index to EPA

diff --git a/Engine/Source/Runtime/Experimental/Chaos/Private/Chaos/CollisionOneShotManifolds.cpp b/Engine/Source/Runtime/Experimental/Chaos/Private/Chaos/CollisionOneShotManifolds.cpp
--- a/Engine/Source/Runtime/Experimental/Chaos/Private/Chaos/CollisionOneShotManifolds.cpp
+++ b/Engine/Source/Runtime/Experimental/Chaos/Private/Chaos/CollisionOneShotManifolds.cpp
@@ -650,6 +650,37 @@ namespace Chaos
 			return Contact;
 		}
 
+		// Find the the most opposing plane given a position and a direction
+		template <typename ConvexImplicitType>
+		void FindBestPlane(
+			const ConvexImplicitType& Convex,
+			const FVec3& X,
+			const FVec3& N,
+			const FReal MaxDistance,
+			const int32 PlaneIndex,
+			int32& BestPlaneIndex,
+			FReal& BestPlaneDot)
+		{
+			const TPlaneConcrete<FReal, 3> Plane = Convex.GetPlane(PlaneIndex);
+			const FReal PlaneNormalDotN = FVec3::DotProduct(N, Plane.Normal());
+				
+			// Ignore planes that do not oppose N
+			if (PlaneNormalDotN <= -SMALL_NUMBER)
+			{
+				// Reject planes farther than MaxDistance
+				const FReal PlaneDistance = Plane.SignedDistance(X);
+				if (FMath::Abs(PlaneDistance) <= MaxDistance)
+				{
+					// Keep the most opposing plane
+					if (PlaneNormalDotN < BestPlaneDot)
+					{
+						BestPlaneDot = PlaneNormalDotN;
+						BestPlaneIndex = PlaneIndex;
+					}
+				}
+			}
+		}
+
 		// Select one of the planes on the convex to use as the contact plane, given an estimated contact position and opposing 
 		// normal from GJK with margins (which gives the shapes rounded corners/edges).
 		template <typename ConvexImplicitType>
@@ -667,34 +698,25 @@ namespace Chaos
 
 			int32 BestPlaneIndex = INDEX_NONE;
 			FReal BestPlaneDot = 1.0f;
-			
 			{
 				int32 PlaneIndices[3] = {INDEX_NONE, INDEX_NONE, INDEX_NONE};
-				const int32 NumPlanes = Convex.GetVertexPlanes3(VertexIndex, PlaneIndices[0], PlaneIndices[1], PlaneIndices[2]);
-			
-				for (int32 PlaneIndex = 0; PlaneIndex < NumPlanes; ++PlaneIndex)
+				int32 NumPlanes = Convex.GetVertexPlanes3(VertexIndex, PlaneIndices[0], PlaneIndices[1], PlaneIndices[2]);
+
+				// If we have more than 3 planes we iterate over the full set of planes since it is faster than using the half edge structure
+				if(NumPlanes > 3)
 				{
-					const TPlaneConcrete<FReal, 3> Plane = Convex.GetPlane(PlaneIndices[PlaneIndex]);
-					const FReal PlaneNormalDotN = FVec3::DotProduct(N, Plane.Normal());
-			
-					// Ignore planes that do not oppose N
-					if (PlaneNormalDotN > -SMALL_NUMBER)
-					{
-						continue;
-					}
-			
-					// Reject planes farther than MaxDistance
-					const FReal PlaneDistance = Plane.SignedDistance(X);
-					if (FMath::Abs(PlaneDistance) > MaxDistance)
+					NumPlanes = Convex.NumPlanes();
+					for (int32 PlaneIndex = 0; PlaneIndex < NumPlanes; ++PlaneIndex)
 					{
-						continue;
+						FindBestPlane(Convex, X, N, MaxDistance, PlaneIndex, BestPlaneIndex, BestPlaneDot);
 					}
-			
-					// Keep the most opposing plane
-					if (PlaneNormalDotN < BestPlaneDot)
+				}
+				// Otherwise we iterate over the cached planes
+				else
+				{
+					for (int32 PlaneIndex = 0; PlaneIndex < NumPlanes; ++PlaneIndex)
 					{
-						BestPlaneDot = PlaneNormalDotN;
-						BestPlaneIndex = PlaneIndices[PlaneIndex];
+						FindBestPlane(Convex, X, N, MaxDistance, PlaneIndices[PlaneIndex], BestPlaneIndex, BestPlaneDot);
 					}
 				}
 			}
diff --git a/Engine/Source/Runtime/Experimental/Chaos/Public/Chaos/ConvexHalfEdgeStructureData.h b/Engine/Source/Runtime/Experimental/Chaos/Public/Chaos/ConvexHalfEdgeStructureData.h
--- a/Engine/Source/Runtime/Experimental/Chaos/Public/Chaos/ConvexHalfEdgeStructureData.h
+++ b/Engine/Source/Runtime/Experimental/Chaos/Public/Chaos/ConvexHalfEdgeStructureData.h
@@ -524,11 +524,12 @@ namespace Chaos
 				{
 					do
 					{
-						VertexPlane.PlaneIndices[VertexPlane.NumPlaneIndices++] = (FIndex)GetHalfEdgePlane(HalfEdgeIndex);
-						if (VertexPlane.NumPlaneIndices == (FIndex)UE_ARRAY_COUNT(VertexPlane.PlaneIndices))
+						if(VertexPlane.NumPlaneIndices < (FIndex)UE_ARRAY_COUNT(VertexPlane.PlaneIndices))
 						{
-							break;
+							VertexPlane.PlaneIndices[VertexPlane.NumPlaneIndices] = (FIndex)GetHalfEdgePlane(HalfEdgeIndex);
 						}
+						// Caching of the Max number of plane indices on this vertex (could be higher than 3)
+						++VertexPlane.NumPlaneIndices;
 
 						const int32 TwinHalfEdgeIndex = GetTwinHalfEdge(HalfEdgeIndex);
 						if (TwinHalfEdgeIndex == InvalidIndex)
diff --git a/Engine/Source/Runtime/Experimental/Chaos/Public/Chaos/GJK.h b/Engine/Source/Runtime/Experimental/Chaos/Public/Chaos/GJK.h
--- a/Engine/Source/Runtime/Experimental/Chaos/Public/Chaos/GJK.h
+++ b/Engine/Source/Runtime/Experimental/Chaos/Public/Chaos/GJK.h
@@ -308,11 +308,10 @@ namespace Chaos
 				VertsB.Add(BToATM.TransformPositionNoScale(InOutSimplexData.Bs[i]));
 			}
 			
-			auto SupportBInAFunc = [&B, &BToATM, &AToBRotation, &SupportDeltaB](const TVec3<T>& V)
+			auto SupportBInAFunc = [&B, &BToATM, &AToBRotation, &SupportDeltaB, &VertexIndexB](const TVec3<T>& V)
 			{
 				const TVec3<T> VInB = AToBRotation * V;
-				int32 VertexIndex = INDEX_NONE;
-				const TVec3<T> SupportBLocal = B.SupportCore(VInB, B.GetMargin(), &SupportDeltaB, VertexIndex);
+				const TVec3<T> SupportBLocal = B.SupportCore(VInB, B.GetMargin(), &SupportDeltaB, VertexIndexB);
 				return BToATM.TransformPositionNoScale(SupportBLocal);
 			};
 
