Chaos - Manifold creation optimization
- Attempt to reuse manifold points that have not moved far between ticks
- Attempt to inject GJK result into existing manifold

Also
- reduce number of temp CollisionConstraints creatd in TriMesh and heightfield collision

diff --git a/Engine/Source/Programs/HeadlessChaos/Private/GeometryCollection/GeometryCollectionTestSimulation.cpp b/Engine/Source/Programs/HeadlessChaos/Private/GeometryCollection/GeometryCollectionTestSimulation.cpp
--- a/Engine/Source/Programs/HeadlessChaos/Private/GeometryCollection/GeometryCollectionTestSimulation.cpp
+++ b/Engine/Source/Programs/HeadlessChaos/Private/GeometryCollection/GeometryCollectionTestSimulation.cpp
@@ -91,7 +91,10 @@ namespace GeometryCollectionTest
 		UnitTest.AddSimulationObject(Collection);
 		UnitTest.AddSimulationObject(Floor);
 		UnitTest.Initialize();
-		for (int i = 0; i < 10; i++) UnitTest.Advance();
+		for (int i = 0; i < 10; i++)
+		{
+			UnitTest.Advance();
+		}
 
 		{ // test results
 			EXPECT_LT(FMath::Abs(Collection->RestCollection->Transform[0].GetTranslation().Z), SMALL_THRESHOLD);
diff --git a/Engine/Source/Programs/HeadlessChaos/Private/HeadlessChaosTestGJK.cpp b/Engine/Source/Programs/HeadlessChaos/Private/HeadlessChaosTestGJK.cpp
--- a/Engine/Source/Programs/HeadlessChaos/Private/HeadlessChaosTestGJK.cpp
+++ b/Engine/Source/Programs/HeadlessChaos/Private/HeadlessChaosTestGJK.cpp
@@ -1872,24 +1872,24 @@ namespace ChaosTest
 		int32 VertexIndexB = INDEX_NONE;
 
 		// Separated (GJK)
-		GJKPenetrationWarmStartable<true>(Box, Box, BTM.GetRelativeTransformNoScale(ATM), FReal(0), FReal(0), Penetration, ClosestA, ClosestB, NormalA, NormalB, VertexIndexA, VertexIndexB, WarmStartData, SupportDelta);
+		GJKPenetrationWarmStartable(Box, Box, BTM.GetRelativeTransformNoScale(ATM), Penetration, ClosestA, ClosestB, NormalA, NormalB, VertexIndexA, VertexIndexB, WarmStartData, SupportDelta);
 		EXPECT_NEAR(Penetration, -5.0f, KINDA_SMALL_NUMBER);
 		EXPECT_NEAR(ClosestA.Z, (FReal)50.0f, (FReal)KINDA_SMALL_NUMBER);
 		EXPECT_NEAR(ClosestB.Z, (FReal)-50.0f, (FReal)KINDA_SMALL_NUMBER);
 
-		GJKPenetrationWarmStartable<true>(Box, Box, BTM.GetRelativeTransformNoScale(ATM), FReal(0), FReal(0), Penetration, ClosestA, ClosestB, NormalA, NormalB, VertexIndexA, VertexIndexB, WarmStartData, SupportDelta);
+		GJKPenetrationWarmStartable(Box, Box, BTM.GetRelativeTransformNoScale(ATM), Penetration, ClosestA, ClosestB, NormalA, NormalB, VertexIndexA, VertexIndexB, WarmStartData, SupportDelta);
 		EXPECT_NEAR(Penetration, -5.0f, KINDA_SMALL_NUMBER);
 
 		BTM = FRigidTransform3(FVec3(0, 0, 145), FRotation3::FromIdentity());
-		GJKPenetrationWarmStartable<true>(Box, Box, BTM.GetRelativeTransformNoScale(ATM), FReal(0), FReal(0), Penetration, ClosestA, ClosestB, NormalA, NormalB, VertexIndexA, VertexIndexB, WarmStartData, SupportDelta);
+		GJKPenetrationWarmStartable(Box, Box, BTM.GetRelativeTransformNoScale(ATM), Penetration, ClosestA, ClosestB, NormalA, NormalB, VertexIndexA, VertexIndexB, WarmStartData, SupportDelta);
 		EXPECT_NEAR(Penetration, -45.0f, KINDA_SMALL_NUMBER);
 
 		BTM = FRigidTransform3(FVec3(0, 0, 145), FRotation3::FromAxisAngle(FVec3(1, 0, 0), FMath::DegreesToRadians(110.0f)));
-		GJKPenetrationWarmStartable<true>(Box, Box, BTM.GetRelativeTransformNoScale(ATM), FReal(0), FReal(0), Penetration, ClosestA, ClosestB, NormalA, NormalB, VertexIndexA, VertexIndexB, WarmStartData, SupportDelta);
+		GJKPenetrationWarmStartable(Box, Box, BTM.GetRelativeTransformNoScale(ATM), Penetration, ClosestA, ClosestB, NormalA, NormalB, VertexIndexA, VertexIndexB, WarmStartData, SupportDelta);
 		EXPECT_NEAR(Penetration, -30.9144f, KINDA_SMALL_NUMBER);
 
 		FReal Penetration2;
-		GJKPenetrationWarmStartable<true>(Box, Box, BTM.GetRelativeTransformNoScale(ATM), FReal(0), FReal(0), Penetration2, ClosestA, ClosestB, NormalA, NormalB, VertexIndexA, VertexIndexB, WarmStartData, SupportDelta);
+		GJKPenetrationWarmStartable(Box, Box, BTM.GetRelativeTransformNoScale(ATM), Penetration2, ClosestA, ClosestB, NormalA, NormalB, VertexIndexA, VertexIndexB, WarmStartData, SupportDelta);
 		EXPECT_NEAR(Penetration2, Penetration, KINDA_SMALL_NUMBER);
 	}
 
@@ -1909,7 +1909,7 @@ namespace ChaosTest
 		int32 VertexIndexB = INDEX_NONE;
 
 		// Deep (EPA) No Margin
-		GJKPenetrationWarmStartable<true>(Box, Box, BTM.GetRelativeTransformNoScale(ATM), FReal(0), FReal(0), Penetration, ClosestA, ClosestB, NormalA, NormalB, VertexIndexA, VertexIndexB,  WarmStartData, SupportDelta);
+		GJKPenetrationWarmStartable(Box, Box, BTM.GetRelativeTransformNoScale(ATM), Penetration, ClosestA, ClosestB, NormalA, NormalB, VertexIndexA, VertexIndexB,  WarmStartData, SupportDelta);
 		EXPECT_NEAR(Penetration, 40.0f, KINDA_SMALL_NUMBER);
 		EXPECT_NEAR(ClosestA.Z, (FReal)50.0f, (FReal)KINDA_SMALL_NUMBER);
 		EXPECT_NEAR(ClosestB.Z, (FReal)-50.0f, (FReal)KINDA_SMALL_NUMBER);
@@ -1917,7 +1917,7 @@ namespace ChaosTest
 		// Deep (EPA) With Margin
 		WarmStartData = FGJKSimplexData();
 		TGJKCoreShape MarginBox(Box, FReal(10));
-		GJKPenetrationWarmStartable<true>(MarginBox, MarginBox, BTM.GetRelativeTransformNoScale(ATM), FReal(0), FReal(0), Penetration, ClosestA, ClosestB, NormalA, NormalB,  VertexIndexA, VertexIndexB, WarmStartData, SupportDelta);
+		GJKPenetrationWarmStartable(MarginBox, MarginBox, BTM.GetRelativeTransformNoScale(ATM), Penetration, ClosestA, ClosestB, NormalA, NormalB,  VertexIndexA, VertexIndexB, WarmStartData, SupportDelta);
 		EXPECT_NEAR(Penetration, 40.0f, KINDA_SMALL_NUMBER);
 		EXPECT_NEAR(ClosestA.Z, (FReal)50.0f, (FReal)KINDA_SMALL_NUMBER);
 		EXPECT_NEAR(ClosestB.Z, (FReal)-50.0f, (FReal)KINDA_SMALL_NUMBER);
@@ -1925,7 +1925,7 @@ namespace ChaosTest
 		// Deep (EPA) With Margin and Relative Rotation
 		WarmStartData = FGJKSimplexData();
 		ATM = FRigidTransform3(FVec3(0, 0, 0), FRotation3::FromAxisAngle(FVec3(1, 0, 0), FMath::DegreesToRadians(180)));
-		GJKPenetrationWarmStartable<true>(MarginBox, MarginBox, BTM.GetRelativeTransformNoScale(ATM), FReal(0), FReal(0), Penetration, ClosestA, ClosestB, NormalA, NormalB,  VertexIndexA, VertexIndexB, WarmStartData, SupportDelta);
+		GJKPenetrationWarmStartable(MarginBox, MarginBox, BTM.GetRelativeTransformNoScale(ATM), Penetration, ClosestA, ClosestB, NormalA, NormalB,  VertexIndexA, VertexIndexB, WarmStartData, SupportDelta);
 		EXPECT_NEAR(Penetration, 40.0f, KINDA_SMALL_NUMBER);
 		EXPECT_NEAR(ClosestA.Z, (FReal)-50.0f, (FReal)KINDA_SMALL_NUMBER);
 		EXPECT_NEAR(ClosestB.Z, (FReal)-50.0f, (FReal)KINDA_SMALL_NUMBER);
@@ -1946,13 +1946,13 @@ namespace ChaosTest
 		int32 VertexIndexA = INDEX_NONE;
 		int32 VertexIndexB = INDEX_NONE;
 
-		GJKPenetrationWarmStartable<true>(Sphere, Sphere, BTM.GetRelativeTransformNoScale(ATM), FReal(0), FReal(0), Penetration, ClosestA, ClosestB, NormalA, NormalB,  VertexIndexA, VertexIndexB, WarmStartData, SupportDelta);
+		GJKPenetrationWarmStartable(Sphere, Sphere, BTM.GetRelativeTransformNoScale(ATM), Penetration, ClosestA, ClosestB, NormalA, NormalB,  VertexIndexA, VertexIndexB, WarmStartData, SupportDelta);
 		EXPECT_NEAR(Penetration, (FReal)-5.0f, (FReal)KINDA_SMALL_NUMBER);
 		EXPECT_NEAR(ClosestA.Z, (FReal)50.0f, (FReal)KINDA_SMALL_NUMBER);
 		EXPECT_NEAR(ClosestB.Z, (FReal)-50.0f, (FReal)KINDA_SMALL_NUMBER);
 
 		BTM = FRigidTransform3(FVec3(0, 0, 105), FRotation3::FromAxisAngle(FVec3(1,0,0), FMath::DegreesToRadians(180)));
-		GJKPenetrationWarmStartable<true>(Sphere, Sphere, BTM.GetRelativeTransformNoScale(ATM), FReal(0), FReal(0), Penetration, ClosestA, ClosestB, NormalA, NormalB, VertexIndexA, VertexIndexB, WarmStartData, SupportDelta);
+		GJKPenetrationWarmStartable(Sphere, Sphere, BTM.GetRelativeTransformNoScale(ATM), Penetration, ClosestA, ClosestB, NormalA, NormalB, VertexIndexA, VertexIndexB, WarmStartData, SupportDelta);
 		EXPECT_NEAR(Penetration, (FReal)-5.0f, (FReal)KINDA_SMALL_NUMBER);
 		EXPECT_NEAR(ClosestA.Z, (FReal)50.0f, (FReal)KINDA_SMALL_NUMBER);
 		EXPECT_NEAR(ClosestB.Z, (FReal)50.0f, (FReal)KINDA_SMALL_NUMBER);
diff --git a/Engine/Source/Programs/HeadlessChaos/Public/HeadlessChaosCollisionConstraints.h b/Engine/Source/Programs/HeadlessChaos/Public/HeadlessChaosCollisionConstraints.h
--- a/Engine/Source/Programs/HeadlessChaos/Public/HeadlessChaosCollisionConstraints.h
+++ b/Engine/Source/Programs/HeadlessChaos/Public/HeadlessChaosCollisionConstraints.h
@@ -99,7 +99,7 @@ public:
 		FRigidTransform3 WorldTransform0 = Constraint.ImplicitTransform[0] * Collisions::GetTransform(Constraint.Particle[0]);
 		FRigidTransform3 WorldTransform1 = Constraint.ImplicitTransform[1] * Collisions::GetTransform(Constraint.Particle[1]);
 
-		Constraint.ResetPhi(TNumericLimits<FReal>::Max());
+		Constraint.ResetManifold();
 		Collisions::UpdateLevelsetLevelsetConstraint<ECollisionUpdateType::Deepest>(WorldTransform0, WorldTransform1, FReal(1 / 30.0f), Constraint);
 	}
 
diff --git a/Engine/Source/Runtime/Experimental/Chaos/Private/Chaos/ChaosDebugDraw.cpp b/Engine/Source/Runtime/Experimental/Chaos/Private/Chaos/ChaosDebugDraw.cpp
--- a/Engine/Source/Runtime/Experimental/Chaos/Private/Chaos/ChaosDebugDraw.cpp
+++ b/Engine/Source/Runtime/Experimental/Chaos/Private/Chaos/ChaosDebugDraw.cpp
@@ -906,13 +906,18 @@ namespace Chaos
 						FMatrix Axes = FRotationMatrix::MakeFromX(WorldPlaneNormal);
 						FDebugDrawQueue::GetInstance().DrawDebugCircle(WorldPointLocation, 0.5f * Settings.DrawScale * Settings.ContactWidth, 12, DiscColor, false, KINDA_SMALL_NUMBER, Settings.DrawPriority, Settings.LineThickness, Axes.GetUnitAxis(EAxis::Y), Axes.GetUnitAxis(EAxis::Z), false);
 
+						// Unfortunately this doesn't work any more - we overwrite the original contacts with the friction contacts
 						// Previous points
-						const FVec3 PrevPointLocation = PointTransform.TransformPosition(ManifoldPoint.CoMContactPoints[ContactPointOwner]);
-						const FVec3 PrevPlaneLocation = PlaneTransform.TransformPosition(ManifoldPoint.CoMContactPoints[ContactPlaneOwner]);
-						const FVec3 WorldPrevPointLocation = SpaceTransform.TransformPosition(PrevPointLocation);
-						const FVec3 WorldPrevPlaneLocation = SpaceTransform.TransformPosition(PrevPlaneLocation);
-						FDebugDrawQueue::GetInstance().DrawDebugLine(WorldPrevPointLocation, WorldPointLocation, FColor::White, false, KINDA_SMALL_NUMBER, Settings.DrawPriority, Settings.LineThickness);
-						FDebugDrawQueue::GetInstance().DrawDebugLine(WorldPrevPlaneLocation, WorldPlaneLocation, FColor::White, false, KINDA_SMALL_NUMBER, Settings.DrawPriority, Settings.LineThickness);
+						//const FManifoldPointSavedData* PrevManifoldPoint = Contact.FindManifoldPointSavedData(ManifoldPoint);
+						//if (PrevManifoldPoint != nullptr)
+						//{
+						//	const FVec3 PrevPointLocation = PointTransform.TransformPosition(PrevManifoldPoint->CoMContactPoints[ContactPointOwner]);
+						//	const FVec3 PrevPlaneLocation = PlaneTransform.TransformPosition(PrevManifoldPoint->CoMContactPoints[ContactPlaneOwner]);
+						//	const FVec3 WorldPrevPointLocation = SpaceTransform.TransformPosition(PrevPointLocation);
+						//	const FVec3 WorldPrevPlaneLocation = SpaceTransform.TransformPosition(PrevPlaneLocation);
+						//	FDebugDrawQueue::GetInstance().DrawDebugLine(WorldPrevPointLocation, WorldPointLocation, FColor::White, false, KINDA_SMALL_NUMBER, Settings.DrawPriority, Settings.LineThickness);
+						//	FDebugDrawQueue::GetInstance().DrawDebugLine(WorldPrevPlaneLocation, WorldPlaneLocation, FColor::White, false, KINDA_SMALL_NUMBER, Settings.DrawPriority, Settings.LineThickness);
+						//}
 
 						// Whether restored
 						if (Contact.WasManifoldRestored())
@@ -920,6 +925,11 @@ namespace Chaos
 							const FReal BoxScale = Settings.DrawScale * Settings.ContactWidth;
 							FDebugDrawQueue::GetInstance().DrawDebugBox(WorldPlaneLocation, FVec3(BoxScale, BoxScale, FReal(0.01)), FRotation3(FRotationMatrix::MakeFromZ(WorldPlaneNormal)), FColor::Blue, false, KINDA_SMALL_NUMBER, Settings.DrawPriority, 0.5f * Settings.LineThickness);
 						}
+						else if (ManifoldPoint.bWasRestored)
+						{
+							const FReal BoxScale = Settings.DrawScale * Settings.ContactWidth;
+							FDebugDrawQueue::GetInstance().DrawDebugBox(WorldPlaneLocation, FVec3(BoxScale, BoxScale, FReal(0.01)), FRotation3(FRotationMatrix::MakeFromZ(WorldPlaneNormal)), FColor::Purple, false, KINDA_SMALL_NUMBER, Settings.DrawPriority, 0.5f * Settings.LineThickness);
+						}
 
 						// How many iterations (0 = green, 8 = red)
 						if (bChaosDebugDebugDrawContactIterations)
diff --git a/Engine/Source/Runtime/Experimental/Chaos/Private/Chaos/Collision/PBDCollisionConstraint.cpp b/Engine/Source/Runtime/Experimental/Chaos/Private/Chaos/Collision/PBDCollisionConstraint.cpp
--- a/Engine/Source/Runtime/Experimental/Chaos/Private/Chaos/Collision/PBDCollisionConstraint.cpp
+++ b/Engine/Source/Runtime/Experimental/Chaos/Private/Chaos/Collision/PBDCollisionConstraint.cpp
@@ -93,10 +93,17 @@ namespace Chaos
 		return TUniquePtr<FPBDCollisionConstraint>(Constraint);
 	}
 
+	FPBDCollisionConstraint FPBDCollisionConstraint::MakeTriangle(const FImplicitObject* Implicit0)
+	{
+		FPBDCollisionConstraint Constraint;
+		Constraint.InitMargins(Implicit0->GetCollisionType(), ImplicitObjectType::Triangle, Implicit0->GetMargin(), FReal(0));
+		return Constraint;
+	}
+
 	FPBDCollisionConstraint FPBDCollisionConstraint::MakeCopy(
 		const FPBDCollisionConstraint& Source)
 	{
-		// @todo(chaos): The resim cache version probably doesn't need all the data, so maybe try to cur this down?
+		// @todo(chaos): The resim cache version probably doesn't need all the data, so maybe try to cut this down?
 		FPBDCollisionConstraint Constraint = Source;
 
 		// Invalidate the data that maps the constraint to its container (we are no longer in the container)
@@ -115,11 +122,17 @@ namespace Chaos
 		, CCDType(ECollisionCCDType::Disabled)
 		, Stiffness(FReal(1))
 		, CullDistance(TNumericLimits<FReal>::Max())
+		, CollisionMargins{ 0, 0 }
+		, CollisionTolerance(0)
 		, bUseManifold(false)
 		, bUseIncrementalManifold(false)
 		, SolverBodies{ nullptr, nullptr }
 		, GJKWarmStartData()
-		, PrevManifoldPoints()
+		, ManifoldPointSavedData{}
+		, NumSavedManifoldPoints(0)
+		, LastShapeWorldTransform0()
+		, LastShapeWorldTransform1()
+		, ExpectedNumManifoldPoints(0)
 		, bWasManifoldRestored(false)
 		, NumActivePositionIterations(0)
 	{
@@ -145,11 +158,17 @@ namespace Chaos
 		, CCDType(ECollisionCCDType::Disabled)
 		, Stiffness(FReal(1))
 		, CullDistance(TNumericLimits<FReal>::Max())
+		, CollisionMargins{ 0, 0 }
+		, CollisionTolerance(0)
 		, bUseManifold(false)
 		, bUseIncrementalManifold(false)
 		, SolverBodies{ nullptr, nullptr }
 		, GJKWarmStartData()
-		, PrevManifoldPoints()
+		, ManifoldPointSavedData{}
+		, NumSavedManifoldPoints(0)
+		, LastShapeWorldTransform0()
+		, LastShapeWorldTransform1()
+		, ExpectedNumManifoldPoints(0)
 		, bWasManifoldRestored(false)
 	{
 		Manifold.Implicit[0] = Implicit0;
@@ -178,6 +197,51 @@ namespace Chaos
 
 		bUseManifold = bInUseManifold && CanUseManifold(Particle[0], Particle[1]);
 		bUseIncrementalManifold = true;	// This will get changed later if we call AddOneShotManifoldContact
+
+		const FReal Margin0 = GetImplicit0()->GetMargin();
+		const FReal Margin1 = GetImplicit1()->GetMargin();
+		const EImplicitObjectType ImplicitType0 = GetInnerType(GetImplicit0()->GetCollisionType());
+		const EImplicitObjectType ImplicitType1 = GetInnerType(GetImplicit1()->GetCollisionType());
+		InitMargins(ImplicitType0, ImplicitType1, Margin0, Margin1);
+	}
+
+	void FPBDCollisionConstraint::InitMargins(const EImplicitObjectType ImplicitType0, const EImplicitObjectType ImplicitType1, const FReal Margin0, const FReal Margin1)
+	{
+		// Set up the margins and tolerances to be used during the narrow phase.
+		// One shape in a collision will always have a margin. Only triangles have zero margin and we don't 
+		// collide two triangles. If we have a triangle, it is always the second shape.
+		// The collision tolerance is used for knowing whether a new contact matches an existing one.
+		// If we have two non-quadratic shapes, we use the smallest margin on both shapes.
+		// If we have a quadratic shape versus a non-quadratic, we don't need a margin on the non-quadratic.
+		// For non-quadratics the collision tolerance is the smallest non-zero margin. 
+		// For quadratic shapes we want a collision tolerance much smaller than the radius.
+		const bool bIsQuadratic0 = ((ImplicitType0 == ImplicitObjectType::Sphere) || (ImplicitType0 == ImplicitObjectType::Capsule));
+		const bool bIsQuadratic1 = ((ImplicitType1 == ImplicitObjectType::Sphere) || (ImplicitType1 == ImplicitObjectType::Capsule));
+		const FReal QuadraticToleranceScale = 0.05f;
+		if (!bIsQuadratic0 && !bIsQuadratic1)
+		{
+			CollisionMargins[0] = FMath::Min(Margin0, Margin1);
+			CollisionMargins[1] = CollisionMargins[0];
+			CollisionTolerance = ((Margin0 < Margin1) || (Margin1 == 0)) ? Margin0 : Margin1;
+		}
+		else if (bIsQuadratic0 && bIsQuadratic1)
+		{
+			CollisionMargins[0] = Margin0;
+			CollisionMargins[1] = Margin1;
+			CollisionTolerance = QuadraticToleranceScale * FMath::Min(Margin0, Margin1);
+		}
+		else if (bIsQuadratic0 && !bIsQuadratic1)
+		{
+			CollisionMargins[0] = Margin0;
+			CollisionMargins[1] = 0;
+			CollisionTolerance = QuadraticToleranceScale * Margin0;
+		}
+		else if (!bIsQuadratic0 && bIsQuadratic1)
+		{
+			CollisionMargins[0] = 0;
+			CollisionMargins[1] = Margin1;
+			CollisionTolerance = QuadraticToleranceScale * Margin1;
+		}
 	}
 
 	void FPBDCollisionConstraint::SetIsSleeping(const bool bInIsSleeping)
@@ -278,7 +342,6 @@ namespace Chaos
 		return BestMatchIndex;
 	}
 
-	// @todo(chaos): this exists primarily for the contact modification callbacks. See if we can eliminate it when not needed
 	void FPBDCollisionConstraint::UpdateManifoldContacts()
 	{
 		FVec3 P0, P1;
@@ -301,6 +364,8 @@ namespace Chaos
 			Q1 = FParticleUtilities::GetCoMWorldRotation(FConstGenericParticleHandle(Particle[1]));
 		}
 
+		Manifold.Reset();
+
 		for (int32 Index = 0; Index < ManifoldPoints.Num(); Index++)
 		{
 			FManifoldPoint& ManifoldPoint = ManifoldPoints[Index];
@@ -316,9 +381,14 @@ namespace Chaos
 		}
 	}
 
-	void FPBDCollisionConstraint::AddOneshotManifoldContact(const FContactPoint& ContactPoint, const FReal Dt)
+	void FPBDCollisionConstraint::AddOneshotManifoldContact(const FContactPoint& ContactPoint)
 	{
-		int32 ManifoldPointIndex = AddManifoldPoint(ContactPoint, Dt);
+		if (ManifoldPoints.Num() == MaxManifoldPoints)
+		{
+			return;
+		}
+
+		int32 ManifoldPointIndex = AddManifoldPoint(ContactPoint);
 
 		// Copy currently active point
 		if (ManifoldPoints[ManifoldPointIndex].ContactPoint.Phi < Manifold.Phi)
@@ -329,8 +399,13 @@ namespace Chaos
 		bUseIncrementalManifold = false;
 	}
 
-	void FPBDCollisionConstraint::AddIncrementalManifoldContact(const FContactPoint& ContactPoint, const FReal Dt)
+	void FPBDCollisionConstraint::AddIncrementalManifoldContact(const FContactPoint& ContactPoint)
 	{
+		if (ManifoldPoints.Num() == MaxManifoldPoints)
+		{
+			return;
+		}
+
 		if (bUseManifold)
 		{
 			// See if the manifold point already exists
@@ -338,12 +413,12 @@ namespace Chaos
 			if (ManifoldPointIndex >= 0)
 			{
 				// This contact point is already in the manifold - update the state
-				UpdateManifoldPoint(ManifoldPointIndex, ContactPoint, Dt);
+				UpdateManifoldPoint(ManifoldPointIndex, ContactPoint);
 			}
 			else
 			{
 				// This is a new manifold point - capture the state and generate initial properties
-				ManifoldPointIndex = AddManifoldPoint(ContactPoint, Dt);
+				ManifoldPointIndex = AddManifoldPoint(ContactPoint);
 			}
 
 			// Copy currently active point
@@ -364,7 +439,7 @@ namespace Chaos
 				ManifoldPoints[0].ContactPoint = ContactPoint;
 			}
 
-			InitManifoldPoint(0, Dt);
+			InitManifoldPoint(0);
 
 			SetActiveContactPoint(ManifoldPoints[0].ContactPoint);
 		}
@@ -377,10 +452,8 @@ namespace Chaos
 		ManifoldPoints.Reset();
 	}
 
-	void FPBDCollisionConstraint::InitManifoldPoint(const int32 ManifoldPointIndex, FReal Dt)
+	void FPBDCollisionConstraint::InitManifoldPoint(const int32 ManifoldPointIndex)
 	{
-		FManifoldPoint& ManifoldPoint = ManifoldPoints[ManifoldPointIndex];
-
 		FConstGenericParticleHandle Particle0 = Particle[0];
 		FConstGenericParticleHandle Particle1 = Particle[1];
 		if (!Particle0.IsValid() || !Particle1.IsValid())
@@ -390,6 +463,10 @@ namespace Chaos
 			return;
 		}
 
+		FManifoldPoint& ManifoldPoint = ManifoldPoints[ManifoldPointIndex];
+		ManifoldPoint.InitialShapeContactPoints[0] = ManifoldPoint.ContactPoint.ShapeContactPoints[0];
+		ManifoldPoint.InitialShapeContactPoints[1] = ManifoldPoint.ContactPoint.ShapeContactPoints[1];
+
 		// Update the derived contact state (CoM relative data)
 		UpdateManifoldPointFromContact(ManifoldPointIndex);
 
@@ -397,16 +474,18 @@ namespace Chaos
 		TryRestoreFrictionData(ManifoldPointIndex);
 	}
 
-	int32 FPBDCollisionConstraint::AddManifoldPoint(const FContactPoint& ContactPoint, const FReal Dt)
+	int32 FPBDCollisionConstraint::AddManifoldPoint(const FContactPoint& ContactPoint)
 	{
+		check(ManifoldPoints.Num() < MaxManifoldPoints);
+
 		int32 ManifoldPointIndex = ManifoldPoints.Add(ContactPoint);
 
-		InitManifoldPoint(ManifoldPointIndex, Dt);
+		InitManifoldPoint(ManifoldPointIndex);
 
 		return ManifoldPointIndex;
 	}
 
-	void FPBDCollisionConstraint::UpdateManifoldPoint(int32 ManifoldPointIndex, const FContactPoint& ContactPoint, const FReal Dt)
+	void FPBDCollisionConstraint::UpdateManifoldPoint(int32 ManifoldPointIndex, const FContactPoint& ContactPoint)
 	{
 		// We really need to know that it's exactly the same contact and not just a close one to update it here
 		// otherwise the PrevLocalContactPoint1 we calculated is not longer for the correct point.
@@ -499,81 +578,394 @@ namespace Chaos
 	}
 
 	void FPBDCollisionConstraint::ResetManifold()
+	{
+		NumSavedManifoldPoints = 0;
+		ResetActiveManifoldContacts();
+	}
+
+	void FPBDCollisionConstraint::ResetActiveManifoldContacts()
 	{
 		ManifoldPoints.Reset();
-		PrevManifoldPoints.Reset();
 		Manifold.Reset();
+		ExpectedNumManifoldPoints = 0;
+		bWasManifoldRestored = false;
 	}
 
 	void FPBDCollisionConstraint::SaveManifold()
 	{
+		check(ManifoldPoints.Num() <= MaxManifoldPoints);
+
 		// Save off the previous data for use by static friction
-		// We do this by swapping the arrays and resetting to avoid allocations
-		Swap(ManifoldPoints, PrevManifoldPoints);
-		ManifoldPoints.Reset();
-		Manifold.Reset();
+		NumSavedManifoldPoints = 0;
+		const int32 NumManifoldPoints = ManifoldPoints.Num();
+		for (int32 ManifoldPointIndex = 0; ManifoldPointIndex < NumManifoldPoints; ++ManifoldPointIndex)
+		{
+			// Only save points that actually produced a response and are holding static friction
+			const FManifoldPoint& ManifoldPoint = ManifoldPoints[ManifoldPointIndex];
+			if (ManifoldPoint.bInsideStaticFrictionCone && !ManifoldPoint.NetPushOut.IsNearlyZero())
+			{
+				ManifoldPointSavedData[NumSavedManifoldPoints++].Save(ManifoldPoint);
+			}
+		}
 
 		bWasManifoldRestored = false;
 	}
 
 	void FPBDCollisionConstraint::RestoreManifold()
 	{
-		// We want to restore the manifold as-is and will skip the narrow phase, which means we just
-		// leave the manifold in place, but we still have some cleanup to do to account for slight movement of the bodies
-		Manifold.Phi = TNumericLimits<FReal>::Max();
-
+		// We want to restore the manifold as-is and will skip the narrow phase, which means we leave the manifold in place, 
+		// but we still have some cleanup to do to account for slight movement of the bodies. E.g., we need to update the 
+		// world-space state for the contact modifiers
 		UpdateManifoldContacts();
 
 		bWasManifoldRestored = true;
 	}
 
-	bool FPBDCollisionConstraint::TryRestoreFrictionData(const int32 ManifoldPointIndex)
+	void FPBDCollisionConstraint::UpdateLastShapeWorldTransforms(const FRigidTransform3& ShapeWorldTransform0, const FRigidTransform3& ShapeWorldTransform1)
 	{
-		// Find the previous manifold point that matches
-		FManifoldPoint& ManifoldPoint = ManifoldPoints[ManifoldPointIndex];
-		
-		// Assume we have no matching point from the previous tick, but assume we can retain friction from now on
-		// Not supported for non-manifolds yet (hopefully we don't need to)
-		ManifoldPoint.bInsideStaticFrictionCone = bUseManifold;
-		ManifoldPoint.StaticFrictionMax = FReal(0);
+		LastShapeWorldTransform0 = ShapeWorldTransform0;
+		LastShapeWorldTransform1 = ShapeWorldTransform1;
+	}
 
-		if (bChaos_Manifold_EnableFrictionRestore)
+	bool FPBDCollisionConstraint::UpdateAndTryRestoreManifold(const FRigidTransform3& ShapeWorldTransform0, const FRigidTransform3& ShapeWorldTransform1)
+	{
+		check(ManifoldPoints.Num() <= MaxManifoldPoints);
+
+		// @todo(chaos): tune the multipliers
+		const FReal ContactPositionTolerance = FReal(0.8) * CollisionTolerance;
+		const FReal ShapePositionTolerance = (ManifoldPoints.Num() > 0) ? FReal(0.2) * CollisionTolerance : FReal(0.5) * CollisionTolerance;
+		const FReal ShapeRotationThreshold = (ManifoldPoints.Num() > 0) ? FReal(0.9999) : FReal(0.9998);
+		const FReal ContactPositionToleranceSq = FMath::Square(ContactPositionTolerance);
+
+		// Reset current closest point
+		Manifold.Reset();
+
+		// How many manifold points we expect. E.g., for Box-box this will be 4 or 1 depending on whether
+		// we have a face or edge contact. We don't reuse the manifold if we lose points after culling
+		// here and potentially adding the new narrow phase result (See TryAddManifoldContact).
+		ExpectedNumManifoldPoints = ManifoldPoints.Num();
+		bWasManifoldRestored = false;
+
+		const FRigidTransform3 Shape0ToShape1Transform = ShapeWorldTransform0.GetRelativeTransformNoScale(ShapeWorldTransform1);
+
+		// Update and prune manifold points.
+		TArray<int32, TInlineAllocator<4>> ManifoldPointsToRemove;
+		const int32 NumManifoldPoints = ManifoldPoints.Num();
+		for (int32 ManifoldPointIndex = 0; ManifoldPointIndex < NumManifoldPoints; ++ManifoldPointIndex)
 		{
-			// ManifoldPoints and PrevManifoldPoints are usually in the same order, so this loop should normally terminate in first iteration
-			int32 PrevManifoldPointIndex = ManifoldPointIndex;
-			for (int32 PrevManifoldPointOffset = 0; PrevManifoldPointOffset < PrevManifoldPoints.Num(); ++PrevManifoldPointOffset, ++PrevManifoldPointIndex)
+			FManifoldPoint& ManifoldPoint = ManifoldPoints[ManifoldPointIndex];
+
+			// Calculate the world-space contact location and separation at the current shape transforms
+			// @todo(chaos): this should use the normal owner. Currently we assume body 1 is the owner
+			const FVec3 Contact0In1 = Shape0ToShape1Transform.TransformPositionNoScale(ManifoldPoint.InitialShapeContactPoints[0]);
+			const FVec3& Contact1In1 = ManifoldPoint.InitialShapeContactPoints[1];
+			const FVec3 ContactNormalIn1 = ShapeWorldTransform1.InverseTransformVectorNoScale(ManifoldPoint.ContactPoint.Normal);
+
+			const FVec3 ContactDeltaIn1 = Contact0In1 - Contact1In1;
+			const FReal ContactPhi = FVec3::DotProduct(ContactDeltaIn1, ContactNormalIn1);
+			const FVec3 ContactLateralDeltaIn1 = ContactDeltaIn1 - ContactPhi * ContactNormalIn1;
+			const FReal ContactLateralDistanceSq = ContactLateralDeltaIn1.SizeSquared();
+
+			// Either update the point or flag it for removal
+			if (ContactLateralDistanceSq < ContactPositionToleranceSq)
 			{
-				if (PrevManifoldPointIndex >= PrevManifoldPoints.Num())
+				// Recalculate the contact points at the new location
+				// @todo(chaos): we should reproject the contact on the plane owner
+				FConstGenericParticleHandle Particle1 = Particle[1];
+				const FVec3 ShapeContactPoint1 = Contact0In1 - ContactPhi * ContactNormalIn1;
+				const FVec3 ActorContactPoint = ImplicitTransform[1].TransformPositionNoScale(ShapeContactPoint1);
+				ManifoldPoint.ContactPoint.ShapeContactPoints[1] = ShapeContactPoint1;
+				ManifoldPoint.CoMContactPoints[1] = Particle1->RotationOfMass().Inverse() * (ActorContactPoint - Particle1->CenterOfMass());
+				ManifoldPoint.ContactPoint.Location = ShapeWorldTransform1.TransformPositionNoScale(FReal(0.5) * (ShapeContactPoint1 + Contact0In1));
+				ManifoldPoint.ContactPoint.Phi = ContactPhi;
+				ManifoldPoint.bWasRestored = true;
+				TryRestoreFrictionData(ManifoldPointIndex);
+				if (ContactPhi < Manifold.Phi)
 				{
-					PrevManifoldPointIndex = 0;
+					// Update closest point
+					SetActiveContactPoint(ManifoldPoint.ContactPoint);
 				}
+			}
+			else
+			{
+				ManifoldPointsToRemove.Add(ManifoldPointIndex);
+			}
+		}
 
-				const FManifoldPoint& PrevManifoldPoint = PrevManifoldPoints[PrevManifoldPointIndex];
-				if (PrevManifoldPoint.bInsideStaticFrictionCone)
+		if ((ManifoldPointsToRemove.Num() == 0) && (ShapePositionTolerance > 0) && (ShapeRotationThreshold > 0))
+		{
+			// If we did not remove any contact points and we have not moved or rotated much we can reuse the manifold as-is.
+			// The transform check is necessary regardless of how many points we have left in the manifold because
+			// as a body moves/rotates we may have to change which faces/edges are colliding. We can't know if the face/edge
+			// will change until we run the closest-point checks (GJK) in the narrow phase.
+			const FVec3 Shape1ToShape0Translation = ShapeWorldTransform0.GetTranslation() - ShapeWorldTransform1.GetTranslation();
+			const FVec3 OriginalShape1ToShape0Translation = LastShapeWorldTransform0.GetTranslation() - LastShapeWorldTransform1.GetTranslation();
+			const FVec3 TranslationDelta = Shape1ToShape0Translation - OriginalShape1ToShape0Translation;
+			if (TranslationDelta.IsNearlyZero(ShapePositionTolerance))
+			{
+				const FRotation3 Shape1toShape0Rotation = ShapeWorldTransform0.GetRotation().Inverse() * ShapeWorldTransform1.GetRotation();
+				const FRotation3 OriginalShape1toShape0Rotation = LastShapeWorldTransform0.GetRotation().Inverse() * LastShapeWorldTransform1.GetRotation();
+				const FReal RotationOverlap = FRotation3::DotProduct(Shape1toShape0Rotation, OriginalShape1toShape0Rotation);
+				if (RotationOverlap > ShapeRotationThreshold)
 				{
-					const FVec3 DP0 = ManifoldPoint.CoMContactPoints[0] - PrevManifoldPoint.CoMContactPoints[0];
-					const FVec3 DP1 = ManifoldPoint.CoMContactPoints[1] - PrevManifoldPoint.CoMContactPoints[1];
-
-					// If the contact point is in the same spot on one of the bodies, assume it is the same contact
-					// @todo(chaos): more robust same-point test. E.g., this won't work for very small or very large objects,
-					// so at least make the tolerance size-dependent
-					const FReal DistanceToleranceSq = FMath::Square(Chaos_Manifold_FrictionPositionTolerance);
-					if ((DP0.SizeSquared() < DistanceToleranceSq) || (DP1.SizeSquared() < DistanceToleranceSq))
+					return true;
+				}
+			}
+		}
+
+		// We removed some points - process in reverse order for simpler index handling and faster removal of multiple elements
+		for (int32 RemoveArrayIndex = ManifoldPointsToRemove.Num() - 1; RemoveArrayIndex >= 0; --RemoveArrayIndex)
+		{
+			ManifoldPoints.RemoveAt(ManifoldPointsToRemove[RemoveArrayIndex]);
+		}
+
+		return false;
+	}
+
+	bool FPBDCollisionConstraint::TryAddManifoldContact(const FContactPoint& NewContactPoint, const FRigidTransform3& ShapeWorldTransform0, const FRigidTransform3& ShapeWorldTransform1)
+	{
+		check(ManifoldPoints.Num() <= MaxManifoldPoints);
+
+		const FReal PositionTolerance = FReal(1.0) * CollisionTolerance;
+		const FReal NormalThreshold = FReal(0.7);
+
+		// We must end up with a full manifold after this if we want to reuse it
+		if ((ManifoldPoints.Num() < ExpectedNumManifoldPoints - 1) || (ExpectedNumManifoldPoints == 0))
+		{
+			// We need to add more than 1 point to restore the manifold so we must rebuild it from scratch
+			return false;
+		}
+
+		// Find the matching manifold point if it exists and replace it
+		// Also check to see if the normal has changed significantly, and if it has, force mnanifold regeneration
+		// NOTE: the normal rejection check assumes all contacts have the same normal - this may not always be true. The worst
+		// case here is that we will regenerate the manifold too often so it will work but could be bad for perf
+		const FReal PositionToleranceSq = FMath::Square(PositionTolerance);
+		int32 MatchedManifoldPointIndex = INDEX_NONE;
+		const int32 NumManifoldPoints = ManifoldPoints.Num();
+		for (int32 ManifoldPointIndex = 0; ManifoldPointIndex < NumManifoldPoints; ++ManifoldPointIndex)
+		{
+			FManifoldPoint& ManifoldPoint = ManifoldPoints[ManifoldPointIndex];
+
+			const FReal NormalOverlap = FVec3::DotProduct(ManifoldPoint.ContactPoint.Normal, NewContactPoint.Normal);
+			if (NormalOverlap < NormalThreshold)
+			{
+				return false;
+			}
+
+			const FVec3 DR0 = ManifoldPoint.ContactPoint.ShapeContactPoints[0] - NewContactPoint.ShapeContactPoints[0];
+			const FVec3 DR1 = ManifoldPoint.ContactPoint.ShapeContactPoints[1] - NewContactPoint.ShapeContactPoints[1];
+			if ((DR0.SizeSquared() < PositionToleranceSq) && (DR1.SizeSquared() < PositionToleranceSq))
+			{
+				// If we should replace a point but will then have too few points we abort
+				if (ManifoldPoints.Num() < ExpectedNumManifoldPoints)
+				{
+					return false;
+				}
+
+				// If the existing point has a deeper penetration, just re-use it. This is common when we have a GJK
+				// result on an edge or corner - the contact created when generating the manifold is on the
+				// surface shape rather than the rounded (margin-reduced) shape.
+				if (ManifoldPoint.ContactPoint.Phi > NewContactPoint.Phi)
+				{
+					ManifoldPoint.ContactPoint = NewContactPoint;
+					ManifoldPoint.InitialShapeContactPoints[0] = NewContactPoint.ShapeContactPoints[0];
+					ManifoldPoint.InitialShapeContactPoints[1] = NewContactPoint.ShapeContactPoints[1];
+					ManifoldPoint.bWasRestored = false;
+					UpdateManifoldPointFromContact(ManifoldPointIndex);
+					TryRestoreFrictionData(ManifoldPointIndex);
+					if (NewContactPoint.Phi < GetPhi())
 					{
-						// If we have a previous point, use it to set the previous-state data required for 
-						// static friction otherwise reset static friction limits
-						ManifoldPoint.CoMContactPoints[0] = PrevManifoldPoint.CoMContactPoints[0];
-						ManifoldPoint.CoMContactPoints[1] = PrevManifoldPoint.CoMContactPoints[1];
-						ManifoldPoint.StaticFrictionMax = PrevManifoldPoint.StaticFrictionMax;
-						return true;
+						SetActiveContactPoint(ManifoldPoint.ContactPoint);
 					}
 				}
+
+				return true;
 			}
 		}
 
+		// If we have a full manifold, see if we can use or reject the GJK point
+		if (ManifoldPoints.Num() == 4)
+		{
+			return TryInsertManifoldContact(NewContactPoint, ShapeWorldTransform0, ShapeWorldTransform1);
+		}
+		
 		return false;
 	}
 
+	bool FPBDCollisionConstraint::TryInsertManifoldContact(const FContactPoint& NewContactPoint, const FRigidTransform3& ShapeWorldTransform0, const FRigidTransform3& ShapeWorldTransform1)
+	{
+		check(ManifoldPoints.Num() == 4);
+
+		const int32 NormalBodyIndex = 1;
+		constexpr int32 NumContactPoints = 5;
+		constexpr int32 NumManifoldPoints = 4;
+
+		// We want to select 4 points from the 5 we have
+		// Create a working set of points, and keep track which points have been selected
+		FVec3 ContactPoints[NumContactPoints];
+		FReal ContactPhis[NumContactPoints];
+		bool bContactSelected[NumContactPoints];
+		int32 SelectedContactIndices[NumManifoldPoints];
+		for (int32 ContactIndex = 0; ContactIndex < NumManifoldPoints; ++ContactIndex)
+		{
+			const FManifoldPoint& ManifoldPoint = ManifoldPoints[ContactIndex];
+			ContactPoints[ContactIndex] = ManifoldPoint.ContactPoint.ShapeContactPoints[NormalBodyIndex];
+			ContactPhis[ContactIndex] = ManifoldPoint.ContactPoint.Phi;
+			bContactSelected[ContactIndex] = false;
+		}
+		ContactPoints[4] = NewContactPoint.ShapeContactPoints[NormalBodyIndex];
+		ContactPhis[4] = NewContactPoint.Phi;
+		bContactSelected[4] = false;
+
+		// We are projecting points into a plane perpendicular to the contact normal, which we assume is the new point's normal
+		const FVec3 ContactNormal = NewContactPoint.ShapeContactNormal;
+
+		// Start with the deepest point. This may not be point 4 despite that being the result of
+		// collision detection because for some shape types we use margin-reduced core shapes which
+		// are effectively rounded at the corners. But...when building a one-shot manifold we 
+		// use the outer shape to get sharp corners. So, if we have a GJK result from a "corner"
+		// the real corner (if it is in the manifold) may actually be deeper than the GJK result.
+		SelectedContactIndices[0] = 0;
+		for (int32 ContactIndex = 1; ContactIndex < NumContactPoints; ++ContactIndex)
+		{
+			if (ContactPhis[ContactIndex] < ContactPhis[SelectedContactIndices[0]])
+			{
+				SelectedContactIndices[0] = ContactIndex;
+			}
+		}
+		bContactSelected[SelectedContactIndices[0]] = true;
+
+		// The second point will be the one farthest from the first
+		SelectedContactIndices[1] = INDEX_NONE;
+		FReal MaxDistanceSq = TNumericLimits<FReal>::Lowest();
+		for (int32 ContactIndex = 0; ContactIndex < NumContactPoints; ++ContactIndex)
+		{
+			if (!bContactSelected[ContactIndex])
+			{
+				const FReal DistanceSq = (ContactPoints[ContactIndex] - ContactPoints[SelectedContactIndices[0]]).SizeSquared();
+				if (DistanceSq > MaxDistanceSq)
+				{
+					SelectedContactIndices[1] = ContactIndex;
+					MaxDistanceSq = DistanceSq;
+				}
+			}
+		}
+		check(SelectedContactIndices[1] != INDEX_NONE);
+		bContactSelected[SelectedContactIndices[1]] = true;
+
+		// The third point is the one which gives us the largest triangle (projected onto a plane perpendicular to the normal)
+		SelectedContactIndices[2] = INDEX_NONE;
+		FReal MaxTriangleArea = 0;
+		FReal WindingOrder = FReal(1.0);
+		for (int32 ContactIndex = 0; ContactIndex < NumContactPoints; ++ContactIndex)
+		{
+			if (!bContactSelected[ContactIndex])
+			{
+				const FVec3 Cross = FVec3::CrossProduct(ContactPoints[SelectedContactIndices[1]] - ContactPoints[SelectedContactIndices[0]], ContactPoints[ContactIndex] - ContactPoints[SelectedContactIndices[1]]);
+				const FReal SignedArea = FVec3::DotProduct(Cross, ContactNormal);
+				if (FMath::Abs(SignedArea) > MaxTriangleArea)
+				{
+					SelectedContactIndices[2] = ContactIndex;
+					MaxTriangleArea = FMath::Abs(SignedArea);
+					WindingOrder = FMath::Sign(SignedArea);
+				}
+			}
+		}
+		if (SelectedContactIndices[2] == INDEX_NONE)
+		{
+			// Degenerate points - all 4 exactly in a line
+			return false;
+		}
+		bContactSelected[SelectedContactIndices[2]] = true;
+
+		// The fourth point is the one which adds the most area to the 3 points we already have
+		SelectedContactIndices[3] = INDEX_NONE;
+		FReal MaxQuadArea = 0;	// Additional area to MaxTriangleArea
+		for (int32 ContactIndex = 0; ContactIndex < NumContactPoints; ++ContactIndex)
+		{
+			if (!bContactSelected[ContactIndex])
+			{
+				// Calculate the area that is added by inserting the point into each edge of the selected triangle
+				// The signed area will be negative for interior points, positive for points that extend the triangle into a quad.
+				const FVec3 Cross0 = FVec3::CrossProduct(ContactPoints[ContactIndex] - ContactPoints[SelectedContactIndices[0]], ContactPoints[SelectedContactIndices[1]] - ContactPoints[ContactIndex]);
+				const FReal SignedArea0 = WindingOrder * FVec3::DotProduct(Cross0, ContactNormal);
+				const FVec3 Cross1 = FVec3::CrossProduct(ContactPoints[ContactIndex] - ContactPoints[SelectedContactIndices[1]], ContactPoints[SelectedContactIndices[2]] - ContactPoints[ContactIndex]);
+				const FReal SignedArea1 = WindingOrder * FVec3::DotProduct(Cross1, ContactNormal);
+				const FVec3 Cross2 = FVec3::CrossProduct(ContactPoints[ContactIndex] - ContactPoints[SelectedContactIndices[2]], ContactPoints[SelectedContactIndices[0]] - ContactPoints[ContactIndex]);
+				const FReal SignedArea2 = WindingOrder * FVec3::DotProduct(Cross2, ContactNormal);
+				const FReal SignedArea = FMath::Max3(SignedArea0, SignedArea1, SignedArea2);
+				if (SignedArea > MaxQuadArea)
+				{
+					SelectedContactIndices[3] = ContactIndex;
+					MaxQuadArea = SignedArea;
+				}
+			}
+		}
+		if (SelectedContactIndices[3] == INDEX_NONE)
+		{
+			// No point is outside the triangle we already have
+			return false;
+		}
+		bContactSelected[SelectedContactIndices[3]] = true;
+
+		// Now we should have exactly 4 selected contacts. If we find that one of the existing points is not
+		// selected, it must be because it is being replaced by the new contact. Otherwise the new contact
+		// is interior to the existing manifiold and is rejected.
+		for (int32 ManifoldPointIndex = 0; ManifoldPointIndex < NumManifoldPoints; ++ManifoldPointIndex)
+		{
+			if (!bContactSelected[ManifoldPointIndex])
+			{
+				ManifoldPoints[ManifoldPointIndex].ContactPoint = NewContactPoint;
+				ManifoldPoints[ManifoldPointIndex].InitialShapeContactPoints[0] = NewContactPoint.ShapeContactPoints[0];
+				ManifoldPoints[ManifoldPointIndex].InitialShapeContactPoints[1] = NewContactPoint.ShapeContactPoints[1];
+				ManifoldPoints[ManifoldPointIndex].bWasRestored = false;
+				UpdateManifoldPointFromContact(ManifoldPointIndex);
+				if (NewContactPoint.Phi < Manifold.Phi)
+				{
+					SetActiveContactPoint(NewContactPoint);
+				}
+			}
+		}
+
+		return true;
+	}
+
+	const FManifoldPointSavedData* FPBDCollisionConstraint::FindManifoldPointSavedData(const FManifoldPoint& ManifoldPoint) const
+	{
+		if (bChaos_Manifold_EnableFrictionRestore)
+		{
+			const FReal DistanceToleranceSq = FMath::Square(Chaos_Manifold_FrictionPositionTolerance);
+			for (int32 PrevManifoldPointIndex = 0; PrevManifoldPointIndex < NumSavedManifoldPoints; ++PrevManifoldPointIndex)
+			{
+				const FManifoldPointSavedData& PrevManifoldPoint = ManifoldPointSavedData[PrevManifoldPointIndex];
+				if (PrevManifoldPoint.bInsideStaticFrictionCone && PrevManifoldPoint.IsMatch(ManifoldPoint, DistanceToleranceSq))
+				{
+					return &PrevManifoldPoint;
+				}
+			}
+		}
+		return nullptr;
+	}
+
+	void FPBDCollisionConstraint::TryRestoreFrictionData(const int32 ManifoldPointIndex)
+	{
+		FManifoldPoint& ManifoldPoint = ManifoldPoints[ManifoldPointIndex];
+
+		// Assume we have no matching point from the previous tick, but that we can retain friction from now on
+		// Not supported for non-manifolds yet (hopefully we don't need to)
+		ManifoldPoint.bInsideStaticFrictionCone = bUseManifold;
+		ManifoldPoint.StaticFrictionMax = FReal(0);
+
+		// Find the previous manifold point that matches if there is one
+		const FManifoldPointSavedData* PrevManifoldPoint = FindManifoldPointSavedData(ManifoldPoint);
+		if (PrevManifoldPoint != nullptr)
+		{
+			// We have data from the previous tick and static friction was enabled - restore the data
+			PrevManifoldPoint->Restore(ManifoldPoint);
+		}
+	}
+
 	ECollisionConstraintDirection FPBDCollisionConstraint::GetConstraintDirection(const FReal Dt) const
 	{
 		if (GetDisabled())
diff --git a/Engine/Source/Runtime/Experimental/Chaos/Private/Chaos/Collision/PBDCollisionSolver.cpp b/Engine/Source/Runtime/Experimental/Chaos/Private/Chaos/Collision/PBDCollisionSolver.cpp
--- a/Engine/Source/Runtime/Experimental/Chaos/Private/Chaos/Collision/PBDCollisionSolver.cpp
+++ b/Engine/Source/Runtime/Experimental/Chaos/Private/Chaos/Collision/PBDCollisionSolver.cpp
@@ -579,9 +579,6 @@ namespace Chaos
 		{
 			FSolverManifoldPoint& SolverManifoldPoint = State.ManifoldPoints[PointIndex];
 
-			//SolverManifoldPoint.UpdateContactPoint(Body0, Body1, State.Constraint->GetManifoldPoints()[PointIndex]);
-			//SolverManifoldPoint.UpdateContactMass(Body0, Body1);
-
 			FVec3 ContactDelta;
 			FReal ContactDeltaNormal;
 			SolverManifoldPoint.CalculateContactPositionError(Body0.SolverBody(), Body1.SolverBody(), MaxPushOut, ContactDelta, ContactDeltaNormal);
diff --git a/Engine/Source/Runtime/Experimental/Chaos/Private/Chaos/Collision/ParticlePairMidPhase.cpp b/Engine/Source/Runtime/Experimental/Chaos/Private/Chaos/Collision/ParticlePairMidPhase.cpp
--- a/Engine/Source/Runtime/Experimental/Chaos/Private/Chaos/Collision/ParticlePairMidPhase.cpp
+++ b/Engine/Source/Runtime/Experimental/Chaos/Private/Chaos/Collision/ParticlePairMidPhase.cpp
@@ -27,6 +27,10 @@ namespace Chaos
 		FAutoConsoleVariableRef CVarChaos_Collision_RestoreTolerance_NoContact_Rotation(TEXT("p.Chaos.Collision.RestoreTolerance.NoContact.Rotation"), Chaos_Collision_RestoreTolerance_NoContact_Rotation, TEXT("Quaternion Dot Product Limit. Particle pairs that move less than this may have their contacts reinstated"));
 		FAutoConsoleVariableRef CVarChaos_Collision_RestoreTolerance_Contact_Position(TEXT("p.Chaos.Collision.RestoreTolerance.WithContact.Position"), Chaos_Collision_RestoreTolerance_Contact_Position, TEXT("Fraction of Size. Particle pairs that move less than this may have their contacts reinstated"));
 		FAutoConsoleVariableRef CVarChaos_Collision_RestoreTolerance_Contact_Rotation(TEXT("p.Chaos.Collision.RestoreTolerance.WithContact.Rotation"), Chaos_Collision_RestoreTolerance_Contact_Rotation, TEXT("Quaternion Dot Product Limit. Particle pairs that move less than this may have their contacts reinstated"));
+
+		bool bChaos_Collision_EnableManifoldUpdate = true;
+		FAutoConsoleVariableRef CVarChaos_Collision_EnableManifoldUpdate(TEXT("p.Chaos.Collision.EnableManifoldUpdate"), bChaos_Collision_EnableManifoldUpdate, TEXT(""));
+
 	}
 	using namespace CVars;
 
@@ -91,18 +95,31 @@ namespace Chaos
 
 
 	FSingleShapePairCollisionDetector::FSingleShapePairCollisionDetector(
-		const FPerShapeData* InShape0, 
-		const FPerShapeData* InShape1, 
+		FGeometryParticleHandle* InParticle0,
+		const FPerShapeData* InShape0,
+		FGeometryParticleHandle* InParticle1,
+		const FPerShapeData* InShape1,
 		const EContactShapesType InShapePairType, 
-		const bool bInSwap,
 		FParticlePairMidPhase& InMidPhase)
 		: MidPhase(InMidPhase)
 		, Constraint(nullptr)
+		, Particle0(InParticle0)
+		, Particle1(InParticle1)
 		, Shape0(InShape0)
 		, Shape1(InShape1)
 		, ShapePairType(InShapePairType)
-		, bSwap(bInSwap)
+		, bEnableOBBCheck0(false)
+		, bEnableOBBCheck1(false)
+		, bEnableManifoldCheck(false)
 	{
+		const EImplicitObjectType ImplicitType0 = (Shape0->GetLeafGeometry() != nullptr) ? GetInnerType(Shape0->GetLeafGeometry()->GetCollisionType()) : ImplicitObjectType::Unknown;
+		const EImplicitObjectType ImplicitType1 = (Shape1->GetLeafGeometry() != nullptr) ? GetInnerType(Shape1->GetLeafGeometry()->GetCollisionType()) : ImplicitObjectType::Unknown;
+		const bool bIsSphere0 = (ImplicitType0 == ImplicitObjectType::Sphere);
+		const bool bIsSphere1 = (ImplicitType1 == ImplicitObjectType::Sphere);
+
+		bEnableOBBCheck0 = Chaos_Collision_NarrowPhase_AABBBoundsCheck && !bIsSphere0 && !bIsSphere1;
+		bEnableOBBCheck1 = Chaos_Collision_NarrowPhase_AABBBoundsCheck && !bIsSphere0 && !bIsSphere1;
+		bEnableManifoldCheck = bChaos_Collision_EnableManifoldUpdate && !bIsSphere0 && !bIsSphere1;
 	}
 
 	FSingleShapePairCollisionDetector::~FSingleShapePairCollisionDetector()
@@ -112,10 +129,14 @@ namespace Chaos
 	FSingleShapePairCollisionDetector::FSingleShapePairCollisionDetector(FSingleShapePairCollisionDetector&& R)
 		: MidPhase(R.MidPhase)
 		, Constraint(MoveTemp(R.Constraint))
+		, Particle0(R.Particle0)
+		, Particle1(R.Particle1)
 		, Shape0(R.Shape0)
 		, Shape1(R.Shape1)
 		, ShapePairType(R.ShapePairType)
-		, bSwap(R.bSwap)
+		, bEnableOBBCheck0(R.bEnableOBBCheck0)
+		, bEnableOBBCheck1(R.bEnableOBBCheck1)
+		, bEnableManifoldCheck(R.bEnableManifoldCheck)
 	{
 	}
 
@@ -126,8 +147,6 @@ namespace Chaos
 	}
 
 	int32 FSingleShapePairCollisionDetector::GenerateCollision(
-		FGeometryParticleHandle* Particle0,
-		FGeometryParticleHandle* Particle1,
 		const FReal CullDistance, 
 		const bool bUseCCD,
 		const FReal Dt)
@@ -148,7 +167,9 @@ namespace Chaos
 				return 0;
 			}
 
-			if (Chaos_Collision_NarrowPhase_AABBBoundsCheck)
+			const int32 LastEpoch = MidPhase.GetCollisionAllocator().GetCurrentEpoch() - 1;
+			const bool bCollidedLastTick = IsUsedSince(LastEpoch);
+			if ((bEnableOBBCheck0 || bEnableOBBCheck1) && !bCollidedLastTick)
 			{
 				// OBB-AABB test on both directions. This is beneficial for shapes which do not fit their AABBs very well,
 				// which includes boxes and other shapes that are not roughly spherical. It is especially beneficial when
@@ -156,18 +177,19 @@ namespace Chaos
 				// However, it is quite expensive to do this all the time so we only do this test when we did not 
 				// collide last frame. This is ok if we assume not much changes from frame to frame, but it means
 				// we might call the narrow phase one time too many when shapes become separated.
-				const int32 LastEpoch = MidPhase.GetCollisionAllocator().GetCurrentEpoch() - 1;
-				const bool bCollidedLastTick = IsUsedSince(LastEpoch);
-
-				if (!bCollidedLastTick)
+				const FRigidTransform3& ShapeWorldTransform0 = Shape0->GetLeafWorldTransform();
+				const FRigidTransform3& ShapeWorldTransform1 = Shape1->GetLeafWorldTransform();
+			
+				if (bEnableOBBCheck0)
 				{
-					const FRigidTransform3& ShapeWorldTransform0 = Shape0->GetLeafWorldTransform();
-					const FRigidTransform3& ShapeWorldTransform1 = Shape1->GetLeafWorldTransform();
 					if (!ImplicitOverlapOBBToAABB(Implicit0, Implicit1, ShapeWorldTransform0, ShapeWorldTransform1, CullDistance))
 					{
 						return 0;
 					}
+				}
 
+				if (bEnableOBBCheck1)
+				{
 					if (!ImplicitOverlapOBBToAABB(Implicit1, Implicit0, ShapeWorldTransform1, ShapeWorldTransform0, CullDistance))
 					{
 						return 0;
@@ -176,123 +198,124 @@ namespace Chaos
 			}
 		}
 
-		if (!bSwap)
-		{
-			return GenerateCollisionImpl(Particle0, Shape0, Particle1, Shape1, CullDistance, bUseCCD, Dt);
-		}
-		else
-		{
-			return GenerateCollisionImpl(Particle1, Shape1, Particle0, Shape0, CullDistance, bUseCCD, Dt);
-		}
+		return GenerateCollisionImpl(CullDistance, bUseCCD, Dt);
 	}
 
-	void FSingleShapePairCollisionDetector::CreateConstraint(
-		FGeometryParticleHandle* Particle0,
-		const FPerShapeData* InShape0,
-		FGeometryParticleHandle* Particle1,
-		const FPerShapeData* InShape1,
-		const FReal CullDistance)
+	void FSingleShapePairCollisionDetector::CreateConstraint(const FReal CullDistance)
 	{
 		PHYSICS_CSV_SCOPED_EXPENSIVE(PhysicsVerbose, NarrowPhase_CreateConstraint);
 		check(Constraint == nullptr);
 
-		// NOTE: Must use InShape0 and InShape1 and not Shape0 and Shape1 in this function because
-		// it may be called with shapes in reverse order
-		Constraint = CreateShapePairConstraint(Particle0, InShape0, Particle1, InShape1, CullDistance, ShapePairType);
+		Constraint = CreateShapePairConstraint(Particle0, Shape0, Particle1, Shape1, CullDistance, ShapePairType);
 
 		Constraint->GetContainerCookie().MidPhase = &MidPhase;
 		Constraint->GetContainerCookie().bIsMultiShapePair = false;
 		Constraint->GetContainerCookie().CreationEpoch = MidPhase.GetCollisionAllocator().GetCurrentEpoch();
 	}
 
-	int32 FSingleShapePairCollisionDetector::GenerateCollisionImpl(
-		FGeometryParticleHandle* Particle0, 
-		const FPerShapeData* InShape0,
-		FGeometryParticleHandle* Particle1, 
-		const FPerShapeData* InShape1,
-		const FReal CullDistance,
-		const bool bUseCCD,
-		const FReal Dt)
+	int32 FSingleShapePairCollisionDetector::GenerateCollisionImpl(const FReal CullDistance, const bool bUseCCD, const FReal Dt)
 	{
-		// NOTE: Must use InShape0 and InShape1 and not Shape0 and Shape1 in this function because
-		// it may be called with shapes in reverse order
-
 		if (Constraint == nullptr)
 		{
 			// Lazy creation of the constraint. If a shape pair never gets within CullDistance of each
 			// other, we never allocate a constraint for them. Once they overlap, we reuse the constraint
 			// until the owing particles are not overlapping. i.e., we keep the constraint even if
 			// the shape pairs stop overlapping, reusing it if they start overlapping again.
-			CreateConstraint(Particle0, InShape0, Particle1, InShape1, CullDistance);
+			CreateConstraint(CullDistance);
 		}
 
 		if (Constraint != nullptr)
 		{
 			PHYSICS_CSV_SCOPED_EXPENSIVE(PhysicsVerbose, NarrowPhase_UpdateConstraint);
 
-			// Save or reset the manifold
-			PrepareConstraint(bUseCCD);
+			const FImplicitObject* Implicit0 = Shape0->GetLeafGeometry();
+			const FImplicitObject* Implicit1 = Shape1->GetLeafGeometry();
+			const FRigidTransform3& ShapeWorldTransform0 = Shape0->GetLeafWorldTransform();
+			const FRigidTransform3& ShapeWorldTransform1 = Shape1->GetLeafWorldTransform();
+			const int32 LastEpoch = MidPhase.GetCollisionAllocator().GetCurrentEpoch() - 1;
+			const bool bWasUpdatedLastTick = IsUsedSince(LastEpoch);
+
+			// If the constraint was not used last frame, it needs to be reset. 
+			// If it was used last frame, its data can be used for static friction etc. (unless CCD is enabled)
+			Constraint->SetCCDEnabled(bUseCCD);
+			if (bWasUpdatedLastTick && !bUseCCD)
+			{
+				// Copy manifold data used for static friction - we are about to overwrite the manifold points
+				// NOTE: does not clear the current manifold points
+				Constraint->SaveManifold();
+			}
+			else
+			{
+				// Clear all manifold data
+				Constraint->ResetManifold();
+			}
 
-			// Run the narrow phase
-			if (!bUseCCD)
+			bool bWasManifoldRestored = false;
+			if (bEnableManifoldCheck)
 			{
-				Collisions::UpdateConstraint(*Constraint.Get(), InShape0->GetLeafWorldTransform(), InShape1->GetLeafWorldTransform(), Dt);
+				// Update the existing manifold. We can re-use as-is if none of the points have moved much and the bodies have not moved much
+				// NOTE: this can succeed in "restoring" even if we have no manifold points
+				bWasManifoldRestored = Constraint->UpdateAndTryRestoreManifold(ShapeWorldTransform0, ShapeWorldTransform1);
 			}
 			else
 			{
-				// Note: This is unusual but we are using a mix of the previous and current transform
-				// This is due to how CCD rewinds the position (not rotation) and then sweeps to find the first contact at the current orientation
-				// NOTE: These are actor transforms, not CoM transforms
-				// @todo(chaos): see if we can easily switch to CoM transforms now in collision loop (shapes are held in actor space)
-				// @todo(chaos): this is broken if both objects are CCD
-				FConstGenericParticleHandle P0 = Particle0;
-				FConstGenericParticleHandle P1 = Particle1;
-				const FRigidTransform3 ParticleWorldTransform0 = FRigidTransform3(P0->CCDEnabled() ? P0->X() : P0->P(), P0->Q());
-				const FRigidTransform3 ParticleWorldTransform1 = FRigidTransform3(P1->CCDEnabled() ? P1->X() : P1->P(), P1->Q());
-				const FRigidTransform3 ShapeWorldTransform0 = Constraint->ImplicitTransform[0] * ParticleWorldTransform0;
-				const FRigidTransform3 ShapeWorldTransform1 = Constraint->ImplicitTransform[1] * ParticleWorldTransform1;
-				Collisions::UpdateConstraintSwept(*Constraint.Get(), ShapeWorldTransform0, ShapeWorldTransform1, Dt);
+				// We are not trying to reuse manifold points, so reset them but leave stored data intact (for friction)
+				Constraint->ResetActiveManifoldContacts();
 			}
 
+			if (!bWasManifoldRestored)
+			{
+				// We will be updating the manifold, if only partially, so update the restore comparison transforms
+				Constraint->UpdateLastShapeWorldTransforms(ShapeWorldTransform0, ShapeWorldTransform1);
+
+				// Run the narrow phase
+				if (!bUseCCD)
+				{
+					Collisions::UpdateConstraint(*Constraint.Get(), Shape0->GetLeafWorldTransform(), Shape1->GetLeafWorldTransform(), Dt);
+				}
+				else
+				{
+					// Note: This is unusual but we are using a mix of the previous and current transform
+					// This is due to how CCD rewinds the position (not rotation) and then sweeps to find the first contact at the current orientation
+					// NOTE: These are actor transforms, not CoM transforms
+					// @todo(chaos): see if we can easily switch to CoM transforms now in collision loop (shapes are held in actor space)
+					// @todo(chaos): this is broken if both objects are CCD
+					FConstGenericParticleHandle P0 = Particle0;
+					FConstGenericParticleHandle P1 = Particle1;
+					const FRigidTransform3 CCDParticleWorldTransform0 = FRigidTransform3(P0->CCDEnabled() ? P0->X() : P0->P(), P0->Q());
+					const FRigidTransform3 CCDParticleWorldTransform1 = FRigidTransform3(P1->CCDEnabled() ? P1->X() : P1->P(), P1->Q());
+					const FRigidTransform3 CCDShapeWorldTransform0 = Constraint->ImplicitTransform[0] * CCDParticleWorldTransform0;
+					const FRigidTransform3 CCDShapeWorldTransform1 = Constraint->ImplicitTransform[1] * CCDParticleWorldTransform1;
+					Collisions::UpdateConstraintSwept(*Constraint.Get(), CCDShapeWorldTransform0, CCDShapeWorldTransform1, Dt);
+				}
+			}
 
 			// If we have a valid contact, add it to the active list
 			if (Constraint->GetPhi() <= CullDistance)
 			{
-				MidPhase.GetCollisionAllocator().ActivateConstraint(Constraint.Get());
-				return 1;
+				if (MidPhase.GetCollisionAllocator().ActivateConstraint(Constraint.Get()))
+				{
+					return 1;
+				}
 			}
 		}
 
 		return 0;
 	}
 
-	void FSingleShapePairCollisionDetector::PrepareConstraint(const bool bUseCCD)
-	{
-		// If the constraint was not used last frame, it needs to be reset. 
-		// If it was used last frame, its data can be used for static friction etc. (unless CCD is enabled)
-		const bool bIsActive = IsUsedSince(MidPhase.GetCollisionAllocator().GetCurrentEpoch() - 1);
-		if (bIsActive && !bUseCCD)
-		{
-			Constraint->SaveManifold();
-		}
-		else
-		{
-			Constraint->ResetManifold();
-		}
-
-		Constraint->SetCCDEnabled(bUseCCD);
-	}
-
-	int32 FSingleShapePairCollisionDetector::RestoreCollision()
+	int32 FSingleShapePairCollisionDetector::RestoreCollision(const FReal CullDistance)
 	{
 		// Only restore constraints if active last tick. Any older than that and the shapes were separated for a bit
 		const int32 LastEpoch = MidPhase.GetCollisionAllocator().GetCurrentEpoch() - 1;
 		if ((Constraint != nullptr) && IsUsedSince(LastEpoch))
 		{
-			if (MidPhase.GetCollisionAllocator().ActivateConstraint(Constraint.Get()))
+			Constraint->RestoreManifold();
+			if (Constraint->GetPhi() <= CullDistance)
 			{
-				Constraint->RestoreManifold();
-				return 1;
+				if (MidPhase.GetCollisionAllocator().ActivateConstraint(Constraint.Get()))
+				{
+					return 1;
+				}
 			}
 		}
 		return 0;
@@ -336,10 +359,15 @@ namespace Chaos
 
 
 	FMultiShapePairCollisionDetector::FMultiShapePairCollisionDetector(
-		const FPerShapeData* InShape0, 
+		FGeometryParticleHandle* InParticle0,
+		const FPerShapeData* InShape0,
+		FGeometryParticleHandle* InParticle1,
 		const FPerShapeData* InShape1,
 		FParticlePairMidPhase& InMidPhase)
 		: MidPhase(InMidPhase)
+		, Constraints()
+		, Particle0(InParticle0)
+		, Particle1(InParticle1)
 		, Shape0(InShape0)
 		, Shape1(InShape1)
 	{
@@ -348,6 +376,8 @@ namespace Chaos
 	FMultiShapePairCollisionDetector::FMultiShapePairCollisionDetector(FMultiShapePairCollisionDetector&& R)
 		: MidPhase(R.MidPhase)
 		, Constraints(MoveTemp(R.Constraints))
+		, Particle0(R.Particle0)
+		, Particle1(R.Particle1)
 		, Shape0(R.Shape0)
 		, Shape1(R.Shape1)
 	{
@@ -359,8 +389,6 @@ namespace Chaos
 	}
 
 	int32 FMultiShapePairCollisionDetector::GenerateCollisions(
-		FGeometryParticleHandle* Particle0,
-		FGeometryParticleHandle* Particle1,
 		const FReal CullDistance,
 		const bool bUseCCD,
 		const FReal Dt,
@@ -399,18 +427,19 @@ namespace Chaos
 		int32 NumActiveConstraints = ProcessNewConstraints();
 
 		// @todo(chaos): we could clean up unused collisions between this pair, but probably not worth it
-		// and we have to prevent cleanup for sleeping particles
+		// and we would have to prevent cleanup for sleeping particles because the collisions will still
+		// be referenced in the IslandManager's constraint graph for the sleeping island.
 		//PruneConstraints();
 
 		return NumActiveConstraints;
 	}
 
 	FPBDCollisionConstraint* FMultiShapePairCollisionDetector::FindOrCreateConstraint(
-		FGeometryParticleHandle* Particle0,
+		FGeometryParticleHandle* InParticle0,
 		const FImplicitObject* Implicit0,
 		const FBVHParticles* BVHParticles0,
 		const FRigidTransform3& ShapeRelativeTransform0,
-		FGeometryParticleHandle* Particle1,
+		FGeometryParticleHandle* InParticle1,
 		const FImplicitObject* Implicit1,
 		const FBVHParticles* BVHParticles1,
 		const FRigidTransform3& ShapeRelativeTransform1,
@@ -418,6 +447,16 @@ namespace Chaos
 		const EContactShapesType ShapePairType,
 		const bool bUseManifold)
 	{
+		// This is a callback from the low-level collision function. It should always be the same two particles, though the
+		// shapes may be some children in the implicit hierarchy. The particles could be in the opposite order though, and
+		// this will depend on the shape types involved. E.g., with two particles each with a sphere and a box in a union
+		// would require up to two Sphere-Box contacts, with the particles in opposite orders.
+		if (!ensure(((InParticle0 == Particle0) && (InParticle1 == Particle1)) || ((InParticle0 == Particle1) && (InParticle1 == Particle0))))
+		{
+			// We somehow received a callback for the wrong particle pair...this should not happen
+			return nullptr;
+		}
+
 		const FCollisionParticlePairConstraintKey Key = FCollisionParticlePairConstraintKey(Implicit0, BVHParticles0, Implicit1, BVHParticles1);
 		FPBDCollisionConstraint* Constraint = FindConstraint(Key);
 
@@ -426,16 +465,18 @@ namespace Chaos
 		if (bIsKeyCollision)
 		{
 			// If we get here, we have a key collision. The key uses a hash of pointers which is very likely to be unique for different implicit pairs, 
-			// especially since the scope is only the implicits on the particle pairs, but ut is not guaranteed.
+			// especially since it only needs to be unique for this particle pair, but it is not guaranteed.
 			// Creating a new constraint with the same key could cause fatal problems (the original constraint will be deleted when we add the new one 
-			// to  the map) so we just abort. This will cause missing collisions. It is extremely unlikely to happen but we should fix it at some point.
+			// to the map, but if it is asleep it will be referenced in the contact graph) so we just abort and accept we will miss collisions. 
+			// It is extremely unlikely to happen but we should fix it at some point.
 			ensure(false);
 			return nullptr;
 		}
 
 		if (Constraint == nullptr)
 		{
-			Constraint = CreateConstraint(Particle0, Implicit0, BVHParticles0, ShapeRelativeTransform0, Particle1, Implicit1, BVHParticles1, ShapeRelativeTransform1, CullDistance, ShapePairType, bUseManifold, Key);
+			// NOTE: Using InParticle0 and InParticle1 here because the order may be different to what we have stored
+			Constraint = CreateConstraint(InParticle0, Implicit0, BVHParticles0, ShapeRelativeTransform0, InParticle1, Implicit1, BVHParticles1, ShapeRelativeTransform1, CullDistance, ShapePairType, bUseManifold, Key);
 		}
 		NewConstraints.Add(Constraint);
 		return Constraint;
@@ -443,26 +484,23 @@ namespace Chaos
 
 
 	FPBDCollisionConstraint* FMultiShapePairCollisionDetector::FindOrCreateSweptConstraint(
-		FGeometryParticleHandle* Particle0,
+		FGeometryParticleHandle* InParticle0,
 		const FImplicitObject* Implicit0,
 		const FBVHParticles* BVHParticles0,
 		const FRigidTransform3& ShapeRelativeTransform0,
-		FGeometryParticleHandle* Particle1,
+		FGeometryParticleHandle* InParticle1,
 		const FImplicitObject* Implicit1,
 		const FBVHParticles* BVHParticles1,
 		const FRigidTransform3& ShapeRelativeTransform1,
 		const FReal CullDistance,
 		const EContactShapesType ShapePairType)
 	{
-		const FCollisionParticlePairConstraintKey Key = FCollisionParticlePairConstraintKey(Implicit0, BVHParticles0, Implicit1, BVHParticles1);
-		FPBDCollisionConstraint* Constraint = FindConstraint(Key);
-		if (Constraint == nullptr)
+		const bool bUseManifold = true;
+		FPBDCollisionConstraint* Constraint = FindOrCreateConstraint(InParticle0, Implicit0, BVHParticles0, ShapeRelativeTransform0, InParticle1, Implicit1, BVHParticles1, ShapeRelativeTransform1, CullDistance, ShapePairType, bUseManifold);
+		if (Constraint != nullptr)
 		{
-			const bool bUseManifold = true;
-			Constraint = CreateConstraint(Particle0, Implicit0, BVHParticles0, ShapeRelativeTransform0, Particle1, Implicit1, BVHParticles1, ShapeRelativeTransform1, CullDistance, ShapePairType, bUseManifold, Key);
+			Constraint->SetCCDEnabled(true);
 		}
-		Constraint->SetCCDEnabled(true);
-		NewConstraints.Add(Constraint);
 		return Constraint;
 	}
 
@@ -477,11 +515,11 @@ namespace Chaos
 	}
 
 	FPBDCollisionConstraint* FMultiShapePairCollisionDetector::CreateConstraint(
-		FGeometryParticleHandle* Particle0,
+		FGeometryParticleHandle* InParticle0,
 		const FImplicitObject* Implicit0,
 		const FBVHParticles* BVHParticles0,
 		const FRigidTransform3& ShapeRelativeTransform0,
-		FGeometryParticleHandle* Particle1,
+		FGeometryParticleHandle* InParticle1,
 		const FImplicitObject* Implicit1,
 		const FBVHParticles* BVHParticles1,
 		const FRigidTransform3& ShapeRelativeTransform1,
@@ -491,7 +529,7 @@ namespace Chaos
 		const FCollisionParticlePairConstraintKey& Key)
 	{
 		PHYSICS_CSV_SCOPED_EXPENSIVE(PhysicsVerbose, NarrowPhase_CreateConstraint);
-		TUniquePtr<FPBDCollisionConstraint> Constraint = CreateImplicitPairConstraint(Particle0, Implicit0, BVHParticles0, ShapeRelativeTransform0, Particle1, Implicit1, BVHParticles1, ShapeRelativeTransform1, CullDistance, ShapePairType, bInUseManifold);
+		TUniquePtr<FPBDCollisionConstraint> Constraint = CreateImplicitPairConstraint(InParticle0, Implicit0, BVHParticles0, ShapeRelativeTransform0, InParticle1, Implicit1, BVHParticles1, ShapeRelativeTransform1, CullDistance, ShapePairType, bInUseManifold);
 		
 		Constraint->GetContainerCookie().MidPhase = &MidPhase;
 		Constraint->GetContainerCookie().bIsMultiShapePair = true;
@@ -500,7 +538,7 @@ namespace Chaos
 		return Constraints.Add(Key.GetKey(), MoveTemp(Constraint)).Get();
 	}
 
-	int32 FMultiShapePairCollisionDetector::RestoreCollisions()
+	int32 FMultiShapePairCollisionDetector::RestoreCollisions(const FReal CullDistance)
 	{
 		int32 NumRestored = 0;
 		const int32 LastEpoch = MidPhase.GetCollisionAllocator().GetCurrentEpoch() - 1;
@@ -509,10 +547,13 @@ namespace Chaos
 			TUniquePtr<FPBDCollisionConstraint>& Constraint = KVP.Value;
 			if ((Constraint != nullptr) && (Constraint->GetContainerCookie().LastUsedEpoch >= LastEpoch))
 			{
-				if (MidPhase.GetCollisionAllocator().ActivateConstraint(Constraint.Get()))
+				Constraint->RestoreManifold();
+				if (Constraint->GetPhi() < CullDistance)
 				{
-					Constraint->RestoreManifold();
-					++NumRestored;
+					if (MidPhase.GetCollisionAllocator().ActivateConstraint(Constraint.Get()))
+					{
+						++NumRestored;
+					}
 				}
 			}
 		}
@@ -702,11 +743,21 @@ namespace Chaos
 
 			if (ShapePairType != EContactShapesType::Unknown)
 			{
-				ShapePairDetectors.Emplace(FSingleShapePairCollisionDetector(Shape0, Shape1, ShapePairType, bSwap, *this));
+				if (!bSwap)
+				{
+					ShapePairDetectors.Emplace(FSingleShapePairCollisionDetector(Particle0, Shape0, Particle1, Shape1, ShapePairType, *this));
+				}
+				else
+				{
+					ShapePairDetectors.Emplace(FSingleShapePairCollisionDetector(Particle1, Shape1, Particle0, Shape0, ShapePairType, *this));
+				}
 			}
 			else
 			{
-				MultiShapePairDetectors.Emplace(FMultiShapePairCollisionDetector(Shape0, Shape1, *this));
+				if (ensure(!bSwap))
+				{
+					MultiShapePairDetectors.Emplace(FMultiShapePairCollisionDetector(Particle0, Shape0, Particle1, Shape1, *this));
+				}
 			}
 
 			// We don't allow full constraint restoration for LevelSets or Unions because small changes in 
@@ -757,16 +808,6 @@ namespace Chaos
 
 	bool FParticlePairMidPhase::ShouldRestoreConstraints(const FReal Dt)
 	{
-		if (!bChaos_Collision_EnableManifoldRestore)
-		{
-			return false;
-		}
-
-		if (!bRestorable)
-		{
-			return false;
-		}
-
 		const FVec3& ParticleP0 = FConstGenericParticleHandle(Particle0)->P();
 		const FRotation3& ParticleQ0 = FConstGenericParticleHandle(Particle0)->Q();
 		const FVec3& ParticleP1 = FConstGenericParticleHandle(Particle1)->P();
@@ -797,7 +838,7 @@ namespace Chaos
 		return false;
 	}
 
-	bool FParticlePairMidPhase::TryRestoreConstraints(const FReal Dt)
+	bool FParticlePairMidPhase::TryRestoreConstraints(const FReal Dt, const FReal CullDistance)
 	{
 		SCOPE_CYCLE_COUNTER(STAT_Collisions_Restore);
 		PHYSICS_CSV_SCOPED_EXPENSIVE(PhysicsVerbose, DetectCollisions_RestoreCollision);
@@ -808,11 +849,11 @@ namespace Chaos
 			int32 NumRestored = 0;
 			for (FSingleShapePairCollisionDetector& ShapePair : ShapePairDetectors)
 			{
-				NumRestored += ShapePair.RestoreCollision();
+				NumRestored += ShapePair.RestoreCollision(CullDistance);
 			}
 			for (FMultiShapePairCollisionDetector& MultiShapePair : MultiShapePairDetectors)
 			{
-				NumRestored += MultiShapePair.RestoreCollisions();
+				NumRestored += MultiShapePair.RestoreCollisions(CullDistance);
 			}
 
 			NumActiveConstraints = NumRestored;
@@ -840,7 +881,8 @@ namespace Chaos
 		const bool bUseCCD = bIsCCD && ShouldEnableCCD(Dt);
 
 		// If the bodies have not moved, we will reuse the constraints as-is
-		bool bWasRestored = !bUseCCD && TryRestoreConstraints(Dt);
+		const bool bCanRestore = bChaos_Collision_EnableManifoldRestore && bRestorable && !bUseCCD;
+		const bool bWasRestored = bCanRestore && TryRestoreConstraints(Dt, CullDistance);
 
 		// If the bodies have moved we need to create or update the constraints
 		if (!bWasRestored)
@@ -853,12 +895,12 @@ namespace Chaos
 			// Run collision detection on all potentially colliding shape pairs
 			for (FSingleShapePairCollisionDetector& ShapePair : ShapePairDetectors)
 			{
-				NumActiveConstraints += ShapePair.GenerateCollision(Particle0, Particle1, CullDistance, bUseCCD, Dt);
+				NumActiveConstraints += ShapePair.GenerateCollision(CullDistance, bUseCCD, Dt);
 			}
 
 			for (FMultiShapePairCollisionDetector& MultiShapePair : MultiShapePairDetectors)
 			{
-				NumActiveConstraints += MultiShapePair.GenerateCollisions(Particle0, Particle1, CullDistance, bUseCCD, Dt, Context);
+				NumActiveConstraints += MultiShapePair.GenerateCollisions(CullDistance, bUseCCD, Dt, Context);
 			}
 		}
 
diff --git a/Engine/Source/Runtime/Experimental/Chaos/Private/Chaos/Collision/SolverCollisionContainer.cpp b/Engine/Source/Runtime/Experimental/Chaos/Private/Chaos/Collision/SolverCollisionContainer.cpp
--- a/Engine/Source/Runtime/Experimental/Chaos/Private/Chaos/Collision/SolverCollisionContainer.cpp
+++ b/Engine/Source/Runtime/Experimental/Chaos/Private/Chaos/Collision/SolverCollisionContainer.cpp
@@ -184,7 +184,7 @@ namespace Chaos
 		SweptCollisionSolvers.Reset();
 	}
 
-	FPBDCollisionSolver* FPBDCollisionSolverContainer::AddConstraintSolver(FReal Dt, FPBDCollisionConstraint& Constraint, const int32 Particle0Level, const int32 Particle1Level, FSolverBodyContainer& SolverBodyContainer)
+	void FPBDCollisionSolverContainer::AddConstraintSolver(FReal Dt, FPBDCollisionConstraint& Constraint, const int32 Particle0Level, const int32 Particle1Level, FSolverBodyContainer& SolverBodyContainer)
 	{
 		// This container is required to allocate pointers that are valid for the whole tick,
 		// so we cannot allow the container to resize during the tick. See Reset()
@@ -199,8 +199,6 @@ namespace Chaos
 		{
 			SweptCollisionSolvers.Add(SolverCollisionIndex);
 		}
-
-		return &CollisionSolvers[SolverCollisionIndex].GetSolver();
 	}
 
 	void FPBDCollisionSolverContainer::UpdatePositionShockPropagation(const FReal Dt, const int32 It, const int32 NumIts, const int32 BeginIndex, const int32 EndIndex)
diff --git a/Engine/Source/Runtime/Experimental/Chaos/Private/Chaos/CollisionOneShotManifolds.cpp b/Engine/Source/Runtime/Experimental/Chaos/Private/Chaos/CollisionOneShotManifolds.cpp
--- a/Engine/Source/Runtime/Experimental/Chaos/Private/Chaos/CollisionOneShotManifolds.cpp
+++ b/Engine/Source/Runtime/Experimental/Chaos/Private/Chaos/CollisionOneShotManifolds.cpp
@@ -30,12 +30,6 @@ namespace Chaos
 	FRealSingle Chaos_Collision_Manifold_CullDistanceMarginMultiplier = 1.0f;
 	FAutoConsoleVariableRef CVarChaosCollisioConvexManifoldCullDistanceMarginMultiplier(TEXT("p.Chaos.Collision.Manifold.CullDistanceMarginMultiplier"), Chaos_Collision_Manifold_CullDistanceMarginMultiplier, TEXT(""));
 
-	bool bChaos_Collision_Manifold_UseMinMargin = true;
-	FAutoConsoleVariableRef CVarChaosCollisioConvexManifoldUseMinMargin(TEXT("p.Chaos.Collision.Manifold.UseMinMargin"), bChaos_Collision_Manifold_UseMinMargin, TEXT(""));
-
-	FRealSingle Chaos_Collision_Manifold_MaxMargin = -1.0f;
-	FAutoConsoleVariableRef CVarChaosCollisioConvexManifoldMaxMargin(TEXT("p.Chaos.Collision.Manifold.MaxMargin"), Chaos_Collision_Manifold_MaxMargin, TEXT(""));
-
 	FRealSingle Chaos_Collision_Manifold_MinFaceSearchDistance = 1.0f;
 	FAutoConsoleVariableRef CVarChaosCollisioConvexManifoldMinFaceSearchDistance(TEXT("p.Chaos.Collision.Manifold.MinFaceSearchDistance"), Chaos_Collision_Manifold_MinFaceSearchDistance, TEXT(""));
 
@@ -45,6 +39,10 @@ namespace Chaos
 	bool ForceOneShotManifoldEdgeEdgeCaseZeroCullDistance = false;
 	FAutoConsoleVariableRef CVarForceOneShotManifoldEdgeEdgeCaseZeroCullDistance(TEXT("p.Chaos.Collision.Manifold.ForceOneShotManifoldEdgeEdgeCaseZeroCullDistance"), ForceOneShotManifoldEdgeEdgeCaseZeroCullDistance,
 	TEXT("If enabled, if one shot manifold hits edge/edge case, we will force a cull distance of zero. That means edge/edge contacts will be thrown out if separated at all. Only applies to Convex/Convex oneshot impl."));
+
+	bool bChaos_Collision_EnableManifoldInject = true;
+	FAutoConsoleVariableRef CVarChaos_Collision_EnableManifoldInject(TEXT("p.Chaos.Collision.EnableManifolInject"), bChaos_Collision_EnableManifoldInject, TEXT(""));
+
 	
 	namespace Collisions
 	{
@@ -221,7 +219,7 @@ namespace Chaos
 				GJKContactPoint.Phi = EdgePhi;
 				GJKContactPoint.Location = 0.5f * (EdgePos1 + EdgePos2);
 
-				Constraint.AddOneshotManifoldContact(GJKContactPoint, Dt);
+				Constraint.AddOneshotManifoldContact(GJKContactPoint);
 				return;
 			}
 
@@ -331,7 +329,7 @@ namespace Chaos
 				ContactPoint.Normal = GJKContactPoint.Normal;
 				ContactPoint.Phi = FVec3::DotProduct(PointProjectedOntoReferenceFace - VertexInReferenceCubeCoordinates, ReferenceFaceBox1 ? SeparationDirectionLocalBox1 : -SeparationDirectionLocalBox2);
 
-				Constraint.AddOneshotManifoldContact(ContactPoint, Dt);
+				Constraint.AddOneshotManifoldContact(ContactPoint);
 			}
 		}
 
@@ -635,7 +633,7 @@ namespace Chaos
 			const TGJKCoreShape<GeometryB> BWithMargin(B, MarginB);
 			const FRigidTransform3 BToATM = BTM.GetRelativeTransformNoScale(ATM);
 
-			if (GJKPenetrationWarmStartable<true>(AWithMargin, BWithMargin, BToATM, FReal(0), FReal(0), Penetration, ClosestA, ClosestB, NormalA, NormalB, VertexIndexA, VertexIndexB, InOutGjkWarmStartData, OutMaxMarginDelta, Epsilon))
+			if (GJKPenetrationWarmStartable(AWithMargin, BWithMargin, BToATM, Penetration, ClosestA, ClosestB, NormalA, NormalB, VertexIndexA, VertexIndexB, InOutGjkWarmStartData, OutMaxMarginDelta, Epsilon))
 			{
 				Contact.ShapeContactPoints[0] = ClosestA;
 				Contact.ShapeContactPoints[1] = ClosestB;
@@ -751,43 +749,41 @@ namespace Chaos
 
 			// We only build one shot manifolds once
 			// All convexes are pre-scaled, or wrapped in TImplicitObjectScaled
-			ensure(Constraint.GetManifoldPoints().Num() == 0);
+			//ensure(Constraint.GetManifoldPoints().Num() == 0);
 			ensure(Convex1Transform.GetScale3D() == FVec3(1.0f, 1.0f, 1.0f));
 			ensure(Convex2Transform.GetScale3D() == FVec3(1.0f, 1.0f, 1.0f));
 
 			// Get the adjusted margins for each convex
-			FReal Margin1 = Convex1.GetMargin();
-			FReal Margin2 = Convex2.GetMargin();
-			if (bChaos_Collision_Manifold_UseMinMargin && !bConvex1IsCapsule)
-			{
-				// Use the smaller of the two margins for both shapes
-				Margin1 = FMath::Min(Margin1, Margin2);
-				Margin2 = Margin1;
-			}
-			if (Chaos_Collision_Manifold_MaxMargin >= 0.0f)
-			{
-				// Clamp the margin
-				if (!bConvex1IsCapsule)
-				{
-					Margin1 = FMath::Min(Margin1, Chaos_Collision_Manifold_MaxMargin);
-				}
-				Margin2 = FMath::Min(Margin2, Chaos_Collision_Manifold_MaxMargin);
-			}
+			const FReal Margin1 = Constraint.GetCollisionMargin0();
+			const FReal Margin2 = Constraint.GetCollisionMargin1();
 
 			// Find the deepest penetration. This is used to determine the planes and points to use for the manifold
 			// MaxMarginDelta is an upper bound on the distance from the contact on the rounded core shape to the actual shape surface. 
 			FReal MaxMarginDelta = FReal(0);
 			int32 VertexIndexA = INDEX_NONE, VertexIndexB = INDEX_NONE;
 			FContactPoint GJKContactPoint = GJKContactPointMargin(Convex1, Convex2, Convex1Transform, Convex2Transform, Margin1, Margin2, Constraint.GetGJKWarmStartData(), MaxMarginDelta, VertexIndexA, VertexIndexB);
+			PHYSICS_CSV_CUSTOM_EXPENSIVE(PhysicsCounters, NumManifoldsGJKCalled, 1, ECsvCustomStatOp::Accumulate);
+
+			const bool bCanUpdateManifold = bChaos_Collision_EnableManifoldInject;
+			if (bCanUpdateManifold && Constraint.TryAddManifoldContact(GJKContactPoint, Convex1Transform, Convex2Transform))
+			{
+				PHYSICS_CSV_CUSTOM_EXPENSIVE(PhysicsCounters, NumManifoldsMaintained, 1, ECsvCustomStatOp::Accumulate);
+				return;
+			}
+
+			Constraint.ResetActiveManifoldContacts();
 
 			// GJK is using margins and rounded corner, so if we have a corner-to-corner contact it will under-report the actual distance by an amount that depends on how
 			// "pointy" the edge/corner is - this error is bounded by MaxMarginDelta.
 			const FReal GJKCullDistance = Constraint.GetCullDistance() + MaxMarginDelta;
 			if (GJKContactPoint.Phi > GJKCullDistance)
 			{
+				PHYSICS_CSV_CUSTOM_EXPENSIVE(PhysicsCounters, NumManifoldsGJKCulled, 1, ECsvCustomStatOp::Accumulate);
 				return;
 			}
 
+			PHYSICS_CSV_CUSTOM_EXPENSIVE(PhysicsCounters, NumManifoldsCreated, 1, ECsvCustomStatOp::Accumulate);
+
 			// @todo(chaos): get the vertex index from GJK and use to to get the plane
 			const FVec3 SeparationDirectionLocalConvex1 = Convex1Transform.InverseTransformVectorNoScale(GJKContactPoint.Normal);
 			const int32 MostOpposingPlaneIndexConvex1 = SelectContactPlane(Convex1, GJKContactPoint.ShapeContactPoints[0], SeparationDirectionLocalConvex1, Margin1, VertexIndexA);
@@ -821,6 +817,8 @@ namespace Chaos
 					return;
 				}
 
+				// @todo(chaos): this does not work well when the edges are parallel. We should always have points with zero
+				// position delta perpendicular to the normal, but that is not the case for parallel edges
 				FVec3 ShapeEdgePos1 = Convex1.GetClosestEdgePosition(MostOpposingPlaneIndexConvex1, GJKContactPoint.ShapeContactPoints[0]);
 				FVec3 ShapeEdgePos2 = Convex2.GetClosestEdgePosition(MostOpposingPlaneIndexConvex2, GJKContactPoint.ShapeContactPoints[1]);
 				if (bConvex1IsCapsule)
@@ -828,17 +826,19 @@ namespace Chaos
 					ShapeEdgePos1 -= Margin1 * SeparationDirectionLocalConvex1;
 				}
 
-				FVec3 EdgePos1 = Convex1Transform.TransformPosition(ShapeEdgePos1);
-				FVec3 EdgePos2 = Convex2Transform.TransformPosition(ShapeEdgePos2);
-				FReal EdgePhi = FVec3::DotProduct(EdgePos1 - EdgePos2, GJKContactPoint.Normal);
+				const FVec3 EdgePos1 = Convex1Transform.TransformPosition(ShapeEdgePos1);
+				const FVec3 EdgePos2 = Convex2Transform.TransformPosition(ShapeEdgePos2);
+				const FReal EdgePhi = FVec3::DotProduct(EdgePos1 - EdgePos2, GJKContactPoint.Normal);
+				const FVec3 WorldPos = FReal(0.5) * (EdgePos1 + EdgePos2);
+				const FVec3& WorldNormal = GJKContactPoint.Normal;
 
-				GJKContactPoint.ShapeContactPoints[0] = ShapeEdgePos1;
-				GJKContactPoint.ShapeContactPoints[1] = ShapeEdgePos2;
+				GJKContactPoint.ShapeContactPoints[0] = Convex1Transform.InverseTransformPositionNoScale(WorldPos + FReal(0.5) * EdgePhi * WorldNormal);
+				GJKContactPoint.ShapeContactPoints[1] = Convex2Transform.InverseTransformPositionNoScale(WorldPos - FReal(0.5) * EdgePhi * WorldNormal);
 				GJKContactPoint.Phi = EdgePhi;
-				GJKContactPoint.Location = FReal(0.5) * (EdgePos1 + EdgePos2);
+				GJKContactPoint.Location = WorldPos;
 				// Normal unchanged from GJK result
 
-				Constraint.AddOneshotManifoldContact(GJKContactPoint, Dt);
+				Constraint.AddOneshotManifoldContact(GJKContactPoint);
 				return;
 			}
 
@@ -945,7 +945,7 @@ namespace Chaos
 					ContactPoint.Normal = GJKContactPoint.Normal;
 					ContactPoint.Phi = FVec3::DotProduct(PointProjectedOntoReferenceFace - VertexInReferenceCoordinates, ReferenceFaceConvex1 ? SeparationDirectionLocalConvex1 : -SeparationDirectionLocalConvex2);				
 
-					Constraint.AddOneshotManifoldContact(ContactPoint, Dt);
+					Constraint.AddOneshotManifoldContact(ContactPoint);
 				}
 			}
 		}
diff --git a/Engine/Source/Runtime/Experimental/Chaos/Private/Chaos/CollisionOneShotManifoldsMiscShapes.cpp b/Engine/Source/Runtime/Experimental/Chaos/Private/Chaos/CollisionOneShotManifoldsMiscShapes.cpp
--- a/Engine/Source/Runtime/Experimental/Chaos/Private/Chaos/CollisionOneShotManifoldsMiscShapes.cpp
+++ b/Engine/Source/Runtime/Experimental/Chaos/Private/Chaos/CollisionOneShotManifoldsMiscShapes.cpp
@@ -37,13 +37,15 @@ namespace Chaos
 			
 			// We only build one shot manifolds once
 			// All convexes are pre-scaled, or wrapped in TImplicitObjectScaled
-			ensure(Constraint.GetManifoldPoints().Num() == 0);
 			ensure(SphereATransform.GetScale3D() == FVec3(1.0f, 1.0f, 1.0f));
 			ensure(SphereBTransform.GetScale3D() == FVec3(1.0f, 1.0f, 1.0f));
 
+			// @todo(chaos): support manifold maintenance
+			Constraint.ResetManifold();
+
 			FContactPoint ContactPoint = SphereSphereContactPoint(SphereA, SphereATransform, SphereB, SphereBTransform, Constraint.Manifold.RestitutionPadding);
 
-			Constraint.AddOneshotManifoldContact(ContactPoint, Dt);
+			Constraint.AddOneshotManifoldContact(ContactPoint);
 		}
 
 		void ConstructSpherePlaneOneShotManifold(
@@ -56,52 +58,60 @@ namespace Chaos
 		{
 			// We only build one shot manifolds once
 			// All convexes are pre-scaled, or wrapped in TImplicitObjectScaled
-			ensure(Constraint.GetManifoldPoints().Num() == 0);
 			ensure(SphereTransform.GetScale3D() == FVec3(1.0f, 1.0f, 1.0f));
 			ensure(PlaneTransform.GetScale3D() == FVec3(1.0f, 1.0f, 1.0f));
 
+			// @todo(chaos): support manifold maintenance
+			Constraint.ResetManifold();
+
 			FContactPoint ContactPoint = SpherePlaneContactPoint(Sphere, SphereTransform, Plane, PlaneTransform, Constraint.Manifold.RestitutionPadding);
 
-			Constraint.AddOneshotManifoldContact(ContactPoint, Dt);
+			Constraint.AddOneshotManifoldContact(ContactPoint);
 		}	
 
 		void ConstructSphereBoxOneShotManifold(const TSphere<FReal, 3>& Sphere, const FRigidTransform3& SphereTransform, const FImplicitBox3& Box, const FRigidTransform3& BoxTransform, const FReal Dt, FPBDCollisionConstraint& Constraint)
 		{
 			// We only build one shot manifolds once
 			// All convexes are pre-scaled, or wrapped in TImplicitObjectScaled
-			ensure(Constraint.GetManifoldPoints().Num() == 0);
 			ensure(SphereTransform.GetScale3D() == FVec3(1.0f, 1.0f, 1.0f));
 			ensure(BoxTransform.GetScale3D() == FVec3(1.0f, 1.0f, 1.0f));
 
+			// @todo(chaos): support manifold maintenance
+			Constraint.ResetManifold();
+
 			FContactPoint ContactPoint = SphereBoxContactPoint(Sphere, SphereTransform, Box, BoxTransform, Constraint.Manifold.RestitutionPadding);
 
-			Constraint.AddOneshotManifoldContact(ContactPoint, Dt);
+			Constraint.AddOneshotManifoldContact(ContactPoint);
 		}
 
 		void ConstructSphereCapsuleOneShotManifold(const TSphere<FReal, 3>& Sphere, const FRigidTransform3& SphereTransform, const FCapsule& Capsule, const FRigidTransform3& CapsuleTransform, const FReal Dt, FPBDCollisionConstraint& Constraint)
 		{
 			// We only build one shot manifolds once
 			// All convexes are pre-scaled, or wrapped in TImplicitObjectScaled
-			ensure(Constraint.GetManifoldPoints().Num() == 0);
 			ensure(SphereTransform.GetScale3D() == FVec3(1.0f, 1.0f, 1.0f));
 			ensure(CapsuleTransform.GetScale3D() == FVec3(1.0f, 1.0f, 1.0f));
 
+			// @todo(chaos): support manifold maintenance
+			Constraint.ResetManifold();
+
 			FContactPoint ContactPoint = SphereCapsuleContactPoint(Sphere, SphereTransform, Capsule, CapsuleTransform, Constraint.Manifold.RestitutionPadding);
 
-			Constraint.AddOneshotManifoldContact(ContactPoint, Dt);
+			Constraint.AddOneshotManifoldContact(ContactPoint);
 		}
 
 		void ConstructSphereConvexManifold(const TSphere<FReal, 3>& Sphere, const FRigidTransform3& SphereTransform, const FImplicitObject3& Convex, const FRigidTransform3& ConvexTransform, const FReal Dt, FPBDCollisionConstraint& Constraint)
 		{
 			// We only build one shot manifolds once
 			// All convexes are pre-scaled, or wrapped in TImplicitObjectScaled
-			ensure(Constraint.GetManifoldPoints().Num() == 0);
 			ensure(SphereTransform.GetScale3D() == FVec3(1.0f, 1.0f, 1.0f));
 			ensure(ConvexTransform.GetScale3D() == FVec3(1.0f, 1.0f, 1.0f));
 
+			// @todo(chaos): support manifold maintenance
+			Constraint.ResetManifold();
+
 			FContactPoint ContactPoint = SphereConvexContactPoint(Sphere, SphereTransform, Convex, ConvexTransform);
 
-			Constraint.AddOneshotManifoldContact(ContactPoint, Dt);
+			Constraint.AddOneshotManifoldContact(ContactPoint);
 		}
 
 		template <typename TriMeshType>
@@ -109,26 +119,27 @@ namespace Chaos
 		{
 			// We only build one shot manifolds once
 			// All convexes are pre-scaled, or wrapped in TImplicitObjectScaled
-			ensure(Constraint.GetManifoldPoints().Num() == 0);
 			ensure(SphereWorldTransform.GetScale3D() == FVec3(1.0f, 1.0f, 1.0f));
 			ensure(TriMeshWorldTransform.GetScale3D() == FVec3(1.0f, 1.0f, 1.0f));
 
+			// @todo(chaos): support manifold maintenance
+			Constraint.ResetManifold();
+
 			FContactPoint ContactPoint = SphereTriangleMeshContactPoint(Sphere, SphereWorldTransform, TriangleMesh, TriMeshWorldTransform, Constraint.GetCullDistance(), 0.0f);
 
-			Constraint.AddOneshotManifoldContact(ContactPoint, Dt);
+			Constraint.AddOneshotManifoldContact(ContactPoint);
 		}
 
 		void ConstructSphereHeightFieldOneShotManifold(const TSphere<FReal, 3>& Sphere, const FRigidTransform3& SphereTransform, const FHeightField& Heightfield, const FRigidTransform3& HeightfieldTransform, const FReal Dt, FPBDCollisionConstraint& Constraint)
 		{
 			// We only build one shot manifolds once
 			// All convexes are pre-scaled, or wrapped in TImplicitObjectScaled
-			ensure(Constraint.GetManifoldPoints().Num() == 0);
 			ensure(SphereTransform.GetScale3D() == FVec3(1.0f, 1.0f, 1.0f));
 			ensure(HeightfieldTransform.GetScale3D() == FVec3(1.0f, 1.0f, 1.0f));
 
 			FContactPoint ContactPoint = SphereHeightFieldContactPoint(Sphere, SphereTransform, Heightfield, HeightfieldTransform, Constraint.GetCullDistance(), 0.0f);
 
-			Constraint.AddOneshotManifoldContact(ContactPoint, Dt);
+			Constraint.AddOneshotManifoldContact(ContactPoint);
 		}
 
 		void ConstructCapsuleCapsuleOneShotManifold(const FCapsule& CapsuleA, const FRigidTransform3& CapsuleATransform, const FCapsule& CapsuleB, const FRigidTransform3& CapsuleBTransform, const FReal Dt, FPBDCollisionConstraint& Constraint)
@@ -137,10 +148,12 @@ namespace Chaos
 
 			// We only build one shot manifolds once
 			// All convexes are pre-scaled, or wrapped in TImplicitObjectScaled
-			ensure(Constraint.GetManifoldPoints().Num() == 0);
 			ensure(CapsuleATransform.GetScale3D() == FVec3(1.0f, 1.0f, 1.0f));
 			ensure(CapsuleBTransform.GetScale3D() == FVec3(1.0f, 1.0f, 1.0f));
 
+			// @todo(chaos): support manifold maintenance
+			Constraint.ResetManifold();
+
 			FVec3 CapsuleADirection(CapsuleATransform.TransformVector(CapsuleA.GetSegment().GetAxis()));
 			const FVec3 CapsuleBDirection(CapsuleBTransform.TransformVector(CapsuleB.GetSegment().GetAxis()));
 
@@ -152,7 +165,7 @@ namespace Chaos
 			if (FMath::Abs(ADotB) < AxisDotMinimum || AHalfLen < KINDA_SMALL_NUMBER || BHalfLen < KINDA_SMALL_NUMBER)
 			{
 				FContactPoint ContactPoint = CapsuleCapsuleContactPoint(CapsuleA, CapsuleATransform, CapsuleB, CapsuleBTransform, Constraint.GetCullDistance());
-				Constraint.AddOneshotManifoldContact(ContactPoint, Dt);
+				Constraint.AddOneshotManifoldContact(ContactPoint);
 				return;
 			}
 			
@@ -173,7 +186,7 @@ namespace Chaos
 			if (DeltaLen < KINDA_SMALL_NUMBER)
 			{
 				FContactPoint ContactPoint = CapsuleCapsuleContactPoint(CapsuleA, CapsuleATransform, CapsuleB, CapsuleBTransform, Constraint.GetCullDistance());
-				Constraint.AddOneshotManifoldContact(ContactPoint, Dt);
+				Constraint.AddOneshotManifoldContact(ContactPoint);
 				return;
 			}
 			
@@ -193,7 +206,7 @@ namespace Chaos
 			if (Clipped1Coord > Clipped2Coord) // No overlap
 			{
 				FContactPoint ContactPoint = CapsuleCapsuleContactPoint(CapsuleA, CapsuleATransform, CapsuleB, CapsuleBTransform, Constraint.GetCullDistance());
-				Constraint.AddOneshotManifoldContact(ContactPoint, Dt);
+				Constraint.AddOneshotManifoldContact(ContactPoint);
 				return;
 			}
 
@@ -217,7 +230,7 @@ namespace Chaos
 				ContactPoint.Location = 0.5f * (LocationA + LocationB);
 				ContactPoint.Phi = FVec3::DotProduct(LocationA - LocationB, Normal);
 
-				Constraint.AddOneshotManifoldContact(ContactPoint, Dt);
+				Constraint.AddOneshotManifoldContact(ContactPoint);
 			};
 
 			AddManifoldPoint(Clipped1Coord);
@@ -229,15 +242,17 @@ namespace Chaos
 		{
 			// We only build one shot manifolds once
 			// All convexes are pre-scaled, or wrapped in TImplicitObjectScaled
-			ensure(Constraint.GetManifoldPoints().Num() == 0);
 			ensure(CapsuleWorldTransform.GetScale3D() == FVec3(1.0f, 1.0f, 1.0f));
 			ensure(TriMeshWorldTransform.GetScale3D() == FVec3(1.0f, 1.0f, 1.0f));
 
+			// @todo(chaos): support manifold maintenance
+			Constraint.ResetManifold();
+
 			TArray<FContactPoint> ContactPoints;
 			GJKImplicitManifold<FCapsule, TriMeshType>(Capsule, CapsuleWorldTransform, TriangleMesh, TriMeshWorldTransform, Constraint.GetCullDistance(), 0.0f, ContactPoints);
 			for (FContactPoint& ContactPoint : ContactPoints)
 			{
-				Constraint.AddOneshotManifoldContact(ContactPoint, Dt);
+				Constraint.AddOneshotManifoldContact(ContactPoint);
 			}
 		}
 
@@ -245,15 +260,17 @@ namespace Chaos
 		{
 			// We only build one shot manifolds once
 			// All convexes are pre-scaled, or wrapped in TImplicitObjectScaled
-			ensure(Constraint.GetManifoldPoints().Num() == 0);
 			ensure(CapsuleTransform.GetScale3D() == FVec3(1.0f, 1.0f, 1.0f));
 			ensure(HeightFieldTransform.GetScale3D() == FVec3(1.0f, 1.0f, 1.0f));
 
+			// @todo(chaos): support manifold maintenance
+			Constraint.ResetManifold();
+
 			TArray<FContactPoint> ContactPoints;
 			GJKImplicitManifold<FCapsule>(Capsule, CapsuleTransform, HeightField, HeightFieldTransform, Constraint.GetCullDistance(), 0.0f, ContactPoints);
 			for (FContactPoint& ContactPoint : ContactPoints)
 			{
-				Constraint.AddOneshotManifoldContact(ContactPoint, Dt);
+				Constraint.AddOneshotManifoldContact(ContactPoint);
 			}
 		}
 
@@ -262,17 +279,19 @@ namespace Chaos
 		{
 			// We only build one shot manifolds once
 			// All convexes are pre-scaled, or wrapped in TImplicitObjectScaled
-			ensure(Constraint.GetManifoldPoints().Num() == 0);
 			ensure(ConvexTransform.GetScale3D() == FVec3(1.0f, 1.0f, 1.0f));
 			ensure(HeightFieldTransform.GetScale3D() == FVec3(1.0f, 1.0f, 1.0f));
 
+			// @todo(chaos): support manifold maintenance
+			Constraint.ResetManifold();
+
 			//FContactPoint ContactPoint = ConvexHeightFieldContactPoint(Convex, ConvexTransform, HeightField, HeightFieldTransform, Constraint.GetCullDistance(), 0.0f);
 
 			TArray<FContactPoint> ContactPoints;
 			GJKImplicitManifold<ConvexType>(Convex, ConvexTransform, HeightField, HeightFieldTransform, Constraint.GetCullDistance(), 0.0f, ContactPoints);
 			for (FContactPoint& ContactPoint : ContactPoints)
 			{
-				Constraint.AddOneshotManifoldContact(ContactPoint, Dt);
+				Constraint.AddOneshotManifoldContact(ContactPoint);
 			}
 		}
 
@@ -281,15 +300,17 @@ namespace Chaos
 		{
 			// We only build one shot manifolds once
 			// All convexes are pre-scaled, or wrapped in TImplicitObjectScaled
-			ensure(Constraint.GetManifoldPoints().Num() == 0);
 			ensure(ConvexTransform.GetScale3D() == FVec3(1.0f, 1.0f, 1.0f));
 			ensure(TriMeshTransform.GetScale3D() == FVec3(1.0f, 1.0f, 1.0f));
 
+			// @todo(chaos): support manifold maintenance
+			Constraint.ResetManifold();
+
 			TArray<FContactPoint> ContactPoints;
 			GJKImplicitManifold<ConvexType, TriMeshType>(Convex, ConvexTransform, TriangleMesh, TriMeshTransform, Constraint.GetCullDistance(), 0.0f, ContactPoints);
 			for (FContactPoint& ContactPoint : ContactPoints)
 			{
-				Constraint.AddOneshotManifoldContact(ContactPoint, Dt);
+				Constraint.AddOneshotManifoldContact(ContactPoint);
 			}
 		}
 
diff --git a/Engine/Source/Runtime/Experimental/Chaos/Private/Chaos/CollisionResolution.cpp b/Engine/Source/Runtime/Experimental/Chaos/Private/Chaos/CollisionResolution.cpp
--- a/Engine/Source/Runtime/Experimental/Chaos/Private/Chaos/CollisionResolution.cpp
+++ b/Engine/Source/Runtime/Experimental/Chaos/Private/Chaos/CollisionResolution.cpp
@@ -301,7 +301,7 @@ namespace Chaos
 			// point within reasonable range despite passing the AABB tests
 			if (ContactPoint.IsSet())
 			{
-				Constraint.AddIncrementalManifoldContact(ContactPoint, Dt);
+				Constraint.AddIncrementalManifoldContact(ContactPoint);
 			}
 			else
 			{
@@ -329,7 +329,7 @@ namespace Chaos
 
 			if (ContactPoint.IsSet())
 			{
-				Constraint.AddIncrementalManifoldContact(ContactPoint, Dt);
+				Constraint.AddIncrementalManifoldContact(ContactPoint);
 			}
 		}
 
@@ -455,16 +455,7 @@ namespace Chaos
 			{
 				if (Constraint.GetUseManifold())
 				{
-					// We only build one shot manifolds once
-					if (Constraint.GetManifoldPoints().Num() == 0)
-					{
-						ConstructConvexConvexOneShotManifold(A, ATM, B, BTM, Dt, Constraint);
-					}
-					else
-					{
-						// @todo(chaos): this is only used for the MinEvolution path - see if it can be removed/moved
-						Constraint.UpdateManifoldContacts();
-					}
+					ConstructConvexConvexOneShotManifold(A, ATM, B, BTM, Dt, Constraint);
 				}
 				else
 				{
@@ -526,15 +517,7 @@ namespace Chaos
 		{
 			if (Constraint.GetUseManifold())
 			{
-				// We only build one shot manifolds once
-				if (Constraint.GetManifoldPoints().Num() == 0)
-				{
-					ConstructBoxBoxOneShotManifold(Box1, Box1Transform, Box2, Box2Transform, Dt, Constraint);
-				}
-				else
-				{
-					Constraint.UpdateManifoldContacts();
-				}
+				ConstructBoxBoxOneShotManifold(Box1, Box1Transform, Box2, Box2Transform, Dt, Constraint);
 			}
 			else
 			{
@@ -569,14 +552,7 @@ namespace Chaos
 		{
 			if(Constraint.GetUseManifold())
 			{
-				if(Constraint.GetManifoldPoints().Num() == 0)
-				{
-					ConstructConvexHeightFieldOneShotManifold<FImplicitBox3>(A, ATransform, B, BTransform, Dt, Constraint);
-				}
-				else
-				{
-					Constraint.UpdateManifoldContacts();
-				}
+				ConstructConvexHeightFieldOneShotManifold<FImplicitBox3>(A, ATransform, B, BTransform, Dt, Constraint);
 			}
 			else
 			{
@@ -707,14 +683,7 @@ namespace Chaos
 		{
 			if(Constraint.GetUseManifold())
 			{
-				if(Constraint.GetManifoldPoints().Num() == 0)
-				{
-					ConstructConvexTriMeshOneShotManifold<FImplicitBox3, TriMeshType>(Box0, WorldTransform0, TriangleMesh1, WorldTransform1, Dt, Constraint);
-				}
-				else
-				{
-					Constraint.UpdateManifoldContacts();
-				}
+				ConstructConvexTriMeshOneShotManifold<FImplicitBox3, TriMeshType>(Box0, WorldTransform0, TriangleMesh1, WorldTransform1, Dt, Constraint);
 			}
 			else
 			{
@@ -766,16 +735,7 @@ namespace Chaos
 		{
 			if (Constraint.GetUseManifold())
 			{
-				// We only build one shot manifolds once
-				if (Constraint.GetManifoldPoints().Num() == 0)
-				{
-					ConstructSphereSphereOneShotManifold(Sphere1, Sphere1Transform, Sphere2, Sphere2Transform, Dt, Constraint);
-				}
-				else
-				{
-					// @todo(chaos): this is only used for the MinEvolution path - see if it can be removed/moved
-					Constraint.UpdateManifoldContacts();
-				}
+				ConstructSphereSphereOneShotManifold(Sphere1, Sphere1Transform, Sphere2, Sphere2Transform, Dt, Constraint);
 			}
 			else
 			{
@@ -808,16 +768,7 @@ namespace Chaos
 		{
 			if (Constraint.GetUseManifold())
 			{
-				// We only build one shot manifolds once
-				if (Constraint.GetManifoldPoints().Num() == 0)
-				{
-					ConstructSphereHeightFieldOneShotManifold(A, ATransform, B, BTransform, Dt, Constraint);
-				}
-				else
-				{
-					// @todo(chaos): this is only used for the MinEvolution path - see if it can be removed/moved
-					Constraint.UpdateManifoldContacts();
-				}
+				ConstructSphereHeightFieldOneShotManifold(A, ATransform, B, BTransform, Dt, Constraint);
 			}
 			else
 			{
@@ -872,16 +823,7 @@ namespace Chaos
 		{
 			if (Constraint.GetUseManifold())
 			{
-				// We only build one shot manifolds once
-				if (Constraint.GetManifoldPoints().Num() == 0)
-				{
-					ConstructSpherePlaneOneShotManifold(Sphere, SphereTransform, Plane, PlaneTransform, Dt, Constraint);
-				}
-				else
-				{
-					// @todo(chaos): this is only used for the MinEvolution path - see if it can be removed/moved
-					Constraint.UpdateManifoldContacts();
-				}
+				ConstructSpherePlaneOneShotManifold(Sphere, SphereTransform, Plane, PlaneTransform, Dt, Constraint);
 			}
 			else
 			{
@@ -916,16 +858,7 @@ namespace Chaos
 		{
 			if (Constraint.GetUseManifold())
 			{
-				// We only build one shot manifolds once
-				if (Constraint.GetManifoldPoints().Num() == 0)
-				{
-					ConstructSphereBoxOneShotManifold(Sphere, SphereTransform, Box, BoxTransform, Dt, Constraint);
-				}
-				else
-				{
-					// @todo(chaos): this is only used for the MinEvolution path - see if it can be removed/moved
-					Constraint.UpdateManifoldContacts();
-				}
+				ConstructSphereBoxOneShotManifold(Sphere, SphereTransform, Box, BoxTransform, Dt, Constraint);
 			}
 			else
 			{
@@ -961,16 +894,7 @@ namespace Chaos
 		{
 			if (Constraint.GetUseManifold())
 			{
-				// We only build one shot manifolds once
-				if (Constraint.GetManifoldPoints().Num() == 0)
-				{
-					ConstructSphereCapsuleOneShotManifold(A, ATransform, B, BTransform, Dt, Constraint);
-				}
-				else
-				{
-					// @todo(chaos): this is only used for the MinEvolution path - see if it can be removed/moved
-					Constraint.UpdateManifoldContacts();
-				}
+				ConstructSphereCapsuleOneShotManifold(A, ATransform, B, BTransform, Dt, Constraint);
 			}
 			else
 			{
@@ -1003,16 +927,7 @@ namespace Chaos
 		{
 			if (Constraint.GetUseManifold())
 			{
-				// We only build one shot manifolds once
-				if (Constraint.GetManifoldPoints().Num() == 0)
-				{
-					ConstructSphereConvexManifold(A, ATransform, B, BTransform, Dt, Constraint);
-				}
-				else
-				{
-					// @todo(chaos): this is only used for the MinEvolution path - see if it can be removed/moved
-					Constraint.UpdateManifoldContacts();
-				}
+				ConstructSphereConvexManifold(A, ATransform, B, BTransform, Dt, Constraint);
 			}
 			else
 			{
@@ -1046,16 +961,7 @@ namespace Chaos
 		{
 			if (Constraint.GetUseManifold())
 			{
-				// We only build one shot manifolds once
-				if (Constraint.GetManifoldPoints().Num() == 0)
-				{
-					ConstructSphereTriangleMeshOneShotManifold(Sphere0, WorldTransform0, TriangleMesh1, WorldTransform1, Dt, Constraint);
-				}
-				else
-				{
-					// @todo(chaos): this is only used for the MinEvolution path - see if it can be removed/moved
-					Constraint.UpdateManifoldContacts();
-				}
+				ConstructSphereTriangleMeshOneShotManifold(Sphere0, WorldTransform0, TriangleMesh1, WorldTransform1, Dt, Constraint);
 			}
 			else
 			{
@@ -1143,16 +1049,7 @@ namespace Chaos
 		{
 			if (Constraint.GetUseManifold())
 			{
-				// We only build one shot manifolds once
-				if (Constraint.GetManifoldPoints().Num() == 0)
-				{
-					ConstructCapsuleCapsuleOneShotManifold(A, ATransform, B, BTransform, Dt, Constraint);
-				}
-				else
-				{
-					// @todo(chaos): this is only used for the MinEvolution path - see if it can be removed/moved
-					Constraint.UpdateManifoldContacts();
-				}
+				ConstructCapsuleCapsuleOneShotManifold(A, ATransform, B, BTransform, Dt, Constraint);
 			}
 			else
 			{
@@ -1185,16 +1082,7 @@ namespace Chaos
 		{
 			if (Constraint.GetUseManifold())
 			{
-				// We only build one shot manifolds once
-				if (Constraint.GetManifoldPoints().Num() == 0)
-				{
-					ConstructConvexConvexOneShotManifold<FCapsule, FImplicitBox3>(A, ATransform, B, BTransform, Dt, Constraint);
-				}
-				else
-				{
-					// @todo(chaos): this is only used for the MinEvolution path - see if it can be removed/moved
-					Constraint.UpdateManifoldContacts();
-				}
+				ConstructConvexConvexOneShotManifold<FCapsule, FImplicitBox3>(A, ATransform, B, BTransform, Dt, Constraint);
 			}
 			else
 			{
@@ -1247,16 +1135,7 @@ namespace Chaos
 			{
 				if (Constraint.GetUseManifold())
 				{
-					// We only build one shot manifolds once
-					if (Constraint.GetManifoldPoints().Num() == 0)
-					{
-						ConstructConvexConvexOneShotManifold<FCapsule, ConvexType>(A, ATransform, B, BTransform, Dt, Constraint);
-					}
-					else
-					{
-						// @todo(chaos): this is only used for the MinEvolution path - see if it can be removed/moved
-						Constraint.UpdateManifoldContacts();
-					}
+					ConstructConvexConvexOneShotManifold<FCapsule, ConvexType>(A, ATransform, B, BTransform, Dt, Constraint);
 				}
 				else
 				{
@@ -1315,16 +1194,7 @@ namespace Chaos
 			CONDITIONAL_SCOPE_CYCLE_COUNTER(STAT_Collisions_UpdateCapsuleHeightFieldConstraint, ConstraintsDetailedStats);
 			if (Constraint.GetUseManifold())
 			{
-				// We only build one shot manifolds once
-				if (Constraint.GetManifoldPoints().Num() == 0)
-				{
-					ConstructCapsuleHeightFieldOneShotManifold(A, ATransform, B, BTransform, Dt, Constraint);
-				}
-				else
-				{
-					// @todo(chaos): this is only used for the MinEvolution path - see if it can be removed/moved
-					Constraint.UpdateManifoldContacts();
-				}
+				ConstructCapsuleHeightFieldOneShotManifold(A, ATransform, B, BTransform, Dt, Constraint);
 			}
 			else
 			{
@@ -1387,16 +1257,7 @@ namespace Chaos
 			CONDITIONAL_SCOPE_CYCLE_COUNTER(STAT_Collisions_UpdateCapsuleTriangleMeshConstraint, ConstraintsDetailedStats);
 			if (Constraint.GetUseManifold())
 			{
-				// We only build one shot manifolds once
-				if (Constraint.GetManifoldPoints().Num() == 0)
-				{
-					ConstructCapsuleTriMeshOneShotManifold(Capsule0, WorldTransform0, TriangleMesh1, WorldTransform1, Dt, Constraint);
-				}
-				else
-				{
-					// @todo(chaos): this is only used for the MinEvolution path - see if it can be removed/moved
-					Constraint.UpdateManifoldContacts();
-				}
+				ConstructCapsuleTriMeshOneShotManifold(Capsule0, WorldTransform0, TriangleMesh1, WorldTransform1, Dt, Constraint);
 			}
 			else
 			{
@@ -1543,16 +1404,7 @@ namespace Chaos
 			CONDITIONAL_SCOPE_CYCLE_COUNTER(STAT_Collisions_UpdateConvexHeightFieldConstraint, ConstraintsDetailedStats);
 			if (Constraint.GetUseManifold())
 			{
-				// We only build one shot manifolds once
-				if (Constraint.GetManifoldPoints().Num() == 0)
-				{
-					ConstructConvexHeightFieldOneShotManifold(A, ATransform, B, BTransform, Dt, Constraint);
-				}
-				else
-				{
-					// @todo(chaos): this is only used for the MinEvolution path - see if it can be removed/moved
-					Constraint.UpdateManifoldContacts();
-				}
+				ConstructConvexHeightFieldOneShotManifold(A, ATransform, B, BTransform, Dt, Constraint);
 			}
 			else
 			{
@@ -1610,16 +1462,7 @@ namespace Chaos
 			CONDITIONAL_SCOPE_CYCLE_COUNTER(STAT_Collisions_UpdateConvexTriangleMeshConstraint, ConstraintsDetailedStats);
 			if (Constraint.GetUseManifold())
 			{
-				// We only build one shot manifolds once
-				if (Constraint.GetManifoldPoints().Num() == 0)
-				{
-					ConstructConvexTriMeshOneShotManifold(Convex0, WorldTransform0, TriangleMesh1, WorldTransform1, Dt, Constraint);
-				}
-				else
-				{
-					// @todo(chaos): this is only used for the MinEvolution path - see if it can be removed/moved
-					Constraint.UpdateManifoldContacts();
-				}
+				ConstructConvexTriMeshOneShotManifold(Convex0, WorldTransform0, TriangleMesh1, WorldTransform1, Dt, Constraint);
 			}
 			else
 			{
@@ -2576,13 +2419,8 @@ namespace Chaos
 			ConstructConstraintsImpl<T_TRAITS>(Particle0, Particle1, Implicit0, Simplicial0, Implicit1, Simplicial1, ParticleWorldTransform0, LocalTransform0, ParticleWorldTransform1, LocalTransform1, CullDistance, Dt, Context);
 		}
 
-		EContactShapesType CalculateShapePairType(const FImplicitObject* Implicit0, const FBVHParticles* BVHParticles0, const FImplicitObject* Implicit1, const FBVHParticles* BVHParticles1, bool& bOutSwap)
+		EContactShapesType CalculateShapePairType(const EImplicitObjectType Implicit0Type, const EImplicitObjectType Implicit1Type, const bool bIsConvex0, const bool bIsConvex1, const bool bIsBVH0, const bool bIsBVH1,  bool &bOutSwap)
 		{
-			const EImplicitObjectType Implicit0Type = (Implicit0 != nullptr) ? GetInnerType(Implicit0->GetCollisionType()) : ImplicitObjectType::Unknown;
-			const EImplicitObjectType Implicit1Type = (Implicit1 != nullptr) ? GetInnerType(Implicit1->GetCollisionType()) : ImplicitObjectType::Unknown;
-			const bool bIsConvex0 = (Implicit0 != nullptr) && Implicit0->IsConvex() && (Implicit0Type != ImplicitObjectType::LevelSet);
-			const bool bIsConvex1 = (Implicit1 != nullptr) && Implicit1->IsConvex() && (Implicit1Type != ImplicitObjectType::LevelSet);
-
 			bOutSwap = false;
 			if (Implicit0Type == TBox<FReal, 3>::StaticType() && Implicit1Type == TBox<FReal, 3>::StaticType())
 			{
@@ -2737,7 +2575,7 @@ namespace Chaos
 			}
 			else if ((Implicit0Type == FLevelSet::StaticType()) || (Implicit1Type == FLevelSet::StaticType()))
 			{
-				bOutSwap = (BVHParticles0 == nullptr) && (BVHParticles1 != nullptr);
+				bOutSwap = !bIsBVH0 && bIsBVH1;
 				return EContactShapesType::LevelSetLevelSet;
 			}
 
@@ -2746,6 +2584,16 @@ namespace Chaos
 			return EContactShapesType::Unknown;
 		}
 
+		EContactShapesType CalculateShapePairType(const FImplicitObject* Implicit0, const FBVHParticles* BVHParticles0, const FImplicitObject* Implicit1, const FBVHParticles* BVHParticles1, bool& bOutSwap)
+		{
+			const EImplicitObjectType Implicit0Type = (Implicit0 != nullptr) ? GetInnerType(Implicit0->GetCollisionType()) : ImplicitObjectType::Unknown;
+			const EImplicitObjectType Implicit1Type = (Implicit1 != nullptr) ? GetInnerType(Implicit1->GetCollisionType()) : ImplicitObjectType::Unknown;
+			const bool bIsConvex0 = (Implicit0 != nullptr) && Implicit0->IsConvex() && (Implicit0Type != ImplicitObjectType::LevelSet);
+			const bool bIsConvex1 = (Implicit1 != nullptr) && Implicit1->IsConvex() && (Implicit1Type != ImplicitObjectType::LevelSet);
+			const bool bIsBVH0 = (BVHParticles0 != nullptr);
+			const bool bIsBVH1 = (BVHParticles1 != nullptr);
+			return CalculateShapePairType(Implicit0Type, Implicit1Type, bIsConvex0, bIsConvex1, bIsBVH0, bIsBVH1, bOutSwap);
+		}
 
 		void ConstructConstraints(TGeometryParticleHandle<FReal, 3>* Particle0, TGeometryParticleHandle<FReal, 3>* Particle1, const FImplicitObject* Implicit0, const FBVHParticles* Simplicial0, const FImplicitObject* Implicit1, const FBVHParticles* Simplicial1, const FRigidTransform3& ParticleWorldTransform0, const FRigidTransform3& LocalTransform0, const FRigidTransform3& ParticleWorldTransform1, const FRigidTransform3& LocalTransform1, const FReal CullDistance, const FReal dT, const FCollisionContext& Context)
 		{
diff --git a/Engine/Source/Runtime/Experimental/Chaos/Private/Chaos/HeightField.cpp b/Engine/Source/Runtime/Experimental/Chaos/Private/Chaos/HeightField.cpp
--- a/Engine/Source/Runtime/Experimental/Chaos/Private/Chaos/HeightField.cpp
+++ b/Engine/Source/Runtime/Experimental/Chaos/Private/Chaos/HeightField.cpp
@@ -1349,6 +1349,9 @@ namespace Chaos
 				continue;
 			}
 
+			// @todo(chaos): we should not be creating constraints just for collecting contacts...
+			FPBDCollisionConstraint Constraint = FPBDCollisionConstraint::MakeTriangle(&QueryGeom);
+
 			// The triangle is solid so proceed to test it
 			FAABB3 CellBounds;
 			GeomData.GetPointsAndBoundsScaled(SingleIndex, Points, CellBounds);
@@ -1356,7 +1359,7 @@ namespace Chaos
 			{
 				// First Triangle
 				{
-					FPBDCollisionConstraint Constraint;
+					Constraint.ResetManifold();
 					OverlapTriangle(Points[0], Points[1], Points[3], Constraint);
 					for (FManifoldPoint& ManifoldPoint : Constraint.GetManifoldPoints())
 					{
@@ -1366,7 +1369,7 @@ namespace Chaos
 				}
 				// Second Triangle
 				{
-					FPBDCollisionConstraint Constraint;
+					Constraint.ResetManifold();
 					OverlapTriangle(Points[0], Points[3], Points[2], Constraint);
 					for (FManifoldPoint& ManifoldPoint : Constraint.GetManifoldPoints())
 					{
diff --git a/Engine/Source/Runtime/Experimental/Chaos/Private/Chaos/PBDCollisionConstraints.cpp b/Engine/Source/Runtime/Experimental/Chaos/Private/Chaos/PBDCollisionConstraints.cpp
--- a/Engine/Source/Runtime/Experimental/Chaos/Private/Chaos/PBDCollisionConstraints.cpp
+++ b/Engine/Source/Runtime/Experimental/Chaos/Private/Chaos/PBDCollisionConstraints.cpp
@@ -347,7 +347,12 @@ namespace Chaos
 	{
 		if (SolverType == EConstraintSolverType::QuasiPbd)
 		{
-			GetConstraintSolverContainer(SolverData).Reset(NumIslandConstraints);
+			FPBDCollisionSolverContainer& SolverContainer = GetConstraintSolverContainer(SolverData);
+			SolverContainer.Reset(NumIslandConstraints);
+
+			// @todo(chaos): we shouldn't have to pass settings down to the solvers every frame. Ideally the solver containers would be created by the constraint containers and
+			// the settings would be passed on then. Currently the solver containers are created here FPBDIslandSolverData::AddConstraintDatas
+			SolverContainer.SetMaxPushOutVelocity(MaxPushOutVelocity);
 		}
 		else
 		{
@@ -370,10 +375,6 @@ namespace Chaos
 
 			FPBDCollisionSolverContainer& SolverContainer = GetConstraintSolverContainer(SolverData);
 			SolverContainer.AddConstraintSolver(Dt, Constraint, Particle0Level, Particle1Level, SolverData.GetBodyContainer());
-
-			// @todo(chaos): we shouldn't have to pass settings down to the solvers every frame. Ideally the solver containers would be created by the constraint containers and
-			// the settings would be passed on then. Currently the solver containers are created here FPBDIslandSolverData::AddConstraintDatas
-			SolverContainer.SetMaxPushOutVelocity(MaxPushOutVelocity);
 		}
 		else
 		{
diff --git a/Engine/Source/Runtime/Experimental/Chaos/Private/Chaos/PBDCollisionConstraintsContact.cpp b/Engine/Source/Runtime/Experimental/Chaos/Private/Chaos/PBDCollisionConstraintsContact.cpp
--- a/Engine/Source/Runtime/Experimental/Chaos/Private/Chaos/PBDCollisionConstraintsContact.cpp
+++ b/Engine/Source/Runtime/Experimental/Chaos/Private/Chaos/PBDCollisionConstraintsContact.cpp
@@ -196,7 +196,14 @@ namespace Chaos
 			for (int32 PairIt = 0; PairIt < IterationParameters.NumPairIterations; ++PairIt)
 			{
 				// Update the contact information based on current particles' positions
-				Collisions::Update(Constraint, IterationParameters.Dt);
+				if (!Constraint.GetUseManifold() || (Constraint.GetManifoldPoints().Num() == 0))
+				{
+					Collisions::Update(Constraint, IterationParameters.Dt);
+				}
+				else
+				{
+					Constraint.UpdateManifoldContacts();
+				}
 
 				// Permanently disable a constraint that is beyond the cull distance
 				if (Constraint.GetPhi() >= Constraint.GetCullDistance())
diff --git a/Engine/Source/Runtime/Experimental/Chaos/Private/Chaos/TriangleMeshImplicitObject.cpp b/Engine/Source/Runtime/Experimental/Chaos/Private/Chaos/TriangleMeshImplicitObject.cpp
--- a/Engine/Source/Runtime/Experimental/Chaos/Private/Chaos/TriangleMeshImplicitObject.cpp
+++ b/Engine/Source/Runtime/Experimental/Chaos/Private/Chaos/TriangleMeshImplicitObject.cpp
@@ -391,11 +391,14 @@ bool FTriangleMeshImplicitObject::ContactManifoldImp(const GeomType& QueryGeom,
 		QueryBounds.ThickenSymmetrically(LocalThickness);
 		const TArray<int32> PotentialIntersections = BVH.FindAllIntersections(QueryBounds);
 
+		// @todo(chaos): we should not be creating constraints just for collecting contacts...
+		FPBDCollisionConstraint Constraint = FPBDCollisionConstraint::MakeTriangle(&QueryGeom);
+
 		for (int32 TriIdx : PotentialIntersections)
 		{
 			FVec3 A, B, C;
 			TransformVertsHelper(TriMeshScale, TriIdx, MParticles, Elements, A, B, C);
-			FPBDCollisionConstraint Constraint;
+			Constraint.ResetManifold();
 			OverlapTriangle(A, B, C, Constraint);
 			for(FManifoldPoint& ManifoldPoint : Constraint.GetManifoldPoints())
 			{
diff --git a/Engine/Source/Runtime/Experimental/Chaos/Private/PBDRigidsSolver.cpp b/Engine/Source/Runtime/Experimental/Chaos/Private/PBDRigidsSolver.cpp
--- a/Engine/Source/Runtime/Experimental/Chaos/Private/PBDRigidsSolver.cpp
+++ b/Engine/Source/Runtime/Experimental/Chaos/Private/PBDRigidsSolver.cpp
@@ -43,6 +43,10 @@ DECLARE_DWORD_COUNTER_STAT(TEXT("NumIslands"), STAT_ChaosCounter_NumIslands, STA
 DECLARE_DWORD_COUNTER_STAT(TEXT("NumContacts"), STAT_ChaosCounter_NumContacts, STATGROUP_ChaosCounters);
 DECLARE_DWORD_COUNTER_STAT(TEXT("NumValidConstraints"), STAT_ChaosCounter_NumValidConstraints, STATGROUP_ChaosCounters);
 DECLARE_DWORD_COUNTER_STAT(TEXT("NumActiveConstraints"), STAT_ChaosCounter_NumActiveConstraints, STATGROUP_ChaosCounters);
+DECLARE_DWORD_COUNTER_STAT(TEXT("NumRestoredConstraints"), STAT_ChaosCounter_NumRestoredConstraints, STATGROUP_ChaosCounters);
+DECLARE_DWORD_COUNTER_STAT(TEXT("NumManifoldPoints"), STAT_ChaosCounter_NumManifoldPoints, STATGROUP_ChaosCounters);
+DECLARE_DWORD_COUNTER_STAT(TEXT("NumActiveManifoldPoints"), STAT_ChaosCounter_NumActiveManifoldPoints, STATGROUP_ChaosCounters);
+DECLARE_DWORD_COUNTER_STAT(TEXT("NumRestoredManifoldPoints"), STAT_ChaosCounter_NumRestoredManifoldPoints, STATGROUP_ChaosCounters);
 DECLARE_DWORD_COUNTER_STAT(TEXT("NumJoints"), STAT_ChaosCounter_NumJoints, STATGROUP_ChaosCounters);
 
 // Stat Iteration counters
@@ -1406,6 +1410,10 @@ CSV_CUSTOM_STAT(PhysicsCounters, Name, Value, ECsvCustomStatOp::Set);
 	{
 		int32 NumValidCollisions = 0;
 		int32 NumActiveCollisions = 0;
+		int32 NumRestoredCollisions = 0;
+		int32 NumManifoldPoints = 0;
+		int32 NumActiveManifoldPoints = 0;
+		int32 NumRestoredManifoldPoints = 0;
 		for (const FPBDCollisionConstraintHandle* Collision : GetEvolution()->GetCollisionConstraints().GetConstraints())
 		{
 			if (Collision->GetContact().IsEnabled())
@@ -1418,11 +1426,32 @@ CSV_CUSTOM_STAT(PhysicsCounters, Name, Value, ECsvCustomStatOp::Set);
 				{
 					++NumActiveCollisions;
 				}
+				if (Collision->GetContact().WasManifoldRestored())
+				{
+					++NumRestoredCollisions;
+				}
+				for (const FManifoldPoint& ManifoldPoint : Collision->GetContact().GetManifoldPoints())
+				{
+					++NumManifoldPoints;
+					if (ManifoldPoint.bWasRestored || Collision->GetContact().WasManifoldRestored())
+					{
+						++NumRestoredManifoldPoints;
+					}
+					if (!ManifoldPoint.NetPushOut.IsNearlyZero())
+					{
+						++NumActiveManifoldPoints;
+					}
+				}
+
 			}
 		}
 
 		CHAOS_COUNTER_STAT(NumValidConstraints, NumValidCollisions);
 		CHAOS_COUNTER_STAT(NumActiveConstraints, NumActiveCollisions);
+		CHAOS_COUNTER_STAT(NumRestoredConstraints, NumRestoredCollisions);
+		CHAOS_COUNTER_STAT(NumManifoldPoints, NumManifoldPoints);
+		CHAOS_COUNTER_STAT(NumActiveManifoldPoints, NumActiveManifoldPoints);
+		CHAOS_COUNTER_STAT(NumRestoredManifoldPoints, NumRestoredManifoldPoints);
 	}
 
 	void FPBDRigidsSolver::PostTickDebugDraw(FReal Dt) const
diff --git a/Engine/Source/Runtime/Experimental/Chaos/Public/Chaos/Collision/CollisionConstraintAllocator.h b/Engine/Source/Runtime/Experimental/Chaos/Public/Chaos/Collision/CollisionConstraintAllocator.h
--- a/Engine/Source/Runtime/Experimental/Chaos/Public/Chaos/Collision/CollisionConstraintAllocator.h
+++ b/Engine/Source/Runtime/Experimental/Chaos/Public/Chaos/Collision/CollisionConstraintAllocator.h
@@ -11,8 +11,6 @@
 #include "Chaos/ParticleHandle.h"
 #include "ChaosStats.h"
 
-extern bool bChaos_Collision_EnableManifoldRestore;
-
 namespace Chaos
 {
 	/**
diff --git a/Engine/Source/Runtime/Experimental/Chaos/Public/Chaos/Collision/PBDCollisionConstraint.h b/Engine/Source/Runtime/Experimental/Chaos/Public/Chaos/Collision/PBDCollisionConstraint.h
--- a/Engine/Source/Runtime/Experimental/Chaos/Public/Chaos/Collision/PBDCollisionConstraint.h
+++ b/Engine/Source/Runtime/Experimental/Chaos/Public/Chaos/Collision/PBDCollisionConstraint.h
@@ -8,6 +8,7 @@
 #include "Chaos/Collision/CollisionKeys.h"
 #include "Chaos/Collision/ContactPoint.h"
 #include "Chaos/Collision/PBDCollisionConstraintHandle.h"
+#include "Chaos/Framework/UncheckedArray.h"
 #include "Chaos/GJK.h"
 #include "Chaos/ParticleHandleFwd.h"
 #include "Chaos/ParticleHandle.h"
@@ -33,32 +34,94 @@ namespace Chaos
 
 	CHAOS_API bool ContactConstraintSortPredicate(const FPBDCollisionConstraint& L, const FPBDCollisionConstraint& R);
 
+	/**
+	 * @brief A single point in a contact manifold.
+	 * Each Collision Constraint will have up to 4 of these.
+	*/
 	class CHAOS_API FManifoldPoint
 	{
 	public:
 		FManifoldPoint() 
-			: CoMContactPoints{ FVec3(0), FVec3(0) }
+			: ContactPoint()
+			, InitialShapeContactPoints{ FVec3(0), FVec3(0) }
+			, CoMContactPoints{ FVec3(0), FVec3(0) }
 			, NetPushOut(0)
 			, NetImpulse(0)
 			, StaticFrictionMax(0)
 			, bInsideStaticFrictionCone(false)
+			, bWasRestored(false)
 		{}
 
 		FManifoldPoint(const FContactPoint& InContactPoint) 
 			: ContactPoint(InContactPoint)
+			, InitialShapeContactPoints{ FVec3(0), FVec3(0) }
 			, CoMContactPoints{ FVec3(0), FVec3(0) }
 			, NetPushOut(0)
 			, NetImpulse(0)
 			, StaticFrictionMax(0)
 			, bInsideStaticFrictionCone(false)
+			, bWasRestored(false)
 		{}
 
 		FContactPoint ContactPoint;			// Contact point results of low-level collision detection
+		FVec3 InitialShapeContactPoints[2];	// ShapeContactPoints when the constraint was first initialized. Used to track reusablility
 		FVec3 CoMContactPoints[2];			// CoM-space contact points on the two bodies
 		FVec3 NetPushOut;					// Total pushout applied at this contact point
 		FVec3 NetImpulse;					// Total impulse applied by this contact point
 		FReal StaticFrictionMax;			// A proxy for the normal impulse used to limit static friction correction. Used for smoothing static friction limits
 		bool bInsideStaticFrictionCone;		// Whether we are inside the static friction cone (used in PushOut)
+		bool bWasRestored;
+	};
+
+	/**
+	 * @brief The friction data for a manifold point
+	 * This is the information that needs to be stored between ticks to implement static friction.
+	*/
+	class CHAOS_API FManifoldPointSavedData
+	{
+	public:
+		FManifoldPointSavedData()
+		{
+		}
+
+		/**
+		 * @brief Copy the ManifoldPoint data needed for static friction next tick
+		*/
+		inline void Save(const FManifoldPoint& ManifoldPoint)
+		{
+			CoMContactPoints[0] = ManifoldPoint.CoMContactPoints[0];
+			CoMContactPoints[1] = ManifoldPoint.CoMContactPoints[1];
+			StaticFrictionMax = ManifoldPoint.StaticFrictionMax;
+			bInsideStaticFrictionCone = ManifoldPoint.bInsideStaticFrictionCone;
+		}
+
+		/**
+		 * @brief Replace the ManifoldPoint's data with the saved data from the previous tick
+		*/
+		inline void Restore(FManifoldPoint& ManifoldPoint) const
+		{
+			ManifoldPoint.CoMContactPoints[0] = CoMContactPoints[0];
+			ManifoldPoint.CoMContactPoints[1] = CoMContactPoints[1];
+			ManifoldPoint.StaticFrictionMax = StaticFrictionMax;
+		}
+
+		/**
+		 * @brief Is the data for the specified manifold point?
+		*/
+		inline bool IsMatch(const FManifoldPoint& ManifoldPoint, const FReal DistanceToleranceSq) const
+		{
+			// If the contact point is in the same spot on one of the bodies, assume it is the same contact
+			// @todo(chaos): more robust same-point test. E.g., this won't work for very small or very large objects,
+			// so at least make the tolerance size-dependent
+			const FVec3 DP0 = ManifoldPoint.CoMContactPoints[0] - CoMContactPoints[0];
+			const FVec3 DP1 = ManifoldPoint.CoMContactPoints[1] - CoMContactPoints[1];
+			return ((DP0.SizeSquared() < DistanceToleranceSq) || (DP1.SizeSquared() < DistanceToleranceSq));
+		}
+
+		FVec3 CoMContactPoints[2];
+		FVec3 NetPushOut;
+		FReal StaticFrictionMax;
+		bool bInsideStaticFrictionCone;
 	};
 
 	/*
@@ -206,6 +269,8 @@ namespace Chaos
 	public:
 		using FConstraintContainerHandle = TIntrusiveConstraintHandle<FPBDCollisionConstraint>;
 
+		static const int32 MaxManifoldPoints = 4;
+
 		/**
 		 * @brief Create a contact constraint
 		 * Allocates a constraint on the heap, with a permanent address.
@@ -224,6 +289,12 @@ namespace Chaos
 			const bool bInUseManifold,
 			const EContactShapesType ShapesType);
 
+		/**
+		 * @brief For use by the tri mesh and heighfield collision detection as a temporary measure
+		 * @see FHeightField::ContactManifoldImp, FTriangleMeshImplicitObject::ContactManifoldImp
+		*/
+		static FPBDCollisionConstraint MakeTriangle(const FImplicitObject* Implicit0);
+
 		/**
 		 * @brief Return a constraint copied from the Source constraint, for use in the Resim Cache or other system
 		 * @note Unlike the other factory method, this version returns a constraint by value for emplacing into an array (ther others are by pointer)
@@ -266,6 +337,9 @@ namespace Chaos
 		const FImplicitObject* GetImplicit1() const { return Manifold.Implicit[1]; }
 		const FBVHParticles* GetCollisionParticles0() const { return Manifold.Simplicial[0]; }
 		const FBVHParticles* GetCollisionParticles1() const { return Manifold.Simplicial[1]; }
+		const FReal GetCollisionMargin0() const { return CollisionMargins[0]; }
+		const FReal GetCollisionMargin1() const { return CollisionMargins[1]; }
+		const FReal GetCollisionTolerance() const { return CollisionTolerance; }
 
 		// @todo(chaos): half of this API is wrong for the new multi-point manifold constraints. Remove it
 
@@ -323,11 +397,17 @@ namespace Chaos
 		TArrayView<FManifoldPoint> GetManifoldPoints() { return MakeArrayView(ManifoldPoints); }
 		TArrayView<const FManifoldPoint> GetManifoldPoints() const { return MakeArrayView(ManifoldPoints); }
 
-		void AddIncrementalManifoldContact(const FContactPoint& ContactPoint, const FReal Dt);
-		void AddOneshotManifoldContact(const FContactPoint& ContactPoint, const FReal Dt);
+		void AddIncrementalManifoldContact(const FContactPoint& ContactPoint);
+		void AddOneshotManifoldContact(const FContactPoint& ContactPoint);
 		void UpdateManifoldContacts();
 		void ClearManifold();
 
+		bool UpdateAndTryRestoreManifold(const FRigidTransform3& ShapeWorldTransform0, const FRigidTransform3& ShapeWorldTransform1);
+		void ResetActiveManifoldContacts();
+		void UpdateLastShapeWorldTransforms(const FRigidTransform3& ShapeWorldTransform0, const FRigidTransform3& ShapeWorldTransform1);
+		bool TryAddManifoldContact(const FContactPoint& ContactPoint, const FRigidTransform3& ShapeWorldTransform0, const FRigidTransform3& ShapeWorldTransform1);
+		bool TryInsertManifoldContact(const FContactPoint& ContactPoint, const FRigidTransform3& ShapeWorldTransform0, const FRigidTransform3& ShapeWorldTransform1);
+
 		//@ todo(chaos): These are for the collision forwarding system - this should use the collision modifier system (which should be extended to support adding collisions)
 		void SetManifoldPoints(const TArray<FManifoldPoint>& InManifoldPoints) { ManifoldPoints = InManifoldPoints; }
 		void UpdateManifoldPointFromContact(const int32 ManifoldPointIndex);
@@ -374,10 +454,8 @@ namespace Chaos
 		 */
 		ECollisionConstraintDirection GetConstraintDirection(const FReal Dt) const;
 
-		/**
-		 * @brief The container key for this contact (for internal use only)
-		*/
-		//const FRigidBodyContactKey& GetKey() const { return GetContainerCookie().Key; }
+
+		const FManifoldPointSavedData* FindManifoldPointSavedData(const FManifoldPoint& ManifoldPoint) const;
 
 	public:
 		const FPBDCollisionConstraintHandle* GetConstraintHandle() const { return this; }
@@ -412,9 +490,9 @@ namespace Chaos
 
 		bool AreMatchingContactPoints(const FContactPoint& A, const FContactPoint& B, FReal& OutScore) const;
 		int32 FindManifoldPoint(const FContactPoint& ContactPoint) const;
-		int32 AddManifoldPoint(const FContactPoint& ContactPoint, const FReal Dt);
-		void InitManifoldPoint(const int32 ManifoldPointIndex, FReal Dt);
-		void UpdateManifoldPoint(int32 ManifoldPointIndex, const FContactPoint& ContactPoint, const FReal Dt);
+		int32 AddManifoldPoint(const FContactPoint& ContactPoint);
+		void InitManifoldPoint(const int32 ManifoldPointIndex);
+		void UpdateManifoldPoint(int32 ManifoldPointIndex, const FContactPoint& ContactPoint);
 		void SetActiveContactPoint(const FContactPoint& ContactPoint);
 		void GetWorldSpaceManifoldPoint(const FManifoldPoint& ManifoldPoint, const FVec3& P0, const FRotation3& Q0, const FVec3& P1, const FRotation3& Q1, FVec3& OutContactLocation, FReal& OutContactPhi);
 
@@ -426,7 +504,9 @@ namespace Chaos
 		/**
 		 * @brief Restore data used for static friction from the previous point(s) to the current
 		*/
-		bool TryRestoreFrictionData(const int32 ManifoldPointIndex);
+		void TryRestoreFrictionData(const int32 ManifoldPointIndex);
+
+		void InitMargins(const EImplicitObjectType ImplicitType0, const EImplicitObjectType ImplicitType1, const FReal Margin0, const FReal Margin1);
 
 	public:
 		//@todo(chaos): make this stuff private
@@ -443,10 +523,17 @@ namespace Chaos
 		ECollisionCCDType CCDType;
 		FReal Stiffness;
 
-		// @todo(chaos): Switching this to inline allocator does not help, but maybe a physics scratchpad would
-		//TArray<FManifoldPoint, TInlineAllocator<4>> ManifoldPoints;
 		TArray<FManifoldPoint> ManifoldPoints;
 		FReal CullDistance;
+
+		// The margins to use during collision detection. We don't always use the margins on the shapes directly.
+		// E.g., we use the smallest non-zero margin for 2 convex shapes. See InitMargins
+		FReal CollisionMargins[2];
+
+		// The collision tolerance is used to determine whether a new contact matches an old on. It is derived from the
+		// margins of the two shapes, as well as their types
+		FReal CollisionTolerance;
+
 		bool bUseManifold;
 		bool bUseIncrementalManifold;
 
@@ -457,9 +544,12 @@ namespace Chaos
 		FGJKSimplexData GJKWarmStartData;
 
 		// The manifold points from the previous tick when we don't reuse the manifold. Used by static friction.
-		// @todo(chaos): this could be a subset of the manifold data
-		TArray<FManifoldPoint> PrevManifoldPoints;
+		FManifoldPointSavedData ManifoldPointSavedData[MaxManifoldPoints];
+		int32 NumSavedManifoldPoints;
 
+		FRigidTransform3 LastShapeWorldTransform0;
+		FRigidTransform3 LastShapeWorldTransform1;
+		int32 ExpectedNumManifoldPoints;
 		bool bWasManifoldRestored;
 
 		int32 NumActivePositionIterations;
diff --git a/Engine/Source/Runtime/Experimental/Chaos/Public/Chaos/Collision/ParticlePairMidPhase.h b/Engine/Source/Runtime/Experimental/Chaos/Public/Chaos/Collision/ParticlePairMidPhase.h
--- a/Engine/Source/Runtime/Experimental/Chaos/Public/Chaos/Collision/ParticlePairMidPhase.h
+++ b/Engine/Source/Runtime/Experimental/Chaos/Public/Chaos/Collision/ParticlePairMidPhase.h
@@ -23,7 +23,8 @@ namespace Chaos
 
 	/**
 	 * @brief Handles collision detection for a pair of simple shapes (i.e., not compound shapes)
-	 * @note this is not used for collisions involving Unions.
+	 * 
+	 * @note this is not used for collisions involving Unions that require a recursive collision test.
 	 * @see FMultiShapePairCollisionDetector
 	*/
 	class CHAOS_API FSingleShapePairCollisionDetector
@@ -32,10 +33,11 @@ namespace Chaos
 		using FCollisionsArray = TArray<FPBDCollisionConstraint*, TInlineAllocator<1>>;
 
 		FSingleShapePairCollisionDetector(
-			const FPerShapeData* InShape0, 
-			const FPerShapeData* InShape1, 
+			FGeometryParticleHandle* InParticle0,
+			const FPerShapeData* InShape0,
+			FGeometryParticleHandle* InParticle1,
+			const FPerShapeData* InShape1,
 			const EContactShapesType InShapePairType, 
-			const bool bInSwap, 
 			FParticlePairMidPhase& MidPhase);
 		FSingleShapePairCollisionDetector(FSingleShapePairCollisionDetector&& R);
 		FSingleShapePairCollisionDetector(const FSingleShapePairCollisionDetector& R) = delete;
@@ -44,6 +46,10 @@ namespace Chaos
 
 		const FPBDCollisionConstraint* GetConstraint() const{ return Constraint.Get(); }
 		FPBDCollisionConstraint* GetConstraint() { return Constraint.Get(); }
+		const FGeometryParticleHandle* GetParticle0() const { return Particle0; }
+		FGeometryParticleHandle* GetParticle0() { return Particle0; }
+		const FGeometryParticleHandle* GetParticle1() const { return Particle1; }
+		FGeometryParticleHandle* GetParticle1() { return Particle1; }
 		const FPerShapeData* GetShape0() const { return Shape0; }
 		const FPerShapeData* GetShape1() const { return Shape1; }
 
@@ -57,8 +63,6 @@ namespace Chaos
 		 * @return The number of collisions constraints that were activated
 		*/
 		int32 GenerateCollision(
-			FGeometryParticleHandle* Particle0,
-			FGeometryParticleHandle* Particle1,
 			const FReal CullDistance,
 			const bool bUseCCD,
 			const FReal Dt);
@@ -67,7 +71,7 @@ namespace Chaos
 		 * @brief Reactivate the collision exactly as it was last frame
 		 * @return The number of collisions constraints that were restored
 		*/
-		int32 RestoreCollision();
+		int32 RestoreCollision(const FReal CullDistance);
 
 		/**
 		 * @brief Reactivate the constraint (essentially the same as Restore but slightly optimized)
@@ -83,26 +87,9 @@ namespace Chaos
 		void SetCollision(const FPBDCollisionConstraint& Constraint);
 
 	private:
-		int32 GenerateCollisionImpl(
-			FGeometryParticleHandle* Particle0,
-			const FPerShapeData* Shape0,
-			FGeometryParticleHandle* Particle1,
-			const FPerShapeData* Shape1,
-			const FReal CullDistance,
-			const bool bUseCCD,
-			const FReal Dt);
+		int32 GenerateCollisionImpl(const FReal CullDistance, const bool bUseCCD, const FReal Dt);
 
-		void CreateConstraint(
-			FGeometryParticleHandle* Particle0,
-			const FPerShapeData* InShape0,
-			FGeometryParticleHandle* Particle1,
-			const FPerShapeData* InShape1,
-			const FReal CullDistance);
-
-		/**
-		 * @brief Reset or save the constraint manifold
-		*/
-		void PrepareConstraint(const bool bUseCCD);
+		void CreateConstraint(const FReal CullDistance);
 
 		/**
 		 * @brief Add the constraint to the scene's active list
@@ -111,10 +98,14 @@ namespace Chaos
 
 		FParticlePairMidPhase& MidPhase;
 		TUniquePtr<FPBDCollisionConstraint> Constraint;
+		FGeometryParticleHandle* Particle0;
+		FGeometryParticleHandle* Particle1;
 		const FPerShapeData* Shape0;
 		const FPerShapeData* Shape1;
 		EContactShapesType ShapePairType;
-		bool bSwap;
+		bool bEnableOBBCheck0;
+		bool bEnableOBBCheck1;
+		bool bEnableManifoldCheck;
 	};
 
 
@@ -127,8 +118,10 @@ namespace Chaos
 	{
 	public:
 		FMultiShapePairCollisionDetector(
-			const FPerShapeData* InShape0, 
-			const FPerShapeData* InShape1, 
+			FGeometryParticleHandle* InParticle0,
+			const FPerShapeData* InShape0,
+			FGeometryParticleHandle* InParticle1,
+			const FPerShapeData* InShape1,
 			FParticlePairMidPhase& MidPhase);
 		FMultiShapePairCollisionDetector(FMultiShapePairCollisionDetector&& R);
 		FMultiShapePairCollisionDetector(const FMultiShapePairCollisionDetector& R) = delete;
@@ -140,19 +133,22 @@ namespace Chaos
 		 * @return The number of collisions constraints that were activated
 		*/
 		int32 GenerateCollisions(
-			FGeometryParticleHandle* Particle0,
-			FGeometryParticleHandle* Particle1,
 			const FReal CullDistance,
 			const bool bUseCCD,
 			const FReal Dt,
 			FCollisionContext& Context);
 
+		/**
+		 * @brief Callback from the narrow phase to create a collision constraint for this particle pair.
+		 * We should never be asked for a collision for a different particle pair, but the 
+		 * implicit objects may be children of the root shape.
+		*/
 		FPBDCollisionConstraint* FindOrCreateConstraint(
-			FGeometryParticleHandle* Particle0,
+			FGeometryParticleHandle* InParticle0,
 			const FImplicitObject* Implicit0,
 			const FBVHParticles* BVHParticles0,
 			const FRigidTransform3& ShapeRelativeTransform0,
-			FGeometryParticleHandle* Particle1,
+			FGeometryParticleHandle* InParticle1,
 			const FImplicitObject* Implicit1,
 			const FBVHParticles* BVHParticles1,
 			const FRigidTransform3& ShapeRelativeTransform1,
@@ -160,12 +156,15 @@ namespace Chaos
 			const EContactShapesType ShapePairType,
 			const bool bInUseManifold);
 
+		/**
+		 * @brief FindOrCreateConstraint for swept constraints 
+		*/
 		FPBDCollisionConstraint* FindOrCreateSweptConstraint(
-			FGeometryParticleHandle* Particle0,
+			FGeometryParticleHandle* InParticle0,
 			const FImplicitObject* Implicit0,
 			const FBVHParticles* BVHParticles0,
 			const FRigidTransform3& ShapeRelativeTransform0,
-			FGeometryParticleHandle* Particle1,
+			FGeometryParticleHandle* InParticle1,
 			const FImplicitObject* Implicit1,
 			const FBVHParticles* BVHParticles1,
 			const FRigidTransform3& ShapeRelativeTransform1,
@@ -176,7 +175,7 @@ namespace Chaos
 		 * @brief Reactivate the collision exactly as it was last frame
 		 * @return The number of collisions constraints that were restored
 		*/
-		int32 RestoreCollisions();
+		int32 RestoreCollisions(const FReal CullDistance);
 
 		/**
 		 * @brief Reactivate the constraint (essentially the same as Restore but slightly optimized)
@@ -210,6 +209,8 @@ namespace Chaos
 		FParticlePairMidPhase& MidPhase;
 		TMap<uint32, TUniquePtr<FPBDCollisionConstraint>> Constraints;
 		TArray<FPBDCollisionConstraint*> NewConstraints;
+		FGeometryParticleHandle* Particle0;
+		FGeometryParticleHandle* Particle1;
 		const FPerShapeData* Shape0;
 		const FPerShapeData* Shape1;
 	};
@@ -359,7 +360,7 @@ namespace Chaos
 		/**
 		 * @brief If the particles have not moved muc, reactivate all the colisions and skip the narrow phase.
 		*/
-		bool TryRestoreConstraints(const FReal Dt);
+		bool TryRestoreConstraints(const FReal Dt, const FReal CullDistance);
 
 
 		FGeometryParticleHandle* Particle0;
diff --git a/Engine/Source/Runtime/Experimental/Chaos/Public/Chaos/Collision/SolverCollisionContainer.h b/Engine/Source/Runtime/Experimental/Chaos/Public/Chaos/Collision/SolverCollisionContainer.h
--- a/Engine/Source/Runtime/Experimental/Chaos/Public/Chaos/Collision/SolverCollisionContainer.h
+++ b/Engine/Source/Runtime/Experimental/Chaos/Public/Chaos/Collision/SolverCollisionContainer.h
@@ -35,7 +35,7 @@ namespace Chaos
 		virtual void Reset(const int32 InMaxCollisions) override;
 
 		// Add a solver constraint for the specified constraint and gather the required solver data
-		FPBDCollisionSolver* AddConstraintSolver(FReal Dt, FPBDCollisionConstraint& Constraint, const int32 Particle0Level, const int32 Particle1Level, FSolverBodyContainer& SolverBodyContainer);
+		void AddConstraintSolver(FReal Dt, FPBDCollisionConstraint& Constraint, const int32 Particle0Level, const int32 Particle1Level, FSolverBodyContainer& SolverBodyContainer);
 
 		void SolveSwept(const FReal Dt);
 
diff --git a/Engine/Source/Runtime/Experimental/Chaos/Public/Chaos/CollisionResolution.h b/Engine/Source/Runtime/Experimental/Chaos/Public/Chaos/CollisionResolution.h
--- a/Engine/Source/Runtime/Experimental/Chaos/Public/Chaos/CollisionResolution.h
+++ b/Engine/Source/Runtime/Experimental/Chaos/Public/Chaos/CollisionResolution.h
@@ -53,6 +53,7 @@ namespace Chaos
 		 * @brief Determine the shape pair type for use in UpdateConstraints
 		*/
 		EContactShapesType CHAOS_API CalculateShapePairType(const FImplicitObject* Implicit0, const FBVHParticles* BVHParticles0, const FImplicitObject* Implicit1, const FBVHParticles* BVHParticles1, bool& bOutSwap);
+		EContactShapesType CHAOS_API CalculateShapePairType(const EImplicitObjectType Implicit0Type, const EImplicitObjectType Implicit1Type, const bool bIsConvex0, const bool bIsConvex1, const bool bIsBVH0, const bool bIsBVH1, bool& bOutSwap);
 
 		/**
 		 * @brief Whether CCD should be enabled for a contact given the current particle velocities etc
diff --git a/Engine/Source/Runtime/Experimental/Chaos/Public/Chaos/Framework/UncheckedArray.h b/Engine/Source/Runtime/Experimental/Chaos/Public/Chaos/Framework/UncheckedArray.h
new file mode 100644
--- /dev/null
+++ b/Engine/Source/Runtime/Experimental/Chaos/Public/Chaos/Framework/UncheckedArray.h
@@ -0,0 +1,62 @@
+// Copyright Epic Games, Inc. All Rights Reserved.
+#pragma once
+
+#include "CoreMinimal.h"
+#include "Containers/Array.h"
+#include "Containers/ContainerAllocationPolicies.h"
+
+// If CHAOS_CHECK_UNCHECKED_ARRAY is 1, we still enable range checks on "unchecked" arrays.
+// This is for debug builds or sanity checks in other builds.
+#ifndef CHAOS_CHECK_UNCHECKED_ARRAY
+#if UE_BUILD_DEBUG
+#define CHAOS_CHECK_UNCHECKED_ARRAY 1
+#else
+#define CHAOS_CHECK_UNCHECKED_ARRAY 0
+#endif
+#endif
+
+namespace Chaos
+{
+	/**
+	 * @brief A fixed allocator without array bounds checking except in Debug builds.
+	 *
+	 * In non-debug builds this offers no saftey at all - it is effectively a C-style array.
+	 *
+	 * This is for use in critical path code where bounds checking would be costly and we want
+	 * to ship a build with most asserts enabled (e.g. the server)
+	*/
+	template<int32 NumInlineElements>
+	class TUncheckedFixedAllocator : public TFixedAllocator<NumInlineElements>
+	{
+	public:
+#if !CHAOS_CHECK_UNCHECKED_ARRAY
+		enum { RequireRangeCheck = false };
+#endif
+	};
+
+	class FUncheckedHeapAllocator : public FHeapAllocator
+	{
+	public:
+#if !CHAOS_CHECK_UNCHECKED_ARRAY
+		enum { RequireRangeCheck = false };
+#endif
+	};
+
+	template<typename T, int32 N>
+	using TUncheckedFixedArray = TArray<T, TUncheckedFixedAllocator<N>>;
+
+	template<typename T>
+	using TUncheckedArray = TArray<T, FUncheckedHeapAllocator>;
+}
+
+template<int32 NumInlineElements>
+struct TAllocatorTraits<Chaos::TUncheckedFixedAllocator<NumInlineElements>> : TAllocatorTraits<TFixedAllocator<NumInlineElements>>
+{
+};
+
+template<>
+struct TAllocatorTraits<Chaos::FUncheckedHeapAllocator> : TAllocatorTraits<FHeapAllocator>
+{
+};
+
+
diff --git a/Engine/Source/Runtime/Experimental/Chaos/Public/Chaos/GJK.h b/Engine/Source/Runtime/Experimental/Chaos/Public/Chaos/GJK.h
--- a/Engine/Source/Runtime/Experimental/Chaos/Public/Chaos/GJK.h
+++ b/Engine/Source/Runtime/Experimental/Chaos/Public/Chaos/GJK.h
@@ -201,8 +201,8 @@ namespace Chaos
 	 * @todo(chaos): convert GJKPenetration() to use this version (but see note above)
 	 *
 	*/
-	template <bool bFindSeparationDistance = false, typename T, typename TGeometryA, typename TGeometryB>
-	bool GJKPenetrationWarmStartable(const TGeometryA& A, const TGeometryB& B, const TRigidTransform<T, 3>& BToATM, const T InThicknessA, const T InThicknessB, T& OutPenetration, TVec3<T>& OutClosestA, TVec3<T>& OutClosestB, TVec3<T>& OutNormalA, TVec3<T>& OutNormalB, int32& OutVertexA, int32& OutVertexB, TGJKSimplexData<T>& InOutSimplexData, T& OutMaxSupportDelta, const T Epsilon = T(1.e-3))
+	template <typename T, typename TGeometryA, typename TGeometryB>
+	bool GJKPenetrationWarmStartable(const TGeometryA& A, const TGeometryB& B, const TRigidTransform<T, 3>& BToATM, T& OutPenetration, TVec3<T>& OutClosestA, TVec3<T>& OutClosestB, TVec3<T>& OutNormalA, TVec3<T>& OutNormalB, int32& OutVertexA, int32& OutVertexB, TGJKSimplexData<T>& InOutSimplexData, T& OutMaxSupportDelta, const T Epsilon = T(1.e-3))
 	{
 		T SupportDeltaA = 0;
 		T SupportDeltaB = 0;
@@ -239,8 +239,8 @@ namespace Chaos
 		bool bIsDegenerate = false;				// True if GJK cannot make any more progress
 		bool bIsContact = false;				// True if shapes are within Epsilon or overlapping - GJK cannot provide a solution
 		int32 NumIterations = 0;
-		const T ThicknessA = InThicknessA + A.GetMargin();
-		const T ThicknessB = InThicknessB + B.GetMargin();
+		const T ThicknessA = A.GetMargin();
+		const T ThicknessB = B.GetMargin();
 		const T SeparatedDistance = ThicknessA + ThicknessB + Epsilon;
 		while (!bIsContact && !bIsDegenerate)
 		{
@@ -260,11 +260,6 @@ namespace Chaos
 			const T ConvergenceTolerance = 1.e-4f;
 			const T ConvergedDistance = (1.0f - ConvergenceTolerance) * Distance;
 			const T VW = TVec3<T>::DotProduct(V, W);
-			if (!bFindSeparationDistance && (VW > SeparatedDistance))
-			{
-				// We are separated and don't care about the separation distance
-				break;
-			}
 
 			InOutSimplexData.As[SimplexIDs.NumVerts] = SupportA;
 			InOutSimplexData.Bs[SimplexIDs.NumVerts] = SupportB;
@@ -370,10 +365,8 @@ namespace Chaos
 
 			OutMaxSupportDelta = MaxSupportDelta;
 
-			// If we don't care about separation distance/normal, the return value is true if we are overlapping, false otherwise.
-			// If we do care about seperation distance/normal, the return value is true if we found a solution.
 			// @todo(chaos): we should pass back failure for the degenerate case so we can decide how to handle it externally.
-			return (bFindSeparationDistance || (Penetration >= 0.0f));
+			return true;
 		}
 	}
 
diff --git a/Engine/Source/Runtime/Experimental/Chaos/Public/Chaos/Transform.h b/Engine/Source/Runtime/Experimental/Chaos/Public/Chaos/Transform.h
--- a/Engine/Source/Runtime/Experimental/Chaos/Public/Chaos/Transform.h
+++ b/Engine/Source/Runtime/Experimental/Chaos/Public/Chaos/Transform.h
@@ -176,9 +176,10 @@ namespace Chaos
 		TRigidTransform<FReal, 3> GetRelativeTransformNoScale(const TRigidTransform<FReal, 3>& Other) const
 		{
 			// @todo(chaos): optimize
-			TRigidTransform<FReal, 3> ThisNoScale(GetTranslation(), GetRotation());
-			TRigidTransform<FReal, 3> OtherNoScale(Other.GetTranslation(), Other.GetRotation());
-			return ThisNoScale.GetRelativeTransform(OtherNoScale);
+			TRotation<FReal, 3> OtherInverse = Other.GetRotation().Inverse();
+			return TRigidTransform<FReal, 3>(
+				(OtherInverse * (GetTranslation() - Other.GetTranslation())),
+				OtherInverse * GetRotation());
 		}
 
 		TVector<FReal, 3> TransformNormalNoScale(const TVector<FReal, 3>& Normal) const
