Chaos - fix typo in comment

diff --git a/Engine/Source/Runtime/Experimental/Chaos/Private/Chaos/Collision/PBDCollisionConstraint.cpp b/Engine/Source/Runtime/Experimental/Chaos/Private/Chaos/Collision/PBDCollisionConstraint.cpp
--- a/Engine/Source/Runtime/Experimental/Chaos/Private/Chaos/Collision/PBDCollisionConstraint.cpp
+++ b/Engine/Source/Runtime/Experimental/Chaos/Private/Chaos/Collision/PBDCollisionConstraint.cpp
@@ -737,7 +737,7 @@ namespace Chaos
 		}
 
 		// Find the matching manifold point if it exists and replace it
-		// Also check to see if the normal has changed significantly, and if it has, force mnanifold regeneration
+		// Also check to see if the normal has changed significantly and if it has force manifold regeneration
 		// NOTE: the normal rejection check assumes all contacts have the same normal - this may not always be true. The worst
 		// case here is that we will regenerate the manifold too often so it will work but could be bad for perf
 		const FReal PositionToleranceSq = FMath::Square(PositionTolerance);
