Optimizing AABB transform

diff --git a/Engine/Source/Programs/HeadlessChaos/Private/HeadlessChaosTestImplicits.cpp b/Engine/Source/Programs/HeadlessChaos/Private/HeadlessChaosTestImplicits.cpp
--- a/Engine/Source/Programs/HeadlessChaos/Private/HeadlessChaosTestImplicits.cpp
+++ b/Engine/Source/Programs/HeadlessChaos/Private/HeadlessChaosTestImplicits.cpp
@@ -1020,7 +1020,7 @@ namespace ChaosTest {
 		FRigidTransform3 Identity(FVec3(0), FQuat::Identity);
 		
 		TUniquePtr<TBox<FReal, 3>> UnitCube = MakeUnique<TBox<FReal, 3>>(FVec3(-1), FVec3(1));
-		TImplicitObjectTransformed<FReal, 3> UnitUnrotated(MakeSerializable(UnitCube), FRigidTransform3(FVec3(0), FQuat(0,0,0, 0)));
+		TImplicitObjectTransformed<FReal, 3> UnitUnrotated(MakeSerializable(UnitCube), FRigidTransform3(FVec3(0), FQuat::Identity));
 		UnitImplicitObjectNormalsInternal(UnitUnrotated, FString("ImplicitTransformed()"));
 		UnitImplicitObjectNormalsExternal(UnitUnrotated, FString("ImplicitTransformed()"));
 		UnitImplicitObjectIntersections(UnitUnrotated, FString("ImplicitTransformed()"));
diff --git a/Engine/Source/Runtime/Experimental/Chaos/Private/Chaos/AABB.cpp b/Engine/Source/Runtime/Experimental/Chaos/Private/Chaos/AABB.cpp
--- a/Engine/Source/Runtime/Experimental/Chaos/Private/Chaos/AABB.cpp
+++ b/Engine/Source/Runtime/Experimental/Chaos/Private/Chaos/AABB.cpp
@@ -255,22 +255,35 @@ bool TAABB<T, d>::Raycast(const TVector<FReal, d>& StartPoint, const TVector<FRe
 	return true;
 }
 
-template<typename T, int d, class TTRANSFORM>
-inline TAABB<T, d> TransformedAABBHelper(const TAABB<T, d>& AABB, const TTRANSFORM& SpaceTransform)
+template<typename T>
+inline TAABB<T, 3> TransformedAABBHelper(const TAABB<T, 3>& AABB, const FMatrix44& SpaceTransform)
 {
-	TVector<T, d> CurrentExtents = AABB.Extents();
-	int32 Idx = 0;
-	const TVector<T, d> MinToNewSpace = SpaceTransform.TransformPosition(AABB.Min());
-	TAABB<T, d> NewAABB(MinToNewSpace, MinToNewSpace);
-	NewAABB.GrowToInclude(SpaceTransform.TransformPosition(AABB.Max()));
+	// Initialize to center
+	FVec3 Translation(SpaceTransform.M[3][0], SpaceTransform.M[3][1], SpaceTransform.M[3][2]);
+	FVec3 Min = Translation;
+	FVec3 Max = Translation;
 
-	for (int32 j = 0; j < d; ++j)
+	// Compute extents per axis
+	for (int32 i = 0; i < 3; ++i)
 	{
-		NewAABB.GrowToInclude(SpaceTransform.TransformPosition(AABB.Min() + TVector<T, d>::AxisVector(j) * CurrentExtents));
-		NewAABB.GrowToInclude(SpaceTransform.TransformPosition(AABB.Max() - TVector<T, d>::AxisVector(j) * CurrentExtents));
+		for (int32 j = 0; j < 3; ++j)	
+		{
+			FReal A = SpaceTransform.M[j][i] * AABB.Min()[j];
+			FReal B = SpaceTransform.M[j][i] * AABB.Max()[j];
+			if (A < B)
+			{
+				Min[i] += A;
+				Max[i] += B;
+			}
+			else 
+			{
+				Min[i] += B;
+				Max[i] += A;
+			}
+		}
 	}
 
-	return NewAABB;
+	return TAABB<T, 3>(Min, Max);
 }
 
 template<typename T>
@@ -298,21 +311,21 @@ TAABB<T, d> TAABB<T, d>::TransformedAABB(const Chaos::TRigidTransform<FReal, 3>&
 	}
 	else
 	{
-		return TransformedAABBHelper<T, d>(*this, SpaceTransform);
+		return TransformedAABBHelper<T>(*this, SpaceTransform.ToMatrixWithScale());
 	}
 }
 
 template<typename T, int d>
 TAABB<T, d> TAABB<T, d>::TransformedAABB(const FMatrix& SpaceTransform) const
 {
-	return TransformedAABBHelper<T, d>(*this, SpaceTransform);
+	return TransformedAABBHelper<T>(*this, SpaceTransform);
 }
 
 
 template<typename T, int d>
 TAABB<T, d> TAABB<T, d>::TransformedAABB(const Chaos::PMatrix<FReal, 4, 4>& SpaceTransform) const
 {
-	return TransformedAABBHelper<T, d>(*this, SpaceTransform);
+	return TransformedAABBHelper<T>(*this, SpaceTransform);
 }
 
 template<typename T, int d>
@@ -324,7 +337,7 @@ TAABB<T, d> TAABB<T, d>::TransformedAABB(const FTransform& SpaceTransform) const
 	}
 	else
 	{
-		return TransformedAABBHelper<T, d>(*this, SpaceTransform);
+		return TransformedAABBHelper<T>(*this, SpaceTransform.ToMatrixWithScale());
 	}
 }
 
diff --git a/Engine/Source/Runtime/Experimental/Chaos/Public/Chaos/Transform.h b/Engine/Source/Runtime/Experimental/Chaos/Public/Chaos/Transform.h
--- a/Engine/Source/Runtime/Experimental/Chaos/Public/Chaos/Transform.h
+++ b/Engine/Source/Runtime/Experimental/Chaos/Public/Chaos/Transform.h
@@ -160,6 +160,11 @@ namespace Chaos
 			return BaseTransform::Inverse();
 		}
 
+		PMatrix<FReal, 4, 4> ToMatrixWithScale() const
+		{
+			return BaseTransform::ToMatrixWithScale();
+		}
+
 		PMatrix<FReal, 4, 4> ToMatrixNoScale() const
 		{
 			return BaseTransform::ToMatrixNoScale();
