Use of the cached planes in Chaos NP only if the number of planes/vertex < 3

diff --git a/Engine/Source/Runtime/Experimental/Chaos/Private/Chaos/CollisionOneShotManifolds.cpp b/Engine/Source/Runtime/Experimental/Chaos/Private/Chaos/CollisionOneShotManifolds.cpp
--- a/Engine/Source/Runtime/Experimental/Chaos/Private/Chaos/CollisionOneShotManifolds.cpp
+++ b/Engine/Source/Runtime/Experimental/Chaos/Private/Chaos/CollisionOneShotManifolds.cpp
@@ -648,6 +648,37 @@ namespace Chaos
 			return Contact;
 		}
 
+		// Find the the most opposing plane given a position and a direction
+		template <typename ConvexImplicitType>
+		void FindBestPlane(
+			const ConvexImplicitType& Convex,
+			const FVec3& X,
+			const FVec3& N,
+			const FReal MaxDistance,
+			const int32 PlaneIndex,
+			int32& BestPlaneIndex,
+			FReal& BestPlaneDot)
+		{
+			const TPlaneConcrete<FReal, 3> Plane = Convex.GetPlane(PlaneIndex);
+			const FReal PlaneNormalDotN = FVec3::DotProduct(N, Plane.Normal());
+				
+			// Ignore planes that do not oppose N
+			if (PlaneNormalDotN <= -SMALL_NUMBER)
+			{
+				// Reject planes farther than MaxDistance
+				const FReal PlaneDistance = Plane.SignedDistance(X);
+				if (FMath::Abs(PlaneDistance) <= MaxDistance)
+				{
+					// Keep the most opposing plane
+					if (PlaneNormalDotN < BestPlaneDot)
+					{
+						BestPlaneDot = PlaneNormalDotN;
+						BestPlaneIndex = PlaneIndex;
+					}
+				}
+			}
+		}
+
 		// Select one of the planes on the convex to use as the contact plane, given an estimated contact position and opposing 
 		// normal from GJK with margins (which gives the shapes rounded corners/edges).
 		template <typename ConvexImplicitType>
@@ -665,34 +696,25 @@ namespace Chaos
 
 			int32 BestPlaneIndex = INDEX_NONE;
 			FReal BestPlaneDot = 1.0f;
-			
 			{
 				int32 PlaneIndices[3] = {INDEX_NONE, INDEX_NONE, INDEX_NONE};
-				const int32 NumPlanes = Convex.GetVertexPlanes3(VertexIndex, PlaneIndices[0], PlaneIndices[1], PlaneIndices[2]);
-			
-				for (int32 PlaneIndex = 0; PlaneIndex < NumPlanes; ++PlaneIndex)
+				int32 NumPlanes = Convex.GetVertexPlanes3(VertexIndex, PlaneIndices[0], PlaneIndices[1], PlaneIndices[2]);
+
+				// If we have more than 3 planes we iterate over the full set of planes since it is faster than using the half edge structure
+				if(NumPlanes > 3)
 				{
-					const TPlaneConcrete<FReal, 3> Plane = Convex.GetPlane(PlaneIndices[PlaneIndex]);
-					const FReal PlaneNormalDotN = FVec3::DotProduct(N, Plane.Normal());
-			
-					// Ignore planes that do not oppose N
-					if (PlaneNormalDotN > -SMALL_NUMBER)
-					{
-						continue;
-					}
-			
-					// Reject planes farther than MaxDistance
-					const FReal PlaneDistance = Plane.SignedDistance(X);
-					if (FMath::Abs(PlaneDistance) > MaxDistance)
+					NumPlanes = Convex.NumPlanes();
+					for (int32 PlaneIndex = 0; PlaneIndex < NumPlanes; ++PlaneIndex)
 					{
-						continue;
+						FindBestPlane(Convex, X, N, MaxDistance, PlaneIndex, BestPlaneIndex, BestPlaneDot);
 					}
-			
-					// Keep the most opposing plane
-					if (PlaneNormalDotN < BestPlaneDot)
+				}
+				// Otherwise we iterate over the cached planes
+				else
+				{
+					for (int32 PlaneIndex = 0; PlaneIndex < NumPlanes; ++PlaneIndex)
 					{
-						BestPlaneDot = PlaneNormalDotN;
-						BestPlaneIndex = PlaneIndices[PlaneIndex];
+						FindBestPlane(Convex, X, N, MaxDistance, PlaneIndices[PlaneIndex], BestPlaneIndex, BestPlaneDot);
 					}
 				}
 			}
diff --git a/Engine/Source/Runtime/Experimental/Chaos/Public/Chaos/ConvexHalfEdgeStructureData.h b/Engine/Source/Runtime/Experimental/Chaos/Public/Chaos/ConvexHalfEdgeStructureData.h
--- a/Engine/Source/Runtime/Experimental/Chaos/Public/Chaos/ConvexHalfEdgeStructureData.h
+++ b/Engine/Source/Runtime/Experimental/Chaos/Public/Chaos/ConvexHalfEdgeStructureData.h
@@ -388,12 +388,18 @@ namespace Chaos
 			}
 
 			// Find the twin half edge for each edge
+			// NOTE: we have to deal with mal-formed convexes which claim to have edges that use the
+			// same vertex pair in the same order.
+			// @todo(chaos): track down to source of the mal-formed convexes
 			// @todo(chaos): could use a map of vertex-index-pair to half edge to eliminate O(N^2) algorithm
+			TArray<bool> HalfEdgeTwinned;
 			TArray<FIndex> TwinHalfEdgeIndices;
 			TwinHalfEdgeIndices.SetNum(HalfEdges.Num());
+			HalfEdgeTwinned.SetNum(HalfEdges.Num());
 			for (int32 HalfEdgeIndex = 0; HalfEdgeIndex < TwinHalfEdgeIndices.Num(); ++HalfEdgeIndex)
 			{
 				TwinHalfEdgeIndices[HalfEdgeIndex] = InvalidIndex;
+				HalfEdgeTwinned[HalfEdgeIndex] = false;
 			}
 			for (int32 HalfEdgeIndex0 = 0; HalfEdgeIndex0 < HalfEdges.Num(); ++HalfEdgeIndex0)
 			{
@@ -405,7 +411,16 @@ namespace Chaos
 				{
 					if ((HalfEdges[HalfEdgeIndex1].VertexIndex == VertexIndex1) && (HalfEdges[HalfEdgeIndex1].TwinHalfEdgeIndex == VertexIndex0))
 					{
-						TwinHalfEdgeIndices[HalfEdgeIndex0] = (FIndex)HalfEdgeIndex1;
+						// We deal with edge duplication by leaving a half edge without a twin
+						if (!HalfEdgeTwinned[HalfEdgeIndex1])
+						{
+							TwinHalfEdgeIndices[HalfEdgeIndex0] = (FIndex)HalfEdgeIndex1;
+							HalfEdgeTwinned[HalfEdgeIndex1] = true;
+						}
+						else
+						{
+							TwinHalfEdgeIndices[HalfEdgeIndex0] = InvalidIndex;
+						}
 						break;
 					}
 				}
@@ -524,9 +539,21 @@ namespace Chaos
 				{
 					do
 					{
-						VertexPlane.PlaneIndices[VertexPlane.NumPlaneIndices++] = (FIndex)GetHalfEdgePlane(HalfEdgeIndex);
-						if (VertexPlane.NumPlaneIndices == (FIndex)UE_ARRAY_COUNT(VertexPlane.PlaneIndices))
+						if(VertexPlane.NumPlaneIndices < (FIndex)UE_ARRAY_COUNT(VertexPlane.PlaneIndices))
+						{
+							VertexPlane.PlaneIndices[VertexPlane.NumPlaneIndices] = (FIndex)GetHalfEdgePlane(HalfEdgeIndex);
+						}
+
+						// Caching of the Max number of plane indices on this vertex (could be higher than 3)
+						// This can be used to determine if a call to GetVertexPlanes3 will actually return all the planes
+						// that use a particular vertex. This is useful in collision detection for box-like objects
+						// to avoid calling FindVertexPlanes
+						++VertexPlane.NumPlaneIndices;
+
+						// If we hit this there's a mal-formed convex case that we did not detect (we should be dealing with all cases - see SetPlaneVertices)
+						if (!ensure(VertexPlane.NumPlaneIndices <= Planes.Num()))
 						{
+							VertexPlane.NumPlaneIndices = 0;
 							break;
 						}
 
diff --git a/Engine/Source/Runtime/Experimental/Chaos/Public/Chaos/GJK.h b/Engine/Source/Runtime/Experimental/Chaos/Public/Chaos/GJK.h
--- a/Engine/Source/Runtime/Experimental/Chaos/Public/Chaos/GJK.h
+++ b/Engine/Source/Runtime/Experimental/Chaos/Public/Chaos/GJK.h
@@ -303,11 +303,10 @@ namespace Chaos
 				VertsB.Add(BToATM.TransformPositionNoScale(InOutSimplexData.Bs[i]));
 			}
 			
-			auto SupportBInAFunc = [&B, &BToATM, &AToBRotation, &SupportDeltaB](const TVec3<T>& V)
+			auto SupportBInAFunc = [&B, &BToATM, &AToBRotation, &SupportDeltaB, &VertexIndexB](const TVec3<T>& V)
 			{
 				const TVec3<T> VInB = AToBRotation * V;
-				int32 VertexIndex = INDEX_NONE;
-				const TVec3<T> SupportBLocal = B.SupportCore(VInB, B.GetMargin(), &SupportDeltaB, VertexIndex);
+				const TVec3<T> SupportBLocal = B.SupportCore(VInB, B.GetMargin(), &SupportDeltaB, VertexIndexB);
 				return BToATM.TransformPositionNoScale(SupportBLocal);
 			};
 
diff --git a/Engine/Source/Runtime/Experimental/ChaosCore/Public/Chaos/Core.h b/Engine/Source/Runtime/Experimental/ChaosCore/Public/Chaos/Core.h
--- a/Engine/Source/Runtime/Experimental/ChaosCore/Public/Chaos/Core.h
+++ b/Engine/Source/Runtime/Experimental/ChaosCore/Public/Chaos/Core.h
@@ -35,5 +35,5 @@ namespace Chaos
 	// NOTE: if you get a merge conflict on the GUID, you must replace it with a new GUID - do not accept the source or target
 	// or you will likely get DDC version conflicts resulting in crashes during load.
 	// Core version string for Chaos data. Any DDC builder dependent on Chaos for serialization should depend on this version
-	static const FString ChaosVersionString = TEXT("6955386D-0525-4192-8927-74B0E57E8926");
+	static const FString ChaosVersionString = TEXT("511D498A-68E5-45F3-A89E-988D57A23C9B");
 }
