Fix static analysis warning on particle access for collision event generation

diff --git a/Engine/Source/Runtime/Experimental/Chaos/Private/EventDefaults.cpp b/Engine/Source/Runtime/Experimental/Chaos/Private/EventDefaults.cpp
--- a/Engine/Source/Runtime/Experimental/Chaos/Private/EventDefaults.cpp
+++ b/Engine/Source/Runtime/Experimental/Chaos/Private/EventDefaults.cpp
@@ -173,9 +173,6 @@ namespace Chaos
 						FGeometryParticleHandle* Particle0 = Constraint.Particle[0];
 						FGeometryParticleHandle* Particle1 = Constraint.Particle[1];
 
-						const FPerShapeData* Shape0 = Particle0->GetImplicitShape(Constraint.GetImplicit0());
-						const FPerShapeData* Shape1 = Particle1->GetImplicitShape(Constraint.GetImplicit1());
-
 						FCollidingData Data;
 						Data.Location = Constraint.GetLocation();
 						Data.AccumulatedImpulse = Constraint.AccumulatedImpulse;
@@ -185,6 +182,9 @@ namespace Chaos
 						Data.Proxy1 = Particle0 ? Particle0->PhysicsProxy() : nullptr;
 						Data.Proxy2 = Particle1 ? Particle1->PhysicsProxy() : nullptr;
 
+						const FPerShapeData* Shape0 = Particle0 ? Particle0->GetImplicitShape(Constraint.GetImplicit0()) : nullptr;
+						const FPerShapeData* Shape1 = Particle1 ? Particle1->GetImplicitShape(Constraint.GetImplicit1()) : nullptr;
+
 						Data.Mat1 = ResolveMaterial(Shape0, Constraint);
 						Data.Mat2 = ResolveMaterial(Shape1, Constraint);
 
