Chaos Cloth - Intialize constraint properties values for easier debugging.

diff --git a/Engine/Source/Runtime/Experimental/Chaos/Private/Chaos/PBDLongRangeConstraintsBase.cpp b/Engine/Source/Runtime/Experimental/Chaos/Private/Chaos/PBDLongRangeConstraintsBase.cpp
--- a/Engine/Source/Runtime/Experimental/Chaos/Private/Chaos/PBDLongRangeConstraintsBase.cpp
+++ b/Engine/Source/Runtime/Experimental/Chaos/Private/Chaos/PBDLongRangeConstraintsBase.cpp
@@ -60,6 +60,9 @@ FPBDLongRangeConstraintsBase::FPBDLongRangeConstraintsBase(
 		ScaleIndices.AddZeroed(1);
 		ScaleTable.AddZeroed(1);
 	}
+
+	// Apply default properties
+	ApplyProperties((FReal)(1. / FPBDStiffness::ParameterFrequency), 1);
 }
 
 void FPBDLongRangeConstraintsBase::ApplyProperties(const FReal Dt, const int32 NumIterations)
