Fracture tool: restore Hide/Unhide button functionality. In UGeometryCollectionComponent, if a Hide array is not available in the GeometryCollection, use the Visible array to determine visibility.

#rnx

diff --git a/Engine/Source/Runtime/Experimental/GeometryCollectionEngine/Private/GeometryCollection/GeometryCollectionComponent.cpp b/Engine/Source/Runtime/Experimental/GeometryCollectionEngine/Private/GeometryCollection/GeometryCollectionComponent.cpp
--- a/Engine/Source/Runtime/Experimental/GeometryCollectionEngine/Private/GeometryCollection/GeometryCollectionComponent.cpp
+++ b/Engine/Source/Runtime/Experimental/GeometryCollectionEngine/Private/GeometryCollection/GeometryCollectionComponent.cpp
@@ -1519,7 +1519,7 @@ void UGeometryCollectionComponent::InitConstantData(FGeometryCollectionConstantD
 		for (int FaceIndex = 0; FaceIndex < NumFaceGroupEntries; ++FaceIndex)
 		{
 #if WITH_EDITOR
-			NumIndices += static_cast<int>(VisibleOverride[FaceIndex]);
+			NumIndices += bUsingHideArray ? static_cast<int>(VisibleOverride[FaceIndex]) : static_cast<int>(Visible[FaceIndex]);
 #else
 			NumIndices += static_cast<int>(Visible[FaceIndex]);
 #endif
@@ -1529,7 +1529,8 @@ void UGeometryCollectionComponent::InitConstantData(FGeometryCollectionConstantD
 		for (int IndexIdx = 0, cdx = 0; IndexIdx < NumFaceGroupEntries; ++IndexIdx)
 		{
 #if WITH_EDITOR
-			if (VisibleOverride[MaterialIndex[IndexIdx]])
+			const bool bUseVisible = bUsingHideArray ? VisibleOverride[MaterialIndex[IndexIdx]] : Visible[MaterialIndex[IndexIdx]];
+			if (bUseVisible)
 #else
 			if (Visible[MaterialIndex[IndexIdx]])
 #endif
@@ -1549,7 +1550,8 @@ void UGeometryCollectionComponent::InitConstantData(FGeometryCollectionConstantD
 			for (int32 TriangleIndex = 0; TriangleIndex < Sections[SectionIndex].FirstIndex / 3; TriangleIndex++)
 			{
 #if WITH_EDITOR
-				if (!VisibleOverride[MaterialIndex[TriangleIndex]])
+				const bool bUseVisible = bUsingHideArray ? VisibleOverride[MaterialIndex[TriangleIndex]] : Visible[MaterialIndex[TriangleIndex]];
+				if (!bUseVisible)
 #else
 				if (!Visible[MaterialIndex[TriangleIndex]])
 #endif
@@ -1562,7 +1564,8 @@ void UGeometryCollectionComponent::InitConstantData(FGeometryCollectionConstantD
 			{
 				int32 FaceIdx = MaterialIndex[Sections[SectionIndex].FirstIndex / 3 + TriangleIndex];
 #if WITH_EDITOR
-				if (!VisibleOverride[FaceIdx])
+				const bool bUseVisible = bUsingHideArray ? VisibleOverride[FaceIdx] : Visible[FaceIdx];
+				if (!bUseVisible)
 #else
 				if (!Visible[FaceIdx])
 #endif
@@ -1594,7 +1597,8 @@ void UGeometryCollectionComponent::InitConstantData(FGeometryCollectionConstantD
 		{
 			// only add visible external faces.  MaterialID that is even is an external material
 #if WITH_EDITOR
-			if (VisibleOverride[FaceIndex] && (MaterialID[FaceIndex] % 2 == 0 || bUsingHideArray))
+			const bool bUseVisible = bUsingHideArray ? VisibleOverride[FaceIndex] : Visible[FaceIndex];
+			if (bUseVisible && (MaterialID[FaceIndex] % 2 == 0 || bUsingHideArray))
 #else
 			if (Visible[FaceIndex] && MaterialID[FaceIndex] % 2 == 0)
 #endif
