Fix bad edge island index

diff --git a/Engine/Source/Runtime/Experimental/Chaos/Private/Chaos/Island/IslandGraph.cpp b/Engine/Source/Runtime/Experimental/Chaos/Private/Chaos/Island/IslandGraph.cpp
--- a/Engine/Source/Runtime/Experimental/Chaos/Private/Chaos/Island/IslandGraph.cpp
+++ b/Engine/Source/Runtime/Experimental/Chaos/Private/Chaos/Island/IslandGraph.cpp
@@ -64,7 +64,32 @@ void FIslandGraph<NodeType, EdgeType, IslandType>::UpdateNode(const NodeType& No
 				MasterIsland = GraphEdge.IslandIndex;
 				GraphIslands[GraphEdge.IslandIndex].bIsPersistent = false;
 			}
-			
+		}
+		// If node is not valid and was valid before we check that all the node edges
+		// have at least one valid node. otherwise we remove the edge. 
+		if (!bValidNode && GraphNode.bValidNode)
+		{
+			for (int32 NodeEdge = GraphNode.NodeEdges.GetMaxIndex() - 1; NodeEdge >= 0; --NodeEdge)
+			{
+				if (GraphNode.NodeEdges.IsValidIndex(NodeEdge))
+				{
+					const int32 EdgeIndex = GraphNode.NodeEdges[NodeEdge];
+
+					if (GraphEdges.IsValidIndex(EdgeIndex))
+					{
+						FGraphEdge& GraphEdge = GraphEdges[EdgeIndex];
+
+						// Get the opposite node index for the given edge
+						const int32 OtherIndex = (NodeIndex == GraphEdge.FirstNode) ?
+							GraphEdge.SecondNode : GraphEdge.FirstNode;
+
+						if (!GraphNodes.IsValidIndex(OtherIndex) || (GraphNodes.IsValidIndex(OtherIndex) && !GraphNodes[OtherIndex].bValidNode))
+						{
+							RemoveEdge(EdgeIndex);
+						}
+					}
+				}
+			}
 		}
 		GraphNode.bValidNode = bValidNode;
 		GraphNode.bStationaryNode = bStationaryNode;
@@ -687,7 +712,7 @@ void FIslandGraph<NodeType, EdgeType, IslandType>::UpdateGraph()
 		{
 			FGraphNode& GraphNode = GraphNodes[NodeIndex];
 			// Add new islands for the all the particles that are not connected into the graph 
-			if(GraphNode.IslandIndex == INDEX_NONE && GraphNode.bValidNode)
+			if(!GraphIslands.IsValidIndex(GraphNode.IslandIndex) && GraphNode.bValidNode)
 			{
 				FGraphIsland GraphIsland = { 0, 1 };
 				GraphNode.IslandIndex = GraphIslands.Emplace(GraphIsland);
