Chaos : Heightfield Optimization: Change the returned cell bounds to be exactly the triangles bounding box
25% gain in raycast
17% in DirectSweep

diff --git a/Engine/Source/Runtime/Experimental/Chaos/Public/Chaos/HeightField.h b/Engine/Source/Runtime/Experimental/Chaos/Public/Chaos/HeightField.h
--- a/Engine/Source/Runtime/Experimental/Chaos/Public/Chaos/HeightField.h
+++ b/Engine/Source/Runtime/Experimental/Chaos/Public/Chaos/HeightField.h
@@ -280,8 +280,7 @@ namespace Chaos
 				OutPts[2] = { (FReal)X, (FReal)Y + 1, H2 };
 				OutPts[3] = { (FReal)X + 1, (FReal)Y + 1, H3 };
 
-				// use MinHeight this assume the heightfield is a halfspace
-				const FReal MinZ = GetMinHeight();
+				const FReal MinZ = FMath::Min<FReal>(H0, FMath::Min<FReal>(H1, FMath::Min<FReal>(H2, H3)));
 				const FReal MaxZ = FMath::Max<FReal>(H0, FMath::Max<FReal>(H1, FMath::Max<FReal>(H2, H3)));
 
 				OutBounds = FAABB3(FVec3((FReal)X, (FReal)Y, MinZ), FVec3((FReal)X + 1, (FReal)Y + 1, MaxZ));
