Chaos - fixes and minor optimizations
- fix excessive veloicty being generated in velocity solve
- avoid sqrt in quat normalize in constraint solver
- don't update mass properties in joint solver
- use same shock propagation mass scale in collision solver position and velocity phases

diff --git a/Engine/Source/Runtime/Experimental/Chaos/Private/Chaos/Collision/PBDCollisionSolver.cpp b/Engine/Source/Runtime/Experimental/Chaos/Private/Chaos/Collision/PBDCollisionSolver.cpp
--- a/Engine/Source/Runtime/Experimental/Chaos/Private/Chaos/Collision/PBDCollisionSolver.cpp
+++ b/Engine/Source/Runtime/Experimental/Chaos/Private/Chaos/Collision/PBDCollisionSolver.cpp
@@ -25,14 +25,14 @@ namespace Chaos
 
 		bool bChaos_PBDCollisionSolver_Position_SolveEnabled = true;
 		int32 Chaos_PBDCollisionSolver_Position_ShockPropagationIterations = 3;
-		float Chaos_PBDCollisionSolver_Position_MinInvMassScale = 0.5f;
+		float Chaos_PBDCollisionSolver_Position_MinInvMassScale = 0.3f;
 		int32 Chaos_PBDCollisionSolver_Position_ZeroFrictionIterations = 4;
 		float Chaos_PBDCollisionSolver_Position_NormalTolerance = 0.1f;
 		bool bChaos_PBDCollisionSolver_Position_NegativePushOutEnabled = true;
 		float Chaos_PBDCollisionSolver_Position_StaticFrictionStiffness = 0.5f;
 		float Chaos_PBDCollisionSolver_Position_StaticFrictionLerpRate = 0.1f;
-		float Chaos_PBDCollisionSolver_Position_PositionSolverTolerance = 0.0001f;		// cms
-		float Chaos_PBDCollisionSolver_Position_RotationSolverTolerance = 0.0001f;		// rads
+		float Chaos_PBDCollisionSolver_Position_PositionSolverTolerance = 0.001f;		// cms
+		float Chaos_PBDCollisionSolver_Position_RotationSolverTolerance = 0.001f;		// rads
 
 		FAutoConsoleVariableRef CVarChaos_PBDCollisionSolver_Position_SolveEnabled(TEXT("p.Chaos.PBDCollisionSolver.Position.SolveEnabled"), bChaos_PBDCollisionSolver_Position_SolveEnabled, TEXT(""));
 		FAutoConsoleVariableRef CVarChaos_PBDCollisionSolver_Position_UseShockPropagation(TEXT("p.Chaos.PBDCollisionSolver.Position.ShockPropagationIterations"), Chaos_PBDCollisionSolver_Position_ShockPropagationIterations, TEXT(""));
@@ -46,10 +46,11 @@ namespace Chaos
 
 		bool bChaos_PBDCollisionSolver_Velocity_SolveEnabled = true;
 		int32 Chaos_PBDCollisionSolver_Velocity_ShockPropagationIterations = 1;
-		float Chaos_PBDCollisionSolver_Velocity_MinInvMassScale = 0.1f;
+		// If this is the same as Chaos_PBDCollisionSolver_Position_MinInvMassScale and all velocity iterations have shockpropagation, we avoid recalculating constraiunt-space mass
+		float Chaos_PBDCollisionSolver_Velocity_MinInvMassScale = Chaos_PBDCollisionSolver_Position_MinInvMassScale;
 		bool bChaos_PBDCollisionSolver_Velocity_DynamicFrictionEnabled = true;
 		bool bChaos_PBDCollisionSolver_Velocity_NegativeImpulseEnabled = true;
-		bool bChaos_PBDCollisionSolver_Velocity_ImpulseClampEnabled = false;
+		bool bChaos_PBDCollisionSolver_Velocity_ImpulseClampEnabled = true;
 
 		FAutoConsoleVariableRef CVarChaos_PBDCollisionSolver_Velocity_SolveEnabled(TEXT("p.Chaos.PBDCollisionSolver.Velocity.SolveEnabled"), bChaos_PBDCollisionSolver_Velocity_SolveEnabled, TEXT(""));
 		FAutoConsoleVariableRef CVarChaos_PBDCollisionSolver_Velocity_UseShockPropagation(TEXT("p.Chaos.PBDCollisionSolver.Velocity.ShockPropagationIterations"), Chaos_PBDCollisionSolver_Velocity_ShockPropagationIterations, TEXT(""));
@@ -220,18 +221,19 @@ namespace Chaos
 				PushOut);									// Out
 		}
 
-		//UE_LOG(LogChaosCollision, VeryVerbose, TEXT("    PushOut %f"), PushOut.Size());
-
 		// Update the particle state based on the pushout
 		if (Body0.IsDynamic())
 		{
 			const FVec3 AngularPushOut = FVec3::CrossProduct(ManifoldPoint.WorldRelativeImpulsePoint0, PushOut);
 			const FVec3 DX0 = Body0.InvM() * PushOut;
 			const FVec3 DR0 = Body0.InvI() * AngularPushOut;
-			//if (!DX0.IsNearlyZero(Chaos_PBDCollisionSolver_Position_PositionSolverTolerance) || !DR0.IsNearlyZero(Chaos_PBDCollisionSolver_Position_RotationSolverTolerance))
+			//if (!DX0.IsNearlyZero(Chaos_PBDCollisionSolver_Position_PositionSolverTolerance))
+			{
+				Body0.ApplyPositionDelta(DX0);
+			}
+			//if (!DR0.IsNearlyZero(Chaos_PBDCollisionSolver_Position_RotationSolverTolerance))
 			{
-				Body0.ApplyTransformDelta(DX0, DR0);
-				//Body0.UpdateRotationDependentState();
+				Body0.ApplyRotationDelta(DR0);
 			}
 		}
 		if (Body1.IsDynamic())
@@ -239,10 +241,13 @@ namespace Chaos
 			const FVec3 AngularPushOut = FVec3::CrossProduct(ManifoldPoint.WorldRelativeImpulsePoint1, PushOut);
 			const FVec3 DX1 = -(Body1.InvM() * PushOut);
 			const FVec3 DR1 = -(Body1.InvI() * AngularPushOut);
-			//if (!DX1.IsNearlyZero(Chaos_PBDCollisionSolver_Position_PositionSolverTolerance) || !DR1.IsNearlyZero(Chaos_PBDCollisionSolver_Position_RotationSolverTolerance))
+			//if (!DX1.IsNearlyZero(Chaos_PBDCollisionSolver_Position_PositionSolverTolerance))
 			{
-				Body1.ApplyTransformDelta(DX1, DR1);
-				//Body1.UpdateRotationDependentState();
+				Body1.ApplyPositionDelta(DX1);
+			}
+			//if (!DR1.IsNearlyZero(Chaos_PBDCollisionSolver_Position_RotationSolverTolerance))
+			{
+				Body1.ApplyRotationDelta(DR1);
 			}
 		}
 	}
@@ -282,11 +287,19 @@ namespace Chaos
 		// but only to correct the velocity that was added by pushout, or in this velocity solve step.
 		if (bChaos_PBDCollisionSolver_Velocity_ImpulseClampEnabled && (Dt > 0))
 		{
+			// @todo(chaos): cache max negative impulse
 			const FVec3 NetImpulse = InOutNetImpulse + Impulse;
-			const FReal PushOutVelocityNormal = FMath::Max(0.0f, FVec3::DotProduct(NetPushOut, ContactNormal) / Dt);
-			if (FVec3::DotProduct(NetImpulse, ContactNormal) < -PushOutVelocityNormal)
+			const FReal PushOutImpulseNormal = FMath::Max(0.0f, FVec3::DotProduct(NetPushOut, ContactNormal) / Dt);
+			const FReal NetImpulseNormal = FVec3::DotProduct(NetImpulse, ContactNormal);
+			if (NetImpulseNormal < -PushOutImpulseNormal)
 			{
-				Impulse = -InOutNetImpulse;
+				// We are trying to apply a negative impulse larger than one to counteract the effective pushout impulse
+				// so clamp the net impulse to be equal to minus the pushout impulse along the normal.
+				// NOTE: NetImpulseNormal is negative here
+				//const FVec3 NewNetImpulse = NetImpulse - NetImpulseNormal * ContactNormal - PushOutImpulseNormal * ContactNormal;
+				//const FVec3 NewImpulse = NewNetImpulse - InOutNetImpulse;
+				//Impulse = InOutNetImpulse + Impulse - NetImpulseNormal * ContactNormal - PushOutImpulseNormal * ContactNormal - InOutNetImpulse;
+				Impulse = Impulse - (NetImpulseNormal + PushOutImpulseNormal) * ContactNormal;
 			}
 		}
 
@@ -319,8 +332,6 @@ namespace Chaos
 			ManifoldPoint.NetImpulse,	// Out
 			Impulse);					// Out
 
-		//UE_LOG(LogChaosCollision, VeryVerbose, TEXT("    Impulse %f"), Impulse.Size());
-
 		// Calculate the velocity deltas from the impulse
 		if (Body0.IsDynamic())
 		{
@@ -485,25 +496,27 @@ namespace Chaos
 		FConstraintSolverBody& Body0 = SolverBody0();
 		FConstraintSolverBody& Body1 = SolverBody1();
 
-		// No need to set an inverse mass scale if ther other body is kinematic (with inv mass of 0)
-		// Also, bodies at the same level do not take part in shock propagation
+		// Shock propagation decreases the inverse mass of bodies that are lower in the pile
+		// of objects. This significantly improves stability of heaps and stacks. Height in the pile is indictaed by the "level". 
+		// No need to set an inverse mass scale if the other body is kinematic (with inv mass of 0).
+		// Bodies at the same level do not take part in shock propagation.
 		if (Body0.IsDynamic() && Body1.IsDynamic() && (Body0.Level() != Body1.Level()))
 		{
 			// Set the inv mass scale of the "lower" body to make it heavier
 			bool bInvMassUpdated = false;
 			if (Body0.Level() < Body1.Level())
 			{
-				if (Body0.InvMassScale() != InvMassScale)
+				if (Body0.InvMScale() != InvMassScale)
 				{
-					Body0.SetInvMassScale(InvMassScale);
+					Body0.SetInvMScale(InvMassScale);
 					bInvMassUpdated = true;
 				}
 			}
 			else
 			{
-				if (Body1.InvMassScale() != InvMassScale)
+				if (Body1.InvMScale() != InvMassScale)
 				{
-					Body1.SetInvMassScale(InvMassScale);
+					Body1.SetInvMScale(InvMassScale);
 					bInvMassUpdated = true;
 				}
 			}
diff --git a/Engine/Source/Runtime/Experimental/Chaos/Private/Chaos/Collision/SolverCollisionContainer.cpp b/Engine/Source/Runtime/Experimental/Chaos/Private/Chaos/Collision/SolverCollisionContainer.cpp
--- a/Engine/Source/Runtime/Experimental/Chaos/Private/Chaos/Collision/SolverCollisionContainer.cpp
+++ b/Engine/Source/Runtime/Experimental/Chaos/Private/Chaos/Collision/SolverCollisionContainer.cpp
@@ -106,7 +106,7 @@ namespace Chaos
 					WorldContactNormal);
 
 				// Calculate the target normal velocity based on restitution
-				FReal WorldContactVelocityTargetNormal = FReal(0); // If two particles are already separating, this can zero out their normal velocity and cause loss of momentum. 
+				FReal WorldContactVelocityTargetNormal = FReal(0);
 				if (Restitution > FReal(0))
 				{
 					const FVec3 ContactVelocity = Solver.GetManifoldPoint(ManifoldPointIndex).CalculateContactVelocity(Solver.SolverBody0(), Solver.SolverBody1());
@@ -115,10 +115,6 @@ namespace Chaos
 					{
 						WorldContactVelocityTargetNormal = -Restitution * ContactVelocityNormal;
 					}
-					else if (ContactVelocityNormal > RestitutionVelocityThreshold)
-					{
-						WorldContactVelocityTargetNormal = ContactVelocityNormal; // @todo(zhenglin): Ideally we should not apply impulse in this case.
-					}
 				}
 
 				// Initialize the material properties (restitution and friction related)
diff --git a/Engine/Source/Runtime/Experimental/Chaos/Private/Chaos/Evolution/SolverBodyContainer.cpp b/Engine/Source/Runtime/Experimental/Chaos/Private/Chaos/Evolution/SolverBodyContainer.cpp
--- a/Engine/Source/Runtime/Experimental/Chaos/Private/Chaos/Evolution/SolverBodyContainer.cpp
+++ b/Engine/Source/Runtime/Experimental/Chaos/Private/Chaos/Evolution/SolverBodyContainer.cpp
@@ -41,8 +41,12 @@ namespace Chaos
 		{
 			if (SolverBody.IsDynamic())
 			{
+				// NOTE: we use a quaternion normalization approximation in the solver
+				// but the integrator asserts that the quaternion is normalized more accurately
+				const FRotation3 Q = SolverBody.Q().GetNormalized();
+
 				// Set the particle state
-				FParticleUtilities::SetCoMWorldTransform(Particle, SolverBody.P(), SolverBody.Q());
+				FParticleUtilities::SetCoMWorldTransform(Particle, SolverBody.P(), Q);
 				Particle->SetV(SolverBody.V());
 				Particle->SetW(SolverBody.W());
 
diff --git a/Engine/Source/Runtime/Experimental/Chaos/Private/Chaos/Joint/PBDJointSolverGaussSeidel.cpp b/Engine/Source/Runtime/Experimental/Chaos/Private/Chaos/Joint/PBDJointSolverGaussSeidel.cpp
--- a/Engine/Source/Runtime/Experimental/Chaos/Private/Chaos/Joint/PBDJointSolverGaussSeidel.cpp
+++ b/Engine/Source/Runtime/Experimental/Chaos/Private/Chaos/Joint/PBDJointSolverGaussSeidel.cpp
@@ -56,13 +56,13 @@ namespace Chaos
 		{
 			ConnectorXs[0] = P(0) + Q(0) * LocalConnectorXs[0].GetTranslation();
 			ConnectorRs[0] = Q(0) * LocalConnectorXs[0].GetRotation();
-			SolverBodies[0].UpdateRotationDependentState();
+			//SolverBodies[0].UpdateRotationDependentState();
 		}
 		if (InvM(1) > SMALL_NUMBER)
 		{
 			ConnectorXs[1] = P(1) + Q(1) * LocalConnectorXs[1].GetTranslation();
 			ConnectorRs[1] = Q(1) * LocalConnectorXs[1].GetRotation();
-			SolverBodies[1].UpdateRotationDependentState();
+			//SolverBodies[1].UpdateRotationDependentState();
 		}
 		ConnectorRs[1].EnforceShortestArcWith(ConnectorRs[0]);
 	}
@@ -74,7 +74,7 @@ namespace Chaos
 		ConnectorRs[BodyIndex] = Q(BodyIndex) * LocalConnectorXs[BodyIndex].GetRotation();
 		ConnectorRs[1].EnforceShortestArcWith(ConnectorRs[0]);
 		
-		SolverBodies[BodyIndex].UpdateRotationDependentState();
+		//SolverBodies[BodyIndex].UpdateRotationDependentState();
 	}
 
 	bool FPBDJointSolver::UpdateIsActive(
@@ -118,8 +118,8 @@ namespace Chaos
 		LocalConnectorXs[1] = XL1;
 
 		// \todo(chaos): joint should support parent/child in either order
-		SolverBodies[0].SetInvMassScale(JointSettings.ParentInvMassScale);
-		SolverBodies[1].SetInvMassScale(FReal(1));
+		SolverBodies[0].SetInvMScale(JointSettings.ParentInvMassScale);
+		SolverBodies[1].SetInvMScale(FReal(1));
 
 		// \chaos(todo): Add support for mass conditioning in ConstraintSolverBody
 		//FPBDJointUtilities::ConditionInverseMassAndInertia(InvM(0), InvM(1), InvILs[0], InvILs[1], SolverSettings.MinParentMassRatio, SolverSettings.MaxInertiaRatio);
@@ -233,7 +233,7 @@ namespace Chaos
 		const FPBDJointSettings& JointSettings)
 	{
 		// @todo(chaos): We need to handle parent/child being the other way round
-		if (InvM(1) < SMALL_NUMBER)
+		if (!Body1().IsDynamic())
 		{
 			// If child is kinematic, return. 
 			return;
@@ -541,11 +541,11 @@ namespace Chaos
 		if (bLinearLocked[0] && bLinearLocked[1] && bLinearLocked[2])
 		{
 			// Hard point constraint (most common case)
-			if (InvM(0) < SMALL_NUMBER)
+			if (!Body0().IsDynamic())
 			{
 				ApplyPointPositionConstraintKD(0, 1, Dt, SolverSettings, JointSettings);
 			}
-			else if (InvM(1) < SMALL_NUMBER)
+			else if (!Body1().IsDynamic())
 			{
 				ApplyPointPositionConstraintKD(1, 0, Dt, SolverSettings, JointSettings);
 			}
@@ -1088,11 +1088,11 @@ namespace Chaos
 		}
 		else
 		{
-			if (InvM(0) > SMALL_NUMBER)
+			if (Body0().IsDynamic())
 			{
 				Body(0).ApplyRotationDelta(DR0);
 			}
-			if (InvM(1) > SMALL_NUMBER)
+			if (Body1().IsDynamic())
 			{
 				Body(1).ApplyRotationDelta(DR1);
 			}
@@ -1304,11 +1304,11 @@ namespace Chaos
 		const FReal Angle,
 		const int32 AngularHardLambdaIndex)
 	{
-		if (InvM(0) < SMALL_NUMBER)
+		if (!Body0().IsDynamic())
 		{
 			ApplyRotationConstraintKD(0, 1, JointStiffness, Axis, Angle, AngularHardLambdaIndex);
 		}
-		else if (InvM(1) < SMALL_NUMBER)
+		else if (!Body1().IsDynamic())
 		{
 			ApplyRotationConstraintKD(1, 0, JointStiffness, Axis, -Angle, AngularHardLambdaIndex);
 		}
@@ -1334,7 +1334,8 @@ namespace Chaos
 		const FReal AngVelTarget,
 		FReal& Lambda)
 	{
-		check(InvM(DIndex) > SMALL_NUMBER);
+		check(!Body(KIndex).IsDynamic());
+		check(Body(DIndex).IsDynamic());
 
 		if (bRealTypeCompatibleWithISPC && bChaos_Joint_ISPC_Enabled)
 		{
@@ -1389,8 +1390,8 @@ namespace Chaos
 		const FReal AngVelTarget,
 		FReal& Lambda)
 	{
-		check(InvM(0) > SMALL_NUMBER);
-		check(InvM(1) > SMALL_NUMBER);
+		check(Body(0).IsDynamic());
+		check(Body(1).IsDynamic());
 
 		if (bRealTypeCompatibleWithISPC && bChaos_Joint_ISPC_Enabled)
 		{
@@ -1448,11 +1449,11 @@ namespace Chaos
 		const FReal AngVelTarget,
 		FReal& Lambda)
 	{
-		if (InvM(0) < SMALL_NUMBER)
+		if (!Body0().IsDynamic())
 		{
 			ApplyRotationConstraintSoftKD(0, 1, Dt, JointStiffness, JointDamping, bAccelerationMode, Axis, Angle, AngVelTarget, Lambda);
 		}
-		else if (InvM(1) < SMALL_NUMBER)
+		else if (!Body1().IsDynamic())
 		{
 			ApplyRotationConstraintSoftKD(1, 0, Dt, JointStiffness, JointDamping, bAccelerationMode, Axis, -Angle, -AngVelTarget, Lambda);
 		}
@@ -1926,7 +1927,8 @@ namespace Chaos
 		const FPBDJointSolverSettings& SolverSettings,
 		const FPBDJointSettings& JointSettings)
 	{
-		check(InvM(DIndex) > SMALL_NUMBER);
+		check(!Body(KIndex).IsDynamic());
+		check(Body(DIndex).IsDynamic());
 
 		FReal Stiffness = SolverStiffness * FPBDJointUtilities::GetLinearStiffness(SolverSettings, JointSettings);
 		const FVec3 CX = ConnectorXs[DIndex] - ConnectorXs[KIndex];
@@ -1968,8 +1970,8 @@ namespace Chaos
 		const FPBDJointSolverSettings& SolverSettings,
 		const FPBDJointSettings& JointSettings)
 	{
-		check(InvM(0) > SMALL_NUMBER);
-		check(InvM(1) > SMALL_NUMBER);
+		check(Body(0).IsDynamic());
+		check(Body(1).IsDynamic());
 
 		FReal Stiffness = SolverStiffness * FPBDJointUtilities::GetLinearStiffness(SolverSettings, JointSettings);
 		const FVec3 CX = ConnectorXs[1] - ConnectorXs[0];
diff --git a/Engine/Source/Runtime/Experimental/Chaos/Public/Chaos/Evolution/SolverBody.h b/Engine/Source/Runtime/Experimental/Chaos/Public/Chaos/Evolution/SolverBody.h
--- a/Engine/Source/Runtime/Experimental/Chaos/Public/Chaos/Evolution/SolverBody.h
+++ b/Engine/Source/Runtime/Experimental/Chaos/Public/Chaos/Evolution/SolverBody.h
@@ -2,6 +2,7 @@
 #pragma once
 
 #include "Chaos/Core.h"
+#include "Math/Quat.h"
 
 namespace Chaos
 {
@@ -12,6 +13,81 @@ namespace Chaos
 	// @note Pointers are only valid for the Constraint Solving phase of the tick
 	using FSolverBodyPtrPair = TVector<FSolverBody*, 2>;
 
+	/**
+	 * @brief An approximate quaternion normalize for use in the solver
+	 * 
+	 * @note we need to correctly normalize the final quaternion before pushing it back to the
+	 * particle otherwise some tolerance checks elsewhere will fail (Integrate)
+	 * 
+	 * This avoids the sqrt which is a massively dominating cost especially with doubles
+	 * when we do not have a fast reciproical sqrt (AVX2)
+	 *
+	 * This uses the first order Pade apprimant instead of a Taylor expansion
+	 * to get more accurate results for small quaternion deltas (i.e., where
+	 * Q.SizeSquared() is already near 1). 
+	 *
+	 * Q.Normalized ~= Q * (2 / (1 + Q.SizeSquared)))
+	 * 
+	 * In practice we can use this for almost any delta generated in collision detection
+	 * but we have an accurate fallback just in case. The fallback adds a branch but this 
+	 * does not seem to cost much.
+	 * 
+	*/
+	FORCEINLINE void SolverQuaternionNormalizeApprox(FRotation3& InOutQ)
+	{
+		//constexpr FReal Tolerance = FReal(2.107342e-08);	// For full double-precision accuracy
+		constexpr FReal ToleranceF = FReal(0.001);
+
+#if PLATFORM_ENABLE_VECTORINTRINSICS
+		using QuatVectorRegister = FRotation3::QuatVectorRegister;
+
+		constexpr VectorRegister Tolerance = MakeVectorRegisterConstant(ToleranceF, ToleranceF, ToleranceF, ToleranceF);
+		constexpr VectorRegister One = MakeVectorRegisterConstant(FReal(1), FReal(1), FReal(1), FReal(1));
+		constexpr VectorRegister Two = MakeVectorRegisterConstant(FReal(2), FReal(2), FReal(2), FReal(2));
+
+		//const FReal QSq = SizeSquared();
+		const QuatVectorRegister Q = VectorLoadAligned(&InOutQ);
+		const QuatVectorRegister QSq = VectorDot4(Q, Q);
+
+		// if (FMath::Abs(FReal(1) - QSq) < Tolerance)
+		const QuatVectorRegister ToleranceCheck = VectorAbs(VectorSubtract(One, QSq));
+		if (VectorMaskBits(VectorCompareLE(ToleranceCheck, Tolerance)) != 0)
+		{
+			// Q * (FReal(2) / (FReal(1) + QSq))
+			const QuatVectorRegister Denom = VectorAdd(One, QSq);
+			const QuatVectorRegister Mult = VectorDivide(Two, Denom);
+			const QuatVectorRegister Result = VectorMultiply(Q, Mult);
+			VectorStoreAligned(Result, &InOutQ);
+		}
+		else
+		{
+			// Q / QLen
+			// @todo(chaos): with doubles VectorReciprocalSqrt does twice as many sqrts as we need and also has a divide
+			const QuatVectorRegister Result = VectorNormalize(Q);
+			VectorStoreAligned(Result, &InOutQ);
+		}
+
+#else
+		const FReal QSq = InOutQ.SizeSquared();
+		if (FMath::Abs(FReal(1) - QSq) < ToleranceF)
+		{
+			InOutQ *= (FReal(2) / (FReal(1) + QSq));
+		}
+		else
+		{
+			InOutQ *= FMath::InvSqrt(QSq);
+		}
+#endif
+	}
+
+	FORCEINLINE FRotation3 SolverQuaternionApplyAngularDeltaApprox(const FRotation3& InQ0, const FVec3& InDR)
+	{
+		FRotation3 Q1 = InQ0 + (FRotation3::FromElements(InDR, FReal(0)) * InQ0) * FReal(0.5);
+		SolverQuaternionNormalizeApprox(Q1);
+		return Q1;
+	}
+
+
 	/**
 	 * Used by the constraint solver loop to cache all state for a particle and accumulate solver results.
 	 * Uses a gather/scatter mechanism to read/write data to the particle SOAs at the beginning/end of the constraint solve.
@@ -27,6 +103,7 @@ namespace Chaos
 	 * impulses and position corrections relative to the center of mass.
 	 * 
 	 * @todo(chaos): layout for cache
+	 * 
 	 */
 	class FSolverBody
 	{
@@ -41,7 +118,7 @@ namespace Chaos
 		/**
 		 * @brief Calculate and set the velocity and angular velocity from the net transform delta
 		*/
-		void SetImplicitVelocity(FReal Dt)
+		inline void SetImplicitVelocity(FReal Dt)
 		{
 			if (IsDynamic())
 			{
@@ -53,32 +130,32 @@ namespace Chaos
 		/**
 		 * @brief Get the inverse mass
 		*/
-		FReal InvM() const { return State.InvM; }
+		inline FReal InvM() const { return State.InvM; }
 
 		/**
 		 * @brief Set the inverse mass
 		*/
-		void SetInvM(FReal InInvM) { State.InvM = InInvM; }
+		inline void SetInvM(FReal InInvM) { State.InvM = InInvM; }
 
 		/**
 		 * @brief Get the world-space inverse inertia
 		*/
-		const FMatrix33& InvI() const { return State.InvI; }
+		inline const FMatrix33& InvI() const { return State.InvI; }
 
 		/**
 		 * @brief Set the world-space inverse inertia
 		*/
-		void SetInvI(const FMatrix33& InInvI) { State.InvI = InInvI; }
+		inline void SetInvI(const FMatrix33& InInvI) { State.InvI = InInvI; }
 
 		/**
 		 * @brief Get the local-space inverse inertia (diagonal elements)
 		*/
-		const FVec3& InvILocal() const { return State.InvILocal; }
+		inline const FVec3& InvILocal() const { return State.InvILocal; }
 
 		/**
 		 * @brief Set the local-space inverse inertia (diagonal elements)
 		*/
-		void SetInvILocal(const FVec3& InInvILocal)
+		inline void SetInvILocal(const FVec3& InInvILocal)
 		{ 
 			State.InvILocal = InInvILocal; 
 			UpdateRotationDependentState();
@@ -87,84 +164,84 @@ namespace Chaos
 		/**
 		 * @brief The current CoM transform
 		*/
-		FRigidTransform3 CoMTransform() const { return FRigidTransform3(P(), Q()); }
+		inline FRigidTransform3 CoMTransform() const { return FRigidTransform3(P(), Q()); }
 
 		/**
 		 * @brief Pre-integration world-space center of mass position
 		*/
-		const FVec3& X() const { return State.X; }
-		void SetX(const FVec3& InX) { State.X = InX; }
+		inline const FVec3& X() const { return State.X; }
+		inline void SetX(const FVec3& InX) { State.X = InX; }
 
 		/**
 		 * @brief Pre-integration world-space center of mass rotation
 		*/
-		const FRotation3& R() const { return State.R; }
-		void SetR(const FRotation3& InR) { State.R = InR; }
+		inline const FRotation3& R() const { return State.R; }
+		inline void SetR(const FRotation3& InR) { State.R = InR; }
 
 		/**
 		 * @brief World-space center of mass position
 		*/
-		const FVec3& P() const { return State.P; }
-		void SetP(const FVec3& InP) { State.P = InP; }
+		inline const FVec3& P() const { return State.P; }
+		inline void SetP(const FVec3& InP) { State.P = InP; }
 
 		/**
 		 * @brief World-space center of mass rotation
 		*/
-		const FRotation3& Q() const { return State.Q; }
-		void SetQ(const FRotation3& InQ) { State.Q = InQ; }
+		inline const FRotation3& Q() const { return State.Q; }
+		inline void SetQ(const FRotation3& InQ) { State.Q = InQ; }
 
 		/**
 		 * @brief World-space center of mass velocity
 		*/
-		const FVec3& V() const { return State.V; }
-		void SetV(const FVec3& InV) { State.V = InV; }
+		inline const FVec3& V() const { return State.V; }
+		inline void SetV(const FVec3& InV) { State.V = InV; }
 
 		/**
 		 * @brief World-space center of mass angular velocity
 		*/
-		const FVec3& W() const { return State.W; }
-		void SetW(const FVec3& InW) { State.W = InW; }
+		inline const FVec3& W() const { return State.W; }
+		inline void SetW(const FVec3& InW) { State.W = InW; }
 
-		const FVec3& CoM() const { return State.CoM; }
-		void SetCoM(const FVec3& InCoM) { State.CoM = InCoM; }
+		inline const FVec3& CoM() const { return State.CoM; }
+		inline void SetCoM(const FVec3& InCoM) { State.CoM = InCoM; }
 
-		const FRotation3& RoM() const { return State.RoM; }
-		void SetRoM(const FRotation3& InRoM) { State.RoM = InRoM; }
+		inline const FRotation3& RoM() const { return State.RoM; }
+		inline void SetRoM(const FRotation3& InRoM) { State.RoM = InRoM; }
 
 		/**
 		 * @brief Get the current world-space Actor position 
 		 * @note This is recalculated from the current CoM transform.
 		*/
-		FVec3 ActorP() const { return P() - ActorQ().RotateVector(CoM()); }
+		inline FVec3 ActorP() const { return P() - ActorQ().RotateVector(CoM()); }
 
 		/**
 		 * @brief Get the current world-space Actor rotation
 		 * @note This is recalculated from the current CoM transform.
 		*/
-		FRotation3 ActorQ() const { return Q() * RoM().Inverse(); }
+		inline FRotation3 ActorQ() const { return Q() * RoM().Inverse(); }
 
 		/**
 		 * @brief Contact graph level. This is used in shock propagation to determine which of two bodies should have its inverse mass scaled
 		*/
-		int32 Level() const { return State.Level; }
-		void SetLevel(int32 InLevel) { State.Level = InLevel; }
+		inline int32 Level() const { return State.Level; }
+		inline void SetLevel(int32 InLevel) { State.Level = InLevel; }
 
 		/**
 		 * @brief Whether there were any active collision constraints on this body
 		*/
-		bool HasActiveCollision() const { return State.bHasActiveCollision; }
-		void SetHasActiveCollision(bool bInHasCollision) { State.bHasActiveCollision = bInHasCollision; }
+		inline bool HasActiveCollision() const { return State.bHasActiveCollision; }
+		inline void SetHasActiveCollision(bool bInHasCollision) { State.bHasActiveCollision = bInHasCollision; }
 
 		/**
 		 * @brief Whether the body has a finite mass
 		 * @note This is based on the current inverse mass, so a "dynamic" particle with 0 inverse mass will return true here.
 		*/
-		bool IsDynamic() const { return (State.InvM > SMALL_NUMBER); }
+		inline bool IsDynamic() const { return (State.InvM > SMALL_NUMBER); }
 
 		/**
 		 * @brief Apply a world-space position and rotation delta to the body center of mass, and update inverse mass
 		*/
-		void ApplyTransformDelta(const FVec3 DP, const FVec3 DR)
+		inline void ApplyTransformDelta(const FVec3 DP, const FVec3 DR)
 		{
 			ApplyPositionDelta(DP);
 			ApplyRotationDelta(DR);
@@ -173,7 +250,7 @@ namespace Chaos
 		/**
 		 * @brief Apply a world-space position delta to the solver body center of mass
 		*/
-		void ApplyPositionDelta(const FVec3 DP)
+		inline void ApplyPositionDelta(const FVec3 DP)
 		{
 			State.P += DP;
 			++State.LastChangeEpoch;
@@ -182,22 +259,16 @@ namespace Chaos
 		/**
 		 * @brief Apply a world-space rotation delta to the solver body and update the inverse mass
 		*/
-		void ApplyRotationDelta(const FVec3 DR)
+		inline void ApplyRotationDelta(const FVec3 DR)
 		{
-			State.Q += FRotation3::FromElements(DR, 0.f) * State.Q * FReal(0.5);
-			State.Q.Normalize();
+			State.Q = SolverQuaternionApplyAngularDeltaApprox(State.Q, DR);
 			++State.LastChangeEpoch;
-
-			//if (bUpdateRotationDependentState)
-			//{
-			//	UpdateRotationDependentState();
-			//}
 		}
 
 		/**
 		 * @brief Apply a world-space velocity delta to the solver body
 		*/
-		void ApplyVelocityDelta(const FVec3& DV, const FVec3& DW)
+		inline void ApplyVelocityDelta(const FVec3& DV, const FVec3& DW)
 		{
 			ApplyLinearVelocityDelta(DV);
 			ApplyAngularVelocityDelta(DW);
@@ -206,7 +277,7 @@ namespace Chaos
 		/**
 		 * @brief Apply a world-space linear velocity delta to the solver body
 		*/
-		void ApplyLinearVelocityDelta(const FVec3& DV)
+		inline void ApplyLinearVelocityDelta(const FVec3& DV)
 		{
 			State.V += DV;
 			++State.LastChangeEpoch;
@@ -215,7 +286,7 @@ namespace Chaos
 		/**
 		 * @brief Apply an world-space angular velocity delta to the solver body
 		*/
-		void ApplyAngularVelocityDelta(const FVec3& DW)
+		inline void ApplyAngularVelocityDelta(const FVec3& DW)
 		{
 			State.W += DW;
 			++State.LastChangeEpoch;
@@ -225,12 +296,12 @@ namespace Chaos
 		 * @brief Update the rotation to be in the same hemisphere as the provided quaternion.
 		 * This is used by joints with angular constraint/drives
 		*/
-		void EnforceShortestRotationTo(const FRotation3& InQ)
+		inline void EnforceShortestRotationTo(const FRotation3& InQ)
 		{
 			State.Q.EnforceShortestArcWith(InQ);
 		}
 
-		int32 LastChangeEpoch() const
+		inline int32 LastChangeEpoch() const
 		{
 			return State.LastChangeEpoch;
 		}
@@ -313,6 +384,7 @@ namespace Chaos
 		FState State;
 	};
 
+
 	/**
 	 * An FSolverBody decorator for adding mass modifiers to a SolverBody. This will scale the
 	 * inverse mass and inverse inertia using the supplied scale. It also updates IsDynamic() to
@@ -340,30 +412,30 @@ namespace Chaos
 		FConstraintSolverBody(FSolverBody& InBody, FReal InInvMassScale)
 			: Body(&InBody)
 		{
-			SetInvMassScale(InInvMassScale);
+			SetInvMScale(InInvMassScale);
 		}
 
 		/**
 		 * @brief True if we have been set up to decorate a SolverBody
 		*/
-		bool IsValid() const { return Body != nullptr; }
+		inline bool IsValid() const { return Body != nullptr; }
 
 		/**
 		 * @brief Invalidate the solver body reference
 		*/
-		void Reset() { Body = nullptr; }
+		inline void Reset() { Body = nullptr; }
 
 		/**
 		 * @brief The decorated SolverBody
 		*/
-		FSolverBody& SolverBody() { check(IsValid()); return *Body; }
-		const FSolverBody& SolverBody() const { check(IsValid()); return *Body; }
+		inline FSolverBody& SolverBody() { check(IsValid()); return *Body; }
+		inline const FSolverBody& SolverBody() const { check(IsValid()); return *Body; }
 
 		/**
 		 * @brief A scale applied to both inverse mass and inverse inertia
 		*/
-		FReal InvMassScale() const { return State.InvMassScale; }
-		void SetInvMassScale(FReal InInvMassScale) { State.InvMassScale = InInvMassScale; }
+		inline FReal InvMScale() const { return State.InvMassScale; }
+		inline void SetInvMScale(FReal InInvMassScale) { State.InvMassScale = InInvMassScale; }
 
 		/**
 		 * @brief The scaled inverse mass
@@ -378,34 +450,34 @@ namespace Chaos
 		/**
 		 * @brief Whether the body is dynamic (i.e., has a finite mass) after InvMassScale is applied
 		*/
-		bool IsDynamic() const { return (InvM() > SMALL_NUMBER); }
+		inline bool IsDynamic() const { return (Body->InvM() != 0) && (InvMScale() != 0); }
 
 		//
 		// From here all methods just forward to the FSolverBody
 		//
 
-		void SetImplicitVelocity(FReal Dt) { Body->SetImplicitVelocity(Dt); }
-		FRigidTransform3 CoMTransform() const { return Body->CoMTransform(); }
-		const FVec3& X() const { return Body->X(); }
-		const FRotation3& R() const { return Body->R(); }
-		const FVec3& P() const { return Body->P(); }
-		const FRotation3& Q() const { return Body->Q(); }
-		const FVec3 ActorP() const { return Body->ActorP(); }
-		const FRotation3 ActorQ() const { return Body->ActorQ(); }
-		const FVec3& V() const { return Body->V(); }
-		const FVec3& W() const { return Body->W(); }
-		int32 Level() const { return Body->Level(); }
-
-		void ApplyTransformDelta(const FVec3 DP, const FVec3 DR) { Body->ApplyTransformDelta(DP, DR); }
-		void ApplyPositionDelta(const FVec3 DP) { Body->ApplyPositionDelta(DP); }
-		void ApplyRotationDelta(const FVec3 DR) { Body->ApplyRotationDelta(DR); }
-		void ApplyVelocityDelta(const FVec3& DV, const FVec3& DW) { Body->ApplyVelocityDelta(DV, DW); }
-		void ApplyLinearVelocityDelta(const FVec3& DV) { Body->ApplyLinearVelocityDelta(DV); }
-		void ApplyAngularVelocityDelta(const FVec3& DW) { Body->ApplyAngularVelocityDelta(DW); }
-		void EnforceShortestRotationTo(const FRotation3& InQ) { Body->EnforceShortestRotationTo(InQ); }
-		void UpdateRotationDependentState() { Body->UpdateRotationDependentState(); }
-
-		int32 LastChangeEpoch() const { return Body->LastChangeEpoch(); }
+		inline void SetImplicitVelocity(FReal Dt) { Body->SetImplicitVelocity(Dt); }
+		inline FRigidTransform3 CoMTransform() const { return Body->CoMTransform(); }
+		inline const FVec3& X() const { return Body->X(); }
+		inline const FRotation3& R() const { return Body->R(); }
+		inline const FVec3& P() const { return Body->P(); }
+		inline const FRotation3& Q() const { return Body->Q(); }
+		inline const FVec3 ActorP() const { return Body->ActorP(); }
+		inline const FRotation3 ActorQ() const { return Body->ActorQ(); }
+		inline const FVec3& V() const { return Body->V(); }
+		inline const FVec3& W() const { return Body->W(); }
+		inline int32 Level() const { return Body->Level(); }
+
+		inline void ApplyTransformDelta(const FVec3 DP, const FVec3 DR) { Body->ApplyTransformDelta(DP, DR); }
+		inline void ApplyPositionDelta(const FVec3 DP) { Body->ApplyPositionDelta(DP); }
+		inline void ApplyRotationDelta(const FVec3 DR) { Body->ApplyRotationDelta(DR); }
+		inline void ApplyVelocityDelta(const FVec3& DV, const FVec3& DW) { Body->ApplyVelocityDelta(DV, DW); }
+		inline void ApplyLinearVelocityDelta(const FVec3& DV) { Body->ApplyLinearVelocityDelta(DV); }
+		inline void ApplyAngularVelocityDelta(const FVec3& DW) { Body->ApplyAngularVelocityDelta(DW); }
+		inline void EnforceShortestRotationTo(const FRotation3& InQ) { Body->EnforceShortestRotationTo(InQ); }
+		inline void UpdateRotationDependentState() { Body->UpdateRotationDependentState(); }
+
+		inline int32 LastChangeEpoch() const { return Body->LastChangeEpoch(); }
 
 	private:
 		// Struct is only so that we can use the same var names as function names
diff --git a/Engine/Source/Runtime/Experimental/Chaos/Public/Chaos/Joint/PBDJointSolverGaussSeidel.h b/Engine/Source/Runtime/Experimental/Chaos/Public/Chaos/Joint/PBDJointSolverGaussSeidel.h
--- a/Engine/Source/Runtime/Experimental/Chaos/Public/Chaos/Joint/PBDJointSolverGaussSeidel.h
+++ b/Engine/Source/Runtime/Experimental/Chaos/Public/Chaos/Joint/PBDJointSolverGaussSeidel.h
@@ -48,6 +48,26 @@ namespace Chaos
 			return SolverBodies[BodyIndex];
 		}
 
+		inline FConstraintSolverBody& Body0()
+		{
+			return SolverBodies[0];
+		}
+
+		inline const FConstraintSolverBody& Body0() const
+		{
+			return SolverBodies[0];
+		}
+
+		inline FConstraintSolverBody& Body1()
+		{
+			return SolverBodies[1];
+		}
+
+		inline const FConstraintSolverBody& Body1() const
+		{
+			return SolverBodies[1];
+		}
+
 		inline const FVec3& X(int BodyIndex) const
 		{
 			return Body(BodyIndex).X();
diff --git a/Engine/Source/Runtime/Experimental/Chaos/Public/Chaos/Rotation.h b/Engine/Source/Runtime/Experimental/Chaos/Public/Chaos/Rotation.h
--- a/Engine/Source/Runtime/Experimental/Chaos/Public/Chaos/Rotation.h
+++ b/Engine/Source/Runtime/Experimental/Chaos/Public/Chaos/Rotation.h
@@ -71,7 +71,7 @@ namespace Chaos
 		 * @note EpsilonSq is approximately the square of the angle below which we cannot calculate the axis. It needs to be "much greater" than square of the error in the
 		 * quaternion values which is usually ~1e-4, so values around 1e-3^2 = 1e-6 or greater are about right.
 		 */
-		bool ToAxisAndAngleSafe(TVector<FReal, 3>& OutAxis, FReal& OutAngle, const TVector<FReal, 3>& DefaultAxis, FReal EpsilionSq = 1e-6f) const
+		inline bool ToAxisAndAngleSafe(TVector<FReal, 3>& OutAxis, FReal& OutAngle, const TVector<FReal, 3>& DefaultAxis, FReal EpsilionSq = 1e-6f) const
 		{
 			OutAngle = GetAngle();
 			return GetRotationAxisSafe(OutAxis, DefaultAxis, EpsilionSq);
@@ -87,7 +87,7 @@ namespace Chaos
 		 * @note EpsilonSq is approximately the square of the angle below which we cannot calculate the axis. It needs to be "much greater" than square of the error in the
 		 * quaternion values which is usually ~1e-4, so values around 1e-3^2 = 1e-6 or greater are about right.
 		 */
-		bool GetRotationAxisSafe(TVector<FReal, 3>& OutAxis, const TVector<FReal, 3>& DefaultAxis, FReal EpsilionSq = 1e-6f) const
+		inline bool GetRotationAxisSafe(TVector<FReal, 3>& OutAxis, const TVector<FReal, 3>& DefaultAxis, FReal EpsilionSq = 1e-6f) const
 		{
 			// Tolerance must be much larger than error in normalized vector (usually ~1e-4) for the 
 			// axis calculation to succeed for small angles. For small angles, W ~= 1, and
@@ -108,7 +108,7 @@ namespace Chaos
 		 * Extract the Swing and Twist rotations, assuming that the Twist Axis is (1,0,0).
 		 * /see ToSwingTwist
 		 */
-		void ToSwingTwistX(BaseQuat& OutSwing, BaseQuat& OutTwist) const
+		inline void ToSwingTwistX(BaseQuat& OutSwing, BaseQuat& OutTwist) const
 		{
 			OutTwist = (X != 0.0f)? BaseQuat(X, 0, 0, W).GetNormalized() : BaseQuat::Identity;
 			OutSwing = *this * OutTwist.Inverse();
@@ -117,7 +117,7 @@ namespace Chaos
 		/**
 		 * @brief Return the dot product of two quaternions
 		*/
-		static FReal DotProduct(const TRotation<FReal, 3>& L, const TRotation<FReal, 3>& R)
+		static inline FReal DotProduct(const TRotation<FReal, 3>& L, const TRotation<FReal, 3>& R)
 		{
 			return (L | R);
 		}
@@ -125,7 +125,7 @@ namespace Chaos
 		/**
 		 * Return the complex conjugate of the rotation
 		 */
-		static TRotation<FReal, 3> Conjugate(const ::Chaos::TRotation<FReal, 3>& InR)
+		static inline TRotation<FReal, 3> Conjugate(const ::Chaos::TRotation<FReal, 3>& InR)
 		{
 			TRotation<FReal, 3> R;
 			R.X = -InR.X;
@@ -138,7 +138,7 @@ namespace Chaos
 		/**
 		 * Negate all values of the quaternion (note: not the inverse rotation. See Conjugate)
 		 */
-		static TRotation<FReal, 3> Negate(const ::Chaos::TRotation<FReal, 3>& InR)
+		static inline TRotation<FReal, 3> Negate(const ::Chaos::TRotation<FReal, 3>& InR)
 		{
 			TRotation<FReal, 3> R;
 			R.X = -InR.X;
@@ -151,7 +151,7 @@ namespace Chaos
 		/**
 		 * Create an identity rotation
 		 */
-		static TRotation<FReal, 3> FromIdentity()
+		static inline TRotation<FReal, 3> FromIdentity()
 		{
 			return BaseQuat(0, 0, 0, 1);
 		}
@@ -159,7 +159,7 @@ namespace Chaos
 		/**
 		 * Create a rotation by explicitly specifying all elements
 		 */
-		static TRotation<FReal, 3> FromElements(const FReal X, const FReal Y, const FReal Z, const FReal W)
+		static inline TRotation<FReal, 3> FromElements(const FReal X, const FReal Y, const FReal Z, const FReal W)
 		{
 			using FQuatReal = BaseQuat::FReal;
 			return BaseQuat((FQuatReal)X, (FQuatReal)Y, (FQuatReal)Z, (FQuatReal)W);
@@ -168,7 +168,7 @@ namespace Chaos
 		/**
 		 * Create a rotation by explicitly specifying all elements
 		 */
-		static TRotation<FReal, 3> FromElements(const ::Chaos::TVector<FReal, 3>& V, const FReal W)
+		static inline TRotation<FReal, 3> FromElements(const ::Chaos::TVector<FReal, 3>& V, const FReal W)
 		{
 			return FromElements(V.X, V.Y, V.Z, W);
 		}
@@ -176,7 +176,7 @@ namespace Chaos
 		/**
 		 * Create a rotation about an axis by an angle specified in radians
 		 */
-		static TRotation<FReal, 3> FromAxisAngle(const ::Chaos::TVector<FReal, 3>& Axis, const FReal AngleRad)
+		static inline TRotation<FReal, 3> FromAxisAngle(const ::Chaos::TVector<FReal, 3>& Axis, const FReal AngleRad)
 		{
 			return BaseQuat(UE::Math::TVector<FReal>(Axis.X, Axis.Y, Axis.Z), (BaseQuat::FReal)AngleRad);
 		}
@@ -184,7 +184,7 @@ namespace Chaos
 		/**
 		 * Create a rotation about an axis V/|V| by angle |V| in radians
 		 */
-		static TRotation<FReal, 3> FromVector(const ::Chaos::TVector<FReal, 3>& V)
+		static inline TRotation<FReal, 3> FromVector(const ::Chaos::TVector<FReal, 3>& V)
 		{
 			using FQuatReal = BaseQuat::FReal;
 			TRotation<FReal, 3> Rot;
@@ -230,7 +230,7 @@ namespace Chaos
 		 *
 		 * Uses the relation: DQ/DT = (W * Q)/2
 		 */
-		static TVector<FReal, 3> CalculateAngularVelocity1(const TRotation<FReal, 3>& InR0, const TRotation<FReal, 3>& InR1, const FReal InDt)
+		static inline TVector<FReal, 3> CalculateAngularVelocity1(const TRotation<FReal, 3>& InR0, const TRotation<FReal, 3>& InR1, const FReal InDt)
 		{
 			if (!ensure(InDt > SMALL_NUMBER))
 			{
@@ -254,7 +254,7 @@ namespace Chaos
 		 *
 		 * Uses the Quaternion to Axis/Angle method.
 		 */
-		static TVector<FReal, 3> CalculateAngularVelocity2(const TRotation<FReal, 3>& InR0, const TRotation<FReal, 3>& InR1, const FReal InDt)
+		static inline TVector<FReal, 3> CalculateAngularVelocity2(const TRotation<FReal, 3>& InR0, const TRotation<FReal, 3>& InR1, const FReal InDt)
 		{
 			// @todo(ccaulfield): ToAxisAndAngle starts to return increasingly random, non-normalized axes for very small angles. This 
 			// underestimates the angular velocity magnitude and randomizes direction.
@@ -278,7 +278,7 @@ namespace Chaos
 		 *
 		 * This should match the algorithm used in PerParticleUpdateFromDeltaPosition rule.
 		 */
-		static TVector<FReal, 3> CalculateAngularVelocity(const TRotation<FReal, 3>& InR0, const TRotation<FReal, 3>& InR1, const FReal InDt)
+		static inline TVector<FReal, 3> CalculateAngularVelocity(const TRotation<FReal, 3>& InR0, const TRotation<FReal, 3>& InR1, const FReal InDt)
 		{
 			return CalculateAngularVelocity1(InR0, InR1, InDt);
 		}
@@ -288,7 +288,7 @@ namespace Chaos
 		 *
 		 * This should match the algorithm used in PerParticleUpdateFromDeltaPosition rule.
 		 */
-		static TVector<FReal, 3> CalculateAngularDelta(const TRotation<FReal, 3>& InR0, const TRotation<FReal, 3>& InR1)
+		static inline TVector<FReal, 3> CalculateAngularDelta(const TRotation<FReal, 3>& InR0, const TRotation<FReal, 3>& InR1)
 		{
 			return CalculateAngularVelocity(InR0, InR1, 1.0f);
 		}
@@ -298,18 +298,19 @@ namespace Chaos
 		 *
 		 * Uses the relation: DQ/DT = (W * Q)/2
 		 */
-		static TRotation<FReal, 3> IntegrateRotationWithAngularVelocity(const TRotation<FReal, 3>& InR0, const TVector<FReal, 3>& InW, const FReal InDt)
+		static inline TRotation<FReal, 3> IntegrateRotationWithAngularVelocity(const TRotation<FReal, 3>& InR0, const TVector<FReal, 3>& InW, const FReal InDt)
 		{
 			TRotation<FReal, 3> R1 = InR0 + (TRotation<FReal, 3>::FromElements(InW.X, InW.Y, InW.Z, 0.f) * InR0) * decltype(InR0.X)(InDt * 0.5f);
 			return R1.GetNormalized();
 		}
 
+
 		/**
 		 * Check that two rotations are approximately equal. Assumes the quaternions are normalized and in the same hemisphere.
 		 * For small values of Epsilon, this is approximately equivalent to checking that the rotations are within 2*Epsilon
 		 * radians of each other.
 		 */
-		static bool IsNearlyEqual(const TRotation<FReal, 3>& A, const TRotation<FReal, 3>& B, const FReal Epsilon)
+		static inline bool IsNearlyEqual(const TRotation<FReal, 3>& A, const TRotation<FReal, 3>& B, const FReal Epsilon)
 		{
 			// Only check imaginary part. This is comparing Epsilon to 2*AngleDelta for small angle deltas
 			return FMath::IsNearlyEqual((FReal)A.X, (FReal)B.X, Epsilon)
