Implementation of a dynamic AABB tree
#[rb] Benn.Gallagher
#[fyi] Michael.Lentine, Chris.Caulfield, Brice.Criswell, Ori.Cohen

diff --git a/Engine/Source/Programs/HeadlessChaos/Private/HeadlessChaosApp.cpp b/Engine/Source/Programs/HeadlessChaos/Private/HeadlessChaosApp.cpp
--- a/Engine/Source/Programs/HeadlessChaos/Private/HeadlessChaosApp.cpp
+++ b/Engine/Source/Programs/HeadlessChaos/Private/HeadlessChaosApp.cpp
@@ -236,6 +236,7 @@ TEST(BP, BroadphaseTests) {
 	ChaosTest::GridBPTest();
 	ChaosTest::GridBPTest2();
 	ChaosTest::AABBTreeTest();
+	ChaosTest::AABBTreeTestDynamic();
 	ChaosTest::AABBTreeDirtyGridTest();
 	ChaosTest::AABBTreeTimesliceTest();
 	ChaosTest::BoundingVolumeNoBoundsTest();
diff --git a/Engine/Source/Programs/HeadlessChaos/Private/HeadlessChaosTestBroadphase.cpp b/Engine/Source/Programs/HeadlessChaos/Private/HeadlessChaosTestBroadphase.cpp
--- a/Engine/Source/Programs/HeadlessChaos/Private/HeadlessChaosTestBroadphase.cpp
+++ b/Engine/Source/Programs/HeadlessChaos/Private/HeadlessChaosTestBroadphase.cpp
@@ -226,8 +226,8 @@ namespace ChaosTest
 	{
 		Box = MakeUnique<TBox<FReal, 3>>(FVec3(0, 0, 0), FVec3(BoxSize, BoxSize, BoxSize));
 		auto Boxes = MakeUnique<FGeometryParticles>();
-		const int32 NumRows = BoxGridDimensions.X;
-		const int32 NumCols = BoxGridDimensions.Y;
+		const int32 NumCols = BoxGridDimensions.X;
+		const int32 NumRows = BoxGridDimensions.Y;
 		const int32 NumHeight = BoxGridDimensions.Z;
 
 		Boxes->AddParticles(NumRows * NumCols * NumHeight);
@@ -297,6 +297,10 @@ namespace ChaosTest
 				Boxes->X(MoveIdx) -= FVec3(1000, 0, 0);
 				NewBounds = Boxes->Geometry(MoveIdx)->template GetObject<TBox<FReal, 3>>()->BoundingBox().TransformedAABB(FRigidTransform3(Boxes->X(MoveIdx), Boxes->R(MoveIdx)));
 				Spatial2->UpdateElementIn(MoveIdx, NewBounds, true, SpatialIdx);
+
+				FVisitor Visitor4(FVec3(10, 0, 0), FVec3(0, 1, 0), 0, *Boxes);
+				Spatial2->Raycast(Visitor4.Start, Visitor4.Dir, 1000, Visitor4);
+				EXPECT_EQ(Visitor4.Instances.Num(), 9);
 			}
 
 			//move other instance into view
@@ -551,6 +555,20 @@ namespace ChaosTest
 	}
 
 
+	void AABBTreeTestDynamic()
+	{
+		using TreeType = TAABBTree<int32, TAABBTreeLeafArray<int32>, true>;
+		{
+			TUniquePtr<TBox<FReal, 3>> Box;
+			auto Boxes = BuildBoxes(Box, 100, FVec3(10,10,10));
+			TreeType Spatial(MakeParticleView(Boxes.Get()), TreeType::DefaultMaxChildrenInLeaf, TreeType::DefaultMaxTreeDepth, TreeType::DefaultMaxPayloadBounds, TreeType::DefaultMaxNumToProcess, true);
+			EXPECT_EQ(Spatial.NumDirtyElements(), 0);
+			SpatialTestHelper(Spatial, Boxes.Get(), Box);
+			EXPECT_EQ(Spatial.NumDirtyElements(), 0);
+		}
+		
+	}
+
 	void AABBTreeDirtyGridTest()
 	{
 		using TreeType = TAABBTree<int32, TBoundingVolume<int32>>;
diff --git a/Engine/Source/Programs/HeadlessChaos/Public/HeadlessChaosTestBroadphase.h b/Engine/Source/Programs/HeadlessChaos/Public/HeadlessChaosTestBroadphase.h
--- a/Engine/Source/Programs/HeadlessChaos/Public/HeadlessChaosTestBroadphase.h
+++ b/Engine/Source/Programs/HeadlessChaos/Public/HeadlessChaosTestBroadphase.h
@@ -13,6 +13,8 @@ namespace ChaosTest {
 
 	void AABBTreeTest();
 
+	void AABBTreeTestDynamic();
+
 	void AABBTreeDirtyGridTest();
 
 	void AABBTreeTimesliceTest();
diff --git a/Engine/Source/Runtime/Experimental/Chaos/Private/Chaos/AABBTree.cpp b/Engine/Source/Runtime/Experimental/Chaos/Private/Chaos/AABBTree.cpp
--- a/Engine/Source/Runtime/Experimental/Chaos/Private/Chaos/AABBTree.cpp
+++ b/Engine/Source/Runtime/Experimental/Chaos/Private/Chaos/AABBTree.cpp
@@ -28,3 +28,6 @@ FAutoConsoleVariableRef FAABBTreeCVars::CVarSplitAtAverageCenter(TEXT("p.aabbtre
 
 int32 FAABBTreeCVars::SplitOnVarianceAxis = 1;
 FAutoConsoleVariableRef FAABBTreeCVars::CVarSplitOnVarianceAxis(TEXT("p.aabbtree.splitonvarianceaxis"), FAABBTreeCVars::SplitOnVarianceAxis, TEXT("Split AABB tree nodes along the axis with the largest element center variance"));
+
+float FAABBTreeCVars::DynamicTreeBoundingBoxPadding = 1.0f;
+FAutoConsoleVariableRef FAABBTreeCVars::CVarDynamicTreeBoundingBoxPadding(TEXT("p.aabbtree.DynamicTreeBoundingBoxPadding"), FAABBTreeCVars::DynamicTreeBoundingBoxPadding, TEXT("Additional padding added to bounding boxes for dynamic AABB trees to amortize update cost"));
diff --git a/Engine/Source/Runtime/Experimental/Chaos/Private/Chaos/PBDRigidsEvolution.cpp b/Engine/Source/Runtime/Experimental/Chaos/Private/Chaos/PBDRigidsEvolution.cpp
--- a/Engine/Source/Runtime/Experimental/Chaos/Private/Chaos/PBDRigidsEvolution.cpp
+++ b/Engine/Source/Runtime/Experimental/Chaos/Private/Chaos/PBDRigidsEvolution.cpp
@@ -7,6 +7,8 @@
 #include "Chaos/ParticleHandle.h"
 #include "Chaos/SpatialAccelerationCollection.h"
 
+CSV_DECLARE_CATEGORY_EXTERN(ChaosPhysicsTimers);
+
 int32 ChaosRigidsEvolutionApplyAllowEarlyOutCVar = 1;
 FAutoConsoleVariableRef CVarChaosRigidsEvolutionApplyAllowEarlyOut(TEXT("p.ChaosRigidsEvolutionApplyAllowEarlyOut"), ChaosRigidsEvolutionApplyAllowEarlyOutCVar, TEXT("Allow Chaos Rigids Evolution apply iterations to early out when resolved.[def:1]"));
 
@@ -30,6 +32,9 @@ namespace Chaos
 	CHAOS_API int32 AccelerationStructureSplitStaticAndDynamic = 1;
 	FAutoConsoleVariableRef CVarAccelerationStructureSplitStaticAndDynamic(TEXT("p.Chaos.AccelerationStructureSplitStaticDynamic"), AccelerationStructureSplitStaticAndDynamic, TEXT("Set to 1: Sort Dynamic and Static bodies into seperate acceleration structures, any other value will disable the feature"));
 
+	CHAOS_API int32 AccelerationStructureUseDynamicTree = 0;
+	FAutoConsoleVariableRef CVarAccelerationStructureUseDynamicTree(TEXT("p.Chaos.AccelerationStructureUseDynamicTree"), AccelerationStructureUseDynamicTree, TEXT("Use a dynamic BVH tree structure for dynamic objects"));
+
 	CHAOS_API int32 AccelerationStructureTimeSlicingMaxQueueSizeBeforeForce = 1000;
 	FAutoConsoleVariableRef CVarAccelerationStructureTimeSlicingMaxQueueSizeBeforeForce(TEXT("p.Chaos.AccelerationStructureTimeSlicingMaxQueueSizeBeforeForce"), AccelerationStructureTimeSlicingMaxQueueSizeBeforeForce, TEXT("If the update queue reaches this limit, time slicing will be disabled, and the acceleration structure will be built at once"));
 
@@ -75,6 +80,7 @@ namespace Chaos
 
 		using BVType = TBoundingVolume<FAccelerationStructureHandle>;
 		using AABBTreeType = TAABBTree<FAccelerationStructureHandle, TAABBTreeLeafArray<FAccelerationStructureHandle>>;
+		using AABBDynamicTreeType = TAABBTree<FAccelerationStructureHandle, TAABBTreeLeafArray<FAccelerationStructureHandle>>;
 		using AABBTreeOfGridsType = TAABBTree<FAccelerationStructureHandle, TBoundingVolume<FAccelerationStructureHandle>>;
 
 		TUniquePtr<ISpatialAccelerationCollection<FAccelerationStructureHandle, FReal, 3>> CreateEmptyCollection() override
@@ -87,11 +93,11 @@ namespace Chaos
 			for (uint16 BucketIdx = 0; BucketIdx < NumBuckets; ++BucketIdx)
 			{
 				// For static bodies
-				Collection->AddSubstructure(CreateAccelerationPerBucket_Threaded(Empty, BucketIdx, true), BucketIdx, 0);
+				Collection->AddSubstructure(CreateAccelerationPerBucket_Threaded(Empty, BucketIdx, true, false), BucketIdx, 0);
 				if (AccelerationStructureSplitStaticAndDynamic == 1)
 				{
 					// Non static bodies
-					Collection->AddSubstructure(CreateAccelerationPerBucket_Threaded(Empty, BucketIdx, true), BucketIdx, 1);
+					Collection->AddSubstructure(CreateAccelerationPerBucket_Threaded(Empty, BucketIdx, true, (bool) AccelerationStructureUseDynamicTree), BucketIdx, 1);
 				}
 			}
 
@@ -140,7 +146,7 @@ namespace Chaos
 			}
 		}
 
-		virtual TUniquePtr<ISpatialAcceleration<FAccelerationStructureHandle, FReal, 3>> CreateAccelerationPerBucket_Threaded(const TConstParticleView<FSpatialAccelerationCache>& Particles, uint16 BucketIdx, bool ForceFullBuild) override
+		virtual TUniquePtr<ISpatialAcceleration<FAccelerationStructureHandle, FReal, 3>> CreateAccelerationPerBucket_Threaded(const TConstParticleView<FSpatialAccelerationCache>& Particles, uint16 BucketIdx, bool ForceFullBuild, bool bDynamicTree) override
 		{
 			// TODO: Unduplicate switch statement here with IsBucketTimeSliced and refactor so that bucket index mapping is better.
 			switch (BucketIdx)
@@ -153,6 +159,10 @@ namespace Chaos
 				}
 				else if (ConfigSettings.BroadphaseType == 1 || ConfigSettings.BroadphaseType == 3)
 				{
+					if (bDynamicTree)
+					{
+						return MakeUnique<AABBDynamicTreeType>(Particles, ConfigSettings.MaxChildrenInLeaf, ConfigSettings.MaxTreeDepth, ConfigSettings.MaxPayloadSize, ForceFullBuild ? 0 : ConfigSettings.IterationsPerTimeSlice, true);
+					}
 					return MakeUnique<AABBTreeType>(Particles, ConfigSettings.MaxChildrenInLeaf, ConfigSettings.MaxTreeDepth, ConfigSettings.MaxPayloadSize, ForceFullBuild ? 0 : ConfigSettings.IterationsPerTimeSlice);
 				}
 				else if (ConfigSettings.BroadphaseType == 4 || ConfigSettings.BroadphaseType == 2)
@@ -432,7 +442,7 @@ namespace Chaos
 
 				auto ParticleView = MakeConstParticleView(MoveTemp(ViewsPerBucket[SpatialAccIdx.Bucket][SpatialAccIdx.InnerIdx]));
 				// Create and run the first slice here
-				auto NewStruct = SpatialCollectionFactory.CreateAccelerationPerBucket_Threaded(ParticleView, SpatialAccIdx.Bucket, IsForceFullBuild);
+				auto NewStruct = SpatialCollectionFactory.CreateAccelerationPerBucket_Threaded(ParticleView, SpatialAccIdx.Bucket, IsForceFullBuild, false);
 				{
 					NewStruct->ClearShouldRebuild();  // Mark as pristine
 				}
@@ -471,7 +481,7 @@ namespace Chaos
 					SCOPE_CYCLE_COUNTER(STAT_CreateNonSlicedStructures);
 
 					auto ParticleView = MakeConstParticleView(MoveTemp(ViewsPerBucket[SpatialAccIdx.Bucket][SpatialAccIdx.InnerIdx]));
-					auto NewStruct = SpatialCollectionFactory.CreateAccelerationPerBucket_Threaded(ParticleView,SpatialAccIdx.Bucket,IsForceFullBuild);
+					auto NewStruct = SpatialCollectionFactory.CreateAccelerationPerBucket_Threaded(ParticleView,SpatialAccIdx.Bucket,IsForceFullBuild, false);
 
 					ISpatialAcceleration<FAccelerationStructureHandle, FReal, 3>* AccelerationSUBStructureCopy = GetSubStructureCopy(SpatialAccIdx);
 					if (AccelerationSUBStructureCopy)
@@ -493,7 +503,7 @@ namespace Chaos
 		LLM_SCOPE(ELLMTag::ChaosAcceleration);
 
 #if !WITH_EDITOR
-		CSV_SCOPED_TIMING_STAT(ChaosPhysicsTimers, AccelerationStructBuildTask);
+		//CSV_SCOPED_TIMING_STAT(ChaosPhysicsTimers, AccelerationStructBuildTask);
 #endif
 		if (AccelerationStructureSplitStaticAndDynamic == 1)
 		{
@@ -534,7 +544,10 @@ namespace Chaos
 		{
 			FGeometryParticleHandle* UpdateParticle = SpatialData.AccelerationHandle.GetGeometryParticleHandle_PhysicsThread();
 
-			AccelerationStructure.UpdateElementIn(UpdateParticle, UpdateParticle->WorldSpaceInflatedBounds(), UpdateParticle->HasBounds(), SpatialData.SpatialIdx);
+			{
+				//CSV_SCOPED_TIMING_STAT(ChaosPhysicsTimers, AABBTreeUpdate)
+				AccelerationStructure.UpdateElementIn(UpdateParticle, UpdateParticle->WorldSpaceInflatedBounds(), UpdateParticle->HasBounds(), SpatialData.SpatialIdx);
+			}
 
 			if (bUpdateCache)
 			{
diff --git a/Engine/Source/Runtime/Experimental/Chaos/Public/Chaos/AABBTree.h b/Engine/Source/Runtime/Experimental/Chaos/Public/Chaos/AABBTree.h
--- a/Engine/Source/Runtime/Experimental/Chaos/Public/Chaos/AABBTree.h
+++ b/Engine/Source/Runtime/Experimental/Chaos/Public/Chaos/AABBTree.h
@@ -28,6 +28,9 @@ struct CHAOS_API FAABBTreeCVars
 
 	static float MaxNonGlobalElementBoundsExtrema; 
 	static FAutoConsoleVariableRef CVarMaxNonGlobalElementBoundsExtrema;
+
+	static float DynamicTreeBoundingBoxPadding;
+	static FAutoConsoleVariableRef CVarDynamicTreeBoundingBoxPadding;
 };
 
 struct CHAOS_API FAABBTreeDirtyGridCVars
@@ -336,6 +339,16 @@ struct TAABBTreeLeafArray : public TBoundsWrapperHelper<TPayloadType, T, bComput
 		}
 	}
 
+	void AddElement(const TPayloadBoundsElement<TPayloadType, T>& Element)
+	{
+		Elems.Add(Element);
+		this->ComputeBounds(Elems);
+	}
+
+	void Reset()
+	{
+		Elems.Reset();
+	}
 
 #if !UE_BUILD_SHIPPING
 	void DebugDrawLeaf(ISpacialDebugDrawInterface<T>& InInterface, const FLinearColor& InLinearColor, float InThickness) const
@@ -379,7 +392,8 @@ struct TAABBTreeNode
 		ChildrenBounds[1] = TAABB<T, 3>();
 	}
 	TAABB<T, 3> ChildrenBounds[2];
-	int32 ChildrenNodes[2] = { 0, 0 };
+	int32 ChildrenNodes[2] = { INDEX_NONE, INDEX_NONE };
+	int32 ParentNode = INDEX_NONE;
 	bool bLeaf = false;
 
 #if !UE_BUILD_SHIPPING
@@ -419,6 +433,12 @@ struct TAABBTreeNode
 		}
 
 		Ar << bLeaf;
+
+		// Dynamic trees are not serialized
+		if (Ar.IsLoading())
+		{
+			ParentNode = INDEX_NONE;
+		}
 	}
 };
 
@@ -435,12 +455,14 @@ struct FAABBTreePayloadInfo
 	int32 DirtyPayloadIdx;
 	int32 LeafIdx;
 	int32 DirtyGridOverflowIdx;
+	int32 NodeIdx;
 
-	FAABBTreePayloadInfo(int32 InGlobalPayloadIdx = INDEX_NONE, int32 InDirtyIdx = INDEX_NONE, int32 InLeafIdx = INDEX_NONE, int32 InDirtyGridOverflowIdx = INDEX_NONE)
+	FAABBTreePayloadInfo(int32 InGlobalPayloadIdx = INDEX_NONE, int32 InDirtyIdx = INDEX_NONE, int32 InLeafIdx = INDEX_NONE, int32 InDirtyGridOverflowIdx = INDEX_NONE, int32 InNodeIdx = INDEX_NONE)
 		: GlobalPayloadIdx(InGlobalPayloadIdx)
 		, DirtyPayloadIdx(InDirtyIdx)
 		, LeafIdx(InLeafIdx)
 		, DirtyGridOverflowIdx(InDirtyGridOverflowIdx)
+		, NodeIdx(InNodeIdx)
 	{}
 
 	void Serialize(FArchive& Ar)
@@ -449,6 +471,11 @@ struct FAABBTreePayloadInfo
 		Ar << DirtyPayloadIdx;
 		Ar << LeafIdx;
 		Ar << DirtyGridOverflowIdx;
+		// Dynamic trees are not serialized
+		if (Ar.IsLoading())
+		{
+			NodeIdx = INDEX_NONE;
+		}
 	}
 };
 
@@ -493,6 +520,10 @@ public:
 		(TIsSame<TBoundingVolume<TPayloadType>, TLeafType>::Value ? ESpatialAcceleration::AABBTreeBV : ESpatialAcceleration::Unknown);
 	TAABBTree()
 		: ISpatialAcceleration<TPayloadType, T, 3>(StaticType)
+		, bDynamicTree(false)
+		, RootNode(INDEX_NONE)
+		, FirstFreeInternalNode(INDEX_NONE)
+		, FirstFreeLeafNode(INDEX_NONE)
 		, MaxChildrenInLeaf(DefaultMaxChildrenInLeaf)
 		, MaxTreeDepth(DefaultMaxTreeDepth)
 		, MaxPayloadBounds(DefaultMaxPayloadBounds)
@@ -521,12 +552,25 @@ public:
 		WorkPool.Reset();
 
 		bShouldRebuild = true;
+
+		RootNode = INDEX_NONE;
+		FirstFreeInternalNode = INDEX_NONE;
+		FirstFreeLeafNode = INDEX_NONE;
+
 		this->SetAsyncTimeSlicingComplete(true);
 	}
 
 	virtual void ProgressAsyncTimeSlicing(bool ForceBuildCompletion) override
 	{
 		SCOPE_CYCLE_COUNTER(STAT_AABBTreeProgressTimeSlice);
+
+		if (bDynamicTree)
+		{
+			// Nothing to do
+			this->SetAsyncTimeSlicingComplete(true);
+			return;
+		}
+
 		// force is to stop time slicing and complete the rest of the build now
 		if (ForceBuildCompletion)
 		{
@@ -542,8 +586,9 @@ public:
 	}
 
 	template <typename TParticles>
-	TAABBTree(const TParticles& Particles, int32 InMaxChildrenInLeaf = DefaultMaxChildrenInLeaf, int32 InMaxTreeDepth = DefaultMaxTreeDepth, T InMaxPayloadBounds = DefaultMaxPayloadBounds, int32 InMaxNumToProcess = DefaultMaxNumToProcess )
+	TAABBTree(const TParticles& Particles, int32 InMaxChildrenInLeaf = DefaultMaxChildrenInLeaf, int32 InMaxTreeDepth = DefaultMaxTreeDepth, T InMaxPayloadBounds = DefaultMaxPayloadBounds, int32 InMaxNumToProcess = DefaultMaxNumToProcess, bool bInDynamicTree = false)
 		: ISpatialAcceleration<TPayloadType, T, 3>(StaticType)
+		, bDynamicTree(bInDynamicTree)
 		, MaxChildrenInLeaf(InMaxChildrenInLeaf)
 		, MaxTreeDepth(InMaxTreeDepth)
 		, MaxPayloadBounds(InMaxPayloadBounds)
@@ -851,6 +896,397 @@ public:
 		return DirtyGridOverflowIdx;
 	}
 
+	// Expensive function: Don't call unless debugging
+	void DynamicTreeDebugStats()
+	{
+		TArray<FElement> AllElements;
+		for (auto& Leaf : Leaves)
+		{
+			Leaf.GatherElements(AllElements);
+		}
+
+		int32 MaxDepth = 0;
+		int32 DepthTotal = 0;
+		for (const FElement& Element : AllElements)
+		{
+			FAABBTreePayloadInfo* PayloadInfo = PayloadToInfo.Find(Element.Payload);
+			int32 Depth = 0;
+			int32 Node = PayloadInfo->NodeIdx;
+			check(Node != INDEX_NONE);
+			while (Node != INDEX_NONE)
+			{
+				Node = Nodes[Node].ParentNode;
+				if (Node != INDEX_NONE)
+				{
+					Depth++;
+				}
+			}
+			if (Depth > MaxDepth)
+			{
+				MaxDepth = Depth;
+			}
+			DepthTotal += Depth;
+		}
+#if !WITH_EDITOR
+		CSV_CUSTOM_STAT(ChaosPhysicsTimers, MaximumTreeDepth, MaxDepth, ECsvCustomStatOp::Max);
+		CSV_CUSTOM_STAT(ChaosPhysicsTimers, AvgTreeDepth, DepthTotal / AllElements.Num(), ECsvCustomStatOp::Max);
+		CSV_CUSTOM_STAT(ChaosPhysicsTimers, Dirty,DirtyElements.Num(), ECsvCustomStatOp::Max);
+#endif
+
+	}
+
+	int32 AllocateInternalNode()
+	{
+		int32 AllocatedNodeIdx = FirstFreeInternalNode;
+		if (FirstFreeInternalNode != INDEX_NONE)
+		{
+			// Unlink from free list
+			FirstFreeInternalNode = Nodes[FirstFreeInternalNode].ChildrenNodes[1];
+		}
+		else
+		{
+			// create the actual node space
+			AllocatedNodeIdx = Nodes.AddUninitialized(1);;
+			Nodes[AllocatedNodeIdx].bLeaf = false;
+		}
+
+		return AllocatedNodeIdx;
+	}
+
+	int32 AllocateLeafNodeAndLeaf(const TPayloadType& Payload, const TAABB<T, 3>& NewBounds)
+	{
+		int32 AllocatedNodeIdx = FirstFreeLeafNode;
+		int32 LeafIndex;
+		if (FirstFreeLeafNode != INDEX_NONE)
+		{
+			FirstFreeLeafNode = Nodes[FirstFreeLeafNode].ChildrenNodes[1];
+			LeafIndex = Nodes[AllocatedNodeIdx].ChildrenNodes[0]; // This is already set when it was allocated for the first time
+			FElement NewElement{ Payload, NewBounds };
+			Leaves[LeafIndex].AddElement(NewElement);
+		}
+		else
+		{
+			LeafIndex = Leaves.Num();
+
+			// create the actual node space
+			AllocatedNodeIdx = Nodes.AddUninitialized(1);
+			Nodes[AllocatedNodeIdx].ChildrenNodes[0] = LeafIndex;
+			Nodes[AllocatedNodeIdx].bLeaf = true;
+
+			FElement NewElement{ Payload, NewBounds };
+			TArray<FElement> SingleElementArray;
+			SingleElementArray.Add(NewElement);
+			Leaves.Add(TLeafType{ SingleElementArray }); // Extra copy
+		}
+
+		// Expand the leaf node bounding box to reduce the number of updates
+		TAABB<T, 3> ExpandedBounds = NewBounds;
+		ExpandedBounds.Thicken(FAABBTreeCVars::DynamicTreeBoundingBoxPadding);
+		Nodes[AllocatedNodeIdx].ChildrenBounds[0] = ExpandedBounds;
+
+		Nodes[AllocatedNodeIdx].ParentNode = INDEX_NONE;
+		FAABBTreePayloadInfo* PayloadInfo = PayloadToInfo.Find(Payload);
+		check(PayloadInfo);
+		*PayloadInfo = FAABBTreePayloadInfo{ INDEX_NONE, INDEX_NONE, LeafIndex, INDEX_NONE, AllocatedNodeIdx };
+
+		return AllocatedNodeIdx;
+	}
+
+	void DeAllocateInternalNode(int32 NodeIdx)
+	{
+		Nodes[NodeIdx].ChildrenNodes[1] = FirstFreeInternalNode;
+		FirstFreeInternalNode = NodeIdx;
+	}
+
+	void  DeAllocateLeafNode(int32 NodeIdx)
+	{
+		
+		Leaves[Nodes[NodeIdx].ChildrenNodes[0]].Reset();
+
+		Nodes[NodeIdx].ChildrenNodes[1] = FirstFreeLeafNode;
+		FirstFreeLeafNode = NodeIdx;
+	}
+
+	// Is the input node Child 0 or Child 1?
+	int32 WhichChildAmI(int32 NodeIdx)
+	{
+		check(NodeIdx != INDEX_NONE);
+		int32 ParentIdx = Nodes[NodeIdx].ParentNode;
+		check(ParentIdx != INDEX_NONE);
+		if (Nodes[ParentIdx].ChildrenNodes[0] == NodeIdx)
+		{
+			return  0;
+		}
+		else
+		{
+			return 1;
+		}
+
+		check(false);
+		return 0;
+	}
+
+	// Is the input node Child 0 or Child 1?
+	int32 GetSiblingIndex(int32 NodeIdx)
+	{
+		return(WhichChildAmI(NodeIdx) ^ 1);
+	}
+
+	int32 FindBestSibling(const TAABB<T, 3>& NewBounds)
+	{
+		
+		//Priority Q of indices to explore
+		TArray<int32> PriorityQ;
+		PriorityQ.Reserve(10);
+
+		TArray<FReal> SumDeltaCostQ;
+		SumDeltaCostQ.Reserve(10);
+
+		int32 QIndex = 0;
+		
+		// Initializing
+		
+		TAABB<T, 3> WorkingAABB{ NewBounds };
+		WorkingAABB.GrowToInclude(Nodes[RootNode].ChildrenBounds[0]);
+		if (!Nodes[RootNode].bLeaf)
+		{
+			WorkingAABB.GrowToInclude(Nodes[RootNode].ChildrenBounds[1]);
+		}
+
+		int32 BestSiblingIdx = RootNode;
+		FReal BestCost = WorkingAABB.GetArea();
+		PriorityQ.Add(RootNode);
+		SumDeltaCostQ.Add(0.0f);
+
+		while (PriorityQ.Num() - QIndex)
+		{
+			// Pop from queue
+			uint32 TestSibling = PriorityQ[QIndex];
+			FReal SumDeltaCost = SumDeltaCostQ[QIndex];
+			QIndex++;
+
+			// Alternative is a stack (this is not very optimal so don't use it)
+			//uint32 TestSibling = PriorityQ[PriorityQ.Num() - 1];
+			//FReal SumDeltaCost = SumDeltaCostQ[SumDeltaCostQ.Num() - 1];
+			//PriorityQ.SetNumUninitialized(PriorityQ.Num() - 1);
+			//SumDeltaCostQ.SetNumUninitialized(SumDeltaCostQ.Num() - 1);
+
+			// TestSibling bounds union with new bounds
+			WorkingAABB = Nodes[TestSibling].ChildrenBounds[0];
+			if (!Nodes[TestSibling].bLeaf)
+			{
+				WorkingAABB.GrowToInclude(Nodes[TestSibling].ChildrenBounds[1]);
+			}
+			FReal TestSiblingArea = WorkingAABB.GetArea();
+			WorkingAABB.GrowToInclude(NewBounds);
+
+			FReal NewPotentialNodeArea = WorkingAABB.GetArea();
+			FReal CostForChoosingNode = NewPotentialNodeArea + SumDeltaCost;
+			FReal NewDeltaCost = NewPotentialNodeArea - TestSiblingArea;
+			// Did we get a better cost?
+			if (CostForChoosingNode < BestCost)
+			{
+				BestCost = CostForChoosingNode;
+				BestSiblingIdx = TestSibling;
+			}
+
+			// Lower bound of Children costs
+			FReal ChildCostLowerBound = NewBounds.GetArea() + NewDeltaCost + SumDeltaCost;
+
+			if (!Nodes[TestSibling].bLeaf && ChildCostLowerBound < BestCost)
+			{
+				// Now we will push the children
+				PriorityQ.Add(Nodes[TestSibling].ChildrenNodes[0]);
+				PriorityQ.Add(Nodes[TestSibling].ChildrenNodes[1]);
+				SumDeltaCostQ.Add(NewDeltaCost + SumDeltaCost);
+				SumDeltaCostQ.Add(NewDeltaCost + SumDeltaCost);
+			}
+
+		}
+
+		return BestSiblingIdx;
+	}
+
+	// Rotate nodes to decrease tree cost
+	// Grandchildren can swap with their aunts
+	void RotateNode(uint32 NodeIdx, bool debugAssert = false)
+	{
+		int32 BestGrandChildToSwap = INDEX_NONE; // GrandChild of NodeIdx
+		int32 BestAuntToSwap = INDEX_NONE; // Aunt of BestGrandChildToSwap
+		FReal BestDeltaCost = 0.0f; // Negative values are cost reductions, doing nothing changes the cost with 0
+
+		check(!Nodes[NodeIdx].bLeaf);
+		// Check both children of NodeIdx
+		for (uint32 AuntLocalIdx = 0; AuntLocalIdx < 2; AuntLocalIdx++)
+		{
+			int32 Aunt = Nodes[NodeIdx].ChildrenNodes[AuntLocalIdx];
+			int32 Mother = Nodes[NodeIdx].ChildrenNodes[AuntLocalIdx ^ 1];
+			if (Nodes[Mother].bLeaf)
+			{
+				continue;
+			}
+			for (int32 GrandChild : Nodes[Mother].ChildrenNodes)
+			{
+				// Only the Mother's cost will change
+				TAABB<T, 3> NewMotherAABB{ Nodes[NodeIdx].ChildrenBounds[AuntLocalIdx]}; // Aunt will be under mother now
+				NewMotherAABB.GrowToInclude(Nodes[Mother].ChildrenBounds[GetSiblingIndex(GrandChild)]); // Add the Grandchild's sibling cost
+				FReal MotherCostWithoutRotation = Nodes[NodeIdx].ChildrenBounds[AuntLocalIdx ^ 1].GetArea();
+				FReal MotherCostWithRotation = NewMotherAABB.GetArea();
+				FReal DeltaCost = MotherCostWithRotation - MotherCostWithoutRotation;
+
+				if (DeltaCost < BestDeltaCost)
+				{
+					BestDeltaCost = DeltaCost;
+					BestAuntToSwap = Aunt;
+					BestGrandChildToSwap = GrandChild;
+				}
+
+			}
+		}
+
+		// Now do the rotation if required
+		if (BestGrandChildToSwap != INDEX_NONE)
+		{
+			if (debugAssert)
+			{
+				check(false);
+			}
+
+			int32 AuntLocalChildIdx = WhichChildAmI(BestAuntToSwap);
+			int32 GrandChildLocalChildIdx = WhichChildAmI(BestGrandChildToSwap);
+
+			int32 MotherOfBestGrandChild = Nodes[BestGrandChildToSwap].ParentNode;
+
+			// Modify NodeIdx 
+			Nodes[NodeIdx].ChildrenNodes[AuntLocalChildIdx] = BestGrandChildToSwap;
+			// Modify BestGrandChildToSwap
+			Nodes[BestGrandChildToSwap].ParentNode = NodeIdx;
+			// Modify BestAuntToSwap
+			Nodes[BestAuntToSwap].ParentNode = MotherOfBestGrandChild;
+			// Modify MotherOfBestGrandChild
+			Nodes[MotherOfBestGrandChild].ChildrenNodes[GrandChildLocalChildIdx] = BestAuntToSwap;
+			// Swap the bounds
+			TAABB<T, 3> AuntAABB = Nodes[NodeIdx].ChildrenBounds[AuntLocalChildIdx];
+			Nodes[NodeIdx].ChildrenBounds[AuntLocalChildIdx] = Nodes[MotherOfBestGrandChild].ChildrenBounds[GrandChildLocalChildIdx];
+			Nodes[MotherOfBestGrandChild].ChildrenBounds[GrandChildLocalChildIdx] = AuntAABB;
+			// Update the other child bound of NodeIdx
+			Nodes[NodeIdx].ChildrenBounds[AuntLocalChildIdx ^ 1] = Nodes[MotherOfBestGrandChild].ChildrenBounds[0];
+			Nodes[NodeIdx].ChildrenBounds[AuntLocalChildIdx ^ 1].GrowToInclude(Nodes[MotherOfBestGrandChild].ChildrenBounds[1]);
+		}
+	}
+
+	void InsertLeaf(const TPayloadType& Payload, const TAABB<T, 3>& NewBounds)
+	{
+
+		// Slow Debug Code
+		//if (GetUniqueIdx(Payload).Idx == 5)
+		//{
+		//	DynamicTreeDebugStats();
+		//}
+
+		int32 NewLeafNode = AllocateLeafNodeAndLeaf(Payload, NewBounds);
+
+		// New tree?
+		if (RootNode == INDEX_NONE)
+		{
+			RootNode = NewLeafNode;
+			return;
+		}
+
+		// Find the best sibling
+		int32 BestSibling = FindBestSibling(NewBounds);
+
+		// New internal parent node
+		int32 oldParent = Nodes[BestSibling].ParentNode;
+		int32 newParent = AllocateInternalNode();
+		Nodes[newParent].ParentNode = oldParent;
+		Nodes[newParent].ChildrenNodes[0] = BestSibling;
+		Nodes[newParent].ChildrenNodes[1] = NewLeafNode;
+		Nodes[newParent].ChildrenBounds[0] = Nodes[BestSibling].ChildrenBounds[0];
+		if (!Nodes[BestSibling].bLeaf)
+		{
+			Nodes[newParent].ChildrenBounds[0].GrowToInclude(Nodes[BestSibling].ChildrenBounds[1]);
+		}
+		Nodes[newParent].ChildrenBounds[1] = Nodes[NewLeafNode].ChildrenBounds[0];
+
+		if (oldParent != INDEX_NONE)
+		{
+			int32 ChildIdx = WhichChildAmI(BestSibling);
+			Nodes[oldParent].ChildrenNodes[ChildIdx] = newParent;
+		}
+		else
+		{
+			RootNode = newParent;
+		}
+
+		Nodes[BestSibling].ParentNode = newParent;
+		Nodes[NewLeafNode].ParentNode = newParent;
+
+		UpdateAncestorBounds(newParent, true);
+
+	}	
+
+	void  UpdateAncestorBounds(int32 NodeIdx, bool bDoRotation = false)
+	{
+		int32 CurrentNodeIdx = NodeIdx;
+		int32 ParentNodeIdx = Nodes[NodeIdx].ParentNode;
+
+
+		// This should not be required
+		/*if (bDoRotation && NodeIdx != INDEX_NONE)
+		{
+			RotateNode(NodeIdx,true); 
+		}*/
+		
+		while (ParentNodeIdx != INDEX_NONE)
+		{
+			int32 ChildIndex = WhichChildAmI(CurrentNodeIdx);
+			Nodes[ParentNodeIdx].ChildrenBounds[ChildIndex] = Nodes[CurrentNodeIdx].ChildrenBounds[0];
+			if (!Nodes[CurrentNodeIdx].bLeaf)
+			{
+				Nodes[ParentNodeIdx].ChildrenBounds[ChildIndex].GrowToInclude(Nodes[CurrentNodeIdx].ChildrenBounds[1]);
+			}
+
+			if (bDoRotation)
+			{
+				RotateNode(ParentNodeIdx);
+			}
+
+			CurrentNodeIdx = ParentNodeIdx;
+			ParentNodeIdx = Nodes[CurrentNodeIdx].ParentNode;
+		}
+	}
+
+	void RemoveLeafNode(int32 LeafNodeIdx)
+	{
+		int32 ParentNodeIdx = Nodes[LeafNodeIdx].ParentNode;
+
+		if (ParentNodeIdx != INDEX_NONE)
+		{
+			int32 GrandParentNodeIdx = Nodes[ParentNodeIdx].ParentNode;
+			int32 SiblingNodeLocalIdx = GetSiblingIndex(LeafNodeIdx);
+			int32 SiblingNodeIdx = Nodes[ParentNodeIdx].ChildrenNodes[SiblingNodeLocalIdx];
+
+			if (GrandParentNodeIdx != INDEX_NONE)
+			{
+				int32 ChildLocalIdx = WhichChildAmI(ParentNodeIdx);
+				Nodes[GrandParentNodeIdx].ChildrenNodes[ChildLocalIdx] = SiblingNodeIdx;
+			}
+			else
+			{
+				RootNode = SiblingNodeIdx;
+			}
+			Nodes[SiblingNodeIdx].ParentNode = GrandParentNodeIdx;
+			UpdateAncestorBounds(SiblingNodeIdx);
+			DeAllocateInternalNode(ParentNodeIdx);
+		}
+		else
+		{
+			RootNode = INDEX_NONE;
+		}
+		DeAllocateLeafNode(LeafNodeIdx);
+	}
+
 	virtual void RemoveElement(const TPayloadType& Payload)
 	{
 		if (ensure(bMutable))
@@ -887,7 +1323,14 @@ public:
 				}
 				else if (ensure(PayloadInfo->LeafIdx != INDEX_NONE))
 				{
-					Leaves[PayloadInfo->LeafIdx].RemoveElement(Payload);
+					if (bDynamicTree)
+					{
+						RemoveLeafNode(PayloadInfo->NodeIdx);
+					}
+					else
+					{
+						Leaves[PayloadInfo->LeafIdx].RemoveElement(Payload);
+					}
 				}
 
 				PayloadToInfo.Remove(Payload);
@@ -896,9 +1339,23 @@ public:
 		}
 	}
 
-	virtual void UpdateElement(const TPayloadType& Payload, const TAABB<T, 3>& NewBounds, bool bHasBounds) override
+	virtual void UpdateElement(const TPayloadType& Payload, const TAABB<T, 3>& NewBounds, bool bInHasBounds) override
 	{
+#if !WITH_EDITOR
 		//CSV_SCOPED_TIMING_STAT(ChaosPhysicsTimers, AABBTreeUpdateElement)
+		//CSV_CUSTOM_STAT(ChaosPhysicsTimers, 1, 1, ECsvCustomStatOp::Accumulate);
+		//CSV_CUSTOM_STAT(PhysicsCounters, NumIntoNP, 1, ECsvCustomStatOp::Accumulate);
+#endif
+
+		bool bHasBounds = bInHasBounds;
+		// If bounds are bad, use global
+		if (bDynamicTree && bHasBounds && ValidateBounds(NewBounds) == false)
+		{
+			bHasBounds = false;
+			ensureMsgf(false, TEXT("AABBTree encountered invalid bounds input. Forcing element to global payload. Min: %s Max: %s. If Bounds are valid but large, increase FAABBTreeCVars::MaxNonGlobalElementBoundsExtrema."),
+				*NewBounds.Min().ToString(), *NewBounds.Max().ToString());
+		}
+
 		if (ensure(bMutable))
 		{
 			FAABBTreePayloadInfo* PayloadInfo = PayloadToInfo.Find(Payload);
@@ -909,17 +1366,42 @@ public:
 					//If we are still within the same leaf bounds, do nothing, don't detect a change either
 					if (bHasBounds) 
 					{
-						const TAABB<T, 3>& LeafBounds = Leaves[PayloadInfo->LeafIdx].GetBounds();
-						if (LeafBounds.Contains(NewBounds.Min()) && LeafBounds.Contains(NewBounds.Max()))
+						if (bDynamicTree)
+						{
+							// The leaf node bounds can be larger than the actual leave bound
+							const TAABB<T, 3>& LeafNodeBounds = Nodes[PayloadInfo->NodeIdx].ChildrenBounds[0];
+							if (LeafNodeBounds.Contains(NewBounds.Min()) && LeafNodeBounds.Contains(NewBounds.Max()))
+							{
+								// We still need to update the constituent bounds
+								Leaves[PayloadInfo->LeafIdx].UpdateElement(Payload, NewBounds, bHasBounds);
+								return;
+							}
+						}
+						else
 						{
-							// We still need to update the constituent bounds
-							Leaves[PayloadInfo->LeafIdx].UpdateElement(Payload, NewBounds, bHasBounds);
-							return;
+							const TAABB<T, 3>& LeafBounds = Leaves[PayloadInfo->LeafIdx].GetBounds();
+							if (LeafBounds.Contains(NewBounds.Min()) && LeafBounds.Contains(NewBounds.Max()))
+							{
+								// We still need to update the constituent bounds
+								Leaves[PayloadInfo->LeafIdx].UpdateElement(Payload, NewBounds, bHasBounds);
+								return;
+							}
 						}
 					}
 
-					Leaves[PayloadInfo->LeafIdx].RemoveElement(Payload);
-					PayloadInfo->LeafIdx = INDEX_NONE;
+					// DBVH remove from tree
+
+					if (bDynamicTree)
+					{
+						RemoveLeafNode(PayloadInfo->NodeIdx);
+						PayloadInfo->LeafIdx = INDEX_NONE;
+						PayloadInfo->NodeIdx = INDEX_NONE;
+					}
+					else
+					{
+						Leaves[PayloadInfo->LeafIdx].RemoveElement(Payload);
+						PayloadInfo->LeafIdx = INDEX_NONE;
+					}
 				}
 			}
 			else
@@ -943,11 +1425,18 @@ public:
 			{
 				if (PayloadInfo->DirtyPayloadIdx == INDEX_NONE)
 				{
-					PayloadInfo->DirtyPayloadIdx = DirtyElements.Add(FElement{ Payload, NewBounds });
-					if (DirtyElementGridEnabled())
+					if (bDynamicTree)
 					{
-						//CSV_SCOPED_TIMING_STAT(ChaosPhysicsTimers, AABBAddElement)
-						PayloadInfo->DirtyGridOverflowIdx = AddDirtyElementToGrid(NewBounds, PayloadInfo->DirtyPayloadIdx);
+						InsertLeaf(Payload, NewBounds);
+					}
+					else
+					{
+						PayloadInfo->DirtyPayloadIdx = DirtyElements.Add(FElement{ Payload, NewBounds });
+						if (DirtyElementGridEnabled())
+						{
+							//CSV_SCOPED_TIMING_STAT(ChaosPhysicsTimers, AABBAddElement)
+							PayloadInfo->DirtyGridOverflowIdx = AddDirtyElementToGrid(NewBounds, PayloadInfo->DirtyPayloadIdx);
+						}
 					}
 				}
 				else
@@ -1011,7 +1500,7 @@ public:
 			}
 		}
 
-		if(DirtyElements.Num() > MaxDirtyElements)
+		if(!bDynamicTree && DirtyElements.Num() > MaxDirtyElements)
 		{
 			UE_LOG(LogChaos, Verbose, TEXT("Bounding volume exceeded maximum dirty elements (%d dirty of max %d) and is forcing a tree rebuild."), DirtyElements.Num(), MaxDirtyElements);
 			ReoptimizeTree();
@@ -1067,10 +1556,13 @@ public:
 	}
 
 
-	virtual bool ShouldRebuild() override { return bShouldRebuild; }  // Used to find out if something changed since last reset for optimizations
+	virtual bool ShouldRebuild() override { return bDynamicTree ? false : bShouldRebuild; }  // Used to find out if something changed since last reset for optimizations
 	// Contract: bShouldRebuild can only ever be cleared by calling the ClearShouldRebuild method, it can be set at will though
 	virtual void ClearShouldRebuild() override { bShouldRebuild = false; }
 
+	virtual bool IsTreeDynamic() const override { return bDynamicTree; }
+	bool SetTreeToDynamic() { bDynamicTree = true; } // Tree cannot be changed back to static for now
+
 	virtual void PrepareCopyTimeSliced(const  ISpatialAcceleration<TPayloadType, T, 3>& InFrom) override
 	{
 
@@ -1097,6 +1589,10 @@ public:
 		NumProcessedThisSlice = From.NumProcessedThisSlice;
 		bShouldRebuild = From.bShouldRebuild;
 
+		RootNode = From.RootNode;
+		FirstFreeInternalNode = From.FirstFreeInternalNode;
+		FirstFreeLeafNode = From.FirstFreeLeafNode;
+
 		// Reserve sizes for arrays etc
 
 		Nodes.Reserve(From.Nodes.Num());
@@ -1232,6 +1728,19 @@ public:
 			DirtyElementGridCellSizeInv = 1.0f;
 			bShouldRebuild = true;
 		}
+
+		// Dynamic trees are not serialized/deserialized for now
+		if (Ar.IsLoading())
+		{
+			bDynamicTree = false;
+			RootNode = INDEX_NONE;
+			FirstFreeInternalNode = INDEX_NONE;
+			FirstFreeLeafNode = INDEX_NONE;
+		}
+		else
+		{
+			ensure(bDynamicTree == false);
+		}
 	}
 
 private:
@@ -1260,7 +1769,7 @@ private:
 			Leaf.GatherElements(AllElements);
 		}
 
-		TAABBTree<TPayloadType,TLeafType,bMutable,T> NewTree(AllElements);
+		TAABBTree<TPayloadType,TLeafType,bMutable, T> NewTree(AllElements);
 		*this = NewTree;
 		bShouldRebuild = true; // No changes since last time tree was built
 	}
@@ -1351,6 +1860,9 @@ private:
 		PHYSICS_CSV_CUSTOM_VERY_EXPENSIVE(PhysicsCounters, MaxNumLeaves, Leaves.Num(), ECsvCustomStatOp::Max);
 		PHYSICS_CSV_SCOPED_VERY_EXPENSIVE(PhysicsVerbose, QueryImp);
 		//QUICK_SCOPE_CYCLE_COUNTER(AABBTreeQueryImp);
+#if !WITH_EDITOR
+		//CSV_SCOPED_TIMING_STAT(ChaosPhysicsTimers, AABBTreeQuery)
+#endif
 		FVec3 TmpPosition;
 		FReal TOI = 0;
 		{
@@ -1500,15 +2012,39 @@ private:
 		};
 
 		TArray<FNodeQueueEntry> NodeStack;
-		if (Nodes.Num())
+		if (bDynamicTree)
+		{
+
+			if (RootNode != INDEX_NONE)
+			{
+				NodeStack.Add(FNodeQueueEntry{ RootNode, 0 });
+			}
+		}
+		else if (Nodes.Num())
 		{
 			NodeStack.Add(FNodeQueueEntry{ 0, 0 });
 		}
 
+// Slow debug code
+//#if !WITH_EDITOR
+//		if (Query == EAABBQueryType::Overlap)
+//		{
+//			CSV_CUSTOM_STAT(ChaosPhysicsTimers, OverlapCount, 1, ECsvCustomStatOp::Accumulate);
+//			CSV_CUSTOM_STAT(ChaosPhysicsTimers, DirtyCount, DirtyElements.Num(), ECsvCustomStatOp::Max);
+//		}
+//#endif
+
 		while (NodeStack.Num())
 		{
 			PHYSICS_CSV_SCOPED_VERY_EXPENSIVE(PhysicsVerbose, QueryImp_NodeTraverse);
 
+//#if !WITH_EDITOR
+//			if (Query == EAABBQueryType::Overlap)
+//			{
+//				CSV_CUSTOM_STAT(ChaosPhysicsTimers, AABBCheckCount, 1, ECsvCustomStatOp::Accumulate);
+//			}
+//#endif
+
 			const FNodeQueueEntry NodeEntry = NodeStack.Pop(false);
 			if (Query != EAABBQueryType::Overlap)
 			{
@@ -1601,6 +2137,9 @@ private:
 		GlobalPayloads.Reset();
 		Leaves.Reset();
 		Nodes.Reset();
+		RootNode = INDEX_NONE;
+		FirstFreeInternalNode = INDEX_NONE;
+		FirstFreeLeafNode = INDEX_NONE;
 		DirtyElements.Reset();
 		CellHashToFlatArray.Reset(); 
 		FlattenedCellArrayOfDirtyIndices.Reset();
@@ -1611,6 +2150,23 @@ private:
 		NumProcessedThisSlice = 0;
 		GetCVars();  // Safe to copy CVARS here
 
+		if (bDynamicTree)
+		{
+			// Todo for now, don't ever return async task
+			int32 Idx = 0;
+			for (auto& Particle : Particles)
+			{
+				bool bHasBoundingBox = HasBoundingBox(Particle);
+				auto Payload = Particle.template GetPayload<TPayloadType>(Idx);
+				TAABB<T, 3> ElemBounds = ComputeWorldSpaceBoundingBox(Particle, false, (T)0);
+				
+				UpdateElement(Payload, ElemBounds, bHasBoundingBox);
+				++Idx;
+			}
+			this->SetAsyncTimeSlicingComplete(true);
+			return;
+		}
+
 		WorkSnapshot.Bounds = TAABB<T, 3>::EmptyAABB();
 
 		{
@@ -2067,6 +2623,10 @@ private:
 		, Nodes(Other.Nodes)
 		, Leaves(Other.Leaves)
 		, DirtyElements(Other.DirtyElements)
+		, bDynamicTree(Other.bDynamicTree)
+		, RootNode(Other.RootNode)
+		, FirstFreeInternalNode(Other.FirstFreeInternalNode)
+		, FirstFreeLeafNode(Other.FirstFreeLeafNode)
 		, CellHashToFlatArray(Other.CellHashToFlatArray)
 		, FlattenedCellArrayOfDirtyIndices(Other.FlattenedCellArrayOfDirtyIndices)
 		, DirtyElementsGridOverflow(Other.DirtyElementsGridOverflow)
@@ -2087,6 +2647,8 @@ private:
 		, bShouldRebuild(Other.bShouldRebuild)
 	{
 
+		ensure(bDynamicTree == Other.IsTreeDynamic());
+
 	}
 
 	virtual ISpatialAcceleration<TPayloadType, T, 3>& operator=(const ISpatialAcceleration<TPayloadType, T, 3>& Other) override
@@ -2110,6 +2672,10 @@ private:
 			Nodes = Rhs.Nodes;
 			Leaves = Rhs.Leaves;
 			DirtyElements = Rhs.DirtyElements;
+			bDynamicTree = Rhs.bDynamicTree;
+			RootNode = Rhs.RootNode;
+			FirstFreeInternalNode = Rhs.FirstFreeInternalNode;
+			FirstFreeLeafNode = Rhs.FirstFreeLeafNode;
 			
 			CellHashToFlatArray = Rhs.CellHashToFlatArray;
 			FlattenedCellArrayOfDirtyIndices = Rhs.FlattenedCellArrayOfDirtyIndices;
@@ -2140,6 +2706,13 @@ private:
 	TArray<TLeafType> Leaves;
 	TArray<FElement> DirtyElements;
 
+	// DynamicTree members
+	bool bDynamicTree = false;
+	int32 RootNode = INDEX_NONE;
+	// Free lists (indices are in Nodes array)
+	int32 FirstFreeInternalNode = INDEX_NONE;
+	int32 FirstFreeLeafNode = INDEX_NONE;
+
 	// Data needed for DirtyElement2DAccelerationGrid
 	TMap<int32, DirtyGridHashEntry> CellHashToFlatArray; // Index, size into flat grid structure (FlattenedCellArrayOfDirtyIndices)
 	TArray<int32> FlattenedCellArrayOfDirtyIndices; // Array of indices into dirty Elements array, indices are always sorted within a 2D cell
diff --git a/Engine/Source/Runtime/Experimental/Chaos/Public/Chaos/BoundingVolume.h b/Engine/Source/Runtime/Experimental/Chaos/Public/Chaos/BoundingVolume.h
--- a/Engine/Source/Runtime/Experimental/Chaos/Public/Chaos/BoundingVolume.h
+++ b/Engine/Source/Runtime/Experimental/Chaos/Public/Chaos/BoundingVolume.h
@@ -282,6 +282,13 @@ public:
 		}
 	}
 
+	inline void AddElement(const TPayloadBoundsElement<TPayloadType, T>& Payload)
+	{
+		// Not implemented 
+		check(false);
+	}
+
+	
 	// Begin ISpatialAcceleration interface
 	virtual TArray<TPayloadType> FindAllIntersections(const FAABB3& Box) const override { return FindAllIntersectionsImp(Box); }
 
diff --git a/Engine/Source/Runtime/Experimental/Chaos/Public/Chaos/ISpatialAcceleration.h b/Engine/Source/Runtime/Experimental/Chaos/Public/Chaos/ISpatialAcceleration.h
--- a/Engine/Source/Runtime/Experimental/Chaos/Public/Chaos/ISpatialAcceleration.h
+++ b/Engine/Source/Runtime/Experimental/Chaos/Public/Chaos/ISpatialAcceleration.h
@@ -264,6 +264,7 @@ public:
 	virtual bool IsAsyncTimeSlicingComplete() { return AsyncTimeSlicingComplete; }
 	virtual void ProgressAsyncTimeSlicing(bool ForceBuildCompletion = false) {}
 	virtual bool ShouldRebuild() { return true; }  // Used to find out if something changed since last reset for optimizations
+	virtual bool IsTreeDynamic() const { return false; }  // Dynamic trees rebuild on the fly without adding dirty elements
 	virtual void ClearShouldRebuild() {}
 	virtual void PrepareCopyTimeSliced(const  ISpatialAcceleration<TPayloadType, T, 3>& InFrom) { check(false); }
 	virtual void ProgressCopyTimeSliced(const  ISpatialAcceleration<TPayloadType, T, 3>& InFrom, int MaximumBytesToCopy) { check(false); }
diff --git a/Engine/Source/Runtime/Experimental/Chaos/Public/Chaos/PBDRigidsEvolution.h b/Engine/Source/Runtime/Experimental/Chaos/Public/Chaos/PBDRigidsEvolution.h
--- a/Engine/Source/Runtime/Experimental/Chaos/Public/Chaos/PBDRigidsEvolution.h
+++ b/Engine/Source/Runtime/Experimental/Chaos/Public/Chaos/PBDRigidsEvolution.h
@@ -220,7 +220,7 @@ struct CHAOS_API ISpatialAccelerationCollectionFactory
 	virtual bool IsBucketTimeSliced(uint16 BucketIdx) const = 0;
 
 	//Chaos creates new acceleration structures per bucket. Factory can change underlying type at runtime as well as number of buckets to AB test
-	virtual TUniquePtr<ISpatialAcceleration<FAccelerationStructureHandle, FReal, 3>> CreateAccelerationPerBucket_Threaded(const TConstParticleView<FSpatialAccelerationCache>& Particles, uint16 BucketIdx, bool ForceFullBuild) = 0;
+	virtual TUniquePtr<ISpatialAcceleration<FAccelerationStructureHandle, FReal, 3>> CreateAccelerationPerBucket_Threaded(const TConstParticleView<FSpatialAccelerationCache>& Particles, uint16 BucketIdx, bool ForceFullBuild, bool bDynamicTree) = 0;
 
 	//Mask indicating which bucket is active. Spatial indices in inactive buckets fallback to bucket 0. Bit 0 indicates bucket 0 is active, Bit 1 indicates bucket 1 is active, etc...
 	virtual uint8 GetActiveBucketsMask() const = 0;
