Chaos - Heightfield overlap optimization
- use cell bounds culling to avoid calling GJK unnecessarily
- when MTD is not needed, transform the corner vertices in the space of the query geometry and use a simpler form of the GJKIntersection
- avoid unecessary clipping of the cell coordinateas they are garanteed to be within the limits
- 25% improvement on heighfield overlaps

diff --git a/Engine/Source/Runtime/Experimental/Chaos/Private/Chaos/HeightField.cpp b/Engine/Source/Runtime/Experimental/Chaos/Private/Chaos/HeightField.cpp
--- a/Engine/Source/Runtime/Experimental/Chaos/Private/Chaos/HeightField.cpp
+++ b/Engine/Source/Runtime/Experimental/Chaos/Private/Chaos/HeightField.cpp
@@ -1311,7 +1311,9 @@ namespace Chaos
 		{
 			for(int32 CurrY = 0; CurrY < NumY; ++CurrY)
 			{
-				OutIntersections.Add(FlatGrid.ClampIndex(TVec2<int32>(MinCell[0] + CurrX, MinCell[1] + CurrY)));
+				const TVec2<int32> Cell(MinCell[0] + CurrX, MinCell[1] + CurrY);
+				check(FlatGrid.IsValid(Cell));
+				OutIntersections.Add(Cell);
 			}
 		}
 
@@ -1731,7 +1733,7 @@ namespace Chaos
 			OutMTD->Penetration = TNumericLimits<FReal>::Lowest();
 		}
 
-		auto OverlapTriangle = [&](const FVec3& A, const FVec3& B, const FVec3& C, FMTDInfo* InnerMTD) -> bool
+		auto OverlapTriangleMTD = [&](const FVec3& A, const FVec3& B, const FVec3& C, FMTDInfo* InnerMTD) -> bool
 		{
 			const FVec3 AB = B - A;
 			const FVec3 AC = C - A;
@@ -1741,30 +1743,37 @@ namespace Chaos
 			const FVec3 Offset = FVec3::CrossProduct(AB, AC);
 
 			FTriangle TriangleConvex(A, B, C);
-			if (InnerMTD)
+			FVec3 TriangleNormal(0);
+			FReal Penetration = 0;
+			FVec3 ClosestA(0);
+			FVec3 ClosestB(0);
+			int32 ClosestVertexIndexA, ClosestVertexIndexB;
+			if (GJKPenetration(TriangleConvex, QueryGeom, QueryTM, Penetration, ClosestA, ClosestB, TriangleNormal, ClosestVertexIndexA, ClosestVertexIndexB, Thickness))
 			{
-				FVec3 TriangleNormal(0);
-				FReal Penetration = 0;
-				FVec3 ClosestA(0);
-				FVec3 ClosestB(0);
-				int32 ClosestVertexIndexA, ClosestVertexIndexB;
-				if (GJKPenetration(TriangleConvex, QueryGeom, QueryTM, Penetration, ClosestA, ClosestB, TriangleNormal, ClosestVertexIndexA, ClosestVertexIndexB, Thickness))
+				// Use Deepest MTD.
+				if (Penetration > InnerMTD->Penetration)
 				{
-					// Use Deepest MTD.
-					if (Penetration > InnerMTD->Penetration)
-					{
-						InnerMTD->Penetration = Penetration;
-						InnerMTD->Normal = TriangleNormal;
-					}
-					return true;
+					InnerMTD->Penetration = Penetration;
+					InnerMTD->Normal = TriangleNormal;
 				}
-
-				return false;
-			}
-			else
-			{
-				return GJKIntersection(TriangleConvex, QueryGeom, QueryTM, Thickness, Offset);
+				return true;
 			}
+
+			return false;
+		};
+
+		auto OverlapTriangleNoMTD = [&](const FVec3& A, const FVec3& B, const FVec3& C) -> bool
+		{
+			// points are assumed to be in the same space as the overlap geometry
+			const FVec3 AB = B - A;
+			const FVec3 AC = C - A;
+
+			//It's most likely that the query object is in front of the triangle since queries tend to be on the outside.
+			//However, maybe we should check if it's behind the triangle plane. Also, we should enforce this winding in some way
+			const FVec3 Offset = FVec3::CrossProduct(AB, AC);
+
+			FTriangle TriangleConvex(A, B, C);
+			return GJKIntersectionSameSpace(TriangleConvex, QueryGeom, Thickness, Offset);
 		};
 
 		bool bResult = false;
@@ -1782,31 +1791,46 @@ namespace Chaos
 		GetGridIntersections(FlatQueryBounds, Intersections);
 
 		bool bOverlaps = false;
-		for(const TVec2<int32>& Cell : Intersections)
+		if (OutMTD)
 		{
-			const int32 SingleIndex = Cell[1] * (GeomData.NumCols) + Cell[0];
-			GeomData.GetPointsScaled(SingleIndex, Points);
-
-			if(OverlapTriangle(Points[0], Points[1], Points[3], OutMTD))
+			for (const TVec2<int32>& Cell : Intersections)
 			{
-				bOverlaps = true;
-				if (!OutMTD)
-				{
-					return true;
-				}
-			}
+				const int32 SingleIndex = Cell[1] * (GeomData.NumCols) + Cell[0];
+				GeomData.GetPointsScaled(SingleIndex, Points);
 
-			if(OverlapTriangle(Points[0], Points[3], Points[2], OutMTD))
+				bOverlaps |= OverlapTriangleMTD(Points[0], Points[1], Points[3], OutMTD);
+				bOverlaps |= OverlapTriangleMTD(Points[0], Points[3], Points[2], OutMTD);
+			}
+			return bOverlaps;
+		}
+		else
+		{
+			FAABB3 CellBounds;
+			for (const TVec2<int32>& Cell : Intersections)
 			{
-				bOverlaps = true;
-				if (!OutMTD)
+				const int32 SingleIndex = Cell[1] * (GeomData.NumCols) + Cell[0];
+				GeomData.GetPointsAndBoundsScaled(SingleIndex, Points, CellBounds);
+
+				if (CellBounds.Intersects(QueryBounds))
 				{
-					return true;
+					// pre-transform the triangle in overlap geometry space
+					Points[0] = QueryTM.InverseTransformPositionNoScale(Points[0]);
+					Points[1] = QueryTM.InverseTransformPositionNoScale(Points[1]);
+					Points[2] = QueryTM.InverseTransformPositionNoScale(Points[2]);
+					Points[3] = QueryTM.InverseTransformPositionNoScale(Points[3]);
+
+					if (OverlapTriangleNoMTD(Points[0], Points[1], Points[3]))
+					{
+						return true;
+					}
+					if (OverlapTriangleNoMTD(Points[0], Points[3], Points[2]))
+					{
+						return true;
+					}
 				}
 			}
+			return false;
 		}
-
-		return bOverlaps;
 	}
 
 	bool FHeightField::OverlapGeom(const TSphere<FReal, 3>& QueryGeom, const FRigidTransform3& QueryTM, const FReal Thickness, FMTDInfo* OutMTD) const
diff --git a/Engine/Source/Runtime/Experimental/Chaos/Public/Chaos/GJK.h b/Engine/Source/Runtime/Experimental/Chaos/Public/Chaos/GJK.h
--- a/Engine/Source/Runtime/Experimental/Chaos/Public/Chaos/GJK.h
+++ b/Engine/Source/Runtime/Experimental/Chaos/Public/Chaos/GJK.h
@@ -15,15 +15,17 @@
 namespace Chaos
 {
 	
-	/** Determines if two convex geometries overlap.
-	 @A The first geometry
-	 @B The second geometry
-	 @BToATM The transform of B in A's local space
-	 @ThicknessA The amount of geometry inflation for Geometry A(for example if the surface distance of two geometries with thickness 0 would be 2, a thickness of 0.5 would give a distance of 1.5)
-	 @InitialDir The first direction we use to search the CSO
-	 @ThicknessB The amount of geometry inflation for Geometry B(for example if the surface distance of two geometries with thickness 0 would be 2, a thickness of 0.5 would give a distance of 1.5)
-	 @return True if the geometries overlap, False otherwise */
-
+	/** 
+		Determines if two convex geometries overlap.
+		
+		@param A The first geometry
+		@param B The second geometry
+		@param BToATM The transform of B in A's local space
+		@param ThicknessA The amount of geometry inflation for Geometry A(for example if the surface distance of two geometries with thickness 0 would be 2, a thickness of 0.5 would give a distance of 1.5)
+		@param InitialDir The first direction we use to search the CSO
+		@param ThicknessB The amount of geometry inflation for Geometry B(for example if the surface distance of two geometries with thickness 0 would be 2, a thickness of 0.5 would give a distance of 1.5)
+		@return True if the geometries overlap, False otherwise 
+	 */
 	template <typename T, typename TGeometryA, typename TGeometryB>
 	bool GJKIntersection(const TGeometryA& A, const TGeometryB& B, const TRigidTransform<T, 3>& BToATM, const T InThicknessA = 0, const TVector<T, 3>& InitialDir = TVector<T, 3>(1, 0, 0), const T InThicknessB = 0)
 	{
@@ -87,7 +89,80 @@ namespace Chaos
 		} while (!bTerminate);
 
 		return bNearZero;
-		
+	}
+
+	/** 
+		Determines if two convex geometries in the same space overlap
+		IMPORTANT: the two convex geometries must be in the same space!
+
+		@param A The first geometry
+		@param B The second geometry
+		@param ThicknessA The amount of geometry inflation for Geometry A(for example if the surface distance of two geometries with thickness 0 would be 2, a thickness of 0.5 would give a distance of 1.5)
+		@param InitialDir The first direction we use to search the CSO
+		@param ThicknessB The amount of geometry inflation for Geometry B(for example if the surface distance of two geometries with thickness 0 would be 2, a thickness of 0.5 would give a distance of 1.5)
+		@return True if the geometries overlap, False otherwise 
+	*/
+	template <typename T, typename TGeometryA, typename TGeometryB>
+	bool GJKIntersectionSameSpace(const TGeometryA& A, const TGeometryB& B, const T InThicknessA = 0, const TVector<T, 3>& InitialDir = TVector<T, 3>(1, 0, 0), const T InThicknessB = 0)
+	{
+		TVector<T, 3> V = -InitialDir;
+		if (V.SafeNormalize() == 0)
+		{
+			V = TVec3<T>(-1, 0, 0);
+		}
+
+		FSimplex SimplexIDs;
+		TVector<T, 3> Simplex[4];
+		T Barycentric[4] = { -1,-1,-1,-1 };	//not needed, but compiler warns
+		bool bTerminate;
+		bool bNearZero = false;
+		int NumIterations = 0;
+		T PrevDist2 = FLT_MAX;
+		const T ThicknessA = A.GetMargin() + InThicknessA;
+		const T ThicknessB = B.GetMargin() + InThicknessB;
+		const T Inflation = ThicknessA + ThicknessB + static_cast<T>(1e-3);
+		const T Inflation2 = Inflation * Inflation;
+		int32 VertexIndexA = INDEX_NONE, VertexIndexB = INDEX_NONE;
+		do
+		{
+			if (!ensure(NumIterations++ < 32))	//todo: take this out
+			{
+				break;	//if taking too long just stop. This should never happen
+			}
+			const TVector<T, 3> NegV = -V;
+			const TVector<T, 3> SupportA = A.SupportCore(NegV, A.GetMargin(), nullptr, VertexIndexA);
+			const TVector<T, 3> VInB = V; // same space
+			const TVector<T, 3> SupportB = B.SupportCore(VInB, B.GetMargin(), nullptr, VertexIndexB);
+			const TVector<T, 3> W = SupportA - SupportB;
+
+			if (TVector<T, 3>::DotProduct(V, W) > Inflation)
+			{
+				return false;
+			}
+
+			SimplexIDs[SimplexIDs.NumVerts] = SimplexIDs.NumVerts;
+			Simplex[SimplexIDs.NumVerts++] = W;
+
+			V = SimplexFindClosestToOrigin(Simplex, SimplexIDs, Barycentric);
+
+			T NewDist2 = V.SizeSquared();
+			bNearZero = NewDist2 < Inflation2;
+
+			//as simplices become degenerate we will stop making progress. This is a side-effect of precision, in that case take V as the current best approximation
+			//question: should we take previous v in case it's better?
+			const bool bMadeProgress = NewDist2 < PrevDist2;
+			bTerminate = bNearZero || !bMadeProgress;
+
+			PrevDist2 = NewDist2;
+
+			if (!bTerminate)
+			{
+				V /= FMath::Sqrt(NewDist2);
+			}
+
+		} while (!bTerminate);
+
+		return bNearZero;
 	}
 
 	/**
