Fix static analysis warning

diff --git a/Engine/Source/Runtime/Experimental/Chaos/Public/Chaos/AABBTree.h b/Engine/Source/Runtime/Experimental/Chaos/Public/Chaos/AABBTree.h
--- a/Engine/Source/Runtime/Experimental/Chaos/Public/Chaos/AABBTree.h
+++ b/Engine/Source/Runtime/Experimental/Chaos/Public/Chaos/AABBTree.h
@@ -1021,9 +1021,6 @@ public:
 		{
 			return 1;
 		}
-
-		check(false);
-		return 0;
 	}
 
 	// Is the input node Child 0 or Child 1?
