Reduce allocations during collision mesh creation
  - 201.3ms to 146.7ms total from engine start to end of skydiving section in replayrun

diff --git a/Engine/Source/Runtime/Experimental/Chaos/Public/Chaos/CollisionConvexMesh.h b/Engine/Source/Runtime/Experimental/Chaos/Public/Chaos/CollisionConvexMesh.h
--- a/Engine/Source/Runtime/Experimental/Chaos/Public/Chaos/CollisionConvexMesh.h
+++ b/Engine/Source/Runtime/Experimental/Chaos/Public/Chaos/CollisionConvexMesh.h
@@ -7,6 +7,7 @@
 #include "TriangleMesh.h"
 #include "Particles.h"
 #include "ChaosLog.h"
+#include "Containers/ChunkedArray.h"
 #include "CompGeom/ConvexHull3.h"
 
 #define DEBUG_HULL_GENERATION 0
@@ -140,7 +141,8 @@ namespace Chaos
 			int32 PoolIdx;
 
 		private:
-			FConvexFace(const TPlaneConcrete<FReal,3>& FacePlane)
+			FConvexFace(int32 InPoolIdx, const TPlaneConcrete<FReal, 3>& FacePlane)
+				: PoolIdx(InPoolIdx)
 			{
 				Reset(FacePlane);
 			}
@@ -151,9 +153,12 @@ namespace Chaos
 				Plane = FacePlane;
 			}
 
+			// Required for TChunkedArray
+			FConvexFace() = default;
 			~FConvexFace() = default;
 
 			friend FMemPool;
+			friend TChunkedArray<FConvexFace>;
 		};
 
 		struct FHalfEdge
@@ -166,7 +171,8 @@ namespace Chaos
 			int32 PoolIdx;
 
 		private:
-			FHalfEdge(int32 InVertex=-1)
+			FHalfEdge(int32 InPoolIdx, int32 InVertex)
+				: PoolIdx(InPoolIdx)
 			{
 				Reset(InVertex);
 			}
@@ -176,9 +182,12 @@ namespace Chaos
 				Vertex = InVertex;
 			}
 
+			// Required for TChunkedArray
+			FHalfEdge() = default;
 			~FHalfEdge() = default;
 
 			friend FMemPool;
+			friend TChunkedArray<FHalfEdge>;
 		};
 
 		class FMemPool
@@ -189,17 +198,15 @@ namespace Chaos
 				if(HalfEdgesFreeIndices.Num())
 				{
 					const uint32 Idx = HalfEdgesFreeIndices.Pop(/*bAllowShrinking=*/false);
-					FHalfEdge* FreeHalfEdge = HalfEdges[Idx];
+					FHalfEdge* FreeHalfEdge = &HalfEdges[Idx];
 					FreeHalfEdge->Reset(InVertex);
 					ensure(FreeHalfEdge->PoolIdx == Idx);
 					return FreeHalfEdge;
 				}
 				else
 				{
-					FHalfEdge* NewHalfEdge = new FHalfEdge(InVertex);
-					NewHalfEdge->PoolIdx = HalfEdges.Num();
-					HalfEdges.Add(NewHalfEdge);
-					return NewHalfEdge;
+					int32 Idx = HalfEdges.AddElement(FHalfEdge(HalfEdges.Num(), InVertex));
+					return &HalfEdges[Idx];
 				}
 			}
 
@@ -208,17 +215,15 @@ namespace Chaos
 				if(FacesFreeIndices.Num())
 				{
 					const uint32 Idx = FacesFreeIndices.Pop(/*bAllowShrinking=*/false);
-					FConvexFace* FreeFace = Faces[Idx];
+					FConvexFace* FreeFace = &Faces[Idx];
 					FreeFace->Reset(FacePlane);
 					ensure(FreeFace->PoolIdx == Idx);
 					return FreeFace;
 				}
 				else
 				{
-					FConvexFace* NewFace = new FConvexFace(FacePlane);
-					NewFace->PoolIdx = Faces.Num();
-					Faces.Add(NewFace);
-					return NewFace;
+					int32 Idx = Faces.AddElement(FConvexFace(Faces.Num(), FacePlane));
+					return &Faces[Idx];
 				}
 			}
 
@@ -232,25 +237,12 @@ namespace Chaos
 				FacesFreeIndices.Add(Face->PoolIdx);
 			}
 
-			~FMemPool()
-			{
-				for(FHalfEdge* HalfEdge : HalfEdges)
-				{
-					delete HalfEdge;
-				}
-
-				for(FConvexFace* Face : Faces)
-				{
-					delete Face;
-				}
-			}
-
 		private:
 			TArray<int32> HalfEdgesFreeIndices;
-			TArray<FHalfEdge*> HalfEdges;
+			TChunkedArray<FHalfEdge> HalfEdges;
 
 			TArray<int32> FacesFreeIndices;
-			TArray<FConvexFace*> Faces;
+			TChunkedArray<FConvexFace> Faces;
 		};
 
 	public:
@@ -259,6 +251,8 @@ namespace Chaos
 
 		static void BuildConvexHull(const TArray<FVec3>& InVertices, TArray<TVec3<int32>>& OutIndices, const Params& InParams = Params())
 		{
+			TRACE_CPUPROFILER_EVENT_SCOPE(Chaos::BuildConvexHull);
+
 			OutIndices.Reset();
 			FMemPool Pool;
 			FConvexFace* Faces = BuildInitialHull(Pool, InVertices);
