Adjust FindBestBounds in AABBTree to workaround code gen issue.

diff --git a/Engine/Source/Runtime/Experimental/Chaos/Public/Chaos/AABBTree.h b/Engine/Source/Runtime/Experimental/Chaos/Public/Chaos/AABBTree.h
--- a/Engine/Source/Runtime/Experimental/Chaos/Public/Chaos/AABBTree.h
+++ b/Engine/Source/Runtime/Experimental/Chaos/Public/Chaos/AABBTree.h
@@ -2319,7 +2319,18 @@ private:
 			const FElement& Elem = CurrentSnapshot.Elems[ElemIdx];
 			int32 BoxIdx = 0;
 			const TVec3<T> ElemCenter = Elem.Bounds.Center();
-			const T CenterVal = ElemCenter[MaxAxis];
+			
+			// This was changed to work around a code generation issue on some platforms, don't change it without testing results of ElemCenter computation.
+			T CenterVal = ElemCenter[0];
+			if (MaxAxis == 1)
+			{
+				CenterVal = ElemCenter[1];
+			}
+			else if(MaxAxis == 2)
+			{
+				CenterVal = ElemCenter[2];
+			}
+		
 			const int32 MinBoxIdx = CenterVal <= SplitVal ? 0 : 1;
 			
 			FSplitInfo& SplitInfo = CurrentSnapshot.SplitInfos[MinBoxIdx];
