Add missing bounds computation for dynamic AABB tree

diff --git a/Engine/Source/Runtime/Experimental/Chaos/Public/Chaos/AABBTree.h b/Engine/Source/Runtime/Experimental/Chaos/Public/Chaos/AABBTree.h
--- a/Engine/Source/Runtime/Experimental/Chaos/Public/Chaos/AABBTree.h
+++ b/Engine/Source/Runtime/Experimental/Chaos/Public/Chaos/AABBTree.h
@@ -244,6 +244,11 @@ struct TAABBTreeLeafArray : public TBoundsWrapperHelper<TPayloadType, T, bComput
 		return Elems.Num();
 	}
 
+	void RecomputeBounds()
+	{
+		this->ComputeBounds(Elems);
+	}
+
 	template <typename TSQVisitor, typename TQueryFastData>
 	FORCEINLINE_DEBUGGABLE bool RaycastFast(const TVec3<T>& Start, TQueryFastData& QueryFastData, TSQVisitor& Visitor, const TVec3<T>& Dir, const TVec3<T> InvDir, const bool bParallel[3]) const
 	{
@@ -1371,6 +1376,7 @@ public:
 							{
 								// We still need to update the constituent bounds
 								Leaves[PayloadInfo->LeafIdx].UpdateElement(Payload, NewBounds, bHasBounds);
+								Leaves[PayloadInfo->LeafIdx].RecomputeBounds();
 								return;
 							}
 						}
diff --git a/Engine/Source/Runtime/Experimental/Chaos/Public/Chaos/BoundingVolume.h b/Engine/Source/Runtime/Experimental/Chaos/Public/Chaos/BoundingVolume.h
--- a/Engine/Source/Runtime/Experimental/Chaos/Public/Chaos/BoundingVolume.h
+++ b/Engine/Source/Runtime/Experimental/Chaos/Public/Chaos/BoundingVolume.h
@@ -288,6 +288,11 @@ public:
 		check(false);
 	}
 
+	inline void RecomputeBounds()
+	{
+
+	}
+
 	
 	// Begin ISpatialAcceleration interface
 	virtual TArray<TPayloadType> FindAllIntersections(const FAABB3& Box) const override { return FindAllIntersectionsImp(Box); }
