[NO NON-PII COMMIT MESSAGE]

diff --git a/Engine/Source/Runtime/Experimental/Chaos/Private/Chaos/PBDRigidsEvolution.cpp b/Engine/Source/Runtime/Experimental/Chaos/Private/Chaos/PBDRigidsEvolution.cpp
--- a/Engine/Source/Runtime/Experimental/Chaos/Private/Chaos/PBDRigidsEvolution.cpp
+++ b/Engine/Source/Runtime/Experimental/Chaos/Private/Chaos/PBDRigidsEvolution.cpp
@@ -32,7 +32,7 @@ namespace Chaos
 	CHAOS_API int32 AccelerationStructureSplitStaticAndDynamic = 1;
 	FAutoConsoleVariableRef CVarAccelerationStructureSplitStaticAndDynamic(TEXT("p.Chaos.AccelerationStructureSplitStaticDynamic"), AccelerationStructureSplitStaticAndDynamic, TEXT("Set to 1: Sort Dynamic and Static bodies into seperate acceleration structures, any other value will disable the feature"));
 
-	CHAOS_API int32 AccelerationStructureUseDynamicTree = 0;
+	CHAOS_API int32 AccelerationStructureUseDynamicTree = 1;
 	FAutoConsoleVariableRef CVarAccelerationStructureUseDynamicTree(TEXT("p.Chaos.AccelerationStructureUseDynamicTree"), AccelerationStructureUseDynamicTree, TEXT("Use a dynamic BVH tree structure for dynamic objects"));
 
 	CHAOS_API int32 AccelerationStructureTimeSlicingMaxQueueSizeBeforeForce = 1000;
