AccelerationStructureSplitStaticAndDynamic now disabled by default but enabled on local physics project.

diff --git a/Engine/Source/Runtime/Experimental/Chaos/Private/Chaos/PBDRigidsEvolution.cpp b/Engine/Source/Runtime/Experimental/Chaos/Private/Chaos/PBDRigidsEvolution.cpp
--- a/Engine/Source/Runtime/Experimental/Chaos/Private/Chaos/PBDRigidsEvolution.cpp
+++ b/Engine/Source/Runtime/Experimental/Chaos/Private/Chaos/PBDRigidsEvolution.cpp
@@ -29,7 +29,7 @@ namespace Chaos
 	CHAOS_API int32 FixBadAccelerationStructureRemoval = 1;
 	FAutoConsoleVariableRef CVarFixBadAccelerationStructureRemoval(TEXT("p.FixBadAccelerationStructureRemoval"), FixBadAccelerationStructureRemoval, TEXT(""));
 
-	CHAOS_API int32 AccelerationStructureSplitStaticAndDynamic = 1;
+	CHAOS_API int32 AccelerationStructureSplitStaticAndDynamic = 0;
 	FAutoConsoleVariableRef CVarAccelerationStructureSplitStaticAndDynamic(TEXT("p.Chaos.AccelerationStructureSplitStaticDynamic"), AccelerationStructureSplitStaticAndDynamic, TEXT("Set to 1: Sort Dynamic and Static bodies into seperate acceleration structures, any other value will disable the feature"));
 
 	CHAOS_API int32 AccelerationStructureUseDynamicTree = 1;
