DynamicTreeBoundingBoxPadding increased to 5cm after perf measurements

diff --git a/Engine/Source/Runtime/Experimental/Chaos/Private/Chaos/AABBTree.cpp b/Engine/Source/Runtime/Experimental/Chaos/Private/Chaos/AABBTree.cpp
--- a/Engine/Source/Runtime/Experimental/Chaos/Private/Chaos/AABBTree.cpp
+++ b/Engine/Source/Runtime/Experimental/Chaos/Private/Chaos/AABBTree.cpp
@@ -29,5 +29,5 @@ FAutoConsoleVariableRef FAABBTreeCVars::CVarSplitAtAverageCenter(TEXT("p.aabbtre
 int32 FAABBTreeCVars::SplitOnVarianceAxis = 1;
 FAutoConsoleVariableRef FAABBTreeCVars::CVarSplitOnVarianceAxis(TEXT("p.aabbtree.splitonvarianceaxis"), FAABBTreeCVars::SplitOnVarianceAxis, TEXT("Split AABB tree nodes along the axis with the largest element center variance"));
 
-float FAABBTreeCVars::DynamicTreeBoundingBoxPadding = 1.0f;
+float FAABBTreeCVars::DynamicTreeBoundingBoxPadding = 5.0f;
 FAutoConsoleVariableRef FAABBTreeCVars::CVarDynamicTreeBoundingBoxPadding(TEXT("p.aabbtree.DynamicTreeBoundingBoxPadding"), FAABBTreeCVars::DynamicTreeBoundingBoxPadding, TEXT("Additional padding added to bounding boxes for dynamic AABB trees to amortize update cost"));
