Removing mass computation from integration by accumulating accel/vel instead of forces/torques.
Remove function calls to rules in integration.

diff --git a/Engine/Plugins/Experimental/ChaosCloth/Source/Chaos/Private/ChaosCloth/ChaosClothingSimulationSolver.cpp b/Engine/Plugins/Experimental/ChaosCloth/Source/Chaos/Private/ChaosCloth/ChaosClothingSimulationSolver.cpp
--- a/Engine/Plugins/Experimental/ChaosCloth/Source/Chaos/Private/ChaosCloth/ChaosClothingSimulationSolver.cpp
+++ b/Engine/Plugins/Experimental/ChaosCloth/Source/Chaos/Private/ChaosCloth/ChaosClothingSimulationSolver.cpp
@@ -722,7 +722,7 @@ void FClothingSimulationSolver::AddExternalForces(uint32 GroupId, bool bUseLegac
 					}
 				}
 
-				Particles.F(Index) += Forces;
+				Particles.Acceleration(Index) += Forces * Particles.InvM(Index);
 			};
 	}
 }
diff --git a/Engine/Plugins/Runtime/NetworkPrediction/Source/NetworkPrediction/Private/NetworkPhysics.cpp b/Engine/Plugins/Runtime/NetworkPrediction/Source/NetworkPrediction/Private/NetworkPhysics.cpp
--- a/Engine/Plugins/Runtime/NetworkPrediction/Source/NetworkPrediction/Private/NetworkPhysics.cpp
+++ b/Engine/Plugins/Runtime/NetworkPrediction/Source/NetworkPrediction/Private/NetworkPhysics.cpp
@@ -335,21 +335,21 @@ struct FNetworkPhysicsRewindCallback : public Chaos::IRewindCallback
 						for (int32 i=(int32)Chaos::FFrameAndPhase::EParticleHistoryPhase::PrePushData; i < (int32)Chaos::FFrameAndPhase::EParticleHistoryPhase::NumPhases; ++i)
 						{
 							const auto LocalPreFrameData = RewindData->GetPastStateAtFrame(*Proxy->GetHandle_LowLevel(), Obj.Frame-1, (Chaos::FFrameAndPhase::EParticleHistoryPhase)i);
-							if (Obj.DebugState[i].Force != LocalPreFrameData.F())
+							if (Obj.DebugState[i].Force != LocalPreFrameData.Acceleration())
 							{
-								UE_LOG(LogNetworkPhysics, Warning, TEXT("Previous Frame [%d] Phase [%d] FORCE mismatch. Local: %s. Server: %s"), PreSimulationFrame, i, *LocalPreFrameData.F().ToString(), *Obj.DebugState[i].Force.ToString());
+								UE_LOG(LogNetworkPhysics, Warning, TEXT("Previous Frame [%d] Phase [%d] FORCE mismatch. Local: %s. Server: %s"), PreSimulationFrame, i, *LocalPreFrameData.Acceleration().ToString(), *Obj.DebugState[i].Force.ToString());
 							}
-							if (Obj.DebugState[i].Torque != LocalPreFrameData.Torque())
+							if (Obj.DebugState[i].Torque != LocalPreFrameData.AngularAcceleration())
 							{
-								UE_LOG(LogNetworkPhysics, Warning, TEXT("Previous Frame [%d] Phase [%d] TORQUE mismatch. Local: %s. Server: %s"), PreSimulationFrame, i, *LocalPreFrameData.Torque().ToString(), *Obj.DebugState[i].Torque.ToString());
+								UE_LOG(LogNetworkPhysics, Warning, TEXT("Previous Frame [%d] Phase [%d] TORQUE mismatch. Local: %s. Server: %s"), PreSimulationFrame, i, *LocalPreFrameData.AngularAcceleration().ToString(), *Obj.DebugState[i].Torque.ToString());
 							}
-							if (Obj.DebugState[i].LinearImpulse != LocalPreFrameData.LinearImpulse())
+							if (Obj.DebugState[i].LinearImpulse != LocalPreFrameData.LinearImpulseVelocity())
 							{
-								UE_LOG(LogNetworkPhysics, Warning, TEXT("Previous Frame [%d] Phase [%d] LINEAR IMPULSE mismatch. Local: %s. Server: %s"), PreSimulationFrame, i, *LocalPreFrameData.LinearImpulse().ToString(), *Obj.DebugState[i].LinearImpulse.ToString());
+								UE_LOG(LogNetworkPhysics, Warning, TEXT("Previous Frame [%d] Phase [%d] LINEAR IMPULSE mismatch. Local: %s. Server: %s"), PreSimulationFrame, i, *LocalPreFrameData.LinearImpulseVelocity().ToString(), *Obj.DebugState[i].LinearImpulse.ToString());
 							}
-							if (Obj.DebugState[i].AngularImpulse != LocalPreFrameData.AngularImpulse())
+							if (Obj.DebugState[i].AngularImpulse != LocalPreFrameData.AngularImpulseVelocity())
 							{
-								UE_LOG(LogNetworkPhysics, Warning, TEXT("Previous Frame [%d] Phase [%d] ANGULAR IMPULSE mismatch. Local: %s. Server: %s"), PreSimulationFrame, i, *LocalPreFrameData.AngularImpulse().ToString(), *Obj.DebugState[i].AngularImpulse.ToString());
+								UE_LOG(LogNetworkPhysics, Warning, TEXT("Previous Frame [%d] Phase [%d] ANGULAR IMPULSE mismatch. Local: %s. Server: %s"), PreSimulationFrame, i, *LocalPreFrameData.AngularImpulseVelocity().ToString(), *Obj.DebugState[i].AngularImpulse.ToString());
 							}
 						}
 					}
@@ -473,10 +473,10 @@ struct FNetworkPhysicsRewindCallback : public Chaos::IRewindCallback
 						for (int32 i=(int32)Chaos::FFrameAndPhase::EParticleHistoryPhase::PrePushData; i < (int32)Chaos::FFrameAndPhase::EParticleHistoryPhase::NumPhases; ++i)
 						{
 							Chaos::FGeometryParticleState State = RewindData->GetPastStateAtFrame(*Proxy->GetHandle_LowLevel(), PreStorageFrame, (Chaos::FFrameAndPhase::EParticleHistoryPhase)i);
-							Obj.DebugState[i].Force = State.F();
-							Obj.DebugState[i].Torque = State.Torque();
-							Obj.DebugState[i].LinearImpulse = State.LinearImpulse();
-							Obj.DebugState[i].AngularImpulse = State.AngularImpulse();
+							Obj.DebugState[i].Force = State.Acceleration();
+							Obj.DebugState[i].Torque = State.AngularAcceleration();
+							Obj.DebugState[i].LinearImpulse = State.LinearImpulseVelocity();
+							Obj.DebugState[i].AngularImpulse = State.AngularImpulseVelocity();
 
 							if (UE_NETWORK_PHYSICS::LogImpulses > 0)
 							{
diff --git a/Engine/Plugins/Runtime/NetworkPredictionExtras/Source/NetworkPredictionExtras/Private/PhysicsMovement.cpp b/Engine/Plugins/Runtime/NetworkPredictionExtras/Source/NetworkPredictionExtras/Private/PhysicsMovement.cpp
--- a/Engine/Plugins/Runtime/NetworkPredictionExtras/Source/NetworkPredictionExtras/Private/PhysicsMovement.cpp
+++ b/Engine/Plugins/Runtime/NetworkPredictionExtras/Source/NetworkPredictionExtras/Private/PhysicsMovement.cpp
@@ -126,7 +126,7 @@ struct FPhysicsMovementSimulation
 			if (InputCmd.bJumpedPressed && NetState.JumpStartFrame + 10 < SimulationFrame)
 			{
 				NetState.JumpStartFrame = SimulationFrame;
-				PT->SetLinearImpulse(FVector(0.f, 0.f, 60000.f + (FMath::FRand() * 50000.f)));
+				PT->SetLinearImpulse(FVector(0.f, 0.f, 60000.f + (FMath::FRand() * 50000.f)), /*bIsVelociy=*/false);
 			}
 			return;
 		}
diff --git a/Engine/Source/Programs/HeadlessChaos/Private/HeadlessChaosTestJointBreak.cpp b/Engine/Source/Programs/HeadlessChaos/Private/HeadlessChaosTestJointBreak.cpp
--- a/Engine/Source/Programs/HeadlessChaos/Private/HeadlessChaosTestJointBreak.cpp
+++ b/Engine/Source/Programs/HeadlessChaos/Private/HeadlessChaosTestJointBreak.cpp
@@ -267,7 +267,7 @@ namespace ChaosTest {
 		// Run the sim
 		for (int32 i = 0; i < NumIts; ++i)
 		{
-			Test.GetParticle(1)->CastToRigidParticle()->Torque() = Torque;
+			Test.GetParticle(1)->CastToRigidParticle()->SetTorque(Torque);
 
 			Test.Evolution.AdvanceOneTimeStep(Dt);
 			Test.Evolution.EndFrame(Dt);
@@ -313,7 +313,7 @@ namespace ChaosTest {
 		// Run the sim
 		for (int32 i = 0; i < NumSteps; ++i)
 		{
-			Test.GetParticle(1)->CastToRigidParticle()->Torque() = Torque;
+			Test.GetParticle(1)->CastToRigidParticle()->SetTorque(Torque);
 
 			Test.Evolution.AdvanceOneTimeStep(Dt);
 			Test.Evolution.EndFrame(Dt);
diff --git a/Engine/Source/Programs/HeadlessChaos/Private/HeadlessChaosTestJointForces.cpp b/Engine/Source/Programs/HeadlessChaos/Private/HeadlessChaosTestJointForces.cpp
--- a/Engine/Source/Programs/HeadlessChaos/Private/HeadlessChaosTestJointForces.cpp
+++ b/Engine/Source/Programs/HeadlessChaos/Private/HeadlessChaosTestJointForces.cpp
@@ -145,7 +145,7 @@ namespace ChaosTest {
 		// Run the sim
 		for (int32 i = 0; i < NumSteps; ++i)
 		{
-			Test.GetParticle(1)->CastToRigidParticle()->Torque() = Torque;
+			Test.GetParticle(1)->CastToRigidParticle()->SetTorque(Torque);
 			
 			Test.Evolution.AdvanceOneTimeStep(Dt);
 			Test.Evolution.EndFrame(Dt);
diff --git a/Engine/Source/Programs/HeadlessChaos/Private/HeadlessChaosTestRewind.cpp b/Engine/Source/Programs/HeadlessChaos/Private/HeadlessChaosTestRewind.cpp
--- a/Engine/Source/Programs/HeadlessChaos/Private/HeadlessChaosTestRewind.cpp
+++ b/Engine/Source/Programs/HeadlessChaos/Private/HeadlessChaosTestRewind.cpp
@@ -251,7 +251,7 @@ namespace ChaosTest {
 			{
 				if (Step == 5)
 				{
-					Particle.SetLinearImpulse(FVec3(0, 0, 100));
+					Particle.SetLinearImpulse(FVec3(0, 0, 100), /*bIsVelocity=*/false);
 				}
 
 				TickSolverHelper(Solver);
@@ -274,7 +274,7 @@ namespace ChaosTest {
 				{
 					if (Step == 5)
 					{
-						Particle.SetLinearImpulse(FVec3(0, 0, 100));
+						Particle.SetLinearImpulse(FVec3(0, 0, 100), /*bIsVelocity=*/false);
 					}
 
 					if(Step == 15)
@@ -1458,8 +1458,8 @@ namespace ChaosTest {
 					{
 						if (PhysicsStep == 3)
 						{
-							Proxy->GetPhysicsThreadAPI()->AddForce(FVec3(0, 0, 10));
-							Proxy->GetPhysicsThreadAPI()->AddTorque(FVec3(0, 0, 10));
+							Proxy->GetPhysicsThreadAPI()->SetAcceleration(FVec3(0, 0, 10) + Proxy->GetPhysicsThreadAPI()->Acceleration());
+							Proxy->GetPhysicsThreadAPI()->SetAngularAcceleration(FVec3(0, 0, 10) + Proxy->GetPhysicsThreadAPI()->AngularAcceleration());
 						}
 
 					for (int32 Step = 0; Step < PhysicsStep; ++Step)
@@ -1468,28 +1468,28 @@ namespace ChaosTest {
 							FGeometryParticleState State = RewindData->GetPastStateAtFrame(*Proxy->GetHandle_LowLevel(), Step, FFrameAndPhase::EParticleHistoryPhase::PostCallbacks);
 								if (Step == 3)
 								{
-									EXPECT_EQ(State.F()[2], 11);	//1 from GT + 10 from callback
-									EXPECT_EQ(State.Torque()[2], 10);	//10 from callback (nothing from GT)
+									EXPECT_EQ(State.Acceleration()[2], 11);	//1 from GT + 10 from callback
+									EXPECT_EQ(State.AngularAcceleration()[2], 10);	//10 from callback (nothing from GT)
 								}
 								else
 								{
-									EXPECT_EQ(State.F()[2], 1);	//GT always sets force of 1
-									EXPECT_EQ(State.Torque()[2], 0);	//GT never sets torque
+									EXPECT_EQ(State.Acceleration()[2], 1);	//GT always sets force of 1
+									EXPECT_EQ(State.AngularAcceleration()[2], 0);	//GT never sets torque
 								}
 							}
 
 							{
 								//Before push no force or torque at all
 							FGeometryParticleState PrePushState = RewindData->GetPastStateAtFrame(*Proxy->GetHandle_LowLevel(), Step, FFrameAndPhase::EParticleHistoryPhase::PrePushData);
-								EXPECT_EQ(PrePushState.F()[2], 0);
-								EXPECT_EQ(PrePushState.Torque()[2], 0);
+								EXPECT_EQ(PrePushState.Acceleration()[2], 0);
+								EXPECT_EQ(PrePushState.AngularAcceleration()[2], 0);
 							}
 
 							{
 								//After push (but before callback) only see GT force, and no torque
 							FGeometryParticleState PostPushState = RewindData->GetPastStateAtFrame(*Proxy->GetHandle_LowLevel(), Step, FFrameAndPhase::EParticleHistoryPhase::PostPushData);
-								EXPECT_EQ(PostPushState.F()[2], 1);	//GT always sets force of 1
-								EXPECT_EQ(PostPushState.Torque()[2], 0);	//GT never sets torque
+								EXPECT_EQ(PostPushState.Acceleration()[2], 1);	//GT always sets force of 1
+								EXPECT_EQ(PostPushState.AngularAcceleration()[2], 0);	//GT never sets torque
 							}
 							
 						}
@@ -1540,7 +1540,7 @@ namespace ChaosTest {
 					//where step is game step: so really it's step * 4
 					ExpectedForce = Step * 4 + 2.5;
 				}
-				EXPECT_EQ(ParticleState.F()[2], ExpectedForce);
+				EXPECT_EQ(ParticleState.Acceleration()[2], ExpectedForce);
 			}
 		});
 	}
@@ -1581,15 +1581,15 @@ namespace ChaosTest {
 					const float SimTime = Step * SimDt;
 					if (SimTime >= 3 && SimTime < 4)
 					{
-						EXPECT_EQ(ParticleState.F()[2], 3);
+						EXPECT_EQ(ParticleState.Acceleration()[2], 3);
 					}
 					else if (SimTime >= 5 && SimTime < 6)
 					{
-						EXPECT_EQ(ParticleState.F()[2], 5);
+						EXPECT_EQ(ParticleState.Acceleration()[2], 5);
 					}
 					else
 					{
-						EXPECT_EQ(ParticleState.F()[2], 0);
+						EXPECT_EQ(ParticleState.Acceleration()[2], 0);
 					}
 				}
 				else
@@ -1597,15 +1597,15 @@ namespace ChaosTest {
 					//we get an average
 					if (Step == 0)
 					{
-						EXPECT_EQ(ParticleState.F()[2], 3 / 4.f);
+						EXPECT_EQ(ParticleState.Acceleration()[2], 3 / 4.f);
 					}
 					else if (Step == 1)
 					{
-						EXPECT_EQ(ParticleState.F()[2], 5 / 4.f);
+						EXPECT_EQ(ParticleState.Acceleration()[2], 5 / 4.f);
 					}
 					else
 					{
-						EXPECT_EQ(ParticleState.F()[2], 0);
+						EXPECT_EQ(ParticleState.Acceleration()[2], 0);
 					}
 				}
 
@@ -2911,7 +2911,7 @@ namespace ChaosTest {
 			{
 				FRewindData* RewindData = Solver->GetRewindData();
 				//EXPECT_TRUE(RewindData->RewindToFrame(7));
-				EXPECT_EQ(Particle.F()[1], 0);
+				EXPECT_EQ(Particle.Acceleration()[1], 0);
 			}
 
 			// Throw out the proxy
diff --git a/Engine/Source/Programs/HeadlessChaos/Private/HeadlessChaosTestSerialization.cpp b/Engine/Source/Programs/HeadlessChaos/Private/HeadlessChaosTestSerialization.cpp
--- a/Engine/Source/Programs/HeadlessChaos/Private/HeadlessChaosTestSerialization.cpp
+++ b/Engine/Source/Programs/HeadlessChaos/Private/HeadlessChaosTestSerialization.cpp
@@ -396,8 +396,9 @@ namespace ChaosTest
 		Particles.AddParticles(2);
 		Particles.R(0) = FRotation3::Identity;
 		Particles.R(1) = FRotation3::Identity;
-		Particles.F(0) = F[0];
-		Particles.F(1) = F[1];
+
+		Particles.Acceleration(0) = F[0];
+		Particles.Acceleration(1) = F[1];
 		Particles.X(0) = X[0];
 		Particles.X(1) = X[1];
 		Particles.RotationOfMass(0) = FRotation3::FromIdentity();
@@ -412,8 +413,8 @@ namespace ChaosTest
 		for (TRigidParticles<FReal, 3> const &TestParticles : ObjectsToTest)
 		{
 			EXPECT_EQ(TestParticles.Size(), Particles.Size());
-			EXPECT_EQ(TestParticles.F(0), Particles.F(0));
-			EXPECT_EQ(TestParticles.F(1), Particles.F(1));
+			EXPECT_EQ(TestParticles.Acceleration(0), Particles.Acceleration(0));
+			EXPECT_EQ(TestParticles.Acceleration(1), Particles.Acceleration(1));
 			EXPECT_EQ(TestParticles.X(0), Particles.X(0));
 			EXPECT_EQ(TestParticles.X(1), Particles.X(1));
 		}
diff --git a/Engine/Source/Programs/HeadlessChaos/Private/HeadlessChaosTestVehicles.cpp b/Engine/Source/Programs/HeadlessChaos/Private/HeadlessChaosTestVehicles.cpp
--- a/Engine/Source/Programs/HeadlessChaos/Private/HeadlessChaosTestVehicles.cpp
+++ b/Engine/Source/Programs/HeadlessChaos/Private/HeadlessChaosTestVehicles.cpp
@@ -1430,15 +1430,13 @@ namespace ChaosTest
 
 	void AddForceAtPosition(FPBDRigidsEvolutionGBF& Evolution, TPBDRigidParticleHandle<FReal, 3>* Rigid, const FVector& InForce, const FVector& InPosition)
 	{
-		const Chaos::FVec3& CurrentForce = Rigid->F();
-		const Chaos::FVec3& CurrentTorque = Rigid->Torque();
 		const Chaos::FVec3 WorldCOM = FParticleUtilitiesGT::GetCoMWorldPosition(Rigid);
 
 		Evolution.SetParticleObjectState(Rigid, EObjectStateType::Dynamic);
 
 		const Chaos::FVec3 WorldTorque = Chaos::FVec3::CrossProduct(InPosition - WorldCOM, InForce);
-		Rigid->SetF(CurrentForce + InForce);
-		Rigid->SetTorque(CurrentTorque + WorldTorque);
+		Rigid->AddForce(InForce);
+		Rigid->AddTorque(WorldTorque);
 
 	}
 
diff --git a/Engine/Source/Runtime/Engine/Private/PhysicsEngine/ImmediatePhysics/ImmediatePhysicsChaos/ImmediatePhysicsActorHandle_Chaos.cpp b/Engine/Source/Runtime/Engine/Private/PhysicsEngine/ImmediatePhysics/ImmediatePhysicsChaos/ImmediatePhysicsActorHandle_Chaos.cpp
--- a/Engine/Source/Runtime/Engine/Private/PhysicsEngine/ImmediatePhysics/ImmediatePhysicsChaos/ImmediatePhysicsActorHandle_Chaos.cpp
+++ b/Engine/Source/Runtime/Engine/Private/PhysicsEngine/ImmediatePhysics/ImmediatePhysicsChaos/ImmediatePhysicsActorHandle_Chaos.cpp
@@ -470,7 +470,7 @@ namespace ImmediatePhysics_Chaos
 
 		if (FPBDRigidParticleHandle* Rigid = Handle()->CastToRigidParticle())
 		{
-			Rigid->F() += Force;
+			Rigid->AddForce(Force);
 		}
 	}
 
@@ -480,7 +480,7 @@ namespace ImmediatePhysics_Chaos
 
 		if (FPBDRigidParticleHandle * Rigid = Handle()->CastToRigidParticle())
 		{
-			Rigid->Torque() += Torque;
+			Rigid->AddTorque(Torque);
 		}
 	}
 
@@ -515,7 +515,7 @@ namespace ImmediatePhysics_Chaos
 			}
 			else
 			{
-				Rigid->F() += ApplyDelta;
+				Rigid->Acceleration() += ApplyDelta * Rigid->InvM();
 			}
 		}
 	}
@@ -527,8 +527,9 @@ namespace ImmediatePhysics_Chaos
 		if (FPBDRigidParticleHandle* Rigid = Handle()->CastToRigidParticle())
 		{
 			FVector CoM = FParticleUtilities::GetCoMWorldPosition(Rigid);
-			Rigid->LinearImpulse() += Impulse;
-			Rigid->AngularImpulse() += FVector::CrossProduct(Location - CoM, Impulse);
+			Chaos::FMatrix33 InvInertia = FParticleUtilities::GetWorldInvInertia(Rigid);
+			Rigid->LinearImpulseVelocity() += Impulse * Rigid->InvM();
+			Rigid->AngularImpulseVelocity() += InvInertia * FVector::CrossProduct(Location - CoM, Impulse);
 		}
 	}
 
diff --git a/Engine/Source/Runtime/Experimental/Chaos/Private/Chaos/ChaosMarshallingManager.cpp b/Engine/Source/Runtime/Experimental/Chaos/Private/Chaos/ChaosMarshallingManager.cpp
--- a/Engine/Source/Runtime/Experimental/Chaos/Private/Chaos/ChaosMarshallingManager.cpp
+++ b/Engine/Source/Runtime/Experimental/Chaos/Private/Chaos/ChaosMarshallingManager.cpp
@@ -235,14 +235,14 @@ void FPushPhysicsData::CopySubstepData(const FPushPhysicsData& FirstStepData)
 		{
 			if (const FParticleDynamics* DynamicsData = Dirty.PropertyData.FindDynamics(FirstManager, FirstDataIdx))
 			{
-				if (DynamicsData->F() != FVec3(0) || DynamicsData->Torque() != FVec3(0))	//don't bother interpolating 0. This is important because the input dirtys rewind data
+				if (DynamicsData->Acceleration() != FVec3(0) || DynamicsData->AngularAcceleration() != FVec3(0))	//don't bother interpolating 0. This is important because the input dirtys rewind data
 				{
 					DirtyProxiesDataBuffer.Add(Dirty.Proxy);
 					FParticleDynamics& SubsteppedDynamics = DirtyPropertiesManager.GetChaosPropertyPool<FParticleDynamics, EChaosProperty::Dynamics>().GetElement(Dirty.Proxy->GetDirtyIdx());
 					SubsteppedDynamics = *DynamicsData;
 					//we don't want to sub-step impulses so those are cleared in the sub-step
-					SubsteppedDynamics.SetAngularImpulse(FVec3(0));
-					SubsteppedDynamics.SetLinearImpulse(FVec3(0));
+					SubsteppedDynamics.SetAngularImpulseVelocity(FVec3(0));
+					SubsteppedDynamics.SetLinearImpulseVelocity(FVec3(0));
 					FDirtyProxy& NewDirtyProxy = DirtyProxiesDataBuffer.GetDirtyProxyAt(Dirty.Proxy->GetType(), Dirty.Proxy->GetDirtyIdx());
 					NewDirtyProxy.PropertyData.DirtyFlag(EChaosPropertyFlags::Dynamics);
 					NewDirtyProxy.PropertyData.SetParticleBufferType(EParticleType::Rigid);
diff --git a/Engine/Source/Runtime/Experimental/Chaos/Private/Chaos/Evolution/PBDMinEvolution.cpp b/Engine/Source/Runtime/Experimental/Chaos/Private/Chaos/Evolution/PBDMinEvolution.cpp
--- a/Engine/Source/Runtime/Experimental/Chaos/Private/Chaos/Evolution/PBDMinEvolution.cpp
+++ b/Engine/Source/Runtime/Experimental/Chaos/Private/Chaos/Evolution/PBDMinEvolution.cpp
@@ -85,10 +85,10 @@ namespace Chaos
 			RotationOfMass = Dynamics.AllRotationOfMass().GetData();
 			InvM = Dynamics.AllInvM().GetData();
 			InvI = Dynamics.AllInvI().GetData();
-			F = Dynamics.AllF().GetData();
-			T = Dynamics.AllT().GetData();
-			LinearImpulse = Dynamics.AllLinearImpulse().GetData();
-			AngularImpulse = Dynamics.AllAngularImpulse().GetData();
+			Acceleration = Dynamics.AllAcceleration().GetData();
+			AngularAcceleration = Dynamics.AllAngularAcceleration().GetData();
+			LinearImpulseVelocity = Dynamics.AllLinearImpulseVelocity().GetData();
+			AngularImpulseVelocity = Dynamics.AllAngularImpulseVelocity().GetData();
 			Disabled = Dynamics.AllDisabled().GetData();
 			GravityEnabled = Dynamics.AllGravityEnabled().GetData();
 			LinearEtherDrag = Dynamics.AllLinearEtherDrag().GetData();
@@ -112,10 +112,10 @@ namespace Chaos
 		FRotation3* RotationOfMass;
 		FReal* InvM;
 		FMatrix33* InvI;
-		FVec3* F;
-		FVec3* T;
-		FVec3* LinearImpulse;
-		FVec3* AngularImpulse;
+		FVec3* Acceleration;
+		FVec3* AngularAcceleration;
+		FVec3* LinearImpulseVelocity;
+		FVec3* AngularImpulseVelocity;
 		bool* Disabled;
 		bool* GravityEnabled;
 		FReal* LinearEtherDrag;
@@ -179,8 +179,8 @@ namespace Chaos
 		{
 			if (Particle.ObjectState() == EObjectStateType::Dynamic)
 			{
-				Particle.F() = FVec3(0);
-				Particle.Torque() = FVec3(0);
+				Particle.Acceleration() = FVec3(0);
+				Particle.AngularAcceleration() = FVec3(0);
 			}
 		}
 
@@ -326,8 +326,8 @@ namespace Chaos
 				
 				// Forces and torques
 				const FMatrix33 WorldInvI = Utilities::ComputeWorldSpaceInertia(RCoM, Particle.InvI());
-				FVec3 DV = Particle.InvM() * (Particle.F() * Dt + Particle.LinearImpulse());
-				FVec3 DW = Utilities::Multiply(WorldInvI, (Particle.Torque() * Dt + Particle.AngularImpulse()));
+				FVec3 DV = Particle.Acceleration() * Dt + Particle.LinearImpulseVelocity();
+				FVec3 DW = Particle.AngularAcceleration() * Dt + Particle.AngularImpulseVelocity();
 				FVec3 TargetV = FVec3(0);
 				FVec3 TargetW = FVec3(0);
 
@@ -367,8 +367,8 @@ namespace Chaos
 				FParticleUtilitiesPQ::SetCoMWorldTransform(&Particle, PCoM, QCoM);
 				Particle.V() = V;
 				Particle.W() = W;
-				Particle.LinearImpulse() = FVec3(0);
-				Particle.AngularImpulse() = FVec3(0);
+				Particle.LinearImpulseVelocity() = FVec3(0);
+				Particle.AngularImpulseVelocity() = FVec3(0);
 
 				// Update cached world space state, including bounds. We use the Swept bounds update so that the bounds includes P,Q and X,Q.
 				// This is because when we have joints, they often pull bodies back to their original positions, so we need to know if there
diff --git a/Engine/Source/Runtime/Experimental/Chaos/Private/Chaos/PBDEvolution.cpp b/Engine/Source/Runtime/Experimental/Chaos/Private/Chaos/PBDEvolution.cpp
--- a/Engine/Source/Runtime/Experimental/Chaos/Private/Chaos/PBDEvolution.cpp
+++ b/Engine/Source/Runtime/Experimental/Chaos/Private/Chaos/PBDEvolution.cpp
@@ -224,7 +224,7 @@ void FPBDEvolution::PreIterationUpdate(
 			if (MParticles.InvM(Index) != (FReal)0.)  // Process dynamic particles
 			{
 				// Init forces with GravityForces
-				MParticles.F(Index) = Gravity * MParticles.M(Index);  // F = M * G
+				MParticles.Acceleration(Index) = Gravity;
 
 				// Force Rule
 				if (bForceRule)
@@ -239,7 +239,7 @@ void FPBDEvolution::PreIterationUpdate(
 				}
 
 				// Euler Step Velocity
-				MParticles.V(Index) += MParticles.F(Index) * MParticles.InvM(Index) * MSmoothDt;
+				MParticles.V(Index) += MParticles.Acceleration(Index) * MSmoothDt;
 
 				// Damp Velocity Rule
 				if (bDampVelocityRule)
diff --git a/Engine/Source/Runtime/Experimental/Chaos/Private/Chaos/ParticleHandle.cpp b/Engine/Source/Runtime/Experimental/Chaos/Private/Chaos/ParticleHandle.cpp
--- a/Engine/Source/Runtime/Experimental/Chaos/Private/Chaos/ParticleHandle.cpp
+++ b/Engine/Source/Runtime/Experimental/Chaos/Private/Chaos/ParticleHandle.cpp
@@ -216,12 +216,42 @@ namespace Chaos
 		});
 	}
 
+	template <typename T, int d, bool bPersistent>
+	void TPBDRigidParticleHandleImp<T, d, bPersistent>::AddTorque(const TVector<T, d>& InTorque, bool bInvalidate)
+	{
+		const FRotation3 RCoM = FParticleUtilitiesPQ::GetCoMWorldRotation(this);
+		const FMatrix33 WorldInvI = Utilities::ComputeWorldSpaceInertia(RCoM, InvI());
+		SetAngularAcceleration(AngularAcceleration() + WorldInvI * InTorque);
+	}
+
+
+	template <typename T, int d, bool bPersistent>
+	void TPBDRigidParticleHandleImp<T, d, bPersistent>::SetTorque(const TVector<T, d>& InTorque, bool bInvalidate)
+	{
+		const FRotation3 RCoM = FParticleUtilitiesPQ::GetCoMWorldRotation(this);
+		const FMatrix33 WorldInvI = Utilities::ComputeWorldSpaceInertia(RCoM, InvI());
+		SetAngularAcceleration(WorldInvI * InTorque);
+	}
+
+	template <typename T, int d>
+	void TPBDRigidParticle<T, d>::AddTorque(const TVector<T, d>& InTorque, bool bInvalidate)
+	{
+		const FRotation3 RCoM = FParticleUtilitiesGT::GetCoMWorldRotation(this);
+		const FMatrix33 WorldInvI = Utilities::ComputeWorldSpaceInertia(RCoM, InvI());
+		SetAngularAcceleration(AngularAcceleration() + WorldInvI * InTorque);
+	}
+
 	template class TGeometryParticle<FReal, 3>;
 
 	template class TKinematicGeometryParticle<FReal, 3>;
 
 	template class TPBDRigidParticle<FReal, 3>;
 
+	template class TParticleHandleBase<FReal, 3>;
+	template class TGeometryParticleHandleImp<FReal, 3, true>;
+	template class TKinematicGeometryParticleHandleImp<FReal, 3, true>;
+	template class TPBDRigidParticleHandleImp<FReal, 3, true>;
+
 	template <>
 	void Chaos::TGeometryParticle<FReal, 3>::MarkDirty(const EChaosPropertyFlags DirtyBits, bool bInvalidate )
 	{
diff --git a/Engine/Source/Runtime/Experimental/Chaos/Private/PhysicsProxy/SingleParticlePhysicsProxy.cpp b/Engine/Source/Runtime/Experimental/Chaos/Private/PhysicsProxy/SingleParticlePhysicsProxy.cpp
--- a/Engine/Source/Runtime/Experimental/Chaos/Private/PhysicsProxy/SingleParticlePhysicsProxy.cpp
+++ b/Engine/Source/Runtime/Experimental/Chaos/Private/PhysicsProxy/SingleParticlePhysicsProxy.cpp
@@ -191,8 +191,8 @@ void FSingleParticlePhysicsProxy::ClearAccumulatedData()
 	{
 		Rigid->ClearForces(false);
 		Rigid->ClearTorques(false);
-		Rigid->SetLinearImpulse(Chaos::FVec3(0), false);
-		Rigid->SetAngularImpulse(Chaos::FVec3(0), false);
+		Rigid->SetLinearImpulseVelocity(Chaos::FVec3(0), false);
+		Rigid->SetAngularImpulseVelocity(Chaos::FVec3(0), false);
 	}
 	
 	Particle->ClearDirtyFlags();
diff --git a/Engine/Source/Runtime/Experimental/Chaos/Public/Chaos/DynamicParticles.h b/Engine/Source/Runtime/Experimental/Chaos/Public/Chaos/DynamicParticles.h
--- a/Engine/Source/Runtime/Experimental/Chaos/Public/Chaos/DynamicParticles.h
+++ b/Engine/Source/Runtime/Experimental/Chaos/Public/Chaos/DynamicParticles.h
@@ -14,16 +14,16 @@ class TDynamicParticles : public TParticles<T, d>
 	    : TParticles<T, d>()
 	{
 		TArrayCollection::AddArray(&MV);
-		TArrayCollection::AddArray(&MF);
+		TArrayCollection::AddArray(&MAcceleration);
 		TArrayCollection::AddArray(&MM);
 		TArrayCollection::AddArray(&MInvM);
 	}
 	TDynamicParticles(const TDynamicParticles<T, d>& Other) = delete;
 	TDynamicParticles(TDynamicParticles<T, d>&& Other)
-	    : TParticles<T, d>(MoveTemp(Other)), MV(MoveTemp(Other.MV)), MF(MoveTemp(Other.MF)), MM(MoveTemp(Other.MM)), MInvM(MoveTemp(Other.MInvM))
+	    : TParticles<T, d>(MoveTemp(Other)), MV(MoveTemp(Other.MV)), MAcceleration(MoveTemp(Other.MAcceleration)), MM(MoveTemp(Other.MM)), MInvM(MoveTemp(Other.MInvM))
 	{
 		TArrayCollection::AddArray(&MV);
-		TArrayCollection::AddArray(&MF);
+		TArrayCollection::AddArray(&MAcceleration);
 		TArrayCollection::AddArray(&MM);
 		TArrayCollection::AddArray(&MInvM);
 	}
@@ -32,8 +32,8 @@ class TDynamicParticles : public TParticles<T, d>
 	TVector<T, d>& V(const int32 Index) { return MV[Index]; }
 	const TArrayCollectionArray<TVector<T, d>>& GetV() const { return MV; }
 
-	const TVector<T, d>& F(const int32 Index) const { return MF[Index]; }
-	TVector<T, d>& F(const int32 Index) { return MF[Index]; }
+	const TVector<T, d>& Acceleration(const int32 Index) const { return MAcceleration[Index]; }
+	TVector<T, d>& Acceleration(const int32 Index) { return MAcceleration[Index]; }
 
 	const T M(const int32 Index) const { return MM[Index]; }
 	T& M(const int32 Index) { return MM[Index]; }
@@ -44,7 +44,7 @@ class TDynamicParticles : public TParticles<T, d>
 	const TArrayCollectionArray<T>& GetInvM() const { return MInvM; }
 
   private:
-	TArrayCollectionArray<TVector<T, d>> MV, MF;
+	TArrayCollectionArray<TVector<T, d>> MV, MAcceleration;
 	TArrayCollectionArray<T> MM, MInvM;
 };
 
diff --git a/Engine/Source/Runtime/Experimental/Chaos/Public/Chaos/PBDRigidParticles.h b/Engine/Source/Runtime/Experimental/Chaos/Public/Chaos/PBDRigidParticles.h
--- a/Engine/Source/Runtime/Experimental/Chaos/Public/Chaos/PBDRigidParticles.h
+++ b/Engine/Source/Runtime/Experimental/Chaos/Public/Chaos/PBDRigidParticles.h
@@ -184,13 +184,13 @@ class TPBDRigidParticles : public TRigidParticles<T, d>
 		// Reset VSmooth to something roughly in the same direction as what V will be after integration.
 		// This is temp fix, if this is only re-computed after solve, island will get incorrectly put back to sleep even if it was just impulsed.
 		FReal FakeDT = (FReal)1. / (FReal)30.;
-		if (this->LinearImpulse(Index).IsNearlyZero() == false || this->F(Index).IsNearlyZero() == false)
+		if (this->LinearImpulseVelocity(Index).IsNearlyZero() == false || this->Acceleration(Index).IsNearlyZero() == false)
 		{
-			this->VSmooth(Index) = this->F(Index) * this->InvM(Index) * FakeDT + this->LinearImpulse(Index) * this->InvM(Index);
+			this->VSmooth(Index) = this->Acceleration(Index) * FakeDT + this->LinearImpulseVelocity(Index);
 		}
-		if (this->AngularImpulse(Index).IsNearlyZero() == false || this->Torque(Index).IsNearlyZero() == false)
+		if (this->AngularImpulseVelocity(Index).IsNearlyZero() == false || this->AngularAcceleration(Index).IsNearlyZero() == false)
 		{
-			this->WSmooth(Index) = this->Torque(Index) * FakeDT + this->AngularImpulse(Index);
+			this->WSmooth(Index) = this->AngularAcceleration(Index) * FakeDT + this->AngularImpulseVelocity(Index);
 		}
 	}
 
diff --git a/Engine/Source/Runtime/Experimental/Chaos/Public/Chaos/PBDRigidsEvolutionGBF.h b/Engine/Source/Runtime/Experimental/Chaos/Public/Chaos/PBDRigidsEvolutionGBF.h
--- a/Engine/Source/Runtime/Experimental/Chaos/Public/Chaos/PBDRigidsEvolutionGBF.h
+++ b/Engine/Source/Runtime/Experimental/Chaos/Public/Chaos/PBDRigidsEvolutionGBF.h
@@ -176,8 +176,8 @@ namespace Chaos
 		CHAOS_API inline void EndFrame(FReal Dt)
 		{
 			Particles.GetNonDisabledDynamicView().ParallelFor([&](auto& Particle, int32 Index) {
-				Particle.F() = FVec3(0);
-				Particle.Torque() = FVec3(0);
+				Particle.Acceleration() = FVec3(0);
+				Particle.AngularAcceleration() = FVec3(0);
 			});
 		}
 
@@ -186,15 +186,12 @@ namespace Chaos
 		{
 			//SCOPE_CYCLE_COUNTER(STAT_Integrate);
 			CHAOS_SCOPED_TIMER(Integrate);
-			FPerParticleEulerStepVelocity EulerStepVelocityRule;
-			FPerParticleAddImpulses AddImpulsesRule;
-			FPerParticleEtherDrag EtherDragRule;
-			FPerParticlePBDEulerStep EulerStepRule;
 
 			const FReal BoundsThickness = GetNarrowPhase().GetBoundsExpansion();
 			const FReal MaxAngularSpeedSq = CVars::HackMaxAngularVelocity * CVars::HackMaxAngularVelocity;
 			const FReal MaxSpeedSq = CVars::HackMaxVelocity * CVars::HackMaxVelocity;
 			InParticles.ParallelFor([&](auto& GeomParticle, int32 Index) {
+	
 				//question: can we enforce this at the API layer? Right now islands contain non dynamic which makes this hard
 				auto PBDParticle = GeomParticle.CastToRigidParticle();
 				if (PBDParticle && PBDParticle->ObjectState() == EObjectStateType::Dynamic)
@@ -209,9 +206,45 @@ namespace Chaos
 					{
 						ForceRule(Particle, Dt);
 					}
-					EulerStepVelocityRule.Apply(Particle, Dt);
-					AddImpulsesRule.Apply(Particle, Dt);
-					EtherDragRule.Apply(Particle, Dt);
+
+					//EulerStepVelocityRule.Apply(Particle, Dt);
+					Particle.V() += Particle.Acceleration() * Dt;
+					Particle.W() += Particle.AngularAcceleration() * Dt;
+
+
+					//AddImpulsesRule.Apply(Particle, Dt);
+					Particle.V() += Particle.LinearImpulseVelocity();
+					Particle.W() += Particle.AngularImpulseVelocity();
+					Particle.LinearImpulseVelocity() = FVec3(0);
+					Particle.AngularImpulseVelocity() = FVec3(0);
+					
+
+					//EtherDragRule.Apply(Particle, Dt);
+					{
+						FVec3& V = Particle.V();
+						FVec3& W = Particle.W();
+
+						const FReal LinearDrag = LinearEtherDragOverride >= 0 ? LinearEtherDragOverride : Particle.LinearEtherDrag() * Dt;
+						const FReal LinearMultiplier = FMath::Max(FReal(0), FReal(1) - LinearDrag);
+						V *= LinearMultiplier;
+
+						const FReal AngularDrag = AngularEtherDragOverride >= 0 ? AngularEtherDragOverride : Particle.AngularEtherDrag() * Dt;
+						const FReal AngularMultiplier = FMath::Max(FReal(0), FReal(1) - AngularDrag);
+						W *= AngularMultiplier;
+
+						const FReal LinearSpeedSq = V.SizeSquared();
+						const FReal AngularSpeedSq = W.SizeSquared();
+
+						if (LinearSpeedSq > Particle.MaxLinearSpeedSq())
+						{
+							V *= FMath::Sqrt(Particle.MaxLinearSpeedSq() / LinearSpeedSq);
+						}
+
+						if (AngularSpeedSq > Particle.MaxAngularSpeedSq())
+						{
+							W *= FMath::Sqrt(Particle.MaxAngularSpeedSq() / AngularSpeedSq);
+						}
+					}
 
 					if (CVars::HackMaxAngularVelocity >= 0.f)
 					{
@@ -231,7 +264,14 @@ namespace Chaos
 						}
 					}
 
-					EulerStepRule.Apply(Particle, Dt);
+					//EulerStepRule.Apply(Particle, Dt);
+					FVec3 PCoM = FParticleUtilitiesXR::GetCoMWorldPosition(&Particle);
+					FRotation3 QCoM = FParticleUtilitiesXR::GetCoMWorldRotation(&Particle);
+
+					PCoM = PCoM + Particle.V() * Dt;
+					QCoM = FRotation3::IntegrateRotationWithAngularVelocity(QCoM, Particle.W(), Dt);
+
+					FParticleUtilitiesPQ::SetCoMWorldTransform(&Particle, PCoM, QCoM);
 
 					if (!Particle.CCDEnabled())
 					{
diff --git a/Engine/Source/Runtime/Experimental/Chaos/Public/Chaos/ParticleDirtyFlags.h b/Engine/Source/Runtime/Experimental/Chaos/Public/Chaos/ParticleDirtyFlags.h
--- a/Engine/Source/Runtime/Experimental/Chaos/Public/Chaos/ParticleDirtyFlags.h
+++ b/Engine/Source/Runtime/Experimental/Chaos/Public/Chaos/ParticleDirtyFlags.h
@@ -145,28 +145,28 @@ class FParticleDynamics
 public:
 	void Serialize(FChaosArchive& Ar)
 	{
-		Ar << MF;
-		Ar << MTorque;
-		Ar << MLinearImpulse;
-		Ar << MAngularImpulse;	
+		Ar << MAcceleration;
+		Ar << MAngularAcceleration;
+		Ar << MLinearImpulseVelocity;
+		Ar << MAngularImpulseVelocity;	
 	}
 
 	template <typename TOther>
 	void CopyFrom(const TOther& Other)
 	{
-		MF = Other.F();
-		MTorque = Other.Torque();
-		MLinearImpulse = Other.LinearImpulse();
-		MAngularImpulse = Other.AngularImpulse();
+		MAcceleration = Other.Acceleration();
+		MAngularAcceleration = Other.AngularAcceleration();
+		MLinearImpulseVelocity = Other.LinearImpulseVelocity();
+		MAngularImpulseVelocity = Other.AngularImpulseVelocity();
 	}
 
 	template <typename TOther>
 	bool IsEqual(const TOther& Other) const
 	{
-		return F() == Other.F()
-			&& Torque() == Other.Torque()
-			&& LinearImpulse() == Other.LinearImpulse()
-			&& AngularImpulse() == Other.AngularImpulse();
+		return Acceleration() == Other.Acceleration()
+			&& AngularAcceleration() == Other.AngularAcceleration()
+			&& LinearImpulseVelocity() == Other.LinearImpulseVelocity()
+			&& AngularImpulseVelocity() == Other.AngularImpulseVelocity();
 	}
 
 	bool operator==(const FParticleDynamics& Other) const
@@ -174,34 +174,34 @@ public:
 		return IsEqual(Other);
 	}
 
-	const FVec3& F() const { return MF; }
-	void SetF(const FVec3& F){ MF = F; }
+	const FVec3& Acceleration() const { return MAcceleration; }
+	void SetAcceleration(const FVec3& Acceleration){ MAcceleration = Acceleration; }
 
-	const FVec3& Torque() const { return MTorque; }
-	void SetTorque(const FVec3& Torque){ MTorque = Torque; }
+	const FVec3& AngularAcceleration() const { return MAngularAcceleration; }
+	void SetAngularAcceleration(const FVec3& AngularAcceleration){ MAngularAcceleration = AngularAcceleration; }
 
-	const FVec3& LinearImpulse() const { return MLinearImpulse; }
-	void SetLinearImpulse(const FVec3& LinearImpulse){ MLinearImpulse = LinearImpulse; }
+	const FVec3& LinearImpulseVelocity() const { return MLinearImpulseVelocity; }
+	void SetLinearImpulseVelocity(const FVec3& LinearImpulseVelocity){ MLinearImpulseVelocity = LinearImpulseVelocity; }
 
-	const FVec3& AngularImpulse() const { return MAngularImpulse; }
-	void SetAngularImpulse(const FVec3& AngularImpulse){ MAngularImpulse = AngularImpulse; }
+	const FVec3& AngularImpulseVelocity() const { return MAngularImpulseVelocity; }
+	void SetAngularImpulseVelocity(const FVec3& AngularImpulseVelocity){ MAngularImpulseVelocity = AngularImpulseVelocity; }
 
 	static FParticleDynamics ZeroValue()
 	{
 		FParticleDynamics Result;
-		Result.MF = FVec3(0);
-		Result.MTorque = FVec3(0);
-		Result.MLinearImpulse = FVec3(0);
-		Result.MAngularImpulse = FVec3(0);
+		Result.MAcceleration = FVec3(0);
+		Result.MAngularAcceleration = FVec3(0);
+		Result.MLinearImpulseVelocity = FVec3(0);
+		Result.MAngularImpulseVelocity = FVec3(0);
 
 		return Result;
 	}
 
 private:
-	FVec3 MF;
-	FVec3 MTorque;
-	FVec3 MLinearImpulse;
-	FVec3 MAngularImpulse;
+	FVec3 MAcceleration;
+	FVec3 MAngularAcceleration;
+	FVec3 MLinearImpulseVelocity;
+	FVec3 MAngularImpulseVelocity;
 
 };
 
diff --git a/Engine/Source/Runtime/Experimental/Chaos/Public/Chaos/ParticleHandle.h b/Engine/Source/Runtime/Experimental/Chaos/Public/Chaos/ParticleHandle.h
--- a/Engine/Source/Runtime/Experimental/Chaos/Public/Chaos/ParticleHandle.h
+++ b/Engine/Source/Runtime/Experimental/Chaos/Public/Chaos/ParticleHandle.h
@@ -89,8 +89,8 @@ void PBDRigidParticleDefaultConstruct(FConcrete& Concrete, const FPBDRigidPartic
 {
 	//don't bother calling parent since the call gets made by the corresponding hierarchy in FConcrete
 	Concrete.SetCollisionGroup(0);
-	Concrete.SetLinearImpulse(TVector<T, d>(0));
-	Concrete.SetAngularImpulse(TVector<T, d>(0));
+	Concrete.SetLinearImpulseVelocity(TVector<T, d>(0));
+	Concrete.SetAngularImpulseVelocity(TVector<T, d>(0));
 	Concrete.SetMaxLinearSpeedSq(TNumericLimits<T>::Max());
 	Concrete.SetMaxAngularSpeedSq(TNumericLimits<T>::Max());
 	Concrete.SetM(1);
@@ -522,7 +522,7 @@ public:
 	TSerializablePtr<FImplicitObject> Geometry() const { return GeometryParticles->Geometry(ParticleIdx); }
 	void SetGeometry(TSerializablePtr<FImplicitObject> InGeometry) { GeometryParticles->SetGeometry(ParticleIdx, InGeometry); }
 
-	TSharedPtr<FImplicitObject, ESPMode::ThreadSafe> SharedGeometry() const { return GeometryParticles->SharedGeometry(ParticleIdx); }
+	TSharedPtr<const FImplicitObject, ESPMode::ThreadSafe> SharedGeometry() const { return GeometryParticles->SharedGeometry(ParticleIdx); }
 	void SetSharedGeometry(TSharedPtr<const FImplicitObject, ESPMode::ThreadSafe> InGeometry) { GeometryParticles->SetSharedGeometry(ParticleIdx, InGeometry); }
 
 	const TSharedPtr<const FImplicitObject, ESPMode::ThreadSafe>& SharedGeometryLowLevel() const { return GeometryParticles->SharedGeometry(ParticleIdx); }
@@ -841,8 +841,8 @@ protected:
 		SetQ(this->R());
 		SetVSmooth(this->V());
 		SetWSmooth(this->W());
-		SetF(TVector<T, d>(0));
-		SetTorque(TVector<T, d>(0));
+		SetAcceleration(TVector<T, d>(0));
+		SetAngularAcceleration(TVector<T, d>(0));
 		SetObjectStateLowLevel(Params.bStartSleeping ? EObjectStateType::Sleeping : EObjectStateType::Dynamic);
 		SetIslandIndex(INDEX_NONE);
 		SetConstraintGraphIndex(INDEX_NONE);
@@ -916,31 +916,29 @@ public:
 	TVector<T, d>& WSmooth() { return PBDRigidParticles->WSmooth(ParticleIdx); }
 	void SetWSmooth(const TVector<T, d>& InWSmooth) { PBDRigidParticles->WSmooth(ParticleIdx) = InWSmooth; }
 
-	const TVector<T, d>& F() const { return PBDRigidParticles->F(ParticleIdx); }
-	TVector<T, d>& F() { return PBDRigidParticles->F(ParticleIdx); }
-	void SetF(const TVector<T, d>& InF) { PBDRigidParticles->F(ParticleIdx) = InF; }
+	const TVector<T, d>& Acceleration() const { return PBDRigidParticles->Acceleration(ParticleIdx); }
+	TVector<T, d>& Acceleration() { return PBDRigidParticles->Acceleration(ParticleIdx); }
+	void SetAcceleration(const TVector<T, d>& InAcceleration) { PBDRigidParticles->Acceleration(ParticleIdx) = InAcceleration; }
 
 	void AddForce(const TVector<T, d>& InF, bool bInvalidate = true)
 	{
-		SetF(F() + InF);
+		SetAcceleration(Acceleration() + InF * InvM());
 	}
 
-	const TVector<T, d>& Torque() const { return PBDRigidParticles->Torque(ParticleIdx); }
-	TVector<T, d>& Torque() { return PBDRigidParticles->Torque(ParticleIdx); }
-	void SetTorque(const TVector<T, d>& InTorque) { PBDRigidParticles->Torque(ParticleIdx) = InTorque; }
+	const TVector<T, d>& AngularAcceleration() const { return PBDRigidParticles->AngularAcceleration(ParticleIdx); }
+	TVector<T, d>& AngularAcceleration() { return PBDRigidParticles->AngularAcceleration(ParticleIdx); }
+	void SetAngularAcceleration(const TVector<T, d>& InAngularAcceleration) { PBDRigidParticles->AngularAcceleration(ParticleIdx) = InAngularAcceleration; }
 
-	void AddTorque(const TVector<T, d>& InTorque, bool bInvalidate = true)
-	{
-		SetTorque(Torque() + InTorque);
-	}
+	CHAOS_API void AddTorque(const TVector<T, d>& InTorque, bool bInvalidate = true);
+	CHAOS_API void SetTorque(const TVector<T, d>& InTorque, bool bInvalidate = true);
 
-	const TVector<T, d>& LinearImpulse() const { return PBDRigidParticles->LinearImpulse(ParticleIdx); }
-	TVector<T, d>& LinearImpulse() { return PBDRigidParticles->LinearImpulse(ParticleIdx); }
-	void SetLinearImpulse(const TVector<T, d>& InLinearImpulse, bool bInvalidate = false) { PBDRigidParticles->LinearImpulse(ParticleIdx) = InLinearImpulse; }
+	const TVector<T, d>& LinearImpulseVelocity() const { return PBDRigidParticles->LinearImpulseVelocity(ParticleIdx); }
+	TVector<T, d>& LinearImpulseVelocity() { return PBDRigidParticles->LinearImpulseVelocity(ParticleIdx); }
+	void SetLinearImpulseVelocity(const TVector<T, d>& InLinearImpulseVelocity, bool bInvalidate = false) { PBDRigidParticles->LinearImpulseVelocity(ParticleIdx) = InLinearImpulseVelocity; }
 
-	const TVector<T, d>& AngularImpulse() const { return PBDRigidParticles->AngularImpulse(ParticleIdx); }
-	TVector<T, d>& AngularImpulse() { return PBDRigidParticles->AngularImpulse(ParticleIdx); }
-	void SetAngularImpulse(const TVector<T, d>& InAngularImpulse, bool bInvalidate = false) { PBDRigidParticles->AngularImpulse(ParticleIdx) = InAngularImpulse; }
+	const TVector<T, d>& AngularImpulseVelocity() const { return PBDRigidParticles->AngularImpulseVelocity(ParticleIdx); }
+	TVector<T, d>& AngularImpulseVelocity() { return PBDRigidParticles->AngularImpulseVelocity(ParticleIdx); }
+	void SetAngularImpulseVelocity(const TVector<T, d>& InAngularImpulseVelocity, bool bInvalidate = false) { PBDRigidParticles->AngularImpulseVelocity(ParticleIdx) = InAngularImpulseVelocity; }
 
 	// Resets VSmooth value to something plausible based on external forces to prevent object from going back to sleep if it was just impulsed.
 	void ResetVSmoothFromForces()
@@ -950,10 +948,10 @@ public:
 
 	void SetDynamics(const FParticleDynamics& Dynamics)
 	{
-		SetF(Dynamics.F());
-		SetTorque(Dynamics.Torque());
-		SetLinearImpulse(Dynamics.LinearImpulse());
-		SetAngularImpulse(Dynamics.AngularImpulse());
+		SetAcceleration(Dynamics.Acceleration());
+		SetAngularAcceleration(Dynamics.AngularAcceleration());
+		SetLinearImpulseVelocity(Dynamics.LinearImpulseVelocity());
+		SetAngularImpulseVelocity(Dynamics.AngularImpulseVelocity());
 	}
 
 	void SetMassProps(const FParticleMassProps& Props)
@@ -1501,20 +1499,20 @@ public:
 		return W();
 	}
 
-	const FVec3& F() const
+	const FVec3& Acceleration() const
 	{ 
 		if (MHandle->CastToRigidParticle() && MHandle->ObjectState() == EObjectStateType::Dynamic)
 		{
-			return MHandle->CastToRigidParticle()->F();
+			return MHandle->CastToRigidParticle()->Acceleration();
 		}
 
 		return ZeroVector;
 	}
-	const FVec3& Torque() const
+	const FVec3& AngularAcceleration() const
 	{ 
 		if (MHandle->CastToRigidParticle() && MHandle->ObjectState() == EObjectStateType::Dynamic)
 		{
-			return MHandle->CastToRigidParticle()->Torque();
+			return MHandle->CastToRigidParticle()->AngularAcceleration();
 		}
 
 		return ZeroVector;
@@ -2478,15 +2476,21 @@ public:
 		this->MInitialized = InInitialized;
 	}
 
-	const TVector<T, d>& F() const { return MDynamics.Read().F(); }
+	const TVector<T, d>& Acceleration() const { return MDynamics.Read().Acceleration(); }
+	void SetAcceleration(const FVec3& Acceleration, bool bInvalidate = true)
+	{ 
+		MDynamics.Modify(bInvalidate, MDirtyFlags, Proxy, [&Acceleration](auto& Data) { Data.SetAcceleration(Acceleration); });
+	}
+
 	void AddForce(const TVector<T, d>& InF, bool bInvalidate = true)
 	{
-		MDynamics.Modify(bInvalidate,MDirtyFlags,Proxy,[&InF](auto& Data){ Data.SetF(InF + Data.F());});
+		FReal InvMass = InvM();
+		MDynamics.Modify(bInvalidate,MDirtyFlags,Proxy,[&InF, InvMass](auto& Data){ Data.SetAcceleration(InF * InvMass + Data.Acceleration());});
 	}
 
 	void ClearForces(bool bInvalidate = true)
 	{
-		MDynamics.Modify(bInvalidate, MDirtyFlags, Proxy, [](auto& Data) { Data.SetF(FVec3(0)); });
+		MDynamics.Modify(bInvalidate, MDirtyFlags, Proxy, [](auto& Data) { Data.SetAcceleration(FVec3(0)); });
 	}
 
 	void ApplyDynamicsWeight(const FReal DynamicsWeight)
@@ -2495,33 +2499,34 @@ public:
 		{
 			MDynamics.Modify(false, MDirtyFlags, Proxy, [DynamicsWeight](auto& Data)
 			{
-				Data.SetF(Data.F() * DynamicsWeight);
-				Data.SetTorque(Data.Torque() * DynamicsWeight);
+				Data.SetAcceleration(Data.Acceleration() * DynamicsWeight);
+				Data.SetAngularAcceleration(Data.AngularAcceleration() * DynamicsWeight);
 			});
 		}
 	}
 
-	const TVector<T, d>& Torque() const { return MDynamics.Read().Torque(); }
-	void AddTorque(const TVector<T, d>& InTorque, bool bInvalidate=true)
+	const TVector<T, d>& AngularAcceleration() const { return MDynamics.Read().AngularAcceleration(); }
+	void SetAngularAcceleration(const TVector<T, d>& InTorque, bool bInvalidate = true)
 	{
-		MDynamics.Modify(bInvalidate,MDirtyFlags,Proxy,[&InTorque](auto& Data){ Data.SetTorque(InTorque + Data.Torque());});
+		MDynamics.Modify(bInvalidate, MDirtyFlags, Proxy, [&InTorque](auto& Data) { Data.SetAngularAcceleration(InTorque);});
 	}
+	CHAOS_API void AddTorque(const TVector<T, d>& InTorque, bool bInvalidate=true);
 
 	void ClearTorques(bool bInvalidate = true)
 	{
-		MDynamics.Modify(bInvalidate, MDirtyFlags, Proxy, [](auto& Data) { Data.SetTorque(FVec3(0)); });
+		MDynamics.Modify(bInvalidate, MDirtyFlags, Proxy, [](auto& Data) { Data.SetAngularAcceleration(FVec3(0)); });
 	}
 
-	const TVector<T, d>& LinearImpulse() const { return MDynamics.Read().LinearImpulse(); }
-	void SetLinearImpulse(const TVector<T, d>& InLinearImpulse, bool bInvalidate = true)
+	const TVector<T, d>& LinearImpulseVelocity() const { return MDynamics.Read().LinearImpulseVelocity(); }
+	void SetLinearImpulseVelocity(const TVector<T, d>& InLinearImpulseVelocity, bool bInvalidate = true)
 	{
-		MDynamics.Modify(bInvalidate,MDirtyFlags,Proxy,[&InLinearImpulse](auto& Data){ Data.SetLinearImpulse(InLinearImpulse);});
+		MDynamics.Modify(bInvalidate,MDirtyFlags,Proxy,[&InLinearImpulseVelocity](auto& Data){ Data.SetLinearImpulseVelocity(InLinearImpulseVelocity);});
 	}
 
-	const TVector<T, d>& AngularImpulse() const { return MDynamics.Read().AngularImpulse(); }
-	void SetAngularImpulse(const TVector<T, d>& InAngularImpulse, bool bInvalidate = true)
+	const TVector<T, d>& AngularImpulseVelocity() const { return MDynamics.Read().AngularImpulseVelocity(); }
+	void SetAngularImpulseVelocity(const TVector<T, d>& InAngularImpulseVelocity, bool bInvalidate = true)
 	{
-		MDynamics.Modify(bInvalidate,MDirtyFlags,Proxy,[&InAngularImpulse](auto& Data){ Data.SetAngularImpulse(InAngularImpulse);});
+		MDynamics.Modify(bInvalidate,MDirtyFlags,Proxy,[&InAngularImpulseVelocity](auto& Data){ Data.SetAngularImpulseVelocity(InAngularImpulseVelocity);});
 	}
 
 	void SetDynamics(const FParticleDynamics& InDynamics,bool bInvalidate = true)
diff --git a/Engine/Source/Runtime/Experimental/Chaos/Public/Chaos/PerParticleAddImpulses.h b/Engine/Source/Runtime/Experimental/Chaos/Public/Chaos/PerParticleAddImpulses.h
--- a/Engine/Source/Runtime/Experimental/Chaos/Public/Chaos/PerParticleAddImpulses.h
+++ b/Engine/Source/Runtime/Experimental/Chaos/Public/Chaos/PerParticleAddImpulses.h
@@ -16,7 +16,7 @@ namespace Chaos
 		template<class T_PARTICLES>
 		inline void ApplyHelper(T_PARTICLES& InParticles, const FReal Dt, const int32 Index) const
 		{
-			InParticles.V(Index) += InParticles.LinearImpulse(Index) * InParticles.InvM(Index);
+			InParticles.V(Index) += InParticles.LinearImpulseVelocity(Index);
 		}
 
 		inline void Apply(FDynamicParticles& InParticles, const FReal Dt, const int32 Index) const override //-V762
@@ -38,28 +38,18 @@ namespace Chaos
 			//       Just using W += InvI * (Torque - W x (I * W)) * dt is not correct, since Torque
 			//		 and W are in an inertial frame.
 			//
-#if CHAOS_PARTICLE_ACTORTRANSFORM
-			const FMatrix33 WorldInvI = Utilities::ComputeWorldSpaceInertia(InParticles.R(Index) * InParticles.RotationOfMass(Index), InParticles.InvI(Index));
-#else
-			const FMatrix33 WorldInvI = Utilities::ComputeWorldSpaceInertia(InParticles.R(Index), InParticles.InvI(Index));
-#endif
 
-			InParticles.W(Index) += WorldInvI * InParticles.AngularImpulse(Index);
-			InParticles.LinearImpulse(Index) = FVec3(0);
-			InParticles.AngularImpulse(Index) = FVec3(0);
+			InParticles.W(Index) += InParticles.AngularImpulseVelocity(Index);
+			InParticles.LinearImpulseVelocity(Index) = FVec3(0);
+			InParticles.AngularImpulseVelocity(Index) = FVec3(0);
 		}
 
 		inline void Apply(TTransientPBDRigidParticleHandle<FReal, 3>& Particle, const FReal Dt) const override //-V762
 		{
-			Particle.V() += Particle.LinearImpulse() * Particle.InvM();
-#if CHAOS_PARTICLE_ACTORTRANSFORM
-			const FMatrix33 WorldInvI = Utilities::ComputeWorldSpaceInertia(Particle.R() * Particle.RotationOfMass(), Particle.InvI());
-#else
-			const FMatrix33 WorldInvI = Utilities::ComputeWorldSpaceInertia(Particle.R(), Particle.InvI());
-#endif
-			Particle.W() += WorldInvI * Particle.AngularImpulse();
-			Particle.LinearImpulse() = FVec3(0);
-			Particle.AngularImpulse() = FVec3(0);
+			Particle.V() += Particle.LinearImpulseVelocity();
+			Particle.W() += Particle.AngularImpulseVelocity();
+			Particle.LinearImpulseVelocity() = FVec3(0);
+			Particle.AngularImpulseVelocity() = FVec3(0);
 		}
 	};
 
diff --git a/Engine/Source/Runtime/Experimental/Chaos/Public/Chaos/PerParticleEulerStepVelocity.h b/Engine/Source/Runtime/Experimental/Chaos/Public/Chaos/PerParticleEulerStepVelocity.h
--- a/Engine/Source/Runtime/Experimental/Chaos/Public/Chaos/PerParticleEulerStepVelocity.h
+++ b/Engine/Source/Runtime/Experimental/Chaos/Public/Chaos/PerParticleEulerStepVelocity.h
@@ -16,7 +16,7 @@ class FPerParticleEulerStepVelocity : public FPerParticleRule
 	template<class T_PARTICLES>
 	inline void ApplyHelper(T_PARTICLES& InParticles, const FReal Dt, const int32 Index) const
 	{
-		InParticles.V(Index) += InParticles.F(Index) * InParticles.InvM(Index) * Dt;
+		InParticles.V(Index) += InParticles.Acceleration(Index) * Dt;
 	}
 
 	inline void Apply(FDynamicParticles& InParticles, const FReal Dt, const int32 Index) const override //-V762
@@ -39,23 +39,13 @@ class FPerParticleEulerStepVelocity : public FPerParticleRule
 		//       Just using W += InvI * (Torque - W x (I * W)) * dt is not correct, since Torque
 		//		 and W are in an inertial frame.
 		//
-#if CHAOS_PARTICLE_ACTORTRANSFORM
-		const FMatrix33 WorldInvI = Utilities::ComputeWorldSpaceInertia(InParticles.R(Index) * InParticles.RotationOfMass(Index), InParticles.InvI(Index));
-#else
-		const FMatrix33 WorldInvI = Utilities::ComputeWorldSpaceInertia(InParticles.R(Index), InParticles.InvI(Index));
-#endif
-		InParticles.W(Index) += WorldInvI * InParticles.Torque(Index) * Dt;
+		InParticles.W(Index) += InParticles.AngularAcceleration(Index) * Dt;
 	}
 	
 	inline void Apply(TTransientPBDRigidParticleHandle<FReal, 3>& Particle, const FReal Dt) const override //-V762
 	{
-		Particle.V() += Particle.F() * Particle.InvM() * Dt;
-#if CHAOS_PARTICLE_ACTORTRANSFORM
-		const FMatrix33 WorldInvI = Utilities::ComputeWorldSpaceInertia(Particle.R() * Particle.RotationOfMass(), Particle.InvI());
-#else
-		const FMatrix33 WorldInvI = Utilities::ComputeWorldSpaceInertia(Particle.R(Index), Particle.InvI());
-#endif
-		Particle.W() += WorldInvI * Particle.Torque() * Dt;
+		Particle.V() += Particle.Acceleration() * Dt;
+		Particle.W() += Particle.AngularAcceleration() * Dt;
 	}
 };
 
diff --git a/Engine/Source/Runtime/Experimental/Chaos/Public/Chaos/PerParticleGravity.h b/Engine/Source/Runtime/Experimental/Chaos/Public/Chaos/PerParticleGravity.h
--- a/Engine/Source/Runtime/Experimental/Chaos/Public/Chaos/PerParticleGravity.h
+++ b/Engine/Source/Runtime/Experimental/Chaos/Public/Chaos/PerParticleGravity.h
@@ -24,7 +24,7 @@ namespace Chaos
 		template<class T_PARTICLES>
 		inline void ApplyHelper(T_PARTICLES& InParticles, const FReal Dt, const int Index) const
 		{
-			InParticles.F(Index) += MAcceleration * InParticles.M(Index);
+			InParticles.Acceleration(Index) += MAcceleration;
 		}
 		inline void Apply(FPBDParticles& InParticles, const FReal Dt, const int Index) const override //-V762
 		{
@@ -35,7 +35,7 @@ namespace Chaos
 		{
 			if(Particle.GravityEnabled())
 			{
-				Particle.F() += MAcceleration * Particle.M();
+				Particle.Acceleration() += MAcceleration;
 			}
 		}
 
diff --git a/Engine/Source/Runtime/Experimental/Chaos/Public/Chaos/PerParticleInitForce.h b/Engine/Source/Runtime/Experimental/Chaos/Public/Chaos/PerParticleInitForce.h
--- a/Engine/Source/Runtime/Experimental/Chaos/Public/Chaos/PerParticleInitForce.h
+++ b/Engine/Source/Runtime/Experimental/Chaos/Public/Chaos/PerParticleInitForce.h
@@ -14,19 +14,19 @@ class FPerParticleInitForce : public FPerParticleRule
 
 	inline void Apply(FDynamicParticles& InParticles, const FReal Dt, const int Index) const override //-V762
 	{
-		InParticles.F(Index) = FVec3(0);
+		InParticles.Acceleration(Index) = FVec3(0);
 	}
 
 	inline void Apply(TRigidParticles<FReal, 3>& InParticles, const FReal Dt, const int Index) const override //-V762
 	{
-		InParticles.F(Index) = FVec3(0);
-		InParticles.Torque(Index) = FVec3(0);
+		InParticles.Acceleration(Index) = FVec3(0);
+		InParticles.AngularAcceleration(Index) = FVec3(0);
 	}
 
 	inline void Apply(TTransientPBDRigidParticleHandle<FReal, 3>& Particle, const FReal Dt) const override //-V762
 	{
-		Particle.F() = FVec3(0);
-		Particle.Torque() = FVec3(0);
+		Particle.Acceleration() = FVec3(0);
+		Particle.AngularAcceleration() = FVec3(0);
 	}
 };
 
diff --git a/Engine/Source/Runtime/Experimental/Chaos/Public/Chaos/RigidParticles.h b/Engine/Source/Runtime/Experimental/Chaos/Public/Chaos/RigidParticles.h
--- a/Engine/Source/Runtime/Experimental/Chaos/Public/Chaos/RigidParticles.h
+++ b/Engine/Source/Runtime/Experimental/Chaos/Public/Chaos/RigidParticles.h
@@ -69,10 +69,10 @@ public:
 	{
 		TArrayCollection::AddArray(&MVSmooth);
 		TArrayCollection::AddArray(&MWSmooth);
-		TArrayCollection::AddArray(&MF);
-		TArrayCollection::AddArray(&MT);
-		TArrayCollection::AddArray(&MLinearImpulse);
-		TArrayCollection::AddArray(&MAngularImpulse);
+		TArrayCollection::AddArray(&MAcceleration);
+		TArrayCollection::AddArray(&MAngularAcceleration);
+		TArrayCollection::AddArray(&MLinearImpulseVelocity);
+		TArrayCollection::AddArray(&MAngularImpulseVelocity);
 		TArrayCollection::AddArray(&MI);
 		TArrayCollection::AddArray(&MInvI);
 		TArrayCollection::AddArray(&MM);
@@ -102,10 +102,10 @@ public:
 	    : TKinematicGeometryParticles<T, d>(MoveTemp(Other))
 		, MVSmooth(MoveTemp(Other.MVSmooth))
 		, MWSmooth(MoveTemp(Other.MWSmooth))
-		, MF(MoveTemp(Other.MF))
-		, MT(MoveTemp(Other.MT))
-		, MLinearImpulse(MoveTemp(Other.MLinearImpulse))
-		, MAngularImpulse(MoveTemp(Other.MAngularImpulse))
+		, MAcceleration(MoveTemp(Other.MAcceleration))
+		, MAngularAcceleration(MoveTemp(Other.MAngularAcceleration))
+		, MLinearImpulseVelocity(MoveTemp(Other.MLinearImpulseVelocity))
+		, MAngularImpulseVelocity(MoveTemp(Other.MAngularImpulseVelocity))
 		, MI(MoveTemp(Other.MI)), MInvI(MoveTemp(Other.MInvI))
 		, MM(MoveTemp(Other.MM))
 		, MInvM(MoveTemp(Other.MInvM))
@@ -125,10 +125,10 @@ public:
 	{
 		TArrayCollection::AddArray(&MVSmooth);
 		TArrayCollection::AddArray(&MWSmooth);
-		TArrayCollection::AddArray(&MF);
-		TArrayCollection::AddArray(&MT);
-		TArrayCollection::AddArray(&MLinearImpulse);
-		TArrayCollection::AddArray(&MAngularImpulse);
+		TArrayCollection::AddArray(&MAcceleration);
+		TArrayCollection::AddArray(&MAngularAcceleration);
+		TArrayCollection::AddArray(&MLinearImpulseVelocity);
+		TArrayCollection::AddArray(&MAngularImpulseVelocity);
 		TArrayCollection::AddArray(&MI);
 		TArrayCollection::AddArray(&MInvI);
 		TArrayCollection::AddArray(&MM);
@@ -163,17 +163,17 @@ public:
 	FORCEINLINE const TVector<T, d>& WSmooth(const int32 Index) const { return MWSmooth[Index]; }
 	FORCEINLINE TVector<T, d>& WSmooth(const int32 Index) { return MWSmooth[Index]; }
 
-	FORCEINLINE const TVector<T, d>& Torque(const int32 Index) const { return MT[Index]; }
-	FORCEINLINE TVector<T, d>& Torque(const int32 Index) { return MT[Index]; }
+	FORCEINLINE const TVector<T, d>& AngularAcceleration(const int32 Index) const { return MAngularAcceleration[Index]; }
+	FORCEINLINE TVector<T, d>& AngularAcceleration(const int32 Index) { return MAngularAcceleration[Index]; }
 
-	FORCEINLINE const TVector<T, d>& F(const int32 Index) const { return MF[Index]; }
-	FORCEINLINE TVector<T, d>& F(const int32 Index) { return MF[Index]; }
+	FORCEINLINE const TVector<T, d>& Acceleration(const int32 Index) const { return MAcceleration[Index]; }
+	FORCEINLINE TVector<T, d>& Acceleration(const int32 Index) { return MAcceleration[Index]; }
 
-	FORCEINLINE const TVector<T, d>& LinearImpulse(const int32 Index) const { return MLinearImpulse[Index]; }
-	FORCEINLINE TVector<T, d>& LinearImpulse(const int32 Index) { return MLinearImpulse[Index]; }
+	FORCEINLINE const TVector<T, d>& LinearImpulseVelocity(const int32 Index) const { return MLinearImpulseVelocity[Index]; }
+	FORCEINLINE TVector<T, d>& LinearImpulseVelocity(const int32 Index) { return MLinearImpulseVelocity[Index]; }
 
-	FORCEINLINE const TVector<T, d>& AngularImpulse(const int32 Index) const { return MAngularImpulse[Index]; }
-	FORCEINLINE TVector<T, d>& AngularImpulse(const int32 Index) { return MAngularImpulse[Index]; }
+	FORCEINLINE const TVector<T, d>& AngularImpulseVelocity(const int32 Index) const { return MAngularImpulseVelocity[Index]; }
+	FORCEINLINE TVector<T, d>& AngularImpulseVelocity(const int32 Index) { return MAngularImpulseVelocity[Index]; }
 
 	FORCEINLINE const PMatrix<T, d, d>& I(const int32 Index) const { return MI[Index]; }
 	FORCEINLINE PMatrix<T, d, d>& I(const int32 Index) { return MI[Index]; }
@@ -287,8 +287,8 @@ public:
 	FORCEINLINE FString ToString(int32 index) const
 	{
 		FString BaseString = TKinematicGeometryParticles<T, d>::ToString(index);
-		return FString::Printf(TEXT("%s, MF:%s, MT:%s, MLinearImpulse:%s, MAngularImpulse:%s, MI:%s, MInvI:%s, MM:%f, MInvM:%f, MCenterOfMass:%s, MRotationOfMass:%s, MCollisionParticles(num):%d, MCollisionGroup:%d, MDisabled:%d, MSleeping:%d, MIslandIndex:%d"),
-			*BaseString, *F(index).ToString(), *Torque(index).ToString(), *LinearImpulse(index).ToString(), *AngularImpulse(index).ToString(),
+		return FString::Printf(TEXT("%s, MAcceleration:%s, MAngularAcceleration:%s, MLinearImpulseVelocity:%s, MAngularImpulseVelocity:%s, MI:%s, MInvI:%s, MM:%f, MInvM:%f, MCenterOfMass:%s, MRotationOfMass:%s, MCollisionParticles(num):%d, MCollisionGroup:%d, MDisabled:%d, MSleeping:%d, MIslandIndex:%d"),
+			*BaseString, *Acceleration(index).ToString(), *AngularAcceleration(index).ToString(), *LinearImpulseVelocity(index).ToString(), *AngularImpulseVelocity(index).ToString(),
 			*I(index).ToString(), *InvI(index).ToString(), M(index), InvM(index), *CenterOfMass(index).ToString(), *RotationOfMass(index).ToString(), CollisionParticlesSize(index),
 			CollisionGroup(index), Disabled(index), Sleeping(index), IslandIndex(index));
 	}
@@ -304,7 +304,7 @@ public:
 			Ar << MRotationOfMass;
 		}
 
-		Ar << MF << MT << MLinearImpulse << MAngularImpulse << MI << MInvI << MM << MInvM;
+		Ar << MAcceleration << MAngularAcceleration << MLinearImpulseVelocity << MAngularImpulseVelocity << MI << MInvI << MM << MInvM;
 
 		Ar.UsingCustomVersion(FExternalPhysicsCustomObjectVersion::GUID);
 		if (Ar.CustomVer(FExternalPhysicsCustomObjectVersion::GUID) >= FExternalPhysicsCustomObjectVersion::AddDampingToRigids)
@@ -316,10 +316,10 @@ public:
 		//todo: add gravity enabled when we decide how we want to handle serialization
 	}
 
-	FORCEINLINE TArray<TVector<T, d>>& AllF() { return MF; }
-	FORCEINLINE TArray<TVector<T, d>>& AllT() { return MT; }
-	FORCEINLINE TArray<TVector<T, d>>& AllLinearImpulse() { return MLinearImpulse; }
-	FORCEINLINE TArray<TVector<T, d>>& AllAngularImpulse() { return MAngularImpulse; }
+	FORCEINLINE TArray<TVector<T, d>>& AllAcceleration() { return MAcceleration; }
+	FORCEINLINE TArray<TVector<T, d>>& AllAngularAcceleration() { return MAngularAcceleration; }
+	FORCEINLINE TArray<TVector<T, d>>& AllLinearImpulseVelocity() { return MLinearImpulseVelocity; }
+	FORCEINLINE TArray<TVector<T, d>>& AllAngularImpulseVelocity() { return MAngularImpulseVelocity; }
 	FORCEINLINE TArray<PMatrix<T, d, d>>& AllI() { return MI; }
 	FORCEINLINE TArray<PMatrix<T, d, d>>& AllInvI() { return MInvI; }
 	FORCEINLINE TArray<FReal>& AllM() { return MM; }
@@ -338,10 +338,10 @@ public:
 private:
 	TArrayCollectionArray<TVector<T, d>> MVSmooth;
 	TArrayCollectionArray<TVector<T, d>> MWSmooth;
-	TArrayCollectionArray<TVector<T, d>> MF;
-	TArrayCollectionArray<TVector<T, d>> MT;
-	TArrayCollectionArray<TVector<T, d>> MLinearImpulse;
-	TArrayCollectionArray<TVector<T, d>> MAngularImpulse;
+	TArrayCollectionArray<TVector<T, d>> MAcceleration;
+	TArrayCollectionArray<TVector<T, d>> MAngularAcceleration;
+	TArrayCollectionArray<TVector<T, d>> MLinearImpulseVelocity;
+	TArrayCollectionArray<TVector<T, d>> MAngularImpulseVelocity;
 	TArrayCollectionArray<PMatrix<T, d, d>> MI;
 	TArrayCollectionArray<PMatrix<T, d, d>> MInvI;
 	TArrayCollectionArray<T> MM;
diff --git a/Engine/Source/Runtime/Experimental/Chaos/Public/Chaos/VelocityField.h b/Engine/Source/Runtime/Experimental/Chaos/Public/Chaos/VelocityField.h
--- a/Engine/Source/Runtime/Experimental/Chaos/Public/Chaos/VelocityField.h
+++ b/Engine/Source/Runtime/Experimental/Chaos/Public/Chaos/VelocityField.h
@@ -70,7 +70,7 @@ namespace Chaos
 			const TArray<int32>& ElementIndices = PointToTriangleMap[Index];
 			for (const int32 ElementIndex : ElementIndices)
 			{
-				InParticles.F(Index) += Forces[ElementIndex];
+				InParticles.Acceleration(Index) += InParticles.InvM(Index) * Forces[ElementIndex];
 			}
 		}
 
diff --git a/Engine/Source/Runtime/Experimental/Chaos/Public/PhysicsProxy/FieldSystemProxyHelper.h b/Engine/Source/Runtime/Experimental/Chaos/Public/PhysicsProxy/FieldSystemProxyHelper.h
--- a/Engine/Source/Runtime/Experimental/Chaos/Public/PhysicsProxy/FieldSystemProxyHelper.h
+++ b/Engine/Source/Runtime/Experimental/Chaos/Public/PhysicsProxy/FieldSystemProxyHelper.h
@@ -641,7 +641,7 @@ namespace Chaos
 						{
 							RigidHandle->SetObjectStateLowLevel(Chaos::EObjectStateType::Dynamic);
 						}
-						RigidHandle->F() += ResultsView[Index.Result];
+						RigidHandle->AddForce(ResultsView[Index.Result]);
 					}
 				}
 			}
@@ -659,7 +659,7 @@ namespace Chaos
 						{
 							RigidHandle->SetObjectStateLowLevel(Chaos::EObjectStateType::Dynamic);
 						}
-						RigidHandle->Torque() += ResultsView[Index.Result];
+						RigidHandle->AddTorque(ResultsView[Index.Result]);
 					}
 				}
 			}
diff --git a/Engine/Source/Runtime/Experimental/Chaos/Public/PhysicsProxy/SingleParticlePhysicsProxy.h b/Engine/Source/Runtime/Experimental/Chaos/Public/PhysicsProxy/SingleParticlePhysicsProxy.h
--- a/Engine/Source/Runtime/Experimental/Chaos/Public/PhysicsProxy/SingleParticlePhysicsProxy.h
+++ b/Engine/Source/Runtime/Experimental/Chaos/Public/PhysicsProxy/SingleParticlePhysicsProxy.h
@@ -442,18 +442,37 @@ public:
 		return EResimType::FullResim;
 	}
 
-	const FVec3 F() const
+	const FVec3 Acceleration() const
 	{
 		return Read([](auto* Particle)
 		{
 			if (auto Rigid = Particle->CastToRigidParticle())
 			{
-				return Rigid->F();
+				return Rigid->Acceleration();
 			}
 
 			return FVec3(0);
 		});
 	}
+	
+	void SetAcceleration(const FVec3& Acceleration, bool bInvalidate = true)
+	{
+		Write([&Acceleration, bInvalidate, this](auto* Particle)
+			{
+				if (auto Rigid = Particle->CastToRigidParticle())
+				{
+					if (Rigid->ObjectState() == EObjectStateType::Sleeping || Rigid->ObjectState() == EObjectStateType::Dynamic)
+					{
+						if (bInvalidate)
+						{
+							SetObjectStateHelper(*GetProxy(), *Rigid, EObjectStateType::Dynamic, true);
+						}
+
+						Rigid->SetAcceleration(Acceleration);
+					}
+				}
+			});
+	}
 
 	void AddForce(const FVec3& InForce, bool bInvalidate = true)
 	{
@@ -474,13 +493,32 @@ public:
 		});
 	}
 
-	const FVec3 Torque() const
+	void SetAngularAcceleration(const FVec3& AngularAcceleration, bool bInvalidate = true)
+	{
+		Write([&AngularAcceleration, bInvalidate, this](auto* Particle)
+			{
+				if (auto Rigid = Particle->CastToRigidParticle())
+				{
+					if (Rigid->ObjectState() == EObjectStateType::Sleeping || Rigid->ObjectState() == EObjectStateType::Dynamic)
+					{
+						if (bInvalidate)
+						{
+							SetObjectStateHelper(*GetProxy(), *Rigid, EObjectStateType::Dynamic, true);
+						}
+
+						Rigid->SetAngularAcceleration(AngularAcceleration);
+					}
+				}
+			});
+	}
+
+	const FVec3 AngularAcceleration() const
 	{
 		return Read([](auto* Particle)
 		{
 			if (auto Rigid = Particle->CastToRigidParticle())
 			{
-				return Rigid->Torque();
+				return Rigid->AngularAcceleration();
 			}
 
 			return FVec3(0);
@@ -506,22 +544,35 @@ public:
 		});
 	}
 
+	const FVec3 LinearImpulseVelocity() const
+	{
+		return Read([](auto* Particle)
+		{
+			if (auto Rigid = Particle->CastToRigidParticle())
+			{
+				return Rigid->LinearImpulseVelocity();
+			}
+
+			return FVec3(0);
+		});
+	}
+
 	const FVec3 LinearImpulse() const
 	{
 		return Read([](auto* Particle)
 		{
 			if (auto Rigid = Particle->CastToRigidParticle())
 			{
-				return Rigid->LinearImpulse();
+				return Rigid->LinearImpulseVelocity() * Rigid->M();
 			}
 
 			return FVec3(0);
 		});
 	}
 
-	void SetLinearImpulse(const FVec3& InLinearImpulse, bool bInvalidate = true)
+	void SetLinearImpulse(const FVec3& InLinearImpulse, bool bIsVelocity, bool bInvalidate = true)
 	{
-		Write([&InLinearImpulse, bInvalidate, this](auto* Particle)
+		Write([&InLinearImpulse, bIsVelocity, bInvalidate, this](auto* Particle)
 		{
 			if (auto Rigid = Particle->CastToRigidParticle())
 			{
@@ -531,28 +582,50 @@ public:
 					{
 						SetObjectStateHelper(*GetProxy(), *Rigid, EObjectStateType::Dynamic, true);
 					}
-					Rigid->SetLinearImpulse(InLinearImpulse, bInvalidate);
+
+					if (bIsVelocity)
+					{
+						Rigid->SetLinearImpulseVelocity(InLinearImpulse, bInvalidate);
+					}
+					else
+					{
+						Rigid->SetLinearImpulseVelocity(InLinearImpulse * Rigid->InvM(), bInvalidate);
+					}
 				}
 			}
 		});
 	}
 
+	const FVec3 AngularImpulseVelocity() const
+	{
+		return Read([](auto* Particle)
+		{
+			if (auto Rigid = Particle->CastToRigidParticle())
+			{
+				return Rigid->AngularImpulseVelocity();
+			}
+
+			return FVec3(0);
+		});
+	}
+
 	const FVec3 AngularImpulse() const
 	{
 		return Read([](auto* Particle)
 		{
 			if (auto Rigid = Particle->CastToRigidParticle())
 			{
-				return Rigid->AngularImpulse();
+				const FMatrix33 WorldI = Utilities::ComputeWorldSpaceInertia(Rigid->R() * Rigid->RotationOfMass(), Rigid->I());
+				return WorldI * Rigid->AngularImpulseVelocity();
 			}
 
 			return FVec3(0);
 		});
 	}
 
-	void SetAngularImpulse(const FVec3& InAngularImpulse, bool bInvalidate = true)
+	void SetAngularImpulse(const FVec3& InAngularImpulse, bool bIsVelocity, bool bInvalidate = true)
 	{
-		Write([&InAngularImpulse, bInvalidate, this](auto* Particle)
+		Write([&InAngularImpulse, bIsVelocity, bInvalidate, this](auto* Particle)
 		{
 			if (auto Rigid = Particle->CastToRigidParticle())
 			{
@@ -562,7 +635,16 @@ public:
 					{
 						SetObjectStateHelper(*GetProxy(), *Rigid, EObjectStateType::Dynamic, true);
 					}
-					Rigid->SetAngularImpulse(InAngularImpulse, bInvalidate);
+
+					if (bIsVelocity)
+					{
+						Rigid->SetAngularImpulseVelocity(InAngularImpulse, bInvalidate);
+					}
+					else
+					{
+						const FMatrix33 WorldInvI = Utilities::ComputeWorldSpaceInertia(Rigid->R() * Rigid->RotationOfMass(), Rigid->InvI());
+						Rigid->SetAngularImpulseVelocity(WorldInvI * InAngularImpulse, bInvalidate);
+					}
 				}
 			}
 		});
diff --git a/Engine/Source/Runtime/Experimental/Chaos/Public/RewindData.h b/Engine/Source/Runtime/Experimental/Chaos/Public/RewindData.h
--- a/Engine/Source/Runtime/Experimental/Chaos/Public/RewindData.h
+++ b/Engine/Source/Runtime/Experimental/Chaos/Public/RewindData.h
@@ -672,10 +672,10 @@ public:
 	REWIND_PARTICLE_STATIC_PROPERTY(NonFrequentData, DebugName)
 #endif
 
-	REWIND_PARTICLE_ZERO_PROPERTY(Dynamics, F)
-	REWIND_PARTICLE_ZERO_PROPERTY(Dynamics, Torque)
-	REWIND_PARTICLE_ZERO_PROPERTY(Dynamics, LinearImpulse)
-	REWIND_PARTICLE_ZERO_PROPERTY(Dynamics, AngularImpulse)
+	REWIND_PARTICLE_ZERO_PROPERTY(Dynamics, Acceleration)
+	REWIND_PARTICLE_ZERO_PROPERTY(Dynamics, AngularAcceleration)
+	REWIND_PARTICLE_ZERO_PROPERTY(Dynamics, LinearImpulseVelocity)
+	REWIND_PARTICLE_ZERO_PROPERTY(Dynamics, AngularImpulseVelocity)
 
 	TShapesArrayState<FGeometryParticleHandle> ShapesArray() const
 	{
diff --git a/Engine/Source/Runtime/PhysicsCore/Private/ChaosEngineInterface.cpp b/Engine/Source/Runtime/PhysicsCore/Private/ChaosEngineInterface.cpp
--- a/Engine/Source/Runtime/PhysicsCore/Private/ChaosEngineInterface.cpp
+++ b/Engine/Source/Runtime/PhysicsCore/Private/ChaosEngineInterface.cpp
@@ -734,7 +734,7 @@ void FChaosEngineInterface::AddImpulse_AssumesLocked(const FPhysicsActorHandle&
 	if(ensure(FChaosEngineInterface::IsValid(InActorReference)))
 	{
 		Chaos::FRigidBodyHandle_External& Body_External = InActorReference->GetGameThreadAPI();
-		Body_External.SetLinearImpulse(Body_External.LinearImpulse() + InForce);
+		Body_External.SetLinearImpulse(Body_External.LinearImpulse() + InForce, /*bIsVelocity=*/false);
 	}
 }
 
@@ -743,7 +743,7 @@ void FChaosEngineInterface::AddAngularImpulseInRadians_AssumesLocked(const FPhys
 	if(ensure(FChaosEngineInterface::IsValid(InActorReference)))
 	{
 		Chaos::FRigidBodyHandle_External& Body_External = InActorReference->GetGameThreadAPI();
-		Body_External.SetAngularImpulse(Body_External.AngularImpulse() + InTorque);
+		Body_External.SetAngularImpulse(Body_External.AngularImpulse() + InTorque, /*bIsVelocity=*/false);
 	}
 }
 
