Fix compiler error in RewindData.h

diff --git a/Engine/Source/Runtime/Experimental/Chaos/Public/RewindData.h b/Engine/Source/Runtime/Experimental/Chaos/Public/RewindData.h
--- a/Engine/Source/Runtime/Experimental/Chaos/Public/RewindData.h
+++ b/Engine/Source/Runtime/Experimental/Chaos/Public/RewindData.h
@@ -728,10 +728,10 @@ public:
 			REWIND_PARTICLE_TO_STR(M)
 			REWIND_PARTICLE_TO_STR(InvM)
 
-			REWIND_PARTICLE_TO_STR(F)
-			REWIND_PARTICLE_TO_STR(Torque)
-			REWIND_PARTICLE_TO_STR(LinearImpulse)
-			REWIND_PARTICLE_TO_STR(AngularImpulse)
+			REWIND_PARTICLE_TO_STR(Acceleration)
+			REWIND_PARTICLE_TO_STR(AngularAcceleration)
+			REWIND_PARTICLE_TO_STR(LinearImpulseVelocity)
+			REWIND_PARTICLE_TO_STR(AngularImpulseVelocity)
 		}
 
 		return Out;
