Improve code that detects wheel wrap around, so it won't stall.

diff --git a/Engine/Source/Runtime/Experimental/ChaosVehicles/ChaosVehiclesCore/Private/WheelSystem.cpp b/Engine/Source/Runtime/Experimental/ChaosVehicles/ChaosVehiclesCore/Private/WheelSystem.cpp
--- a/Engine/Source/Runtime/Experimental/ChaosVehicles/ChaosVehiclesCore/Private/WheelSystem.cpp
+++ b/Engine/Source/Runtime/Experimental/ChaosVehicles/ChaosVehiclesCore/Private/WheelSystem.cpp
@@ -176,21 +176,16 @@ namespace Chaos
 		}
 		else
 		{ 
-			float GroundOmega = GroundVelocityVector.X / Re;
+			float GroundOmega = GroundVelocityVector.X / FMath::Max(Re, KINDA_SMALL_NUMBER);
 			Omega += ((GroundOmega - Omega + SlipOmega));
 		}
 
-		// Wheel angular position
+		// Wheel angular position integrated
 		AngularPosition += Omega * DeltaTime;
 
-		while (AngularPosition >= PI * 2.f)
-		{
-			AngularPosition -= PI * 2.f;
-		}
-		while (AngularPosition <= -PI * 2.f)
-		{
-			AngularPosition += PI * 2.f;
-		}
+		// Handle wrap around of wheel position
+		float IntegerPart = 0;
+		AngularPosition = FMath::Modf(AngularPosition / TWO_PI, &IntegerPart);
 
 		if (!bInContact)
 		{
