Correction in math of last change CL 17473281

diff --git a/Engine/Source/Runtime/Experimental/ChaosVehicles/ChaosVehiclesCore/Private/WheelSystem.cpp b/Engine/Source/Runtime/Experimental/ChaosVehicles/ChaosVehiclesCore/Private/WheelSystem.cpp
--- a/Engine/Source/Runtime/Experimental/ChaosVehicles/ChaosVehiclesCore/Private/WheelSystem.cpp
+++ b/Engine/Source/Runtime/Experimental/ChaosVehicles/ChaosVehiclesCore/Private/WheelSystem.cpp
@@ -184,8 +184,8 @@ namespace Chaos
 		AngularPosition += Omega * DeltaTime;
 
 		// Handle wrap around of wheel position
-		float IntegerPart = 0;
-		AngularPosition = FMath::Modf(AngularPosition / TWO_PI, &IntegerPart);
+		float IntegerPart = 0.f;
+		AngularPosition = FMath::Modf(AngularPosition / TWO_PI, &IntegerPart) * TWO_PI;
 
 		if (!bInContact)
 		{
