Chaos: Add ability to clamp constraint motor position targets set by plasticity ( linear and angular ) to set limits if necessary

diff --git a/Engine/Source/Runtime/Experimental/Chaos/Private/Chaos/Joint/JointConstraintsCVars.cpp b/Engine/Source/Runtime/Experimental/Chaos/Private/Chaos/Joint/JointConstraintsCVars.cpp
--- a/Engine/Source/Runtime/Experimental/Chaos/Private/Chaos/Joint/JointConstraintsCVars.cpp
+++ b/Engine/Source/Runtime/Experimental/Chaos/Private/Chaos/Joint/JointConstraintsCVars.cpp
@@ -20,6 +20,9 @@ FAutoConsoleVariableRef CVarChaosJointVelProjectionScale(TEXT("p.Chaos.Joint.Vel
 bool bChaos_Joint_DisableSoftLimits = false;
 FAutoConsoleVariableRef CVarChaosJointDisableSoftLimits(TEXT("p.Chaos.Joint.DisableSoftLimits"), bChaos_Joint_DisableSoftLimits, TEXT("Disable soft limits (for debugging only)"));
 
+bool bChaos_Joint_Plasticity_ClampToLimits = true;
+FAutoConsoleVariableRef CVarChaosJointPlasticityClampToLimits(TEXT("p.Chaos.Joint.Plasticity.ClampToLimits"), bChaos_Joint_Plasticity_ClampToLimits, TEXT("Clamp drive position targets to defined limits after plasticity computation"));
+
 float Chaos_Joint_LinearVelocityThresholdToApplyRestitution = 1e-2f;
 FAutoConsoleVariableRef CVarChaosJointLinearVelocityThresholdToApplyRestitution(TEXT("p.Chaos.Joint.LinearVelocityThresholdToApplyRestitution"), Chaos_Joint_LinearVelocityThresholdToApplyRestitution, TEXT("Apply restitution only if initial velocity is higher than this threshold (used in Quasipbd)"));
 
diff --git a/Engine/Source/Runtime/Experimental/Chaos/Private/Chaos/Joint/JointConstraintsCVars.h b/Engine/Source/Runtime/Experimental/Chaos/Private/Chaos/Joint/JointConstraintsCVars.h
--- a/Engine/Source/Runtime/Experimental/Chaos/Private/Chaos/Joint/JointConstraintsCVars.h
+++ b/Engine/Source/Runtime/Experimental/Chaos/Private/Chaos/Joint/JointConstraintsCVars.h
@@ -20,6 +20,8 @@ extern float Chaos_Joint_VelProjectionAlpha;
 
 extern bool bChaos_Joint_DisableSoftLimits;
 
+extern bool bChaos_Joint_Plasticity_ClampToLimits;
+
 extern float Chaos_Joint_LinearVelocityThresholdToApplyRestitution;
 
 extern float Chaos_Joint_AngularVelocityThresholdToApplyRestitution;
