GeometryCollectionComponent
-- added ForceMotionBlur property to make the Nanite scene proxy think that the geometry collection is always moving so that motion blur optimizaztion is forced to evaluate motion.

#deadend

diff --git a/Engine/Source/Runtime/Experimental/GeometryCollectionEngine/Private/GeometryCollection/GeometryCollectionComponent.cpp b/Engine/Source/Runtime/Experimental/GeometryCollectionEngine/Private/GeometryCollection/GeometryCollectionComponent.cpp
--- a/Engine/Source/Runtime/Experimental/GeometryCollectionEngine/Private/GeometryCollection/GeometryCollectionComponent.cpp
+++ b/Engine/Source/Runtime/Experimental/GeometryCollectionEngine/Private/GeometryCollection/GeometryCollectionComponent.cpp
@@ -219,6 +219,7 @@ UGeometryCollectionComponent::UGeometryCollectionComponent(const FObjectInitiali
 	, ChaosSolverActor(nullptr)
 	, InitializationState(ESimulationInitializationState::Unintialized)
 	, ObjectType(EObjectStateTypeEnum::Chaos_Object_Dynamic)
+	, bForceMotionBlur()
 	, EnableClustering(true)
 	, ClusterGroupIndex(0)
 	, MaxClusterLevel(100)
@@ -365,7 +366,7 @@ void UGeometryCollectionComponent::BeginPlay()
 	//////////////////////////////////////////////////////////////////////////
 
 	// default current cache time
-	CurrentCacheTime = MAX_flt;
+	CurrentCacheTime = MAX_flt;	
 }
 
 
@@ -523,6 +524,22 @@ FPrimitiveSceneProxy* UGeometryCollectionComponent::CreateSceneProxy()
 			GGeometryCollectionNanite != 0)
 		{
 			LocalSceneProxy = new FNaniteGeometryCollectionSceneProxy(this);
+
+			// ForceMotionBlur means we maintain bIsMoving, regardless of actual state.
+			if (bForceMotionBlur)
+			{
+				bIsMoving = true;
+				if (LocalSceneProxy)
+				{
+					FNaniteGeometryCollectionSceneProxy* NaniteProxy = static_cast<FNaniteGeometryCollectionSceneProxy*>(LocalSceneProxy);
+					ENQUEUE_RENDER_COMMAND(NaniteProxyOnMotionEnd)(
+						[NaniteProxy](FRHICommandListImmediate& RHICmdList)
+						{
+							NaniteProxy->OnMotionBegin();
+						}
+					);
+				}
+			}
 		}
 		// If we didn't get a proxy, but Nanite was enabled on the asset when it was built, evaluate proxy creation
 		else if (RestCollection->EnableNanite && NaniteProxyRenderMode != 0)
@@ -846,6 +863,7 @@ UPhysicalMaterial* UGeometryCollectionComponent::GetPhysicalMaterial() const
 
 void UGeometryCollectionComponent::RefreshEmbeddedGeometry()
 {
+	
 	const TManagedArray<int32>& ExemplarIndexArray = GetExemplarIndexArray();
 	const int32 TransformCount = GlobalMatrices.Num();
 	
@@ -1681,7 +1699,7 @@ FGeometryCollectionDynamicData* UGeometryCollectionComponent::InitDynamicData(bo
 			DynamicData = nullptr;
 
 			// Change of state?
-			if (bIsMoving)
+			if (bIsMoving && !bForceMotionBlur)
 			{
 				bIsMoving = false;
 				if (SceneProxy && SceneProxy->IsNaniteMesh())
@@ -1699,7 +1717,7 @@ FGeometryCollectionDynamicData* UGeometryCollectionComponent::InitDynamicData(bo
 		else
 		{
 			// Change of state?
-			if (!bIsMoving)
+			if (!bIsMoving && !bForceMotionBlur)
 			{
 				bIsMoving = true;
 				if (SceneProxy && SceneProxy->IsNaniteMesh())
@@ -2860,13 +2878,13 @@ void UGeometryCollectionComponent::CalculateGlobalMatrices()
 		GlobalMatrices.Append(Results->GlobalTransforms);	
 	}
 	else
-	{	
+	{
 		// If hierarchy topology has changed, the RestTransforms is invalidated.
 		if (RestTransforms.Num() != GetTransformArray().Num())
 		{
 			RestTransforms.Empty();
 		}
-		
+
 		if (!DynamicCollection && RestTransforms.Num() > 0)
 		{
 			GeometryCollectionAlgo::GlobalMatrices(RestTransforms, GetParentArray(), GlobalMatrices);
@@ -3090,7 +3108,7 @@ void UGeometryCollectionComponent::InitializeEmbeddedGeometry()
 	if (RestCollection)
 	{
 		ClearEmbeddedGeometry();
-
+		
 		AActor* ActorOwner = GetOwner();
 		check(ActorOwner);
 
@@ -3123,6 +3141,7 @@ void UGeometryCollectionComponent::InitializeEmbeddedGeometry()
 
 		CalculateGlobalMatrices();
 		RefreshEmbeddedGeometry();
+		
 	}
 }
 
diff --git a/Engine/Source/Runtime/Experimental/GeometryCollectionEngine/Public/GeometryCollection/GeometryCollectionComponent.h b/Engine/Source/Runtime/Experimental/GeometryCollectionEngine/Public/GeometryCollection/GeometryCollectionComponent.h
--- a/Engine/Source/Runtime/Experimental/GeometryCollectionEngine/Public/GeometryCollection/GeometryCollectionComponent.h
+++ b/Engine/Source/Runtime/Experimental/GeometryCollectionEngine/Public/GeometryCollection/GeometryCollectionComponent.h
@@ -454,6 +454,9 @@ public:
 	UPROPERTY(EditAnywhere, BlueprintReadWrite, Category = "ChaosPhysics|General")
 	EObjectStateTypeEnum ObjectType;
 
+	/** If ForceMotionBlur is on, motion blur will always be active, even if the GeometryCollection is at rest. */
+	UPROPERTY(EditAnywhere, BlueprintReadWrite, Category = "ChaosPhysics|General")
+	bool bForceMotionBlur;
 
 	UPROPERTY(EditAnywhere, BlueprintReadWrite, Category = "ChaosPhysics|Clustering")
 	bool EnableClustering;
