Revs not hitting limit due to float accuracy issue.

diff --git a/Engine/Source/Runtime/Experimental/ChaosVehicles/ChaosVehiclesCore/Private/EngineSystem.cpp b/Engine/Source/Runtime/Experimental/ChaosVehicles/ChaosVehiclesCore/Private/EngineSystem.cpp
--- a/Engine/Source/Runtime/Experimental/ChaosVehicles/ChaosVehiclesCore/Private/EngineSystem.cpp
+++ b/Engine/Source/Runtime/Experimental/ChaosVehicles/ChaosVehiclesCore/Private/EngineSystem.cpp
@@ -28,7 +28,7 @@ namespace Chaos
 
 	float FSimpleEngineSim::GetTorqueFromRPM(float RPM, bool LimitToIdle /*= true*/)
 	{
-		if (!EngineStarted || RPM >= Setup().MaxRPM || Setup().MaxRPM == 0)
+		if (!EngineStarted || (FMath::Abs(RPM - Setup().MaxRPM) < 1.0f) || Setup().MaxRPM == 0)
 		{
 			return 0.f;
 		}
