Optimized FClothingSimulation::CreateActor inclusive avg 8.5ms down to 5.6ms (Reduce allocations done in FGraphColoring::ComputeGraphColoring).

diff --git a/Engine/Source/Runtime/Experimental/Chaos/Public/Chaos/GraphColoring.h b/Engine/Source/Runtime/Experimental/Chaos/Public/Chaos/GraphColoring.h
--- a/Engine/Source/Runtime/Experimental/Chaos/Public/Chaos/GraphColoring.h
+++ b/Engine/Source/Runtime/Experimental/Chaos/Public/Chaos/GraphColoring.h
@@ -11,7 +11,7 @@ namespace Chaos
 
 class FGraphColoring
 {
-	typedef TSet<int32> FColorSet;
+	typedef TArray<int32, TInlineAllocator<8>> FColorSet;
 	
 	struct FGraphNode
 	{
@@ -20,7 +20,7 @@ class FGraphColoring
 		{
 		}
 
-		TArray<int32> Edges;
+		TArray<int32, TInlineAllocator<8>> Edges;
 		int32 NextColor;
 		FColorSet UsedColors;
 	};
