Removed unnecessary updates for dynamic trees

diff --git a/Engine/Source/Runtime/Experimental/Chaos/Private/Chaos/PBDRigidsEvolution.cpp b/Engine/Source/Runtime/Experimental/Chaos/Private/Chaos/PBDRigidsEvolution.cpp
--- a/Engine/Source/Runtime/Experimental/Chaos/Private/Chaos/PBDRigidsEvolution.cpp
+++ b/Engine/Source/Runtime/Experimental/Chaos/Private/Chaos/PBDRigidsEvolution.cpp
@@ -83,6 +83,12 @@ namespace Chaos
 		using AABBDynamicTreeType = TAABBTree<FAccelerationStructureHandle, TAABBTreeLeafArray<FAccelerationStructureHandle>>;
 		using AABBTreeOfGridsType = TAABBTree<FAccelerationStructureHandle, TBoundingVolume<FAccelerationStructureHandle>>;
 
+
+		static bool IsDynamicTree(FSpatialAccelerationIdx SpatialAccelerationIdx)
+		{
+			return AccelerationStructureUseDynamicTree && SpatialAccelerationIdx.InnerIdx == 1;
+		}
+
 		TUniquePtr<ISpatialAccelerationCollection<FAccelerationStructureHandle, FReal, 3>> CreateEmptyCollection() override
 		{
 			TConstParticleView<FSpatialAccelerationCache> Empty;
@@ -97,7 +103,7 @@ namespace Chaos
 				if (AccelerationStructureSplitStaticAndDynamic == 1)
 				{
 					// Non static bodies
-					Collection->AddSubstructure(CreateAccelerationPerBucket_Threaded(Empty, BucketIdx, true, (bool) AccelerationStructureUseDynamicTree), BucketIdx, 1);
+					Collection->AddSubstructure(CreateAccelerationPerBucket_Threaded(Empty, BucketIdx, true, IsDynamicTree(FSpatialAccelerationIdx{BucketIdx, 1})), BucketIdx, 1);
 				}
 			}
 
@@ -517,11 +523,14 @@ namespace Chaos
 		}
 	}
 
-	void FPBDRigidsEvolutionBase::ApplyParticlePendingData(const FPendingSpatialData& SpatialData, FAccelerationStructure& AccelerationStructure, bool bUpdateCache)
+	void FPBDRigidsEvolutionBase::ApplyParticlePendingData(const FPendingSpatialData& SpatialData, FAccelerationStructure& AccelerationStructure, bool bUpdateCache, bool bUpdateDynamicTrees)
 	{
 		if (SpatialData.bDelete)
 		{
-			AccelerationStructure.RemoveElementFrom(SpatialData.AccelerationHandle, SpatialData.SpatialIdx);
+			if (bUpdateDynamicTrees || !FDefaultCollectionFactory::IsDynamicTree(SpatialData.SpatialIdx))
+			{
+				AccelerationStructure.RemoveElementFrom(SpatialData.AccelerationHandle, SpatialData.SpatialIdx);
+			}
 
 			if (bUpdateCache)
 			{
@@ -544,6 +553,7 @@ namespace Chaos
 		{
 			FGeometryParticleHandle* UpdateParticle = SpatialData.AccelerationHandle.GetGeometryParticleHandle_PhysicsThread();
 
+			if (bUpdateDynamicTrees || !FDefaultCollectionFactory::IsDynamicTree(SpatialData.SpatialIdx))
 			{
 				//CSV_SCOPED_TIMING_STAT(ChaosPhysicsTimers, AABBTreeUpdate)
 				AccelerationStructure.UpdateElementIn(UpdateParticle, UpdateParticle->WorldSpaceInflatedBounds(), UpdateParticle->HasBounds(), SpatialData.SpatialIdx);
@@ -584,7 +594,7 @@ namespace Chaos
 	{
 		for (const FPendingSpatialData& PendingData : InternalAccelerationQueue.PendingData)
 		{
-			ApplyParticlePendingData(PendingData, *InternalAcceleration, false);
+			ApplyParticlePendingData(PendingData, *InternalAcceleration, false, true);
 		}
 		InternalAcceleration->SetSyncTimestamp(LatestExternalTimestampConsumed_Internal);
 		InternalAccelerationQueue.Reset();
@@ -594,8 +604,8 @@ namespace Chaos
 	{
 		for (const FPendingSpatialData& PendingData : AsyncAccelerationQueue.PendingData)
 		{
-			ApplyParticlePendingData(PendingData, *AsyncInternalAcceleration, true); //only the first queue needs to update the cached acceleration
-			ApplyParticlePendingData(PendingData, *AsyncExternalAcceleration, false);
+			ApplyParticlePendingData(PendingData, *AsyncInternalAcceleration, true, false); //only the first queue needs to update the cached acceleration
+			ApplyParticlePendingData(PendingData, *AsyncExternalAcceleration, false, false);
 		}
 
 				//NOTE: This assumes that we are never creating a PT particle that is replicated to GT
@@ -743,6 +753,12 @@ namespace Chaos
 				std::swap(InternalAcceleration, AsyncInternalAcceleration);
 				// The old InternalAcceleration is now in AsyncInternalAcceleration, and the constituent structures will be reused if no changes have been detected
 
+				if (AccelerationStructureUseDynamicTree)
+				{
+					// Dynamic Acceleration structures were not built by the async task, so copy them where required
+					CopyUnBuiltDynamicAccelerationStructures(SpatialAccelerationCache, InternalAcceleration, AsyncInternalAcceleration, AsyncExternalAcceleration);
+				}
+
 				if (AccelerationStructureSplitStaticAndDynamic == 1)
 				{
 					// If the new InternalAcceleration structure have constituents with no changes, we can copy it to the  AsyncInternalAcceleration for reuse
@@ -774,6 +790,50 @@ namespace Chaos
 		}
 	}
 
+	void FPBDRigidsEvolutionBase::CopyUnBuiltDynamicAccelerationStructures(const TMap<FSpatialAccelerationIdx, TUniquePtr<FSpatialAccelerationCache>>& SpatialAccelerationCache, FAccelerationStructure* InternalAcceleration, FAccelerationStructure* AsyncInternalAcceleration, FAccelerationStructure* AsyncExternalAcceleration)
+	{
+		for (const auto& Itr : SpatialAccelerationCache)
+		{
+			const FSpatialAccelerationIdx SpatialIdx = Itr.Key;
+			auto* InternalAccelerationSubStructure = InternalAcceleration->GetSubstructure(SpatialIdx);
+			if (!InternalAccelerationSubStructure)
+			{
+				ensure(false);
+				continue;
+			}
+
+			if (!InternalAccelerationSubStructure->IsTreeDynamic())
+			{
+				continue;
+			}
+
+			auto* AsyncInternalAccelerationSubStructure = AsyncInternalAcceleration->GetSubstructure(SpatialIdx);
+			if (!AsyncInternalAccelerationSubStructure)
+			{
+				ensure(false);
+				continue;
+			}
+
+			check(AsyncInternalAccelerationSubStructure->IsTreeDynamic());
+			// The good up to date data is in AsyncInternalAcceleration
+			if (AsyncExternalAcceleration)
+			{
+				auto* ExternalAccelerationSubStructure = AsyncExternalAcceleration->GetSubstructure(SpatialIdx);
+				if (!ExternalAccelerationSubStructure)
+				{
+					ensure(false);
+					continue;
+				}
+
+				check(ExternalAccelerationSubStructure->IsTreeDynamic());
+				*ExternalAccelerationSubStructure = *AsyncInternalAccelerationSubStructure;
+			}
+
+			// ToDo: We should be able to change this copy to a swap if we remove the copy in the async task
+			*InternalAccelerationSubStructure = *AsyncInternalAccelerationSubStructure;
+			//AsyncInternalAcceleration->SwapSubstructure(*InternalAcceleration, SpatialIdx);
+		}
+	}
 
 	// copy Pristine substructures from FromStructure to ToStructure, if the substructure is not already pristine
 	// Assumption: Both structures represent the same state, one can just be dirtier than the other
diff --git a/Engine/Source/Runtime/Experimental/Chaos/Public/Chaos/ISpatialAccelerationCollection.h b/Engine/Source/Runtime/Experimental/Chaos/Public/Chaos/ISpatialAccelerationCollection.h
--- a/Engine/Source/Runtime/Experimental/Chaos/Public/Chaos/ISpatialAccelerationCollection.h
+++ b/Engine/Source/Runtime/Experimental/Chaos/Public/Chaos/ISpatialAccelerationCollection.h
@@ -27,6 +27,7 @@ public:
 	virtual FSpatialAccelerationIdx AddSubstructure(TUniquePtr<ISpatialAcceleration<TPayloadType, T, d>>&& Substructure, uint16 Bucket, uint16 BucketInnerIdx) = 0;
 	virtual TUniquePtr<ISpatialAcceleration<TPayloadType, T, d>> RemoveSubstructure(FSpatialAccelerationIdx Idx) = 0;
 	virtual ISpatialAcceleration<TPayloadType, T, d>* GetSubstructure(FSpatialAccelerationIdx Idx) = 0;
+	virtual void SwapSubstructure(ISpatialAccelerationCollection<TPayloadType, T, d>& Other, FSpatialAccelerationIdx Idx) = 0;	
 
 	/** This is kind of a hack to avoid virtuals. We simply route calls into templated functions */
 	virtual void PBDComputeConstraintsLowLevel(T Dt, FSpatialAccelerationBroadPhase& BroadPhase, FNarrowPhase& NarrowPhase, IResimCacheBase* ResimCache) const = 0;
diff --git a/Engine/Source/Runtime/Experimental/Chaos/Public/Chaos/PBDRigidsEvolution.h b/Engine/Source/Runtime/Experimental/Chaos/Public/Chaos/PBDRigidsEvolution.h
--- a/Engine/Source/Runtime/Experimental/Chaos/Public/Chaos/PBDRigidsEvolution.h
+++ b/Engine/Source/Runtime/Experimental/Chaos/Public/Chaos/PBDRigidsEvolution.h
@@ -909,6 +909,7 @@ protected:
 	void FlushInternalAccelerationQueue();
 	void FlushAsyncAccelerationQueue();
 	void WaitOnAccelerationStructure();
+	static void CopyUnBuiltDynamicAccelerationStructures(const TMap<FSpatialAccelerationIdx, TUniquePtr<FSpatialAccelerationCache>>& SpatialAccelerationCache, FAccelerationStructure* InternalAcceleration, FAccelerationStructure* AsyncInternalAcceleration, FAccelerationStructure* AsyncExternalAcceleration);
 	static void CopyPristineAccelerationStructures(const TMap<FSpatialAccelerationIdx, TUniquePtr<FSpatialAccelerationCache>>& SpatialAccelerationCache, FAccelerationStructure* FromStructure, FAccelerationStructure* ToStructure, bool CheckPristine);
 
 	TArray<FForceRule> ForceRules;
@@ -982,7 +983,7 @@ protected:
 
 	TMap<FSpatialAccelerationIdx, TUniquePtr<FSpatialAccelerationCache>> SpatialAccelerationCache;
 
-	FORCEINLINE_DEBUGGABLE void ApplyParticlePendingData(const FPendingSpatialData& PendingData, FAccelerationStructure& SpatialAcceleration, bool bUpdateCache);
+	FORCEINLINE_DEBUGGABLE void ApplyParticlePendingData(const FPendingSpatialData& PendingData, FAccelerationStructure& SpatialAcceleration, bool bUpdateCache, bool bUpdateDynamicTrees);
 
 	class FChaosAccelerationStructureTask
 	{
diff --git a/Engine/Source/Runtime/Experimental/Chaos/Public/Chaos/SpatialAccelerationCollection.h b/Engine/Source/Runtime/Experimental/Chaos/Public/Chaos/SpatialAccelerationCollection.h
--- a/Engine/Source/Runtime/Experimental/Chaos/Public/Chaos/SpatialAccelerationCollection.h
+++ b/Engine/Source/Runtime/Experimental/Chaos/Public/Chaos/SpatialAccelerationCollection.h
@@ -432,6 +432,19 @@ public:
 
 		return nullptr;
 	}
+
+	virtual void SwapSubstructure(ISpatialAccelerationCollection<TPayloadType, T, d>& InOther, FSpatialAccelerationIdx Idx) override
+	{
+		using ThisType = decltype(this);
+		ThisType Other = static_cast<ThisType>(&InOther);
+
+		check(Idx.Bucket < MaxBuckets);
+		check(Idx.Bucket < Other->MaxBuckets);
+		check(Buckets[Idx.Bucket].Objects.Num() > Idx.InnerIdx);
+		check(Other->Buckets[Idx.Bucket].Objects.Num() > Idx.InnerIdx);
+
+		std::swap(Buckets[Idx.Bucket].Objects[Idx.InnerIdx].Acceleration, Other->Buckets[Idx.Bucket].Objects[Idx.InnerIdx].Acceleration);
+	}
 	
 	virtual void Reset() override
 	{
