Make an override to run chaos only as single threaded
Added console flags for physics playthrough

diff --git a/Engine/Source/Runtime/Experimental/Chaos/Private/Chaos/Framework/Parallel.cpp b/Engine/Source/Runtime/Experimental/Chaos/Private/Chaos/Framework/Parallel.cpp
--- a/Engine/Source/Runtime/Experimental/Chaos/Private/Chaos/Framework/Parallel.cpp
+++ b/Engine/Source/Runtime/Experimental/Chaos/Private/Chaos/Framework/Parallel.cpp
@@ -8,6 +8,8 @@ using namespace Chaos;
 
 namespace Chaos
 {
+	CHAOS_API int32 GSingleThreadedPhysics = 0;
+	
 #if !UE_BUILD_SHIPPING
 	CHAOS_API bool bDisablePhysicsParallelFor = false;
 	CHAOS_API bool bDisableParticleParallelFor = false;
@@ -41,7 +43,7 @@ void Chaos::PhysicsParallelFor(int32 InNum, TFunctionRef<void(int32)> InCallable
 #endif
 		InCallable(Idx);
 	};
-	::ParallelFor(InNum, PassThrough, bDisablePhysicsParallelFor || bForceSingleThreaded);
+	::ParallelFor(InNum, PassThrough, !!GSingleThreadedPhysics || bDisablePhysicsParallelFor || bForceSingleThreaded);
 }
 
 //class FRecursiveDivideTask
diff --git a/Engine/Source/Runtime/Experimental/Chaos/Private/Chaos/Framework/PhysicsSolverBase.cpp b/Engine/Source/Runtime/Experimental/Chaos/Private/Chaos/Framework/PhysicsSolverBase.cpp
--- a/Engine/Source/Runtime/Experimental/Chaos/Private/Chaos/Framework/PhysicsSolverBase.cpp
+++ b/Engine/Source/Runtime/Experimental/Chaos/Private/Chaos/Framework/PhysicsSolverBase.cpp
@@ -13,6 +13,7 @@
 
 namespace Chaos
 {	
+	extern int GSingleThreadedPhysics;
 	void FPhysicsSolverBase::ChangeBufferMode(EMultiBufferMode InBufferMode)
 	{
 		BufferMode = InBufferMode;
@@ -139,7 +140,7 @@ namespace Chaos
 
 	FPhysicsSolverBase::FPhysicsSolverBase(const EMultiBufferMode BufferingModeIn,const EThreadingModeTemp InThreadingMode,UObject* InOwner, Chaos::FReal InAsyncDt)
 		: BufferMode(BufferingModeIn)
-		, ThreadingMode(InThreadingMode)
+		, ThreadingMode(!!GSingleThreadedPhysics ? EThreadingModeTemp::SingleThread : InThreadingMode)
 		, PullResultsManager(MakeUnique<FChaosResultsManager>(MarshallingManager))
 		, PendingSpatialOperations_External(MakeUnique<FPendingSpatialDataQueue>())
 		, bUseCollisionResimCache(false)
diff --git a/Engine/Source/Runtime/Experimental/Chaos/Public/Chaos/Framework/PhysicsSolverBase.h b/Engine/Source/Runtime/Experimental/Chaos/Public/Chaos/Framework/PhysicsSolverBase.h
--- a/Engine/Source/Runtime/Experimental/Chaos/Public/Chaos/Framework/PhysicsSolverBase.h
+++ b/Engine/Source/Runtime/Experimental/Chaos/Public/Chaos/Framework/PhysicsSolverBase.h
@@ -28,6 +28,8 @@ namespace Chaos
 	class FRewindData;
 	class IRewindCallback;
 
+	extern CHAOS_API int32 GSingleThreadedPhysics;
+
 	extern CHAOS_API int32 UseAsyncInterpolation;
 	extern CHAOS_API int32 ForceDisableAsyncPhysics;
 	extern CHAOS_API FRealSingle AsyncInterpolationMultiplier;
@@ -252,6 +254,8 @@ namespace Chaos
 
 		void SetThreadingMode_External(EThreadingModeTemp InThreadingMode)
 		{
+			if (!!GSingleThreadedPhysics) { InThreadingMode = EThreadingModeTemp::SingleThread; }
+
 			if(InThreadingMode != ThreadingMode)
 			{
 				if(InThreadingMode == EThreadingModeTemp::SingleThread)
diff --git a/Engine/Source/Runtime/Experimental/Chaos/Public/Chaos/ParticleIterator.h b/Engine/Source/Runtime/Experimental/Chaos/Public/Chaos/ParticleIterator.h
--- a/Engine/Source/Runtime/Experimental/Chaos/Public/Chaos/ParticleIterator.h
+++ b/Engine/Source/Runtime/Experimental/Chaos/Public/Chaos/ParticleIterator.h
@@ -65,7 +65,7 @@ void ParticlesParallelForImp(const TArray<TParticle>& Particles, const Lambda& F
 	// When ParticlesParallelFor is called with a plain old TArray,
 	// just do normal parallelization.
 	const int32 Num = Particles.Num();
-	::ParallelFor(Num, [&Func, &Particles](const int32 Index)
+	PhysicsParallelFor(Num, [&Func, &Particles](const int32 Index)
 	{
 		Func(Particles[Index], Index);
 	});
@@ -697,7 +697,7 @@ void ParticleViewParallelForImp(const TParticleView& Particles, const Lambda& Us
 		{
 			// Do a regular parallel for over the handles in this SOA view
 			const int32 HandleCount = CurHandlesArray->Num();
-			::ParallelFor(HandleCount, [&Func, CurHandlesArray, ParticleIdxOff](const int32 HandleIdx)
+			PhysicsParallelFor(HandleCount, [&Func, CurHandlesArray, ParticleIdxOff](const int32 HandleIdx)
 			{
 				// Reconstruct the TransientHandle so that it has a chance to update
 				// other data members, like the particle type in the case of geometry 
@@ -711,7 +711,7 @@ void ParticleViewParallelForImp(const TParticleView& Particles, const Lambda& Us
 		else
 		{
 			// Do a regular parallel for over the particles in this SOA view
-			::ParallelFor(ParticleCount, [&Func, &SOAView, ParticleIdxOff](const int32 ParticleIdx)
+			PhysicsParallelFor(ParticleCount, [&Func, &SOAView, ParticleIdxOff](const int32 ParticleIdx)
 			{
 				// Reconstruct the TransientHandle so that it has a chance to update
 				// other data members, like the particle type in the case of geometry 
@@ -731,7 +731,7 @@ void HandleViewParallelForImp(const THandleView& HandleView, const Lambda& Func)
 
 	using THandle = typename THandleView::THandle;
 	const int32 HandleCount = HandleView.Handles.Num();
-	::ParallelFor(HandleCount, [&HandleView, &Func](const int32 Index)
+	PhysicsParallelFor(HandleCount, [&HandleView, &Func](const int32 Index)
 	{
 		if (!HandleView.Handles[Index]->LightWeightDisabled())
 		{
