Fix slip angle reported at very low speeds

diff --git a/Engine/Source/Runtime/Experimental/ChaosVehicles/ChaosVehiclesCore/Public/VehicleUtility.h b/Engine/Source/Runtime/Experimental/ChaosVehicles/ChaosVehiclesCore/Public/VehicleUtility.h
--- a/Engine/Source/Runtime/Experimental/ChaosVehicles/ChaosVehiclesCore/Public/VehicleUtility.h
+++ b/Engine/Source/Runtime/Experimental/ChaosVehicles/ChaosVehiclesCore/Public/VehicleUtility.h
@@ -130,10 +130,16 @@ namespace Chaos
 
 		static float CalculateSlipAngle(float Y, float X)
 		{
-			float Value = FMath::Abs(FMath::Atan2(Y, X));
-			if (Value > HALF_PI)
+			float Value = 0.0f;
+
+			float LateralSpeedThreshold = 0.05f;
+			if (FMath::Abs(Y) > LateralSpeedThreshold)
 			{
-				Value = PI - Value;
+				Value = FMath::Abs(FMath::Atan2(Y, X));
+				if (Value > HALF_PI)
+				{
+					Value = PI - Value;
+				}
 			}
 
 			return Value;
