Add mass scaling back for quasipbd

diff --git a/Engine/Source/Runtime/Experimental/Chaos/Private/Chaos/Joint/PBDJointSolverGaussSeidel.cpp b/Engine/Source/Runtime/Experimental/Chaos/Private/Chaos/Joint/PBDJointSolverGaussSeidel.cpp
--- a/Engine/Source/Runtime/Experimental/Chaos/Private/Chaos/Joint/PBDJointSolverGaussSeidel.cpp
+++ b/Engine/Source/Runtime/Experimental/Chaos/Private/Chaos/Joint/PBDJointSolverGaussSeidel.cpp
@@ -123,8 +123,10 @@ namespace Chaos
 		SolverBodies[0].SetInvMScale(JointSettings.ParentInvMassScale);
 		SolverBodies[1].SetInvMScale(FReal(1));
 
-		// \chaos(todo): Add support for mass conditioning in ConstraintSolverBody
-		//FPBDJointUtilities::ConditionInverseMassAndInertia(InvM(0), InvM(1), InvILs[0], InvILs[1], SolverSettings.MinParentMassRatio, SolverSettings.MaxInertiaRatio);
+		FVec3 ConditionedInvILs[2];
+		FPBDJointUtilities::ConditionInverseMassAndInertia(Body0().InvM(), Body1().InvM(), Body0().InvILocal(), Body1().InvILocal(), SolverSettings.MinParentMassRatio, SolverSettings.MaxInertiaRatio, ConditionedInvMs[0], ConditionedInvMs[1], ConditionedInvILs[0], ConditionedInvILs[1]);
+		ConditionedInvIs[0] = Utilities::ComputeWorldSpaceInertia(Q(0), ConditionedInvILs[0]);
+		ConditionedInvIs[1] = Utilities::ComputeWorldSpaceInertia(Q(1), ConditionedInvILs[1]);
 
 		NetLinearImpulse = FVec3(0);
 		NetAngularImpulse = FVec3(0);
diff --git a/Engine/Source/Runtime/Experimental/Chaos/Private/Chaos/PBDJointConstraintUtilities.cpp b/Engine/Source/Runtime/Experimental/Chaos/Private/Chaos/PBDJointConstraintUtilities.cpp
--- a/Engine/Source/Runtime/Experimental/Chaos/Private/Chaos/PBDJointConstraintUtilities.cpp
+++ b/Engine/Source/Runtime/Experimental/Chaos/Private/Chaos/PBDJointConstraintUtilities.cpp
@@ -659,12 +659,16 @@ namespace Chaos
 	
 	// @todo(ccaulfield): should also take into account the length of the joint connector to prevent over-rotation
 	void FPBDJointUtilities::ConditionInverseMassAndInertia(
-		FReal& InOutInvMParent,
-		FReal& InOutInvMChild,
-		FVec3& InOutInvIParent,
-		FVec3& InOutInvIChild,
+		const FReal& InInvMParent,
+		const FReal& InInvMChild,
+		const FVec3& InInvIParent,
+		const FVec3& InInvIChild,
 		const FReal MinParentMassRatio,
-		const FReal MaxInertiaRatio)
+		const FReal MaxInertiaRatio,
+		FReal& OutInvMParent,
+		FReal& OutInvMChild,
+		FVec3& OutInvIParent,
+		FVec3& OutInvIChild)
 	{
 		FReal MParent = 0.0f;
 		FVec3 IParent = FVec3(0);
@@ -672,37 +676,48 @@ namespace Chaos
 		FVec3 IChild = FVec3(0);
 
 		// Set up inertia so that it is more uniform (reduce the maximum ratio of the inertia about each axis)
-		if (InOutInvMParent > 0)
+		if (InInvMParent > 0)
 		{
-			MParent = 1.0f / InOutInvMParent;
-			IParent = ConditionInertia(FVec3(1.0f / InOutInvIParent.X, 1.0f / InOutInvIParent.Y, 1.0f / InOutInvIParent.Z), MaxInertiaRatio);
+			MParent = 1.0f / InInvMParent;
+			IParent = ConditionInertia(FVec3(1.0f / InInvIParent.X, 1.0f / InInvIParent.Y, 1.0f / InInvIParent.Z), MaxInertiaRatio);
 		}
-		if (InOutInvMChild > 0)
+		if (InInvMChild > 0)
 		{
-			MChild = 1.0f / InOutInvMChild;
-			IChild = ConditionInertia(FVec3(1.0f / InOutInvIChild.X, 1.0f / InOutInvIChild.Y, 1.0f / InOutInvIChild.Z), MaxInertiaRatio);
+			MChild = 1.0f / InInvMChild;
+			IChild = ConditionInertia(FVec3(1.0f / InInvIChild.X, 1.0f / InInvIChild.Y, 1.0f / InInvIChild.Z), MaxInertiaRatio);
 		}
 
 		// Set up relative mass and inertia so that the parent cannot be much lighter than the child
-		if ((InOutInvMParent > 0) && (InOutInvMChild > 0))
+		if ((InInvMParent > 0) && (InInvMChild > 0))
 		{
 			MParent = ConditionParentMass(MParent, MChild, MinParentMassRatio);
 			IParent = ConditionParentInertia(IParent, IChild, MinParentMassRatio);
 		}
 
 		// Map back to inverses
-		if (InOutInvMParent > 0)
+		if (InInvMParent > 0)
 		{
-			InOutInvMParent = (FReal)1 / MParent;
-			InOutInvIParent = FVec3((FReal)1 / IParent.X, (FReal)1 / IParent.Y, (FReal)1 / IParent.Z);
+			OutInvMParent = (FReal)1 / MParent;
+			OutInvIParent = FVec3((FReal)1 / IParent.X, (FReal)1 / IParent.Y, (FReal)1 / IParent.Z);
 		}
-		if (InOutInvMChild > 0)
+		if (InInvMChild > 0)
 		{
-			InOutInvMChild = (FReal)1 / MChild;
-			InOutInvIChild = FVec3((FReal)1 / IChild.X, (FReal)1 / IChild.Y, (FReal)1 / IChild.Z);
+			OutInvMChild = (FReal)1 / MChild;
+			OutInvIChild = FVec3((FReal)1 / IChild.X, (FReal)1 / IChild.Y, (FReal)1 / IChild.Z);
 		}
 	}
 
+	void FPBDJointUtilities::ConditionInverseMassAndInertia(
+		FReal& InOutInvMParent,
+		FReal& InOutInvMChild,
+		FVec3& InOutInvIParent,
+		FVec3& InOutInvIChild,
+		const FReal MinParentMassRatio,
+		const FReal MaxInertiaRatio)
+	{
+		ConditionInverseMassAndInertia(InOutInvMParent, InOutInvMChild, InOutInvIParent, InOutInvIChild, MinParentMassRatio, MaxInertiaRatio, InOutInvMParent, InOutInvMChild, InOutInvIParent, InOutInvIChild);
+	}
+
 
 	FVec3 FPBDJointUtilities::GetSphereLimitedPositionError(const FVec3& CX, const FReal Radius)
 	{
diff --git a/Engine/Source/Runtime/Experimental/Chaos/Public/Chaos/Evolution/SolverBody.h b/Engine/Source/Runtime/Experimental/Chaos/Public/Chaos/Evolution/SolverBody.h
--- a/Engine/Source/Runtime/Experimental/Chaos/Public/Chaos/Evolution/SolverBody.h
+++ b/Engine/Source/Runtime/Experimental/Chaos/Public/Chaos/Evolution/SolverBody.h
@@ -500,6 +500,11 @@ namespace Chaos
 		*/
 		FMatrix33 InvI() const { return State.InvMassScale * Body->InvI(); }
 
+		/**
+		 * @brief The scaled local space inverse inertia
+		*/
+		FVec3 InvILocal() const { return State.InvMassScale * Body->InvILocal(); }
+
 		/**
 		 * @brief Whether the body is dynamic (i.e., has a finite mass) after InvMassScale is applied
 		*/
diff --git a/Engine/Source/Runtime/Experimental/Chaos/Public/Chaos/Joint/JointSolverConstraints.inl b/Engine/Source/Runtime/Experimental/Chaos/Public/Chaos/Joint/JointSolverConstraints.inl
--- a/Engine/Source/Runtime/Experimental/Chaos/Public/Chaos/Joint/JointSolverConstraints.inl
+++ b/Engine/Source/Runtime/Experimental/Chaos/Public/Chaos/Joint/JointSolverConstraints.inl
@@ -24,7 +24,7 @@ namespace Chaos
 		InvMs[0] = JointSettings.ParentInvMassScale * InvM0;
 		InvMs[1] = InvM1;
 
-		FPBDJointUtilities::ConditionInverseMassAndInertia(InvMs[0], InvMs[1], InvILs[0], InvILs[1], SolverSettings.MinParentMassRatio, SolverSettings.MaxInertiaRatio);
+		FPBDJointUtilities::ConditionInverseMassAndInertia(InvMs[0], InvMs[1], InvILs[0], InvILs[1], SolverSettings.MinParentMassRatio, SolverSettings.MaxInertiaRatio, InvMs[0], InvMs[1], InvILs[0], InvILs[1]);
 
 		PrevPs[0] = PrevP0;
 		PrevPs[1] = PrevP1;
diff --git a/Engine/Source/Runtime/Experimental/Chaos/Public/Chaos/Joint/PBDJointSolverGaussSeidel.h b/Engine/Source/Runtime/Experimental/Chaos/Public/Chaos/Joint/PBDJointSolverGaussSeidel.h
--- a/Engine/Source/Runtime/Experimental/Chaos/Public/Chaos/Joint/PBDJointSolverGaussSeidel.h
+++ b/Engine/Source/Runtime/Experimental/Chaos/Public/Chaos/Joint/PBDJointSolverGaussSeidel.h
@@ -104,12 +104,12 @@ namespace Chaos
 
 		inline FReal InvM(int32 BodyIndex) const
 		{
-			return Body(BodyIndex).InvM();
+			return ConditionedInvMs[BodyIndex];
 		}
 
 		inline FMatrix33 InvI(int32 BodyIndex) const
 		{
-			return Body(BodyIndex).InvI();
+			return ConditionedInvIs[BodyIndex];
 		}
 
 		inline const FVec3& GetNetLinearImpulse() const
@@ -705,6 +705,10 @@ namespace Chaos
 		FVec3 InitConnectorXs[MaxConstrainedBodies];		// World-space joint connector positions
 		FRotation3 InitConnectorRs[MaxConstrainedBodies];	// World-space joint connector rotations
 
+		// Conditioned InvM and InvI
+		FReal ConditionedInvMs[MaxConstrainedBodies];
+		FMatrix33 ConditionedInvIs[MaxConstrainedBodies];
+
 		// Accumulated Impulse and AngularImpulse (Impulse * Dt since they are mass multiplied position corrections)
 		FVec3 NetLinearImpulse;
 		FVec3 NetAngularImpulse;
diff --git a/Engine/Source/Runtime/Experimental/Chaos/Public/Chaos/PBDJointConstraintUtilities.h b/Engine/Source/Runtime/Experimental/Chaos/Public/Chaos/PBDJointConstraintUtilities.h
--- a/Engine/Source/Runtime/Experimental/Chaos/Public/Chaos/PBDJointConstraintUtilities.h
+++ b/Engine/Source/Runtime/Experimental/Chaos/Public/Chaos/PBDJointConstraintUtilities.h
@@ -144,6 +144,18 @@ namespace Chaos
 			const FReal MChild, 
 			const FReal MinRatio);
 
+		static CHAOS_API void ConditionInverseMassAndInertia(
+			const FReal& InInvMParent,
+			const FReal& InInvMChild,
+			const FVec3& InInvIParent,
+			const FVec3& InInvIChild,
+			const FReal MinParentMassRatio,
+			const FReal MaxInertiaRatio,
+			FReal& OutInvMParent,
+			FReal& OutInvMChild,
+			FVec3& OutInvIParent,
+			FVec3& OutInvIChild);
+
 		static CHAOS_API void ConditionInverseMassAndInertia(
 			FReal& InOutInvMParent,
 			FReal& InOutInvMChild,
