Chaos - enable OBB-AABB tests for convex-sphere/capsule
- tight fitting capsule bounds

diff --git a/Engine/Source/Runtime/Experimental/Chaos/Private/Chaos/Collision/ParticlePairMidPhase.cpp b/Engine/Source/Runtime/Experimental/Chaos/Private/Chaos/Collision/ParticlePairMidPhase.cpp
--- a/Engine/Source/Runtime/Experimental/Chaos/Private/Chaos/Collision/ParticlePairMidPhase.cpp
+++ b/Engine/Source/Runtime/Experimental/Chaos/Private/Chaos/Collision/ParticlePairMidPhase.cpp
@@ -117,8 +117,8 @@ namespace Chaos
 		const bool bIsSphere0 = (ImplicitType0 == ImplicitObjectType::Sphere);
 		const bool bIsSphere1 = (ImplicitType1 == ImplicitObjectType::Sphere);
 
-		bEnableOBBCheck0 = Chaos_Collision_NarrowPhase_AABBBoundsCheck && !bIsSphere0 && !bIsSphere1;
-		bEnableOBBCheck1 = Chaos_Collision_NarrowPhase_AABBBoundsCheck && !bIsSphere0 && !bIsSphere1;
+		bEnableOBBCheck0 = Chaos_Collision_NarrowPhase_AABBBoundsCheck && !bIsSphere0;
+		bEnableOBBCheck1 = Chaos_Collision_NarrowPhase_AABBBoundsCheck && !bIsSphere1;
 		bEnableManifoldCheck = bChaos_Collision_EnableManifoldUpdate && !bIsSphere0 && !bIsSphere1;
 	}
 
diff --git a/Engine/Source/Runtime/Experimental/Chaos/Private/Chaos/CollisionOneShotManifoldsMiscShapes.cpp b/Engine/Source/Runtime/Experimental/Chaos/Private/Chaos/CollisionOneShotManifoldsMiscShapes.cpp
--- a/Engine/Source/Runtime/Experimental/Chaos/Private/Chaos/CollisionOneShotManifoldsMiscShapes.cpp
+++ b/Engine/Source/Runtime/Experimental/Chaos/Private/Chaos/CollisionOneShotManifoldsMiscShapes.cpp
@@ -164,8 +164,11 @@ namespace Chaos
 
 			if (FMath::Abs(ADotB) < AxisDotMinimum || AHalfLen < KINDA_SMALL_NUMBER || BHalfLen < KINDA_SMALL_NUMBER)
 			{
-				FContactPoint ContactPoint = CapsuleCapsuleContactPoint(CapsuleA, CapsuleATransform, CapsuleB, CapsuleBTransform, Constraint.GetCullDistance());
-				Constraint.AddOneshotManifoldContact(ContactPoint);
+				FContactPoint ContactPoint = CapsuleCapsuleContactPoint(CapsuleA, CapsuleATransform, CapsuleB, CapsuleBTransform, FReal(0));
+				if (ContactPoint.Phi < Constraint.GetCullDistance())
+				{
+					Constraint.AddOneshotManifoldContact(ContactPoint);
+				}
 				return;
 			}
 			
@@ -185,8 +188,11 @@ namespace Chaos
 
 			if (DeltaLen < KINDA_SMALL_NUMBER)
 			{
-				FContactPoint ContactPoint = CapsuleCapsuleContactPoint(CapsuleA, CapsuleATransform, CapsuleB, CapsuleBTransform, Constraint.GetCullDistance());
-				Constraint.AddOneshotManifoldContact(ContactPoint);
+				FContactPoint ContactPoint = CapsuleCapsuleContactPoint(CapsuleA, CapsuleATransform, CapsuleB, CapsuleBTransform, FReal(0));
+				if (ContactPoint.Phi < Constraint.GetCullDistance())
+				{
+					Constraint.AddOneshotManifoldContact(ContactPoint);
+				}
 				return;
 			}
 			
@@ -205,8 +211,11 @@ namespace Chaos
 			const FReal Clipped2Coord = FMath::Min(ProjA2OntoB, BHalfLen);
 			if (Clipped1Coord > Clipped2Coord) // No overlap
 			{
-				FContactPoint ContactPoint = CapsuleCapsuleContactPoint(CapsuleA, CapsuleATransform, CapsuleB, CapsuleBTransform, Constraint.GetCullDistance());
-				Constraint.AddOneshotManifoldContact(ContactPoint);
+				FContactPoint ContactPoint = CapsuleCapsuleContactPoint(CapsuleA, CapsuleATransform, CapsuleB, CapsuleBTransform, FReal(0));
+				if (ContactPoint.Phi < Constraint.GetCullDistance())
+				{
+					Constraint.AddOneshotManifoldContact(ContactPoint);
+				}
 				return;
 			}
 
@@ -225,12 +234,15 @@ namespace Chaos
 				// Note location A is calculated by rotation (effectively) instead of the usual plane clipping
 				FVec3 LocationA = (ClippedCoord - ProjCentreAOntoB) * CapsuleADirection + ACenter - Normal * (CapsuleA.GetRadius());
 
-				ContactPoint.ShapeContactPoints[0] = CapsuleATransform.InverseTransformPosition(LocationA);
-				ContactPoint.ShapeContactPoints[1] = CapsuleBTransform.InverseTransformPosition(LocationB);
-				ContactPoint.Location = 0.5f * (LocationA + LocationB);
-				ContactPoint.Phi = FVec3::DotProduct(LocationA - LocationB, Normal);
-
-				Constraint.AddOneshotManifoldContact(ContactPoint);
+				const FReal Phi = FVec3::DotProduct(LocationA - LocationB, Normal);
+				if (Phi < Constraint.GetCullDistance())
+				{
+					ContactPoint.ShapeContactPoints[0] = CapsuleATransform.InverseTransformPosition(LocationA);
+					ContactPoint.ShapeContactPoints[1] = CapsuleBTransform.InverseTransformPosition(LocationB);
+					ContactPoint.Location = 0.5f * (LocationA + LocationB);
+					ContactPoint.Phi = Phi;
+					Constraint.AddOneshotManifoldContact(ContactPoint);
+				}
 			};
 
 			AddManifoldPoint(Clipped1Coord);
diff --git a/Engine/Source/Runtime/Experimental/Chaos/Public/Chaos/Capsule.h b/Engine/Source/Runtime/Experimental/Chaos/Public/Chaos/Capsule.h
--- a/Engine/Source/Runtime/Experimental/Chaos/Public/Chaos/Capsule.h
+++ b/Engine/Source/Runtime/Experimental/Chaos/Public/Chaos/Capsule.h
@@ -125,6 +125,17 @@ namespace Chaos
 			return Box;
 		}
 
+		virtual FAABB3 CalculateTransformedBounds(const FRigidTransform3& Transform) const override
+		{
+			const FVec3 X1 = Transform.TransformPositionNoScale(MSegment.GetX1());
+			const FVec3 X2 = Transform.TransformPositionNoScale(MSegment.GetX2());
+			const FVec3 MinSegment = X1.ComponentwiseMin(X2);
+			const FVec3 MaxSegment = X1.ComponentwiseMax(X2);
+
+			const FVec3 RadiusV = FVec3(GetRadius());
+			return FAABB3(MinSegment - RadiusV, MaxSegment + RadiusV);
+		}
+
 		static bool RaycastFast(FReal MRadius, FReal MHeight, const FVec3& MVector, const FVec3& X1, const FVec3& X2, const FVec3& StartPoint, const FVec3& Dir, const FReal Length, const FReal Thickness, FReal& OutTime, FVec3& OutPosition, FVec3& OutNormal, int32& OutFaceIndex)
 		{
 			ensure(FMath::IsNearlyEqual(MVector.SizeSquared(), (FReal)1, (FReal)KINDA_SMALL_NUMBER));
