Chaos: Fix capsule going through scaled trimesh
- make sure we unwrap the scale implicit if the scale is 1 so that the original shape can be properly considered by the MakeTriangle

diff --git a/Engine/Source/Runtime/Experimental/Chaos/Private/Chaos/TriangleMeshImplicitObject.cpp b/Engine/Source/Runtime/Experimental/Chaos/Private/Chaos/TriangleMeshImplicitObject.cpp
--- a/Engine/Source/Runtime/Experimental/Chaos/Private/Chaos/TriangleMeshImplicitObject.cpp
+++ b/Engine/Source/Runtime/Experimental/Chaos/Private/Chaos/TriangleMeshImplicitObject.cpp
@@ -21,8 +21,8 @@ namespace Chaos
 	FAutoConsoleVariableRef CVarPerPolySupport(TEXT("p.Chaos.TriMeshPerPolySupport"), TriMeshPerPolySupport, TEXT("Disabling removes memory cost of vertex map on triangle mesh. Note: Changing at runtime will not work."));
 
 
-template <typename QueryGeomType, typename T, int d>
-static auto MakeScaledHelper(const QueryGeomType& B, const TVector<T, d>& InvScale)
+template <typename QueryGeomType>
+static auto MakeScaledHelper(const QueryGeomType& B, const FVec3& InvScale)
 {
 	TUniquePtr<QueryGeomType> HackBPtr(const_cast<QueryGeomType*>(&B));	//todo: hack, need scaled object to accept raw ptr similar to transformed implicit
 	TImplicitObjectScaled<QueryGeomType> ScaledB(MakeSerializable(HackBPtr), InvScale);
@@ -30,8 +30,8 @@ static auto MakeScaledHelper(const QueryGeomType& B, const TVector<T, d>& InvSca
 	return ScaledB;
 }
 
-template <typename QueryGeomType, typename T, int d>
-static auto MakeScaledHelper(const TImplicitObjectScaled<QueryGeomType>& B, const TVector<T, d>& InvScale)
+template <typename QueryGeomType>
+static auto MakeScaledHelper(const TImplicitObjectScaled<QueryGeomType>& B, const FVec3& InvScale)
 {
 	//if scaled of scaled just collapse into one scaled
 	TImplicitObjectScaled<QueryGeomType> ScaledB(B.Object(), InvScale * B.GetScale());
@@ -78,6 +78,23 @@ void TransformSweepOutputsHelper(FVec3 TriMeshScale, const FVec3& HitNormal, con
 	}
 }
 
+template <typename QueryGeomType>
+const auto MakeTriangleHelper(const QueryGeomType& Geom)
+{
+	return FPBDCollisionConstraint::MakeTriangle(&Geom);
+}
+
+// collapse scale object into it's inner shape if scale is 1, because MakeTRionagle need to be able to infer properties on the shape
+template <typename QueryGeomType>
+const auto MakeTriangleHelper(const TImplicitObjectScaled<QueryGeomType>& ScaledGeom)
+{
+	if (FVec3::IsNearlyEqual(ScaledGeom.GetScale(), FVec3(1), SMALL_NUMBER))
+	{
+		return FPBDCollisionConstraint::MakeTriangle(ScaledGeom.GetUnscaledObject());
+	}
+	return FPBDCollisionConstraint::MakeTriangle(&ScaledGeom);
+}
+
 template <typename IdxType>
 struct FTriangleMeshRaycastVisitor
 {
@@ -391,8 +408,8 @@ bool FTriangleMeshImplicitObject::ContactManifoldImp(const GeomType& QueryGeom,
 		QueryBounds.ThickenSymmetrically(LocalThickness);
 		const TArray<int32> PotentialIntersections = BVH.FindAllIntersections(QueryBounds);
 
-		// @todo(chaos): we should not be creating constraints just for collecting contacts...
-		FPBDCollisionConstraint Constraint = FPBDCollisionConstraint::MakeTriangle(&QueryGeom);
+		// MakeTriangleHelper get rid of the scale wrapper if necessary as MakeTriangle will try to infer properties from it 
+		FPBDCollisionConstraint Constraint = MakeTriangleHelper(WorldScaleGeom);
 
 		for (int32 TriIdx : PotentialIntersections)
 		{
