Chaos - add option to enable/disable determinism
- don't sort constraints in non-deterministic mode

diff --git a/Engine/Source/Programs/HeadlessChaos/Private/HeadlessChaosTestRewind.cpp b/Engine/Source/Programs/HeadlessChaos/Private/HeadlessChaosTestRewind.cpp
--- a/Engine/Source/Programs/HeadlessChaos/Private/HeadlessChaosTestRewind.cpp
+++ b/Engine/Source/Programs/HeadlessChaos/Private/HeadlessChaosTestRewind.cpp
@@ -3938,6 +3938,7 @@ namespace ChaosTest {
 		// Make a solver
 		auto* Solver = Module->CreateSolver(nullptr, /*AsyncDt=*/-1);
 		InitSolverSettings(Solver);
+		Solver->SetIsDeterministic(true);
 
 		TArray<FPhysicsActorHandle> Storage = InitFunc(Solver, Mapping);
 
diff --git a/Engine/Source/Runtime/Experimental/Chaos/Private/Chaos/PBDCollisionConstraints.cpp b/Engine/Source/Runtime/Experimental/Chaos/Private/Chaos/PBDCollisionConstraints.cpp
--- a/Engine/Source/Runtime/Experimental/Chaos/Private/Chaos/PBDCollisionConstraints.cpp
+++ b/Engine/Source/Runtime/Experimental/Chaos/Private/Chaos/PBDCollisionConstraints.cpp
@@ -119,6 +119,7 @@ namespace Chaos
 		, bEnableCollisions(true)
 		, bEnableRestitution(true)
 		, bHandlesEnabled(true)
+		, bIsDeterministic(false)
 		, bCanDisableContacts(true)
 		, GravityDirection(FVec3(0,0,-1))
 		, GravitySize(980)
@@ -283,6 +284,11 @@ namespace Chaos
 		// Prune the unused contacts
 		ConstraintAllocator.EndDetectCollisions();
 
+		if (bIsDeterministic)
+		{
+			ConstraintAllocator.SortConstraintsHandles();
+		}
+
 		// Initialize container properties
 		for (FPBDCollisionConstraint* Contact : GetConstraints())
 		{
diff --git a/Engine/Source/Runtime/Experimental/Chaos/Private/Chaos/PBDRigidsEvolutionGBF.cpp b/Engine/Source/Runtime/Experimental/Chaos/Private/Chaos/PBDRigidsEvolutionGBF.cpp
--- a/Engine/Source/Runtime/Experimental/Chaos/Private/Chaos/PBDRigidsEvolutionGBF.cpp
+++ b/Engine/Source/Runtime/Experimental/Chaos/Private/Chaos/PBDRigidsEvolutionGBF.cpp
@@ -589,6 +589,11 @@ void FPBDRigidsEvolutionGBF::ApplyConstraintsPhase2(const FReal Dt, int32 Island
 	}
 }
 
+void FPBDRigidsEvolutionGBF::SetIsDeterministic(const bool bInIsDeterministic)
+{
+	// We detect collisions in parallel, so order is non-deterministic without additional processing
+	CollisionConstraints.SetIsDeterministic(bInIsDeterministic);
+}
 
 FPBDRigidsEvolutionGBF::FPBDRigidsEvolutionGBF(FPBDRigidsSOAs& InParticles,THandleArray<FChaosPhysicsMaterial>& SolverPhysicsMaterials, const TArray<ISimCallbackObject*>* InCollisionModifiers, bool InIsSingleThreaded)
 	: Base(InParticles, SolverPhysicsMaterials, DefaultNumIterations, DefaultNumPushOutIterations, InIsSingleThreaded)
@@ -607,6 +612,7 @@ FPBDRigidsEvolutionGBF::FPBDRigidsEvolutionGBF(FPBDRigidsSOAs& InParticles,THand
 	, CurrentStepResimCacheImp(nullptr)
 	, CollisionModifiers(InCollisionModifiers)
 	, CCDManager()
+	, bIsDeterministic(false)
 {
 	CollisionConstraints.SetCanDisableContacts(!!CollisionDisableCulledContacts);
 
diff --git a/Engine/Source/Runtime/Experimental/Chaos/Private/PBDRigidsSolver.cpp b/Engine/Source/Runtime/Experimental/Chaos/Private/PBDRigidsSolver.cpp
--- a/Engine/Source/Runtime/Experimental/Chaos/Private/PBDRigidsSolver.cpp
+++ b/Engine/Source/Runtime/Experimental/Chaos/Private/PBDRigidsSolver.cpp
@@ -178,7 +178,7 @@ namespace Chaos
 		// These override the engine config if >= 0
 
 		int32 ChaosSolverIterations = -1;
-		FAutoConsoleVariableRef CVarChaosSolverIterations(TEXT("p.Chaos.Solver.Iterations"), ChaosSolverIterations, TEXT("Override umber of solver iterations (-1 to use config)"));
+		FAutoConsoleVariableRef CVarChaosSolverIterations(TEXT("p.Chaos.Solver.Iterations"), ChaosSolverIterations, TEXT("Override number of solver iterations (-1 to use config)"));
 
 		int32 ChaosSolverCollisionIterations = -1;
 		FAutoConsoleVariableRef CVarChaosSolverCollisionIterations(TEXT("p.Chaos.Solver.Collision.Iterations"), ChaosSolverCollisionIterations, TEXT("Override number of collision iterations per solver iteration (-1 to use config)"));
@@ -195,6 +195,9 @@ namespace Chaos
 		int32 ChaosSolverJointPushOutPairIterations = -1;
 		FAutoConsoleVariableRef CVarChaosSolverJointPushOutPairIterations(TEXT("p.Chaos.Solver.Joint.PushOutPairIterations"), ChaosSolverJointPushOutPairIterations, TEXT("Override number of push out iterations per joint during a solver iteration (-1 to use config)"));
 
+		int32 ChaosSolverDeterministic = -1;
+		FAutoConsoleVariableRef CVarChaosSolverDeterministic(TEXT("p.Chaos.Solver.Deterministic"), ChaosSolverDeterministic, TEXT("Override determinism. 0: disabled; 1: enabled; -1: default(disabled)"));
+
 		// Copied from RBAN
 		Chaos::FRealSingle ChaosSolverJointPositionTolerance = 0.025f;
 		FAutoConsoleVariableRef CVarChaosSolverJointPositionTolerance(TEXT("p.Chaos.Solver.Joint.PositionTolerance"), ChaosSolverJointPositionTolerance, TEXT("PositionTolerance."));
@@ -429,6 +432,7 @@ namespace Chaos
 		, bHasFloor(true)
 		, bIsFloorAnalytic(false)
 		, FloorHeight(0.f)
+		, bIsDeterministic(false)
 		, Particles(UniqueIndices)
 		, MEvolution(new FPBDRigidsEvolution(Particles, SimMaterials, &ContactModifiers, BufferingModeIn == Chaos::EMultiBufferMode::Single))
 		, MEventManager(new FEventManager(BufferingModeIn))
@@ -697,6 +701,8 @@ namespace Chaos
 		MRewindCallback = MoveTemp(RewindCallback);
 		MarshallingManager.SetHistoryLength_Internal(NumFrames);
 		MEvolution->SetRewindData(GetRewindData());
+		
+		UpdateIsDeterministic();
 	}
 
 	void FPBDRigidsSolver::Reset()
@@ -835,6 +841,10 @@ namespace Chaos
 			{
 				SetCollisionMaxPushOutVelocity(ChaosSolverMaxPushOutVelocity);
 			}
+			if (ChaosSolverDeterministic >= 0)
+			{
+				UpdateIsDeterministic();
+			}
 		}
 
 		UE_LOG(LogPBDRigidsSolver, Verbose, TEXT("PBDRigidsSolver::Tick(%3.5f)"), DeltaTime);
@@ -1603,6 +1613,26 @@ CSV_CUSTOM_STAT(PhysicsCounters, Name, Value, ECsvCustomStatOp::Set);
 		GetEvolution()->UpdateExternalAccelerationStructure_External(ExternalStructure,*PendingSpatialOperations_External);
 	}
 
+	bool FPBDRigidsSolver::IsDetemerministic() const
+	{
+		return bIsDeterministic || (MRewindData != nullptr) || (ChaosSolverDeterministic >= 1);
+	}
+
+	void FPBDRigidsSolver::SetIsDeterministic(const bool bInIsDeterministic)
+	{
+		if (bIsDeterministic != bInIsDeterministic)
+		{
+			bIsDeterministic = bInIsDeterministic;
+			UpdateIsDeterministic();
+		}
+	}
+
+	void FPBDRigidsSolver::UpdateIsDeterministic()
+	{
+		GetEvolution()->SetIsDeterministic(IsDetemerministic());
+	}
+
+
 	Chaos::FClusterCreationParameters::EConnectionMethod ToInternalConnectionMethod(EClusterUnionMethod InMethod)
 	{
 		using ETargetEnum = Chaos::FClusterCreationParameters::EConnectionMethod;
diff --git a/Engine/Source/Runtime/Experimental/Chaos/Public/Chaos/Collision/CollisionConstraintAllocator.h b/Engine/Source/Runtime/Experimental/Chaos/Public/Chaos/Collision/CollisionConstraintAllocator.h
--- a/Engine/Source/Runtime/Experimental/Chaos/Public/Chaos/Collision/CollisionConstraintAllocator.h
+++ b/Engine/Source/Runtime/Experimental/Chaos/Public/Chaos/Collision/CollisionConstraintAllocator.h
@@ -395,7 +395,6 @@ namespace Chaos
 					// Add the constraint to the active list
 					ActivateConstraintImp(NewConstraint);
 				}
-				SortConstraintsHandles();
 			}
 		}
 
diff --git a/Engine/Source/Runtime/Experimental/Chaos/Public/Chaos/PBDCollisionConstraints.h b/Engine/Source/Runtime/Experimental/Chaos/Public/Chaos/PBDCollisionConstraints.h
--- a/Engine/Source/Runtime/Experimental/Chaos/Public/Chaos/PBDCollisionConstraints.h
+++ b/Engine/Source/Runtime/Experimental/Chaos/Public/Chaos/PBDCollisionConstraints.h
@@ -83,6 +83,15 @@ public:
 		SolverType = InSolverType;
 	}
 
+	/**
+	 * @brief Enable or disable determinism.
+	 * Support for determinism requires that we sort active constraints each tick, so there is additional cost.
+	*/
+	void SetIsDeterministic(const bool bInIsDeterministic)
+	{
+		bIsDeterministic = bInIsDeterministic;
+	}
+
 	/**
 	 *  Clears the list of active constraints.
 	 * @todo(chaos): This is only required because of the way events work (see AdvanceOneTimeStepTask::DoWork)
@@ -309,6 +318,7 @@ private:
 	bool bEnableCollisions;
 	bool bEnableRestitution;
 	bool bHandlesEnabled;
+	bool bIsDeterministic;
 
 	// This is passed to IterationParameters. If true, then an iteration can cull a contact
 	// permanently (ie, for the remaining iterations) if it is ignored due to culldistance.
diff --git a/Engine/Source/Runtime/Experimental/Chaos/Public/Chaos/PBDRigidsEvolutionGBF.h b/Engine/Source/Runtime/Experimental/Chaos/Public/Chaos/PBDRigidsEvolutionGBF.h
--- a/Engine/Source/Runtime/Experimental/Chaos/Public/Chaos/PBDRigidsEvolutionGBF.h
+++ b/Engine/Source/Runtime/Experimental/Chaos/Public/Chaos/PBDRigidsEvolutionGBF.h
@@ -147,6 +147,8 @@ namespace Chaos
 			}
 		}
 
+		void SetIsDeterministic(const bool bInIsDeterministic);
+
 		CHAOS_API void Advance(const FReal Dt, const FReal MaxStepDt, const int32 MaxSteps);
 		CHAOS_API void AdvanceOneTimeStep(const FReal dt, const FSubStepInfo& SubStepInfo = FSubStepInfo());
 
@@ -367,6 +369,8 @@ namespace Chaos
 		const TArray<ISimCallbackObject*>* CollisionModifiers;
 
 		FCCDManager CCDManager;
+
+		bool bIsDeterministic;
 	};
 
 }
diff --git a/Engine/Source/Runtime/Experimental/Chaos/Public/PBDRigidsSolver.h b/Engine/Source/Runtime/Experimental/Chaos/Public/PBDRigidsSolver.h
--- a/Engine/Source/Runtime/Experimental/Chaos/Public/PBDRigidsSolver.h
+++ b/Engine/Source/Runtime/Experimental/Chaos/Public/PBDRigidsSolver.h
@@ -206,6 +206,20 @@ namespace Chaos
 		void SetTrailingFilterSettings(const FSolverTrailingFilterSettings& InTrailingFilterSettings) { GetEventFilters()->GetTrailingFilter()->UpdateFilterSettings(InTrailingFilterSettings); }
 		void SetRemovalFilterSettings(const FSolverRemovalFilterSettings& InRemovalFilterSettings) { GetEventFilters()->GetRemovalFilter()->UpdateFilterSettings(InRemovalFilterSettings); }
 
+		/**
+		 * @brief True if the simulation is running in deterministic mode
+		 * This will be true if determinism is explicitly requested (via SetIsDeterministic()) or if required
+		 * by some otyher system like Rewind/Resim support.
+		*/
+		bool IsDetemerministic() const;
+
+		/**
+		 * @brief Request that the sim be determinitic (or not)
+		 * @note Even if set to false, the sim may still be deterministic if some other feature is enabled and requires it.
+		 * @see IsDetemerministic()
+		*/
+		void SetIsDeterministic(const bool bInIsDeterministic);
+
 		/**/
 		FJointConstraints& GetJointConstraints() { return MEvolution->GetJointConstraints(); }
 		const FJointConstraints& GetJointConstraints() const { return MEvolution->GetJointConstraints(); }
@@ -311,6 +325,8 @@ namespace Chaos
 		virtual void PushPhysicsState(const FReal ExternalDt, const int32 NumSteps, const int32 NumExternalSteps) override;
 		virtual void SetExternalTimestampConsumed_Internal(const int32 Timestamp) override;
 
+		void UpdateIsDeterministic();
+
 		//
 		// Solver Data
 		//
@@ -320,6 +336,7 @@ namespace Chaos
 		bool bHasFloor;
 		bool bIsFloorAnalytic;
 		FReal FloorHeight;
+		bool bIsDeterministic;
 
 		FParticleUniqueIndicesMultithreaded UniqueIndices;
 		FParticlesType Particles;
