Chaos - add check for cull distance before adding manifold points in various implicit pair collision functions

diff --git a/Engine/Source/Runtime/Experimental/Chaos/Private/Chaos/Collision/ParticlePairMidPhase.cpp b/Engine/Source/Runtime/Experimental/Chaos/Private/Chaos/Collision/ParticlePairMidPhase.cpp
--- a/Engine/Source/Runtime/Experimental/Chaos/Private/Chaos/Collision/ParticlePairMidPhase.cpp
+++ b/Engine/Source/Runtime/Experimental/Chaos/Private/Chaos/Collision/ParticlePairMidPhase.cpp
@@ -271,7 +271,7 @@ namespace Chaos
 				// Run the narrow phase
 				if (!bUseCCD)
 				{
-					Collisions::UpdateConstraint(*Constraint.Get(), Shape0->GetLeafWorldTransform(), Shape1->GetLeafWorldTransform(), Dt);
+					Collisions::UpdateConstraint(*Constraint.Get(), ShapeWorldTransform0, ShapeWorldTransform1, Dt);
 				}
 				else
 				{
diff --git a/Engine/Source/Runtime/Experimental/Chaos/Private/Chaos/CollisionOneShotManifoldsMiscShapes.cpp b/Engine/Source/Runtime/Experimental/Chaos/Private/Chaos/CollisionOneShotManifoldsMiscShapes.cpp
--- a/Engine/Source/Runtime/Experimental/Chaos/Private/Chaos/CollisionOneShotManifoldsMiscShapes.cpp
+++ b/Engine/Source/Runtime/Experimental/Chaos/Private/Chaos/CollisionOneShotManifoldsMiscShapes.cpp
@@ -44,8 +44,10 @@ namespace Chaos
 			Constraint.ResetManifold();
 
 			FContactPoint ContactPoint = SphereSphereContactPoint(SphereA, SphereATransform, SphereB, SphereBTransform, Constraint.Manifold.RestitutionPadding);
-
-			Constraint.AddOneshotManifoldContact(ContactPoint);
+			if (ContactPoint.Phi < Constraint.GetCullDistance())
+			{
+				Constraint.AddOneshotManifoldContact(ContactPoint);
+			}
 		}
 
 		void ConstructSpherePlaneOneShotManifold(
@@ -65,8 +67,10 @@ namespace Chaos
 			Constraint.ResetManifold();
 
 			FContactPoint ContactPoint = SpherePlaneContactPoint(Sphere, SphereTransform, Plane, PlaneTransform, Constraint.Manifold.RestitutionPadding);
-
-			Constraint.AddOneshotManifoldContact(ContactPoint);
+			if (ContactPoint.Phi < Constraint.GetCullDistance())
+			{
+				Constraint.AddOneshotManifoldContact(ContactPoint);
+			}
 		}	
 
 		void ConstructSphereBoxOneShotManifold(const TSphere<FReal, 3>& Sphere, const FRigidTransform3& SphereTransform, const FImplicitBox3& Box, const FRigidTransform3& BoxTransform, const FReal Dt, FPBDCollisionConstraint& Constraint)
@@ -80,8 +84,10 @@ namespace Chaos
 			Constraint.ResetManifold();
 
 			FContactPoint ContactPoint = SphereBoxContactPoint(Sphere, SphereTransform, Box, BoxTransform, Constraint.Manifold.RestitutionPadding);
-
-			Constraint.AddOneshotManifoldContact(ContactPoint);
+			if (ContactPoint.Phi < Constraint.GetCullDistance())
+			{
+				Constraint.AddOneshotManifoldContact(ContactPoint);
+			}
 		}
 
 		void ConstructSphereCapsuleOneShotManifold(const TSphere<FReal, 3>& Sphere, const FRigidTransform3& SphereTransform, const FCapsule& Capsule, const FRigidTransform3& CapsuleTransform, const FReal Dt, FPBDCollisionConstraint& Constraint)
@@ -95,8 +101,10 @@ namespace Chaos
 			Constraint.ResetManifold();
 
 			FContactPoint ContactPoint = SphereCapsuleContactPoint(Sphere, SphereTransform, Capsule, CapsuleTransform, Constraint.Manifold.RestitutionPadding);
-
-			Constraint.AddOneshotManifoldContact(ContactPoint);
+			if (ContactPoint.Phi < Constraint.GetCullDistance())
+			{
+				Constraint.AddOneshotManifoldContact(ContactPoint);
+			}
 		}
 
 		void ConstructSphereConvexManifold(const TSphere<FReal, 3>& Sphere, const FRigidTransform3& SphereTransform, const FImplicitObject3& Convex, const FRigidTransform3& ConvexTransform, const FReal Dt, FPBDCollisionConstraint& Constraint)
@@ -110,8 +118,10 @@ namespace Chaos
 			Constraint.ResetManifold();
 
 			FContactPoint ContactPoint = SphereConvexContactPoint(Sphere, SphereTransform, Convex, ConvexTransform);
-
-			Constraint.AddOneshotManifoldContact(ContactPoint);
+			if (ContactPoint.Phi < Constraint.GetCullDistance())
+			{
+				Constraint.AddOneshotManifoldContact(ContactPoint);
+			}
 		}
 
 		template <typename TriMeshType>
@@ -126,8 +136,10 @@ namespace Chaos
 			Constraint.ResetManifold();
 
 			FContactPoint ContactPoint = SphereTriangleMeshContactPoint(Sphere, SphereWorldTransform, TriangleMesh, TriMeshWorldTransform, Constraint.GetCullDistance(), 0.0f);
-
-			Constraint.AddOneshotManifoldContact(ContactPoint);
+			if (ContactPoint.Phi < Constraint.GetCullDistance())
+			{
+				Constraint.AddOneshotManifoldContact(ContactPoint);
+			}
 		}
 
 		void ConstructSphereHeightFieldOneShotManifold(const TSphere<FReal, 3>& Sphere, const FRigidTransform3& SphereTransform, const FHeightField& Heightfield, const FRigidTransform3& HeightfieldTransform, const FReal Dt, FPBDCollisionConstraint& Constraint)
@@ -138,8 +150,10 @@ namespace Chaos
 			ensure(HeightfieldTransform.GetScale3D() == FVec3(1.0f, 1.0f, 1.0f));
 
 			FContactPoint ContactPoint = SphereHeightFieldContactPoint(Sphere, SphereTransform, Heightfield, HeightfieldTransform, Constraint.GetCullDistance(), 0.0f);
-
-			Constraint.AddOneshotManifoldContact(ContactPoint);
+			if (ContactPoint.Phi < Constraint.GetCullDistance())
+			{
+				Constraint.AddOneshotManifoldContact(ContactPoint);
+			}
 		}
 
 		void ConstructCapsuleCapsuleOneShotManifold(const FCapsule& CapsuleA, const FRigidTransform3& CapsuleATransform, const FCapsule& CapsuleB, const FRigidTransform3& CapsuleBTransform, const FReal Dt, FPBDCollisionConstraint& Constraint)
@@ -264,7 +278,10 @@ namespace Chaos
 			GJKImplicitManifold<FCapsule, TriMeshType>(Capsule, CapsuleWorldTransform, TriangleMesh, TriMeshWorldTransform, Constraint.GetCullDistance(), 0.0f, ContactPoints);
 			for (FContactPoint& ContactPoint : ContactPoints)
 			{
-				Constraint.AddOneshotManifoldContact(ContactPoint);
+				if (ContactPoint.Phi < Constraint.GetCullDistance())
+				{
+					Constraint.AddOneshotManifoldContact(ContactPoint);
+				}
 			}
 		}
 
@@ -282,7 +299,10 @@ namespace Chaos
 			GJKImplicitManifold<FCapsule>(Capsule, CapsuleTransform, HeightField, HeightFieldTransform, Constraint.GetCullDistance(), 0.0f, ContactPoints);
 			for (FContactPoint& ContactPoint : ContactPoints)
 			{
-				Constraint.AddOneshotManifoldContact(ContactPoint);
+				if (ContactPoint.Phi < Constraint.GetCullDistance())
+				{
+					Constraint.AddOneshotManifoldContact(ContactPoint);
+				}
 			}
 		}
 
@@ -303,7 +323,10 @@ namespace Chaos
 			GJKImplicitManifold<ConvexType>(Convex, ConvexTransform, HeightField, HeightFieldTransform, Constraint.GetCullDistance(), 0.0f, ContactPoints);
 			for (FContactPoint& ContactPoint : ContactPoints)
 			{
-				Constraint.AddOneshotManifoldContact(ContactPoint);
+				if (ContactPoint.Phi < Constraint.GetCullDistance())
+				{
+					Constraint.AddOneshotManifoldContact(ContactPoint);
+				}
 			}
 		}
 
@@ -322,7 +345,10 @@ namespace Chaos
 			GJKImplicitManifold<ConvexType, TriMeshType>(Convex, ConvexTransform, TriangleMesh, TriMeshTransform, Constraint.GetCullDistance(), 0.0f, ContactPoints);
 			for (FContactPoint& ContactPoint : ContactPoints)
 			{
-				Constraint.AddOneshotManifoldContact(ContactPoint);
+				if (ContactPoint.Phi < Constraint.GetCullDistance())
+				{
+					Constraint.AddOneshotManifoldContact(ContactPoint);
+				}
 			}
 		}
 
