Chaos - only update collision materials when contafct is first created

diff --git a/Engine/Source/Runtime/Experimental/Chaos/Private/Chaos/PBDCollisionConstraints.cpp b/Engine/Source/Runtime/Experimental/Chaos/Private/Chaos/PBDCollisionConstraints.cpp
--- a/Engine/Source/Runtime/Experimental/Chaos/Private/Chaos/PBDCollisionConstraints.cpp
+++ b/Engine/Source/Runtime/Experimental/Chaos/Private/Chaos/PBDCollisionConstraints.cpp
@@ -289,12 +289,15 @@ namespace Chaos
 			ConstraintAllocator.SortConstraintsHandles();
 		}
 
-		// Initialize container properties
+		// Bind the constraints to this container and initialize other properties
+		// @todo(chaos): this could be set on creation if the midphase knew about the container
 		for (FPBDCollisionConstraint* Contact : GetConstraints())
 		{
-			// @todo(chaos): this could be set on creation if the allocator knew about the container
-			Contact->SetContainer(this);
-			UpdateConstraintMaterialProperties(*Contact);
+			if (Contact->GetContainer() == nullptr)
+			{
+				Contact->SetContainer(this);
+				UpdateConstraintMaterialProperties(*Contact);
+			}
 		}
 	}
 
diff --git a/Engine/Source/Runtime/Experimental/Chaos/Public/Chaos/ConstraintHandle.h b/Engine/Source/Runtime/Experimental/Chaos/Public/Chaos/ConstraintHandle.h
--- a/Engine/Source/Runtime/Experimental/Chaos/Public/Chaos/ConstraintHandle.h
+++ b/Engine/Source/Runtime/Experimental/Chaos/Public/Chaos/ConstraintHandle.h
@@ -100,6 +100,16 @@ namespace Chaos
 			return (ConstraintContainer != nullptr) && IsEnabled();
 		}
 
+		FPBDConstraintContainer* GetContainer()
+		{
+			return ConstraintContainer;
+		}
+
+		const FPBDConstraintContainer* GetContainer() const 
+		{
+			return ConstraintContainer;
+		}
+
 		int32 ConstraintGraphIndex() const
 		{
 			return GraphIndex;
