Chaos - prep work to separately allocate solver manifold points for better cache behaviour

diff --git a/Engine/Source/Runtime/Experimental/Chaos/Private/Chaos/Collision/PBDCollisionSolver.cpp b/Engine/Source/Runtime/Experimental/Chaos/Private/Chaos/Collision/PBDCollisionSolver.cpp
--- a/Engine/Source/Runtime/Experimental/Chaos/Private/Chaos/Collision/PBDCollisionSolver.cpp
+++ b/Engine/Source/Runtime/Experimental/Chaos/Private/Chaos/Collision/PBDCollisionSolver.cpp
@@ -193,7 +193,7 @@ namespace Chaos
 		const FReal DynamicFriction,
 		const FVec3& ContactDelta,
 		const FReal ContactDeltaNormal,
-		FPBDCollisionSolver::FSolverManifoldPoint& ManifoldPoint,
+		FPBDCollisionSolverManifoldPoint& ManifoldPoint,
 		FConstraintSolverBody& Body0,
 		FConstraintSolverBody& Body1)
 	{
@@ -319,7 +319,7 @@ namespace Chaos
 		const FReal DynamicFriction,
 		const FVec3& ContactVelocityDelta,
 		const FReal ContactVelocityDeltaNormal,
-		FPBDCollisionSolver::FSolverManifoldPoint& ManifoldPoint,
+		FPBDCollisionSolverManifoldPoint& ManifoldPoint,
 		FConstraintSolverBody& Body0,
 		FConstraintSolverBody& Body1)
 	{
@@ -351,26 +351,7 @@ namespace Chaos
 		}
 	}
 
-	FPBDCollisionSolver::FState::FState()
-		: SolverBodies()
-		, ManifoldPoints()
-		, NumManifoldPoints(0)
-		, StaticFriction(0)
-		, DynamicFriction(0)
-		, Stiffness(1)
-		, BodyEpochs{ INDEX_NONE, INDEX_NONE }
-		, NumPositionSolves(0)
-		, NumVelocitySolves(0)
-		, bIsSolved(false)
-	{
-	}
-
-	FPBDCollisionSolver::FPBDCollisionSolver()
-		: State()
-	{
-	}
-
-	void FPBDCollisionSolver::FSolverManifoldPoint::InitContact(
+	void FPBDCollisionSolverManifoldPoint::InitContact(
 		const FConstraintSolverBody& Body0,
 		const FConstraintSolverBody& Body1,
 		const FVec3& InCoMAnchorPoint0,
@@ -383,7 +364,7 @@ namespace Chaos
 		UpdateMass(Body0, Body1);
 	}
 
-	void FPBDCollisionSolver::FSolverManifoldPoint::InitMaterial(
+	void FPBDCollisionSolverManifoldPoint::InitMaterial(
 		const FReal InWorldContactVelocityTargetNormal,
 		const bool bInEnableStaticFriction,
 		const FReal InStaticFrictionMax)
@@ -393,7 +374,7 @@ namespace Chaos
 		WorldContactVelocityTargetNormal = InWorldContactVelocityTargetNormal;
 	}
 
-	inline void FPBDCollisionSolver::FSolverManifoldPoint::UpdateContact(
+	inline void FPBDCollisionSolverManifoldPoint::UpdateContact(
 		const FConstraintSolverBody& Body0, 
 		const FConstraintSolverBody& Body1,
 		const FVec3& InCoMAnchorPoint0,
@@ -415,7 +396,7 @@ namespace Chaos
 		WorldContactDelta = WorldAnchorPoint0 - WorldAnchorPoint1;
 	}
 
-	inline void FPBDCollisionSolver::FSolverManifoldPoint::UpdateMass(const FConstraintSolverBody& Body0, const FConstraintSolverBody& Body1)
+	inline void FPBDCollisionSolverManifoldPoint::UpdateMass(const FConstraintSolverBody& Body0, const FConstraintSolverBody& Body1)
 	{
 		const FMatrix33 ContactMassInv =
 			(Body0.IsDynamic() ? Collisions::ComputeFactorMatrix3(WorldContactPosition - Body0.P(), Body0.InvI(), Body0.InvM()) : FMatrix33(0)) +
@@ -427,14 +408,14 @@ namespace Chaos
 		WorldContactMassNormal = ContactMassNormal;
 	}
 
-	FVec3 FPBDCollisionSolver::FSolverManifoldPoint::CalculateContactVelocity(const FConstraintSolverBody& Body0, const FConstraintSolverBody& Body1) const
+	FVec3 FPBDCollisionSolverManifoldPoint::CalculateContactVelocity(const FConstraintSolverBody& Body0, const FConstraintSolverBody& Body1) const
 	{
 		const FVec3 ContactVelocity0 = Body0.V() + FVec3::CrossProduct(Body0.W(), WorldContactPosition - Body0.P());
 		const FVec3 ContactVelocity1 = Body1.V() + FVec3::CrossProduct(Body1.W(), WorldContactPosition - Body1.P());
 		return ContactVelocity0 - ContactVelocity1;
 	}
 
-	inline void FPBDCollisionSolver::FSolverManifoldPoint::CalculateContactPositionError(const FConstraintSolverBody& Body0, const FConstraintSolverBody& Body1, const FReal MaxPushOut, FVec3& OutContactDelta, FReal& OutContactDeltaNormal) const
+	inline void FPBDCollisionSolverManifoldPoint::CalculateContactPositionError(const FConstraintSolverBody& Body0, const FConstraintSolverBody& Body1, const FReal MaxPushOut, FVec3& OutContactDelta, FReal& OutContactDeltaNormal) const
 	{
 		const FVec3 WorldRelativeContactPosition0 = WorldContactPosition - Body0.P();
 		const FVec3 WorldRelativeContactPosition1 = WorldContactPosition - Body1.P();
@@ -462,7 +443,7 @@ namespace Chaos
 		}
 	}
 
-	inline void FPBDCollisionSolver::FSolverManifoldPoint::CalculateContactVelocityError(const FConstraintSolverBody& Body0, const FConstraintSolverBody& Body1, const FReal DynamicFriction, const FReal Dt, FVec3& OutContactVelocityDelta, FReal& OutContactVelocityDeltaNormal) const
+	inline void FPBDCollisionSolverManifoldPoint::CalculateContactVelocityError(const FConstraintSolverBody& Body0, const FConstraintSolverBody& Body1, const FReal DynamicFriction, const FReal Dt, FVec3& OutContactVelocityDelta, FReal& OutContactVelocityDeltaNormal) const
 	{
 		const FVec3 ContactVelocity0 = Body0.V() + FVec3::CrossProduct(Body0.W(), WorldContactPosition - Body0.P());
 		const FVec3 ContactVelocity1 = Body1.V() + FVec3::CrossProduct(Body1.W(), WorldContactPosition - Body1.P());
@@ -497,6 +478,18 @@ namespace Chaos
 		}
 	}
 
+
+	//////////////////////////////////////////////////////////////////////////////////////////////////
+	//////////////////////////////////////////////////////////////////////////////////////////////////
+	//////////////////////////////////////////////////////////////////////////////////////////////////
+	//////////////////////////////////////////////////////////////////////////////////////////////////
+
+
+	FPBDCollisionSolver::FPBDCollisionSolver()
+		: State()
+	{
+	}
+
 	void FPBDCollisionSolver::EnablePositionShockPropagation()
 	{
 		SetShockPropagationInvMassScale(Chaos_PBDCollisionSolver_Position_MinInvMassScale);
@@ -611,7 +604,7 @@ namespace Chaos
 		// Apply the position correction so that all contacts have zero separation
 		for (int32 PointIndex = 0; PointIndex < State.NumManifoldPoints; ++PointIndex)
 		{
-			FSolverManifoldPoint& SolverManifoldPoint = State.ManifoldPoints[PointIndex];
+			FPBDCollisionSolverManifoldPoint& SolverManifoldPoint = State.ManifoldPoints[PointIndex];
 
 			FVec3 ContactDelta;
 			FReal ContactDeltaNormal;
@@ -660,7 +653,7 @@ namespace Chaos
 
 		for (int32 PointIndex = 0; PointIndex < State.NumManifoldPoints; ++PointIndex)
 		{
-			FSolverManifoldPoint& SolverManifoldPoint = State.ManifoldPoints[PointIndex];
+			FPBDCollisionSolverManifoldPoint& SolverManifoldPoint = State.ManifoldPoints[PointIndex];
 
 			//SolverManifoldPoint.UpdateContactPoint(Body0, Body1, State.Constraint->GetManifoldPoints()[PointIndex]);
 			//SolverManifoldPoint.UpdateContactMass(Body0, Body1);
diff --git a/Engine/Source/Runtime/Experimental/Chaos/Private/Chaos/Collision/PBDCollisionSolver.h b/Engine/Source/Runtime/Experimental/Chaos/Private/Chaos/Collision/PBDCollisionSolver.h
--- a/Engine/Source/Runtime/Experimental/Chaos/Private/Chaos/Collision/PBDCollisionSolver.h
+++ b/Engine/Source/Runtime/Experimental/Chaos/Private/Chaos/Collision/PBDCollisionSolver.h
@@ -20,100 +20,103 @@ namespace Chaos
 	class FPBDCollisionConstraint;
 
 	/**
-	 * @brief 
-	 * @todo(chaos): Make this solver operate on a single contact point rather than all points in a manifold.
-	 * This would be beneficial if we have many contacts with less than 4 points in the manifold. However this
-	 * is dificult to do while we are still supporting non-manifold collisions.
+	 * @brief A single contact point in a FPBDCollisionSolver
 	*/
-	class FPBDCollisionSolver
+	class FPBDCollisionSolverManifoldPoint
 	{
 	public:
-		static const int32 MaxConstrainedBodies = 2;
-		static const int32 MaxPointsPerConstraint = 4;
+		/**
+		 * @brief Initialize the geometric data for the contact
+		*/
+		void InitContact(
+			const FConstraintSolverBody& Body0,
+			const FConstraintSolverBody& Body1,
+			const FVec3& InCoMAnchorPoint0,
+			const FVec3& InCoMAnchorPoint1,
+			const FVec3& InWorldContactNormal);
 
-		class FSolverManifoldPoint
-		{
-		public:
-			/**
-			 * @brief Initialize the geometric data for the contact
-			*/
-			void InitContact(
-				const FConstraintSolverBody& Body0,
-				const FConstraintSolverBody& Body1,
-				const FVec3& InCoMAnchorPoint0,
-				const FVec3& InCoMAnchorPoint1,
-				const FVec3& InWorldContactNormal);
-
-			/**
-			 * @brief Initialize the material related properties of the contact
-			*/
-			void InitMaterial(
-				const FReal InWorldContactVelocityTargetNormal,
-				const bool bInEnableStaticFriction,
-				const FReal InStaticFrictionMax);
-
-			/**
-			 * @brief Update the world-space relative contact points based on current body transforms and body-space contact positions
-			*/
-			void UpdateContact(
-				const FConstraintSolverBody& Body0,
-				const FConstraintSolverBody& Body1,
-				const FVec3& CoMAnchorPoint0,
-				const FVec3& CoMAnchorPoint1,
-				const FVec3& WorldContactNormal);
-
-			/**
-			 * @brief Update the cached mass properties based on the current body transforms
-			*/
-			void UpdateMass(const FConstraintSolverBody& Body0, const FConstraintSolverBody& Body1);
-
-			/**
-			 * @brief Calculate the relative velocity at the contact point
-			 * @note InitContact must be called before calling this function
-			*/
-			FVec3 CalculateContactVelocity(const FConstraintSolverBody& Body0, const FConstraintSolverBody& Body1) const;
-
-			/**
-			 * @brief Calculate the position error at the current transforms
-			 * @param MaxPushOut a limit on the position error for this iteration to prevent initial-penetration explosion (a common PBD problem)
-			*/
-			void CalculateContactPositionError(const FConstraintSolverBody& Body0, const FConstraintSolverBody& Body1, const FReal MaxPushOut, FVec3& OutContactDelta, FReal& OutContactDeltaNormal) const;
-
-			/**
-			 * @brief Calculate the velocity error at the current transforms
-			*/
-			void CalculateContactVelocityError(const FConstraintSolverBody& Body0, const FConstraintSolverBody& Body1, const FReal DynamicFriction, const FReal Dt, FVec3& OutContactVelocityDelta, FReal& OutContactVelocityDeltaNormal) const;
+		/**
+		 * @brief Initialize the material related properties of the contact
+		*/
+		void InitMaterial(
+			const FReal InWorldContactVelocityTargetNormal,
+			const bool bInEnableStaticFriction,
+			const FReal InStaticFrictionMax);
+
+		/**
+		 * @brief Update the world-space relative contact points based on current body transforms and body-space contact positions
+		*/
+		void UpdateContact(
+			const FConstraintSolverBody& Body0,
+			const FConstraintSolverBody& Body1,
+			const FVec3& CoMAnchorPoint0,
+			const FVec3& CoMAnchorPoint1,
+			const FVec3& WorldContactNormal);
+
+		/**
+		 * @brief Update the cached mass properties based on the current body transforms
+		*/
+		void UpdateMass(const FConstraintSolverBody& Body0, const FConstraintSolverBody& Body1);
+
+		/**
+		 * @brief Calculate the relative velocity at the contact point
+		 * @note InitContact must be called before calling this function
+		*/
+		FVec3 CalculateContactVelocity(const FConstraintSolverBody& Body0, const FConstraintSolverBody& Body1) const;
+
+		/**
+		 * @brief Calculate the position error at the current transforms
+		 * @param MaxPushOut a limit on the position error for this iteration to prevent initial-penetration explosion (a common PBD problem)
+		*/
+		void CalculateContactPositionError(const FConstraintSolverBody& Body0, const FConstraintSolverBody& Body1, const FReal MaxPushOut, FVec3& OutContactDelta, FReal& OutContactDeltaNormal) const;
+
+		/**
+		 * @brief Calculate the velocity error at the current transforms
+		*/
+		void CalculateContactVelocityError(const FConstraintSolverBody& Body0, const FConstraintSolverBody& Body1, const FReal DynamicFriction, const FReal Dt, FVec3& OutContactVelocityDelta, FReal& OutContactVelocityDeltaNormal) const;
 
 		// @todo(chaos): make private
-		public:
-			friend class FPBDCollisionSolver;
+	public:
+		friend class FPBDCollisionSolver;
 
-			// World-space contact point
-			FVec3 WorldContactPosition;
+		// World-space contact point
+		FVec3 WorldContactPosition;
 
-			// World-space contact normal
-			FVec3 WorldContactNormal;
+		// World-space contact normal
+		FVec3 WorldContactNormal;
 
-			// Contact mass
-			FMatrix33 WorldContactMass;
-			FReal WorldContactMassNormal;
+		// Contact mass
+		FMatrix33 WorldContactMass;
+		FReal WorldContactMassNormal;
 
-			// World-space contact separation that we are trying to correct
-			FVec3 WorldContactDelta;
+		// World-space contact separation that we are trying to correct
+		FVec3 WorldContactDelta;
 
-			// Desired final normal velocity, taking Restitution into account
-			FReal WorldContactVelocityTargetNormal;
+		// Desired final normal velocity, taking Restitution into account
+		FReal WorldContactVelocityTargetNormal;
 
-			// Solver outputs
-			FVec3 NetPushOut;
-			FVec3 NetImpulse;
+		// Solver outputs
+		FVec3 NetPushOut;
+		FVec3 NetImpulse;
 
-			// A smoothed NetImpulse along the normal, used for clipping to the static friction cone
-			FReal StaticFrictionMax;
+		// A smoothed NetImpulse along the normal, used for clipping to the static friction cone
+		FReal StaticFrictionMax;
 
-			// Whether we are still in the static friction cone
-			bool bInsideStaticFrictionCone;
-		};
+		// Whether we are still in the static friction cone
+		bool bInsideStaticFrictionCone;
+	};
+
+	/**
+	 * @brief 
+	 * @todo(chaos): Make this solver operate on a single contact point rather than all points in a manifold.
+	 * This would be beneficial if we have many contacts with less than 4 points in the manifold. However this
+	 * is dificult to do while we are still supporting non-manifold collisions.
+	*/
+	class FPBDCollisionSolver
+	{
+	public:
+		static const int32 MaxConstrainedBodies = 2;
+		static const int32 MaxPointsPerConstraint = 4;
 
 		FPBDCollisionSolver();
 
@@ -148,12 +151,13 @@ namespace Chaos
 			return State.NumManifoldPoints;
 		}
 
-		void SetNumManifoldPoints(const int32 InNumManifoldPoints)
+		int32 SetNumManifoldPoints(const int32 InNumManifoldPoints)
 		{
-			State.NumManifoldPoints = InNumManifoldPoints;
+			State.NumManifoldPoints = FMath::Min(InNumManifoldPoints, MaxPointsPerConstraint);
+			return State.NumManifoldPoints;
 		}
 
-		const FSolverManifoldPoint& GetManifoldPoint(const int32 ManifoldPointIndex) const
+		const FPBDCollisionSolverManifoldPoint& GetManifoldPoint(const int32 ManifoldPointIndex) const
 		{
 			check(ManifoldPointIndex < NumManifoldPoints());
 			return State.ManifoldPoints[ManifoldPointIndex];
@@ -237,10 +241,22 @@ namespace Chaos
 
 		struct FState
 		{
-			FState();
+			FState()
+				: SolverBodies()
+				, ManifoldPoints()
+				, NumManifoldPoints(0)
+				, StaticFriction(0)
+				, DynamicFriction(0)
+				, Stiffness(1)
+				, BodyEpochs{ INDEX_NONE, INDEX_NONE }
+				, NumPositionSolves(0)
+				, NumVelocitySolves(0)
+				, bIsSolved(false)
+			{
+			}
 
 			FConstraintSolverBody SolverBodies[MaxConstrainedBodies];
-			FSolverManifoldPoint ManifoldPoints[MaxPointsPerConstraint];
+			FPBDCollisionSolverManifoldPoint ManifoldPoints[MaxPointsPerConstraint];
 			int32 NumManifoldPoints;
 			FReal StaticFriction;
 			FReal DynamicFriction;
diff --git a/Engine/Source/Runtime/Experimental/Chaos/Private/Chaos/Collision/SolverCollisionContainer.cpp b/Engine/Source/Runtime/Experimental/Chaos/Private/Chaos/Collision/SolverCollisionContainer.cpp
--- a/Engine/Source/Runtime/Experimental/Chaos/Private/Chaos/Collision/SolverCollisionContainer.cpp
+++ b/Engine/Source/Runtime/Experimental/Chaos/Private/Chaos/Collision/SolverCollisionContainer.cpp
@@ -29,6 +29,12 @@ namespace Chaos
 	extern bool bChaos_Collision_Manifold_FixNormalsInWorldSpace;
 
 
+	//////////////////////////////////////////////////////////////////////////////////////////////////
+	//////////////////////////////////////////////////////////////////////////////////////////////////
+	//////////////////////////////////////////////////////////////////////////////////////////////////
+	//////////////////////////////////////////////////////////////////////////////////////////////////
+
+
 	/**
 	 * @brief A wrapper for FPBDCollisionSolver which binds to a Collision Constraint and adds Gather/Scatter from/to the constraint
 	*/
@@ -75,17 +81,15 @@ namespace Chaos
 
 			GatherManifoldPoints(Dt);
 
-			// We should try to remove this
+			// We should try to remove this - the Constraint should need to know about solver objects
 			Constraint->SetSolverBodies(Body0, Body1);
 		}
 
 		void GatherManifoldPoints(const FReal Dt)
 		{
-			// We handle incremental manifolds by just colecting any new contacts
+			// We handle incremental manifolds by just collecting any new contacts
 			const int32 BeginPointIndex = Solver.NumManifoldPoints();
-			const int32 EndPointIndex = FMath::Min(Constraint->GetManifoldPoints().Num(), FPBDCollisionSolver::MaxPointsPerConstraint);
-
-			Solver.SetNumManifoldPoints(EndPointIndex);
+			const int32 EndPointIndex = Solver.SetNumManifoldPoints(Constraint->GetManifoldPoints().Num());
 
 			const FReal RestitutionVelocityThreshold = Constraint->Manifold.RestitutionThreshold * Dt;
 			const FReal Restitution = Constraint->Manifold.Restitution;
@@ -95,22 +99,19 @@ namespace Chaos
 				TArrayView<FManifoldPoint> ManifoldPoints = Constraint->GetManifoldPoints();
 				const FManifoldPoint& ManifoldPoint = ManifoldPoints[ManifoldPointIndex];
 
-				// The world-space contact normal
-				const FVec3& WorldContactNormal = ManifoldPoint.ContactPoint.Normal;
-
 				// Initialize the structural data in the contact (relative contact points, contact mass etc)
 				Solver.InitContact(
 					ManifoldPointIndex,
 					ManifoldPoint.CoMContactPoints[0],
 					ManifoldPoint.CoMContactPoints[1],
-					WorldContactNormal);
+					ManifoldPoint.ContactPoint.Normal);
 
 				// Calculate the target normal velocity based on restitution
 				FReal WorldContactVelocityTargetNormal = FReal(0);
 				if (Restitution > FReal(0))
 				{
 					const FVec3 ContactVelocity = Solver.GetManifoldPoint(ManifoldPointIndex).CalculateContactVelocity(Solver.SolverBody0(), Solver.SolverBody1());
-					const FReal ContactVelocityNormal = FVec3::DotProduct(ContactVelocity, WorldContactNormal);
+					const FReal ContactVelocityNormal = FVec3::DotProduct(ContactVelocity, ManifoldPoint.ContactPoint.Normal);
 					if (ContactVelocityNormal < -RestitutionVelocityThreshold)
 					{
 						WorldContactVelocityTargetNormal = -Restitution * ContactVelocityNormal;
@@ -136,7 +137,7 @@ namespace Chaos
 
 			for (int32 PointIndex = 0; PointIndex < Solver.NumManifoldPoints(); ++PointIndex)
 			{
-				const FPBDCollisionSolver::FSolverManifoldPoint& SolverManifoldPoint = Solver.GetManifoldPoint(PointIndex);
+				const FPBDCollisionSolverManifoldPoint& SolverManifoldPoint = Solver.GetManifoldPoint(PointIndex);
 
 				TArrayView<FManifoldPoint> ManifoldPoints = Constraint->GetManifoldPoints();
 				FManifoldPoint& ManifoldPoint = ManifoldPoints[PointIndex];
@@ -164,6 +165,12 @@ namespace Chaos
 	};
 
 
+	//////////////////////////////////////////////////////////////////////////////////////////////////
+	//////////////////////////////////////////////////////////////////////////////////////////////////
+	//////////////////////////////////////////////////////////////////////////////////////////////////
+	//////////////////////////////////////////////////////////////////////////////////////////////////
+
+
 	FPBDCollisionSolverContainer::FPBDCollisionSolverContainer()
 		: FConstraintSolverContainer()
 		, MaxPushOutVelocity(0)
diff --git a/Engine/Source/Runtime/Experimental/Chaos/Public/Chaos/Collision/SolverCollisionContainer.h b/Engine/Source/Runtime/Experimental/Chaos/Public/Chaos/Collision/SolverCollisionContainer.h
--- a/Engine/Source/Runtime/Experimental/Chaos/Public/Chaos/Collision/SolverCollisionContainer.h
+++ b/Engine/Source/Runtime/Experimental/Chaos/Public/Chaos/Collision/SolverCollisionContainer.h
@@ -17,6 +17,7 @@ namespace Chaos
 	class FPBDCollisionConstraint;
 	class FPBDCollisionSolver;
 	class FPBDCollisionSolverAdapter;
+	class FPBDCollisionSolverManifoldPoint;
 	class FSolverBodyContainer;
 
 	/**
