Chaos Overlap Caching on Dynamic Tree with Fulll update (Dirty one is WIP) + remove reset indices (constraints no longer valid)

diff --git a/Engine/Source/Runtime/Experimental/Chaos/Private/Chaos/Island/IslandManager.cpp b/Engine/Source/Runtime/Experimental/Chaos/Private/Chaos/Island/IslandManager.cpp
--- a/Engine/Source/Runtime/Experimental/Chaos/Private/Chaos/Island/IslandManager.cpp
+++ b/Engine/Source/Runtime/Experimental/Chaos/Private/Chaos/Island/IslandManager.cpp
@@ -312,7 +312,6 @@ void FPBDIslandManager::InitializeGraph(const TParticleView<FPBDRigidParticles>&
 	{
 		if(!IslandSolver->IsSleeping())
 		{
-			IslandSolver->ResetIndices();
 			IslandSolver->ClearConstraints();
 		}
 	}
diff --git a/Engine/Source/Runtime/Experimental/Chaos/Private/Chaos/PBDRigidsEvolution.cpp b/Engine/Source/Runtime/Experimental/Chaos/Private/Chaos/PBDRigidsEvolution.cpp
--- a/Engine/Source/Runtime/Experimental/Chaos/Private/Chaos/PBDRigidsEvolution.cpp
+++ b/Engine/Source/Runtime/Experimental/Chaos/Private/Chaos/PBDRigidsEvolution.cpp
@@ -35,6 +35,10 @@ namespace Chaos
 	CHAOS_API int32 AccelerationStructureUseDynamicTree = 0;
 	FAutoConsoleVariableRef CVarAccelerationStructureUseDynamicTree(TEXT("p.Chaos.AccelerationStructureUseDynamicTree"), AccelerationStructureUseDynamicTree, TEXT("Use a dynamic BVH tree structure for dynamic objects"));
 
+	/** Console variable to enable the caching of the overlapping leaves if the dynamic tree is enable */
+	CHAOS_API int32 GAccelerationStructureCacheOverlappingLeaves = 1;
+	FAutoConsoleVariableRef CVarAccelerationStructureCacheOverlappingLeaves(TEXT("p.Chaos.AccelerationStructureCacheOverlappingLeaves"), GAccelerationStructureCacheOverlappingLeaves, TEXT("Set to 1: Cache the overlapping leaves for faster overlap query, any other value will disable the feature"));
+
 	CHAOS_API int32 AccelerationStructureTimeSlicingMaxQueueSizeBeforeForce = 1000;
 	FAutoConsoleVariableRef CVarAccelerationStructureTimeSlicingMaxQueueSizeBeforeForce(TEXT("p.Chaos.AccelerationStructureTimeSlicingMaxQueueSizeBeforeForce"), AccelerationStructureTimeSlicingMaxQueueSizeBeforeForce, TEXT("If the update queue reaches this limit, time slicing will be disabled, and the acceleration structure will be built at once"));
 
@@ -349,7 +353,7 @@ namespace Chaos
 		bool IsTimeSlicingProgressing = false;
 		bool CreatingNewTimeSlicedStructures = bNeedsReset;
 
-		if (CreatingNewTimeSlicedStructures)
+		if (CreatingNewTimeSlicedStructures && AccelerationStructureCopy)
 		{
 			// Cannot assume all substructures that exist are in cache!!
 			// Make sure AccelerationStructureCopy is completely reset and does not have dirty substructures that will be ignored below.
@@ -421,7 +425,9 @@ namespace Chaos
 				SCOPE_CYCLE_COUNTER(STAT_AccelerationStructureTimeSliceCopy);
 				AccelerationSUBStructureCopy->ProgressCopyTimeSliced(*SubStructure, IsForceFullBuild ? -1 : AccelerationStructureTimeSlicingMaxBytesCopy);
 				if (!AccelerationSUBStructureCopy->IsAsyncTimeSlicingComplete())
+				{
 					IsTimeSlicingProgressing = true;
+				}
 			}
 			else
 			{
@@ -472,7 +478,6 @@ namespace Chaos
 				AccelerationStructure->AddSubstructure(MoveTemp(NewStruct), SpatialAccIdx.Bucket, SpatialAccIdx.InnerIdx);
 			}
 		}
-
 		// If it's not progressing then it is finished so we can perform the final copy if required
 		AccelerationStructure->SetAllAsyncTasksComplete(!IsTimeSlicingProgressing);
 
@@ -757,8 +762,13 @@ namespace Chaos
 				{
 					// Dynamic Acceleration structures were not built by the async task, so copy them where required
 					CopyUnBuiltDynamicAccelerationStructures(SpatialAccelerationCache, InternalAcceleration, AsyncInternalAcceleration, AsyncExternalAcceleration);
-				}
 
+					if(GAccelerationStructureCacheOverlappingLeaves)
+					{
+						// Caching of the overlapping leaves
+						InternalAcceleration->CacheOverlappingLeaves();
+					}
+				}
 				if (AccelerationStructureSplitStaticAndDynamic == 1)
 				{
 					// If the new InternalAcceleration structure have constituents with no changes, we can copy it to the  AsyncInternalAcceleration for reuse
diff --git a/Engine/Source/Runtime/Experimental/Chaos/Private/Chaos/TriangleMeshImplicitObject.cpp b/Engine/Source/Runtime/Experimental/Chaos/Private/Chaos/TriangleMeshImplicitObject.cpp
--- a/Engine/Source/Runtime/Experimental/Chaos/Private/Chaos/TriangleMeshImplicitObject.cpp
+++ b/Engine/Source/Runtime/Experimental/Chaos/Private/Chaos/TriangleMeshImplicitObject.cpp
@@ -126,6 +126,12 @@ struct FTriangleMeshRaycastVisitor
 	{
 		return nullptr;
 	}
+	
+	/** Return a pointer to the payload on which we are querying the acceleration structure */
+	const void* GetQueryPayload() const
+	{
+		return nullptr;
+	}
 
 	template <ERaycastType SQType>
 	bool Visit(int32 TriIdx, FQueryFastData& CurData)
@@ -850,6 +856,12 @@ struct FTriangleMeshSweepVisitor
 	const void* GetQueryData() const { return nullptr; }
 	const void* GetSimData() const { return nullptr; }
 
+	/** Return a pointer to the payload on which we are querying the acceleration structure */
+	const void* GetQueryPayload() const
+	{
+		return nullptr;
+	}
+
 	bool VisitOverlap(const TSpatialVisitorData<int32>& VisitData)
 	{
 		check(false);
diff --git a/Engine/Source/Runtime/Experimental/Chaos/Public/Chaos/AABBTree.h b/Engine/Source/Runtime/Experimental/Chaos/Public/Chaos/AABBTree.h
--- a/Engine/Source/Runtime/Experimental/Chaos/Public/Chaos/AABBTree.h
+++ b/Engine/Source/Runtime/Experimental/Chaos/Public/Chaos/AABBTree.h
@@ -249,6 +249,22 @@ struct TAABBTreeLeafArray : public TBoundsWrapperHelper<TPayloadType, T, bComput
 		this->ComputeBounds(Elems);
 	}
 
+	/** Check if the leaf is dirty (if one of the payload have been updated)
+	 * @return Dirty boolean that indicates if the leaf is dirty or not
+ 	 */
+	bool IsLeafDirty() const
+	{
+		return bDirtyLeaf;
+	}
+
+	/** Set thye dirty flag onto the leaf 
+	 * @param  bDirtyState Disrty flag to set 
+	 */
+	void SetDirtyState(const bool bDirtyState)
+	{
+		bDirtyLeaf = bDirtyState;
+	}
+
 	template <typename TSQVisitor, typename TQueryFastData>
 	FORCEINLINE_DEBUGGABLE bool RaycastFast(const TVec3<T>& Start, TQueryFastData& QueryFastData, TSQVisitor& Visitor, const TVec3<T>& Dir, const TVec3<T> InvDir, const bool bParallel[3]) const
 	{
@@ -326,6 +342,7 @@ struct TAABBTreeLeafArray : public TBoundsWrapperHelper<TPayloadType, T, bComput
 				break;
 			}
 		}
+		bDirtyLeaf = true;
 	}
 
 	void UpdateElement(const TPayloadType& Payload, const TAABB<T, 3>& NewBounds, bool bHasBounds)
@@ -342,17 +359,20 @@ struct TAABBTreeLeafArray : public TBoundsWrapperHelper<TPayloadType, T, bComput
 				break;
 			}
 		}
+		bDirtyLeaf = true;
 	}
 
 	void AddElement(const TPayloadBoundsElement<TPayloadType, T>& Element)
 	{
 		Elems.Add(Element);
 		this->ComputeBounds(Elems);
+		bDirtyLeaf = true;
 	}
 
 	void Reset()
 	{
 		Elems.Reset();
+		bDirtyLeaf = true;
 	}
 
 #if !UE_BUILD_SHIPPING
@@ -373,12 +393,28 @@ struct TAABBTreeLeafArray : public TBoundsWrapperHelper<TPayloadType, T, bComput
 	}
 #endif
 
+	/** Print leaf information (bounds) for debugging purpose*/
+	void PrintLeaf() const
+	{
+		int32 ElemIndex = 0;
+		for (const auto& Elem : Elems)
+		{
+			UE_LOG(LogChaos, Log, TEXT("Elem[%d] with bounds = %f %f %f | %f %f %f"), ElemIndex, 
+					Elem.Bounds.Min()[0], Elem.Bounds.Min()[1], Elem.Bounds.Min()[2], 
+					Elem.Bounds.Max()[0], Elem.Bounds.Max()[1], Elem.Bounds.Max()[2]);
+			++ElemIndex;
+		}
+	}
+
 	void Serialize(FChaosArchive& Ar)
 	{
 		Ar << Elems;
 	}
 
 	TArray<TPayloadBoundsElement<TPayloadType, T>> Elems;
+
+	/** Flag on the leaf to know if it has been updated */
+	bool bDirtyLeaf = true;
 };
 
 template <typename TPayloadType, bool bComputeBounds, typename T>
@@ -550,7 +586,10 @@ public:
 		TreeExpensiveStats.Reset();
 		GlobalPayloads.Reset();
 		PayloadToInfo.Reset();
-
+		
+		OverlappingLeaves.Reset();
+		OverlappingOffsets.Reset();
+		
 		NumProcessedThisSlice = 0;
 		WorkStack.Reset();
 		WorkPoolFreeList.Reset();
@@ -1568,7 +1607,6 @@ public:
 
 	virtual void PrepareCopyTimeSliced(const  ISpatialAcceleration<TPayloadType, T, 3>& InFrom) override
 	{
-
 		check(this != &InFrom);
 		check(InFrom.GetType() == ESpatialAcceleration::AABBTree);
 		const TAABBTree<TPayloadType, TLeafType, bMutable, T>& From = static_cast<const TAABBTree<TPayloadType, TLeafType, bMutable, T>&>(InFrom);
@@ -1606,10 +1644,12 @@ public:
 		DirtyElementsGridOverflow.Reserve(From.DirtyElementsGridOverflow.Num());
 		GlobalPayloads.Reserve(From.GlobalPayloads.Num());
 		PayloadToInfo.Reserve(From.PayloadToInfo.Num());
+		OverlappingLeaves.Reserve(From.OverlappingLeaves.Num());
+		OverlappingOffsets.Reserve(From.OverlappingOffsets.Num());
 
 		this->SetAsyncTimeSlicingComplete(false);
 	}
-
+	
 	virtual void ProgressCopyTimeSliced(const  ISpatialAcceleration<TPayloadType, T, 3>& InFrom, int MaximumBytesToCopy) override
 	{
 		check(this != &InFrom);
@@ -1648,6 +1688,14 @@ public:
 		{
 			return;
 		}
+		if (!ContinueTimeSliceCopy(From.OverlappingLeaves, OverlappingLeaves, SizeToCopyLeft))
+		{
+			return;
+		}
+		if (!ContinueTimeSliceCopy(From.OverlappingOffsets, OverlappingOffsets, SizeToCopyLeft))
+		{
+			return;
+		}
 
 		this->SetAsyncTimeSlicingComplete(true);
 	}
@@ -1745,6 +1793,87 @@ public:
 			ensure(bDynamicTree == false);
 		}
 	}
+	
+	/** Given a first node and a leaf index find the overlapping leaves and update the node stack 
+	 * @param FirstNode First node to be added to the stack
+	 * @param LeafIndex Leaf index for which we want to find the overlapping leaves
+	 * @param Nodestack Node stack that will be used to traverse the tree and find the overlapping leaves
+	 */
+	void FindOverlappingLeaf(const int32 FirstNode, const int32 LeafIndex, TArray<int32>& NodeStack) 
+	{
+		const TAABB<T,3>& LeafBounds = Leaves[LeafIndex].GetBounds();
+				
+		NodeStack.Reset();
+		NodeStack.Add(FirstNode);
+		
+		int32 NodeIndex = INDEX_NONE;
+		while (NodeStack.Num())
+		{
+			NodeIndex = NodeStack.Pop(false);
+			const FNode& Node = Nodes[NodeIndex];
+			
+			// If a leaf directly test the bounds
+			if (Node.bLeaf)
+			{
+				if (LeafBounds.Intersects(Leaves[Node.ChildrenNodes[0]].GetBounds()))
+				{
+					OverlappingLeaves.Add(Node.ChildrenNodes[0]);
+				}
+			}
+			else
+			{
+				// If not loop over all the children nodes to check if they intersect the bounds
+				for(int32 ChildIndex = 0; ChildIndex < 2; ++ChildIndex)
+				{
+					if(LeafBounds.Intersects(Node.ChildrenBounds[ChildIndex]) && Node.ChildrenNodes[ChildIndex] != INDEX_NONE)
+					{
+						NodeStack.Add(Node.ChildrenNodes[ChildIndex]);
+					}
+				}
+			}
+		}
+	}
+
+	/** Cache for each leaves all the overlapping leaves*/
+	virtual void CacheOverlappingLeaves() override
+	{
+		OverlappingOffsets.SetNum(Leaves.Num()+1, false);
+		OverlappingLeaves.Reset();
+
+		if(!bDynamicTree || (bDynamicTree && RootNode == INDEX_NONE)) return;
+
+		TArray<int32> NodeStack;
+		const int32 FirstNode = bDynamicTree ? RootNode : 0;
+		
+		for(int32 LeafIndex = 0, NumLeaves = Leaves.Num(); LeafIndex < NumLeaves; ++LeafIndex)
+		{ 
+			OverlappingOffsets[LeafIndex] = OverlappingLeaves.Num();
+			FindOverlappingLeaf(FirstNode, LeafIndex, NodeStack);
+		}
+		OverlappingOffsets.Last() = OverlappingLeaves.Num();
+		for(int32 LeafIndex = 0, NumLeaves = Leaves.Num(); LeafIndex < NumLeaves; ++LeafIndex)
+		{
+			Leaves[LeafIndex].SetDirtyState(false);
+		}
+	}
+
+	/** Print the overlapping leaves data structure */
+	void PrintOverlappingLeaves()
+	{
+		UE_LOG(LogChaos, Log, TEXT("Num Leaves = %d"), Leaves.Num());
+		for(int32 LeafIndex = 0, NumLeaves = Leaves.Num(); LeafIndex < NumLeaves; ++LeafIndex)
+		{
+			auto& Leaf = Leaves[LeafIndex];
+			UE_LOG(LogChaos, Log, TEXT("Overlapping Count[%d] = %d with bounds = %f %f %f | %f %f %f"), LeafIndex,
+				OverlappingOffsets[LeafIndex+1] - OverlappingOffsets[LeafIndex], Leaf.GetBounds().Min()[0],
+				Leaf.GetBounds().Min()[1], Leaf.GetBounds().Min()[2], Leaf.GetBounds().Max()[0], Leaf.GetBounds().Max()[1], Leaf.GetBounds().Max()[2]);
+			
+			for(int32 OverlappingIndex = OverlappingOffsets[LeafIndex]; OverlappingIndex < OverlappingOffsets[LeafIndex+1]; ++OverlappingIndex)
+			{
+				UE_LOG(LogChaos, Log, TEXT("Overlapping Leaf[%d] = %d"), LeafIndex, OverlappingLeaves[OverlappingIndex]);
+			}
+		}
+	}
 
 private:
 
@@ -1854,7 +1983,45 @@ private:
 		}
 		return true;
 	}
-	
+
+	/** Cached version of the overlap function 
+	 * @param QueryBounds Bounds we want to query the tree on
+	 * @param Visitor Object owning the payload and used in the overlap function to store the result
+	 * @param bOverlapResult Result of the overlap function that will be used in the overlap fast
+	 * @return Boolean to check if we have used or not the cached overlapping leaves (for static we don't have the cache yet)
+	 */
+	template <EAABBQueryType Query, typename SQVisitor>
+	bool OverlapCached(const FAABB3& QueryBounds, SQVisitor& Visitor, bool& bOverlapResult) const
+	{
+		bool bCouldUseCache = false;
+		bOverlapResult = true;
+		// Only the overlap queries could use the caching
+		if (Query == EAABBQueryType::Overlap && Visitor.GetQueryPayload())
+		{
+			// Grab the payload from the visitor (only available on physics thread for now) and retrieve the info
+			const TPayloadType& QueryPayload = *static_cast<const TPayloadType*>(Visitor.GetQueryPayload());
+			if( const FAABBTreePayloadInfo* QueryInfo =  PayloadToInfo.Find(QueryPayload))
+			{
+				const int32 LeafIndex = QueryInfo->LeafIdx;
+				if(LeafIndex != INDEX_NONE && LeafIndex < (OverlappingOffsets.Num()-1))
+				{
+					//Once we have the leaf index we can loop over the overlapping leaves
+					for(int32 OverlappingIndex = OverlappingOffsets[LeafIndex];
+								OverlappingIndex < OverlappingOffsets[LeafIndex+1]; ++OverlappingIndex)
+					{
+						const TLeafType& OverlappingLeaf = Leaves[OverlappingLeaves[OverlappingIndex]];
+						if (OverlappingLeaf.OverlapFast(QueryBounds, Visitor) == false)
+						{
+							bOverlapResult = false;
+							break;
+						}
+					}
+					bCouldUseCache = true;
+				}
+			}
+		}
+		return bCouldUseCache;
+	}
 
 	template <EAABBQueryType Query, typename TQueryFastData, typename SQVisitor>
 	bool QueryImp(const FVec3& RESTRICT Start, TQueryFastData& CurData, const FVec3& QueryHalfExtents, const FAABB3& QueryBounds, SQVisitor& Visitor, const FVec3& Dir, const FVec3& InvDir, const bool bParallel[3]) const
@@ -2014,6 +2181,17 @@ private:
 			FReal TOI;
 		};
 
+		// Caching is for now only available for dyanmic tree
+		if (bDynamicTree && !OverlappingLeaves.IsEmpty())
+		{
+			// For overlap query and dynamic tree we are using the cached overlapping leaves
+			bool bOverlapResult = true;
+			if (OverlapCached<Query, SQVisitor>(QueryBounds, Visitor, bOverlapResult))
+			{
+				return bOverlapResult;
+			}
+		}
+		
 		TArray<FNodeQueueEntry> NodeStack;
 		if (bDynamicTree)
 		{
@@ -2352,7 +2530,7 @@ private:
 
 		NumProcessedThisSlice += LastElem - StartElemIdx;
 	}
-
+	
 	void SplitNode()
 	{
 		const bool WeAreTimeslicing = (MaxNumToProcess > 0);
@@ -2542,6 +2720,13 @@ private:
 
 			const void* GetQueryData() const { return nullptr; }
 			const void* GetSimData() const { return nullptr; }
+
+			/** Return a pointer to the payload on which we are querying the acceleration structure */
+			const void* GetQueryPayload() const 
+			{ 
+				return nullptr; 
+			}
+
 			bool ShouldIgnore(const TSpatialVisitorData<TPayloadType>& Instance) const { return false; }
 			TArray<TPayloadType>& CollectedResults;
 		};
@@ -2659,6 +2844,8 @@ private:
 		, MaxNumToProcess(Other.MaxNumToProcess)
 		, NumProcessedThisSlice(Other.NumProcessedThisSlice)
 		, bShouldRebuild(Other.bShouldRebuild)
+		, OverlappingLeaves(Other.OverlappingLeaves)
+		, OverlappingOffsets(Other.OverlappingOffsets)
 	{
 
 		ensure(bDynamicTree == Other.IsTreeDynamic());
@@ -2755,6 +2942,12 @@ private:
 	TArray<FWorkSnapshot> WorkPool;
 
 	bool bShouldRebuild;  // Contract: this can only ever be cleared by calling the ClearShouldRebuild method
+
+	/** Flat array of overlapping leaves.  */
+	TArray<int32> OverlappingLeaves;
+
+	/** Offset for each leaf in the overlapping leaves (OverlappingOffsets[LeafIndex]/OverlappingOffsets[LeafIndex+1] will be start/end indices of the leaf index in the overallping leaves array )*/
+	TArray<int32> OverlappingOffsets;
 };
 
 template<typename TPayloadType, typename TLeafType, bool bMutable, typename T>
diff --git a/Engine/Source/Runtime/Experimental/Chaos/Public/Chaos/BoundingVolume.h b/Engine/Source/Runtime/Experimental/Chaos/Public/Chaos/BoundingVolume.h
--- a/Engine/Source/Runtime/Experimental/Chaos/Public/Chaos/BoundingVolume.h
+++ b/Engine/Source/Runtime/Experimental/Chaos/Public/Chaos/BoundingVolume.h
@@ -219,6 +219,7 @@ public:
 			}
 			const void* GetQueryData() const { return nullptr; }
 			const void* GetSimData() const { return nullptr; }
+			const void* GetQueryPayload() const { return nullptr; }
 			bool ShouldIgnore(const TSpatialVisitorData<TPayloadType>& Instance) const { return false; }
 			TArray<TPayloadType>& CollectedResults;
 		};
@@ -357,7 +358,21 @@ public:
 	{
 		OverlapImp<SQVisitor, bPruneDuplicates>(QueryBounds, Visitor);
 	}
-
+	
+	/** Check if the leaf is dirty (if one of the payload have been updated)
+	* @return Dirty boolean that indicates if the leaf is dirty or not
+	*/
+	bool IsLeafDirty() const
+	{
+		return false;
+	}
+	
+	/** Set thye dirty flag onto the leaf 
+	* @param  bDirtyState Disrty flag to set 
+	*/
+	void SetDirtyState(const bool bDirtyState)
+	{}
+	
 	virtual void Serialize(FChaosArchive& Ar) override
 	{
 		Ar.UsingCustomVersion(FExternalPhysicsCustomObjectVersion::GUID);
diff --git a/Engine/Source/Runtime/Experimental/Chaos/Public/Chaos/Collision/SpatialAccelerationBroadPhase.h b/Engine/Source/Runtime/Experimental/Chaos/Public/Chaos/Collision/SpatialAccelerationBroadPhase.h
--- a/Engine/Source/Runtime/Experimental/Chaos/Public/Chaos/Collision/SpatialAccelerationBroadPhase.h
+++ b/Engine/Source/Runtime/Experimental/Chaos/Public/Chaos/Collision/SpatialAccelerationBroadPhase.h
@@ -26,10 +26,11 @@ namespace Chaos
 	struct FSimOverlapVisitor
 	{
 		TArray<FAccelerationStructureHandle>& Intersections;
-		FSimOverlapVisitor(FUniqueIdx InParticleUniqueIdx, const FCollisionFilterData& InSimFilterData, TArray<FAccelerationStructureHandle>& InIntersections)
+		FSimOverlapVisitor(FGeometryParticleHandle* ParticleHandle, const FCollisionFilterData& InSimFilterData, TArray<FAccelerationStructureHandle>& InIntersections)
 			: Intersections(InIntersections)
 			, SimFilterData(InSimFilterData)
-			, ParticleUniqueIdx(InParticleUniqueIdx)
+			, ParticleUniqueIdx(ParticleHandle ? ParticleHandle->UniqueIdx() : FUniqueIdx(0))
+			, AccelerationHandle(ParticleHandle)
 		{
 		}
 
@@ -59,10 +60,18 @@ namespace Chaos
 		{
 			return Instance.Payload.UniqueIdx() == ParticleUniqueIdx;
 		}
+		/** Return a pointer to the payload on which we are querying the acceleration structure */
+		const void* GetQueryPayload() const
+		{
+			return &AccelerationHandle;
+		}
 
 	private:
 		FCollisionFilterData SimFilterData;
 		FUniqueIdx ParticleUniqueIdx; // unique id of the particle visiting, used to skip self intersection as early as possible
+
+		/** Handle to be stored to retrieve the payload on which we are querying the acceleration structure*/
+		FAccelerationStructureHandle AccelerationHandle;
 	};
 
 	/**
@@ -136,10 +145,10 @@ namespace Chaos
 		void ComputeParticlesOverlaps(ViewType& OverlapView, FReal Dt,
 			const SpatialAccelerationType& InSpatialAcceleration, FNarrowPhase& NarrowPhase)
 		{
-			OverlapView.ParallelFor([&](auto& Particle1,int32 ActiveIdxIdx)
-			{
-				FGenericParticleHandleImp GenericHandle(Particle1.Handle());
-				ProduceParticleOverlaps<bNeedsResim,bOnlyRigid>(Dt,GenericHandle, InSpatialAcceleration,NarrowPhase,ActiveIdxIdx);
+			 OverlapView.ParallelFor([&](auto& Particle1,int32 ActiveIdxIdx)
+			 {
+			 	FGenericParticleHandleImp GenericHandle(Particle1.Handle());
+			 	ProduceParticleOverlaps<bNeedsResim,bOnlyRigid>(Dt,GenericHandle, InSpatialAcceleration,NarrowPhase,ActiveIdxIdx);
 			},bDisableCollisionParallelFor);
 		}
 
@@ -208,7 +217,7 @@ namespace Chaos
 
 						{
 							PHYSICS_CSV_SCOPED_EXPENSIVE(PhysicsVerbose, DetectCollisions_BroadPhase);
-							FSimOverlapVisitor OverlapVisitor(Particle1.UniqueIdx(), ParticleSimData, PotentialIntersections);
+							FSimOverlapVisitor OverlapVisitor(Particle1.Handle() , ParticleSimData, PotentialIntersections);
 							InSpatialAcceleration.Overlap(Box1, OverlapVisitor);
 						}
 						
@@ -363,7 +372,7 @@ namespace Chaos
 					{
 						Swap(ParticleA, ParticleB);
 					}
-
+					
 					NarrowPhase.GenerateCollisions(Dt, ParticleA, ParticleB, false);
 				}
 
diff --git a/Engine/Source/Runtime/Experimental/Chaos/Public/Chaos/ISpatialAcceleration.h b/Engine/Source/Runtime/Experimental/Chaos/Public/Chaos/ISpatialAcceleration.h
--- a/Engine/Source/Runtime/Experimental/Chaos/Public/Chaos/ISpatialAcceleration.h
+++ b/Engine/Source/Runtime/Experimental/Chaos/Public/Chaos/ISpatialAcceleration.h
@@ -145,6 +145,9 @@ public:
 	virtual const void* GetSimData() const { return nullptr; }
 
 	virtual bool ShouldIgnore(const TSpatialVisitorData<TPayloadType>& Instance) const { return false; }
+
+	/** Return a pointer to the payload on which we are querying the acceleration structure */
+	virtual const void* GetQueryPayload() const { return nullptr; }
 };
 
 /**
@@ -269,6 +272,9 @@ public:
 	virtual void PrepareCopyTimeSliced(const  ISpatialAcceleration<TPayloadType, T, 3>& InFrom) { check(false); }
 	virtual void ProgressCopyTimeSliced(const  ISpatialAcceleration<TPayloadType, T, 3>& InFrom, int MaximumBytesToCopy) { check(false); }
 
+	/** Cache for each leaves all the overlapping leaves */
+	virtual void CacheOverlappingLeaves() {}
+
 	// IMPORTANT : (LWC) this API should be typed on Freal not T, as we want the query API to be using the highest precision while maintaining arbitrary internal precision for the acceleration structure ( based on T )
 	virtual TArray<TPayloadType> FindAllIntersections(const FAABB3& Box) const { check(false); return TArray<TPayloadType>(); }
 	virtual void Raycast(const FVec3& Start, const FVec3& Dir, const FReal Length, ISpatialVisitor<TPayloadType, FReal>& Visitor) const { check(false); }
@@ -431,6 +437,12 @@ public:
 		return Visitor.ShouldIgnore(Instance);
 	}
 
+	/** Return a pointer to the payload on which we are querying the acceleration structure */
+	FORCEINLINE const void* GetQueryPayload() const
+	{
+		return Visitor.GetQueryPayload();
+	}
+
 private:
 	ISpatialVisitor<TPayloadType, T>& Visitor;
 };
@@ -603,7 +615,7 @@ public:
 	}
 
 private:
-
+	
 	struct FEntry
 	{
 		TValue Value;
@@ -615,7 +627,7 @@ private:
 
 		}
 	};
-
+	
 	TArray<FEntry> Entries;
 
 #if CHAOS_SERIALIZE_OUT
diff --git a/Engine/Source/Runtime/Experimental/Chaos/Public/Chaos/SpatialAccelerationCollection.h b/Engine/Source/Runtime/Experimental/Chaos/Public/Chaos/SpatialAccelerationCollection.h
--- a/Engine/Source/Runtime/Experimental/Chaos/Public/Chaos/SpatialAccelerationCollection.h
+++ b/Engine/Source/Runtime/Experimental/Chaos/Public/Chaos/SpatialAccelerationCollection.h
@@ -315,6 +315,24 @@ struct TSpatialAccelerationCollectionHelper
 		}
 	}
 
+	/** Compute the overlapping leaves for each leaf */
+	static void CacheOverlappingLeaves(const Tuple& Types)
+	{
+		auto& Accelerations = GetAccelerationsPerType<TypeIdx>(Types).Objects;
+		for (auto& Accelerator : Accelerations)
+		{
+			if (Accelerator)
+			{
+				Accelerator->CacheOverlappingLeaves();
+			}
+		}
+		constexpr int NextType = TypeIdx + 1;
+		if (NextType < NumTypes)
+		{
+			TSpatialAccelerationCollectionHelper < NextType < NumTypes ? NextType : 0, NumTypes, Tuple, TPayloadType, T, d>::CacheOverlappingLeaves(Types);
+		}
+	}
+	
 	static uint16 FindTypeIdx(const Tuple& Types, SpatialAccelerationType Type)
 	{
 		using AccelType = typename std::remove_pointer<typename decltype(GetAccelerationsPerType<TypeIdx>(Types).Objects)::ElementType>::type;
@@ -395,6 +413,12 @@ public:
 		return Result;
 	}
 
+	/** Dispatch the compute the overlapping leaves for each collection helper */
+	virtual void CacheOverlappingLeaves() override
+	{
+		TSpatialAccelerationCollectionHelper<0, NumTypes, decltype(Types), TPayloadType, T, d>::CacheOverlappingLeaves(Types);
+	}
+
 	virtual TUniquePtr <ISpatialAcceleration<TPayloadType, T, d>> RemoveSubstructure(FSpatialAccelerationIdx Idx) override
 	{
 		TUniquePtr <ISpatialAcceleration<TPayloadType, T, d>> Removed;
diff --git a/Engine/Source/Runtime/PhysicsCore/Private/SQAccelerator.cpp b/Engine/Source/Runtime/PhysicsCore/Private/SQAccelerator.cpp
--- a/Engine/Source/Runtime/PhysicsCore/Private/SQAccelerator.cpp
+++ b/Engine/Source/Runtime/PhysicsCore/Private/SQAccelerator.cpp
@@ -186,6 +186,12 @@ struct TSQVisitor : public Chaos::ISpatialVisitor<TPayload, Chaos::FReal>
 		return nullptr;
 	}
 
+	/** Return a pointer to the payload on which we are querying the acceleration structure */
+	virtual const void* GetQueryPayload() const override
+	{
+		return nullptr;
+	}
+
 private:
 
 	enum class ESQType
@@ -466,6 +472,12 @@ struct TBPVisitor : public Chaos::ISpatialVisitor<TPayload, Chaos::FReal>
 	{
 		return nullptr;
 	}
+	
+	/** Return a pointer to the payload on which we are querying the acceleration structure */
+	virtual const void* GetQueryPayload() const override
+	{
+		return nullptr;
+	}
 
 private:
 
