Fixes for contact generation
- fix to not add manifold contact points with Phi not set
- fix the way bounds are computed to find triangle in scaled trimeshes

diff --git a/Engine/Source/Runtime/Experimental/Chaos/Private/Chaos/Collision/PBDCollisionConstraint.cpp b/Engine/Source/Runtime/Experimental/Chaos/Private/Chaos/Collision/PBDCollisionConstraint.cpp
--- a/Engine/Source/Runtime/Experimental/Chaos/Private/Chaos/Collision/PBDCollisionConstraint.cpp
+++ b/Engine/Source/Runtime/Experimental/Chaos/Private/Chaos/Collision/PBDCollisionConstraint.cpp
@@ -383,20 +383,23 @@ namespace Chaos
 
 	void FPBDCollisionConstraint::AddOneshotManifoldContact(const FContactPoint& ContactPoint)
 	{
-		if (ManifoldPoints.Num() == MaxManifoldPoints)
+		if (ContactPoint.IsSet())
 		{
-			return;
-		}
+			if (ManifoldPoints.Num() == MaxManifoldPoints)
+			{
+				return;
+			}
+			
+			int32 ManifoldPointIndex = AddManifoldPoint(ContactPoint);
 
-		int32 ManifoldPointIndex = AddManifoldPoint(ContactPoint);
+			// Copy currently active point
+			if (ManifoldPoints[ManifoldPointIndex].ContactPoint.Phi < Manifold.Phi)
+			{
+				SetActiveContactPoint(ManifoldPoints[ManifoldPointIndex].ContactPoint);
+			}
 
-		// Copy currently active point
-		if (ManifoldPoints[ManifoldPointIndex].ContactPoint.Phi < Manifold.Phi)
-		{
-			SetActiveContactPoint(ManifoldPoints[ManifoldPointIndex].ContactPoint);
+			bUseIncrementalManifold = false;
 		}
-
-		bUseIncrementalManifold = false;
 	}
 
 	void FPBDCollisionConstraint::AddIncrementalManifoldContact(const FContactPoint& ContactPoint)
diff --git a/Engine/Source/Runtime/Experimental/Chaos/Private/Chaos/TriangleMeshImplicitObject.cpp b/Engine/Source/Runtime/Experimental/Chaos/Private/Chaos/TriangleMeshImplicitObject.cpp
--- a/Engine/Source/Runtime/Experimental/Chaos/Private/Chaos/TriangleMeshImplicitObject.cpp
+++ b/Engine/Source/Runtime/Experimental/Chaos/Private/Chaos/TriangleMeshImplicitObject.cpp
@@ -358,6 +358,7 @@ bool FTriangleMeshImplicitObject::ContactManifoldImp(const GeomType& QueryGeom,
 	bool bResult = false;
 
 	const auto& WorldScaleGeom = ScaleGeomIntoWorldHelper(QueryGeom, TriMeshScale);
+	const FVec3 InvTriMeshScale = FVec3(FReal(1) / TriMeshScale.X, FReal(1) / TriMeshScale.Y, FReal(1) / TriMeshScale.Z);
 
 	TRigidTransform<FReal, 3> WorldScaleQueryTM;
 	ScaleTransformHelper(TriMeshScale, QueryTM, WorldScaleQueryTM);
@@ -406,12 +407,11 @@ bool FTriangleMeshImplicitObject::ContactManifoldImp(const GeomType& QueryGeom,
 		FVec3 LocalContactLocation, LocalContactNormal;
 
 		// NOTE: BVH test is done in tri-mesh local space (whereas collision detection is done in world space becaused you can't non-uniformly scale all shapes)
-		// We must apply the local-space scaled thickness after transforming the world bounds into local space
-		const FVec3 LocalThickness = FVec3(WorldThickness / TriMeshScale.X, WorldThickness / TriMeshScale.Y, WorldThickness / TriMeshScale.Z);
-
-		FAABB3 QueryBounds = QueryGeom.BoundingBox();
+		FAABB3 QueryBounds = WorldScaleGeom.BoundingBox();
 		QueryBounds = QueryBounds.TransformedAABB(QueryTM);
-		QueryBounds.ThickenSymmetrically(LocalThickness);
+		QueryBounds.ThickenSymmetrically(FVec3(WorldThickness));
+		QueryBounds.LocalScale(InvTriMeshScale);
+
 		const TArray<int32> PotentialIntersections = BVH.FindAllIntersections(QueryBounds);
 
 		// MakeTriangleHelper get rid of the scale wrapper if necessary as MakeTriangle will try to infer properties from it 
@@ -467,6 +467,7 @@ bool FTriangleMeshImplicitObject::GJKContactPointImp(const QueryGeomType& QueryG
 	bool bResult = false;
 
 	const auto& WorldScaleGeom = ScaleGeomIntoWorldHelper(QueryGeom, TriMeshScale);
+	const FVec3 InvTriMeshScale = FVec3(FReal(1) / TriMeshScale.X, FReal(1) / TriMeshScale.Y, FReal(1) / TriMeshScale.Z);
 
 	TRigidTransform<FReal, 3> WorldScaleQueryTM;
 	ScaleTransformHelper(TriMeshScale, QueryTM, WorldScaleQueryTM);
@@ -488,7 +489,6 @@ bool FTriangleMeshImplicitObject::GJKContactPointImp(const QueryGeomType& QueryG
 			LocalContactNormal = LambdaNormal;
 			LocalContactPhi = -LambdaPenetration;
 		}
-
 		return GJKValidResult;
 	};
 
@@ -499,12 +499,11 @@ bool FTriangleMeshImplicitObject::GJKContactPointImp(const QueryGeomType& QueryG
 		FVec3 LocalContactLocation, LocalContactNormal;
 
 		// NOTE: BVH test is done in tri-mesh local space (whereas collision detection is done in world space becaused you can't non-uniformly scale all shapes)
-		// We must apply the local-space scaled thickness after transforming the world bounds into local space
-		const FVec3 LocalThickness = FVec3(WorldThickness / TriMeshScale.X, WorldThickness / TriMeshScale.Y, WorldThickness / TriMeshScale.Z);
-
-		FAABB3 QueryBounds = QueryGeom.BoundingBox();
+		FAABB3 QueryBounds = WorldScaleGeom.BoundingBox();
 		QueryBounds = QueryBounds.TransformedAABB(QueryTM);
-		QueryBounds.ThickenSymmetrically(LocalThickness);
+		QueryBounds.ThickenSymmetrically(FVec3(WorldThickness));
+		QueryBounds.LocalScale(InvTriMeshScale);
+
 		const TArray<int32> PotentialIntersections = BVH.FindAllIntersections(QueryBounds);
 
 		for (int32 TriIdx : PotentialIntersections)
@@ -523,7 +522,6 @@ bool FTriangleMeshImplicitObject::GJKContactPointImp(const QueryGeomType& QueryG
 			}
 
 		}
-
 		return OutContactPhi < WorldThickness;
 	};
 
diff --git a/Engine/Source/Runtime/Experimental/Chaos/Public/Chaos/AABB.h b/Engine/Source/Runtime/Experimental/Chaos/Public/Chaos/AABB.h
--- a/Engine/Source/Runtime/Experimental/Chaos/Public/Chaos/AABB.h
+++ b/Engine/Source/Runtime/Experimental/Chaos/Public/Chaos/AABB.h
@@ -459,6 +459,10 @@ namespace Chaos
 			}
 		}
 
+		/**
+		* Scale the AABB relative to the origin
+		* IMPORTANT : this does not support negative scale
+		*/
 		FORCEINLINE TAABB<T, d>& Scale(const TVector<T, d>& InScale)
 		{
 			MMin *= InScale;
@@ -466,6 +470,19 @@ namespace Chaos
 			return *this;
 		}
 
+		/**
+		* Scale the AABB relative to its center
+		* IMPORTANT : this does not support negative scale 
+		*/
+		FORCEINLINE TAABB<T, d>& LocalScale(const TVector<T, d>& InScale)
+		{
+			const TVector<T, d> BoxCenter = Center();
+			const TVector<T, d> ScaledHalfExtents = Extents() * InScale * FReal(0.5);
+			MMin = BoxCenter - ScaledHalfExtents;
+			MMax = BoxCenter + ScaledHalfExtents;
+			return *this;
+		}
+
 		FORCEINLINE const TVector<T, d>& Min() const { return MMin; }
 		FORCEINLINE const TVector<T, d>& Max() const { return MMax; }
 
