Chaos - CullDistance increases with object size so that it is always bigger than the restore-collision movement threshold

diff --git a/Engine/Source/Runtime/Experimental/Chaos/Private/Chaos/Collision/ParticlePairMidPhase.cpp b/Engine/Source/Runtime/Experimental/Chaos/Private/Chaos/Collision/ParticlePairMidPhase.cpp
--- a/Engine/Source/Runtime/Experimental/Chaos/Private/Chaos/Collision/ParticlePairMidPhase.cpp
+++ b/Engine/Source/Runtime/Experimental/Chaos/Private/Chaos/Collision/ParticlePairMidPhase.cpp
@@ -31,6 +31,11 @@ namespace Chaos
 		bool bChaos_Collision_EnableManifoldUpdate = true;
 		FAutoConsoleVariableRef CVarChaos_Collision_EnableManifoldUpdate(TEXT("p.Chaos.Collision.EnableManifoldUpdate"), bChaos_Collision_EnableManifoldUpdate, TEXT(""));
 
+		Chaos::FRealSingle Chaos_Collision_CullDistanceScaleInverseSize = 0.01f;	// 100cm
+		Chaos::FRealSingle Chaos_Collision_MinCullDistanceScale = 1.0f;
+		FAutoConsoleVariableRef CVarChaos_Collision_CullDistanceReferenceSize(TEXT("p.Chaos.Collision.CullDistanceReferenceSize"), Chaos_Collision_CullDistanceScaleInverseSize, TEXT(""));
+		FAutoConsoleVariableRef CVarChaos_Collision_MinCullDistanecScale(TEXT("p.Chaos.Collision.MinCullDistanceScale"), Chaos_Collision_MinCullDistanceScale, TEXT(""));
+
 	}
 	using namespace CVars;
 
@@ -656,6 +661,7 @@ namespace Chaos
 		, RestoreParticleP1(FVec3(0))
 		, RestoreParticleQ0(FRotation3::FromIdentity())
 		, RestoreParticleQ1(FRotation3::FromIdentity())
+		, CullDistanceScale(1)
 	{
 		if (ensure(Particle0 != Particle1))
 		{
@@ -702,7 +708,7 @@ namespace Chaos
 
 		BuildDetectors();
 
-		InitRestoreThresholds();
+		InitThresholds();
 
 		bIsInitialized = true;
 	}
@@ -789,21 +795,33 @@ namespace Chaos
 		return false;
 	}
 
-	void FParticlePairMidPhase::InitRestoreThresholds()
+	void FParticlePairMidPhase::InitThresholds()
 	{
 		// @todo(chaos): improve this threshold calculation for thin objects? Dynamic thin objects have bigger problems so maybe we don't care
 		// @todo(chaos): Spheres and capsules need smaller position tolerance - the restore test doesn't work well with rolling
 		bool bIsDynamic0 = FConstGenericParticleHandle(Particle0)->IsDynamic();
 		bool bIsDynamic1 = FConstGenericParticleHandle(Particle1)->IsDynamic();
-		const FReal BoundsSize0 = bIsDynamic0 ? Particle0->LocalBounds().Extents().GetMin() : TNumericLimits<FReal>::Max();
-		const FReal BoundsSize1 = bIsDynamic1 ? Particle1->LocalBounds().Extents().GetMin() : TNumericLimits<FReal>::Max();
-		const FReal ThresholdSize = FMath::Min(BoundsSize0, BoundsSize1);
+		const FReal MinBoundsSize0 = bIsDynamic0 ? Particle0->LocalBounds().Extents().GetMin() : TNumericLimits<FReal>::Max();
+		const FReal MinBoundsSize1 = bIsDynamic1 ? Particle1->LocalBounds().Extents().GetMin() : TNumericLimits<FReal>::Max();
+		const FReal ThresholdSize = FMath::Min(MinBoundsSize0, MinBoundsSize1);
 
 		RestoreThresholdZeroContacts.PositionThreshold = Chaos_Collision_RestoreTolerance_NoContact_Position * ThresholdSize;
 		RestoreThresholdZeroContacts.RotationThreshold = Chaos_Collision_RestoreTolerance_NoContact_Rotation;
 
 		RestoreThreshold.PositionThreshold = Chaos_Collision_RestoreTolerance_Contact_Position * ThresholdSize;
 		RestoreThreshold.RotationThreshold = Chaos_Collision_RestoreTolerance_Contact_Rotation;
+	
+		// NOTE: If CullDistance ends up smaller than the thresholds used to restore collisions, we can end up missing
+		// collisions as the objects move if we restore a "zero contact" manifold after movement greater than the cull distance. 
+		// Currently this should not happen, but it is not explicitly ensured by the way the thresholds and CullDistanceScale are calculated.
+		// @todo(chaos): Add a way to enforce a CullDistance big enough to support the reuse thresholds
+		const FReal CullDistanceReferenceSizeInv = FReal(Chaos_Collision_CullDistanceScaleInverseSize);
+		const FReal MinCullDistanceScale = FReal(Chaos_Collision_MinCullDistanceScale);
+		const FReal MaxBoundsSize0 = bIsDynamic0 ? Particle0->LocalBounds().Extents().GetMax() : FReal(0);
+		const FReal MaxBoundsSize1 = bIsDynamic1 ? Particle1->LocalBounds().Extents().GetMax() : FReal(0);
+		const FReal CullDistanceScale0 = MaxBoundsSize0 * CullDistanceReferenceSizeInv;
+		const FReal CullDistanceScale1 = MaxBoundsSize1 * CullDistanceReferenceSizeInv;
+		CullDistanceScale = FMath::Max3(CullDistanceScale0, CullDistanceScale1, MinCullDistanceScale);
 	}
 
 	bool FParticlePairMidPhase::ShouldRestoreConstraints(const FReal Dt)
@@ -868,7 +886,7 @@ namespace Chaos
 	}
 
 	void FParticlePairMidPhase::GenerateCollisions(
-		const FReal CullDistance,
+		const FReal InCullDistance,
 		const FReal Dt,
 		FCollisionContext& Context)
 	{
@@ -877,6 +895,9 @@ namespace Chaos
 			return;
 		}
 
+		// CullDistance is scaled by the size of the dynamic objects.
+		const FReal CullDistance = InCullDistance * CullDistanceScale;
+
 		// Enable CCD?
 		const bool bUseCCD = bIsCCD && ShouldEnableCCD(Dt);
 
diff --git a/Engine/Source/Runtime/Experimental/Chaos/Public/Chaos/Collision/ParticlePairMidPhase.h b/Engine/Source/Runtime/Experimental/Chaos/Public/Chaos/Collision/ParticlePairMidPhase.h
--- a/Engine/Source/Runtime/Experimental/Chaos/Public/Chaos/Collision/ParticlePairMidPhase.h
+++ b/Engine/Source/Runtime/Experimental/Chaos/Public/Chaos/Collision/ParticlePairMidPhase.h
@@ -348,7 +348,7 @@ namespace Chaos
 		*/
 		bool ShouldEnableCCD(const FReal Dt);
 
-		void InitRestoreThresholds();
+		void InitThresholds();
 
 		/**
 		 * @brief Whether we should reuse the constraint as-is and skip the narrow phase
@@ -385,5 +385,8 @@ namespace Chaos
 		FVec3 RestoreParticleP1;
 		FRotation3 RestoreParticleQ0;
 		FRotation3 RestoreParticleQ1;
+
+		// A number based on the size of the dynamic objects used to scale cull distance
+		FReal CullDistanceScale;
 	};
 }
