Chaos : fix ghost contacts generated by stale GJKWarmStartData because we reuse constraints in Trimesh manifold generation

diff --git a/Engine/Source/Runtime/Experimental/Chaos/Private/Chaos/TriangleMeshImplicitObject.cpp b/Engine/Source/Runtime/Experimental/Chaos/Private/Chaos/TriangleMeshImplicitObject.cpp
--- a/Engine/Source/Runtime/Experimental/Chaos/Private/Chaos/TriangleMeshImplicitObject.cpp
+++ b/Engine/Source/Runtime/Experimental/Chaos/Private/Chaos/TriangleMeshImplicitObject.cpp
@@ -398,6 +398,10 @@ bool FTriangleMeshImplicitObject::ContactManifoldImp(const GeomType& QueryGeom,
 		FPBDCollisionConstraint& Constraint)
 	{
 		FTriangle TriangleConvex(A, B, C);
+		// make sure the constraint does not contain any stale data  ( it is shared between trinagles )
+		// @todo(chaos) : we should eventually not use a constraint here and just get a list of contact points
+		Constraint.ResetManifold();
+		Constraint.GetGJKWarmStartData().Reset();
 		Collisions::ConstructConvexConvexOneShotManifold(WorldScaleGeom, WorldScaleQueryTM, TriangleConvex, FRigidTransform3::Identity, 0, Constraint);
 	};
 
@@ -421,7 +425,6 @@ bool FTriangleMeshImplicitObject::ContactManifoldImp(const GeomType& QueryGeom,
 		{
 			FVec3 A, B, C;
 			TransformVertsHelper(TriMeshScale, TriIdx, MParticles, Elements, A, B, C);
-			Constraint.ResetManifold();
 			OverlapTriangle(A, B, C, Constraint);
 			for(FManifoldPoint& ManifoldPoint : Constraint.GetManifoldPoints())
 			{
