Chaos : fix scaled trimesh overlap issues where query bounds would be incorrect in the space of the trimesh and collect wrong set of trinagle sometimes

diff --git a/Engine/Plugins/Experimental/ChaosVehiclesPlugin/Source/ChaosVehicles/Private/ChaosWheeledVehicleMovementComponent.cpp b/Engine/Plugins/Experimental/ChaosVehiclesPlugin/Source/ChaosVehicles/Private/ChaosWheeledVehicleMovementComponent.cpp
--- a/Engine/Plugins/Experimental/ChaosVehiclesPlugin/Source/ChaosVehicles/Private/ChaosWheeledVehicleMovementComponent.cpp
+++ b/Engine/Plugins/Experimental/ChaosVehiclesPlugin/Source/ChaosVehicles/Private/ChaosWheeledVehicleMovementComponent.cpp
@@ -360,6 +360,16 @@ void UChaosWheeledVehicleSimulation::PerformSuspensionTraces(const TArray<FSuspe
 		if (GWheeledVehicleDebugParams.ShowBatchQueryExtents)
 		{
 			Chaos::FDebugDrawQueue::GetInstance().DrawDebugBox(QueryBox.GetCenter(), QueryBox.GetExtent(), FQuat::Identity, FColor::Yellow, false, -1.0f, 0, 2.0f);
+
+			// draw all corresponding results bounding boxes
+			for (FOverlapResult OverlapResult : OverlapResults)
+			{
+				if (OverlapResult.bBlockingHit)
+				{
+					const FBoxSphereBounds Bounds = OverlapResult.Component->CalcBounds(OverlapResult.Component->GetComponentTransform());
+					Chaos::FDebugDrawQueue::GetInstance().DrawDebugBox(Bounds.GetBox().GetCenter(), Bounds.GetBox().GetExtent(), FQuat::Identity, FColor::Purple, false, -1.0f, 0, 2.0f);
+				}
+			}
 		}
 	#endif
 
diff --git a/Engine/Source/Runtime/Experimental/Chaos/Private/Chaos/TriangleMeshImplicitObject.cpp b/Engine/Source/Runtime/Experimental/Chaos/Private/Chaos/TriangleMeshImplicitObject.cpp
--- a/Engine/Source/Runtime/Experimental/Chaos/Private/Chaos/TriangleMeshImplicitObject.cpp
+++ b/Engine/Source/Runtime/Experimental/Chaos/Private/Chaos/TriangleMeshImplicitObject.cpp
@@ -713,9 +713,17 @@ template <typename QueryGeomType>
 bool FTriangleMeshImplicitObject::OverlapGeomImp(const QueryGeomType& QueryGeom, const FRigidTransform3& QueryTM, const FReal Thickness, FMTDInfo* OutMTD, FVec3 TriMeshScale) const
 {
 	bool bResult = false;
-	FAABB3 QueryBounds = QueryGeom.BoundingBox();
-	QueryBounds.Thicken(Thickness);
+
+	const auto& WorldScaleQueryGeom = ScaleGeomIntoWorldHelper(QueryGeom, TriMeshScale);
+
+	const FVec3 InvTriMeshScale = FVec3(FReal(1) / TriMeshScale.X, FReal(1) / TriMeshScale.Y, FReal(1) / TriMeshScale.Z);
+
+	// NOTE: BVH test is done in tri-mesh local space (whereas collision detection is done in world space becaused you can't non-uniformly scale all shapes)
+	FAABB3 QueryBounds = WorldScaleQueryGeom.BoundingBox();
 	QueryBounds = QueryBounds.TransformedAABB(QueryTM);
+	QueryBounds.Thicken(Thickness);
+	QueryBounds.LocalScale(InvTriMeshScale);
+
 	const TArray<int32> PotentialIntersections = BVH.FindAllIntersections(QueryBounds);
 
 	if (OutMTD)
@@ -724,8 +732,6 @@ bool FTriangleMeshImplicitObject::OverlapGeomImp(const QueryGeomType& QueryGeom,
 		OutMTD->Penetration = TNumericLimits<FReal>::Lowest();
 	}
 
-	const auto& WorldScaleQueryGeom = ScaleGeomIntoWorldHelper(QueryGeom, TriMeshScale);
-
 	TRigidTransform<FReal, 3> WorldScaleQueryTM;
 	ScaleTransformHelper(TriMeshScale, QueryTM, WorldScaleQueryTM);
 
