Chaos - Remove CoM transforms from collisions

diff --git a/Engine/Source/Programs/HeadlessChaos/Private/HeadlessChaosTestContactModification.cpp b/Engine/Source/Programs/HeadlessChaos/Private/HeadlessChaosTestContactModification.cpp
--- a/Engine/Source/Programs/HeadlessChaos/Private/HeadlessChaosTestContactModification.cpp
+++ b/Engine/Source/Programs/HeadlessChaos/Private/HeadlessChaosTestContactModification.cpp
@@ -432,118 +432,6 @@ namespace ChaosTest
 		Module->DestroySolver(Solver);
 	}
 
-	GTEST_TEST(AllTraits, ContactModification_ModifyLocationCoMSpace)
-	{
-		FChaosSolversModule* Module = FChaosSolversModule::GetModule();
-		auto* Solver = Module->CreateSolver(nullptr, /*AsyncDt=*/-1);
-		InitSolverSettings(Solver);
-		Solver->SetThreadingMode_External(EThreadingModeTemp::SingleThread);
-
-		// create a static floor and two boxes colliding on edge of floor each with center of mass over the edge.
-		// One cube should rotate off side, the other has contact point locations moved under center of mass,
-		// so cube does not rotate and fall, instead remains on floor.
-
-		// simulated cube falling onto floor with center of mass over hanging past edge, should fall under floor,
-		FSingleParticlePhysicsProxy* FallingCubeProxy = FSingleParticlePhysicsProxy::Create(Chaos::FPBDRigidParticle::CreateParticle());
-		auto& FallingCubeParticle = FallingCubeProxy->GetGameThreadAPI();
-		auto FallingCubeGeom = TSharedPtr<FImplicitObject, ESPMode::ThreadSafe>(new TBox<FReal, 3>(FVec3(-100), FVec3(100)));
-		FallingCubeParticle.SetGeometry(FallingCubeGeom);
-		Solver->RegisterObject(FallingCubeProxy);
-		FallingCubeParticle.SetGravityEnabled(true);
-		FallingCubeParticle.SetX(FVec3(550, 0, 110));
-		SetCubeInertiaTensor(FallingCubeParticle, /*Dimension=*/200, /*Mass=*/1);
-		ChaosTest::SetParticleSimDataToCollide({ FallingCubeProxy->GetParticle_LowLevel() });
-
-		// cube with CoM hanging past edge of floor, contact mod moves contact under CoM so it will not tip off edge.
-		FSingleParticlePhysicsProxy* ModifiedCubeProxy = FSingleParticlePhysicsProxy::Create(Chaos::FPBDRigidParticle::CreateParticle());
-		auto& ModifiedCubeParticle = ModifiedCubeProxy->GetGameThreadAPI();
-		auto ModifiedCubeGeom = TSharedPtr<FImplicitObject, ESPMode::ThreadSafe>(new TBox<FReal, 3>(FVec3(-100), FVec3(100)));
-		ModifiedCubeParticle.SetGeometry(ModifiedCubeGeom);
-		Solver->RegisterObject(ModifiedCubeProxy);
-		ModifiedCubeParticle.SetGravityEnabled(true);
-		ModifiedCubeParticle.SetX(FVec3(-550, 0, 110));
-		SetCubeInertiaTensor(ModifiedCubeParticle, /*Dimension=*/200, /*Mass=*/1);
-		ChaosTest::SetParticleSimDataToCollide({ ModifiedCubeProxy->GetParticle_LowLevel() });
-
-
-		// static floor at origin, X/Y spanning [-500, 500] and Z spanning [-100,0]
-		FSingleParticlePhysicsProxy* FloorProxy = FSingleParticlePhysicsProxy::Create(Chaos::FGeometryParticle::CreateParticle());
-		auto& FloorParticle = FloorProxy->GetGameThreadAPI();
-		auto FloorGeom = TSharedPtr<FImplicitObject, ESPMode::ThreadSafe>(new TBox<FReal, 3>(FVec3(-500, -500, -100), FVec3(500, 500, 0)));
-		FloorParticle.SetGeometry(FloorGeom);
-		Solver->RegisterObject(FloorProxy);
-		FloorParticle.SetX(FVec3(0, 0, 0));
-		ChaosTest::SetParticleSimDataToCollide({ FloorProxy->GetParticle_LowLevel() });
-
-
-		// Save Unique indices of floor and modified cube to disable in contact mod.
-		TVec2<FUniqueIdx> UniqueIndices({ ModifiedCubeParticle.UniqueIdx(), FloorParticle.UniqueIdx() });
-
-		FContactModificationTestCallback* Callback = Solver->CreateAndRegisterSimCallbackObject_External<FContactModificationTestCallback>(true);
-		Callback->TestLambda = [UniqueIndices](Chaos::FCollisionContactModifier& Modifier)
-		{
-			for (FContactPairModifier& PairModifier : Modifier)
-			{
-				TVec2<FGeometryParticleHandle*> Particles = PairModifier.GetParticlePair();
-				FUniqueIdx Idx0 = Particles[0]->UniqueIdx();
-				FUniqueIdx Idx1 = Particles[1]->UniqueIdx();
-
-				// If unique indices match disable the pair.
-				if ((UniqueIndices[0] == Idx0 && UniqueIndices[1] == Idx1) ||
-					(UniqueIndices[0] == Idx1 && UniqueIndices[1] == Idx0))
-				{
-					int32 NumContacts = PairModifier.GetNumContacts();
-					for (int32 PointIdx = 0; PointIdx < NumContacts; ++PointIdx)
-					{
-						// Move contact locations below center of mass
-						FVec3 CoMPos0;
-						FVec3 CoMPos1;
-						PairModifier.GetCoMContactLocations(PointIdx, CoMPos0, CoMPos1);
-
-						int32 DynamicIdx = (UniqueIndices[0] == Idx0) ? 0 : 1;
-						if (DynamicIdx == 0)
-						{
-							// Move point0 under center of mass and move second point under CoM but keep the same distance between bodiesz
-							FVec3 PointUnderCoM0(0, 0, CoMPos0.Z);
-							FVec3 Delta = PointUnderCoM0 - CoMPos0;
-							FVec3 PointUnderCoM1(CoMPos1 + Delta);
-							PairModifier.ModifyCoMContactLocations(PointUnderCoM0, PointUnderCoM1, PointIdx);
-						}
-						else
-						{
-							FVec3 PointUnderCoM1(0, 0, CoMPos1.Z);
-							FVec3 Delta = PointUnderCoM1 - CoMPos1;
-							FVec3 PointUnderCoM0(CoMPos0 + Delta);
-							PairModifier.ModifyCoMContactLocations(PointUnderCoM0, PointUnderCoM1, PointIdx);
-						}
-					}
-				}
-			}
-		};
-
-
-		const float Dt = 0.1f;
-		const int32 Steps = 10;
-		for (int Step = 0; Step < Steps; ++Step)
-		{
-			Solver->AdvanceAndDispatch_External(Dt);
-			Solver->UpdateGameThreadStructures();
-		}
-
-		// Modified contact points to be below CoM, cube should not tip off edge of floor, but rest on it instead.
-		EXPECT_NEAR(ModifiedCubeParticle.X().Z, 100.f, KINDA_SMALL_NUMBER);
-
-		// Expected to tip off edge as CoM hangs off of floor.
-		EXPECT_LT(FallingCubeParticle.X().Z, 0.f);
-
-		Solver->UnregisterAndFreeSimCallbackObject_External(Callback);
-		Solver->UnregisterObject(FallingCubeProxy);
-		Solver->UnregisterObject(ModifiedCubeProxy);
-		Solver->UnregisterObject(FloorProxy);
-		Module->DestroySolver(Solver);
-	}
-
-
 	GTEST_TEST(AllTraits, ContactModification_ModifyRestitution)
 	{
 		FChaosSolversModule* Module = FChaosSolversModule::GetModule();
@@ -704,7 +592,8 @@ namespace ChaosTest
 				if ((UniqueIndices[0] == Idx0 && UniqueIndices[1] == Idx1) ||
 					(UniqueIndices[0] == Idx1 && UniqueIndices[1] == Idx0))
 				{
-					PairModifier.ModifyFriction(1);
+					PairModifier.ModifyDynamicFriction(1);
+					PairModifier.ModifyStaticFriction(1);
 				}
 			}
 		};
@@ -964,6 +853,8 @@ namespace ChaosTest
 				TVec2<FGeometryParticleHandle*> Particles = PairModifier.GetParticlePair();
 				int32 DynamicParticleIdx = (Particles[0]->ObjectState() == EObjectStateType::Dynamic) ? 0 : 1;
 
+				PairModifier.ModifyDynamicFriction(0);
+				PairModifier.ModifyStaticFriction(0);
 				PairModifier.ModifyParticleRotation(ModificationRotation, /*bMaintainVelocity=*/true, DynamicParticleIdx);
 			}
 		};
diff --git a/Engine/Source/Runtime/Experimental/Chaos/Private/Chaos/CCDUtilities.cpp b/Engine/Source/Runtime/Experimental/Chaos/Private/Chaos/CCDUtilities.cpp
--- a/Engine/Source/Runtime/Experimental/Chaos/Private/Chaos/CCDUtilities.cpp
+++ b/Engine/Source/Runtime/Experimental/Chaos/Private/Chaos/CCDUtilities.cpp
@@ -438,6 +438,19 @@ namespace Chaos
 				}
 			}
 		}
+
+		// We need to update the world-space contact points at the final locations
+		for (int32 i = ConstraintStart; i < ConstraintEnd; i++)
+		{
+			FPBDCollisionConstraint* Constraint = SortedCCDConstraints[i]->SweptConstraint;
+			FRigidTransform3 ShapeWorldTransform0 = Constraint->GetShapeWorldTransform0();
+			FRigidTransform3 ShapeWorldTransform1 = Constraint->GetShapeWorldTransform1();
+			ShapeWorldTransform0.SetTranslation(FConstGenericParticleHandle(Constraint->GetParticle0())->P());
+			ShapeWorldTransform1.SetTranslation(FConstGenericParticleHandle(Constraint->GetParticle1())->P());
+
+			Constraint->SetShapeWorldTransforms(ShapeWorldTransform0, ShapeWorldTransform1);
+			Constraint->UpdateManifoldContacts();
+		}
 	}
 
 
diff --git a/Engine/Source/Runtime/Experimental/Chaos/Private/Chaos/ChaosDebugDraw.cpp b/Engine/Source/Runtime/Experimental/Chaos/Private/Chaos/ChaosDebugDraw.cpp
--- a/Engine/Source/Runtime/Experimental/Chaos/Private/Chaos/ChaosDebugDraw.cpp
+++ b/Engine/Source/Runtime/Experimental/Chaos/Private/Chaos/ChaosDebugDraw.cpp
@@ -856,12 +856,9 @@ namespace Chaos
 
 						const int32 ContactPlaneOwner = 1;
 						const int32 ContactPointOwner = 1 - ContactPlaneOwner;
-						const FRigidTransform3& PlaneTransform = (ContactPlaneOwner == 0) ? WorldCoMTransform0 : WorldCoMTransform1;
-						const FRigidTransform3& PointTransform = (ContactPlaneOwner == 0) ? WorldCoMTransform1 : WorldCoMTransform0;
-						const FConstGenericParticleHandle PlaneParticle = (ContactPlaneOwner == 0) ? Particle0 : Particle1;
 						const FVec3 PlaneNormal = ManifoldPoint.ContactPoint.Normal;
-						const FVec3 PointLocation = PointTransform.TransformPosition(ManifoldPoint.CoMContactPoints[ContactPointOwner]);
-						const FVec3 PlaneLocation = PlaneTransform.TransformPosition(ManifoldPoint.CoMContactPoints[ContactPlaneOwner]);
+						const FVec3 PointLocation = ManifoldPoint.WorldContactPoints[ContactPointOwner];
+						const FVec3 PlaneLocation = ManifoldPoint.WorldContactPoints[ContactPlaneOwner];
 						const FVec3 PointPlaneLocation = PointLocation - FVec3::DotProduct(PointLocation - PlaneLocation, PlaneNormal) * PlaneNormal;
 
 						// Dynamic friction, restitution = red
@@ -906,19 +903,6 @@ namespace Chaos
 						FMatrix Axes = FRotationMatrix::MakeFromX(WorldPlaneNormal);
 						FDebugDrawQueue::GetInstance().DrawDebugCircle(WorldPointLocation, 0.5f * Settings.DrawScale * Settings.ContactWidth, 12, DiscColor, false, KINDA_SMALL_NUMBER, Settings.DrawPriority, Settings.LineThickness, Axes.GetUnitAxis(EAxis::Y), Axes.GetUnitAxis(EAxis::Z), false);
 
-						// Unfortunately this doesn't work any more - we overwrite the original contacts with the friction contacts
-						// Previous points
-						//const FManifoldPointSavedData* PrevManifoldPoint = Contact.FindManifoldPointSavedData(ManifoldPoint);
-						//if (PrevManifoldPoint != nullptr)
-						//{
-						//	const FVec3 PrevPointLocation = PointTransform.TransformPosition(PrevManifoldPoint->CoMContactPoints[ContactPointOwner]);
-						//	const FVec3 PrevPlaneLocation = PlaneTransform.TransformPosition(PrevManifoldPoint->CoMContactPoints[ContactPlaneOwner]);
-						//	const FVec3 WorldPrevPointLocation = SpaceTransform.TransformPosition(PrevPointLocation);
-						//	const FVec3 WorldPrevPlaneLocation = SpaceTransform.TransformPosition(PrevPlaneLocation);
-						//	FDebugDrawQueue::GetInstance().DrawDebugLine(WorldPrevPointLocation, WorldPointLocation, FColor::White, false, KINDA_SMALL_NUMBER, Settings.DrawPriority, Settings.LineThickness);
-						//	FDebugDrawQueue::GetInstance().DrawDebugLine(WorldPrevPlaneLocation, WorldPlaneLocation, FColor::White, false, KINDA_SMALL_NUMBER, Settings.DrawPriority, Settings.LineThickness);
-						//}
-
 						// Whether restored
 						if (Contact.WasManifoldRestored())
 						{
@@ -1152,21 +1136,21 @@ namespace Chaos
 		{
 			auto DrawGraphCollision = [&](const FRigidTransform3& SpaceTransform, const FPBDCollisionConstraint* Constraint,  int32 IslandIndex, int32 LevelIndex, int32 ColorIndex, int32 OrderIndex, bool bIsUsed, const FChaosDebugDrawSettings& Settings)
 			{
-				const FRigidTransform3 Transform0 = FParticleUtilities::GetCoMWorldTransform(FConstGenericParticleHandle(Constraint->GetConstrainedParticles()[0])) * SpaceTransform;
-				const FRigidTransform3 Transform1 = FParticleUtilities::GetCoMWorldTransform(FConstGenericParticleHandle(Constraint->GetConstrainedParticles()[1])) * SpaceTransform;
-
 				FVec3 ContactPos = Constraint->GetContactLocation();
 				if (Constraint->GetManifoldPoints().Num() > 0)
 				{
 					ContactPos = FVec3(0);
 					for (const FManifoldPoint& ManifoldPoint : Constraint->GetManifoldPoints())
 					{
-						ContactPos += Transform0.TransformPositionNoScale(ManifoldPoint.CoMContactPoints[0]);
-						ContactPos += Transform1.TransformPositionNoScale(ManifoldPoint.CoMContactPoints[1]);
+						ContactPos += SpaceTransform.TransformPosition(ManifoldPoint.WorldContactPoints[0]);
+						ContactPos += SpaceTransform.TransformPosition(ManifoldPoint.WorldContactPoints[1]);
 					}
 					ContactPos /= (FReal)(2 * Constraint->GetManifoldPoints().Num());
 				}
 
+				const FRigidTransform3 Transform0 = FParticleUtilities::GetCoMWorldTransform(FConstGenericParticleHandle(Constraint->GetConstrainedParticles()[0])) * SpaceTransform;
+				const FRigidTransform3 Transform1 = FParticleUtilities::GetCoMWorldTransform(FConstGenericParticleHandle(Constraint->GetConstrainedParticles()[1])) * SpaceTransform;
+
 				if ((bChaosDebugDebugDrawContactGraphUsed && bIsUsed) || (bChaosDebugDebugDrawContactGraphUnused && !bIsUsed))
 				{
 					FColor Color = bIsUsed ? FColor::Green : FColor::Red;
diff --git a/Engine/Source/Runtime/Experimental/Chaos/Private/Chaos/Collision/PBDCollisionConstraint.cpp b/Engine/Source/Runtime/Experimental/Chaos/Private/Chaos/Collision/PBDCollisionConstraint.cpp
--- a/Engine/Source/Runtime/Experimental/Chaos/Private/Chaos/Collision/PBDCollisionConstraint.cpp
+++ b/Engine/Source/Runtime/Experimental/Chaos/Private/Chaos/Collision/PBDCollisionConstraint.cpp
@@ -37,6 +37,29 @@ namespace Chaos
 
 	bool bChaos_Manifold_EnableFrictionRestore = true;
 	FAutoConsoleVariableRef CVarChaos_Manifold_EnableFrictionRestore(TEXT("p.Chaos.Collision.Manifold.EnableFrictionRestore"), bChaos_Manifold_EnableFrictionRestore, TEXT(""));
+	
+	struct FCollisionTolerances
+	{
+		// Multiplied by the contact margin to produce a distance within which contacts are considered to be the same point
+		FReal ContactPositionToleranceScale = FReal(0.8);
+
+		// Multiplied by the contact margin to produce a max distance that a shape can move if we want to reuse contacts
+		FReal ShapePositionToleranceScale0 = FReal(0.5);	// 0 contacts
+		FReal ShapePositionToleranceScaleN = FReal(0.2);	// >0 contacts
+
+		// A threshold on the quaternion change the tells us when we cannot reuse contacts
+		FReal ShapeRotationThreshold0 = FReal(0.9998);		// 0 contacts
+		FReal ShapeRotationThresholdN = FReal(0.9999);		// >0 contacts
+
+		// Thresholds used to restore individual manifold points
+		FReal ManifoldPointPositionToleranceScale = FReal(1);
+		FReal ManifoldPointNormalThreshold  = FReal(0.7);
+	};
+
+	// @todo(chaos): put these tolerances on cvars
+	// @todo(chaos): tune the tolerances used in FPBDCollisionConstraint::UpdateAndTryRestoreManifold
+	FCollisionTolerances Chaos_Manifold_Tolerances;
+
 
 	extern bool bChaos_Collision_Manifold_FixNormalsInWorldSpace;
 
@@ -130,6 +153,8 @@ namespace Chaos
 		, GJKWarmStartData()
 		, ManifoldPointSavedData{}
 		, NumSavedManifoldPoints(0)
+		, ShapeWorldTransform0()
+		, ShapeWorldTransform1()
 		, LastShapeWorldTransform0()
 		, LastShapeWorldTransform1()
 		, ExpectedNumManifoldPoints(0)
@@ -166,6 +191,8 @@ namespace Chaos
 		, GJKWarmStartData()
 		, ManifoldPointSavedData{}
 		, NumSavedManifoldPoints(0)
+		, ShapeWorldTransform0()
+		, ShapeWorldTransform1()
 		, LastShapeWorldTransform0()
 		, LastShapeWorldTransform1()
 		, ExpectedNumManifoldPoints(0)
@@ -344,32 +371,21 @@ namespace Chaos
 
 	void FPBDCollisionConstraint::UpdateManifoldContacts()
 	{
-		FVec3 P0, P1;
-		FRotation3 Q0, Q1;
-
+		// This is only used when calling collision detection in the solver loop, which is only for incremental manifolds.
 		// @todo(chaos): Remove this when we don't need to support incremental manifolds (this will only be called on creation/restore)
 		if ((GetSolverBody0() != nullptr) && (GetSolverBody1() != nullptr))
 		{
-			P0 = GetSolverBody0()->CorrectedP();
-			Q0 = GetSolverBody0()->CorrectedQ();
-			P1 = GetSolverBody1()->CorrectedP();
-			Q1 = GetSolverBody1()->CorrectedQ();
-		}
-		else
-		{
-			// @todo(chaos): we should not need to regenerate the CoM transform
-			P0 = FParticleUtilities::GetCoMWorldPosition(FConstGenericParticleHandle(Particle[0]));
-			Q0 = FParticleUtilities::GetCoMWorldRotation(FConstGenericParticleHandle(Particle[0]));
-			P1 = FParticleUtilities::GetCoMWorldPosition(FConstGenericParticleHandle(Particle[1]));
-			Q1 = FParticleUtilities::GetCoMWorldRotation(FConstGenericParticleHandle(Particle[1]));
+			ShapeWorldTransform0 = FRigidTransform3(GetSolverBody0()->CorrectedP(), GetSolverBody0()->CorrectedQ());
+			ShapeWorldTransform1 = FRigidTransform3(GetSolverBody1()->CorrectedP(), GetSolverBody1()->CorrectedQ());
 		}
 
 		Manifold.Reset();
 
-		for (int32 Index = 0; Index < ManifoldPoints.Num(); Index++)
+		for (int32 ManifoldPointIndex = 0; ManifoldPointIndex < ManifoldPoints.Num(); ManifoldPointIndex++)
 		{
-			FManifoldPoint& ManifoldPoint = ManifoldPoints[Index];
-			GetWorldSpaceManifoldPoint(ManifoldPoint, P0, Q0, P1, Q1, ManifoldPoint.ContactPoint.Location, ManifoldPoint.ContactPoint.Phi);
+			FManifoldPoint& ManifoldPoint = ManifoldPoints[ManifoldPointIndex];
+
+			UpdateManifoldPointFromContact(ManifoldPointIndex);
 
 			ManifoldPoint.bInsideStaticFrictionCone = bUseManifold;
 
@@ -416,7 +432,9 @@ namespace Chaos
 			if (ManifoldPointIndex >= 0)
 			{
 				// This contact point is already in the manifold - update the state
-				UpdateManifoldPoint(ManifoldPointIndex, ContactPoint);
+				ManifoldPoints[ManifoldPointIndex].ContactPoint = ContactPoint;
+
+				UpdateManifoldPointFromContact(ManifoldPointIndex);
 			}
 			else
 			{
@@ -470,11 +488,11 @@ namespace Chaos
 		ManifoldPoint.InitialShapeContactPoints[0] = ManifoldPoint.ContactPoint.ShapeContactPoints[0];
 		ManifoldPoint.InitialShapeContactPoints[1] = ManifoldPoint.ContactPoint.ShapeContactPoints[1];
 
-		// Update the derived contact state (CoM relative data)
-		UpdateManifoldPointFromContact(ManifoldPointIndex);
-
 		// Initialize the previous contact transforms if the data is available, otherwise reset them to current
 		TryRestoreFrictionData(ManifoldPointIndex);
+
+		// Update the derived contact state (World-Space data)
+		UpdateManifoldPointFromContact(ManifoldPointIndex);
 	}
 
 	int32 FPBDCollisionConstraint::AddManifoldPoint(const FContactPoint& ContactPoint)
@@ -488,35 +506,11 @@ namespace Chaos
 		return ManifoldPointIndex;
 	}
 
-	void FPBDCollisionConstraint::UpdateManifoldPoint(int32 ManifoldPointIndex, const FContactPoint& ContactPoint)
-	{
-		// We really need to know that it's exactly the same contact and not just a close one to update it here
-		// otherwise the PrevLocalContactPoint1 we calculated is not longer for the correct point.
-		FManifoldPoint& ManifoldPoint = ManifoldPoints[ManifoldPointIndex];
-		ManifoldPoint.ContactPoint = ContactPoint;
-		UpdateManifoldPointFromContact(ManifoldPointIndex);
-	}
-
-	// Update the derived contact state (CoM relative data)
 	void FPBDCollisionConstraint::UpdateManifoldPointFromContact(const int32 ManifoldPointIndex)
 	{
 		FManifoldPoint& ManifoldPoint = ManifoldPoints[ManifoldPointIndex];
-		FConstGenericParticleHandle Particle0 = Particle[0];
-		FConstGenericParticleHandle Particle1 = Particle[1];
-
-		const FVec3 LocalContactPoint0 = ImplicitTransform[0].TransformPositionNoScale(ManifoldPoint.ContactPoint.ShapeContactPoints[0]);	// Particle Space on body 0
-		const FVec3 LocalContactPoint1 = ImplicitTransform[1].TransformPositionNoScale(ManifoldPoint.ContactPoint.ShapeContactPoints[1]);	// Particle Space on body 1
-
-		const FVec3 CoMContactPoint0 = Particle0->RotationOfMass().Inverse() * (LocalContactPoint0 - Particle0->CenterOfMass());	// CoM Space on Body 0
-		const FVec3 CoMContactPoint1 = Particle1->RotationOfMass().Inverse() * (LocalContactPoint1 - Particle1->CenterOfMass());	// CoM Space on Body 1
-		ManifoldPoint.CoMContactPoints[0] = CoMContactPoint0;
-		ManifoldPoint.CoMContactPoints[1] = CoMContactPoint1;
-
-		// We now assume that the low-level collision detection functions initialize the world-space contact in the way we want, which is as below...
-		//const FVec3 WorldContactPoint0 = Particle0->P() + Particle0->Q().RotateVector(LocalContactPoint0);
-		//const FVec3 WorldContactPoint1 = Particle1->P() + Particle1->Q().RotateVector(LocalContactPoint1);
-		//ManifoldPoint.ContactPoint.Location = FReal(0.5) * (WorldContactPoint0 + WorldContactPoint1);
-		//ManifoldPoint.ContactPoint.Phi = FVec3::DotProduct(WorldContactPoint0 - WorldContactPoint1, ManifoldPoint.ContactPoint.Normal);
+		ManifoldPoint.WorldContactPoints[0] = ShapeWorldTransform0.TransformPositionNoScale(ManifoldPoint.ContactPoint.ShapeContactPoints[0]);
+		ManifoldPoint.WorldContactPoints[1] = ShapeWorldTransform1.TransformPositionNoScale(ManifoldPoint.ContactPoint.ShapeContactPoints[1]);
 	}
 
 	void FPBDCollisionConstraint::SetActiveContactPoint(const FContactPoint& ContactPoint)
@@ -527,52 +521,6 @@ namespace Chaos
 		Manifold.Phi = ContactPoint.Phi;
 	}
 
-	void FPBDCollisionConstraint::GetWorldSpaceContactPositions(
-		const FManifoldPoint& ManifoldPoint,
-		const FVec3& P0,			// World-Space CoM
-		const FRotation3& Q0,		// World-Space CoM
-		const FVec3& P1,			// World-Space CoM
-		const FRotation3& Q1,		// World-Space CoM
-		FVec3& OutWorldPosition0,
-		FVec3& OutWorldPosition1)
-	{
-		OutWorldPosition0 = P0 + Q0.RotateVector(ManifoldPoint.CoMContactPoints[0]);
-		OutWorldPosition1 = P1 + Q1.RotateVector(ManifoldPoint.CoMContactPoints[1]);
-	}
-
-	void FPBDCollisionConstraint::GetCoMContactPositionsFromWorld(
-		const FManifoldPoint& ManifoldPoint,
-		const FVec3& PCoM0,
-		const FRotation3& QCoM0,
-		const FVec3& PCoM1,
-		const FRotation3& QCoM1,
-		const FVec3& WorldPoint0,
-		const FVec3& WorldPoint1,
-		FVec3& OutCoMPoint0,
-		FVec3& OutCoMPoint1)
-	{
-		// Invert the transformation in GetWorldSpaceContactPositions() and return CoM space contact locations.
-		OutCoMPoint0 = QCoM0.UnrotateVector(WorldPoint0 - PCoM0);
-		OutCoMPoint1 = QCoM1.UnrotateVector(WorldPoint1 - PCoM1);
-	}
-
-	void FPBDCollisionConstraint::GetWorldSpaceManifoldPoint(
-		const FManifoldPoint& ManifoldPoint,
-		const FVec3& P0,			// World-Space CoM
-		const FRotation3& Q0,		// World-Space CoM
-		const FVec3& P1,			// World-Space CoM
-		const FRotation3& Q1,		// World-Space CoM
-		FVec3& OutContactLocation,
-		FReal& OutContactPhi)
-	{
-		FVec3 ContactPos0;
-		FVec3 ContactPos1;
-		FPBDCollisionConstraint::GetWorldSpaceContactPositions(ManifoldPoint, P0, Q0, P1, Q1, ContactPos0, ContactPos1);
-
-		OutContactLocation = 0.5f * (ContactPos0 + ContactPos1);
-		OutContactPhi = FVec3::DotProduct(ContactPos0 - ContactPos1, ManifoldPoint.ContactPoint.Normal);
-	}
-
 	bool FPBDCollisionConstraint::CanUseManifold(FGeometryParticleHandle* Particle0, FGeometryParticleHandle* Particle1) const
 	{
 		// Do not use manifolds when a body is connected by a joint to another. Manifolds do not work when the bodies may be moved
@@ -624,20 +572,27 @@ namespace Chaos
 		bWasManifoldRestored = true;
 	}
 
-	void FPBDCollisionConstraint::UpdateLastShapeWorldTransforms(const FRigidTransform3& ShapeWorldTransform0, const FRigidTransform3& ShapeWorldTransform1)
+
+	void FPBDCollisionConstraint::SetShapeWorldTransforms(const FRigidTransform3& InShapeWorldTransform0, const FRigidTransform3& InShapeWorldTransform1)
 	{
-		LastShapeWorldTransform0 = ShapeWorldTransform0;
-		LastShapeWorldTransform1 = ShapeWorldTransform1;
+		ShapeWorldTransform0 = InShapeWorldTransform0;
+		ShapeWorldTransform1 = InShapeWorldTransform1;
 	}
 
-	bool FPBDCollisionConstraint::UpdateAndTryRestoreManifold(const FRigidTransform3& ShapeWorldTransform0, const FRigidTransform3& ShapeWorldTransform1)
+	void FPBDCollisionConstraint::SetLastShapeWorldTransforms(const FRigidTransform3& InShapeWorldTransform0, const FRigidTransform3& InShapeWorldTransform1)
+	{
+		LastShapeWorldTransform0 = InShapeWorldTransform0;
+		LastShapeWorldTransform1 = InShapeWorldTransform1;
+	}
+
+	bool FPBDCollisionConstraint::UpdateAndTryRestoreManifold()
 	{
 		check(ManifoldPoints.Num() <= MaxManifoldPoints);
 
-		// @todo(chaos): tune the multipliers
-		const FReal ContactPositionTolerance = FReal(0.8) * CollisionTolerance;
-		const FReal ShapePositionTolerance = (ManifoldPoints.Num() > 0) ? FReal(0.2) * CollisionTolerance : FReal(0.5) * CollisionTolerance;
-		const FReal ShapeRotationThreshold = (ManifoldPoints.Num() > 0) ? FReal(0.9999) : FReal(0.9998);
+		const FCollisionTolerances& Tolerances = Chaos_Manifold_Tolerances;
+		const FReal ContactPositionTolerance = Tolerances.ContactPositionToleranceScale * CollisionTolerance;
+		const FReal ShapePositionTolerance = (ManifoldPoints.Num() > 0) ? Tolerances.ShapePositionToleranceScaleN * CollisionTolerance : Tolerances.ShapePositionToleranceScale0 * CollisionTolerance;
+		const FReal ShapeRotationThreshold = (ManifoldPoints.Num() > 0) ? Tolerances.ShapeRotationThresholdN : Tolerances.ShapeRotationThreshold0;
 		const FReal ContactPositionToleranceSq = FMath::Square(ContactPositionTolerance);
 
 		// Reset current closest point
@@ -674,15 +629,19 @@ namespace Chaos
 			{
 				// Recalculate the contact points at the new location
 				// @todo(chaos): we should reproject the contact on the plane owner
-				FConstGenericParticleHandle Particle1 = Particle[1];
 				const FVec3 ShapeContactPoint1 = Contact0In1 - ContactPhi * ContactNormalIn1;
-				const FVec3 ActorContactPoint = ImplicitTransform[1].TransformPositionNoScale(ShapeContactPoint1);
 				ManifoldPoint.ContactPoint.ShapeContactPoints[1] = ShapeContactPoint1;
-				ManifoldPoint.CoMContactPoints[1] = Particle1->RotationOfMass().Inverse() * (ActorContactPoint - Particle1->CenterOfMass());
-				ManifoldPoint.ContactPoint.Location = ShapeWorldTransform1.TransformPositionNoScale(FReal(0.5) * (ShapeContactPoint1 + Contact0In1));
+
+				// Restore friction anchors if possible
+				TryRestoreFrictionData(ManifoldPointIndex);
+
+				// Update world-space contact locations
+				UpdateManifoldPointFromContact(ManifoldPointIndex);
+				ManifoldPoint.ContactPoint.Location = FReal(0.5) * (ManifoldPoint.WorldContactPoints[0] + ManifoldPoint.WorldContactPoints[1]);
 				ManifoldPoint.ContactPoint.Phi = ContactPhi;
+
 				ManifoldPoint.bWasRestored = true;
-				TryRestoreFrictionData(ManifoldPointIndex);
+
 				if (ContactPhi < Manifold.Phi)
 				{
 					// Update closest point
@@ -725,12 +684,13 @@ namespace Chaos
 		return false;
 	}
 
-	bool FPBDCollisionConstraint::TryAddManifoldContact(const FContactPoint& NewContactPoint, const FRigidTransform3& ShapeWorldTransform0, const FRigidTransform3& ShapeWorldTransform1)
+	bool FPBDCollisionConstraint::TryAddManifoldContact(const FContactPoint& NewContactPoint)
 	{
 		check(ManifoldPoints.Num() <= MaxManifoldPoints);
 
-		const FReal PositionTolerance = FReal(1.0) * CollisionTolerance;
-		const FReal NormalThreshold = FReal(0.7);
+		const FCollisionTolerances& Tolerances = Chaos_Manifold_Tolerances;
+		const FReal PositionTolerance = Tolerances.ManifoldPointPositionToleranceScale * CollisionTolerance;
+		const FReal NormalThreshold = Tolerances.ManifoldPointNormalThreshold;
 
 		// We must end up with a full manifold after this if we want to reuse it
 		if ((ManifoldPoints.Num() < ExpectedNumManifoldPoints - 1) || (ExpectedNumManifoldPoints == 0))
@@ -775,8 +735,8 @@ namespace Chaos
 					ManifoldPoint.InitialShapeContactPoints[0] = NewContactPoint.ShapeContactPoints[0];
 					ManifoldPoint.InitialShapeContactPoints[1] = NewContactPoint.ShapeContactPoints[1];
 					ManifoldPoint.bWasRestored = false;
-					UpdateManifoldPointFromContact(ManifoldPointIndex);
 					TryRestoreFrictionData(ManifoldPointIndex);
+					UpdateManifoldPointFromContact(ManifoldPointIndex);
 					if (NewContactPoint.Phi < GetPhi())
 					{
 						SetActiveContactPoint(ManifoldPoint.ContactPoint);
@@ -790,13 +750,13 @@ namespace Chaos
 		// If we have a full manifold, see if we can use or reject the GJK point
 		if (ManifoldPoints.Num() == 4)
 		{
-			return TryInsertManifoldContact(NewContactPoint, ShapeWorldTransform0, ShapeWorldTransform1);
+			return TryInsertManifoldContact(NewContactPoint);
 		}
 		
 		return false;
 	}
 
-	bool FPBDCollisionConstraint::TryInsertManifoldContact(const FContactPoint& NewContactPoint, const FRigidTransform3& ShapeWorldTransform0, const FRigidTransform3& ShapeWorldTransform1)
+	bool FPBDCollisionConstraint::TryInsertManifoldContact(const FContactPoint& NewContactPoint)
 	{
 		check(ManifoldPoints.Num() == 4);
 
diff --git a/Engine/Source/Runtime/Experimental/Chaos/Private/Chaos/Collision/PBDCollisionSolver.cpp b/Engine/Source/Runtime/Experimental/Chaos/Private/Chaos/Collision/PBDCollisionSolver.cpp
--- a/Engine/Source/Runtime/Experimental/Chaos/Private/Chaos/Collision/PBDCollisionSolver.cpp
+++ b/Engine/Source/Runtime/Experimental/Chaos/Private/Chaos/Collision/PBDCollisionSolver.cpp
@@ -354,13 +354,13 @@ namespace Chaos
 	void FPBDCollisionSolverManifoldPoint::InitContact(
 		const FConstraintSolverBody& Body0,
 		const FConstraintSolverBody& Body1,
-		const FVec3& InCoMAnchorPoint0,
-		const FVec3& InCoMAnchorPoint1,
+		const FVec3& InWorldAnchorPoint0,
+		const FVec3& InWorldAnchorPoint1,
 		const FVec3& InWorldContactNormal)
 	{
 		NetPushOut = FVec3(0);
 		NetImpulse = FVec3(0);
-		UpdateContact(Body0, Body1, InCoMAnchorPoint0, InCoMAnchorPoint1, InWorldContactNormal);
+		UpdateContact(Body0, Body1, InWorldAnchorPoint0, InWorldAnchorPoint1, InWorldContactNormal);
 		UpdateMass(Body0, Body1);
 	}
 
@@ -377,23 +377,18 @@ namespace Chaos
 	inline void FPBDCollisionSolverManifoldPoint::UpdateContact(
 		const FConstraintSolverBody& Body0, 
 		const FConstraintSolverBody& Body1,
-		const FVec3& InCoMAnchorPoint0,
-		const FVec3& InCoMAnchorPoint1, 
+		const FVec3& InWorldAnchorPoint0,
+		const FVec3& InWorldAnchorPoint1, 
 		const FVec3& InWorldContactNormal)
 	{
-		// The contact points on each body
-		const FVec3 WorldRelativeAnchorPoint0 = Body0.Q() * InCoMAnchorPoint0;
-		const FVec3 WorldRelativeAnchorPoint1 = Body1.Q() * InCoMAnchorPoint1;
-		const FVec3 WorldAnchorPoint0 = Body0.P() + WorldRelativeAnchorPoint0;
-		const FVec3 WorldAnchorPoint1 = Body1.P() + WorldRelativeAnchorPoint1;
-
 		// The world-space point where we apply impulses/corrections (same world-space point for momentum conservation)
-		WorldContactPosition = FReal(0.5) * (WorldAnchorPoint0 + WorldAnchorPoint1);
+		WorldContactPosition = FReal(0.5) * (InWorldAnchorPoint0 + InWorldAnchorPoint1);
 
 		// The world-space contact normal
 		WorldContactNormal = InWorldContactNormal;
 
-		WorldContactDelta = WorldAnchorPoint0 - WorldAnchorPoint1;
+		// The contact point error we are trying to correct in this solver
+		WorldContactDelta = InWorldAnchorPoint0 - InWorldAnchorPoint1;
 	}
 
 	inline void FPBDCollisionSolverManifoldPoint::UpdateMass(const FConstraintSolverBody& Body0, const FConstraintSolverBody& Body1)
@@ -548,11 +543,11 @@ namespace Chaos
 
 	void FPBDCollisionSolver::InitContact(
 		const int32 ManifoldPoiontIndex,
-		const FVec3& InCoMAnchorPoint0,
-		const FVec3& InCoMAnchorPoint1,
+		const FVec3& InWorldAnchorPoint0,
+		const FVec3& InWorldAnchorPoint1,
 		const FVec3& InWorldContactNormal)
 	{
-		State.ManifoldPoints[ManifoldPoiontIndex].InitContact(State.SolverBodies[0], State.SolverBodies[1], InCoMAnchorPoint0, InCoMAnchorPoint1, InWorldContactNormal);
+		State.ManifoldPoints[ManifoldPoiontIndex].InitContact(State.SolverBodies[0], State.SolverBodies[1], InWorldAnchorPoint0, InWorldAnchorPoint1, InWorldContactNormal);
 	}
 
 	void FPBDCollisionSolver::InitMaterial(
@@ -566,11 +561,11 @@ namespace Chaos
 
 	void FPBDCollisionSolver::UpdateContact(
 		const int32 ManifoldPoiontIndex,
-		const FVec3& InCoMAnchorPoint0,
-		const FVec3& InCoMAnchorPoint1,
+		const FVec3& InWorldAnchorPoint0,
+		const FVec3& InWorldAnchorPoint1,
 		const FVec3& InWorldContactNormal)
 	{
-		State.ManifoldPoints[ManifoldPoiontIndex].UpdateContact(State.SolverBodies[0], State.SolverBodies[1], InCoMAnchorPoint0, InCoMAnchorPoint1, InWorldContactNormal);
+		State.ManifoldPoints[ManifoldPoiontIndex].UpdateContact(State.SolverBodies[0], State.SolverBodies[1], InWorldAnchorPoint0, InWorldAnchorPoint1, InWorldContactNormal);
 	}
 
 	bool FPBDCollisionSolver::SolvePosition(const FReal Dt, const FReal MaxPushOut, const bool bApplyStaticFriction)
diff --git a/Engine/Source/Runtime/Experimental/Chaos/Private/Chaos/Collision/PBDCollisionSolver.h b/Engine/Source/Runtime/Experimental/Chaos/Private/Chaos/Collision/PBDCollisionSolver.h
--- a/Engine/Source/Runtime/Experimental/Chaos/Private/Chaos/Collision/PBDCollisionSolver.h
+++ b/Engine/Source/Runtime/Experimental/Chaos/Private/Chaos/Collision/PBDCollisionSolver.h
@@ -31,8 +31,8 @@ namespace Chaos
 		void InitContact(
 			const FConstraintSolverBody& Body0,
 			const FConstraintSolverBody& Body1,
-			const FVec3& InCoMAnchorPoint0,
-			const FVec3& InCoMAnchorPoint1,
+			const FVec3& InWorldAnchorPoint0,
+			const FVec3& InWorldAnchorPoint1,
 			const FVec3& InWorldContactNormal);
 
 		/**
@@ -49,9 +49,9 @@ namespace Chaos
 		void UpdateContact(
 			const FConstraintSolverBody& Body0,
 			const FConstraintSolverBody& Body1,
-			const FVec3& CoMAnchorPoint0,
-			const FVec3& CoMAnchorPoint1,
-			const FVec3& WorldContactNormal);
+			const FVec3& InWorldAnchorPoint0,
+			const FVec3& InWorldAnchorPoint1,
+			const FVec3& InWorldContactNormal);
 
 		/**
 		 * @brief Update the cached mass properties based on the current body transforms
@@ -175,8 +175,8 @@ namespace Chaos
 
 		void InitContact(
 			const int32 ManifoldPoiontIndex,
-			const FVec3& InCoMAnchorPoint0,
-			const FVec3& InCoMAnchorPoint1,
+			const FVec3& InWorldAnchorPoint0,
+			const FVec3& InWorldAnchorPoint1,
 			const FVec3& InWorldContactNormal);
 
 		void InitMaterial(
@@ -187,8 +187,8 @@ namespace Chaos
 
 		void UpdateContact(
 			const int32 ManifoldPoiontIndex,
-			const FVec3& InCoMAnchorPoint0,
-			const FVec3& InCoMAnchorPoint1,
+			const FVec3& InWorldAnchorPoint0,
+			const FVec3& InWorldAnchorPoint1,
 			const FVec3& InWorldContactNormal);
 
 		/**
diff --git a/Engine/Source/Runtime/Experimental/Chaos/Private/Chaos/Collision/ParticlePairMidPhase.cpp b/Engine/Source/Runtime/Experimental/Chaos/Private/Chaos/Collision/ParticlePairMidPhase.cpp
--- a/Engine/Source/Runtime/Experimental/Chaos/Private/Chaos/Collision/ParticlePairMidPhase.cpp
+++ b/Engine/Source/Runtime/Experimental/Chaos/Private/Chaos/Collision/ParticlePairMidPhase.cpp
@@ -240,6 +240,11 @@ namespace Chaos
 			const int32 LastEpoch = MidPhase.GetCollisionAllocator().GetCurrentEpoch() - 1;
 			const bool bWasUpdatedLastTick = IsUsedSince(LastEpoch);
 
+			// Update the world shape transforms on the constraint (we cannot just give it the PerShapeData 
+			// pointer because of Unions - see FMultiShapePairCollisionDetector)
+			// NOTE: these are not used by CCD which continuously moves the particles
+			Constraint->SetShapeWorldTransforms(ShapeWorldTransform0, ShapeWorldTransform1);
+
 			// If the constraint was not used last frame, it needs to be reset. 
 			// If it was used last frame, its data can be used for static friction etc. (unless CCD is enabled)
 			Constraint->SetCCDEnabled(bUseCCD);
@@ -260,7 +265,7 @@ namespace Chaos
 			{
 				// Update the existing manifold. We can re-use as-is if none of the points have moved much and the bodies have not moved much
 				// NOTE: this can succeed in "restoring" even if we have no manifold points
-				bWasManifoldRestored = Constraint->UpdateAndTryRestoreManifold(ShapeWorldTransform0, ShapeWorldTransform1);
+				bWasManifoldRestored = Constraint->UpdateAndTryRestoreManifold();
 			}
 			else
 			{
@@ -271,7 +276,7 @@ namespace Chaos
 			if (!bWasManifoldRestored)
 			{
 				// We will be updating the manifold, if only partially, so update the restore comparison transforms
-				Constraint->UpdateLastShapeWorldTransforms(ShapeWorldTransform0, ShapeWorldTransform1);
+				Constraint->SetLastShapeWorldTransforms(ShapeWorldTransform0, ShapeWorldTransform1);
 
 				// Run the narrow phase
 				if (!bUseCCD)
@@ -283,7 +288,6 @@ namespace Chaos
 					// Note: This is unusual but we are using a mix of the previous and current transform
 					// This is due to how CCD rewinds the position (not rotation) and then sweeps to find the first contact at the current orientation
 					// NOTE: These are actor transforms, not CoM transforms
-					// @todo(chaos): see if we can easily switch to CoM transforms now in collision loop (shapes are held in actor space)
 					// @todo(chaos): this is broken if both objects are CCD
 					FConstGenericParticleHandle P0 = Particle0;
 					FConstGenericParticleHandle P1 = Particle1;
@@ -314,6 +318,10 @@ namespace Chaos
 		const int32 LastEpoch = MidPhase.GetCollisionAllocator().GetCurrentEpoch() - 1;
 		if ((Constraint != nullptr) && IsUsedSince(LastEpoch))
 		{
+			const FRigidTransform3& ShapeWorldTransform0 = Shape0->GetLeafWorldTransform();
+			const FRigidTransform3& ShapeWorldTransform1 = Shape1->GetLeafWorldTransform();
+			Constraint->SetShapeWorldTransforms(ShapeWorldTransform0, ShapeWorldTransform1);
+
 			Constraint->RestoreManifold();
 			if (Constraint->GetPhi() <= CullDistance)
 			{
@@ -483,6 +491,14 @@ namespace Chaos
 			// NOTE: Using InParticle0 and InParticle1 here because the order may be different to what we have stored
 			Constraint = CreateConstraint(InParticle0, Implicit0, BVHParticles0, ShapeRelativeTransform0, InParticle1, Implicit1, BVHParticles1, ShapeRelativeTransform1, CullDistance, ShapePairType, bUseManifold, Key);
 		}
+
+		// @todo(chaos): we already have the shape world transforms at the calling site - pass them in
+		const FRigidTransform3 ParticleTransform0 = FParticleUtilitiesPQ::GetActorWorldTransform(FConstGenericParticleHandle(InParticle0));
+		const FRigidTransform3 ParticleTransform1 = FParticleUtilitiesPQ::GetActorWorldTransform(FConstGenericParticleHandle(InParticle1));
+		const FRigidTransform3 ShapeWorldTransform0 = ShapeRelativeTransform0 * ParticleTransform0;
+		const FRigidTransform3 ShapeWorldTransform1 = ShapeRelativeTransform1 * ParticleTransform1;
+		Constraint->SetShapeWorldTransforms(ShapeWorldTransform0, ShapeWorldTransform1);
+
 		NewConstraints.Add(Constraint);
 		return Constraint;
 	}
@@ -545,13 +561,25 @@ namespace Chaos
 
 	int32 FMultiShapePairCollisionDetector::RestoreCollisions(const FReal CullDistance)
 	{
+		if (Constraints.Num() == 0)
+		{
+			return 0;
+		}
+
 		int32 NumRestored = 0;
 		const int32 LastEpoch = MidPhase.GetCollisionAllocator().GetCurrentEpoch() - 1;
+		const FRigidTransform3 ParticleTransform0 = FParticleUtilitiesPQ::GetActorWorldTransform(FConstGenericParticleHandle(Particle0));
+		const FRigidTransform3 ParticleTransform1 = FParticleUtilitiesPQ::GetActorWorldTransform(FConstGenericParticleHandle(Particle1));
+
 		for (auto& KVP : Constraints)
 		{
 			TUniquePtr<FPBDCollisionConstraint>& Constraint = KVP.Value;
 			if ((Constraint != nullptr) && (Constraint->GetContainerCookie().LastUsedEpoch >= LastEpoch))
 			{
+				const FRigidTransform3 ShapeWorldTransform0 = Constraint->ImplicitTransform[0] * ParticleTransform0;
+				const FRigidTransform3 ShapeWorldTransform1 = Constraint->ImplicitTransform[1] * ParticleTransform1;
+				Constraint->SetShapeWorldTransforms(ShapeWorldTransform0, ShapeWorldTransform1);
+
 				Constraint->RestoreManifold();
 				if (Constraint->GetPhi() < CullDistance)
 				{
diff --git a/Engine/Source/Runtime/Experimental/Chaos/Private/Chaos/Collision/SolverCollisionContainer.cpp b/Engine/Source/Runtime/Experimental/Chaos/Private/Chaos/Collision/SolverCollisionContainer.cpp
--- a/Engine/Source/Runtime/Experimental/Chaos/Private/Chaos/Collision/SolverCollisionContainer.cpp
+++ b/Engine/Source/Runtime/Experimental/Chaos/Private/Chaos/Collision/SolverCollisionContainer.cpp
@@ -102,8 +102,8 @@ namespace Chaos
 				// Initialize the structural data in the contact (relative contact points, contact mass etc)
 				Solver.InitContact(
 					ManifoldPointIndex,
-					ManifoldPoint.CoMContactPoints[0],
-					ManifoldPoint.CoMContactPoints[1],
+					ManifoldPoint.WorldContactPoints[0],
+					ManifoldPoint.WorldContactPoints[1],
 					ManifoldPoint.ContactPoint.Normal);
 
 				// Calculate the target normal velocity based on restitution
diff --git a/Engine/Source/Runtime/Experimental/Chaos/Private/Chaos/CollisionOneShotManifolds.cpp b/Engine/Source/Runtime/Experimental/Chaos/Private/Chaos/CollisionOneShotManifolds.cpp
--- a/Engine/Source/Runtime/Experimental/Chaos/Private/Chaos/CollisionOneShotManifolds.cpp
+++ b/Engine/Source/Runtime/Experimental/Chaos/Private/Chaos/CollisionOneShotManifolds.cpp
@@ -765,7 +765,7 @@ namespace Chaos
 			PHYSICS_CSV_CUSTOM_EXPENSIVE(PhysicsCounters, NumManifoldsGJKCalled, 1, ECsvCustomStatOp::Accumulate);
 
 			const bool bCanUpdateManifold = bChaos_Collision_EnableManifoldInject;
-			if (bCanUpdateManifold && Constraint.TryAddManifoldContact(GJKContactPoint, Convex1Transform, Convex2Transform))
+			if (bCanUpdateManifold && Constraint.TryAddManifoldContact(GJKContactPoint))
 			{
 				PHYSICS_CSV_CUSTOM_EXPENSIVE(PhysicsCounters, NumManifoldsMaintained, 1, ECsvCustomStatOp::Accumulate);
 				return;
diff --git a/Engine/Source/Runtime/Experimental/Chaos/Private/Chaos/ContactModification.cpp b/Engine/Source/Runtime/Experimental/Chaos/Private/Chaos/ContactModification.cpp
--- a/Engine/Source/Runtime/Experimental/Chaos/Private/Chaos/ContactModification.cpp
+++ b/Engine/Source/Runtime/Experimental/Chaos/Private/Chaos/ContactModification.cpp
@@ -101,6 +101,8 @@ namespace Chaos
 		FManifoldPoint& ManifoldPoint = ManifoldPoints[ContactPointIdx];
 
 		ManifoldPoint.ContactPoint.Normal = Normal;
+		ManifoldPoint.ContactPoint.Phi = FVec3::DotProduct(ManifoldPoint.WorldContactPoints[0] - ManifoldPoint.WorldContactPoints[1], ManifoldPoint.ContactPoint.Normal);
+
 		Modifier->MarkConstraintForManifoldUpdate(*Constraint);
 	}
 
@@ -108,60 +110,26 @@ namespace Chaos
 	{
 		TArrayView<FManifoldPoint> ManifoldPoints = Constraint->GetManifoldPoints();
 		FManifoldPoint& ManifoldPoint = ManifoldPoints[ContactPointIdx];
-		FGeometryParticleHandle* Particle0 = Constraint->Particle[0];
-		FGeometryParticleHandle* Particle1 = Constraint->Particle[1];
-
-		const FVec3 PCoM0 = FParticleUtilities::GetCoMWorldPosition(FConstGenericParticleHandle(Particle0));
-		const FVec3 PCoM1 = FParticleUtilities::GetCoMWorldPosition(FConstGenericParticleHandle(Particle1));
-		const FRotation3 QCoM0 = FParticleUtilities::GetCoMWorldRotation(FConstGenericParticleHandle(Particle0));
-		const FRotation3 QCoM1 = FParticleUtilities::GetCoMWorldRotation(FConstGenericParticleHandle(Particle1));
-
-		FPBDCollisionConstraint::GetWorldSpaceContactPositions(ManifoldPoint, PCoM0, QCoM0, PCoM1, QCoM1, OutLocation0, OutLocation1);
+		OutLocation0 = ManifoldPoint.WorldContactPoints[0];
+		OutLocation1 = ManifoldPoint.WorldContactPoints[1];
 	}
 
 	FVec3 FContactPairModifier::GetWorldContactLocation(int32 ContactPointIdx) const
 	{
 		FVec3 WorldPos0, WorldPos1;
 		GetWorldContactLocations(ContactPointIdx, WorldPos0, WorldPos1);
-		return (WorldPos0 + WorldPos1) * 0.5;
+		return (WorldPos0 + WorldPos1) * FReal(0.5);
 	}
 
 	void FContactPairModifier::ModifyWorldContactLocations(const FVec3& Location0, const FVec3& Location1, int32 ContactPointIdx)
 	{
 		TArrayView<FManifoldPoint> ManifoldPoints = Constraint->GetManifoldPoints();
 		FManifoldPoint& ManifoldPoint = ManifoldPoints[ContactPointIdx];
-		FGeometryParticleHandle* Particle0 = Constraint->Particle[0];
-		FGeometryParticleHandle* Particle1 = Constraint->Particle[1];
-
-		const FVec3 PCoM0 = FParticleUtilities::GetCoMWorldPosition(FConstGenericParticleHandle(Particle0));
-		const FVec3 PCoM1 = FParticleUtilities::GetCoMWorldPosition(FConstGenericParticleHandle(Particle1));
-		const FRotation3 QCoM0 = FParticleUtilities::GetCoMWorldRotation(FConstGenericParticleHandle(Particle0));
-		const FRotation3 QCoM1 = FParticleUtilities::GetCoMWorldRotation(FConstGenericParticleHandle(Particle1));
-
-		FPBDCollisionConstraint::GetCoMContactPositionsFromWorld(ManifoldPoint, PCoM0, QCoM0, PCoM1, QCoM1, Location0, Location1, ManifoldPoint.CoMContactPoints[0], ManifoldPoint.CoMContactPoints[1]);
-
-		ManifoldPoint.ContactPoint.Location = 0.5 * (Location0 + Location1);
-
-		Modifier->MarkConstraintForManifoldUpdate(*Constraint);
-	}
-
-	void FContactPairModifier::GetCoMContactLocations(int32 ContactPointIdx, FVec3& OutLocation0, FVec3& OutLocation1) const
-	{
-		TArrayView<FManifoldPoint> ManifoldPoints = Constraint->GetManifoldPoints();
-		FManifoldPoint& ManifoldPoint = ManifoldPoints[ContactPointIdx];
-
-		OutLocation0 = ManifoldPoint.CoMContactPoints[0];
-		OutLocation1 = ManifoldPoint.CoMContactPoints[1];
-	}
-
-	void FContactPairModifier::ModifyCoMContactLocations(const FVec3& Location0, const FVec3& Location1, int32 ContactPointIdx)
-	{
-		TArrayView<FManifoldPoint> ManifoldPoints = Constraint->GetManifoldPoints();
-		FManifoldPoint& ManifoldPoint = ManifoldPoints[ContactPointIdx];
-
-		ManifoldPoint.CoMContactPoints[0] = Location0;
-		ManifoldPoint.CoMContactPoints[1] = Location1;
 
+		ManifoldPoint.ContactPoint.ShapeContactPoints[0] = Constraint->GetShapeWorldTransform0().InverseTransformPositionNoScale(Location0);
+		ManifoldPoint.ContactPoint.ShapeContactPoints[1] = Constraint->GetShapeWorldTransform1().InverseTransformPositionNoScale(Location1);
+		ManifoldPoint.WorldContactPoints[0] = Location0;
+		ManifoldPoint.WorldContactPoints[1] = Location1;
 		ManifoldPoint.ContactPoint.Location = 0.5 * (Location0 + Location1);
 
 		Modifier->MarkConstraintForManifoldUpdate(*Constraint);
@@ -187,24 +155,24 @@ namespace Chaos
 		Constraint->Manifold.RestitutionThreshold = Threshold;
 	}
 
-	FReal FContactPairModifier::GetFriction() const
+	FReal FContactPairModifier::GetDynamicFriction() const
 	{
 		return Constraint->Manifold.Friction;
 	}
 
-	void FContactPairModifier::ModifyFriction(FReal Friction)
+	void FContactPairModifier::ModifyDynamicFriction(FReal DynamicFriction)
 	{
-		Constraint->Manifold.Friction = Friction;
+		Constraint->Manifold.Friction = DynamicFriction;
 	}
 
-	FReal FContactPairModifier::GetAngularFriction() const
+	FReal FContactPairModifier::GetStaticFriction() const
 	{
 		return Constraint->Manifold.AngularFriction;
 	}
 
-	void FContactPairModifier::ModifyAngularFriction(FReal AngularFriction)
+	void FContactPairModifier::ModifyStaticFriction(FReal StaticFriction)
 	{
-		Constraint->Manifold.AngularFriction = AngularFriction;
+		Constraint->Manifold.AngularFriction = StaticFriction;
 	}
 
 	FVec3 FContactPairModifier::GetParticleVelocity(int32 ParticleIdx) const
@@ -292,6 +260,13 @@ namespace Chaos
 		return Particle->X();
 	}
 
+	void FContactPairModifier::UpdateConstraintShapeTransforms()
+	{
+		const FRigidTransform3 ShapeWorldTransform0 = Constraint->GetShapeRelativeTransform0() * FParticleUtilitiesPQ::GetActorWorldTransform(FConstGenericParticleHandle(Constraint->GetParticle0()));
+		const FRigidTransform3 ShapeWorldTransform1 = Constraint->GetShapeRelativeTransform1() * FParticleUtilitiesPQ::GetActorWorldTransform(FConstGenericParticleHandle(Constraint->GetParticle1()));
+		Constraint->SetShapeWorldTransforms(ShapeWorldTransform0, ShapeWorldTransform1);
+	}
+
 	void FContactPairModifier::ModifyParticlePosition(FVec3 Position, bool bMaintainVelocity, int32 ParticleIdx)
 	{
 		FGeometryParticleHandle* Particle = Constraint->Particle[ParticleIdx];
@@ -314,6 +289,7 @@ namespace Chaos
 					// Update V to new implicit velocity
 					RigidHandle->SetV((RigidHandle->P() - RigidHandle->X()) / Modifier->Dt);
 				}
+				UpdateConstraintShapeTransforms();
 				return;
 			}
 			else
@@ -321,6 +297,7 @@ namespace Chaos
 				// Kinematic must keep P/X in sync
 				RigidHandle->SetX(Position);
 				RigidHandle->SetP(Position);
+				UpdateConstraintShapeTransforms();
 				return;
 			}
 		}
@@ -330,6 +307,7 @@ namespace Chaos
 		if (KinematicHandle)
 		{
 			KinematicHandle->SetX(Position);
+			UpdateConstraintShapeTransforms();
 			return;
 		}
 
@@ -374,6 +352,7 @@ namespace Chaos
 					// Update W to new implicit velocity
 					RigidHandle->SetW(FRotation3::CalculateAngularVelocity(RigidHandle->R(), RigidHandle->Q(), Modifier->Dt));
 				}
+				UpdateConstraintShapeTransforms();
 				return;
 			}
 			else
@@ -381,6 +360,7 @@ namespace Chaos
 				// Kinematic must keep Q/R in sync
 				RigidHandle->SetR(Rotation);
 				RigidHandle->SetQ(Rotation);
+				UpdateConstraintShapeTransforms();
 				return;
 			}
 		}
@@ -389,6 +369,7 @@ namespace Chaos
 		if (KinematicHandle)
 		{
 			KinematicHandle->SetR(Rotation);
+			UpdateConstraintShapeTransforms();
 			return;
 		}
 
diff --git a/Engine/Source/Runtime/Experimental/Chaos/Public/Chaos/Collision/PBDCollisionConstraint.h b/Engine/Source/Runtime/Experimental/Chaos/Public/Chaos/Collision/PBDCollisionConstraint.h
--- a/Engine/Source/Runtime/Experimental/Chaos/Public/Chaos/Collision/PBDCollisionConstraint.h
+++ b/Engine/Source/Runtime/Experimental/Chaos/Public/Chaos/Collision/PBDCollisionConstraint.h
@@ -44,7 +44,7 @@ namespace Chaos
 		FManifoldPoint() 
 			: ContactPoint()
 			, InitialShapeContactPoints{ FVec3(0), FVec3(0) }
-			, CoMContactPoints{ FVec3(0), FVec3(0) }
+			, WorldContactPoints{ FVec3(0), FVec3(0) }
 			, NetPushOut(0)
 			, NetImpulse(0)
 			, StaticFrictionMax(0)
@@ -55,7 +55,7 @@ namespace Chaos
 		FManifoldPoint(const FContactPoint& InContactPoint) 
 			: ContactPoint(InContactPoint)
 			, InitialShapeContactPoints{ FVec3(0), FVec3(0) }
-			, CoMContactPoints{ FVec3(0), FVec3(0) }
+			, WorldContactPoints{ FVec3(0), FVec3(0) }
 			, NetPushOut(0)
 			, NetImpulse(0)
 			, StaticFrictionMax(0)
@@ -65,7 +65,7 @@ namespace Chaos
 
 		FContactPoint ContactPoint;			// Contact point results of low-level collision detection
 		FVec3 InitialShapeContactPoints[2];	// ShapeContactPoints when the constraint was first initialized. Used to track reusablility
-		FVec3 CoMContactPoints[2];			// CoM-space contact points on the two bodies
+		FVec3 WorldContactPoints[2];		// World-space contact points on the two bodies
 		FVec3 NetPushOut;					// Total pushout applied at this contact point
 		FVec3 NetImpulse;					// Total impulse applied by this contact point
 		FReal StaticFrictionMax;			// A proxy for the normal impulse used to limit static friction correction. Used for smoothing static friction limits
@@ -89,8 +89,8 @@ namespace Chaos
 		*/
 		inline void Save(const FManifoldPoint& ManifoldPoint)
 		{
-			CoMContactPoints[0] = ManifoldPoint.CoMContactPoints[0];
-			CoMContactPoints[1] = ManifoldPoint.CoMContactPoints[1];
+			ShapeContactPoints[0] = ManifoldPoint.ContactPoint.ShapeContactPoints[0];
+			ShapeContactPoints[1] = ManifoldPoint.ContactPoint.ShapeContactPoints[1];
 			StaticFrictionMax = ManifoldPoint.StaticFrictionMax;
 			bInsideStaticFrictionCone = ManifoldPoint.bInsideStaticFrictionCone;
 		}
@@ -100,8 +100,8 @@ namespace Chaos
 		*/
 		inline void Restore(FManifoldPoint& ManifoldPoint) const
 		{
-			ManifoldPoint.CoMContactPoints[0] = CoMContactPoints[0];
-			ManifoldPoint.CoMContactPoints[1] = CoMContactPoints[1];
+			ManifoldPoint.ContactPoint.ShapeContactPoints[0] = ShapeContactPoints[0];
+			ManifoldPoint.ContactPoint.ShapeContactPoints[1] = ShapeContactPoints[1];
 			ManifoldPoint.StaticFrictionMax = StaticFrictionMax;
 		}
 
@@ -113,12 +113,12 @@ namespace Chaos
 			// If the contact point is in the same spot on one of the bodies, assume it is the same contact
 			// @todo(chaos): more robust same-point test. E.g., this won't work for very small or very large objects,
 			// so at least make the tolerance size-dependent
-			const FVec3 DP0 = ManifoldPoint.CoMContactPoints[0] - CoMContactPoints[0];
-			const FVec3 DP1 = ManifoldPoint.CoMContactPoints[1] - CoMContactPoints[1];
+			const FVec3 DP0 = ManifoldPoint.ContactPoint.ShapeContactPoints[0] - ShapeContactPoints[0];
+			const FVec3 DP1 = ManifoldPoint.ContactPoint.ShapeContactPoints[1] - ShapeContactPoints[1];
 			return ((DP0.SizeSquared() < DistanceToleranceSq) || (DP1.SizeSquared() < DistanceToleranceSq));
 		}
 
-		FVec3 CoMContactPoints[2];
+		FVec3 ShapeContactPoints[2];
 		FVec3 NetPushOut;
 		FReal StaticFrictionMax;
 		bool bInsideStaticFrictionCone;
@@ -161,8 +161,8 @@ namespace Chaos
 		FVec3 Location;
 		FReal Phi;
 
-		FReal Friction;
-		FReal AngularFriction;
+		FReal Friction;			// @todo(chaos): rename DynamicFriction
+		FReal AngularFriction;	// @todo(chaos): rename StaticFriction
 		FReal Restitution;
 		FReal RestitutionPadding; // For PBD implementation of resitution, we pad constraints on initial contact to enforce outward velocity
 		FReal RestitutionThreshold;
@@ -402,21 +402,24 @@ namespace Chaos
 		void UpdateManifoldContacts();
 		void ClearManifold();
 
-		bool UpdateAndTryRestoreManifold(const FRigidTransform3& ShapeWorldTransform0, const FRigidTransform3& ShapeWorldTransform1);
+		const FRigidTransform3& GetShapeRelativeTransform0() const { return ImplicitTransform[0]; }
+		const FRigidTransform3& GetShapeRelativeTransform1() const { return ImplicitTransform[1]; }
+
+		const FRigidTransform3& GetShapeWorldTransform0() const { return ShapeWorldTransform0; }
+		const FRigidTransform3& GetShapeWorldTransform1() const { return ShapeWorldTransform1; }
+
+		void SetShapeWorldTransforms(const FRigidTransform3& InShapeWorldTransform0, const FRigidTransform3& InShapeWorldTransform1);
+		void SetLastShapeWorldTransforms(const FRigidTransform3& InShapeWorldTransform0, const FRigidTransform3& InShapeWorldTransform1);
+
+		bool UpdateAndTryRestoreManifold();
 		void ResetActiveManifoldContacts();
-		void UpdateLastShapeWorldTransforms(const FRigidTransform3& ShapeWorldTransform0, const FRigidTransform3& ShapeWorldTransform1);
-		bool TryAddManifoldContact(const FContactPoint& ContactPoint, const FRigidTransform3& ShapeWorldTransform0, const FRigidTransform3& ShapeWorldTransform1);
-		bool TryInsertManifoldContact(const FContactPoint& ContactPoint, const FRigidTransform3& ShapeWorldTransform0, const FRigidTransform3& ShapeWorldTransform1);
+		bool TryAddManifoldContact(const FContactPoint& ContactPoint);
+		bool TryInsertManifoldContact(const FContactPoint& ContactPoint);
 
 		//@ todo(chaos): These are for the collision forwarding system - this should use the collision modifier system (which should be extended to support adding collisions)
 		void SetManifoldPoints(const TArray<FManifoldPoint>& InManifoldPoints) { ManifoldPoints = InManifoldPoints; }
 		void UpdateManifoldPointFromContact(const int32 ManifoldPointIndex);
 
-		// Helpers for interacting with constraint from world space
-		static void GetWorldSpaceContactPositions(const FManifoldPoint& ManifoldPoint, const FVec3& PCoM0, const FRotation3& QCoM0, const FVec3& PCoM1, const FRotation3& QCoM1, FVec3& OutWorldPosition0, FVec3& OutWorldPosition1);
-		static void GetCoMContactPositionsFromWorld(const FManifoldPoint& ManifoldPoint, const FVec3& PCoM0, const FRotation3& QCoM0, const FVec3& PCoM1, const FRotation3& QCoM1, const FVec3& WorldPoint0, const FVec3& WorldPoint1, FVec3& OutCoMPoint0, FVec3& OutCoMPoint1);
-
-
 		// The GJK warm-start data. This is updated directly in the narrow phase
 		FGJKSimplexData& GetGJKWarmStartData() { return GJKWarmStartData; }
 
@@ -492,9 +495,7 @@ namespace Chaos
 		int32 FindManifoldPoint(const FContactPoint& ContactPoint) const;
 		int32 AddManifoldPoint(const FContactPoint& ContactPoint);
 		void InitManifoldPoint(const int32 ManifoldPointIndex);
-		void UpdateManifoldPoint(int32 ManifoldPointIndex, const FContactPoint& ContactPoint);
 		void SetActiveContactPoint(const FContactPoint& ContactPoint);
-		void GetWorldSpaceManifoldPoint(const FManifoldPoint& ManifoldPoint, const FVec3& P0, const FRotation3& Q0, const FVec3& P1, const FRotation3& Q1, FVec3& OutContactLocation, FReal& OutContactPhi);
 
 		/**
 		 * @brief Move the current manifold to the previous manifold, and reset the current manifold ready for it to be rebuilt
@@ -547,6 +548,8 @@ namespace Chaos
 		FManifoldPointSavedData ManifoldPointSavedData[MaxManifoldPoints];
 		int32 NumSavedManifoldPoints;
 
+		FRigidTransform3 ShapeWorldTransform0;
+		FRigidTransform3 ShapeWorldTransform1;
 		FRigidTransform3 LastShapeWorldTransform0;
 		FRigidTransform3 LastShapeWorldTransform1;
 		int32 ExpectedNumManifoldPoints;
diff --git a/Engine/Source/Runtime/Experimental/Chaos/Public/Chaos/ContactModification.h b/Engine/Source/Runtime/Experimental/Chaos/Public/Chaos/ContactModification.h
--- a/Engine/Source/Runtime/Experimental/Chaos/Public/Chaos/ContactModification.h
+++ b/Engine/Source/Runtime/Experimental/Chaos/Public/Chaos/ContactModification.h
@@ -91,16 +91,6 @@ namespace Chaos
 		*/
 		void ModifyWorldContactLocations(const FVec3& Location0, const FVec3& Location1, int32 ContactPointIdx);
 
-		/**
-		* Get contact location on each body relative to center of mass.
-		*/
-		void GetCoMContactLocations(int32 ContactPointIdx, FVec3& OutLocation0, FVec3& OutLocation1)  const;
-
-		/**
-		* Modify contact location of each body relative to center of mass. If modifying locations and separation, order of operations should be considered.
-		*/
-		void ModifyCoMContactLocations(const FVec3& Location0, const FVec3& Location1, int32 ContactPointIdx);
-
 		/**
 		* @return Restitution of contact.
 		*/
@@ -122,24 +112,24 @@ namespace Chaos
 		void ModifyRestitutionThreshold(FReal Restitution);
 
 		/**
-		* @return Friction coefficient of contact.
+		* @return Dynamic friction coefficient of contact.
 		*/
-		FReal GetFriction() const;
+		FReal GetDynamicFriction() const;
 
 		/**
-		* Modify friction coefficient of contact.
+		* Modify dynamic friction coefficient of contact.
 		*/
-		void ModifyFriction(FReal Friction);
+		void ModifyDynamicFriction(FReal DynamicFriction);
 
 		/**
-		* @return Angular friction coefficient of contact.
+		* @return Static friction coefficient of contact.
 		*/
-		FReal GetAngularFriction() const;
+		FReal GetStaticFriction() const;
 
 		/**
-		* Modify angular friction coefficient of contact.
+		* Modify static friction coefficient of contact.
 		*/
-		void ModifyAngularFriction(FReal Friction);
+		void ModifyStaticFriction(FReal StaticFriction);
 
 		/**
 		* @return Linear velocity of particle.
@@ -209,6 +199,11 @@ namespace Chaos
 		TVec2<FGeometryParticleHandle*> GetParticlePair() const;
 
 	private:
+		/**
+		 * @brief Update cached shape transforms in the constraint after modifying particle positions
+		*/
+		void UpdateConstraintShapeTransforms();
+
 		FPBDCollisionConstraint* Constraint;
 		FCollisionContactModifier* Modifier;
 	};
