Missing file to make SingleThreadedPhysics work with Chaos

diff --git a/Engine/Source/Runtime/Experimental/Chaos/Private/ChaosModule.cpp b/Engine/Source/Runtime/Experimental/Chaos/Private/ChaosModule.cpp
--- a/Engine/Source/Runtime/Experimental/Chaos/Private/ChaosModule.cpp
+++ b/Engine/Source/Runtime/Experimental/Chaos/Private/ChaosModule.cpp
@@ -3,13 +3,24 @@
 #include "ChaosModule.h"
 #include "CoreMinimal.h"
 #include "Modules/ModuleManager.h"
+#include "Misc/CommandLine.h"
 
 IMPLEMENT_MODULE(FChaosEngineModule, Chaos);
 
+namespace Chaos
+{
+	extern CHAOS_API int GSingleThreadedPhysics;
+}
+
 void FChaosEngineModule::StartupModule()
 {
 	if (FModuleManager::Get().ModuleExists(TEXT("GeometryCore")))
 	{
 		FModuleManager::Get().LoadModule("GeometryCore");
 	}
+
+	if(FParse::Param(FCommandLine::Get(), TEXT("SingleThreadedPhysics")))
+	{
+		Chaos::GSingleThreadedPhysics = 1;
+	}
 }
