Chaos - fix collision debug draw to show contacts at final position rather than initial world space position (knock on from previous change)

diff --git a/Engine/Source/Runtime/Experimental/Chaos/Private/Chaos/ChaosDebugDraw.cpp b/Engine/Source/Runtime/Experimental/Chaos/Private/Chaos/ChaosDebugDraw.cpp
--- a/Engine/Source/Runtime/Experimental/Chaos/Private/Chaos/ChaosDebugDraw.cpp
+++ b/Engine/Source/Runtime/Experimental/Chaos/Private/Chaos/ChaosDebugDraw.cpp
@@ -831,17 +831,17 @@ namespace Chaos
 				{
 					const FConstGenericParticleHandle Particle0 = Contact.Particle[0];
 					const FConstGenericParticleHandle Particle1 = Contact.Particle[1];
-					const FRigidTransform3 WorldCoMTransform0 = FParticleUtilities::GetCoMWorldTransform(Particle0);
-					const FRigidTransform3 WorldCoMTransform1 = FParticleUtilities::GetCoMWorldTransform(Particle1);
+					const FRigidTransform3 WorldActorTransform0 = FParticleUtilities::GetActorWorldTransform(Particle0);
+					const FRigidTransform3 WorldActorTransform1 = FParticleUtilities::GetActorWorldTransform(Particle1);
 
 					// Are we within the region of interest?
 					const FReal Particle0Size = Contact.Particle[0]->HasBounds() ? 0.5f * Contact.Particle[0]->LocalBounds().Extents().Size() : TNumericLimits<FReal>::Max();
 					const FReal Particle1Size = Contact.Particle[1]->HasBounds() ? 0.5f * Contact.Particle[1]->LocalBounds().Extents().Size() : TNumericLimits<FReal>::Max();
-					if (!FDebugDrawQueue::GetInstance().IsInRegionOfInterest(WorldCoMTransform0.GetLocation(), Particle0Size))
+					if (!FDebugDrawQueue::GetInstance().IsInRegionOfInterest(WorldActorTransform0.GetLocation(), Particle0Size))
 					{
 						return;
 					}
-					if (!FDebugDrawQueue::GetInstance().IsInRegionOfInterest(WorldCoMTransform1.GetLocation(), Particle1Size))
+					if (!FDebugDrawQueue::GetInstance().IsInRegionOfInterest(WorldActorTransform1.GetLocation(), Particle1Size))
 					{
 						return;
 					}
@@ -856,9 +856,12 @@ namespace Chaos
 
 						const int32 ContactPlaneOwner = 1;
 						const int32 ContactPointOwner = 1 - ContactPlaneOwner;
+						const FRigidTransform3& PlaneTransform = (ContactPlaneOwner == 0) ? Contact.ImplicitTransform[0] * WorldActorTransform0 : Contact.ImplicitTransform[1] * WorldActorTransform1;
+						const FRigidTransform3& PointTransform = (ContactPlaneOwner == 0) ? Contact.ImplicitTransform[1] * WorldActorTransform1 : Contact.ImplicitTransform[0] * WorldActorTransform0;
+						const FConstGenericParticleHandle PlaneParticle = (ContactPlaneOwner == 0) ? Particle0 : Particle1;
 						const FVec3 PlaneNormal = ManifoldPoint.ContactPoint.Normal;
-						const FVec3 PointLocation = ManifoldPoint.WorldContactPoints[ContactPointOwner];
-						const FVec3 PlaneLocation = ManifoldPoint.WorldContactPoints[ContactPlaneOwner];
+						const FVec3 PointLocation = PointTransform.TransformPosition(ManifoldPoint.ContactPoint.ShapeContactPoints[ContactPointOwner]);
+						const FVec3 PlaneLocation = PlaneTransform.TransformPosition(ManifoldPoint.ContactPoint.ShapeContactPoints[ContactPlaneOwner]);
 						const FVec3 PointPlaneLocation = PointLocation - FVec3::DotProduct(PointLocation - PlaneLocation, PlaneNormal) * PlaneNormal;
 
 						// Dynamic friction, restitution = red
@@ -903,6 +906,12 @@ namespace Chaos
 						FMatrix Axes = FRotationMatrix::MakeFromX(WorldPlaneNormal);
 						FDebugDrawQueue::GetInstance().DrawDebugCircle(WorldPointLocation, 0.5f * Settings.DrawScale * Settings.ContactWidth, 12, DiscColor, false, KINDA_SMALL_NUMBER, Settings.DrawPriority, Settings.LineThickness, Axes.GetUnitAxis(EAxis::Y), Axes.GetUnitAxis(EAxis::Z), false);
 
+						// Previous points
+						const FVec3 WorldPrevPointLocation = SpaceTransform.TransformPosition(ManifoldPoint.WorldContactPoints[ContactPointOwner]);
+						const FVec3 WorldPrevPlaneLocation = SpaceTransform.TransformPosition(ManifoldPoint.WorldContactPoints[ContactPlaneOwner]);
+						FDebugDrawQueue::GetInstance().DrawDebugLine(WorldPrevPointLocation, WorldPointLocation, FColor::White, false, KINDA_SMALL_NUMBER, Settings.DrawPriority, Settings.LineThickness);
+						FDebugDrawQueue::GetInstance().DrawDebugLine(WorldPrevPlaneLocation, WorldPlaneLocation, FColor::White, false, KINDA_SMALL_NUMBER, Settings.DrawPriority, Settings.LineThickness);
+
 						// Whether restored
 						if (Contact.WasManifoldRestored())
 						{
