Chaos - remove unused CCD path

diff --git a/Engine/Source/Runtime/Experimental/Chaos/Private/Chaos/Collision/SolverCollisionContainer.cpp b/Engine/Source/Runtime/Experimental/Chaos/Private/Chaos/Collision/SolverCollisionContainer.cpp
--- a/Engine/Source/Runtime/Experimental/Chaos/Private/Chaos/Collision/SolverCollisionContainer.cpp
+++ b/Engine/Source/Runtime/Experimental/Chaos/Private/Chaos/Collision/SolverCollisionContainer.cpp
@@ -81,7 +81,7 @@ namespace Chaos
 
 			GatherManifoldPoints(Dt);
 
-			// We should try to remove this - the Constraint should need to know about solver objects
+			// We should try to remove this - the Constraint should not need to know about solver objects
 			Constraint->SetSolverBodies(Body0, Body1);
 		}
 
@@ -184,7 +184,6 @@ namespace Chaos
 	void FPBDCollisionSolverContainer::Reset(const int32 MaxCollisions)
 	{
 		CollisionSolvers.Reset(MaxCollisions);
-		SweptCollisionSolvers.Reset();
 	}
 
 	void FPBDCollisionSolverContainer::AddConstraintSolver(FReal Dt, FPBDCollisionConstraint& Constraint, const int32 Particle0Level, const int32 Particle1Level, FSolverBodyContainer& SolverBodyContainer)
@@ -196,12 +195,6 @@ namespace Chaos
 		const int32 SolverCollisionIndex = CollisionSolvers.AddDefaulted();
 
 		CollisionSolvers[SolverCollisionIndex].GatherInput(Dt, Constraint, Particle0Level, Particle1Level, SolverBodyContainer);
-
-		// Keep a list of CCD contacts - they will get resolved before we move on to the main solver loop
-		if (Constraint.GetCCDType() == ECollisionCCDType::Enabled)
-		{
-			SweptCollisionSolvers.Add(SolverCollisionIndex);
-		}
 	}
 
 	void FPBDCollisionSolverContainer::UpdatePositionShockPropagation(const FReal Dt, const int32 It, const int32 NumIts, const int32 BeginIndex, const int32 EndIndex)
@@ -240,14 +233,6 @@ namespace Chaos
 		}
 	}
 
-	void FPBDCollisionSolverContainer::SolveSwept(const FReal Dt)
-	{
-		for (int32 SweptSolverCollisionIndex : SweptCollisionSolvers)
-		{
-			SolveSwept(Dt, *CollisionSolvers[SweptSolverCollisionIndex].GetConstraint(), CollisionSolvers[SweptSolverCollisionIndex].GetSolver());
-		}
-	}
-
 	bool FPBDCollisionSolverContainer::SolvePositionSerial(const FReal Dt, const int32 It, const int32 NumIts, const int32 BeginIndex, const int32 EndIndex)
 	{
 		return SolvePositionImpl(Dt, It, NumIts, BeginIndex, EndIndex, false);
@@ -268,19 +253,6 @@ namespace Chaos
 		return SolveVelocityImpl(Dt, It, NumIts, BeginIndex, EndIndex, true);
 	}
 
-	void FPBDCollisionSolverContainer::SolveSwept(const FReal Dt, FPBDCollisionConstraint& Constraint, FPBDCollisionSolver& CollisionSolver)
-	{
-		// @todo(chaos): Fix CCD. This "solution" just rewinds the CCD object to the first contact and will result in lost momentum
-		if (Constraint.TimeOfImpact < FReal(1))
-		{
-			Collisions::UpdateSwept(Constraint, Dt);
-
-			// Just rewind the body to the impact. This is not good at all, but it will prevent penetration for now
-			FSolverBody& Body0 = CollisionSolver.SolverBody0().SolverBody();
-			Body0.SetP(FMath::Lerp(Body0.X(), Body0.P(), Constraint.TimeOfImpact));
-		}
-	}
-
 	bool FPBDCollisionSolverContainer::SolvePositionImpl(const FReal Dt, const int32 It, const int32 NumIts, const int32 BeginIndex, const int32 EndIndex, const bool bParallel)
 	{
 		SCOPE_CYCLE_COUNTER(STAT_Collisions_Apply);
diff --git a/Engine/Source/Runtime/Experimental/Chaos/Public/Chaos/Collision/SolverCollisionContainer.h b/Engine/Source/Runtime/Experimental/Chaos/Public/Chaos/Collision/SolverCollisionContainer.h
--- a/Engine/Source/Runtime/Experimental/Chaos/Public/Chaos/Collision/SolverCollisionContainer.h
+++ b/Engine/Source/Runtime/Experimental/Chaos/Public/Chaos/Collision/SolverCollisionContainer.h
@@ -38,8 +38,6 @@ namespace Chaos
 		// Add a solver constraint for the specified constraint and gather the required solver data
 		void AddConstraintSolver(FReal Dt, FPBDCollisionConstraint& Constraint, const int32 Particle0Level, const int32 Particle1Level, FSolverBodyContainer& SolverBodyContainer);
 
-		void SolveSwept(const FReal Dt);
-
 		bool SolvePositionSerial(const FReal Dt, const int32 It, const int32 NumIts, const int32 BeginIndex, const int32 EndIndex);
 		bool SolvePositionParallel(const FReal Dt, const int32 It, const int32 NumIts, const int32 BeginIndex, const int32 EndIndex);
 
@@ -51,13 +49,11 @@ namespace Chaos
 	private:
 		void UpdatePositionShockPropagation(const FReal Dt, const int32 It, const int32 NumIts, const int32 BeginIndex, const int32 EndIndex);
 		void UpdateVelocityShockPropagation(const FReal Dt, const int32 It, const int32 NumIts, const int32 BeginIndex, const int32 EndIndex);
-		void SolveSwept(const FReal Dt, FPBDCollisionConstraint& Constraint, FPBDCollisionSolver& CollisionSolver);
 		bool SolvePositionImpl(const FReal Dt, const int32 It, const int32 NumIts, const int32 BeginIndex, const int32 EndIndex, const bool bParallel);
 		bool SolveVelocityImpl(const FReal Dt, const int32 It, const int32 NumIts, const int32 BeginIndex, const int32 EndIndex, const bool bParallel);
 		void ScatterOutputImpl(const FReal Dt, const int32 BeginIndex, const int32 EndIndex, const bool bParallel);
 
 		TArray<FPBDCollisionSolverAdapter> CollisionSolvers;
-		TArray<int32> SweptCollisionSolvers;
 		FReal MaxPushOutVelocity;
 	};
 }
