Chaos - better CalculateTransformedBounds in ImplicitObjectTransformed

diff --git a/Engine/Source/Runtime/Experimental/Chaos/Public/Chaos/GeometryParticles.h b/Engine/Source/Runtime/Experimental/Chaos/Public/Chaos/GeometryParticles.h
--- a/Engine/Source/Runtime/Experimental/Chaos/Public/Chaos/GeometryParticles.h
+++ b/Engine/Source/Runtime/Experimental/Chaos/Public/Chaos/GeometryParticles.h
@@ -78,11 +78,11 @@ namespace Chaos
 
 		void UpdateWorldSpaceState(const FRigidTransform3& WorldTransform, const FVec3& BoundsExpansion)
 		{
-			if (GetGeometry()->HasBoundingBox())
+			LeafWorldTransform = FRigidTransform3::MultiplyNoScale(LeafRelativeTransform, WorldTransform);
+			if ((LeafGeometry != nullptr) && LeafGeometry->HasBoundingBox())
 			{
-				UpdateShapeBounds(WorldTransform, BoundsExpansion);
+				WorldSpaceInflatedShapeBounds = LeafGeometry->CalculateTransformedBounds(LeafWorldTransform).ThickenSymmetrically(BoundsExpansion);;
 			}
-			LeafWorldTransform = LeafRelativeTransform * WorldTransform;
 		}
 
 		// The leaf shape (with transformed and implicit wrapper removed)
diff --git a/Engine/Source/Runtime/Experimental/Chaos/Public/Chaos/ImplicitObjectTransformed.h b/Engine/Source/Runtime/Experimental/Chaos/Public/Chaos/ImplicitObjectTransformed.h
--- a/Engine/Source/Runtime/Experimental/Chaos/Public/Chaos/ImplicitObjectTransformed.h
+++ b/Engine/Source/Runtime/Experimental/Chaos/Public/Chaos/ImplicitObjectTransformed.h
@@ -243,6 +243,12 @@ public:
 
 	virtual const TAABB<T, d> BoundingBox() const override { return MLocalBoundingBox; }
 
+	// Calculate the tight-fitting world-space bounding box
+	virtual FAABB3 CalculateTransformedBounds(const FRigidTransform3& InTransform) const
+	{
+		return MObject->CalculateTransformedBounds(FRigidTransform3::MultiplyNoScale(MTransform ,InTransform));
+	}
+
 	const FReal GetVolume() const
 	{
 		// TODO: More precise volume!
diff --git a/Engine/Source/Runtime/Experimental/Chaos/Public/Chaos/Transform.h b/Engine/Source/Runtime/Experimental/Chaos/Public/Chaos/Transform.h
--- a/Engine/Source/Runtime/Experimental/Chaos/Public/Chaos/Transform.h
+++ b/Engine/Source/Runtime/Experimental/Chaos/Public/Chaos/Transform.h
@@ -207,6 +207,35 @@ namespace Chaos
 				return RotatedNormal;
 			}
 		}
+
+		/**
+		 * @brief Equivalent to (A * B) but assuming both have unit scale
+		*/
+		static TRigidTransform<FReal, 3> MultiplyNoScale(const TRigidTransform<FReal, 3>& A, const TRigidTransform<FReal, 3>& B)
+		{
+			TRigidTransform<FReal, 3> Result;
+
+#if ENABLE_VECTORIZED_TRANSFORM
+			const TransformVectorRegister QuatA = A.Rotation;
+			const TransformVectorRegister QuatB = B.Rotation;
+			const TransformVectorRegister TranslateA = A.Translation;
+			const TransformVectorRegister TranslateB = B.Translation;
+
+			const TransformVectorRegister Rotation = VectorQuaternionMultiply2(QuatB, QuatA);
+			const TransformVectorRegister RotatedTranslate = VectorQuaternionRotateVector(QuatB, TranslateA);
+			const TransformVectorRegister Translation = VectorAdd(RotatedTranslate, TranslateB);
+
+			Result.Rotation = Rotation;
+			Result.Translation = Translation;
+			Result.Scale3D = VectorOne();
+#else
+			Result.Rotation = B.Rotation * A.Rotation;
+			Result.Translation = B.Rotation * A.Translation + B.Translation;
+			Result.Scale3D = FVector::OneVector;
+#endif
+
+			return Result;
+		}
 	};
 }
 
