Chaos - Midphase sphere bounds test (for sphere-sphere)

diff --git a/Engine/Source/Runtime/Experimental/Chaos/Private/Chaos/Collision/ParticlePairMidPhase.cpp b/Engine/Source/Runtime/Experimental/Chaos/Private/Chaos/Collision/ParticlePairMidPhase.cpp
--- a/Engine/Source/Runtime/Experimental/Chaos/Private/Chaos/Collision/ParticlePairMidPhase.cpp
+++ b/Engine/Source/Runtime/Experimental/Chaos/Private/Chaos/Collision/ParticlePairMidPhase.cpp
@@ -113,6 +113,8 @@ namespace Chaos
 		, Shape0(InShape0)
 		, Shape1(InShape1)
 		, ShapePairType(InShapePairType)
+		, SphereBoundsCheckSize(0)
+		, bEnableAABBCheck(false)
 		, bEnableOBBCheck0(false)
 		, bEnableOBBCheck1(false)
 		, bEnableManifoldCheck(false)
@@ -122,9 +124,15 @@ namespace Chaos
 		const bool bIsSphere0 = (ImplicitType0 == ImplicitObjectType::Sphere);
 		const bool bIsSphere1 = (ImplicitType1 == ImplicitObjectType::Sphere);
 
+		bEnableAABBCheck = !bIsSphere0 || !bIsSphere1;
 		bEnableOBBCheck0 = Chaos_Collision_NarrowPhase_AABBBoundsCheck && !bIsSphere0;
 		bEnableOBBCheck1 = Chaos_Collision_NarrowPhase_AABBBoundsCheck && !bIsSphere1;
 		bEnableManifoldCheck = bChaos_Collision_EnableManifoldUpdate && !bIsSphere0 && !bIsSphere1;
+
+		if (bIsSphere0 && bIsSphere1)
+		{
+			SphereBoundsCheckSize = Shape0->GetLeafGeometry()->GetMargin() + Shape1->GetLeafGeometry()->GetMargin();
+		}
 	}
 
 	FSingleShapePairCollisionDetector::~FSingleShapePairCollisionDetector()
@@ -139,13 +147,15 @@ namespace Chaos
 		, Shape0(R.Shape0)
 		, Shape1(R.Shape1)
 		, ShapePairType(R.ShapePairType)
+		, SphereBoundsCheckSize(R.SphereBoundsCheckSize)
+		, bEnableAABBCheck(R.bEnableAABBCheck)
 		, bEnableOBBCheck0(R.bEnableOBBCheck0)
 		, bEnableOBBCheck1(R.bEnableOBBCheck1)
 		, bEnableManifoldCheck(R.bEnableManifoldCheck)
 	{
 	}
 
-	bool FSingleShapePairCollisionDetector::IsUsedSince(const int32 Epoch) const
+	inline bool FSingleShapePairCollisionDetector::IsUsedSince(const int32 Epoch) const
 	{
 		// If we have no constraint it was never used, so this check is always false
 		return (Constraint != nullptr) && (Constraint->GetContainerCookie().LastUsedEpoch >= Epoch);
@@ -164,12 +174,27 @@ namespace Chaos
 		{
 			PHYSICS_CSV_SCOPED_EXPENSIVE(PhysicsVerbose, NarrowPhase_ShapeBounds);
 
-			// World-space expanded bounds check
 			const FAABB3& ShapeWorldBounds0 = Shape0->GetWorldSpaceInflatedShapeBounds();
 			const FAABB3& ShapeWorldBounds1 = Shape1->GetWorldSpaceInflatedShapeBounds();
-			if (!ShapeWorldBounds0.Intersects(ShapeWorldBounds1))
+
+			// World-space expanded bounds check
+			if (bEnableAABBCheck)
+			{
+				if (!ShapeWorldBounds0.Intersects(ShapeWorldBounds1))
+				{
+					return 0;
+				}
+			}
+
+			if (SphereBoundsCheckSize > FReal(0))
 			{
-				return 0;
+				const FVec3 Separation = ShapeWorldBounds0.GetCenter() - ShapeWorldBounds1.GetCenter();
+				const FReal SeparationSq = Separation.SizeSquared();
+				const FReal CullDistanceSq = FMath::Square(CullDistance + SphereBoundsCheckSize);
+				if (SeparationSq > CullDistanceSq)
+				{
+					return 0;
+				}
 			}
 
 			const int32 LastEpoch = MidPhase.GetCollisionAllocator().GetCurrentEpoch() - 1;
@@ -184,7 +209,7 @@ namespace Chaos
 				// we might call the narrow phase one time too many when shapes become separated.
 				const FRigidTransform3& ShapeWorldTransform0 = Shape0->GetLeafWorldTransform();
 				const FRigidTransform3& ShapeWorldTransform1 = Shape1->GetLeafWorldTransform();
-			
+
 				if (bEnableOBBCheck0)
 				{
 					if (!ImplicitOverlapOBBToAABB(Implicit0, Implicit1, ShapeWorldTransform0, ShapeWorldTransform1, CullDistance))
@@ -981,12 +1006,6 @@ namespace Chaos
 		}
 	}
 
-
-	bool FParticlePairMidPhase::IsUsedSince(const int32 Epoch) const
-	{
-		return (LastUsedEpoch >= Epoch);
-	}
-
 	void FParticlePairMidPhase::SetIsSleeping(const bool bInIsSleeping)
 	{
 		// This can be called from two locations:
diff --git a/Engine/Source/Runtime/Experimental/Chaos/Public/Chaos/Collision/ParticlePairMidPhase.h b/Engine/Source/Runtime/Experimental/Chaos/Public/Chaos/Collision/ParticlePairMidPhase.h
--- a/Engine/Source/Runtime/Experimental/Chaos/Public/Chaos/Collision/ParticlePairMidPhase.h
+++ b/Engine/Source/Runtime/Experimental/Chaos/Public/Chaos/Collision/ParticlePairMidPhase.h
@@ -103,6 +103,8 @@ namespace Chaos
 		const FPerShapeData* Shape0;
 		const FPerShapeData* Shape1;
 		EContactShapesType ShapePairType;
+		FReal SphereBoundsCheckSize;
+		bool bEnableAABBCheck;
 		bool bEnableOBBCheck0;
 		bool bEnableOBBCheck1;
 		bool bEnableManifoldCheck;
@@ -261,25 +263,29 @@ namespace Chaos
 		~FParticlePairMidPhase();
 
 
-		FGeometryParticleHandle* GetParticle0() { return Particle0; }
+		inline FGeometryParticleHandle* GetParticle0() { return Particle0; }
 
-		FGeometryParticleHandle* GetParticle1() { return Particle1; }
+		inline FGeometryParticleHandle* GetParticle1() { return Particle1; }
 
-		const FCollisionParticlePairKey& GetKey() const { return Key; }
+		inline const FCollisionParticlePairKey& GetKey() const { return Key; }
 
-		FCollisionConstraintAllocator& GetCollisionAllocator() { return *CollisionAllocator; }
+		inline FCollisionConstraintAllocator& GetCollisionAllocator() { return *CollisionAllocator; }
 
-		bool IsValid() const { return (Particle0 != nullptr) && (Particle1 != nullptr); }
+		inline bool IsValid() const { return (Particle0 != nullptr) && (Particle1 != nullptr); }
 
 		/**
 		 * @brief Have we run collision detection since this Epoch (inclusive)
 		*/
-		bool IsUsedSince(const int32 Epoch) const;
+		inline bool IsUsedSince(const int32 Epoch) const
+		{
+			return (LastUsedEpoch >= Epoch);
+		}
+
 
 		/**
 		 * @brief Whether the particle pair is sleeping and therefore contacts should not be culled (they will be reused on wake)
 		*/
-		bool IsSleeping() const { return bIsSleeping; }
+		inline bool IsSleeping() const { return bIsSleeping; }
 
 		/**
 		 * @brief Update the sleeping state
