Chaos - Add cull distance to low-level Sphere-Sphere collision test

diff --git a/Engine/Source/Runtime/Experimental/Chaos/Private/Chaos/Collision/ContactPointsMiscShapes.cpp b/Engine/Source/Runtime/Experimental/Chaos/Private/Chaos/Collision/ContactPointsMiscShapes.cpp
--- a/Engine/Source/Runtime/Experimental/Chaos/Private/Chaos/Collision/ContactPointsMiscShapes.cpp
+++ b/Engine/Source/Runtime/Experimental/Chaos/Private/Chaos/Collision/ContactPointsMiscShapes.cpp
@@ -192,7 +192,7 @@ namespace Chaos
 	}
 
 
-	FContactPoint SphereSphereContactPoint(const TSphere<FReal, 3>& Sphere1, const FRigidTransform3& Sphere1Transform, const TSphere<FReal, 3>& Sphere2, const FRigidTransform3& Sphere2Transform, const FReal ShapePadding)
+	FContactPoint SphereSphereContactPoint(const TSphere<FReal, 3>& Sphere1, const FRigidTransform3& Sphere1Transform, const TSphere<FReal, 3>& Sphere2, const FRigidTransform3& Sphere2Transform, const FReal CullDistance, const FReal ShapePadding)
 	{
 		FContactPoint Result;
 
@@ -203,17 +203,21 @@ namespace Chaos
 		const FVec3 Center1 = Sphere1Transform.TransformPosition(Sphere1.GetCenter());
 		const FVec3 Center2 = Sphere2Transform.TransformPosition(Sphere2.GetCenter());
 		const FVec3 Direction = Center1 - Center2;
-		const FReal Size = Direction.Size();
-		const FVec3 Normal = Size > SMALL_NUMBER ? Direction / Size : FVec3(0, 0, 1);
-		const FReal NewPhi = Size - (R1 + R2);
-
-		Result.ShapeContactPoints[0] = Sphere1.GetCenter() - Sphere1Transform.InverseTransformVector(R1 * Normal);
-		Result.ShapeContactPoints[1] = Sphere2.GetCenter() + Sphere2Transform.InverseTransformVector(R2 * Normal);
-		Result.ShapeContactNormal = Sphere2Transform.InverseTransformVector(Normal);
-		Result.Phi = NewPhi;
-		Result.Normal = Normal;
-		Result.Location = Center2 + R2 * Result.Normal;
+		const FReal SizeSq = Direction.SizeSquared();
+		const FReal CullDistanceSq = FMath::Square(CullDistance + R1 + R2);
+		if (SizeSq < CullDistanceSq)
+		{
+			const FReal Size = FMath::Sqrt(SizeSq);
+			const FVec3 Normal = Size > SMALL_NUMBER ? Direction / Size : FVec3(0, 0, 1);
+			const FReal NewPhi = Size - (R1 + R2);
 
+			Result.ShapeContactPoints[0] = Sphere1.GetCenter() - Sphere1Transform.InverseTransformVector(R1 * Normal);
+			Result.ShapeContactPoints[1] = Sphere2.GetCenter() + Sphere2Transform.InverseTransformVector(R2 * Normal);
+			Result.ShapeContactNormal = Sphere2Transform.InverseTransformVector(Normal);
+			Result.Phi = NewPhi;
+			Result.Normal = Normal;
+			Result.Location = Center2 + R2 * Result.Normal;
+		}
 		return Result;
 	}
 
@@ -288,7 +292,7 @@ namespace Chaos
 			Result.ShapeContactNormal = BTransform.InverseTransformVector(Normal);
 			Result.Phi = NewPhi;
 			Result.Normal = Normal;
-			Result.Location = 0.5f * (LocationA + LocationB);
+			Result.Location = FReal(0.5) * (LocationA + LocationB);
 		}
 
 		return Result;
@@ -418,7 +422,7 @@ namespace Chaos
 			Result.ShapeContactNormal = BTransform.InverseTransformVector(Normal);
 			Result.Phi = NewPhi;
 			Result.Normal = Normal;
-			Result.Location = 0.5f * (LocationA + LocationB);
+			Result.Location = FReal(0.5) * (LocationA + LocationB);
 		}
 
 		return Result;
diff --git a/Engine/Source/Runtime/Experimental/Chaos/Private/Chaos/CollisionOneShotManifolds.cpp b/Engine/Source/Runtime/Experimental/Chaos/Private/Chaos/CollisionOneShotManifolds.cpp
--- a/Engine/Source/Runtime/Experimental/Chaos/Private/Chaos/CollisionOneShotManifolds.cpp
+++ b/Engine/Source/Runtime/Experimental/Chaos/Private/Chaos/CollisionOneShotManifolds.cpp
@@ -214,8 +214,8 @@ namespace Chaos
 				FVec3 EdgePos2 = Box2Transform.TransformPosition(ShapeEdgePos2);
 				FReal EdgePhi = FVec3::DotProduct(EdgePos1 - EdgePos2, GJKContactPoint.Normal);
 
-				GJKContactPoint.ShapeContactPoints[0] = EdgePos1;
-				GJKContactPoint.ShapeContactPoints[1] = EdgePos2;
+				GJKContactPoint.ShapeContactPoints[0] = ShapeEdgePos1;
+				GJKContactPoint.ShapeContactPoints[1] = ShapeEdgePos2;
 				GJKContactPoint.Phi = EdgePhi;
 				GJKContactPoint.Location = 0.5f * (EdgePos1 + EdgePos2);
 
diff --git a/Engine/Source/Runtime/Experimental/Chaos/Private/Chaos/CollisionOneShotManifoldsMiscShapes.cpp b/Engine/Source/Runtime/Experimental/Chaos/Private/Chaos/CollisionOneShotManifoldsMiscShapes.cpp
--- a/Engine/Source/Runtime/Experimental/Chaos/Private/Chaos/CollisionOneShotManifoldsMiscShapes.cpp
+++ b/Engine/Source/Runtime/Experimental/Chaos/Private/Chaos/CollisionOneShotManifoldsMiscShapes.cpp
@@ -43,7 +43,7 @@ namespace Chaos
 			// @todo(chaos): support manifold maintenance
 			Constraint.ResetManifold();
 
-			FContactPoint ContactPoint = SphereSphereContactPoint(SphereA, SphereATransform, SphereB, SphereBTransform, Constraint.Manifold.RestitutionPadding);
+			FContactPoint ContactPoint = SphereSphereContactPoint(SphereA, SphereATransform, SphereB, SphereBTransform, Constraint.GetCullDistance(), Constraint.Manifold.RestitutionPadding);
 			if (ContactPoint.Phi < Constraint.GetCullDistance())
 			{
 				Constraint.AddOneshotManifoldContact(ContactPoint);
diff --git a/Engine/Source/Runtime/Experimental/Chaos/Private/Chaos/CollisionResolution.cpp b/Engine/Source/Runtime/Experimental/Chaos/Private/Chaos/CollisionResolution.cpp
--- a/Engine/Source/Runtime/Experimental/Chaos/Private/Chaos/CollisionResolution.cpp
+++ b/Engine/Source/Runtime/Experimental/Chaos/Private/Chaos/CollisionResolution.cpp
@@ -739,7 +739,7 @@ namespace Chaos
 			}
 			else
 			{
-				UpdateContactPoint(Constraint, SphereSphereContactPoint(Sphere1, Sphere1Transform, Sphere2, Sphere2Transform, Constraint.Manifold.RestitutionPadding), Dt);
+				UpdateContactPoint(Constraint, SphereSphereContactPoint(Sphere1, Sphere1Transform, Sphere2, Sphere2Transform, Constraint.GetCullDistance(), Constraint.Manifold.RestitutionPadding), Dt);
 			}
 		}
 
diff --git a/Engine/Source/Runtime/Experimental/Chaos/Public/Chaos/Collision/ContactPoint.h b/Engine/Source/Runtime/Experimental/Chaos/Public/Chaos/Collision/ContactPoint.h
--- a/Engine/Source/Runtime/Experimental/Chaos/Public/Chaos/Collision/ContactPoint.h
+++ b/Engine/Source/Runtime/Experimental/Chaos/Public/Chaos/Collision/ContactPoint.h
@@ -25,7 +25,7 @@ namespace Chaos
 	class CHAOS_API FContactPoint
 	{
 	public:
-		// Shape-space contact points on the two bodies, without the margin added (i.e., contact point on the core shape)
+		// Shape-space contact points on the two bodies
 		FVec3 ShapeContactPoints[2];
 
 		// Shape-space contact normal on the second shape with direction that points away from shape 1
diff --git a/Engine/Source/Runtime/Experimental/Chaos/Public/Chaos/Collision/ContactPointsMiscShapes.h b/Engine/Source/Runtime/Experimental/Chaos/Public/Chaos/Collision/ContactPointsMiscShapes.h
--- a/Engine/Source/Runtime/Experimental/Chaos/Public/Chaos/Collision/ContactPointsMiscShapes.h
+++ b/Engine/Source/Runtime/Experimental/Chaos/Public/Chaos/Collision/ContactPointsMiscShapes.h
@@ -26,7 +26,7 @@ namespace Chaos
 	template <typename GeometryA, typename GeometryB>
 	void GJKImplicitManifold(const FImplicitObject& A, const FRigidTransform3& ATransform, const GeometryB& B, const FRigidTransform3& BTransform, const FReal CullDistance, const FReal ShapePadding, TArray<FContactPoint>& ContactPoints);
 
-	FContactPoint SphereSphereContactPoint(const TSphere<FReal, 3>& Sphere1, const FRigidTransform3& Sphere1Transform, const TSphere<FReal, 3>& Sphere2, const FRigidTransform3& Sphere2Transform, const FReal ShapePadding);
+	FContactPoint SphereSphereContactPoint(const TSphere<FReal, 3>& Sphere1, const FRigidTransform3& Sphere1Transform, const TSphere<FReal, 3>& Sphere2, const FRigidTransform3& Sphere2Transform, const FReal CullDistance, const FReal ShapePadding);
 	FContactPoint SpherePlaneContactPoint(const TSphere<FReal, 3>& Sphere, const FRigidTransform3& SphereTransform, const TPlane<FReal, 3>& Plane, const FRigidTransform3& PlaneTransform, const FReal ShapePadding);
 	FContactPoint SphereBoxContactPoint(const TSphere<FReal, 3>& Sphere, const FRigidTransform3& SphereTransform, const FImplicitBox3& Box, const FRigidTransform3& BoxTransform, const FReal ShapePadding);
 	FContactPoint SphereCapsuleContactPoint(const TSphere<FReal, 3>& A, const FRigidTransform3& ATransform, const FCapsule& B, const FRigidTransform3& BTransform, const FReal ShapePadding);
