Chaos - Remove MidPhase map from coliison constraints
- Cache ParticleCollisions array index on MidPhase

diff --git a/Engine/Source/Runtime/Experimental/Chaos/Private/Chaos/Collision/ParticleCollisions.cpp b/Engine/Source/Runtime/Experimental/Chaos/Private/Chaos/Collision/ParticleCollisions.cpp
--- a/Engine/Source/Runtime/Experimental/Chaos/Private/Chaos/Collision/ParticleCollisions.cpp
+++ b/Engine/Source/Runtime/Experimental/Chaos/Private/Chaos/Collision/ParticleCollisions.cpp
@@ -14,28 +14,42 @@ namespace Chaos
 	{
 	}
 
-	void FParticleCollisions::Reset()
+	void FParticleCollisions::AddMidPhase(FGeometryParticleHandle* InParticle, FParticlePairMidPhase* InMidPhase)
 	{
-		ParticlePairs.Reset();
-	}
+		// Double add?
+		check(InMidPhase->GetParticleCollisionsIndex(InParticle) == INDEX_NONE);
 
-	void FParticleCollisions::AddParticlePair(FParticlePairMidPhase* MidPhase)
-	{
-		ParticlePairs.Add(MidPhase);
+		const int32 Index = MidPhases.Emplace(typename FContainerType::ElementType(InMidPhase->GetKey().GetKey(), InMidPhase));
+
+		// Set the cookie for the midphase for fast removal
+		InMidPhase->SetParticleCollisionsIndex(InParticle, Index);
 	}
 
-	void FParticleCollisions::RemoveParticlePair(FParticlePairMidPhase* MidPhase)
+	void FParticleCollisions::RemoveMidPhase(FGeometryParticleHandle* InParticle, FParticlePairMidPhase* InMidPhase)
 	{
-		// @todo(chaos): store the index on the midphase (need one index for each particle)
-		ParticlePairs.Remove(MidPhase);
+		// Retrieve the cookie from the midphase
+		const int32 Index = InMidPhase->GetParticleCollisionsIndex(InParticle);
+
+		// Double remove?
+		check(Index != INDEX_NONE);
+		
+		MidPhases.RemoveAtSwap(Index);
+		InMidPhase->SetParticleCollisionsIndex(InParticle, INDEX_NONE);
+
+		// Update the cookie of the midphase that moved
+		if (Index < MidPhases.Num())
+		{
+			MidPhases[Index].Value->SetParticleCollisionsIndex(InParticle, Index);
+		}
 	}
 
+
 	void FParticleCollisions::VisitCollisions(const FPBDCollisionVisitor& Visitor) const
 	{
-		for (FParticlePairMidPhase* MidPhase : ParticlePairs)
-		{
-			MidPhase->VisitCollisions(Visitor);
-		}
+		VisitConstMidPhases([&Visitor](const FParticlePairMidPhase& MidPhase)
+			{
+				MidPhase.VisitCollisions(Visitor);
+			});
 	}
 
 }
\ No newline at end of file
diff --git a/Engine/Source/Runtime/Experimental/Chaos/Private/Chaos/Collision/ParticlePairMidPhase.cpp b/Engine/Source/Runtime/Experimental/Chaos/Private/Chaos/Collision/ParticlePairMidPhase.cpp
--- a/Engine/Source/Runtime/Experimental/Chaos/Private/Chaos/Collision/ParticlePairMidPhase.cpp
+++ b/Engine/Source/Runtime/Experimental/Chaos/Private/Chaos/Collision/ParticlePairMidPhase.cpp
@@ -699,6 +699,8 @@ namespace Chaos
 		: Particle0(InParticle0)
 		, Particle1(InParticle1)
 		, Key(InKey)
+		, ParticleCollisionsIndex0(INDEX_NONE)
+		, ParticleCollisionsIndex1(INDEX_NONE)
 		, ShapePairDetectors()
 		, MultiShapePairDetectors()
 		, CollisionAllocator(&InCollisionAllocator)
diff --git a/Engine/Source/Runtime/Experimental/Chaos/Public/Chaos/Collision/CollisionConstraintAllocator.h b/Engine/Source/Runtime/Experimental/Chaos/Public/Chaos/Collision/CollisionConstraintAllocator.h
--- a/Engine/Source/Runtime/Experimental/Chaos/Public/Chaos/Collision/CollisionConstraintAllocator.h
+++ b/Engine/Source/Runtime/Experimental/Chaos/Public/Chaos/Collision/CollisionConstraintAllocator.h
@@ -165,12 +165,10 @@ namespace Chaos
 		*/
 		FParticlePairMidPhase* GetParticlePairMidPhase(FGeometryParticleHandle* Particle0, FGeometryParticleHandle* Particle1)
 		{
-			FCollisionParticlePairKey Key = FCollisionParticlePairKey(Particle0, Particle1);
-
-			FParticlePairMidPhase* MidPhase = FindParticlePairMidPhaseImpl(Key);
+			FParticlePairMidPhase* MidPhase = FindParticlePairMidPhaseImpl(Particle0, Particle1);
 			if (MidPhase == nullptr)
 			{
-				MidPhase = CreateParticlePairMidPhase(Particle0, Particle1, Key);
+				MidPhase = CreateParticlePairMidPhase(Particle0, Particle1);
 			}
 
 			return MidPhase;
@@ -181,8 +179,7 @@ namespace Chaos
 		*/
 		FParticlePairMidPhase* FindParticlePairMidPhase(FGeometryParticleHandle* Particle0, FGeometryParticleHandle* Particle1)
 		{
-			FCollisionParticlePairKey Key = FCollisionParticlePairKey(Particle0, Particle1);
-			return FindParticlePairMidPhaseImpl(Key);
+			return FindParticlePairMidPhaseImpl(Particle0, Particle1);
 		}
 
 		/**
@@ -275,11 +272,10 @@ namespace Chaos
 			// Lop over all particle pairs involving this particle.
 			// Tell each Particle Pair MidPhase that one of its particles is gone. 
 			// It will get pruned at the next collision detection phase.
-			FParticleCollisions& ParticleCollisions = Particle->ParticleCollisions();
-			for (FParticlePairMidPhase* MidPhase : ParticleCollisions.GetParticlePairs())
-			{
-				MidPhase->DetachParticle(Particle);
-			}
+			Particle->ParticleCollisions().VisitMidPhases([Particle](FParticlePairMidPhase& MidPhase)
+				{
+					MidPhase.DetachParticle(Particle);
+				});
 		}
 
 		/**
@@ -287,9 +283,8 @@ namespace Chaos
 		*/
 		void VisitCollisions(const TFunction<void(const FPBDCollisionConstraint*)>& Visitor) const
 		{
-			for (const auto& KVP : ParticlePairMidPhases)
+			for (const auto& MidPhase : ParticlePairMidPhases)
 			{
-				const TUniquePtr<FParticlePairMidPhase>& MidPhase = KVP.Value;
 				MidPhase->VisitCollisions(Visitor);
 			}
 		}
@@ -307,18 +302,19 @@ namespace Chaos
 			PruneExpiredMidPhases();
 		}
 
-		FParticlePairMidPhase* FindParticlePairMidPhaseImpl(const FCollisionParticlePairKey& Key)
+		FParticlePairMidPhase* FindParticlePairMidPhaseImpl(FGeometryParticleHandle* Particle0, FGeometryParticleHandle* Particle1)
 		{
-			TUniquePtr<FParticlePairMidPhase>* ExistingMidPhase = ParticlePairMidPhases.Find(Key.GetKey());
-			if (ExistingMidPhase != nullptr)
-			{
-				return (*ExistingMidPhase).Get();
-			}
-			return nullptr;
+			// Find the existing midphase from one of the particle's lists of midphases
+			FGeometryParticleHandle* SearchParticle = (Particle0->ParticleCollisions().Num() < Particle1->ParticleCollisions().Num()) ? Particle0 : Particle1;
+
+			const FCollisionParticlePairKey Key = FCollisionParticlePairKey(Particle0, Particle1);
+			return SearchParticle->ParticleCollisions().FindMidPhase(Key.GetKey());
 		}
 
-		FParticlePairMidPhase* CreateParticlePairMidPhase(FGeometryParticleHandle* Particle0, FGeometryParticleHandle* Particle1, const FCollisionParticlePairKey& Key)
+		FParticlePairMidPhase* CreateParticlePairMidPhase(FGeometryParticleHandle* Particle0, FGeometryParticleHandle* Particle1)
 		{
+			const FCollisionParticlePairKey Key = FCollisionParticlePairKey(Particle0, Particle1);
+
 			// We enqueue a raw pointer and wrap it in a UniquePtr later
 			FParticlePairMidPhase* MidPhase = new FParticlePairMidPhase(Particle0, Particle1, Key, *this);
 
@@ -336,30 +332,26 @@ namespace Chaos
 
 				FCollisionParticlePairKey Key = FCollisionParticlePairKey(Particle0, Particle1);
 
-				ParticlePairMidPhases.Add(Key.GetKey(), TUniquePtr<FParticlePairMidPhase>(MidPhase));
-
-				Particle0->ParticleCollisions().AddParticlePair(MidPhase);
-				Particle1->ParticleCollisions().AddParticlePair(MidPhase);
+				ParticlePairMidPhases.Emplace(TUniquePtr<FParticlePairMidPhase>(MidPhase));
+				Particle0->ParticleCollisions().AddMidPhase(Particle0, MidPhase);
+				Particle1->ParticleCollisions().AddMidPhase(Particle1, MidPhase);
 			}
 		}
 
-		void DestroyParticlePairMidPhase(FParticlePairMidPhase* MidPhase)
+		void DetachParticlePairMidPhase(FParticlePairMidPhase* MidPhase)
 		{
 			// Remove ParticlePairMidPhase from each Particles list of collisions
 			// NOTE: One or both particles may have been destroyed in which case it will
 			// have been set to null on the midphase.
 			if (FGeometryParticleHandle* Particle0 = MidPhase->GetParticle0())
 			{
-				Particle0->ParticleCollisions().RemoveParticlePair(MidPhase);
+				Particle0->ParticleCollisions().RemoveMidPhase(Particle0, MidPhase);
 			}
 
 			if (FGeometryParticleHandle* Particle1 = MidPhase->GetParticle1())
 			{
-				Particle1->ParticleCollisions().RemoveParticlePair(MidPhase);
+				Particle1->ParticleCollisions().RemoveMidPhase(Particle1, MidPhase);
 			}
-
-			// Destroy the midphase and all its constraints
-			ParticlePairMidPhases.Remove(MidPhase->GetKey().GetKey());
 		}
 
 		void PruneExpiredMidPhases()
@@ -369,21 +361,18 @@ namespace Chaos
 			// Determine which particle pairs are no longer overlapping
 			// Prune all pairs which wer not updated this tick as part of the collision
 			// detection loop and are not asleep
-			// @todo(chaos): is there any value in waiting a few frames before destroying?
-			TArray<FParticlePairMidPhase*> Pruned;
-			for (auto& KVP : ParticlePairMidPhases)
+			int32 NumParticlePairMidPhases = ParticlePairMidPhases.Num();
+			for (int32 Index = 0; Index < NumParticlePairMidPhases; ++Index)
 			{
-				TUniquePtr<FParticlePairMidPhase>& MidPhase = KVP.Value;
+				TUniquePtr<FParticlePairMidPhase>& MidPhase = ParticlePairMidPhases[Index];
 				if ((MidPhase != nullptr) && !MidPhase->IsUsedSince(CurrentEpoch) && !MidPhase->IsSleeping())
 				{
-					Pruned.Add(MidPhase.Get());
+					DetachParticlePairMidPhase(MidPhase.Get());
+					ParticlePairMidPhases.RemoveAtSwap(Index);
+					--Index;
+					--NumParticlePairMidPhases;
 				}
 			}
-
-			for (FParticlePairMidPhase* MidPhase : Pruned)
-			{
-				DestroyParticlePairMidPhase(MidPhase);
-			}
 		}
 
 		void ProcessNewConstraints()
@@ -416,7 +405,7 @@ namespace Chaos
 		}
 
 		// All of the overlapping particle pairs in the scene
-		TMap<uint64, TUniquePtr<FParticlePairMidPhase>> ParticlePairMidPhases;
+		TArray<TUniquePtr<FParticlePairMidPhase>> ParticlePairMidPhases;
 
 		// The active constraints (added or recovered this tick)
 		TArray<FPBDCollisionConstraint*> ActiveConstraints;
diff --git a/Engine/Source/Runtime/Experimental/Chaos/Public/Chaos/Collision/ParticleCollisions.h b/Engine/Source/Runtime/Experimental/Chaos/Public/Chaos/Collision/ParticleCollisions.h
--- a/Engine/Source/Runtime/Experimental/Chaos/Public/Chaos/Collision/ParticleCollisions.h
+++ b/Engine/Source/Runtime/Experimental/Chaos/Public/Chaos/Collision/ParticleCollisions.h
@@ -7,36 +7,110 @@
 #include "Chaos/Collision/CollisionVisitor.h"
 #include "Chaos/ParticleHandleFwd.h"
 
-
 namespace Chaos
 {
 	class FParticlePairMidPhase;
 
+
 	/**
 	 * @brief Knows about all the collisions detectors associated with a particular particle.
 	 * Used when particles are destroyed to remove perisstent collisions from the system, or
 	 * when Islands are woken to restore the collisions.
+	 * 
 	*/
 	class CHAOS_API FParticleCollisions
 	{
 	public:
+		// In a mostly stationary scene the choice of container doesn't matter much. In a highly dynamic
+		// scene it matters a lot. This was tested with TMap, TSortedMap, and TArray. TArray was better in all cases.
+		// We store the list of overlapping particle pairs for every particle. This list can get quite
+		// large for world objects. We need to make sure that Add / Remove are not O(N).
+		// We also search this list to find an existing MidPhase for a specific particle pair. However
+		// we always search the particle with the fewest overlaps, which will be a dynamic particle
+		// that should not have too many contacts. 
+		using FContainerType = TArray<TPair<uint64, FParticlePairMidPhase*>>;
+
 		FParticleCollisions();
 		~FParticleCollisions();
 
-		TArrayView<FParticlePairMidPhase*> GetParticlePairs()
+		int32 Num() const 
+		{ 
+			return MidPhases.Num();
+		}
+
+		/**
+		 * @brief Clear the list of midphases. Only for use in shutdown.
+		*/
+		void Reset()
 		{
-			return MakeArrayView(ParticlePairs);
+			MidPhases.Reset();
 		}
 
-		void Reset();
+		/**
+		 * @brief Add a mid phase to the list
+		 * We are passing the particle and midphase here rather than just the key because we store a cookie
+		 * on the midphase that we want to retrieve, and it has one cookie per particle. 
+		 * This could probably be cleaned up a bit...
+		*/
+		void AddMidPhase(FGeometryParticleHandle* InParticle, FParticlePairMidPhase* InMidPhase);
+
+		/**
+		 * @brief Remove a mid phase
+		*/
+		void RemoveMidPhase(FGeometryParticleHandle* InParticle, FParticlePairMidPhase* InMidPhase);
 
-		void AddParticlePair(FParticlePairMidPhase* MidPhase);
-		void RemoveParticlePair(FParticlePairMidPhase* MidPhase);
+		/**
+		 * @brief Find the mid phase with the matching key
+		 * @param InKey The internal key from a FCollisionParticlePairKey
+		 * @todo(chaos): we should use FCollisionParticlePairKey here
+		*/
+		inline FParticlePairMidPhase* FindMidPhase(const uint64 InKey)
+		{
+			for (int32 Index = 0; Index < MidPhases.Num(); ++Index)
+			{
+				if (MidPhases[Index].Key == InKey)
+				{
+					return MidPhases[Index].Value;
+				}
+			}
+			return nullptr;
+		}
+
+		/**
+		 * @brief Visit all of the midphases on the particle and call the specified function
+		 * @note Do not call RemoveMidPhase from the visitor
+		 * 
+		 * Lambda must have signature void(FParticlePairMidPhase&)
+		*/
+		template<typename TLambda>
+		void VisitMidPhases(const TLambda& Lambda)
+		{
+			for (int32 Index = 0; Index < MidPhases.Num(); ++Index)
+			{
+				Lambda(*MidPhases[Index].Value);
+			}
+		}
+
+		/**
+		 * @brief Visit all of the midphases on the particle and call the specified function
+		 * @note Do not call RemoveMidPhase from the visitor
+		 *
+		 * Lambda must have signature void(const FParticlePairMidPhase&)
+		*/
+		template<typename TLambda>
+		void VisitConstMidPhases(const TLambda& Lambda) const
+		{
+			for (int32 Index = 0; Index < MidPhases.Num(); ++Index)
+			{
+				Lambda(*MidPhases[Index].Value);
+			}
+		}
 
 		void VisitCollisions(const FPBDCollisionVisitor& Visitor) const;
 
+
 	private:
-		TArray<FParticlePairMidPhase*> ParticlePairs;
+		FContainerType MidPhases;
 	};
 
 }
\ No newline at end of file
diff --git a/Engine/Source/Runtime/Experimental/Chaos/Public/Chaos/Collision/ParticlePairMidPhase.h b/Engine/Source/Runtime/Experimental/Chaos/Public/Chaos/Collision/ParticlePairMidPhase.h
--- a/Engine/Source/Runtime/Experimental/Chaos/Public/Chaos/Collision/ParticlePairMidPhase.h
+++ b/Engine/Source/Runtime/Experimental/Chaos/Public/Chaos/Collision/ParticlePairMidPhase.h
@@ -328,6 +328,38 @@ namespace Chaos
 		*/
 		void VisitCollisions(const FPBDCollisionVisitor& Visitor) const;
 
+		/**
+		 * @brief Cookie for use by FParticleCollisions
+		*/
+		int32 GetParticleCollisionsIndex(FGeometryParticleHandle* InParticle) const
+		{
+			check((InParticle == Particle0) || (InParticle == Particle1));
+			if (InParticle == Particle0)
+			{
+				return ParticleCollisionsIndex0;
+			}
+			else
+			{
+				return ParticleCollisionsIndex1;
+			}
+		}
+
+		/**
+		 * @brief Cookie for use by FParticleCollisions
+		*/
+		void SetParticleCollisionsIndex(FGeometryParticleHandle* InParticle, const int32 InIndex)
+		{
+			check((InParticle == Particle0) || (InParticle == Particle1));
+			if (InParticle == Particle0)
+			{
+				ParticleCollisionsIndex0 = InIndex;
+			}
+			else
+			{
+				ParticleCollisionsIndex1 = InIndex;
+			}
+		}
+
 	private:
 		/**
 		 * @brief Set up the midphase based on the SHapesArrays of the two particles
@@ -373,6 +405,10 @@ namespace Chaos
 		FGeometryParticleHandle* Particle1;
 		FCollisionParticlePairKey Key;
 
+		// Indices into the arrays of collisions on the particles. This is a cookie for use by FParticleCollisions
+		int32 ParticleCollisionsIndex0;
+		int32 ParticleCollisionsIndex1;
+
 		TArray<FSingleShapePairCollisionDetector, TInlineAllocator<1>> ShapePairDetectors;
 		TArray<FMultiShapePairCollisionDetector> MultiShapePairDetectors;
 		FCollisionConstraintAllocator* CollisionAllocator;
