Chaos - Move collision friction data copy to a better spot (scatter rather than gather)

diff --git a/Engine/Source/Runtime/Experimental/Chaos/Private/Chaos/Collision/PBDCollisionConstraint.cpp b/Engine/Source/Runtime/Experimental/Chaos/Private/Chaos/Collision/PBDCollisionConstraint.cpp
--- a/Engine/Source/Runtime/Experimental/Chaos/Private/Chaos/Collision/PBDCollisionConstraint.cpp
+++ b/Engine/Source/Runtime/Experimental/Chaos/Private/Chaos/Collision/PBDCollisionConstraint.cpp
@@ -144,6 +144,8 @@ namespace Chaos
 		, TimeOfImpact(0)
 		, CCDType(ECollisionCCDType::Disabled)
 		, Stiffness(FReal(1))
+		, ManifoldPoints()
+		, ManifoldPointsSavedData()
 		, CullDistance(TNumericLimits<FReal>::Max())
 		, CollisionMargins{ 0, 0 }
 		, CollisionTolerance(0)
@@ -151,8 +153,6 @@ namespace Chaos
 		, bUseIncrementalManifold(false)
 		, SolverBodies{ nullptr, nullptr }
 		, GJKWarmStartData()
-		, ManifoldPointSavedData{}
-		, NumSavedManifoldPoints(0)
 		, ShapeWorldTransform0()
 		, ShapeWorldTransform1()
 		, LastShapeWorldTransform0()
@@ -182,6 +182,8 @@ namespace Chaos
 		, TimeOfImpact(0)
 		, CCDType(ECollisionCCDType::Disabled)
 		, Stiffness(FReal(1))
+		, ManifoldPoints()
+		, ManifoldPointsSavedData()
 		, CullDistance(TNumericLimits<FReal>::Max())
 		, CollisionMargins{ 0, 0 }
 		, CollisionTolerance(0)
@@ -189,8 +191,6 @@ namespace Chaos
 		, bUseIncrementalManifold(false)
 		, SolverBodies{ nullptr, nullptr }
 		, GJKWarmStartData()
-		, ManifoldPointSavedData{}
-		, NumSavedManifoldPoints(0)
 		, ShapeWorldTransform0()
 		, ShapeWorldTransform1()
 		, LastShapeWorldTransform0()
@@ -530,7 +530,7 @@ namespace Chaos
 
 	void FPBDCollisionConstraint::ResetManifold()
 	{
-		NumSavedManifoldPoints = 0;
+		ManifoldPointsSavedData.Reset();
 		ResetActiveManifoldContacts();
 	}
 
@@ -542,26 +542,6 @@ namespace Chaos
 		bWasManifoldRestored = false;
 	}
 
-	void FPBDCollisionConstraint::SaveManifold()
-	{
-		check(ManifoldPoints.Num() <= MaxManifoldPoints);
-
-		// Save off the previous data for use by static friction
-		NumSavedManifoldPoints = 0;
-		const int32 NumManifoldPoints = ManifoldPoints.Num();
-		for (int32 ManifoldPointIndex = 0; ManifoldPointIndex < NumManifoldPoints; ++ManifoldPointIndex)
-		{
-			// Only save points that actually produced a response and are holding static friction
-			const FManifoldPoint& ManifoldPoint = ManifoldPoints[ManifoldPointIndex];
-			if (ManifoldPoint.bInsideStaticFrictionCone && !ManifoldPoint.NetPushOut.IsNearlyZero())
-			{
-				ManifoldPointSavedData[NumSavedManifoldPoints++].Save(ManifoldPoint);
-			}
-		}
-
-		bWasManifoldRestored = false;
-	}
-
 	void FPBDCollisionConstraint::RestoreManifold()
 	{
 		// We want to restore the manifold as-is and will skip the narrow phase, which means we leave the manifold in place, 
@@ -606,10 +586,9 @@ namespace Chaos
 
 		const FRigidTransform3 Shape0ToShape1Transform = ShapeWorldTransform0.GetRelativeTransformNoScale(ShapeWorldTransform1);
 
-		// Update and prune manifold points.
+		// Update or prune manifold points.
 		TArray<int32, TInlineAllocator<4>> ManifoldPointsToRemove;
-		const int32 NumManifoldPoints = ManifoldPoints.Num();
-		for (int32 ManifoldPointIndex = 0; ManifoldPointIndex < NumManifoldPoints; ++ManifoldPointIndex)
+		for (int32 ManifoldPointIndex = 0; ManifoldPointIndex < ManifoldPoints.Num(); ++ManifoldPointIndex)
 		{
 			FManifoldPoint& ManifoldPoint = ManifoldPoints[ManifoldPointIndex];
 
@@ -631,22 +610,7 @@ namespace Chaos
 				// @todo(chaos): we should reproject the contact on the plane owner
 				const FVec3 ShapeContactPoint1 = Contact0In1 - ContactPhi * ContactNormalIn1;
 				ManifoldPoint.ContactPoint.ShapeContactPoints[1] = ShapeContactPoint1;
-
-				// Restore friction anchors if possible
-				TryRestoreFrictionData(ManifoldPointIndex);
-
-				// Update world-space contact locations
-				UpdateManifoldPointFromContact(ManifoldPointIndex);
-				ManifoldPoint.ContactPoint.Location = FReal(0.5) * (ManifoldPoint.WorldContactPoints[0] + ManifoldPoint.WorldContactPoints[1]);
 				ManifoldPoint.ContactPoint.Phi = ContactPhi;
-
-				ManifoldPoint.bWasRestored = true;
-
-				if (ContactPhi < Manifold.Phi)
-				{
-					// Update closest point
-					SetActiveContactPoint(ManifoldPoint.ContactPoint);
-				}
 			}
 			else
 			{
@@ -654,9 +618,44 @@ namespace Chaos
 			}
 		}
 
+		// If we removed more than 1 point, we are going to rebuild the manifold anyway, so we're done
+		if (ManifoldPointsToRemove.Num() > 1)
+		{
+			ResetActiveManifoldContacts();
+			return false;
+		}
+
+		// Remove points - process in reverse order for simpler index handling and faster removal of multiple elements
+		for (int32 RemoveArrayIndex = ManifoldPointsToRemove.Num() - 1; RemoveArrayIndex >= 0; --RemoveArrayIndex)
+		{
+			constexpr bool bAllowShrink = false;
+			ManifoldPoints.RemoveAt(ManifoldPointsToRemove[RemoveArrayIndex], 1, bAllowShrink);
+		}
+
+		// Update world-space state for the points we kept
+		for (int32 ManifoldPointIndex = 0; ManifoldPointIndex < ManifoldPoints.Num(); ++ManifoldPointIndex)
+		{
+			FManifoldPoint& ManifoldPoint = ManifoldPoints[ManifoldPointIndex];
+
+			// Restore friction anchors if we have them for this point
+			TryRestoreFrictionData(ManifoldPointIndex);
+
+			// Update world-space contact locations
+			UpdateManifoldPointFromContact(ManifoldPointIndex);
+			ManifoldPoint.ContactPoint.Location = FReal(0.5) * (ManifoldPoint.WorldContactPoints[0] + ManifoldPoint.WorldContactPoints[1]);
+
+			ManifoldPoint.bWasRestored = true;
+
+			if (ManifoldPoint.ContactPoint.Phi < Manifold.Phi)
+			{
+				// Update closest point
+				SetActiveContactPoint(ManifoldPoint.ContactPoint);
+			}
+		}
+
+		// If we did not remove any contact points and we have not moved or rotated much we may reuse the manifold as-is.
 		if ((ManifoldPointsToRemove.Num() == 0) && (ShapePositionTolerance > 0) && (ShapeRotationThreshold > 0))
 		{
-			// If we did not remove any contact points and we have not moved or rotated much we can reuse the manifold as-is.
 			// The transform check is necessary regardless of how many points we have left in the manifold because
 			// as a body moves/rotates we may have to change which faces/edges are colliding. We can't know if the face/edge
 			// will change until we run the closest-point checks (GJK) in the narrow phase.
@@ -675,12 +674,6 @@ namespace Chaos
 			}
 		}
 
-		// We removed some points - process in reverse order for simpler index handling and faster removal of multiple elements
-		for (int32 RemoveArrayIndex = ManifoldPointsToRemove.Num() - 1; RemoveArrayIndex >= 0; --RemoveArrayIndex)
-		{
-			ManifoldPoints.RemoveAt(ManifoldPointsToRemove[RemoveArrayIndex]);
-		}
-
 		return false;
 	}
 
@@ -899,10 +892,10 @@ namespace Chaos
 		if (bChaos_Manifold_EnableFrictionRestore)
 		{
 			const FReal DistanceToleranceSq = FMath::Square(Chaos_Manifold_FrictionPositionTolerance);
-			for (int32 PrevManifoldPointIndex = 0; PrevManifoldPointIndex < NumSavedManifoldPoints; ++PrevManifoldPointIndex)
+			for (int32 PrevManifoldPointIndex = 0; PrevManifoldPointIndex < ManifoldPointsSavedData.Num(); ++PrevManifoldPointIndex)
 			{
-				const FManifoldPointSavedData& PrevManifoldPoint = ManifoldPointSavedData[PrevManifoldPointIndex];
-				if (PrevManifoldPoint.bInsideStaticFrictionCone && PrevManifoldPoint.IsMatch(ManifoldPoint, DistanceToleranceSq))
+				const FManifoldPointSavedData& PrevManifoldPoint = ManifoldPointsSavedData[PrevManifoldPointIndex];
+				if (PrevManifoldPoint.IsMatch(ManifoldPoint, DistanceToleranceSq))
 				{
 					return &PrevManifoldPoint;
 				}
diff --git a/Engine/Source/Runtime/Experimental/Chaos/Private/Chaos/Collision/ParticlePairMidPhase.cpp b/Engine/Source/Runtime/Experimental/Chaos/Private/Chaos/Collision/ParticlePairMidPhase.cpp
--- a/Engine/Source/Runtime/Experimental/Chaos/Private/Chaos/Collision/ParticlePairMidPhase.cpp
+++ b/Engine/Source/Runtime/Experimental/Chaos/Private/Chaos/Collision/ParticlePairMidPhase.cpp
@@ -270,16 +270,11 @@ namespace Chaos
 			// NOTE: these are not used by CCD which continuously moves the particles
 			Constraint->SetShapeWorldTransforms(ShapeWorldTransform0, ShapeWorldTransform1);
 
-			// If the constraint was not used last frame, it needs to be reset. 
-			// If it was used last frame, its data can be used for static friction etc. (unless CCD is enabled)
 			Constraint->SetCCDEnabled(bUseCCD);
-			if (bWasUpdatedLastTick && !bUseCCD)
-			{
-				// Copy manifold data used for static friction - we are about to overwrite the manifold points
-				// NOTE: does not clear the current manifold points
-				Constraint->SaveManifold();
-			}
-			else
+
+			// If the constraint was not used last frame, it needs to be reset. 
+			// Otherwise we will try to reuse it below
+			if (!bWasUpdatedLastTick || bUseCCD)
 			{
 				// Clear all manifold data
 				Constraint->ResetManifold();
diff --git a/Engine/Source/Runtime/Experimental/Chaos/Private/Chaos/Collision/SolverCollisionContainer.cpp b/Engine/Source/Runtime/Experimental/Chaos/Private/Chaos/Collision/SolverCollisionContainer.cpp
--- a/Engine/Source/Runtime/Experimental/Chaos/Private/Chaos/Collision/SolverCollisionContainer.cpp
+++ b/Engine/Source/Runtime/Experimental/Chaos/Private/Chaos/Collision/SolverCollisionContainer.cpp
@@ -130,29 +130,24 @@ namespace Chaos
 		/**
 		 * @brief Send all solver results to the constraint
 		*/
-		void ScatterOutput(
-			const FReal Dt)
+		void ScatterOutput(const FReal Dt)
 		{
 			FVec3 AccumulatedImpulse = FVec3(0);
 
+			Constraint->ResetSolverResults();
+			Constraint->SetNumActivePositionIterations(Solver.NumPositionSolves());
+
 			for (int32 PointIndex = 0; PointIndex < Solver.NumManifoldPoints(); ++PointIndex)
 			{
 				const FPBDCollisionSolverManifoldPoint& SolverManifoldPoint = Solver.GetManifoldPoint(PointIndex);
-
-				TArrayView<FManifoldPoint> ManifoldPoints = Constraint->GetManifoldPoints();
-				FManifoldPoint& ManifoldPoint = ManifoldPoints[PointIndex];
-
-				ManifoldPoint.NetPushOut = SolverManifoldPoint.NetPushOut;
-				ManifoldPoint.NetImpulse = SolverManifoldPoint.NetImpulse;
-				ManifoldPoint.bInsideStaticFrictionCone = SolverManifoldPoint.bInsideStaticFrictionCone;
-				ManifoldPoint.StaticFrictionMax = SolverManifoldPoint.StaticFrictionMax;
-
-				AccumulatedImpulse += SolverManifoldPoint.NetImpulse + (SolverManifoldPoint.NetPushOut / Dt);
+				Constraint->SetSolverResults(PointIndex, 
+					SolverManifoldPoint.NetPushOut, 
+					SolverManifoldPoint.NetImpulse, 
+					SolverManifoldPoint.bInsideStaticFrictionCone,
+					SolverManifoldPoint.StaticFrictionMax,
+					Dt);
 			}
 
-			Constraint->AccumulatedImpulse = AccumulatedImpulse;
-			Constraint->SetNumActivePositionIterations(Solver.NumPositionSolves());
-
 			Constraint->SetSolverBodies(nullptr, nullptr);
 			Constraint = nullptr;
 			Solver.ResetSolverBodies();
diff --git a/Engine/Source/Runtime/Experimental/Chaos/Public/Chaos/Collision/CollisionConstraintAllocator.h b/Engine/Source/Runtime/Experimental/Chaos/Public/Chaos/Collision/CollisionConstraintAllocator.h
--- a/Engine/Source/Runtime/Experimental/Chaos/Public/Chaos/Collision/CollisionConstraintAllocator.h
+++ b/Engine/Source/Runtime/Experimental/Chaos/Public/Chaos/Collision/CollisionConstraintAllocator.h
@@ -368,7 +368,11 @@ namespace Chaos
 				if ((MidPhase != nullptr) && !MidPhase->IsUsedSince(CurrentEpoch) && !MidPhase->IsSleeping())
 				{
 					DetachParticlePairMidPhase(MidPhase.Get());
-					ParticlePairMidPhases.RemoveAtSwap(Index);
+
+					// This array can get very large so we allow it to shrink when possible
+					constexpr bool bAllowShrink = true;
+					ParticlePairMidPhases.RemoveAtSwap(Index, 1, bAllowShrink);
+
 					--Index;
 					--NumParticlePairMidPhases;
 				}
diff --git a/Engine/Source/Runtime/Experimental/Chaos/Public/Chaos/Collision/PBDCollisionConstraint.h b/Engine/Source/Runtime/Experimental/Chaos/Public/Chaos/Collision/PBDCollisionConstraint.h
--- a/Engine/Source/Runtime/Experimental/Chaos/Public/Chaos/Collision/PBDCollisionConstraint.h
+++ b/Engine/Source/Runtime/Experimental/Chaos/Public/Chaos/Collision/PBDCollisionConstraint.h
@@ -80,19 +80,14 @@ namespace Chaos
 	class CHAOS_API FManifoldPointSavedData
 	{
 	public:
-		FManifoldPointSavedData()
-		{
-		}
-
 		/**
 		 * @brief Copy the ManifoldPoint data needed for static friction next tick
 		*/
-		inline void Save(const FManifoldPoint& ManifoldPoint)
+		FManifoldPointSavedData(const FManifoldPoint& ManifoldPoint)
 		{
 			ShapeContactPoints[0] = ManifoldPoint.ContactPoint.ShapeContactPoints[0];
 			ShapeContactPoints[1] = ManifoldPoint.ContactPoint.ShapeContactPoints[1];
 			StaticFrictionMax = ManifoldPoint.StaticFrictionMax;
-			bInsideStaticFrictionCone = ManifoldPoint.bInsideStaticFrictionCone;
 		}
 
 		/**
@@ -121,7 +116,6 @@ namespace Chaos
 		FVec3 ShapeContactPoints[2];
 		FVec3 NetPushOut;
 		FReal StaticFrictionMax;
-		bool bInsideStaticFrictionCone;
 	};
 
 	/*
@@ -460,6 +454,36 @@ namespace Chaos
 
 		const FManifoldPointSavedData* FindManifoldPointSavedData(const FManifoldPoint& ManifoldPoint) const;
 
+		/**
+		 * @brief Called before SetSolverResults() to reset accumulators
+		*/
+		void ResetSolverResults()
+		{
+			ManifoldPointsSavedData.Reset();
+			AccumulatedImpulse = FVec3(0);
+		}
+
+		/**
+		 * @brief Store the data (from the solver) that is retained between ticks for the specified manifold point
+		*/
+		void SetSolverResults(const int32 ManifoldPointIndex, const FVec3& NetPushOut, const FVec3& NetImpulse, const bool bInsideStaticFrictionCone, const FReal StaticFrictionMax, const FReal Dt)
+		{
+			FManifoldPoint& ManifoldPoint = ManifoldPoints[ManifoldPointIndex];
+
+			ManifoldPoint.NetPushOut = NetPushOut;
+			ManifoldPoint.NetImpulse = NetImpulse;
+			ManifoldPoint.bInsideStaticFrictionCone = bInsideStaticFrictionCone;
+			ManifoldPoint.StaticFrictionMax = StaticFrictionMax;
+
+			AccumulatedImpulse += NetImpulse + (NetPushOut / Dt);
+
+			// If are still satisfying the static friction condition, save off the state required for next tick
+			if (bInsideStaticFrictionCone && !NetPushOut.IsNearlyZero() && (ManifoldPointsSavedData.Num() < ManifoldPointsSavedData.Max()))
+			{
+				ManifoldPointsSavedData.Emplace(FManifoldPointSavedData(ManifoldPoint));
+			}
+		}
+
 	public:
 		const FPBDCollisionConstraintHandle* GetConstraintHandle() const { return this; }
 		FPBDCollisionConstraintHandle* GetConstraintHandle() { return this; }
@@ -497,11 +521,6 @@ namespace Chaos
 		void InitManifoldPoint(const int32 ManifoldPointIndex);
 		void SetActiveContactPoint(const FContactPoint& ContactPoint);
 
-		/**
-		 * @brief Move the current manifold to the previous manifold, and reset the current manifold ready for it to be rebuilt
-		*/
-		void SaveManifold();
-
 		/**
 		 * @brief Restore data used for static friction from the previous point(s) to the current
 		*/
@@ -525,6 +544,10 @@ namespace Chaos
 		FReal Stiffness;
 
 		TArray<FManifoldPoint> ManifoldPoints;
+
+		// The manifold points from the previous tick when we don't reuse the manifold. Used by static friction.
+		TFixedArray<FManifoldPointSavedData, MaxManifoldPoints> ManifoldPointsSavedData;
+
 		FReal CullDistance;
 
 		// The margins to use during collision detection. We don't always use the margins on the shapes directly.
@@ -544,10 +567,6 @@ namespace Chaos
 		// Simplex data from the last call to GJK, used to warm-start GJK
 		FGJKSimplexData GJKWarmStartData;
 
-		// The manifold points from the previous tick when we don't reuse the manifold. Used by static friction.
-		FManifoldPointSavedData ManifoldPointSavedData[MaxManifoldPoints];
-		int32 NumSavedManifoldPoints;
-
 		FRigidTransform3 ShapeWorldTransform0;
 		FRigidTransform3 ShapeWorldTransform1;
 		FRigidTransform3 LastShapeWorldTransform0;
diff --git a/Engine/Source/Runtime/Experimental/Chaos/Public/Chaos/Framework/UncheckedArray.h b/Engine/Source/Runtime/Experimental/Chaos/Public/Chaos/Framework/UncheckedArray.h
--- a/Engine/Source/Runtime/Experimental/Chaos/Public/Chaos/Framework/UncheckedArray.h
+++ b/Engine/Source/Runtime/Experimental/Chaos/Public/Chaos/Framework/UncheckedArray.h
@@ -42,6 +42,9 @@ namespace Chaos
 #endif
 	};
 
+	template<typename T, int32 N>
+	using TFixedArray = TArray<T, TFixedAllocator<N>>;
+
 	template<typename T, int32 N>
 	using TUncheckedFixedArray = TArray<T, TUncheckedFixedAllocator<N>>;
 
