Simultaneous tree descent overlap

diff --git a/Engine/Source/Runtime/Experimental/Chaos/Public/Chaos/AABBTree.h b/Engine/Source/Runtime/Experimental/Chaos/Public/Chaos/AABBTree.h
--- a/Engine/Source/Runtime/Experimental/Chaos/Public/Chaos/AABBTree.h
+++ b/Engine/Source/Runtime/Experimental/Chaos/Public/Chaos/AABBTree.h
@@ -372,7 +372,7 @@ struct TAABBTreeLeafArray : public TBoundsWrapperHelper<TPayloadType, T, bComput
 	void Reset()
 	{
 		Elems.Reset();
-		bDirtyLeaf = true;
+		bDirtyLeaf = false;
 	}
 
 #if !UE_BUILD_SHIPPING
@@ -414,7 +414,7 @@ struct TAABBTreeLeafArray : public TBoundsWrapperHelper<TPayloadType, T, bComput
 	TArray<TPayloadBoundsElement<TPayloadType, T>> Elems;
 
 	/** Flag on the leaf to know if it has been updated */
-	bool bDirtyLeaf = true;
+	bool bDirtyLeaf = false;
 };
 
 template <typename TPayloadType, bool bComputeBounds, typename T>
@@ -436,6 +436,7 @@ struct TAABBTreeNode
 	int32 ChildrenNodes[2] = { INDEX_NONE, INDEX_NONE };
 	int32 ParentNode = INDEX_NONE;
 	bool bLeaf = false;
+	bool bDirtyNode = false;
 
 #if !UE_BUILD_SHIPPING
 	void DebugDraw(ISpacialDebugDrawInterface<T>& InInterface, const TArray<TAABBTreeNode<T>>& Nodes, const FVec3& InLinearColor, float InThickness) const
@@ -589,6 +590,8 @@ public:
 		
 		OverlappingLeaves.Reset();
 		OverlappingOffsets.Reset();
+		OverlappingPairs.Reset();
+		OverlappingCounts.Reset();
 		
 		NumProcessedThisSlice = 0;
 		WorkStack.Reset();
@@ -1646,6 +1649,8 @@ public:
 		PayloadToInfo.Reserve(From.PayloadToInfo.Num());
 		OverlappingLeaves.Reserve(From.OverlappingLeaves.Num());
 		OverlappingOffsets.Reserve(From.OverlappingOffsets.Num());
+		OverlappingPairs.Reserve(From.OverlappingPairs.Num());
+		OverlappingCounts.Reserve(From.OverlappingCounts.Num());
 
 		this->SetAsyncTimeSlicingComplete(false);
 	}
@@ -1696,6 +1701,14 @@ public:
 		{
 			return;
 		}
+		if (!ContinueTimeSliceCopy(From.OverlappingCounts, OverlappingCounts, SizeToCopyLeft))
+		{
+			return;
+		}
+		if (!ContinueTimeSliceCopy(From.OverlappingPairs, OverlappingPairs, SizeToCopyLeft))
+		{
+			return;
+		}
 
 		this->SetAsyncTimeSlicingComplete(true);
 	}
@@ -1833,15 +1846,180 @@ public:
 			}
 		}
 	}
+	
+	/** Recursively add overlapping leaves given 2 nodes in the tree */
+    void AddNodesOverlappingLeaves(const TAABBTreeNode<T>& LeftNode, const TAABB<T, 3>& LeftBounds,
+								  const TAABBTreeNode<T>& RightNode, const TAABB<T, 3>& RightBounds, const bool bDirtyFilter)
+	{
+		// If dirty filter enabled only look for overlapping leaves if one of the 2 nodes are dirty 
+		if(!bDirtyFilter || (bDirtyFilter && (LeftNode.bDirtyNode || RightNode.bDirtyNode)))
+		{
+			if(LeftBounds.Intersects(RightBounds))
+			{
+				// If left and right are leaves check for intersection
+				if(LeftNode.bLeaf && RightNode.bLeaf)
+				{
+					const int32 LeftLeaf = LeftNode.ChildrenNodes[0];
+					const int32 RightLeaf = RightNode.ChildrenNodes[0];
 
-	/** Cache for each leaves all the overlapping leaves*/
-	virtual void CacheOverlappingLeaves() override
+					// Same condition as for the nodes
+					if(!bDirtyFilter || (bDirtyFilter && (Leaves[LeftLeaf].IsLeafDirty() || Leaves[RightLeaf].IsLeafDirty())))
+					{
+						if(Leaves[LeftLeaf].GetBounds().Intersects(Leaves[RightLeaf].GetBounds()))
+						{
+							OverlappingPairs.Add(FIntVector2(LeftLeaf, RightLeaf));
+							++OverlappingCounts[LeftLeaf];
+							++OverlappingCounts[RightLeaf];
+						}
+					}
+				}
+				// If only left is a leaf continue recursion with the right node children
+				else if(LeftNode.bLeaf)
+				{
+					AddNodesOverlappingLeaves(LeftNode, LeftBounds, Nodes[RightNode.ChildrenNodes[0]], RightNode.ChildrenBounds[0], bDirtyFilter);
+					AddNodesOverlappingLeaves(LeftNode, LeftBounds, Nodes[RightNode.ChildrenNodes[1]], RightNode.ChildrenBounds[1], bDirtyFilter);
+				}
+				// Otherwise continue recursion with the left node children
+				else
+				{
+					AddNodesOverlappingLeaves(Nodes[LeftNode.ChildrenNodes[0]], LeftNode.ChildrenBounds[0], RightNode, RightBounds, bDirtyFilter);
+					AddNodesOverlappingLeaves(Nodes[LeftNode.ChildrenNodes[1]], LeftNode.ChildrenBounds[1], RightNode, RightBounds, bDirtyFilter);
+				}
+			}
+		}
+	}
+	
+	/** Recursively add overlapping leaves given a root node in the tree */
+	void AddRootOverlappingLeaves(const TAABBTreeNode<T>& TreeNode, const bool bDirtyFilter)
 	{
-		OverlappingOffsets.SetNum(Leaves.Num()+1, false);
-		OverlappingLeaves.Reset();
+		if(!TreeNode.bLeaf)
+		{
+			// Find overlapping leaves within the left and right children
+			AddRootOverlappingLeaves(Nodes[TreeNode.ChildrenNodes[0]], bDirtyFilter);
+			AddRootOverlappingLeaves(Nodes[TreeNode.ChildrenNodes[1]], bDirtyFilter);
+
+			// Then try finding some overlaps in between the 2 children
+			AddNodesOverlappingLeaves(Nodes[TreeNode.ChildrenNodes[0]], TreeNode.ChildrenBounds[0],
+									  Nodes[TreeNode.ChildrenNodes[1]], TreeNode.ChildrenBounds[1], bDirtyFilter);
+		}
+	}
+
+	/** Fill the overlapping pairs from the previous persistent and not dirty leaves */
+	void FillPersistentOverlappingPairs()
+	{
+		for(int32 LeafIndex = 0, NumLeaves = Leaves.Num(); LeafIndex < NumLeaves; ++LeafIndex)
+		{
+			int32 NumOverlaps = 0;
+			if(!Leaves[LeafIndex].IsLeafDirty())
+			{
+				for(int32 OverlappingIndex = OverlappingOffsets[LeafIndex]; OverlappingIndex < OverlappingOffsets[LeafIndex+1]; ++OverlappingIndex)
+				{
+					if(!Leaves[OverlappingLeaves[OverlappingIndex]].IsLeafDirty())
+					{
+						if(LeafIndex < OverlappingLeaves[OverlappingIndex])
+						{
+							OverlappingPairs.Add(FIntVector2(LeafIndex, OverlappingLeaves[OverlappingIndex]));
+						}
+						if(LeafIndex != OverlappingLeaves[OverlappingIndex])
+						{
+							++NumOverlaps;
+						}
+					}
+				}
+			}
+			// Make sure the leaf is intersecting itself if several elements per leaf
+			OverlappingPairs.Add(FIntVector2(LeafIndex, LeafIndex));
+			++NumOverlaps;
+			
+			OverlappingCounts[LeafIndex] = NumOverlaps;
+		}
+	}
+	/** Propagates the leaves dirty flag up to the root node */
+	void PropagateLeavesDirtyFlag()
+	{
+		for(int32 NodeIndex = 0, NumNodes = Nodes.Num(); NodeIndex < NumNodes; ++NodeIndex)
+		{
+			if(Nodes[NodeIndex].bLeaf)
+			{
+				Nodes[NodeIndex].bDirtyNode = Leaves[Nodes[NodeIndex].ChildrenNodes[0]].IsLeafDirty();
+				if(Nodes[NodeIndex].bDirtyNode)
+				{
+					int32 NodeParent = Nodes[NodeIndex].ParentNode;
+					while(NodeParent != INDEX_NONE && !Nodes[NodeParent].bDirtyNode)
+					{
+						Nodes[NodeParent].bDirtyNode = true;
+						NodeParent = Nodes[NodeParent].ParentNode;
+					}
+				}
+			}
+		}
+	}
 
+	/** Simultaneous tree descent to compute the overlapping leaves */
+	void ComputeOverlappingCacheFromRoot(const bool bDirtyFilter)
+	{
 		if(!bDynamicTree || (bDynamicTree && RootNode == INDEX_NONE)) return;
+		
+		OverlappingOffsets.SetNum(Leaves.Num()+1, false);
+		OverlappingCounts.SetNum(Leaves.Num(), false);
+		OverlappingPairs.Reset();
+		
+		if(bDirtyFilter)
+		{
+			FillPersistentOverlappingPairs();
+			PropagateLeavesDirtyFlag();
+		}
+		else
+		{
+			for(int32 LeafIndex = 0, NumLeaves = Leaves.Num(); LeafIndex < NumLeaves; ++LeafIndex)
+			{
+				// Make sure the leaf is intersecting itself if several elements per leaf
+				OverlappingPairs.Add(FIntVector2(LeafIndex, LeafIndex));
+				OverlappingCounts[LeafIndex] = 1;
+			}
+		}
+		AddRootOverlappingLeaves(Nodes[RootNode], bDirtyFilter);
+
+		OverlappingOffsets[0] = 0;
+		for(int32 LeafIndex = 0, NumLeaves = Leaves.Num(); LeafIndex < NumLeaves; ++LeafIndex)
+		{
+			OverlappingOffsets[LeafIndex+1] = OverlappingOffsets[LeafIndex] + OverlappingCounts[LeafIndex];
+			OverlappingCounts[LeafIndex] = OverlappingOffsets[LeafIndex];
+		}
+		OverlappingLeaves.SetNum(OverlappingOffsets.Last(), false);
+		for(auto& OverlappingPair : OverlappingPairs)
+		{
+			if(OverlappingPair[0] != OverlappingPair[1])
+			{
+				OverlappingLeaves[OverlappingCounts[OverlappingPair[0]]++] = OverlappingPair[1];
+				OverlappingLeaves[OverlappingCounts[OverlappingPair[1]]++] = OverlappingPair[0];
+			}
+			else
+			{
+				OverlappingLeaves[OverlappingCounts[OverlappingPair[0]]++] = OverlappingPair[0];
+			}
+		}
+		if(bDirtyFilter)
+		{
+			for(int32 LeafIndex = 0, NumLeaves = Leaves.Num(); LeafIndex < NumLeaves; ++LeafIndex)
+			{
+				Leaves[LeafIndex].SetDirtyState(false);
+			}
+			for(int32 NodeIndex = 0, NumNodes = Nodes.Num(); NodeIndex < NumNodes; ++NodeIndex)
+			{
+				Nodes[NodeIndex].bDirtyNode = false;
+			}
+		}
+	}
 
+	/** Sequential loop over the leaves to fill the overlapping pairs */
+	void ComputeOverlappingCacheFromLeaf()
+	{
+		OverlappingOffsets.SetNum(Leaves.Num()+1, false);
+		OverlappingLeaves.Reset();
+		
+		if(!bDynamicTree || (bDynamicTree && RootNode == INDEX_NONE)) return;
+		
 		TArray<int32> NodeStack;
 		const int32 FirstNode = bDynamicTree ? RootNode : 0;
 		
@@ -1851,9 +2029,23 @@ public:
 			FindOverlappingLeaf(FirstNode, LeafIndex, NodeStack);
 		}
 		OverlappingOffsets.Last() = OverlappingLeaves.Num();
-		for(int32 LeafIndex = 0, NumLeaves = Leaves.Num(); LeafIndex < NumLeaves; ++LeafIndex)
+	}
+
+	/** Cache for each leaves all the overlapping leaves*/
+	virtual void CacheOverlappingLeaves() override
+	{
+		// Dev settings to switch easily algorithms
+		// Will switch to cvars if the leaf version could be faster 
+		const bool bCachingRoot = true;
+		const bool bDirtyFilter = false;
+		
+		if(bCachingRoot)
+		{
+			ComputeOverlappingCacheFromRoot(bDirtyFilter);
+		}
+		else
 		{
-			Leaves[LeafIndex].SetDirtyState(false);
+			ComputeOverlappingCacheFromLeaf();
 		}
 	}
 
@@ -2846,6 +3038,8 @@ private:
 		, bShouldRebuild(Other.bShouldRebuild)
 		, OverlappingLeaves(Other.OverlappingLeaves)
 		, OverlappingOffsets(Other.OverlappingOffsets)
+		, OverlappingPairs(Other.OverlappingPairs)
+		, OverlappingCounts(Other.OverlappingCounts)
 	{
 
 		ensure(bDynamicTree == Other.IsTreeDynamic());
@@ -2948,6 +3142,12 @@ private:
 
 	/** Offset for each leaf in the overlapping leaves (OverlappingOffsets[LeafIndex]/OverlappingOffsets[LeafIndex+1] will be start/end indices of the leaf index in the overallping leaves array )*/
 	TArray<int32> OverlappingOffsets;
+
+	/** List of leaves intersecting pairs that will be used to build the flat arrays (offsets,leaves) */
+	TArray<FIntVector2> OverlappingPairs;
+
+	/** Number of overlapping leaves per leaf */
+	TArray<int32> OverlappingCounts;
 };
 
 template<typename TPayloadType, typename TLeafType, bool bMutable, typename T>
