Chaos Dynamic tree performance improvements by eliminating copies

diff --git a/Engine/Source/Runtime/Experimental/Chaos/Private/Chaos/PBDRigidsEvolution.cpp b/Engine/Source/Runtime/Experimental/Chaos/Private/Chaos/PBDRigidsEvolution.cpp
--- a/Engine/Source/Runtime/Experimental/Chaos/Private/Chaos/PBDRigidsEvolution.cpp
+++ b/Engine/Source/Runtime/Experimental/Chaos/Private/Chaos/PBDRigidsEvolution.cpp
@@ -339,6 +339,11 @@ namespace Chaos
 				{
 					ensure(false);
 				}
+				// We should not copy dynamic trees here
+				if (Result->IsTreeDynamic())
+				{
+					Result = nullptr;
+				}
 			}
 			return Result;
 		};
@@ -836,10 +841,10 @@ namespace Chaos
 				check(ExternalAccelerationSubStructure->IsTreeDynamic());
 				*ExternalAccelerationSubStructure = *AsyncInternalAccelerationSubStructure;
 			}
-
-			// ToDo: We should be able to change this copy to a swap if we remove the copy in the async task
-			*InternalAccelerationSubStructure = *AsyncInternalAccelerationSubStructure;
-			//AsyncInternalAcceleration->SwapSubstructure(*InternalAcceleration, SpatialIdx);
+			
+			InternalAcceleration->SwapSubstructure(*AsyncInternalAcceleration, SpatialIdx);
+			//AsyncInternalAcceleration->GetSubstructure(SpatialIdx)->Reset(); // No need to reset
+			
 		}
 	}
 
diff --git a/Engine/Source/Runtime/Experimental/Chaos/Public/Chaos/SpatialAccelerationCollection.h b/Engine/Source/Runtime/Experimental/Chaos/Public/Chaos/SpatialAccelerationCollection.h
--- a/Engine/Source/Runtime/Experimental/Chaos/Public/Chaos/SpatialAccelerationCollection.h
+++ b/Engine/Source/Runtime/Experimental/Chaos/Public/Chaos/SpatialAccelerationCollection.h
@@ -467,7 +467,11 @@ public:
 		check(Buckets[Idx.Bucket].Objects.Num() > Idx.InnerIdx);
 		check(Other->Buckets[Idx.Bucket].Objects.Num() > Idx.InnerIdx);
 
-		std::swap(Buckets[Idx.Bucket].Objects[Idx.InnerIdx].Acceleration, Other->Buckets[Idx.Bucket].Objects[Idx.InnerIdx].Acceleration);
+		TUniquePtr<ISpatialAcceleration<TPayloadType, T, d>> OtherSubStructure = Other->RemoveSubstructure(Idx);
+		TUniquePtr<ISpatialAcceleration<TPayloadType, T, d>> ThisSubStructure = RemoveSubstructure(Idx);
+
+		Other->AddSubstructure(MoveTemp(ThisSubStructure), Idx.Bucket, Idx.InnerIdx);
+		AddSubstructure(MoveTemp(OtherSubStructure), Idx.Bucket, Idx.InnerIdx);
 	}
 	
 	virtual void Reset() override
