AABB Dynamic Tree: Use thickened bounds for calculating the best sibling

diff --git a/Engine/Source/Runtime/Experimental/Chaos/Public/Chaos/AABBTree.h b/Engine/Source/Runtime/Experimental/Chaos/Public/Chaos/AABBTree.h
--- a/Engine/Source/Runtime/Experimental/Chaos/Public/Chaos/AABBTree.h
+++ b/Engine/Source/Runtime/Experimental/Chaos/Public/Chaos/AABBTree.h
@@ -1076,9 +1076,12 @@ public:
 		return(WhichChildAmI(NodeIdx) ^ 1);
 	}
 
-	int32 FindBestSibling(const TAABB<T, 3>& NewBounds)
+	int32 FindBestSibling(const TAABB<T, 3>& InNewBounds)
 	{
 		
+		TAABB<T, 3> NewBounds = InNewBounds;
+		NewBounds.Thicken(FAABBTreeCVars::DynamicTreeBoundingBoxPadding);
+
 		//Priority Q of indices to explore
 		TArray<int32> PriorityQ;
 		PriorityQ.Reserve(10);
