Chaos : Destruction
-Removed stale code path for MultiChildProxy

diff --git a/Engine/Source/Programs/HeadlessChaos/Private/GeometryCollection/GeometryCollectionTestClustering.cpp b/Engine/Source/Programs/HeadlessChaos/Private/GeometryCollection/GeometryCollectionTestClustering.cpp
--- a/Engine/Source/Programs/HeadlessChaos/Private/GeometryCollection/GeometryCollectionTestClustering.cpp
+++ b/Engine/Source/Programs/HeadlessChaos/Private/GeometryCollection/GeometryCollectionTestClustering.cpp
@@ -1665,11 +1665,11 @@ namespace GeometryCollectionTest
 		auto& Clustering = UnitTest.Solver->GetEvolution()->GetRigidClustering();
 		const auto& ClusterMap = Clustering.GetChildrenMap();
 		const auto& ParticleHandles = UnitTest.Solver->GetParticles().GetParticleHandles();
+		const TArrayCollectionArray<FRigidTransform3>& ChildToParent = Clustering.GetChildToParentMap();
 
-		//FCollisionFilterData FilterData;
-		//FilterData.Word1 = 0xFFFF;
-		//FilterData.Word3 = 0xFFFF;
-		//ParticleHandles.Handle(6)->ShapesArray()[0]->SetQueryData(FilterData);
+		FVector TestOffset;
+		FVector InitialRootPosition;
+		FVector RelativeChildOffsets[4];
 
 		UnitTest.Solver->RegisterSimOneShotCallback([&]()
 		{
@@ -1697,6 +1697,43 @@ namespace GeometryCollectionTest
 			EXPECT_EQ(ClusterMap.Num(), 1);
 			EXPECT_TRUE(ClusterMapContains(ClusterMap, Root, { CollectionParticles[0],CollectionParticles[1], CollectionParticles2[0], CollectionParticles2[1] }));
 
+			//
+			// TEST
+			// Validate that the relative translations stay the same
+			// for the union clustered children, and that the root actually
+			// moves. 
+			//
+			if (Frame == 0)
+			{
+				InitialRootPosition = Root->X();
+				RelativeChildOffsets[0] = CollectionParticles[0]->X() - Root->X();
+				RelativeChildOffsets[1] = CollectionParticles[1]->X() - Root->X();
+				RelativeChildOffsets[2] = CollectionParticles2[0]->X() - Root->X();
+				RelativeChildOffsets[3] = CollectionParticles2[1]->X() - Root->X();
+			}
+			else
+			{
+				FTransform RootTransform(Root->R(), Root->X());
+
+				TArray<FTransform> GlobalTransform1;
+				GeometryCollectionAlgo::GlobalMatrices(DynamicCollection->Transform, DynamicCollection->Parent, GlobalTransform1);
+
+				TArray<FTransform> GlobalTransform2;
+				GeometryCollectionAlgo::GlobalMatrices(DynamicCollection2->Transform, DynamicCollection2->Parent, GlobalTransform2);
+
+				EXPECT_TRUE(!InitialRootPosition.Equals(Root->X())); // root moves
+
+				EXPECT_TRUE(RelativeChildOffsets[0].Equals(GlobalTransform1[0].GetRelativeTransform(RootTransform).GetTranslation()));
+				EXPECT_TRUE(RelativeChildOffsets[1].Equals(GlobalTransform1[1].GetRelativeTransform(RootTransform).GetTranslation()));
+				EXPECT_TRUE(RelativeChildOffsets[2].Equals(GlobalTransform2[0].GetRelativeTransform(RootTransform).GetTranslation()));
+				EXPECT_TRUE(RelativeChildOffsets[3].Equals(GlobalTransform2[1].GetRelativeTransform(RootTransform).GetTranslation()));
+			}
+
+			//
+			// TEST
+			// Validate that the children have been removed from the 
+			// parenting hierarchy.
+			//
 			EXPECT_TRUE(DynamicCollection->Parent[0] == INDEX_NONE);
 			EXPECT_TRUE(DynamicCollection->Parent[1] == INDEX_NONE);
 			EXPECT_TRUE(DynamicCollection->Parent[2] == INDEX_NONE);
@@ -1705,17 +1742,6 @@ namespace GeometryCollectionTest
 			EXPECT_TRUE(DynamicCollection2->Parent[1] == INDEX_NONE);
 			EXPECT_TRUE(DynamicCollection2->Parent[2] == INDEX_NONE);
 		}
-
-		TArray<FTransform> GlobalTransform;
-		GeometryCollectionAlgo::GlobalMatrices(DynamicCollection->Transform, DynamicCollection->Parent, GlobalTransform);
-
-		TArray<FTransform> GlobalTransform2;
-		GeometryCollectionAlgo::GlobalMatrices(DynamicCollection2->Transform, DynamicCollection2->Parent, GlobalTransform2);
-
-		EXPECT_TRUE(GlobalTransform[0].GetTranslation().Z > 0.0f);
-		EXPECT_TRUE(GlobalTransform[1].GetTranslation().Z > 0.0f);
-		EXPECT_TRUE(GlobalTransform2[0].GetTranslation().Z > 0.0f);
-		EXPECT_TRUE(GlobalTransform2[1].GetTranslation().Z > 0.0f);
 	}
 	
 
diff --git a/Engine/Source/Runtime/Experimental/Chaos/Private/Chaos/PBDRigidClustering.cpp b/Engine/Source/Runtime/Experimental/Chaos/Private/Chaos/PBDRigidClustering.cpp
--- a/Engine/Source/Runtime/Experimental/Chaos/Private/Chaos/PBDRigidClustering.cpp
+++ b/Engine/Source/Runtime/Experimental/Chaos/Private/Chaos/PBDRigidClustering.cpp
@@ -68,9 +68,6 @@ namespace Chaos
 	int32 DeactivateClusterChildren = 0;
 	FAutoConsoleVariableRef CVarDeactivateClusterChildren(TEXT("p.DeactivateClusterChildren"), DeactivateClusterChildren, TEXT("If children should be decativated when broken and put into another cluster."));
 
-	int32 MassPropertiesFromMultiChildProxy = 1;
-	FAutoConsoleVariableRef CVarMassPropertiesFromMultiChildProxy(TEXT("p.MassPropertiesFromMultiChildProxy"), MassPropertiesFromMultiChildProxy, TEXT(""));
-
 	//==========================================================================
 	// Free helper functions
 	//==========================================================================
@@ -194,8 +191,8 @@ namespace Chaos
 
 	DECLARE_CYCLE_STAT(TEXT("TPBDRigidClustering<>::UpdateClusterMassProperties()"), STAT_UpdateClusterMassProperties, STATGROUP_Chaos);
 	void UpdateClusterMassProperties(
-		Chaos::FPBDRigidClusteredParticleHandle* Parent, 
-		TSet<FPBDRigidParticleHandle*>& Children, 
+		Chaos::FPBDRigidClusteredParticleHandle* Parent,
+		TSet<FPBDRigidParticleHandle*>& Children,
 		const FRigidTransform3* ForceMassOrientation)
 	{
 		SCOPE_CYCLE_COUNTER(STAT_UpdateClusterMassProperties);
@@ -212,42 +209,16 @@ namespace Chaos
 		bool bHasProxyChild = false;
 		for (FPBDRigidParticleHandle* OriginalChild : Children)
 		{
-			FMultiChildProxyId MultiChildProxyId; // sizeof(FMultiChildProxyId) = sizeof(void*), so copy
-			TMultiChildProxyData<FReal, 3>* ProxyData = nullptr;
-			if (FPBDRigidClusteredParticleHandle* ClusteredOriginalChild = OriginalChild->CastToClustered())
-			{
-				MultiChildProxyId = ClusteredOriginalChild->MultiChildProxyId();
-				ProxyData = ClusteredOriginalChild->MultiChildProxyData().Get();
-			}
-
-			//int32 Child;
 			FPBDRigidParticleHandle* Child;
 			FVec3 ChildPosition;
 			FRotation3 ChildRotation;
-			if (MultiChildProxyId.Id == nullptr)
-			{
-				Child = OriginalChild;
-				ChildPosition = Child->X();
-				ChildRotation = Child->R();
-			}
-			else if (ProxyData->KeyChild == OriginalChild)
-			{
-				Child = MultiChildProxyId.Id;
 
-				const FRigidTransform3 ProxyWorldTM =
-					ProxyData->RelativeToKeyChild *
-					FRigidTransform3(
-						OriginalChild->X(), OriginalChild->R());
-				ChildPosition = ProxyWorldTM.GetLocation();
-				ChildRotation = ProxyWorldTM.GetRotation();
-			}
-			else
-			{
-				continue; //using a proxy but we are not the key child
-			}
+			Child = OriginalChild;
+			ChildPosition = Child->X();
+			ChildRotation = Child->R();
 
 			const FReal ChildMass = Child->M();
-			const FMatrix33 ChildWorldSpaceI = 
+			const FMatrix33 ChildWorldSpaceI =
 				(ChildRotation * FMatrix::Identity) * Child->I() * (ChildRotation * FMatrix::Identity).GetTransposed();
 			if (ChildWorldSpaceI.ContainsNaN())
 			{
@@ -325,48 +296,27 @@ namespace Chaos
 		Parent->P() = Parent->X();
 		for (FPBDRigidParticleHandle* OriginalChild : Children)
 		{
-			FMultiChildProxyId MultiChildProxyId; // sizeof(FMultiChildProxyId) = sizeof(void*), so copy
-			TMultiChildProxyData<FReal, 3>* ProxyData = nullptr;
-			if (FPBDRigidClusteredParticleHandle* ClusteredOriginalChild = OriginalChild->CastToClustered())
-			{
-				MultiChildProxyId = ClusteredOriginalChild->MultiChildProxyId();
-				ProxyData = ClusteredOriginalChild->MultiChildProxyData().Get();
-			}
-
 			FPBDRigidParticleHandle* Child;
 			FVec3 ChildPosition;
-			if (MultiChildProxyId.Id == nullptr)
-			{
-				Child = OriginalChild;
-				ChildPosition = Child->X();
-			}
-			else if (ProxyData && ProxyData->KeyChild == OriginalChild)
-			{
-				Child = MultiChildProxyId.Id;
-				const FRigidTransform3 ProxyWorldTM = 
-					ProxyData->RelativeToKeyChild * FRigidTransform3(OriginalChild->X(), OriginalChild->R());
-				ChildPosition = ProxyWorldTM.GetLocation();
-			}
-			else
-			{
-				continue; //using a proxy but we are not the key child
-			}
+
+			Child = OriginalChild;
+			ChildPosition = Child->X();
 
 			FVec3 ParentToChild = ChildPosition - Parent->X();
 
 			const FReal ChildMass = Child->M();
 			// taking v from original child since we are not actually simulating the proxy child
-			Parent->W() += 
-				FVec3::CrossProduct(ParentToChild, 
+			Parent->W() +=
+				FVec3::CrossProduct(ParentToChild,
 					OriginalChild->V() * ChildMass);
 			{
 				const FReal& p0 = ParentToChild[0];
 				const FReal& p1 = ParentToChild[1];
 				const FReal& p2 = ParentToChild[2];
 				const FReal& m = ChildMass;
-				Parent->I() += 
+				Parent->I() +=
 					FMatrix33(
-						m * (p1 * p1 + p2 * p2), -m * p1 * p0, -m * p2 * p0, 
+						m * (p1 * p1 + p2 * p2), -m * p1 * p0, -m * p2 * p0,
 						m * (p2 * p2 + p0 * p0), -m * p2 * p1, m * (p1 * p1 + p0 * p0));
 			}
 		}
@@ -395,14 +345,15 @@ namespace Chaos
 
 		if (Parent->R().ContainsNaN())
 		{
-			InertiaTensor = PMatrix<FReal, 3,3>(1.f, 1.f, 1.f);
-			Parent->R() = TRotation<FReal,3>(FMatrix::Identity);
+			InertiaTensor = PMatrix<FReal, 3, 3>(1.f, 1.f, 1.f);
+			Parent->R() = TRotation<FReal, 3>(FMatrix::Identity);
 		}
 
 		Parent->Q() = Parent->R();
 		Parent->InvI() = Parent->I().Inverse();
 	}
 
+
 	//==========================================================================
 	// TPBDRigidClustering
 	//==========================================================================
@@ -588,12 +539,6 @@ namespace Chaos
 		return NewParticle;
 	}
 
-	int32 UseMultiChildProxy = 1;
-	FAutoConsoleVariableRef CVarUseMultiChildProxy(TEXT("p.UseMultiChildProxy"), UseMultiChildProxy, TEXT("Whether to merge multiple children into a single collision proxy when one is available"));
-
-	int32 MinChildrenForMultiProxy = 1;
-	FAutoConsoleVariableRef CVarMinChildrenForMultiProxy(TEXT("p.MinChildrenForMultiProxy"), MinChildrenForMultiProxy, TEXT("Min number of children needed for multi child proxy optimization"));
-
 	DECLARE_CYCLE_STAT(TEXT("TPBDRigidClustering<>::UnionClusterGroups"), STAT_UnionClusterGroups, STATGROUP_Chaos);
 	template<class T_FPBDRigidsEvolution, class T_FPBDCollisionConstraint>
 	void TPBDRigidClustering<T_FPBDRigidsEvolution, T_FPBDCollisionConstraint>::UnionClusterGroups()
@@ -719,9 +664,6 @@ namespace Chaos
 			MEvolution.EnableParticle(Child, ClusteredParticle);
 			TopLevelClusterParents.Add(ClusteredChild);
 
-			//make sure to remove multi child proxy if it exists
-			ClusteredChild->MultiChildProxyData().Reset();
-			ClusteredChild->MultiChildProxyId().Id = nullptr;
 			ClusteredChild->SetClusterId(ClusterId(nullptr, ClusteredChild->ClusterIds().NumChildren)); // clear Id but retain number of children
 
 			const FRigidTransform3 ChildFrame = ClusteredChild->ChildToParent() * PreSolveTM;
@@ -881,7 +823,7 @@ namespace Chaos
 					int32 NumNewClusters = 0;
 					for (TArray<FPBDRigidParticleHandle*>& ConnectedPieces : ConnectedPiecesArray)
 					{
-						if (ConnectedPieces.Num() == 1) //need to break single pieces first in case multi child proxy needs to be invalidated
+						if (ConnectedPieces.Num() == 1) //need to break single pieces first
 						{
 							FPBDRigidParticleHandle* Child = ConnectedPieces[0];
 							RemoveChildLambda(Child);
@@ -999,9 +941,6 @@ namespace Chaos
 			MEvolution.EnableParticle(Child, ClusteredParticle);
 			TopLevelClusterParents.Add(ClusteredChild);
 
-			//make sure to remove multi child proxy if it exists
-			ClusteredChild->MultiChildProxyData().Reset();
-			ClusteredChild->MultiChildProxyId().Id = nullptr;
 			ClusteredChild->SetClusterId(ClusterId(nullptr, ClusteredChild->ClusterIds().NumChildren)); // clear Id but retain number of children
 
 			const FRigidTransform3 ChildFrame = ClusteredChild->ChildToParent() * PreSolveTM;
@@ -1162,7 +1101,7 @@ namespace Chaos
 					int32 NumNewClusters = 0;
 					for (TArray<FPBDRigidParticleHandle*>& ConnectedPieces : ConnectedPiecesArray)
 					{
-						if (ConnectedPieces.Num() == 1) //need to break single pieces first in case multi child proxy needs to be invalidated
+						if (ConnectedPieces.Num() == 1) //need to break single pieces first
 						{
 							FPBDRigidParticleHandle* Child = ConnectedPieces[0];
 							RemoveChildLambda(Child);
@@ -1651,12 +1590,11 @@ namespace Chaos
 		const FRigidTransform3 ClusterWorldTM(Parent->X(), Parent->R());
 
 		TArray<FVec3> OriginalPoints;
-		TArray<FPBDRigidParticleHandle*> GeomToOriginalParticlesHack;
-		GeomToOriginalParticlesHack.Reserve(Children.Num());
+		TArray<FPBDRigidParticleHandle*> ChildParticleHandles;
+		ChildParticleHandles.Reserve(Children.Num());
 
 		const bool bUseCollisionPoints = (ProxyGeometry || Parameters.bCopyCollisionParticles) && !Parameters.CollisionParticles;
 		bool bUseParticleImplicit = false;
-		bool bUsingMultiChildProxy = false;
 
 		// Need to extract a filter off one of the cluster children 
 		FCollisionFilterData Filter;
@@ -1704,29 +1642,9 @@ namespace Chaos
 				FPBDRigidParticleHandle* UsedGeomChild = Child;
 				if (Child->Geometry())
 				{
-					FPBDRigidClusteredParticleHandle* ClusteredChild = Child->CastToClustered();
-
-					const FMultiChildProxyId MultiChildProxyId = 
-						ClusteredChild ? ClusteredChild->MultiChildProxyId() : FMultiChildProxyId();
-					const TUniquePtr<TMultiChildProxyData<FReal,3>>& MultiChildProxyData = ClusteredChild->MultiChildProxyData();
-
-					if (UseLevelsetCollision || MultiChildProxyId.Id == nullptr || !MultiChildProxyData)
-					{
-						Objects.Add(TUniquePtr<FImplicitObject>(new TImplicitObjectTransformed<FReal, 3>(Child->Geometry(), Frame)));
-						Objects2.Add(TUniquePtr<FImplicitObject>(new TImplicitObjectTransformed<FReal, 3>(Child->Geometry(), Frame)));
-						GeomToOriginalParticlesHack.Add(Child);
-					}
-					else if (MultiChildProxyData->KeyChild == Child)
-					{
-						//using multi child proxy and this child is the key
-						const FRigidTransform3 ProxyWorldTM = MultiChildProxyData->RelativeToKeyChild * ChildWorldTM;
-						const FRigidTransform3 ProxyRelativeTM = ProxyWorldTM.GetRelativeTransform(ClusterWorldTM);
-						Objects.Add(TUniquePtr<FImplicitObject>(new TImplicitObjectTransformed<FReal, 3>(MultiChildProxyId.Id->Geometry(), ProxyRelativeTM)));
-						Objects2.Add(TUniquePtr<FImplicitObject>(new TImplicitObjectTransformed<FReal, 3>(MultiChildProxyId.Id->Geometry(), ProxyRelativeTM)));
-						UsedGeomChild = MultiChildProxyId.Id;
-						GeomToOriginalParticlesHack.Add(UsedGeomChild);
-						bUsingMultiChildProxy = true;
-					}
+					Objects.Add(TUniquePtr<FImplicitObject>(new TImplicitObjectTransformed<FReal, 3>(Child->Geometry(), Frame)));
+					Objects2.Add(TUniquePtr<FImplicitObject>(new TImplicitObjectTransformed<FReal, 3>(Child->Geometry(), Frame)));
+					ChildParticleHandles.Add(Child);
 				}
 
 				ensure(Child->Disabled() == true);
@@ -1757,7 +1675,7 @@ namespace Chaos
 			TUniquePtr<FImplicitObjectUnionClustered>& ChildrenSpatial = Parent->ChildrenSpatial();
 			ChildrenSpatial.Reset(
 				Objects2.Num() ? 
-				new Chaos::FImplicitObjectUnionClustered(MoveTemp(Objects2), GeomToOriginalParticlesHack) : 
+				new Chaos::FImplicitObjectUnionClustered(MoveTemp(Objects2), ChildParticleHandles) : 
 				nullptr);
 		}
 
@@ -1840,20 +1758,15 @@ namespace Chaos
 			}
 			else // !UseLevelsetCollision
 			{
-				if (!bUsingMultiChildProxy && Objects.Num() == 1)
+				if (Objects.Num() == 1)
 				{
-					// @coverage:{confidence tests}
-					//ensureMsgf(false, TEXT("Checking no proxy, not level set, a single object"));
 					Parent->SetDynamicGeometry(MoveTemp(Objects[0]));
 				}
 				else
 				{
-					QUICK_SCOPE_CYCLE_COUNTER(UnionBVH);
-					// @coverage : { confidence tests}
-					//ensureMsgf(false, TEXT("Checking no proxy, not levelset, and multiple objects"));
 					Parent->SetDynamicGeometry(
 						MakeUnique<FImplicitObjectUnionClustered>(
-							MoveTemp(Objects), GeomToOriginalParticlesHack));
+							MoveTemp(Objects), ChildParticleHandles));
 				}
 			}
 		}
@@ -1973,33 +1886,7 @@ namespace Chaos
 					{
 						if (TPBDRigidClusteredParticleHandle<FReal, 3>*ClusteredChild = Child->CastToClustered())
 						{
-							const TUniquePtr<TMultiChildProxyData<FReal, 3>>& ProxyData = ClusteredChild->MultiChildProxyData();
-							const FPBDRigidParticleHandle* KeyChild = ProxyData ? ProxyData->KeyChild : nullptr;
-							const FPBDRigidClusteredParticleHandle* ClusteredKeyChild = KeyChild ? KeyChild->CastToClustered() : nullptr;
-							if (ClusteredKeyChild)
-							{
-								//multi child so get its children
-								const FRigidTransform3 ProxyToCluster = ProxyData->RelativeToKeyChild * ClusteredKeyChild->ChildToParent();
-								const FVec3 ContactLocationProxyLocal = ProxyToCluster.InverseTransformPosition(ContactLocationClusterLocal);
-								FAABB3 ContactBoxProxy(ContactLocationProxyLocal, ContactLocationProxyLocal);
-								ContactBoxProxy.Thicken(ClusterDistanceThreshold);
-								if (ClusteredChild->ChildrenSpatial())
-								{
-									const TArray<FPBDRigidParticleHandle*> SubIntersections =
-										ClusteredChild->ChildrenSpatial()->FindAllIntersectingChildren(ContactBoxProxy);
-									for (FPBDRigidParticleHandle* SubChild : SubIntersections)
-									{
-										if (FPBDRigidClusteredParticleHandle* ClusteredSubChild = SubChild->CastToClustered())
-										{
-											ClusteredSubChild->CollisionImpulses() += ContactHandle->GetAccumulatedImpulse().Size();
-										}
-									}
-								}
-							}
-							else
-							{
-								ClusteredChild->CollisionImpulses() += ContactHandle->GetAccumulatedImpulse().Size();
-							}
+							ClusteredChild->CollisionImpulses() += ContactHandle->GetAccumulatedImpulse().Size();
 						}
 					}
 				}
@@ -2231,7 +2118,6 @@ namespace Chaos
 		Chaos::FPBDRigidClusteredParticleHandle* Parent,
 		const FClusterCreationParameters& Parameters)
 	{
-		// @todo(investigate) : This is trying to set multiple connections and throwing a warning in ConnectNodes
 		SCOPE_CYCLE_COUNTER(STAT_FixConnectivityGraphUsingDelaunayTriangulation);
 
 		const TArray<FPBDRigidParticleHandle*>& Children = MChildren[Parent];
@@ -2419,28 +2305,6 @@ namespace Chaos
 		}
 	}
 
-	//template<class T_FPBDRigidsEvolution, class T_FPBDCollisionConstraint>
-	//void TPBDRigidClustering<T_FPBDRigidsEvolution, T_FPBDCollisionConstraint, T, d>::AddUniqueConnection(uint32 Index1, uint32 Index2, T Strain)
-	//{
-	//	if (Index1 != Index2)
-	//	{
-	//		//todo(pref): This can be removed if we are sure there are no duplicate connections generated.
-	//		for (int32 i = 0; i < MParticles.ConnectivityEdges(Index1).Num(); i++)
-	//		{
-	//			if (MParticles.ConnectivityEdges(Index1)[i].Sibling == Index2)
-	//			{
-	//				// @todo(duplication connection) : re-enable post GDC.  
-	//				// FixConnectivityGraphUsingDelaunayTriangulation attempts to add multiple connections.
-	//				// so commenting out this msg to remove the noise from the confidence test. 
-	//				// ensureMsgf(false, TEXT("Duplicate graph connection."));
-	//				return;
-	//			}
-	//		}
-	//
-	//		MParticles.ConnectivityEdges(Index1).Add({ Index2, Strain });
-	//	}
-	//}
-
 	template<class T_FPBDRigidsEvolution, class T_FPBDCollisionConstraint>
 	void TPBDRigidClustering<T_FPBDRigidsEvolution, T_FPBDCollisionConstraint>::ConnectNodes(
 		FPBDRigidParticleHandle* Child1,
diff --git a/Engine/Source/Runtime/Experimental/Chaos/Public/Chaos/PBDRigidClusteredParticles.h b/Engine/Source/Runtime/Experimental/Chaos/Public/Chaos/PBDRigidClusteredParticles.h
--- a/Engine/Source/Runtime/Experimental/Chaos/Public/Chaos/PBDRigidClusteredParticles.h
+++ b/Engine/Source/Runtime/Experimental/Chaos/Public/Chaos/PBDRigidClusteredParticles.h
@@ -30,32 +30,6 @@ struct ClusterId
 	int32 NumChildren;
 };
 
-/** 
- * When multiple children are active and can share one collision proxy. Only 
- * valid if all original children are still in the cluster.
- */
-template <typename T, int d>
-struct CHAOS_API TMultiChildProxyData
-{
-	TMultiChildProxyData()
-		: KeyChild(nullptr)
-	{}
-	TRigidTransform<T, d> RelativeToKeyChild;	//Use one child's transform to determine where to place the geometry. Needed for partial fracture where all children are still present and can therefore use proxy
-	TPBDRigidParticleHandle<T, d>* KeyChild;
-};
-
-/** 
- * Used with \c TMultiChildProxyData. \c nullptr indicates no proxy data available.
- *
- * TODO: Chaos - Add dimension template param?  Add floating point param?
- */
-struct FMultiChildProxyId
-{
-	FMultiChildProxyId()
-		: Id(nullptr) 
-	{}
-	FPBDRigidParticleHandle* Id;
-};
 
 /**
  * An entry in a clustered particle's \c ConnectivityEdges array, indicating a
@@ -106,8 +80,6 @@ class TPBDRigidClusteredParticles : public TPBDRigidParticles<T, d>
 	, MChildToParent(MoveTemp(Other.MChildToParent))
 	, MClusterGroupIndex(MoveTemp(Other.MClusterGroupIndex))
 	, MInternalCluster(MoveTemp(Other.MInternalCluster))
-	, MMultiChildProxyId(MoveTemp(Other.MMultiChildProxyId))
-	, MMultiChildProxyData(MoveTemp(Other.MMultiChildProxyData))
 	, MPhysicsProxies(MoveTemp(Other.MPhysicsProxies))
 	, MCollisionImpulses(MoveTemp(Other.MCollisionImpulses))
 	, MStrains(MoveTemp(Other.MStrains))
@@ -133,12 +105,6 @@ class TPBDRigidClusteredParticles : public TPBDRigidParticles<T, d>
 	const auto& ChildrenSpatial(int32 Idx) const { return MChildrenSpatial[Idx]; }
 	auto& ChildrenSpatial(int32 Idx) { return MChildrenSpatial[Idx]; }
 
-	const auto& MultiChildProxyId(int32 Idx) const { return MMultiChildProxyId[Idx]; }
-	auto& MultiChildProxyId(int32 Idx) { return MMultiChildProxyId[Idx]; }
-
-	const auto& MultiChildProxyData(int32 Idx) const { return MMultiChildProxyData[Idx]; }
-	auto& MultiChildProxyData(int32 Idx) { return MMultiChildProxyData[Idx]; }
-
 	const auto& PhysicsProxies(int32 Idx) const { return MPhysicsProxies[Idx]; }
 	auto& PhysicsProxies(int32 Idx) { return MPhysicsProxies[Idx]; }
 
@@ -152,10 +118,7 @@ class TPBDRigidClusteredParticles : public TPBDRigidParticles<T, d>
 	const auto& ConnectivityEdges(int32 Idx) const { return MConnectivityEdges[Idx]; }
 	auto& ConnectivityEdges(int32 Idx) { return MConnectivityEdges[Idx]; }
 
-	//TODO_CHAOSPARTICLE_HANDLE: is this really needed?
 	const auto& ConnectivityEdgesArray() const { return MConnectivityEdges; }
-	const auto& MultiChildProxyDataArray() const { return MMultiChildProxyData; }
-	const auto& MultiChildProxyIdArray() const { return MMultiChildProxyId; }
 
 	const auto& ClusterIdsArray() const { return MClusterIds; }
 	auto& ClusterIdsArray() { return MClusterIds; }
@@ -189,8 +152,6 @@ class TPBDRigidClusteredParticles : public TPBDRigidParticles<T, d>
 		  TArrayCollection::AddArray(&MClusterGroupIndex);
 		  TArrayCollection::AddArray(&MInternalCluster);
 		  TArrayCollection::AddArray(&MChildrenSpatial);
-		  TArrayCollection::AddArray(&MMultiChildProxyId);
-		  TArrayCollection::AddArray(&MMultiChildProxyData);
 		  TArrayCollection::AddArray(&MPhysicsProxies);
 		  TArrayCollection::AddArray(&MCollisionImpulses);
 		  TArrayCollection::AddArray(&MStrains);
@@ -202,8 +163,6 @@ class TPBDRigidClusteredParticles : public TPBDRigidParticles<T, d>
 	  TArrayCollectionArray<int32> MClusterGroupIndex;
 	  TArrayCollectionArray<bool> MInternalCluster;
 	  TArrayCollectionArray<TUniquePtr<FImplicitObjectUnionClustered>> MChildrenSpatial;
-	  TArrayCollectionArray<FMultiChildProxyId> MMultiChildProxyId;
-	  TArrayCollectionArray<TUniquePtr<TMultiChildProxyData<T, d>>> MMultiChildProxyData;
 
 	  // Multiple proxy pointers required for internal clusters
 	  TArrayCollectionArray<TSet<IPhysicsProxyBase*>> MPhysicsProxies;
diff --git a/Engine/Source/Runtime/Experimental/Chaos/Public/Chaos/PBDRigidClustering.h b/Engine/Source/Runtime/Experimental/Chaos/Public/Chaos/PBDRigidClustering.h
--- a/Engine/Source/Runtime/Experimental/Chaos/Public/Chaos/PBDRigidClustering.h
+++ b/Engine/Source/Runtime/Experimental/Chaos/Public/Chaos/PBDRigidClustering.h
@@ -253,12 +253,6 @@ public:
 	*/
 	TArrayCollectionArray<int32>& GetClusterGroupIndexArray() { return MParticles.ClusterGroupIndexArray(); }
 
-	/** Indicates if the child geometry is approximated by a single proxy */
-	const TArrayCollectionArray<FMultiChildProxyId>& GetMultiChildProxyIdArray() const { return MParticles.MultiChildProxyIdArray(); }
-
-	/** If multi child proxy is used, this is the data needed */
-	const TArrayCollectionArray<TUniquePtr<TMultiChildProxyData<FReal, 3>>>& GetMultiChildProxyDataArray() const { return MParticles.MultiChildProxyDataArray(); }
-
 	void AddToClusterUnion(int32 ClusterID, FPBDRigidClusteredParticleHandle* Handle)
 	{
 		if(ClusterID <= 0)
diff --git a/Engine/Source/Runtime/Experimental/Chaos/Public/Chaos/ParticleHandle.h b/Engine/Source/Runtime/Experimental/Chaos/Public/Chaos/ParticleHandle.h
--- a/Engine/Source/Runtime/Experimental/Chaos/Public/Chaos/ParticleHandle.h
+++ b/Engine/Source/Runtime/Experimental/Chaos/Public/Chaos/ParticleHandle.h
@@ -1134,14 +1134,6 @@ public:
 	TUniquePtr<FImplicitObjectUnionClustered>& ChildrenSpatial() { return PBDRigidClusteredParticles->ChildrenSpatial(ParticleIdx); }
 	void SetChildrenSpatial(TUniquePtr<FImplicitObjectUnion>& Obj) { PBDRigidClusteredParticles->ChildrenSpatial(ParticleIdx) = Obj; }
 
-	const FMultiChildProxyId& MultiChildProxyId() const { return PBDRigidClusteredParticles->MultiChildProxyId(ParticleIdx); }
-	FMultiChildProxyId& MultiChildProxyId() { return PBDRigidClusteredParticles->MultiChildProxyId(ParticleIdx); }
-	void SetMultiChildProxyId(const FMultiChildProxyId& Id) { PBDRigidClusteredParticles->MultiChildProxyId(ParticleIdx) = Id; }
-
-	const TUniquePtr<TMultiChildProxyData<T, d>>& MultiChildProxyData() const { return PBDRigidClusteredParticles->MultiChildProxyData(ParticleIdx); }
-	TUniquePtr<TMultiChildProxyData<T, d>>& MultiChildProxyData() { return PBDRigidClusteredParticles->MultiChildProxyData(ParticleIdx); }
-	void SetMultiChildProxyData(TUniquePtr<TMultiChildProxyData<T, d>>&& ProxyData) { PBDRigidClusteredParticles->MultiChildProxyData(ParticleIdx) = MoveTemp(ProxyData); }
-
 	const T& CollisionImpulse() const { return PBDRigidClusteredParticles->CollisionImpulses(ParticleIdx); }
 	T& CollisionImpulse() { return PBDRigidClusteredParticles->CollisionImpulses(ParticleIdx); }
 	void SetCollisionImpulse(const T Value) { PBDRigidClusteredParticles->CollisionImpulses(ParticleIdx) = Value; }
