CIS fix -- change logging on geometry collection upgrades from warning to log, because the upgrade does not indicate an error

diff --git a/Engine/Source/Runtime/Experimental/Chaos/Private/GeometryCollection/GeometryCollection.cpp b/Engine/Source/Runtime/Experimental/Chaos/Private/GeometryCollection/GeometryCollection.cpp
--- a/Engine/Source/Runtime/Experimental/Chaos/Private/GeometryCollection/GeometryCollection.cpp
+++ b/Engine/Source/Runtime/Experimental/Chaos/Private/GeometryCollection/GeometryCollection.cpp
@@ -961,7 +961,7 @@ void FGeometryCollection::Serialize(Chaos::FChaosArchive& Ar)
 		// Version 5 introduced accurate SimulationType tagging
 		if (Version < 5)
 		{
-			UE_LOG(FGeometryCollectionLogging, Warning, TEXT("GeometryCollection has inaccurate simulation type tags. Updating tags based on transform topology."));
+			UE_LOG(FGeometryCollectionLogging, Log, TEXT("GeometryCollection has inaccurate simulation type tags. Updating tags based on transform topology."));
 			TManagedArray<bool>* SimulatableParticles = FindAttribute<bool>(FGeometryCollection::SimulatableParticlesAttribute, FTransformCollection::TransformGroup);
 			TArray<bool> RigidChildren; RigidChildren.Init(false,NumElements(FTransformCollection::TransformGroup));
 			const TArray<int32> RecursiveOrder = GeometryCollectionAlgo::ComputeRecursiveOrder(*this);
@@ -1042,7 +1042,7 @@ void FGeometryCollection::Serialize(Chaos::FChaosArchive& Ar)
 		{
 			if (!HasAttribute("UVs", FGeometryCollection::VerticesGroup))
 			{
-				UE_LOG(FGeometryCollectionLogging, Warning, TEXT("GeometryCollection updated to multiple UV sets."));
+				UE_LOG(FGeometryCollectionLogging, Log, TEXT("GeometryCollection updated to multiple UV sets."));
 				AddAttribute<TArray<FVector2f>>("UVs", FGeometryCollection::VerticesGroup);				
 			}
 
@@ -1137,13 +1137,13 @@ void FGeometryCollection::Serialize(Chaos::FChaosArchive& Ar)
 			}
 			if (InvalidGeometry.Num() > 0)
 			{
-				UE_LOG(FGeometryCollectionLogging, Warning, TEXT("Removing %d invalid, fully-invisible geometries from geometry collection."), InvalidGeometry.Num());
+				UE_LOG(FGeometryCollectionLogging, Log, TEXT("Removing %d invalid, fully-invisible geometries from geometry collection."), InvalidGeometry.Num());
 				InvalidGeometry.Sort();
 				RemoveElements(GeometryGroup, InvalidGeometry);
 			}
 			if (InvalidTransforms.Num() > 0)
 			{
-				UE_LOG(FGeometryCollectionLogging, Warning, TEXT("Removing %d invalid, empty transforms from geometry collection."), InvalidTransforms.Num());
+				UE_LOG(FGeometryCollectionLogging, Log, TEXT("Removing %d invalid, empty transforms from geometry collection."), InvalidTransforms.Num());
 				InvalidTransforms.Sort();
 				RemoveElements(TransformGroup, InvalidTransforms);
 			}
