Static analysis fix: eliminate possibility of null ptr dereference

diff --git a/Engine/Source/Runtime/Experimental/Chaos/Private/Chaos/PBDRigidsEvolution.cpp b/Engine/Source/Runtime/Experimental/Chaos/Private/Chaos/PBDRigidsEvolution.cpp
--- a/Engine/Source/Runtime/Experimental/Chaos/Private/Chaos/PBDRigidsEvolution.cpp
+++ b/Engine/Source/Runtime/Experimental/Chaos/Private/Chaos/PBDRigidsEvolution.cpp
@@ -340,7 +340,7 @@ namespace Chaos
 					ensure(false);
 				}
 				// We should not copy dynamic trees here
-				if (Result->IsTreeDynamic())
+				else if (Result->IsTreeDynamic())
 				{
 					Result = nullptr;
 				}
