Chaos - Remove unnecessary positive epsilon in collision processing
This was causing small but incorrect impulses from friction at small positive separations. No visible behaviour change but stable stack sleep slightly faster now for example.

diff --git a/Engine/Source/Runtime/Experimental/Chaos/Private/Chaos/Collision/PBDCollisionSolver.cpp b/Engine/Source/Runtime/Experimental/Chaos/Private/Chaos/Collision/PBDCollisionSolver.cpp
--- a/Engine/Source/Runtime/Experimental/Chaos/Private/Chaos/Collision/PBDCollisionSolver.cpp
+++ b/Engine/Source/Runtime/Experimental/Chaos/Private/Chaos/Collision/PBDCollisionSolver.cpp
@@ -33,7 +33,6 @@ namespace Chaos
 		int32 Chaos_PBDCollisionSolver_Position_ShockPropagationIterations = 3;
 		float Chaos_PBDCollisionSolver_Position_MinInvMassScale = 0.3f;
 		int32 Chaos_PBDCollisionSolver_Position_ZeroFrictionIterations = 4;
-		float Chaos_PBDCollisionSolver_Position_NormalTolerance = 0.1f;
 		bool bChaos_PBDCollisionSolver_Position_NegativePushOutEnabled = true;
 		float Chaos_PBDCollisionSolver_Position_StaticFrictionStiffness = 0.5f;
 		float Chaos_PBDCollisionSolver_Position_StaticFrictionLerpRate = 0.1f;
@@ -44,7 +43,6 @@ namespace Chaos
 		FAutoConsoleVariableRef CVarChaos_PBDCollisionSolver_Position_UseShockPropagation(TEXT("p.Chaos.PBDCollisionSolver.Position.ShockPropagationIterations"), Chaos_PBDCollisionSolver_Position_ShockPropagationIterations, TEXT(""));
 		FAutoConsoleVariableRef CVarChaos_PBDCollisionSolver_Position_MinInvMassScale(TEXT("p.Chaos.PBDCollisionSolver.Position.MinInvMassScale"), Chaos_PBDCollisionSolver_Position_MinInvMassScale, TEXT(""));
 		FAutoConsoleVariableRef CVarChaos_PBDCollisionSolver_Position_ZeroFrictionIterations(TEXT("p.Chaos.PBDCollisionSolver.Position.ZeroFrictionIterations"), Chaos_PBDCollisionSolver_Position_ZeroFrictionIterations, TEXT(""));
-		FAutoConsoleVariableRef CVarChaos_PBDCollisionSolver_Position_NormalTolerance(TEXT("p.Chaos.PBDCollisionSolver.Position.NormalTolerance"), Chaos_PBDCollisionSolver_Position_NormalTolerance, TEXT(""));
 		FAutoConsoleVariableRef CVarChaos_PBDCollisionSolver_Position_StaticFrictionStiffness(TEXT("p.Chaos.PBDCollisionSolver.Position.StaticFriction.Stiffness"), Chaos_PBDCollisionSolver_Position_StaticFrictionStiffness, TEXT(""));
 		FAutoConsoleVariableRef CVarChaos_PBDCollisionSolver_Position_StaticFrictionLerpRate(TEXT("p.Chaos.PBDCollisionSolver.Position.StaticFriction.LerpRate"), Chaos_PBDCollisionSolver_Position_StaticFrictionLerpRate, TEXT(""));
 		FAutoConsoleVariableRef CVarChaos_PBDCollisionSolver_Position_PositionSolverTolerance(TEXT("p.Chaos.PBDCollisionSolver.Position.PositionTolerance"), Chaos_PBDCollisionSolver_Position_PositionSolverTolerance, TEXT(""));
@@ -605,7 +603,7 @@ namespace Chaos
 			FReal ContactDeltaNormal;
 			SolverManifoldPoint.CalculateContactPositionError(Body0.SolverBody(), Body1.SolverBody(), MaxPushOut, ContactDelta, ContactDeltaNormal);
 
-			const bool bProcessManifoldPoint = (ContactDeltaNormal < Chaos_PBDCollisionSolver_Position_NormalTolerance) || !SolverManifoldPoint.NetPushOut.IsNearlyZero();
+			const bool bProcessManifoldPoint = (ContactDeltaNormal < FReal(0)) || !SolverManifoldPoint.NetPushOut.IsNearlyZero();
 			if (bProcessManifoldPoint)
 			{
 				ApplyPositionCorrection(
diff --git a/Engine/Source/Runtime/Experimental/Chaos/Private/Chaos/Collision/ParticlePairMidPhase.cpp b/Engine/Source/Runtime/Experimental/Chaos/Private/Chaos/Collision/ParticlePairMidPhase.cpp
--- a/Engine/Source/Runtime/Experimental/Chaos/Private/Chaos/Collision/ParticlePairMidPhase.cpp
+++ b/Engine/Source/Runtime/Experimental/Chaos/Private/Chaos/Collision/ParticlePairMidPhase.cpp
@@ -17,6 +17,9 @@ namespace Chaos
 {
 	namespace CVars
 	{
+		bool bChaos_Collision_MidPhase_EnableBoundsChecks = true;
+		FAutoConsoleVariableRef CVarChaos_Collision_EnableBoundsChecks(TEXT("p.Chaos.Collision.EnableBoundsChecks"), bChaos_Collision_MidPhase_EnableBoundsChecks, TEXT(""));
+
 		bool bChaos_Collision_EnableManifoldRestore = true;
 		Chaos::FRealSingle Chaos_Collision_RestoreTolerance_NoContact_Position = 0.005f;	// About 0.5cm for a meter cube
 		Chaos::FRealSingle Chaos_Collision_RestoreTolerance_NoContact_Rotation = 0.1f;		// About 10deg
@@ -124,12 +127,12 @@ namespace Chaos
 		const bool bIsSphere0 = (ImplicitType0 == ImplicitObjectType::Sphere);
 		const bool bIsSphere1 = (ImplicitType1 == ImplicitObjectType::Sphere);
 
-		bEnableAABBCheck = !bIsSphere0 || !bIsSphere1;
-		bEnableOBBCheck0 = Chaos_Collision_NarrowPhase_AABBBoundsCheck && !bIsSphere0;
-		bEnableOBBCheck1 = Chaos_Collision_NarrowPhase_AABBBoundsCheck && !bIsSphere1;
+		bEnableAABBCheck = bChaos_Collision_MidPhase_EnableBoundsChecks && (!bIsSphere0 || !bIsSphere1);
+		bEnableOBBCheck0 = bChaos_Collision_MidPhase_EnableBoundsChecks && !bIsSphere0;
+		bEnableOBBCheck1 = bChaos_Collision_MidPhase_EnableBoundsChecks && !bIsSphere1;
 		bEnableManifoldCheck = bChaos_Collision_EnableManifoldUpdate && !bIsSphere0 && !bIsSphere1;
 
-		if (bIsSphere0 && bIsSphere1)
+		if (bIsSphere0 && bIsSphere1 && bChaos_Collision_MidPhase_EnableBoundsChecks)
 		{
 			SphereBoundsCheckSize = Shape0->GetLeafGeometry()->GetMargin() + Shape1->GetLeafGeometry()->GetMargin();
 		}
