Chaos - fix incorrect edge-edge determination in convex collision

diff --git a/Engine/Source/Runtime/Experimental/Chaos/Private/Chaos/CollisionOneShotManifolds.cpp b/Engine/Source/Runtime/Experimental/Chaos/Private/Chaos/CollisionOneShotManifolds.cpp
--- a/Engine/Source/Runtime/Experimental/Chaos/Private/Chaos/CollisionOneShotManifolds.cpp
+++ b/Engine/Source/Runtime/Experimental/Chaos/Private/Chaos/CollisionOneShotManifolds.cpp
@@ -788,13 +788,13 @@ namespace Chaos
 			const FVec3 SeparationDirectionLocalConvex1 = Convex1Transform.InverseTransformVectorNoScale(GJKContactPoint.Normal);
 			const int32 MostOpposingPlaneIndexConvex1 = SelectContactPlane(Convex1, GJKContactPoint.ShapeContactPoints[0], SeparationDirectionLocalConvex1, Margin1, VertexIndexA);
 			const TPlaneConcrete<FReal, 3> BestPlaneConvex1 = Convex1.GetPlane(MostOpposingPlaneIndexConvex1);
-			const FReal BestPlaneDotNormalConvex1 = !bConvex1IsCapsule ? FVec3::DotProduct(-SeparationDirectionLocalConvex1, BestPlaneConvex1.Normal()) : -FLT_MAX;
+			const FReal BestPlaneDotNormalConvex1 = !bConvex1IsCapsule ? FMath::Abs(FVec3::DotProduct(-SeparationDirectionLocalConvex1, BestPlaneConvex1.Normal())) : -FLT_MAX;
 
 			// Now for Convex2
 			const FVec3 SeparationDirectionLocalConvex2 = Convex2Transform.InverseTransformVectorNoScale(GJKContactPoint.Normal);
 			const int32 MostOpposingPlaneIndexConvex2 = SelectContactPlane(Convex2, GJKContactPoint.ShapeContactPoints[1], -SeparationDirectionLocalConvex2, Margin2, VertexIndexB);
 			const TPlaneConcrete<FReal, 3> BestPlaneConvex2 = Convex2.GetPlane(MostOpposingPlaneIndexConvex2);
-			const FReal BestPlaneDotNormalConvex2 = FVec3::DotProduct(SeparationDirectionLocalConvex2, BestPlaneConvex2.Normal());
+			const FReal BestPlaneDotNormalConvex2 = FMath::Abs(FVec3::DotProduct(SeparationDirectionLocalConvex2, BestPlaneConvex2.Normal()));
 
 			const FReal SmallBiasToPreventFeatureFlipping = 0.002f; // This improves frame coherence by penalizing convex 1 in favour of convex 2
 			bool ReferenceFaceConvex1 = true; // Is the reference face on convex1 or convex2?
