Chaos
- inline manifold points array
- reduce size of midphase objects

diff --git a/Engine/Source/Runtime/Experimental/Chaos/Private/Chaos/Collision/PBDCollisionConstraint.cpp b/Engine/Source/Runtime/Experimental/Chaos/Private/Chaos/Collision/PBDCollisionConstraint.cpp
--- a/Engine/Source/Runtime/Experimental/Chaos/Private/Chaos/Collision/PBDCollisionConstraint.cpp
+++ b/Engine/Source/Runtime/Experimental/Chaos/Private/Chaos/Collision/PBDCollisionConstraint.cpp
@@ -475,9 +475,7 @@ namespace Chaos
 
 	void FPBDCollisionConstraint::InitManifoldPoint(const int32 ManifoldPointIndex)
 	{
-		FConstGenericParticleHandle Particle0 = Particle[0];
-		FConstGenericParticleHandle Particle1 = Particle[1];
-		if (!Particle0.IsValid() || !Particle1.IsValid())
+		if ((Particle[0] == nullptr) || (Particle[1] == nullptr))
 		{
 			// @todo(chaos): This is just for unit tests testing one-shot manifolds - remove it somehow... 
 			// maybe ConstructConvexConvexOneShotManifold should not take a Constraint
diff --git a/Engine/Source/Runtime/Experimental/Chaos/Private/Chaos/Collision/ParticlePairMidPhase.cpp b/Engine/Source/Runtime/Experimental/Chaos/Private/Chaos/Collision/ParticlePairMidPhase.cpp
--- a/Engine/Source/Runtime/Experimental/Chaos/Private/Chaos/Collision/ParticlePairMidPhase.cpp
+++ b/Engine/Source/Runtime/Experimental/Chaos/Private/Chaos/Collision/ParticlePairMidPhase.cpp
@@ -117,25 +117,29 @@ namespace Chaos
 		, Shape1(InShape1)
 		, ShapePairType(InShapePairType)
 		, SphereBoundsCheckSize(0)
-		, bEnableAABBCheck(false)
-		, bEnableOBBCheck0(false)
-		, bEnableOBBCheck1(false)
-		, bEnableManifoldCheck(false)
+		, LastUsedEpoch(-1)
+		, Flags()
 	{
-		const EImplicitObjectType ImplicitType0 = (Shape0->GetLeafGeometry() != nullptr) ? GetInnerType(Shape0->GetLeafGeometry()->GetCollisionType()) : ImplicitObjectType::Unknown;
-		const EImplicitObjectType ImplicitType1 = (Shape1->GetLeafGeometry() != nullptr) ? GetInnerType(Shape1->GetLeafGeometry()->GetCollisionType()) : ImplicitObjectType::Unknown;
+		const FImplicitObject* Implicit0 = Shape0->GetLeafGeometry();
+		const FImplicitObject* Implicit1 = Shape1->GetLeafGeometry();
+		const bool bHasBounds0 = (Implicit0 != nullptr) && Implicit0->HasBoundingBox();
+		const bool bHasBounds1 = (Implicit1 != nullptr) && Implicit1->HasBoundingBox();
+		const EImplicitObjectType ImplicitType0 = (Implicit0 != nullptr) ? GetInnerType(Implicit0->GetCollisionType()) : ImplicitObjectType::Unknown;
+		const EImplicitObjectType ImplicitType1 = (Implicit1 != nullptr) ? GetInnerType(Implicit1->GetCollisionType()) : ImplicitObjectType::Unknown;
 		const bool bIsSphere0 = (ImplicitType0 == ImplicitObjectType::Sphere);
 		const bool bIsSphere1 = (ImplicitType1 == ImplicitObjectType::Sphere);
 
-		bEnableAABBCheck = bChaos_Collision_MidPhase_EnableBoundsChecks && (!bIsSphere0 || !bIsSphere1);
-		bEnableOBBCheck0 = bChaos_Collision_MidPhase_EnableBoundsChecks && !bIsSphere0;
-		bEnableOBBCheck1 = bChaos_Collision_MidPhase_EnableBoundsChecks && !bIsSphere1;
-		bEnableManifoldCheck = bChaos_Collision_EnableManifoldUpdate && !bIsSphere0 && !bIsSphere1;
-
-		if (bIsSphere0 && bIsSphere1 && bChaos_Collision_MidPhase_EnableBoundsChecks)
+		const bool bAllowBoundsChecked = bChaos_Collision_MidPhase_EnableBoundsChecks && bHasBounds0 && bHasBounds1;
+		Flags.bEnableAABBCheck = bAllowBoundsChecked && (!bIsSphere0 || !bIsSphere1);
+		Flags.bEnableOBBCheck0 = bAllowBoundsChecked && !bIsSphere0;
+		Flags.bEnableOBBCheck1 = bAllowBoundsChecked && !bIsSphere1;
+		
+		if (bAllowBoundsChecked && bIsSphere0 && bIsSphere1)
 		{
 			SphereBoundsCheckSize = Shape0->GetLeafGeometry()->GetMargin() + Shape1->GetLeafGeometry()->GetMargin();
 		}
+
+		Flags.bEnableManifoldCheck = bChaos_Collision_EnableManifoldUpdate && !bIsSphere0 && !bIsSphere1;
 	}
 
 	FSingleShapePairCollisionDetector::~FSingleShapePairCollisionDetector()
@@ -151,87 +155,89 @@ namespace Chaos
 		, Shape1(R.Shape1)
 		, ShapePairType(R.ShapePairType)
 		, SphereBoundsCheckSize(R.SphereBoundsCheckSize)
-		, bEnableAABBCheck(R.bEnableAABBCheck)
-		, bEnableOBBCheck0(R.bEnableOBBCheck0)
-		, bEnableOBBCheck1(R.bEnableOBBCheck1)
-		, bEnableManifoldCheck(R.bEnableManifoldCheck)
+		, Flags(R.Flags)
 	{
 	}
 
 	inline bool FSingleShapePairCollisionDetector::IsUsedSince(const int32 Epoch) const
 	{
 		// If we have no constraint it was never used, so this check is always false
-		return (Constraint != nullptr) && (Constraint->GetContainerCookie().LastUsedEpoch >= Epoch);
+		return (Constraint != nullptr) && (LastUsedEpoch >= Epoch);
 	}
 
-	int32 FSingleShapePairCollisionDetector::GenerateCollision(
-		const FReal CullDistance, 
-		const bool bUseCCD,
-		const FReal Dt)
+	bool FSingleShapePairCollisionDetector::DoBoundsOverlap(const FReal CullDistance)
 	{
+		PHYSICS_CSV_SCOPED_EXPENSIVE(PhysicsVerbose, NarrowPhase_ShapeBounds);
 
-		// Shape-pair Bounds check (not for CCD)
-		const FImplicitObject* Implicit0 = Shape0->GetLeafGeometry();
-		const FImplicitObject* Implicit1 = Shape1->GetLeafGeometry();
-		if (Implicit0->HasBoundingBox() && Implicit1->HasBoundingBox() && !bUseCCD)
-		{
-			PHYSICS_CSV_SCOPED_EXPENSIVE(PhysicsVerbose, NarrowPhase_ShapeBounds);
-
-			const FAABB3& ShapeWorldBounds0 = Shape0->GetWorldSpaceInflatedShapeBounds();
-			const FAABB3& ShapeWorldBounds1 = Shape1->GetWorldSpaceInflatedShapeBounds();
+		const FAABB3& ShapeWorldBounds0 = Shape0->GetWorldSpaceInflatedShapeBounds();
+		const FAABB3& ShapeWorldBounds1 = Shape1->GetWorldSpaceInflatedShapeBounds();
 
-			// World-space expanded bounds check
-			if (bEnableAABBCheck)
+		// World-space expanded bounds check
+		if (Flags.bEnableAABBCheck)
+		{
+			if (!ShapeWorldBounds0.Intersects(ShapeWorldBounds1))
 			{
-				if (!ShapeWorldBounds0.Intersects(ShapeWorldBounds1))
-				{
-					return 0;
-				}
+				return false;
 			}
+		}
 
-			if (SphereBoundsCheckSize > FReal(0))
+		// World-space sphere bounds check
+		if (SphereBoundsCheckSize > FReal(0))
+		{
+			const FVec3 Separation = ShapeWorldBounds0.GetCenter() - ShapeWorldBounds1.GetCenter();
+			const FReal SeparationSq = Separation.SizeSquared();
+			const FReal CullDistanceSq = FMath::Square(CullDistance + SphereBoundsCheckSize);
+			if (SeparationSq > CullDistanceSq)
 			{
-				const FVec3 Separation = ShapeWorldBounds0.GetCenter() - ShapeWorldBounds1.GetCenter();
-				const FReal SeparationSq = Separation.SizeSquared();
-				const FReal CullDistanceSq = FMath::Square(CullDistance + SphereBoundsCheckSize);
-				if (SeparationSq > CullDistanceSq)
-				{
-					return 0;
-				}
+				return false;
 			}
+		}
+
+		// OBB-AABB test on both directions. This is beneficial for shapes which do not fit their AABBs very well,
+		// which includes boxes and other shapes that are not roughly spherical. It is especially beneficial when
+		// one shape is long and thin (i.e., it does not fit an AABB well when the shape is rotated).
+		// However, it is quite expensive to do this all the time so we only do this test when we did not 
+		// collide last frame. This is ok if we assume not much changes from frame to frame, but it means
+		// we might call the narrow phase one time too many when shapes become separated.
+		const int32 LastEpoch = MidPhase.GetCollisionAllocator().GetCurrentEpoch() - 1;
+		const bool bCollidedLastTick = IsUsedSince(LastEpoch);
+		if ((Flags.bEnableOBBCheck0 || Flags.bEnableOBBCheck1) && !bCollidedLastTick)
+		{
+			const FRigidTransform3& ShapeWorldTransform0 = Shape0->GetLeafWorldTransform();
+			const FRigidTransform3& ShapeWorldTransform1 = Shape1->GetLeafWorldTransform();
+			const FImplicitObject* Implicit0 = Shape0->GetLeafGeometry();
+			const FImplicitObject* Implicit1 = Shape1->GetLeafGeometry();
 
-			const int32 LastEpoch = MidPhase.GetCollisionAllocator().GetCurrentEpoch() - 1;
-			const bool bCollidedLastTick = IsUsedSince(LastEpoch);
-			if ((bEnableOBBCheck0 || bEnableOBBCheck1) && !bCollidedLastTick)
+			if (Flags.bEnableOBBCheck0)
 			{
-				// OBB-AABB test on both directions. This is beneficial for shapes which do not fit their AABBs very well,
-				// which includes boxes and other shapes that are not roughly spherical. It is especially beneficial when
-				// one shape is long and thin (i.e., it does not fit an AABB well when the shape is rotated).
-				// However, it is quite expensive to do this all the time so we only do this test when we did not 
-				// collide last frame. This is ok if we assume not much changes from frame to frame, but it means
-				// we might call the narrow phase one time too many when shapes become separated.
-				const FRigidTransform3& ShapeWorldTransform0 = Shape0->GetLeafWorldTransform();
-				const FRigidTransform3& ShapeWorldTransform1 = Shape1->GetLeafWorldTransform();
-
-				if (bEnableOBBCheck0)
+				if (!ImplicitOverlapOBBToAABB(Implicit0, Implicit1, ShapeWorldTransform0, ShapeWorldTransform1, CullDistance))
 				{
-					if (!ImplicitOverlapOBBToAABB(Implicit0, Implicit1, ShapeWorldTransform0, ShapeWorldTransform1, CullDistance))
-					{
-						return 0;
-					}
+					return false;
 				}
+			}
 
-				if (bEnableOBBCheck1)
+			if (Flags.bEnableOBBCheck1)
+			{
+				if (!ImplicitOverlapOBBToAABB(Implicit1, Implicit0, ShapeWorldTransform1, ShapeWorldTransform0, CullDistance))
 				{
-					if (!ImplicitOverlapOBBToAABB(Implicit1, Implicit0, ShapeWorldTransform1, ShapeWorldTransform0, CullDistance))
-					{
-						return 0;
-					}
+					return false;
 				}
 			}
 		}
 
-		return GenerateCollisionImpl(CullDistance, bUseCCD, Dt);
+		return true;
+	}
+
+	int32 FSingleShapePairCollisionDetector::GenerateCollision(
+		const FReal CullDistance,
+		const bool bUseCCD,
+		const FReal Dt)
+	{
+		if (bUseCCD || DoBoundsOverlap(CullDistance))
+		{
+			return GenerateCollisionImpl(CullDistance, bUseCCD, Dt);
+		}
+		return 0;
 	}
 
 	void FSingleShapePairCollisionDetector::CreateConstraint(const FReal CullDistance)
@@ -244,6 +250,7 @@ namespace Chaos
 		Constraint->GetContainerCookie().MidPhase = &MidPhase;
 		Constraint->GetContainerCookie().bIsMultiShapePair = false;
 		Constraint->GetContainerCookie().CreationEpoch = MidPhase.GetCollisionAllocator().GetCurrentEpoch();
+		LastUsedEpoch = -1;
 	}
 
 	int32 FSingleShapePairCollisionDetector::GenerateCollisionImpl(const FReal CullDistance, const bool bUseCCD, const FReal Dt)
@@ -265,7 +272,8 @@ namespace Chaos
 			const FImplicitObject* Implicit1 = Shape1->GetLeafGeometry();
 			const FRigidTransform3& ShapeWorldTransform0 = Shape0->GetLeafWorldTransform();
 			const FRigidTransform3& ShapeWorldTransform1 = Shape1->GetLeafWorldTransform();
-			const int32 LastEpoch = MidPhase.GetCollisionAllocator().GetCurrentEpoch() - 1;
+			const int32 CurrentEpoch = MidPhase.GetCollisionAllocator().GetCurrentEpoch();
+			const int32 LastEpoch = CurrentEpoch - 1;
 			const bool bWasUpdatedLastTick = IsUsedSince(LastEpoch);
 
 			// Update the world shape transforms on the constraint (we cannot just give it the PerShapeData 
@@ -284,7 +292,7 @@ namespace Chaos
 			}
 
 			bool bWasManifoldRestored = false;
-			if (bEnableManifoldCheck)
+			if (Flags.bEnableManifoldCheck)
 			{
 				// Update the existing manifold. We can re-use as-is if none of the points have moved much and the bodies have not moved much
 				// NOTE: this can succeed in "restoring" even if we have no manifold points
@@ -327,6 +335,7 @@ namespace Chaos
 			{
 				if (MidPhase.GetCollisionAllocator().ActivateConstraint(Constraint.Get()))
 				{
+					LastUsedEpoch = CurrentEpoch;
 					return 1;
 				}
 			}
@@ -338,7 +347,8 @@ namespace Chaos
 	int32 FSingleShapePairCollisionDetector::RestoreCollision(const FReal CullDistance)
 	{
 		// Only restore constraints if active last tick. Any older than that and the shapes were separated for a bit
-		const int32 LastEpoch = MidPhase.GetCollisionAllocator().GetCurrentEpoch() - 1;
+		const int32 CurrentEpoch = MidPhase.GetCollisionAllocator().GetCurrentEpoch();
+		const int32 LastEpoch = CurrentEpoch - 1;
 		if ((Constraint != nullptr) && IsUsedSince(LastEpoch))
 		{
 			const FRigidTransform3& ShapeWorldTransform0 = Shape0->GetLeafWorldTransform();
@@ -350,6 +360,7 @@ namespace Chaos
 			{
 				if (MidPhase.GetCollisionAllocator().ActivateConstraint(Constraint.Get()))
 				{
+					LastUsedEpoch = CurrentEpoch;
 					return 1;
 				}
 			}
@@ -363,18 +374,22 @@ namespace Chaos
 		{
 			// We just need to refresh the epoch so that the constraint state will be used as the previous
 			// state iof the pair is still colliding in the next tick
+			const int32 CurrentEpoch = MidPhase.GetCollisionAllocator().GetCurrentEpoch();
 			Constraint->GetContainerCookie().LastUsedEpoch = MidPhase.GetCollisionAllocator().GetCurrentEpoch();
+			LastUsedEpoch = CurrentEpoch;
 		}
 	}
 
 	void FSingleShapePairCollisionDetector::SetCollision(const FPBDCollisionConstraint& SourceConstraint)
 	{
+		const int32 CurrentEpoch = MidPhase.GetCollisionAllocator().GetCurrentEpoch();
+
 		if (Constraint == nullptr)
 		{
 			Constraint = MakeUnique<FPBDCollisionConstraint>();
 			Constraint->GetContainerCookie().MidPhase = &MidPhase;
 			Constraint->GetContainerCookie().bIsMultiShapePair = false;
-			Constraint->GetContainerCookie().CreationEpoch = MidPhase.GetCollisionAllocator().GetCurrentEpoch();
+			Constraint->GetContainerCookie().CreationEpoch = CurrentEpoch;
 		}
 
 		// Copy the constraint over the existing one, taking care to leave the cookie intact
@@ -385,6 +400,7 @@ namespace Chaos
 		// Add the constraint to the active list
 		// If the constraint already existed and was already active, this will do nothing
 		MidPhase.GetCollisionAllocator().ActivateConstraint(Constraint.Get());
+		LastUsedEpoch = CurrentEpoch;
 	}
 
 
@@ -689,25 +705,18 @@ namespace Chaos
 	////////////////////////////////////////////////////////////////////////////////////////
 
 
-	FParticlePairMidPhase::FParticlePairMidPhase(
-		FGeometryParticleHandle* InParticle0, 
-		FGeometryParticleHandle* InParticle1, 
-		const FCollisionParticlePairKey& InKey,
-		FCollisionConstraintAllocator& InCollisionAllocator)
-		: Particle0(InParticle0)
-		, Particle1(InParticle1)
-		, Key(InKey)
-		, ParticleCollisionsIndex0(INDEX_NONE)
-		, ParticleCollisionsIndex1(INDEX_NONE)
+	FParticlePairMidPhase::FParticlePairMidPhase()
+		: Particle0(nullptr)
+		, Particle1(nullptr)
+		, Key()
 		, ShapePairDetectors()
 		, MultiShapePairDetectors()
-		, CollisionAllocator(&InCollisionAllocator)
-		, bIsCCD(false)
-		, bIsInitialized(false)
-		, bRestorable(false)
-		, bIsSleeping(false)
+		, CollisionAllocator(nullptr)
+		, Flags()
 		, LastUsedEpoch(INDEX_NONE)
 		, NumActiveConstraints(0)
+		, ParticleCollisionsIndex0(INDEX_NONE)
+		, ParticleCollisionsIndex1(INDEX_NONE)
 		, RestoreThresholdZeroContacts()
 		, RestoreThreshold()
 		, RestoreParticleP0(FVec3(0))
@@ -716,10 +725,6 @@ namespace Chaos
 		, RestoreParticleQ1(FRotation3::FromIdentity())
 		, CullDistanceScale(1)
 	{
-		if (ensure(Particle0 != Particle1))
-		{
-			Init();
-		}
 	}
 
 	FParticlePairMidPhase::~FParticlePairMidPhase()
@@ -746,24 +751,34 @@ namespace Chaos
 		ShapePairDetectors.Reset();
 		MultiShapePairDetectors.Reset();
 
-		bIsCCD = false;
-		bIsInitialized = false;
-		bIsSleeping = false;
+		Flags.bIsCCD = false;
+		Flags.bIsInitialized = false;
+		Flags.bRestorable = false;
+		Flags.bIsSleeping = false;
 	}
 
-	void FParticlePairMidPhase::Init()
+	void FParticlePairMidPhase::Init(
+		FGeometryParticleHandle* InParticle0,
+		FGeometryParticleHandle* InParticle1,
+		const FCollisionParticlePairKey& InKey,
+		FCollisionConstraintAllocator& InCollisionAllocator)
 	{
 		PHYSICS_CSV_SCOPED_EXPENSIVE(PhysicsVerbose, NarrowPhase_Filter);
 
-		bIsCCD = FConstGenericParticleHandle(Particle0)->CCDEnabled() || FConstGenericParticleHandle(Particle1)->CCDEnabled();
+		Particle0 = InParticle0;
+		Particle1 = InParticle1;
+		Key = InKey;
+		CollisionAllocator = &InCollisionAllocator;
+
+		Flags.bIsCCD = FConstGenericParticleHandle(Particle0)->CCDEnabled() || FConstGenericParticleHandle(Particle1)->CCDEnabled();
 
-		bRestorable = true;
+		Flags.bRestorable = true;
 
 		BuildDetectors();
 
 		InitThresholds();
 
-		bIsInitialized = true;
+		Flags.bIsInitialized = true;
 	}
 
 	void FParticlePairMidPhase::BuildDetectors()
@@ -825,14 +840,14 @@ namespace Chaos
 			// @todo(chaos): Unions may need to reactivate constraints that were not used last frame to support full restore
 			if ((ShapePairType == EContactShapesType::LevelSetLevelSet) || (ShapePairType == EContactShapesType::Unknown))
 			{
-				bRestorable = false;
+				Flags.bRestorable = false;
 			}
 		}
 	}
 
 	bool FParticlePairMidPhase::ShouldEnableCCD(const FReal Dt)
 	{
-		if (bIsCCD)
+		if (Flags.bIsCCD)
 		{
 			FConstGenericParticleHandle ConstParticle0 = FConstGenericParticleHandle(Particle0);
 			FConstGenericParticleHandle ConstParticle1 = FConstGenericParticleHandle(Particle1);
@@ -952,10 +967,10 @@ namespace Chaos
 		const FReal CullDistance = InCullDistance * CullDistanceScale;
 
 		// Enable CCD?
-		const bool bUseCCD = bIsCCD && ShouldEnableCCD(Dt);
+		const bool bUseCCD = Flags.bIsCCD && ShouldEnableCCD(Dt);
 
 		// If the bodies have not moved, we will reuse the constraints as-is
-		const bool bCanRestore = bChaos_Collision_EnableManifoldRestore && bRestorable && !bUseCCD;
+		const bool bCanRestore = bChaos_Collision_EnableManifoldRestore && Flags.bRestorable && !bUseCCD;
 		const bool bWasRestored = bCanRestore && TryRestoreConstraints(Dt, CullDistance);
 
 		// If the bodies have moved we need to create or update the constraints
@@ -1019,7 +1034,7 @@ namespace Chaos
 		// We don't need to do anything when going to sleep because sleeping particles pairs are ignored in collision detection 
 		// so the next set of active collisions generated will not contain these collisions.
 
-		if (bIsSleeping != bInIsSleeping)
+		if (Flags.bIsSleeping != bInIsSleeping)
 		{
 			// If we are waking particles, reactivate all collisions that were
 			// active when we were put to sleep, i.e., all collisions whose LastUsedEpoch
@@ -1043,7 +1058,7 @@ namespace Chaos
 			}
 			// If we are going to sleep, there is nothing to do (see comments above)
 
-			bIsSleeping = bInIsSleeping;
+			Flags.bIsSleeping = bInIsSleeping;
 		}
 	}
 
diff --git a/Engine/Source/Runtime/Experimental/Chaos/Public/Chaos/Collision/CollisionConstraintAllocator.h b/Engine/Source/Runtime/Experimental/Chaos/Public/Chaos/Collision/CollisionConstraintAllocator.h
--- a/Engine/Source/Runtime/Experimental/Chaos/Public/Chaos/Collision/CollisionConstraintAllocator.h
+++ b/Engine/Source/Runtime/Experimental/Chaos/Public/Chaos/Collision/CollisionConstraintAllocator.h
@@ -165,6 +165,8 @@ namespace Chaos
 		*/
 		FParticlePairMidPhase* GetParticlePairMidPhase(FGeometryParticleHandle* Particle0, FGeometryParticleHandle* Particle1)
 		{
+			// NOTE: Called from the CollisionDetection parellel-for loop.
+
 			FParticlePairMidPhase* MidPhase = FindParticlePairMidPhaseImpl(Particle0, Particle1);
 			if (MidPhase == nullptr)
 			{
@@ -187,12 +189,17 @@ namespace Chaos
 		*/
 		bool ActivateConstraint(FPBDCollisionConstraint* Constraint)
 		{
+			// NOTE: Called from the CollisionDetection parellel-for loop.
+			// We need to lock the arrays, but we can freely read/write to the constraint without the lock
+			// because each constraint is only processed once and not accessed by any other collision detection threads
+
 			// When we wake an Island, we reactivate all constraints for all dynamic particles in the island. This
-			// results in duplicate calls to active for constraints involving two dynamic particles.
+			// results in duplicate calls to active for constraints involving two dynamic particles, hence the check for CurrentEpoch.
 			// @todo(chaos): fix duplicate calls from island wake. See UpdateSleepState in IslandManager.cpp
-			if (Constraint->GetContainerCookie().LastUsedEpoch != CurrentEpoch)
+			FPBDCollisionConstraintContainerCookie& Cookie = Constraint->GetContainerCookie();
+			if (Cookie.LastUsedEpoch != CurrentEpoch)
 			{
-				Constraint->GetContainerCookie().LastUsedEpoch = CurrentEpoch;
+				Cookie.LastUsedEpoch = CurrentEpoch;
 
 				NewConstraints.Push(Constraint);
 
@@ -266,10 +273,10 @@ namespace Chaos
 		*/
 		void RemoveParticle(FGeometryParticleHandle* Particle)
 		{
-			// We will be removing collisions, and don't want to have to prune the NewConstraints queue
+			// We will be removing collisions, and don't want to have to prune the queues
 			check(!bInCollisionDetectionPhase);
 
-			// Lop over all particle pairs involving this particle.
+			// Loop over all particle pairs involving this particle.
 			// Tell each Particle Pair MidPhase that one of its particles is gone. 
 			// It will get pruned at the next collision detection phase.
 			Particle->ParticleCollisions().VisitMidPhases([Particle](FParticlePairMidPhase& MidPhase)
@@ -316,7 +323,8 @@ namespace Chaos
 			const FCollisionParticlePairKey Key = FCollisionParticlePairKey(Particle0, Particle1);
 
 			// We enqueue a raw pointer and wrap it in a UniquePtr later
-			FParticlePairMidPhase* MidPhase = new FParticlePairMidPhase(Particle0, Particle1, Key, *this);
+			FParticlePairMidPhase* MidPhase = new FParticlePairMidPhase();
+			MidPhase->Init(Particle0, Particle1, Key, *this);
 
 			NewParticlePairMidPhases.Push(MidPhase);
 
@@ -357,9 +365,11 @@ namespace Chaos
 		void PruneExpiredMidPhases()
 		{
 			check(NewParticlePairMidPhases.IsEmpty());
+			
+			// NOTE: Called from the physics thread. No need for locks.
 
 			// Determine which particle pairs are no longer overlapping
-			// Prune all pairs which wer not updated this tick as part of the collision
+			// Prune all pairs which were not updated this tick as part of the collision
 			// detection loop and are not asleep
 			int32 NumParticlePairMidPhases = ParticlePairMidPhases.Num();
 			for (int32 Index = 0; Index < NumParticlePairMidPhases; ++Index)
diff --git a/Engine/Source/Runtime/Experimental/Chaos/Public/Chaos/Collision/PBDCollisionConstraint.h b/Engine/Source/Runtime/Experimental/Chaos/Public/Chaos/Collision/PBDCollisionConstraint.h
--- a/Engine/Source/Runtime/Experimental/Chaos/Public/Chaos/Collision/PBDCollisionConstraint.h
+++ b/Engine/Source/Runtime/Experimental/Chaos/Public/Chaos/Collision/PBDCollisionConstraint.h
@@ -543,7 +543,7 @@ namespace Chaos
 		ECollisionCCDType CCDType;
 		FReal Stiffness;
 
-		TArray<FManifoldPoint> ManifoldPoints;
+		TFixedArray<FManifoldPoint, MaxManifoldPoints> ManifoldPoints;
 
 		// The manifold points from the previous tick when we don't reuse the manifold. Used by static friction.
 		TFixedArray<FManifoldPointSavedData, MaxManifoldPoints> ManifoldPointsSavedData;
diff --git a/Engine/Source/Runtime/Experimental/Chaos/Public/Chaos/Collision/ParticlePairMidPhase.h b/Engine/Source/Runtime/Experimental/Chaos/Public/Chaos/Collision/ParticlePairMidPhase.h
--- a/Engine/Source/Runtime/Experimental/Chaos/Public/Chaos/Collision/ParticlePairMidPhase.h
+++ b/Engine/Source/Runtime/Experimental/Chaos/Public/Chaos/Collision/ParticlePairMidPhase.h
@@ -89,12 +89,16 @@ namespace Chaos
 	private:
 		int32 GenerateCollisionImpl(const FReal CullDistance, const bool bUseCCD, const FReal Dt);
 
-		void CreateConstraint(const FReal CullDistance);
+		/**
+		 * @brief Whether the two shapes are separated by less than CullDistance (i.e., we should run the narrow phase).
+		 * Also returns true if bounds checking is disabled (globally or for this pair)
+		*/
+		bool DoBoundsOverlap(const FReal CullDistance);
 
 		/**
-		 * @brief Add the constraint to the scene's active list
+		 * @brief Create a constraint
 		*/
-		bool ActivateConstraint();
+		void CreateConstraint(const FReal CullDistance);
 
 		FParticlePairMidPhase& MidPhase;
 		TUniquePtr<FPBDCollisionConstraint> Constraint;
@@ -104,10 +108,19 @@ namespace Chaos
 		const FPerShapeData* Shape1;
 		EContactShapesType ShapePairType;
 		FReal SphereBoundsCheckSize;
-		bool bEnableAABBCheck;
-		bool bEnableOBBCheck0;
-		bool bEnableOBBCheck1;
-		bool bEnableManifoldCheck;
+		int32 LastUsedEpoch;
+		union FFlags
+		{
+			FFlags() : Bits(0) {}
+			struct 
+			{
+				uint32 bEnableAABBCheck : 1;
+				uint32 bEnableOBBCheck0 : 1;
+				uint32 bEnableOBBCheck1 : 1;
+				uint32 bEnableManifoldCheck : 1;
+			};
+			uint32 Bits;
+		} Flags;
 	};
 
 
@@ -252,16 +265,23 @@ namespace Chaos
 	class CHAOS_API FParticlePairMidPhase
 	{
 	public:
-		FParticlePairMidPhase(
-			FGeometryParticleHandle* InParticle0, 
-			FGeometryParticleHandle* InParticle1, 
-			const FCollisionParticlePairKey& InKey,
-			FCollisionConstraintAllocator& InCollisionAllocator);
+		FParticlePairMidPhase();
 
 		UE_NONCOPYABLE(FParticlePairMidPhase);
 
 		~FParticlePairMidPhase();
 
+		/**
+		 * @brief Set up the midphase based on the SHapesArrays of the two particles
+		 * Only intended to be called once right after constructor. We don't do this work in
+		 * the constructor so that we can reduce the time that the lock is held when allocating
+		 * new MidPhases.
+		*/
+		void Init(
+			FGeometryParticleHandle* InParticle0,
+			FGeometryParticleHandle* InParticle1,
+			const FCollisionParticlePairKey& InKey,
+			FCollisionConstraintAllocator& InCollisionAllocator);
 
 		inline FGeometryParticleHandle* GetParticle0() { return Particle0; }
 
@@ -285,7 +305,7 @@ namespace Chaos
 		/**
 		 * @brief Whether the particle pair is sleeping and therefore contacts should not be culled (they will be reused on wake)
 		*/
-		inline bool IsSleeping() const { return bIsSleeping; }
+		inline bool IsSleeping() const { return Flags.bIsSleeping; }
 
 		/**
 		 * @brief Update the sleeping state
@@ -361,11 +381,6 @@ namespace Chaos
 		}
 
 	private:
-		/**
-		 * @brief Set up the midphase based on the SHapesArrays of the two particles
-		*/
-		void Init();
-
 		/**
 		 * @brief Build the list of potentially colliding shape pairs.
 		 * This is all the shape pairs in the partilces' shapes arrays that pass the collision filter.
@@ -405,21 +420,29 @@ namespace Chaos
 		FGeometryParticleHandle* Particle1;
 		FCollisionParticlePairKey Key;
 
-		// Indices into the arrays of collisions on the particles. This is a cookie for use by FParticleCollisions
-		int32 ParticleCollisionsIndex0;
-		int32 ParticleCollisionsIndex1;
-
 		TArray<FSingleShapePairCollisionDetector, TInlineAllocator<1>> ShapePairDetectors;
 		TArray<FMultiShapePairCollisionDetector> MultiShapePairDetectors;
 		FCollisionConstraintAllocator* CollisionAllocator;
 
-		bool bIsCCD;
-		bool bIsInitialized;
-		bool bRestorable;
-		bool bIsSleeping;
+		union FFlags
+		{
+			FFlags() : Bits(0) {}
+			struct
+			{
+				uint32 bIsCCD : 1;
+				uint32 bIsInitialized : 1;
+				uint32 bRestorable : 1;
+				uint32 bIsSleeping : 1;
+			};
+			uint32 Bits;
+		} Flags;
 		int32 LastUsedEpoch;
 		int32 NumActiveConstraints;
 
+		// Indices into the arrays of collisions on the particles. This is a cookie for use by FParticleCollisions
+		int32 ParticleCollisionsIndex0;
+		int32 ParticleCollisionsIndex1;
+
 		// The particle transforms the last time the collisions were updated (used to determine whether we can restore contacts)
 		FMidPhaseRestoreThresholds RestoreThresholdZeroContacts;
 		FMidPhaseRestoreThresholds RestoreThreshold;
diff --git a/Engine/Source/Runtime/Experimental/ChaosCore/Public/Chaos/ObjectPool.h b/Engine/Source/Runtime/Experimental/ChaosCore/Public/Chaos/ObjectPool.h
--- a/Engine/Source/Runtime/Experimental/ChaosCore/Public/Chaos/ObjectPool.h
+++ b/Engine/Source/Runtime/Experimental/ChaosCore/Public/Chaos/ObjectPool.h
@@ -365,7 +365,7 @@ namespace Chaos
 					AsItem->NextFree = FreeList;
 				}
 
-				FreeList = IndexInBlock;
+				FreeList = (int32)IndexInBlock;
 				++NumFree;
 			}
 
