Chaos Heightfield: Set bOneSidedHeightfieldAlwaysSweep to false by default to workaround under the world issues.

Steven.Barnett

diff --git a/Engine/Source/Runtime/Experimental/Chaos/Private/Chaos/HeightField.cpp b/Engine/Source/Runtime/Experimental/Chaos/Private/Chaos/HeightField.cpp
--- a/Engine/Source/Runtime/Experimental/Chaos/Private/Chaos/HeightField.cpp
+++ b/Engine/Source/Runtime/Experimental/Chaos/Private/Chaos/HeightField.cpp
@@ -20,7 +20,7 @@ namespace Chaos
 	int32 bOneSidedHeightField = 1;
 	static FAutoConsoleVariableRef CVarOneSidedHeightField(TEXT("p.Chaos.OneSidedHeightField"), bOneSidedHeightField, TEXT("When enabled, extra steps will ensure that FHeightField::GJKContactPointImp never results in internal-facing contact data."));
 
-	int32 bOneSidedHeightfieldAlwaysSweep = 0;
+	int32 bOneSidedHeightfieldAlwaysSweep = 1;
 	static FAutoConsoleVariableRef CVarOneSidedHeightfieldAlwaysSweep(TEXT("p.Chaos.OneSidedHeightfieldAlwaysSweep"), bOneSidedHeightfieldAlwaysSweep, TEXT("When enabled, always use a sweep to ensure FHeightField::GJKContactPointImp never results \
 	in internal-facing contact data. Else, we only sweep if we detect an inward facing normal. Note that the sweep results can be inaccurate in some cases."));
 
