Chaos Physics Dynamic Trees: Maximum leaf capacity increased to more than one

diff --git a/Engine/Source/Runtime/Experimental/Chaos/Private/Chaos/AABBTree.cpp b/Engine/Source/Runtime/Experimental/Chaos/Private/Chaos/AABBTree.cpp
--- a/Engine/Source/Runtime/Experimental/Chaos/Private/Chaos/AABBTree.cpp
+++ b/Engine/Source/Runtime/Experimental/Chaos/Private/Chaos/AABBTree.cpp
@@ -31,3 +31,6 @@ FAutoConsoleVariableRef FAABBTreeCVars::CVarSplitOnVarianceAxis(TEXT("p.aabbtree
 
 float FAABBTreeCVars::DynamicTreeBoundingBoxPadding = 5.0f;
 FAutoConsoleVariableRef FAABBTreeCVars::CVarDynamicTreeBoundingBoxPadding(TEXT("p.aabbtree.DynamicTreeBoundingBoxPadding"), FAABBTreeCVars::DynamicTreeBoundingBoxPadding, TEXT("Additional padding added to bounding boxes for dynamic AABB trees to amortize update cost"));
+
+int32 FAABBTreeCVars::DynamicTreeLeafCapacity = 8;
+FAutoConsoleVariableRef FAABBTreeCVars::CVarDynamicTreeLeafCapacity(TEXT("p.aabbtree.DynamicTreeLeafCapacity"), FAABBTreeCVars::DynamicTreeLeafCapacity, TEXT("Dynamic Tree Leaf Capacity"));
diff --git a/Engine/Source/Runtime/Experimental/Chaos/Public/Chaos/AABBTree.h b/Engine/Source/Runtime/Experimental/Chaos/Public/Chaos/AABBTree.h
--- a/Engine/Source/Runtime/Experimental/Chaos/Public/Chaos/AABBTree.h
+++ b/Engine/Source/Runtime/Experimental/Chaos/Public/Chaos/AABBTree.h
@@ -31,6 +31,9 @@ struct CHAOS_API FAABBTreeCVars
 
 	static float DynamicTreeBoundingBoxPadding;
 	static FAutoConsoleVariableRef CVarDynamicTreeBoundingBoxPadding;
+
+	static int32 DynamicTreeLeafCapacity;
+	static FAutoConsoleVariableRef CVarDynamicTreeLeafCapacity;
 };
 
 struct CHAOS_API FAABBTreeDirtyGridCVars
@@ -1076,8 +1079,13 @@ public:
 		return(WhichChildAmI(NodeIdx) ^ 1);
 	}
 
-	int32 FindBestSibling(const TAABB<T, 3>& InNewBounds)
+	int32 FindBestSibling(const TAABB<T, 3>& InNewBounds, bool& bOutAddToLeaf)
 	{
+		bOutAddToLeaf = false;
+		if (RootNode == INDEX_NONE)
+		{
+			return INDEX_NONE;
+		}
 		
 		TAABB<T, 3> NewBounds = InNewBounds;
 		NewBounds.Thicken(FAABBTreeCVars::DynamicTreeBoundingBoxPadding);
@@ -1119,22 +1127,35 @@ public:
 			//SumDeltaCostQ.SetNumUninitialized(SumDeltaCostQ.Num() - 1);
 
 			// TestSibling bounds union with new bounds
+			bool bAddToLeaf = false;
 			WorkingAABB = Nodes[TestSibling].ChildrenBounds[0];
 			if (!Nodes[TestSibling].bLeaf)
 			{
 				WorkingAABB.GrowToInclude(Nodes[TestSibling].ChildrenBounds[1]);
 			}
+			else
+			{
+				int32 LeafIdx = Nodes[TestSibling].ChildrenNodes[0];
+				bAddToLeaf = Leaves[LeafIdx].GetElementCount() < FAABBTreeCVars::DynamicTreeLeafCapacity;
+			}
 			FReal TestSiblingArea = WorkingAABB.GetArea();
 			WorkingAABB.GrowToInclude(NewBounds);
 
 			FReal NewPotentialNodeArea = WorkingAABB.GetArea();
 			FReal CostForChoosingNode = NewPotentialNodeArea + SumDeltaCost;
+			if (bAddToLeaf)
+			{
+				// No new node is added (we can experiment with this cost function
+				// It is faster overall if we don't subtract here
+				//CostForChoosingNode -= TestSiblingArea;
+			}
 			FReal NewDeltaCost = NewPotentialNodeArea - TestSiblingArea;
 			// Did we get a better cost?
 			if (CostForChoosingNode < BestCost)
 			{
 				BestCost = CostForChoosingNode;
 				BestSiblingIdx = TestSibling;
+				bOutAddToLeaf = bAddToLeaf;
 			}
 
 			// Lower bound of Children costs
@@ -1231,6 +1252,24 @@ public:
 		//	DynamicTreeDebugStats();
 		//}
 
+		// Find the best sibling
+		bool bAddToLeaf;
+		int32 BestSibling = FindBestSibling(NewBounds, bAddToLeaf);
+
+		if (bAddToLeaf)
+		{
+			int32 LeafIdx = Nodes[BestSibling].ChildrenNodes[0];
+			Leaves[LeafIdx].AddElement(TPayloadBoundsElement<TPayloadType, T>{Payload, NewBounds});
+			Leaves[LeafIdx].RecomputeBounds();
+			TAABB<T, 3> ExpandedBounds = Leaves[LeafIdx].GetBounds();
+			ExpandedBounds.Thicken(FAABBTreeCVars::DynamicTreeBoundingBoxPadding);
+			Nodes[BestSibling].ChildrenBounds[0] = ExpandedBounds;
+			UpdateAncestorBounds(BestSibling, true);
+			FAABBTreePayloadInfo* PayloadInfo = PayloadToInfo.Find(Payload);
+			*PayloadInfo = FAABBTreePayloadInfo{ INDEX_NONE, INDEX_NONE, LeafIdx, INDEX_NONE, BestSibling };
+			return;
+		}
+
 		int32 NewLeafNode = AllocateLeafNodeAndLeaf(Payload, NewBounds);
 
 		// New tree?
@@ -1240,9 +1279,6 @@ public:
 			return;
 		}
 
-		// Find the best sibling
-		int32 BestSibling = FindBestSibling(NewBounds);
-
 		// New internal parent node
 		int32 oldParent = Nodes[BestSibling].ParentNode;
 		int32 newParent = AllocateInternalNode();
@@ -1304,8 +1340,21 @@ public:
 		}
 	}
 
-	void RemoveLeafNode(int32 LeafNodeIdx)
+	void RemoveLeafNode(int32 LeafNodeIdx, const TPayloadType& Payload)
 	{
+		int32 LeafIdx = Nodes[LeafNodeIdx].ChildrenNodes[0];
+
+		if (Leaves[LeafIdx].GetElementCount() > 1)
+		{
+			Leaves[LeafIdx].RemoveElement(Payload);
+			Leaves[LeafIdx].RecomputeBounds();
+			TAABB<T, 3> ExpandedBounds = Leaves[LeafIdx].GetBounds();
+			ExpandedBounds.Thicken(FAABBTreeCVars::DynamicTreeBoundingBoxPadding);
+			Nodes[LeafNodeIdx].ChildrenBounds[0] = ExpandedBounds;
+			UpdateAncestorBounds(LeafNodeIdx);
+			return;
+		}
+
 		int32 ParentNodeIdx = Nodes[LeafNodeIdx].ParentNode;
 
 		if (ParentNodeIdx != INDEX_NONE)
@@ -1372,7 +1421,7 @@ public:
 				{
 					if (bDynamicTree)
 					{
-						RemoveLeafNode(PayloadInfo->NodeIdx);
+						RemoveLeafNode(PayloadInfo->NodeIdx, Payload);
 					}
 					else
 					{
@@ -1441,7 +1490,7 @@ public:
 
 					if (bDynamicTree)
 					{
-						RemoveLeafNode(PayloadInfo->NodeIdx);
+						RemoveLeafNode(PayloadInfo->NodeIdx, Payload);
 						PayloadInfo->LeafIdx = INDEX_NONE;
 						PayloadInfo->NodeIdx = INDEX_NONE;
 					}
diff --git a/Engine/Source/Runtime/Experimental/Chaos/Public/Chaos/BoundingVolume.h b/Engine/Source/Runtime/Experimental/Chaos/Public/Chaos/BoundingVolume.h
--- a/Engine/Source/Runtime/Experimental/Chaos/Public/Chaos/BoundingVolume.h
+++ b/Engine/Source/Runtime/Experimental/Chaos/Public/Chaos/BoundingVolume.h
@@ -294,6 +294,13 @@ public:
 
 	}
 
+	inline int32 GetElementCount()
+	{
+		// Not implemented 
+		check(false);
+		return 0;
+	}
+
 	
 	// Begin ISpatialAcceleration interface
 	virtual TArray<TPayloadType> FindAllIntersections(const FAABB3& Box) const override { return FindAllIntersectionsImp(Box); }
