Adding GRHISupportsMapWriteNoOverwrite to check at runtime if RLM_WriteOnly_NoOverwrite is supported.

Motivated by changes from PR #7423

diff --git a/Engine/Source/Runtime/D3D12RHI/Private/D3D12RHI.cpp b/Engine/Source/Runtime/D3D12RHI/Private/D3D12RHI.cpp
--- a/Engine/Source/Runtime/D3D12RHI/Private/D3D12RHI.cpp
+++ b/Engine/Source/Runtime/D3D12RHI/Private/D3D12RHI.cpp
@@ -230,6 +230,8 @@ FD3D12DynamicRHI::FD3D12DynamicRHI(const TArray<TSharedPtr<FD3D12Adapter>>& Chos
 	GRHISupportsRayTracingAsyncBuildAccelerationStructure = true;
 
 	GRHISupportsPipelineFileCache = PLATFORM_WINDOWS;
+
+	GRHISupportsMapWriteNoOverwrite = true;
 }
 
 FD3D12DynamicRHI::~FD3D12DynamicRHI()
diff --git a/Engine/Source/Runtime/Experimental/GeometryCollectionEngine/Private/GeometryCollection/GeometryCollectionSceneProxy.cpp b/Engine/Source/Runtime/Experimental/GeometryCollectionEngine/Private/GeometryCollection/GeometryCollectionSceneProxy.cpp
--- a/Engine/Source/Runtime/Experimental/GeometryCollectionEngine/Private/GeometryCollection/GeometryCollectionSceneProxy.cpp
+++ b/Engine/Source/Runtime/Experimental/GeometryCollectionEngine/Private/GeometryCollection/GeometryCollectionSceneProxy.cpp
@@ -108,7 +108,7 @@ FGeometryCollectionSceneProxy::FGeometryCollectionSceneProxy(UGeometryCollection
 	, NumIndices(0)
 	, VertexFactory(GetScene().GetFeatureLevel())
 	, bSupportsManualVertexFetch(VertexFactory.SupportsManualVertexFetch(GetScene().GetFeatureLevel()))
-	, bSupportsTripleBufferVertexUpload(!IsMetalPlatform(GetScene().GetShaderPlatform()) && !IsVulkanPlatform(GetScene().GetShaderPlatform()))
+	, bSupportsTripleBufferVertexUpload(GRHISupportsMapWriteNoOverwrite)
 #if GEOMETRYCOLLECTION_EDITOR_SELECTION
 	, SubSections()
 	, SubSectionHitProxies()
diff --git a/Engine/Source/Runtime/RHI/Private/RHI.cpp b/Engine/Source/Runtime/RHI/Private/RHI.cpp
--- a/Engine/Source/Runtime/RHI/Private/RHI.cpp
+++ b/Engine/Source/Runtime/RHI/Private/RHI.cpp
@@ -1464,6 +1464,8 @@ bool GRHISupportsShaderTimestamp = false;
 
 bool GRHISupportsEfficientUploadOnResourceCreation = false;
 
+bool GRHISupportsMapWriteNoOverwrite = false;
+
 FVertexElementTypeSupportInfo GVertexElementTypeSupport;
 
 RHI_API int32 volatile GCurrentTextureMemorySize = 0;
diff --git a/Engine/Source/Runtime/RHI/Public/RHI.h b/Engine/Source/Runtime/RHI/Public/RHI.h
--- a/Engine/Source/Runtime/RHI/Public/RHI.h
+++ b/Engine/Source/Runtime/RHI/Public/RHI.h
@@ -753,6 +753,9 @@ extern RHI_API bool GRHISupportsShaderTimestamp;
 
 extern RHI_API bool GRHISupportsEfficientUploadOnResourceCreation;
 
+/** true if the RHI supports RLM_WriteOnly_NoOverwrite */
+extern RHI_API bool GRHISupportsMapWriteNoOverwrite;
+
 /** Tables of all MSAA sample offset for all MSAA supported. Use GetMSAASampleOffsets() to read it. */
 extern RHI_API FVector2f GRHIDefaultMSAASampleOffsets[1 + 2 + 4 + 8 + 16];
 
diff --git a/Engine/Source/Runtime/RHI/Public/RHIValidation.h b/Engine/Source/Runtime/RHI/Public/RHIValidation.h
--- a/Engine/Source/Runtime/RHI/Public/RHIValidation.h
+++ b/Engine/Source/Runtime/RHI/Public/RHIValidation.h
@@ -1427,6 +1427,7 @@ public:
 
 	virtual void* LockBuffer_BottomOfPipe(class FRHICommandListImmediate& RHICmdList, FRHIBuffer* Buffer, uint32 Offset, uint32 SizeRHI, EResourceLockMode LockMode) override final
 	{
+		RHI_VALIDATION_CHECK(LockMode != RLM_WriteOnly_NoOverwrite || GRHISupportsMapWriteNoOverwrite, TEXT("Using RLM_WriteOnly_NoOverwrite when the RHI doesn't support it."));
 		return RHI->LockBuffer_BottomOfPipe(RHICmdList, Buffer, Offset, SizeRHI, LockMode);
 	}
 	
diff --git a/Engine/Source/Runtime/Windows/D3D11RHI/Private/D3D11Buffer.cpp b/Engine/Source/Runtime/Windows/D3D11RHI/Private/D3D11Buffer.cpp
--- a/Engine/Source/Runtime/Windows/D3D11RHI/Private/D3D11Buffer.cpp
+++ b/Engine/Source/Runtime/Windows/D3D11RHI/Private/D3D11Buffer.cpp
@@ -146,7 +146,7 @@ void* FD3D11DynamicRHI::LockBuffer_BottomOfPipe(FRHICommandListImmediate& RHICmd
 		// If the buffer is dynamic, map its memory for writing.
 		D3D11_MAPPED_SUBRESOURCE MappedSubresource;
 
-		D3D11_MAP MapType = (LockMode == RLM_WriteOnly)? D3D11_MAP_WRITE_DISCARD : D3D11_MAP_WRITE_NO_OVERWRITE;
+		D3D11_MAP MapType = (LockMode == RLM_WriteOnly || !GRHISupportsMapWriteNoOverwrite) ? D3D11_MAP_WRITE_DISCARD : D3D11_MAP_WRITE_NO_OVERWRITE;
 		VERIFYD3D11RESULT_EX(Direct3DDeviceIMContext->Map(Buffer->Resource, 0, MapType, 0, &MappedSubresource), Direct3DDevice);
 
 		LockedData.SetData(MappedSubresource.pData);
diff --git a/Engine/Source/Runtime/Windows/D3D11RHI/Private/D3D11Device.cpp b/Engine/Source/Runtime/Windows/D3D11RHI/Private/D3D11Device.cpp
--- a/Engine/Source/Runtime/Windows/D3D11RHI/Private/D3D11Device.cpp
+++ b/Engine/Source/Runtime/Windows/D3D11RHI/Private/D3D11Device.cpp
@@ -501,6 +501,17 @@ void FD3D11DynamicRHI::SetupAfterDeviceCreation()
 		UE_LOG(LogD3D11RHI, Log, TEXT("Async texture creation disabled: %s"),
 			D3D11RHI_ShouldAllowAsyncResourceCreation() ? TEXT("no driver support") : TEXT("disabled by user"));
 	}
+	{
+		D3D11_FEATURE_DATA_D3D11_OPTIONS Data;
+		if (const HRESULT Result = Direct3DDevice->CheckFeatureSupport(D3D11_FEATURE_D3D11_OPTIONS, &Data, sizeof(Data)); SUCCEEDED(Result))
+		{
+			GRHISupportsMapWriteNoOverwrite = Data.MapNoOverwriteOnDynamicBufferSRV;
+			if (GRHISupportsMapWriteNoOverwrite)
+			{
+				UE_LOG(LogD3D11RHI, Log, TEXT("D3D11_MAP_WRITE_NO_OVERWRITE for dynamic buffer SRVs is supported"));
+			}
+		}
+	}
 
 	{
 #if 0
