Chaos - fix issue with scaled trimesh where off center query geometry may not find the right triangles ( contacts, overlap )

diff --git a/Engine/Source/Runtime/Experimental/Chaos/Private/Chaos/TriangleMeshImplicitObject.cpp b/Engine/Source/Runtime/Experimental/Chaos/Private/Chaos/TriangleMeshImplicitObject.cpp
--- a/Engine/Source/Runtime/Experimental/Chaos/Private/Chaos/TriangleMeshImplicitObject.cpp
+++ b/Engine/Source/Runtime/Experimental/Chaos/Private/Chaos/TriangleMeshImplicitObject.cpp
@@ -358,7 +358,16 @@ bool FTriangleMeshImplicitObject::ContactManifoldImp(const GeomType& QueryGeom,
 	bool bResult = false;
 
 	const auto& WorldScaleGeom = ScaleGeomIntoWorldHelper(QueryGeom, TriMeshScale);
-	const FVec3 InvTriMeshScale = FVec3(FReal(1) / TriMeshScale.X, FReal(1) / TriMeshScale.Y, FReal(1) / TriMeshScale.Z);
+	const FVec3 InvTriMeshScale = FReal(1) / TriMeshScale;
+
+	// IMPORTANT QueryTM comes with a invscaled translation so we need a version of the TM with world space translation to properly compute the bounds
+	FRigidTransform3 TriMeshToGeomNoScale{ QueryTM };
+	TriMeshToGeomNoScale.SetTranslation(TriMeshToGeomNoScale.GetTranslation() * TriMeshScale);
+	// NOTE: BVH test is done in tri-mesh local space (whereas collision detection is done in world space becaused you can't non-uniformly scale all shapes)
+	FAABB3 QueryBounds = WorldScaleGeom.BoundingBox();
+	QueryBounds = QueryBounds.TransformedAABB(TriMeshToGeomNoScale);
+	QueryBounds.ThickenSymmetrically(FVec3(WorldThickness));
+	QueryBounds.Scale(InvTriMeshScale);
 
 	TRigidTransform<FReal, 3> WorldScaleQueryTM;
 	ScaleTransformHelper(TriMeshScale, QueryTM, WorldScaleQueryTM);
@@ -410,12 +419,6 @@ bool FTriangleMeshImplicitObject::ContactManifoldImp(const GeomType& QueryGeom,
 		FReal LocalContactPhi = FLT_MAX;
 		FVec3 LocalContactLocation, LocalContactNormal;
 
-		// NOTE: BVH test is done in tri-mesh local space (whereas collision detection is done in world space becaused you can't non-uniformly scale all shapes)
-		FAABB3 QueryBounds = WorldScaleGeom.BoundingBox();
-		QueryBounds = QueryBounds.TransformedAABB(QueryTM);
-		QueryBounds.ThickenSymmetrically(FVec3(WorldThickness));
-		QueryBounds.LocalScale(InvTriMeshScale);
-
 		const TArray<int32> PotentialIntersections = BVH.FindAllIntersections(QueryBounds);
 
 		// MakeTriangleHelper get rid of the scale wrapper if necessary as MakeTriangle will try to infer properties from it 
@@ -472,6 +475,15 @@ bool FTriangleMeshImplicitObject::GJKContactPointImp(const QueryGeomType& QueryG
 	const auto& WorldScaleGeom = ScaleGeomIntoWorldHelper(QueryGeom, TriMeshScale);
 	const FVec3 InvTriMeshScale = FVec3(FReal(1) / TriMeshScale.X, FReal(1) / TriMeshScale.Y, FReal(1) / TriMeshScale.Z);
 
+	// IMPORTANT QueryTM comes with a invscaled translation so we need a version of the TM with world space translation to properly compute the bounds
+	FRigidTransform3 TriMeshToGeomNoScale{ QueryTM };
+	TriMeshToGeomNoScale.SetTranslation(TriMeshToGeomNoScale.GetTranslation() * TriMeshScale);
+	// NOTE: BVH test is done in tri-mesh local space (whereas collision detection is done in world space becaused you can't non-uniformly scale all shapes)
+	FAABB3 QueryBounds = WorldScaleGeom.BoundingBox();
+	QueryBounds = QueryBounds.TransformedAABB(TriMeshToGeomNoScale);
+	QueryBounds.ThickenSymmetrically(FVec3(WorldThickness));
+	QueryBounds.Scale(InvTriMeshScale);
+
 	TRigidTransform<FReal, 3> WorldScaleQueryTM;
 	ScaleTransformHelper(TriMeshScale, QueryTM, WorldScaleQueryTM);
 
@@ -501,12 +513,6 @@ bool FTriangleMeshImplicitObject::GJKContactPointImp(const QueryGeomType& QueryG
 		FReal LocalContactPhi = FLT_MAX;
 		FVec3 LocalContactLocation, LocalContactNormal;
 
-		// NOTE: BVH test is done in tri-mesh local space (whereas collision detection is done in world space becaused you can't non-uniformly scale all shapes)
-		FAABB3 QueryBounds = WorldScaleGeom.BoundingBox();
-		QueryBounds = QueryBounds.TransformedAABB(QueryTM);
-		QueryBounds.ThickenSymmetrically(FVec3(WorldThickness));
-		QueryBounds.LocalScale(InvTriMeshScale);
-
 		const TArray<int32> PotentialIntersections = BVH.FindAllIntersections(QueryBounds);
 
 		for (int32 TriIdx : PotentialIntersections)
@@ -708,6 +714,25 @@ void FTriangleMeshImplicitObject::VisitTriangles(const FAABB3& QueryBounds, cons
 	}
 }
 
+void FTriangleMeshImplicitObject::VisitTriangle(int32 TriangleIndex, const TFunction<void(const FTriangle& Triangle)>& Visitor) const
+{
+	const auto TriangleProducer = [&](int32 TriIdx, const auto& Elements)
+	{
+		TVec3<FReal> A, B, C;
+		TransformVertsHelper(FVec3(1), TriIdx, MParticles, Elements, A, B, C);
+
+		Visitor(FTriangle(A, B, C));
+	};
+
+	if (MElements.RequiresLargeIndices())
+	{
+		return TriangleProducer(TriangleIndex, MElements.GetLargeIndexBuffer());
+	}
+	else
+	{
+		return TriangleProducer(TriangleIndex, MElements.GetSmallIndexBuffer());
+	}
+}
 
 template <typename QueryGeomType>
 bool FTriangleMeshImplicitObject::OverlapGeomImp(const QueryGeomType& QueryGeom, const FRigidTransform3& QueryTM, const FReal Thickness, FMTDInfo* OutMTD, FVec3 TriMeshScale) const
@@ -718,11 +743,14 @@ bool FTriangleMeshImplicitObject::OverlapGeomImp(const QueryGeomType& QueryGeom,
 
 	const FVec3 InvTriMeshScale = FVec3(FReal(1) / TriMeshScale.X, FReal(1) / TriMeshScale.Y, FReal(1) / TriMeshScale.Z);
 
+	// IMPORTANT QueryTM comes with a invscaled translation so we need a version of the TM with world space translation to properly compute the bounds
+	FRigidTransform3 TriMeshToGeomNoScale{ QueryTM };
+	TriMeshToGeomNoScale.SetTranslation(TriMeshToGeomNoScale.GetTranslation() * TriMeshScale);
 	// NOTE: BVH test is done in tri-mesh local space (whereas collision detection is done in world space becaused you can't non-uniformly scale all shapes)
 	FAABB3 QueryBounds = WorldScaleQueryGeom.BoundingBox();
-	QueryBounds = QueryBounds.TransformedAABB(QueryTM);
-	QueryBounds.Thicken(Thickness);
-	QueryBounds.LocalScale(InvTriMeshScale);
+	QueryBounds = QueryBounds.TransformedAABB(TriMeshToGeomNoScale);
+	QueryBounds.ThickenSymmetrically(FVec3(Thickness));
+	QueryBounds.Scale(InvTriMeshScale);
 
 	const TArray<int32> PotentialIntersections = BVH.FindAllIntersections(QueryBounds);
 
diff --git a/Engine/Source/Runtime/Experimental/Chaos/Public/Chaos/TriangleMeshImplicitObject.h b/Engine/Source/Runtime/Experimental/Chaos/Public/Chaos/TriangleMeshImplicitObject.h
--- a/Engine/Source/Runtime/Experimental/Chaos/Public/Chaos/TriangleMeshImplicitObject.h
+++ b/Engine/Source/Runtime/Experimental/Chaos/Public/Chaos/TriangleMeshImplicitObject.h
@@ -305,6 +305,8 @@ namespace Chaos
 
 		void VisitTriangles(const FAABB3& InQueryBounds, const TFunction<void(const FTriangle& Triangle)>& Visitor) const;
 
+		void VisitTriangle(int32 TriangleIndex, const TFunction<void(const FTriangle& Triangle)>& Visitor) const;
+
 	private:
 		void RebuildBV();
 
