Chaos - solve for restitution at average contact poin
- fixes box bouncing so that they don't start spinning on bounce

diff --git a/Engine/Source/Runtime/Experimental/Chaos/Private/Chaos/Collision/PBDCollisionSolver.cpp b/Engine/Source/Runtime/Experimental/Chaos/Private/Chaos/Collision/PBDCollisionSolver.cpp
--- a/Engine/Source/Runtime/Experimental/Chaos/Private/Chaos/Collision/PBDCollisionSolver.cpp
+++ b/Engine/Source/Runtime/Experimental/Chaos/Private/Chaos/Collision/PBDCollisionSolver.cpp
@@ -363,13 +363,27 @@ namespace Chaos
 	}
 
 	void FPBDCollisionSolverManifoldPoint::InitMaterial(
-		const FReal InWorldContactVelocityTargetNormal,
+		const FConstraintSolverBody& Body0,
+		const FConstraintSolverBody& Body1,
+		const FReal InRestitution,
+		const FReal InRestitutionVelocityThreshold,
 		const bool bInEnableStaticFriction,
 		const FReal InStaticFrictionMax)
 	{
 		StaticFrictionMax = InStaticFrictionMax;
 		bInsideStaticFrictionCone = bInEnableStaticFriction;
-		WorldContactVelocityTargetNormal = InWorldContactVelocityTargetNormal;
+		WorldContactVelocityTargetNormal = FReal(0);
+
+		if (InRestitution > FReal(0))
+		{
+			const FVec3 ContactVelocity = CalculateContactVelocity(Body0, Body1);
+			const FReal ContactVelocityNormal = FVec3::DotProduct(ContactVelocity, WorldContactNormal);
+			if (ContactVelocityNormal < -InRestitutionVelocityThreshold)
+			{
+				WorldContactVelocityTargetNormal = -InRestitution * ContactVelocityNormal;
+			}
+		}
+
 	}
 
 	inline void FPBDCollisionSolverManifoldPoint::UpdateContact(
@@ -550,11 +564,15 @@ namespace Chaos
 
 	void FPBDCollisionSolver::InitMaterial(
 		const int32 ManifoldPoiontIndex,
-		const FReal InWorldContactVelocityTargetNormal,
+		const FReal InRestitution,
+		const FReal InRestitutionVelocityThreshold,
 		const bool bInEnableStaticFriction,
 		const FReal InStaticFrictionMax)
 	{
-		State.ManifoldPoints[ManifoldPoiontIndex].InitMaterial(InWorldContactVelocityTargetNormal, bInEnableStaticFriction, InStaticFrictionMax);
+		State.ManifoldPoints[ManifoldPoiontIndex].InitMaterial(State.SolverBodies[0], State.SolverBodies[1], InRestitution, InRestitutionVelocityThreshold, bInEnableStaticFriction, InStaticFrictionMax);
+
+		// Track if any points have restitution enabled. See SolveVelocity
+		State.bHaveRestitution = State.bHaveRestitution || (State.ManifoldPoints[ManifoldPoiontIndex].WorldContactVelocityTargetNormal > 0);
 	}
 
 	void FPBDCollisionSolver::UpdateContact(
@@ -629,6 +647,71 @@ namespace Chaos
 		return !State.bIsSolved;
 	}
 
+	void FPBDCollisionSolver::SolveVelocityAverage(const FReal Dt)
+	{
+		FConstraintSolverBody& Body0 = SolverBody0();
+		FConstraintSolverBody& Body1 = SolverBody1();
+
+		// Generate a new contact point at the average of all the active contacts
+		int32 NumActiveManifoldPoints = 0;
+		FVec3 WorldContactPosition = FVec3(0);
+		FVec3 WorldContactNormal = FVec3(0);
+		FVec3 NetPushOut = FVec3(0);
+		FReal WorldContactVelocityTargetNormal = FReal(0);
+		for (int32 PointIndex = 0; PointIndex < State.NumManifoldPoints; ++PointIndex)
+		{
+			FPBDCollisionSolverManifoldPoint& SolverManifoldPoint = State.ManifoldPoints[PointIndex];
+			if (!SolverManifoldPoint.NetPushOut.IsNearlyZero())
+			{
+				WorldContactPosition += SolverManifoldPoint.WorldContactPosition;
+				WorldContactVelocityTargetNormal += SolverManifoldPoint.WorldContactVelocityTargetNormal;
+				WorldContactNormal = SolverManifoldPoint.WorldContactNormal;	// Take last value - should all be similar
+				++NumActiveManifoldPoints;
+			}
+		}
+
+		// Solve for velocity at the new conatct point
+		// We only do this if we have multiple active contacts
+		if (NumActiveManifoldPoints > 1)
+		{
+			const FReal DynamicFriction = FReal(0);
+			const FReal Scale = FReal(1) / FReal(NumActiveManifoldPoints);
+
+			FPBDCollisionSolverManifoldPoint AverageManifoldPoint;
+			AverageManifoldPoint.WorldContactPosition = WorldContactPosition * Scale;
+			AverageManifoldPoint.WorldContactNormal = WorldContactNormal;
+
+			const FMatrix33 ContactMassInv =
+				(Body0.IsDynamic() ? Collisions::ComputeFactorMatrix3(AverageManifoldPoint.WorldContactPosition - Body0.P(), Body0.InvI(), Body0.InvM()) : FMatrix33(0)) +
+				(Body1.IsDynamic() ? Collisions::ComputeFactorMatrix3(AverageManifoldPoint.WorldContactPosition - Body1.P(), Body1.InvI(), Body1.InvM()) : FMatrix33(0));
+			const FReal ContactMassInvNormal = FVec3::DotProduct(AverageManifoldPoint.WorldContactNormal, Utilities::Multiply(ContactMassInv, AverageManifoldPoint.WorldContactNormal));
+			const FReal ContactMassNormal = (ContactMassInvNormal > FReal(SMALL_NUMBER)) ? FReal(1) / ContactMassInvNormal : FReal(0);
+			AverageManifoldPoint.WorldContactMass = FMatrix33(0);
+			AverageManifoldPoint.WorldContactMassNormal = ContactMassNormal;
+
+			AverageManifoldPoint.WorldContactDelta = FVec3(0);
+			AverageManifoldPoint.WorldContactVelocityTargetNormal = WorldContactVelocityTargetNormal * Scale;
+			AverageManifoldPoint.NetPushOut = FVec3(0);
+			AverageManifoldPoint.NetImpulse = FVec3(0);
+			AverageManifoldPoint.StaticFrictionMax = FReal(0);
+			AverageManifoldPoint.bInsideStaticFrictionCone = false;
+
+			FVec3 ContactVelocityDelta;
+			FReal ContactVelocityDeltaNormal;
+			AverageManifoldPoint.CalculateContactVelocityError(Body0, Body1, DynamicFriction, Dt, ContactVelocityDelta, ContactVelocityDeltaNormal);
+
+			ApplyVelocityCorrection(
+				State.Stiffness,
+				Dt,
+				DynamicFriction,
+				ContactVelocityDelta,
+				ContactVelocityDeltaNormal,
+				AverageManifoldPoint,
+				Body0,
+				Body1);
+		}
+	}
+
 	bool FPBDCollisionSolver::SolveVelocity(const FReal Dt, const bool bApplyDynamicFriction)
 	{
 		if (!bChaos_PBDCollisionSolver_Velocity_SolveEnabled)
@@ -636,28 +719,30 @@ namespace Chaos
 			return false;
 		}
 
+		// Apply restitution at the average contact point
+		// This means we don't need to run as many iterations to get stable bouncing
+		if (State.bHaveRestitution && (State.NumManifoldPoints > 1))
+		{
+			SolveVelocityAverage(Dt);
+		}
+
 		FConstraintSolverBody& Body0 = SolverBody0();
 		FConstraintSolverBody& Body1 = SolverBody1();
 
-		// Only apply friction for 1 iteration (the last one) because the solver is iteration-count sensitive
-		// @todo(chaos): use xpbd dynamic friction in QuasiPBD Solver
+		// NOTE: this dynamic friction implementation is iteration-count sensitive
+		// @todo(chaos): fix iteration count dependence of dynamic friction
 		const FReal DynamicFriction = (bApplyDynamicFriction && bChaos_PBDCollisionSolver_Velocity_DynamicFrictionEnabled) ? State.DynamicFriction : (FReal)0.0f;
-		const FReal Stiffness = State.Stiffness;
 
 		for (int32 PointIndex = 0; PointIndex < State.NumManifoldPoints; ++PointIndex)
 		{
 			FPBDCollisionSolverManifoldPoint& SolverManifoldPoint = State.ManifoldPoints[PointIndex];
-
-			//SolverManifoldPoint.UpdateContactPoint(Body0, Body1, State.Constraint->GetManifoldPoints()[PointIndex]);
-			//SolverManifoldPoint.UpdateContactMass(Body0, Body1);
-
 			if (!SolverManifoldPoint.NetPushOut.IsNearlyZero())
 			{
 				FVec3 ContactVelocityDelta;
 				FReal ContactVelocityDeltaNormal;
 				SolverManifoldPoint.CalculateContactVelocityError(Body0, Body1, DynamicFriction, Dt, ContactVelocityDelta, ContactVelocityDeltaNormal);
 				ApplyVelocityCorrection(
-					Stiffness,
+					State.Stiffness,
 					Dt,
 					DynamicFriction,
 					ContactVelocityDelta,
diff --git a/Engine/Source/Runtime/Experimental/Chaos/Private/Chaos/Collision/PBDCollisionSolver.h b/Engine/Source/Runtime/Experimental/Chaos/Private/Chaos/Collision/PBDCollisionSolver.h
--- a/Engine/Source/Runtime/Experimental/Chaos/Private/Chaos/Collision/PBDCollisionSolver.h
+++ b/Engine/Source/Runtime/Experimental/Chaos/Private/Chaos/Collision/PBDCollisionSolver.h
@@ -39,7 +39,10 @@ namespace Chaos
 		 * @brief Initialize the material related properties of the contact
 		*/
 		void InitMaterial(
-			const FReal InWorldContactVelocityTargetNormal,
+			const FConstraintSolverBody& Body0,
+			const FConstraintSolverBody& Body1,
+			const FReal InRestitution,
+			const FReal InRestitutionVelocityThreshold,
 			const bool bInEnableStaticFriction,
 			const FReal InStaticFrictionMax);
 
@@ -181,7 +184,8 @@ namespace Chaos
 
 		void InitMaterial(
 			const int32 ManifoldPoiontIndex,
-			const FReal InWorldContactVelocityTargetNormal,
+			const FReal InRestitution,
+			const FReal InRestitutionVelocityThreshold,
 			const bool bInEnableStaticFriction,
 			const FReal InStaticFrictionMax);
 
@@ -232,6 +236,14 @@ namespace Chaos
 		*/
 		bool SolveVelocity(const FReal Dt, const bool bApplyDynamicFriction);
 
+		/**
+		 * @brief Run a velocity solve on the average point from all the points that received a position impulse
+		 * This is used to enforce Restitution constraints without introducing rotation artefacts without
+		 * adding more velocity iterations.
+		 * This will only perform work if there is more than one active contact.
+		*/
+		void SolveVelocityAverage(const FReal Dt);
+
 	private:
 		/**
 		 * @brief Apply the inverse mass scale the body with the lower level
@@ -251,6 +263,7 @@ namespace Chaos
 				, BodyEpochs{ INDEX_NONE, INDEX_NONE }
 				, NumPositionSolves(0)
 				, NumVelocitySolves(0)
+				, bHaveRestitution(false)
 				, bIsSolved(false)
 			{
 			}
@@ -264,6 +277,7 @@ namespace Chaos
 			int32 BodyEpochs[MaxConstrainedBodies];
 			int32 NumPositionSolves;
 			int32 NumVelocitySolves;
+			bool bHaveRestitution;
 			bool bIsSolved;
 		};
 
diff --git a/Engine/Source/Runtime/Experimental/Chaos/Private/Chaos/Collision/SolverCollisionContainer.cpp b/Engine/Source/Runtime/Experimental/Chaos/Private/Chaos/Collision/SolverCollisionContainer.cpp
--- a/Engine/Source/Runtime/Experimental/Chaos/Private/Chaos/Collision/SolverCollisionContainer.cpp
+++ b/Engine/Source/Runtime/Experimental/Chaos/Private/Chaos/Collision/SolverCollisionContainer.cpp
@@ -106,22 +106,11 @@ namespace Chaos
 					ManifoldPoint.WorldContactPoints[1],
 					ManifoldPoint.ContactPoint.Normal);
 
-				// Calculate the target normal velocity based on restitution
-				FReal WorldContactVelocityTargetNormal = FReal(0);
-				if (Restitution > FReal(0))
-				{
-					const FVec3 ContactVelocity = Solver.GetManifoldPoint(ManifoldPointIndex).CalculateContactVelocity(Solver.SolverBody0(), Solver.SolverBody1());
-					const FReal ContactVelocityNormal = FVec3::DotProduct(ContactVelocity, ManifoldPoint.ContactPoint.Normal);
-					if (ContactVelocityNormal < -RestitutionVelocityThreshold)
-					{
-						WorldContactVelocityTargetNormal = -Restitution * ContactVelocityNormal;
-					}
-				}
-
 				// Initialize the material properties (restitution and friction related)
 				Solver.InitMaterial(
 					ManifoldPointIndex,
-					WorldContactVelocityTargetNormal,
+					Restitution,
+					RestitutionVelocityThreshold,
 					(Solver.StaticFriction() > 0),
 					ManifoldPoint.StaticFrictionMax);
 			}
