Chaos - eliminating cache miss in Collision Apply loop (Incremental check)
- ManifoldPoint arrays to TCArray to avoid non-lineable function in TArray with TFixedAllocator

diff --git a/Engine/Source/Programs/HeadlessChaos/Private/HeadlessChaosTestCollisions.cpp b/Engine/Source/Programs/HeadlessChaos/Private/HeadlessChaosTestCollisions.cpp
--- a/Engine/Source/Programs/HeadlessChaos/Private/HeadlessChaosTestCollisions.cpp
+++ b/Engine/Source/Programs/HeadlessChaos/Private/HeadlessChaosTestCollisions.cpp
@@ -504,7 +504,6 @@ namespace ChaosTest {
 
 		Collisions.GatherInput(Dt);
 
-		Collisions.UpdateLevelsetConstraint(Constraint);
 		EXPECT_EQ(Constraint.Particle[0], DynamicCube);
 		EXPECT_EQ(Constraint.Particle[1], StaticCube);
 		EXPECT_NEAR(Constraint.GetNormal().Z, FReal(1), KINDA_SMALL_NUMBER);
diff --git a/Engine/Source/Runtime/Experimental/Chaos/Private/Chaos/Collision/PBDCollisionConstraint.cpp b/Engine/Source/Runtime/Experimental/Chaos/Private/Chaos/Collision/PBDCollisionConstraint.cpp
--- a/Engine/Source/Runtime/Experimental/Chaos/Private/Chaos/Collision/PBDCollisionConstraint.cpp
+++ b/Engine/Source/Runtime/Experimental/Chaos/Private/Chaos/Collision/PBDCollisionConstraint.cpp
@@ -142,15 +142,15 @@ namespace Chaos
 		, AccumulatedImpulse(0)
 		, Manifold()
 		, TimeOfImpact(0)
+		, ContainerCookie()
 		, CCDType(ECollisionCCDType::Disabled)
 		, Stiffness(FReal(1))
 		, ManifoldPoints()
-		, ManifoldPointsSavedData()
+		, SavedManifoldPoints()
 		, CullDistance(TNumericLimits<FReal>::Max())
 		, CollisionMargins{ 0, 0 }
 		, CollisionTolerance(0)
-		, bUseManifold(false)
-		, bUseIncrementalManifold(false)
+		, Flags()
 		, SolverBodies{ nullptr, nullptr }
 		, GJKWarmStartData()
 		, ShapeWorldTransform0()
@@ -158,7 +158,6 @@ namespace Chaos
 		, LastShapeWorldTransform0()
 		, LastShapeWorldTransform1()
 		, ExpectedNumManifoldPoints(0)
-		, bWasManifoldRestored(false)
 		, NumActivePositionIterations(0)
 	{
 		Manifold.Implicit[0] = nullptr;
@@ -180,15 +179,15 @@ namespace Chaos
 		, AccumulatedImpulse(0)
 		, Manifold()
 		, TimeOfImpact(0)
+		, ContainerCookie()
 		, CCDType(ECollisionCCDType::Disabled)
 		, Stiffness(FReal(1))
 		, ManifoldPoints()
-		, ManifoldPointsSavedData()
+		, SavedManifoldPoints()
 		, CullDistance(TNumericLimits<FReal>::Max())
 		, CollisionMargins{ 0, 0 }
 		, CollisionTolerance(0)
-		, bUseManifold(false)
-		, bUseIncrementalManifold(false)
+		, Flags()
 		, SolverBodies{ nullptr, nullptr }
 		, GJKWarmStartData()
 		, ShapeWorldTransform0()
@@ -196,7 +195,7 @@ namespace Chaos
 		, LastShapeWorldTransform0()
 		, LastShapeWorldTransform1()
 		, ExpectedNumManifoldPoints(0)
-		, bWasManifoldRestored(false)
+		, NumActivePositionIterations(0)
 	{
 		Manifold.Implicit[0] = Implicit0;
 		Manifold.Implicit[1] = Implicit1;
@@ -222,8 +221,8 @@ namespace Chaos
 
 		CullDistance = InCullDistance;
 
-		bUseManifold = bInUseManifold && CanUseManifold(Particle[0], Particle[1]);
-		bUseIncrementalManifold = true;	// This will get changed later if we call AddOneShotManifoldContact
+		Flags.bUseManifold = bInUseManifold && CanUseManifold(Particle[0], Particle[1]);
+		Flags.bUseIncrementalManifold = true;	// This will get changed later if we call AddOneShotManifoldContact
 
 		const FReal Margin0 = GetImplicit0()->GetMargin();
 		const FReal Margin1 = GetImplicit1()->GetMargin();
@@ -387,7 +386,7 @@ namespace Chaos
 
 			UpdateManifoldPointFromContact(ManifoldPointIndex);
 
-			ManifoldPoint.bInsideStaticFrictionCone = bUseManifold;
+			ManifoldPoint.bInsideStaticFrictionCone = Flags.bUseManifold;
 
 			// Copy currently active point
 			if (ManifoldPoint.ContactPoint.Phi < Manifold.Phi)
@@ -401,7 +400,7 @@ namespace Chaos
 	{
 		if (ContactPoint.IsSet())
 		{
-			if (ManifoldPoints.Num() == MaxManifoldPoints)
+			if (ManifoldPoints.IsFull())
 			{
 				return;
 			}
@@ -414,18 +413,18 @@ namespace Chaos
 				SetActiveContactPoint(ManifoldPoints[ManifoldPointIndex].ContactPoint);
 			}
 
-			bUseIncrementalManifold = false;
+			Flags.bUseIncrementalManifold = false;
 		}
 	}
 
 	void FPBDCollisionConstraint::AddIncrementalManifoldContact(const FContactPoint& ContactPoint)
 	{
-		if (ManifoldPoints.Num() == MaxManifoldPoints)
+		if (ManifoldPoints.IsFull())
 		{
 			return;
 		}
 
-		if (bUseManifold)
+		if (Flags.bUseManifold)
 		{
 			// See if the manifold point already exists
 			int32 ManifoldPointIndex = FindManifoldPoint(ContactPoint);
@@ -451,26 +450,15 @@ namespace Chaos
 		else 
 		{
 			// We are not using manifolds - reuse the first and only point
-			if (ManifoldPoints.Num() == 0)
-			{
-				ManifoldPoints.Add(ContactPoint);
-			}
-			else
-			{
-				ManifoldPoints[0].ContactPoint = ContactPoint;
-			}
+			ManifoldPoints.SetNum(1);
+			ManifoldPoints[0].ContactPoint = ContactPoint;
 
 			InitManifoldPoint(0);
 
 			SetActiveContactPoint(ManifoldPoints[0].ContactPoint);
 		}
 
-		bUseIncrementalManifold = true;
-	}
-
-	void FPBDCollisionConstraint::ClearManifold()
-	{
-		ManifoldPoints.Reset();
+		Flags.bUseIncrementalManifold = true;
 	}
 
 	void FPBDCollisionConstraint::InitManifoldPoint(const int32 ManifoldPointIndex)
@@ -495,9 +483,9 @@ namespace Chaos
 
 	int32 FPBDCollisionConstraint::AddManifoldPoint(const FContactPoint& ContactPoint)
 	{
-		check(ManifoldPoints.Num() < MaxManifoldPoints);
+		int32 ManifoldPointIndex = ManifoldPoints.Add();
 
-		int32 ManifoldPointIndex = ManifoldPoints.Add(ContactPoint);
+		ManifoldPoints[ManifoldPointIndex].ContactPoint = ContactPoint;
 
 		InitManifoldPoint(ManifoldPointIndex);
 
@@ -528,16 +516,16 @@ namespace Chaos
 
 	void FPBDCollisionConstraint::ResetManifold()
 	{
-		ManifoldPointsSavedData.Reset();
+		SavedManifoldPoints.Reset();
 		ResetActiveManifoldContacts();
 	}
 
 	void FPBDCollisionConstraint::ResetActiveManifoldContacts()
 	{
-		ManifoldPoints.Reset();
 		Manifold.Reset();
+		ManifoldPoints.Reset();
 		ExpectedNumManifoldPoints = 0;
-		bWasManifoldRestored = false;
+		Flags.bWasManifoldRestored = false;
 	}
 
 	void FPBDCollisionConstraint::RestoreManifold()
@@ -547,7 +535,7 @@ namespace Chaos
 		// world-space state for the contact modifiers
 		UpdateManifoldContacts();
 
-		bWasManifoldRestored = true;
+		Flags.bWasManifoldRestored = true;
 	}
 
 
@@ -565,8 +553,6 @@ namespace Chaos
 
 	bool FPBDCollisionConstraint::UpdateAndTryRestoreManifold()
 	{
-		check(ManifoldPoints.Num() <= MaxManifoldPoints);
-
 		const FCollisionTolerances& Tolerances = Chaos_Manifold_Tolerances;
 		const FReal ContactPositionTolerance = Tolerances.ContactPositionToleranceScale * CollisionTolerance;
 		const FReal ShapePositionTolerance = (ManifoldPoints.Num() > 0) ? Tolerances.ShapePositionToleranceScaleN * CollisionTolerance : Tolerances.ShapePositionToleranceScale0 * CollisionTolerance;
@@ -580,79 +566,83 @@ namespace Chaos
 		// we have a face or edge contact. We don't reuse the manifold if we lose points after culling
 		// here and potentially adding the new narrow phase result (See TryAddManifoldContact).
 		ExpectedNumManifoldPoints = ManifoldPoints.Num();
-		bWasManifoldRestored = false;
+		Flags.bWasManifoldRestored = false;
 
-		const FRigidTransform3 Shape0ToShape1Transform = ShapeWorldTransform0.GetRelativeTransformNoScale(ShapeWorldTransform1);
-
-		// Update or prune manifold points.
-		TArray<int32, TInlineAllocator<4>> ManifoldPointsToRemove;
-		for (int32 ManifoldPointIndex = 0; ManifoldPointIndex < ManifoldPoints.Num(); ++ManifoldPointIndex)
+		// Either update or remove each manifold point depending on how far it has moved from its initial relative point
+		// NOTE: We do not reset if we have 0 points - we can still "restore" a zero point manifold if the bodies have not moved
+		int32 ManifoldPointToRemove = INDEX_NONE;
+		if (ManifoldPoints.Num() > 0)
 		{
-			FManifoldPoint& ManifoldPoint = ManifoldPoints[ManifoldPointIndex];
+			const FRigidTransform3 Shape0ToShape1Transform = ShapeWorldTransform0.GetRelativeTransformNoScale(ShapeWorldTransform1);
+			
+			// Update or prune manifold points. If we would end up removing more than 1 point, we just throw the 
+			// whole manifold away because it will get rebuilt in the narrow phasee anyway.
+			for (int32 ManifoldPointIndex = 0; ManifoldPointIndex < ManifoldPoints.Num(); ++ManifoldPointIndex)
+			{
+				FManifoldPoint& ManifoldPoint = ManifoldPoints[ManifoldPointIndex];
 
-			// Calculate the world-space contact location and separation at the current shape transforms
-			// @todo(chaos): this should use the normal owner. Currently we assume body 1 is the owner
-			const FVec3 Contact0In1 = Shape0ToShape1Transform.TransformPositionNoScale(ManifoldPoint.InitialShapeContactPoints[0]);
-			const FVec3& Contact1In1 = ManifoldPoint.InitialShapeContactPoints[1];
-			const FVec3 ContactNormalIn1 = ShapeWorldTransform1.InverseTransformVectorNoScale(ManifoldPoint.ContactPoint.Normal);
+				// Calculate the world-space contact location and separation at the current shape transforms
+				// @todo(chaos): this should use the normal owner. Currently we assume body 1 is the owner
+				const FVec3 Contact0In1 = Shape0ToShape1Transform.TransformPositionNoScale(ManifoldPoint.InitialShapeContactPoints[0]);
+				const FVec3& Contact1In1 = ManifoldPoint.InitialShapeContactPoints[1];
+				const FVec3 ContactNormalIn1 = ShapeWorldTransform1.InverseTransformVectorNoScale(ManifoldPoint.ContactPoint.Normal);
 
-			const FVec3 ContactDeltaIn1 = Contact0In1 - Contact1In1;
-			const FReal ContactPhi = FVec3::DotProduct(ContactDeltaIn1, ContactNormalIn1);
-			const FVec3 ContactLateralDeltaIn1 = ContactDeltaIn1 - ContactPhi * ContactNormalIn1;
-			const FReal ContactLateralDistanceSq = ContactLateralDeltaIn1.SizeSquared();
+				const FVec3 ContactDeltaIn1 = Contact0In1 - Contact1In1;
+				const FReal ContactPhi = FVec3::DotProduct(ContactDeltaIn1, ContactNormalIn1);
+				const FVec3 ContactLateralDeltaIn1 = ContactDeltaIn1 - ContactPhi * ContactNormalIn1;
+				const FReal ContactLateralDistanceSq = ContactLateralDeltaIn1.SizeSquared();
 
-			// Either update the point or flag it for removal
-			if (ContactLateralDistanceSq < ContactPositionToleranceSq)
-			{
-				// Recalculate the contact points at the new location
-				// @todo(chaos): we should reproject the contact on the plane owner
-				const FVec3 ShapeContactPoint1 = Contact0In1 - ContactPhi * ContactNormalIn1;
-				ManifoldPoint.ContactPoint.ShapeContactPoints[1] = ShapeContactPoint1;
-				ManifoldPoint.ContactPoint.Phi = ContactPhi;
+				// Either update the point or flag it for removal
+				if (ContactLateralDistanceSq < ContactPositionToleranceSq)
+				{
+					// Recalculate the contact points at the new location
+					// @todo(chaos): we should reproject the contact on the plane owner
+					const FVec3 ShapeContactPoint1 = Contact0In1 - ContactPhi * ContactNormalIn1;
+					ManifoldPoint.ContactPoint.ShapeContactPoints[1] = ShapeContactPoint1;
+					ManifoldPoint.ContactPoint.Phi = ContactPhi;
+				}
+				else if (ManifoldPointToRemove == INDEX_NONE)
+				{
+					ManifoldPointToRemove = ManifoldPointIndex;
+				}
+				else
+				{
+					// We want to remove a second point, but we will never reuse it now so throw away the whole manifold
+					ResetActiveManifoldContacts();
+					return false;
+				}
 			}
-			else
+
+			// Remove points - only one point removal suport required (see above)
+			if (ManifoldPointToRemove != INDEX_NONE)
 			{
-				ManifoldPointsToRemove.Add(ManifoldPointIndex);
+				ManifoldPoints.RemoveAt(ManifoldPointToRemove);
 			}
-		}
-
-		// If we removed more than 1 point, we are going to rebuild the manifold anyway, so we're done
-		if (ManifoldPointsToRemove.Num() > 1)
-		{
-			ResetActiveManifoldContacts();
-			return false;
-		}
 
-		// Remove points - process in reverse order for simpler index handling and faster removal of multiple elements
-		for (int32 RemoveArrayIndex = ManifoldPointsToRemove.Num() - 1; RemoveArrayIndex >= 0; --RemoveArrayIndex)
-		{
-			constexpr bool bAllowShrink = false;
-			ManifoldPoints.RemoveAt(ManifoldPointsToRemove[RemoveArrayIndex], 1, bAllowShrink);
-		}
-
-		// Update world-space state for the points we kept
-		for (int32 ManifoldPointIndex = 0; ManifoldPointIndex < ManifoldPoints.Num(); ++ManifoldPointIndex)
-		{
-			FManifoldPoint& ManifoldPoint = ManifoldPoints[ManifoldPointIndex];
+			// Update world-space state for the points we kept
+			for (int32 ManifoldPointIndex = 0; ManifoldPointIndex < ManifoldPoints.Num(); ++ManifoldPointIndex)
+			{
+				FManifoldPoint& ManifoldPoint = ManifoldPoints[ManifoldPointIndex];
 
-			// Restore friction anchors if we have them for this point
-			TryRestoreFrictionData(ManifoldPointIndex);
+				// Restore friction anchors if we have them for this point
+				TryRestoreFrictionData(ManifoldPointIndex);
 
-			// Update world-space contact locations
-			UpdateManifoldPointFromContact(ManifoldPointIndex);
-			ManifoldPoint.ContactPoint.Location = FReal(0.5) * (ManifoldPoint.WorldContactPoints[0] + ManifoldPoint.WorldContactPoints[1]);
+				// Update world-space contact locations
+				UpdateManifoldPointFromContact(ManifoldPointIndex);
+				ManifoldPoint.ContactPoint.Location = FReal(0.5) * (ManifoldPoint.WorldContactPoints[0] + ManifoldPoint.WorldContactPoints[1]);
 
-			ManifoldPoint.bWasRestored = true;
+				ManifoldPoint.bWasRestored = true;
 
-			if (ManifoldPoint.ContactPoint.Phi < Manifold.Phi)
-			{
-				// Update closest point
-				SetActiveContactPoint(ManifoldPoint.ContactPoint);
+				if (ManifoldPoint.ContactPoint.Phi < Manifold.Phi)
+				{
+					// Update closest point
+					SetActiveContactPoint(ManifoldPoint.ContactPoint);
+				}
 			}
 		}
 
 		// If we did not remove any contact points and we have not moved or rotated much we may reuse the manifold as-is.
-		if ((ManifoldPointsToRemove.Num() == 0) && (ShapePositionTolerance > 0) && (ShapeRotationThreshold > 0))
+		if ((ManifoldPointToRemove == INDEX_NONE) && (ShapePositionTolerance > 0) && (ShapeRotationThreshold > 0))
 		{
 			// The transform check is necessary regardless of how many points we have left in the manifold because
 			// as a body moves/rotates we may have to change which faces/edges are colliding. We can't know if the face/edge
@@ -677,8 +667,6 @@ namespace Chaos
 
 	bool FPBDCollisionConstraint::TryAddManifoldContact(const FContactPoint& NewContactPoint)
 	{
-		check(ManifoldPoints.Num() <= MaxManifoldPoints);
-
 		const FCollisionTolerances& Tolerances = Chaos_Manifold_Tolerances;
 		const FReal PositionTolerance = Tolerances.ManifoldPointPositionToleranceScale * CollisionTolerance;
 		const FReal NormalThreshold = Tolerances.ManifoldPointNormalThreshold;
@@ -696,8 +684,7 @@ namespace Chaos
 		// case here is that we will regenerate the manifold too often so it will work but could be bad for perf
 		const FReal PositionToleranceSq = FMath::Square(PositionTolerance);
 		int32 MatchedManifoldPointIndex = INDEX_NONE;
-		const int32 NumManifoldPoints = ManifoldPoints.Num();
-		for (int32 ManifoldPointIndex = 0; ManifoldPointIndex < NumManifoldPoints; ++ManifoldPointIndex)
+		for (int32 ManifoldPointIndex = 0; ManifoldPointIndex < ManifoldPoints.Num(); ++ManifoldPointIndex)
 		{
 			FManifoldPoint& ManifoldPoint = ManifoldPoints[ManifoldPointIndex];
 
@@ -890,12 +877,12 @@ namespace Chaos
 		if (bChaos_Manifold_EnableFrictionRestore)
 		{
 			const FReal DistanceToleranceSq = FMath::Square(Chaos_Manifold_FrictionPositionTolerance);
-			for (int32 PrevManifoldPointIndex = 0; PrevManifoldPointIndex < ManifoldPointsSavedData.Num(); ++PrevManifoldPointIndex)
+			for (int32 SavedManifoldPointIndex = 0; SavedManifoldPointIndex < SavedManifoldPoints.Num(); ++SavedManifoldPointIndex)
 			{
-				const FManifoldPointSavedData& PrevManifoldPoint = ManifoldPointsSavedData[PrevManifoldPointIndex];
-				if (PrevManifoldPoint.IsMatch(ManifoldPoint, DistanceToleranceSq))
+				const FManifoldPointSavedData& SavedManifoldPoint = SavedManifoldPoints[SavedManifoldPointIndex];
+				if (SavedManifoldPoint.IsMatch(ManifoldPoint, DistanceToleranceSq))
 				{
-					return &PrevManifoldPoint;
+					return &SavedManifoldPoint;
 				}
 			}
 		}
@@ -908,7 +895,7 @@ namespace Chaos
 
 		// Assume we have no matching point from the previous tick, but that we can retain friction from now on
 		// Not supported for non-manifolds yet (hopefully we don't need to)
-		ManifoldPoint.bInsideStaticFrictionCone = bUseManifold;
+		ManifoldPoint.bInsideStaticFrictionCone = Flags.bUseManifold;
 		ManifoldPoint.StaticFrictionMax = FReal(0);
 
 		// Find the previous manifold point that matches if there is one
diff --git a/Engine/Source/Runtime/Experimental/Chaos/Private/Chaos/Collision/SolverCollisionContainer.cpp b/Engine/Source/Runtime/Experimental/Chaos/Private/Chaos/Collision/SolverCollisionContainer.cpp
--- a/Engine/Source/Runtime/Experimental/Chaos/Private/Chaos/Collision/SolverCollisionContainer.cpp
+++ b/Engine/Source/Runtime/Experimental/Chaos/Private/Chaos/Collision/SolverCollisionContainer.cpp
@@ -43,13 +43,14 @@ namespace Chaos
 	public:
 		FPBDCollisionSolverAdapter()
 			: Constraint(nullptr)
+			, bIsIncremental(false)
 		{
 		}
 
 		FPBDCollisionSolver& GetSolver() { return Solver; }
 		FPBDCollisionConstraint* GetConstraint() { return Constraint; }
 
-		bool IsIncrementalManifold() const { return Constraint->GetUseIncrementalCollisionDetection(); }
+		bool IsIncrementalManifold() const { return bIsIncremental; }
 
 		/**
 		 * @brief Initialize the data required for the solver, and bind to the bodies
@@ -62,6 +63,7 @@ namespace Chaos
 			FSolverBodyContainer& SolverBodyContainer)
 		{
 			Constraint = &InConstraint;
+			bIsIncremental = Constraint->GetUseIncrementalCollisionDetection();
 
 			// Find the solver bodies for the particles we constrain. This will add them to the container
 			// if they aren't there already, and ensure that they are populated with the latest data.
@@ -146,6 +148,7 @@ namespace Chaos
 	private:
 		FPBDCollisionSolver Solver;
 		FPBDCollisionConstraint* Constraint;
+		bool bIsIncremental;
 	};
 
 
diff --git a/Engine/Source/Runtime/Experimental/Chaos/Private/Chaos/CollisionResolution.cpp b/Engine/Source/Runtime/Experimental/Chaos/Private/Chaos/CollisionResolution.cpp
--- a/Engine/Source/Runtime/Experimental/Chaos/Private/Chaos/CollisionResolution.cpp
+++ b/Engine/Source/Runtime/Experimental/Chaos/Private/Chaos/CollisionResolution.cpp
@@ -1784,10 +1784,6 @@ namespace Chaos
 				ensure(false);
 				break;
 			}
-
-			//const FVec3 NewContactPositionLocal0 = WorldTransform0.InverseTransformPosition(Constraint.Manifold.Location);
-			//const FVec3 NewContactPositionLocal1 = WorldTransform1.InverseTransformPosition(Constraint.Manifold.Location);
-			//Constraint.Manifold.ContactMoveSQRDistance = FMath::Max((NewContactPositionLocal0 - OriginalContactPositionLocal0).SizeSquared(), (NewContactPositionLocal1 - OriginalContactPositionLocal1).SizeSquared());
 		}
 
 		// Run collision detection for the swept constraints
diff --git a/Engine/Source/Runtime/Experimental/Chaos/Public/Chaos/Collision/PBDCollisionConstraint.h b/Engine/Source/Runtime/Experimental/Chaos/Public/Chaos/Collision/PBDCollisionConstraint.h
--- a/Engine/Source/Runtime/Experimental/Chaos/Public/Chaos/Collision/PBDCollisionConstraint.h
+++ b/Engine/Source/Runtime/Experimental/Chaos/Public/Chaos/Collision/PBDCollisionConstraint.h
@@ -2,8 +2,9 @@
 #pragma once
 
 #include "CoreMinimal.h"
-
 #include "Chaos/Core.h"
+
+#include "Chaos/BVHParticles.h"
 #include "Chaos/CollisionResolutionTypes.h"
 #include "Chaos/Collision/CollisionKeys.h"
 #include "Chaos/Collision/ContactPoint.h"
@@ -13,7 +14,6 @@
 #include "Chaos/ParticleHandleFwd.h"
 #include "Chaos/ParticleHandle.h"
 #include "Chaos/Vector.h"
-#include "Chaos/BVHParticles.h"
 
 namespace Chaos
 {
@@ -80,10 +80,19 @@ namespace Chaos
 	class CHAOS_API FManifoldPointSavedData
 	{
 	public:
+		FManifoldPointSavedData()
+		{
+		}
+
+		FManifoldPointSavedData(const FManifoldPoint& ManifoldPoint)
+		{
+			Save(ManifoldPoint);
+		}
+
 		/**
 		 * @brief Copy the ManifoldPoint data needed for static friction next tick
 		*/
-		FManifoldPointSavedData(const FManifoldPoint& ManifoldPoint)
+		inline void Save(const FManifoldPoint& ManifoldPoint)
 		{
 			ShapeContactPoints[0] = ManifoldPoint.ContactPoint.ShapeContactPoints[0];
 			ShapeContactPoints[1] = ManifoldPoint.ContactPoint.ShapeContactPoints[1];
@@ -126,7 +135,6 @@ namespace Chaos
 	public:
 		FCollisionContact(const FImplicitObject* InImplicit0 = nullptr, const FBVHParticles* InSimplicial0 = nullptr, const FImplicitObject* InImplicit1 = nullptr, const FBVHParticles* InSimplicial1 = nullptr)
 			: bDisabled(false)
-			, bUseAccumalatedImpulseInSolve(true)
 			, Normal(0)
 			, Location(0)
 			, Phi(FLT_MAX)
@@ -139,7 +147,6 @@ namespace Chaos
 			, InvMassScale1(1.f)
 			, InvInertiaScale0(1.f)
 			, InvInertiaScale1(1.f)
-			, ContactMoveSQRDistance(0)
 			, ShapesType(EContactShapesType::Unknown)
 		{
 			Implicit[0] = InImplicit0;
@@ -149,7 +156,6 @@ namespace Chaos
 		}
 
 		bool bDisabled;
-		bool bUseAccumalatedImpulseInSolve; // Accumulated Impulse will only be used when the contact did not move significantly during iterations (This will reduce jitter)
 
 		FVec3 Normal;
 		FVec3 Location;
@@ -164,7 +170,6 @@ namespace Chaos
 		FReal InvMassScale1;
 		FReal InvInertiaScale0;
 		FReal InvInertiaScale1;
-		FReal ContactMoveSQRDistance; // How much the contact position moved after the last update
 		EContactShapesType ShapesType;
 
 		void Reset()
@@ -376,9 +381,9 @@ namespace Chaos
 		FReal GetCullDistance() const { return CullDistance; }
 		void SetCullDistance(FReal InCullDistance) { CullDistance = InCullDistance; }
 
-		bool GetUseManifold() const { return bUseManifold; }
+		bool GetUseManifold() const { return Flags.bUseManifold; }
 		
-		bool GetUseIncrementalCollisionDetection() const { return !bUseManifold || bUseIncrementalManifold; }
+		bool GetUseIncrementalCollisionDetection() const { return !Flags.bUseManifold || Flags.bUseIncrementalManifold; }
 
 		void SetNumActivePositionIterations(const int32 InNumActivePositionIterations) { NumActivePositionIterations = InNumActivePositionIterations; }
 		int32 GetNumActivePositionIterations() const { return NumActivePositionIterations; }
@@ -388,13 +393,13 @@ namespace Chaos
 		*/
 		void ResetManifold();
 
-		TArrayView<FManifoldPoint> GetManifoldPoints() { return MakeArrayView(ManifoldPoints); }
-		TArrayView<const FManifoldPoint> GetManifoldPoints() const { return MakeArrayView(ManifoldPoints); }
+		// @todo(chaos): remove array view and provide per-point accessor
+		TArrayView<FManifoldPoint> GetManifoldPoints() { return MakeArrayView(ManifoldPoints.begin(), ManifoldPoints.Num()); }
+		TArrayView<const FManifoldPoint> GetManifoldPoints() const { return MakeArrayView(ManifoldPoints.begin(), ManifoldPoints.Num()); }
 
 		void AddIncrementalManifoldContact(const FContactPoint& ContactPoint);
 		void AddOneshotManifoldContact(const FContactPoint& ContactPoint);
 		void UpdateManifoldContacts();
-		void ClearManifold();
 
 		const FRigidTransform3& GetShapeRelativeTransform0() const { return ImplicitTransform[0]; }
 		const FRigidTransform3& GetShapeRelativeTransform1() const { return ImplicitTransform[1]; }
@@ -411,7 +416,14 @@ namespace Chaos
 		bool TryInsertManifoldContact(const FContactPoint& ContactPoint);
 
 		//@ todo(chaos): These are for the collision forwarding system - this should use the collision modifier system (which should be extended to support adding collisions)
-		void SetManifoldPoints(const TArray<FManifoldPoint>& InManifoldPoints) { ManifoldPoints = InManifoldPoints; }
+		void SetManifoldPoints(const TArray<FManifoldPoint>& InManifoldPoints)
+		{ 
+			ManifoldPoints.SetNum(FMath::Min(MaxManifoldPoints, InManifoldPoints.Num()));
+			for (int32 ManifoldPointIndex = 0; ManifoldPoints.Num(); ++ManifoldPointIndex)
+			{
+				ManifoldPoints[ManifoldPointIndex] = InManifoldPoints[ManifoldPointIndex];
+			}
+		}
 		void UpdateManifoldPointFromContact(const int32 ManifoldPointIndex);
 
 		// The GJK warm-start data. This is updated directly in the narrow phase
@@ -437,7 +449,7 @@ namespace Chaos
 		/**
 		 * @brief Whether this constraint was fully restored from a previous tick, and the manifold should be reused as-is
 		*/
-		bool WasManifoldRestored() const { return bWasManifoldRestored; }
+		bool WasManifoldRestored() const { return Flags.bWasManifoldRestored; }
 
 		/**
 		 * @brief Restore the contact manifold (assumes relative motion of the two bodies is small)
@@ -459,7 +471,7 @@ namespace Chaos
 		*/
 		void ResetSolverResults()
 		{
-			ManifoldPointsSavedData.Reset();
+			SavedManifoldPoints.Reset();
 			AccumulatedImpulse = FVec3(0);
 		}
 
@@ -478,9 +490,10 @@ namespace Chaos
 			AccumulatedImpulse += NetImpulse + (NetPushOut / Dt);
 
 			// If are still satisfying the static friction condition, save off the state required for next tick
-			if (bInsideStaticFrictionCone && !NetPushOut.IsNearlyZero() && (ManifoldPointsSavedData.Num() < ManifoldPointsSavedData.Max()))
+			if (bInsideStaticFrictionCone && !NetPushOut.IsNearlyZero() && !SavedManifoldPoints.IsFull())
 			{
-				ManifoldPointsSavedData.Emplace(FManifoldPointSavedData(ManifoldPoint));
+				const int32 SavedIndex = SavedManifoldPoints.Add();
+				SavedManifoldPoints[SavedIndex].Save(ManifoldPoint);
 			}
 		}
 
@@ -543,11 +556,10 @@ namespace Chaos
 		ECollisionCCDType CCDType;
 		FReal Stiffness;
 
-		TFixedArray<FManifoldPoint, MaxManifoldPoints> ManifoldPoints;
+		TCArray<FManifoldPoint, MaxManifoldPoints> ManifoldPoints;
 
 		// The manifold points from the previous tick when we don't reuse the manifold. Used by static friction.
-		TFixedArray<FManifoldPointSavedData, MaxManifoldPoints> ManifoldPointsSavedData;
-
+		TCArray<FManifoldPointSavedData, MaxManifoldPoints> SavedManifoldPoints;
 		FReal CullDistance;
 
 		// The margins to use during collision detection. We don't always use the margins on the shapes directly.
@@ -558,10 +570,19 @@ namespace Chaos
 		// margins of the two shapes, as well as their types
 		FReal CollisionTolerance;
 
-		bool bUseManifold;
-		bool bUseIncrementalManifold;
-
-		// These are only needed here while we still have thge lgacy solvers (not QuasiPBD)
+		union FFlags
+		{
+			FFlags() : Bits(0) {}
+			struct
+			{
+				bool bUseManifold;
+				bool bUseIncrementalManifold;
+				bool bWasManifoldRestored;
+			};
+			uint32 Bits;
+		} Flags;
+
+		// These are only needed here while we still have the legacy solvers (not QuasiPBD)
 		FSolverBody* SolverBodies[2];
 
 		// Simplex data from the last call to GJK, used to warm-start GJK
@@ -572,8 +593,6 @@ namespace Chaos
 		FRigidTransform3 LastShapeWorldTransform0;
 		FRigidTransform3 LastShapeWorldTransform1;
 		int32 ExpectedNumManifoldPoints;
-		bool bWasManifoldRestored;
-
 		int32 NumActivePositionIterations;
 	};
 
diff --git a/Engine/Source/Runtime/Experimental/Chaos/Public/Chaos/Framework/UncheckedArray.h b/Engine/Source/Runtime/Experimental/Chaos/Public/Chaos/Framework/UncheckedArray.h
--- a/Engine/Source/Runtime/Experimental/Chaos/Public/Chaos/Framework/UncheckedArray.h
+++ b/Engine/Source/Runtime/Experimental/Chaos/Public/Chaos/Framework/UncheckedArray.h
@@ -17,6 +17,160 @@
 
 namespace Chaos
 {
+	/**
+	 * @brief A c-style array of objects with non-shipping bounds checking.
+	 * 
+	 * This behaves exactly like a C-style array, although it also keeps track of the number of
+	 * elements in the array. This element count is artficial - all elements in the array are 
+	 * default constructed and will not be destructed until the array itself is destroyed.
+	 * 
+	 * @note the element type must have a default constructor
+	*/
+	template<typename T, int32 N>
+	class TCArray
+	{
+	public:
+		static const int32 MaxElements = N;
+		using ElementType = T;
+
+		inline TCArray() : NumElements(0) {}
+
+		inline int32 Num() const { return NumElements; }
+		inline int32 Max() const { return MaxElements; }
+		inline bool IsEmpty() const { return NumElements == 0; }
+		inline bool IsFull() const { return NumElements == MaxElements; }
+
+		inline ElementType& operator[](const int32 Index)
+		{
+			check(Index < NumElements);
+			return Elements[Index];
+		}
+
+		inline const ElementType& operator[](const int32 Index) const
+		{
+			check(Index < NumElements);
+			return Elements[Index];
+		}
+
+		/**
+		 * @brief Set the number of elements in the array
+		 * @note Does not reset any added or removed elements, it only changes the number-of-elements counter
+		*/
+		inline void SetNum(const int32 InNum)
+		{
+			check(InNum < NumElements);
+			NumElements = InNum;
+		}
+
+		/**
+		 * @brief Set the number of elements to 0
+		 * @note Does not reset any elements
+		*/
+		inline void Reset()
+		{
+			NumElements = 0;
+		}
+
+		/**
+		 * @brief Set the number of elements to 0
+		 * @note Does not reset any elements
+		*/
+		inline void Empty()
+		{
+			NumElements = 0;
+		}
+
+		/**
+		 * @brief Increase the size of the array without re-initializing the new element
+		*/
+		inline int32 Add()
+		{
+			check(NumElements < MaxElements);
+			return NumElements++;
+		}
+
+		/**
+		 * @brief Copy the element to the end of the array
+		*/
+		inline int32 Add(const ElementType& V)
+		{
+			check(NumElements < MaxElements);
+			Elements[NumElements] = V;
+			return NumElements++;
+		}
+
+		/**
+		 * @brief Move the element to the end of the array
+		*/
+		inline int32 Add(ElementType&& V)
+		{
+			check(NumElements < MaxElements);
+			Elements[NumElements] = MoveTemp(V);
+			return NumElements++;
+		}
+
+		/**
+		 * @brief Move the element to the end of the array
+		*/
+		inline int32 Emplace(ElementType&& V)
+		{
+			return Add(MoveTemp(V));
+		}
+
+		/**
+		 * @brief Remove the element at the specified index
+		 * Moves all higher elements down to fill the gap
+		*/
+		inline void RemoveAt(const int32 Index)
+		{
+			check(Index < NumElements);
+			for (int32 MoveIndex = Index; MoveIndex < NumElements - 1; ++MoveIndex)
+			{
+				Elements[MoveIndex] = MoveTemp(Elements[MoveIndex + 1]);
+			}
+			--NumElements;
+		}
+
+		/**
+		 * @brief Remove the element at the specified index
+		 * Moves the last element into the gap.
+		*/
+		inline void RemoveAtSwap(const int32 Index)
+		{
+			check(Index < NumElements);
+			if (Index < NumElements - 1)
+			{
+				Elements[Index] = MoveTemp(Elements[NumElements - 1]);
+			}
+			--NumElements;
+		}
+
+		inline ElementType* begin()
+		{
+			return &Elements[0];
+		}
+
+		inline const ElementType* begin() const
+		{
+			return &Elements[0];
+		}
+
+		inline ElementType* end()
+		{
+			return &Elements[NumElements];
+		}
+
+		inline const ElementType* end() const
+		{
+			return &Elements[NumElements];
+		}
+
+	private:
+		ElementType Elements[MaxElements];
+		int32 NumElements;
+	};
+
+
 	/**
 	 * @brief A fixed allocator without array bounds checking except in Debug builds.
 	 *
diff --git a/Engine/Source/Runtime/Experimental/Chaos/Public/Chaos/Utilities.h b/Engine/Source/Runtime/Experimental/Chaos/Public/Chaos/Utilities.h
--- a/Engine/Source/Runtime/Experimental/Chaos/Public/Chaos/Utilities.h
+++ b/Engine/Source/Runtime/Experimental/Chaos/Public/Chaos/Utilities.h
@@ -235,27 +235,23 @@ namespace Chaos
 		 * Multiple a vector by a matrix: C = L.R
 		 * If L is a rotation matrix, then this will return R rotated by that rotation.
 		 */
-		inline FVec3 Multiply(const FMatrix33& LIn, const FVec3& R)
+		inline FVec3 Multiply(const FMatrix33& L, const FVec3& R)
 		{
-			// @todo(ccaulfield): optimize: remove transposes and use simd etc
-			FMatrix33 L = LIn.GetTransposed();
-
+			// @todo(chaos): optimize: use simd
 			return FVec3(
-			    L.M[0][0] * R.X + L.M[0][1] * R.Y + L.M[0][2] * R.Z,
-			    L.M[1][0] * R.X + L.M[1][1] * R.Y + L.M[1][2] * R.Z,
-			    L.M[2][0] * R.X + L.M[2][1] * R.Y + L.M[2][2] * R.Z);
+			    L.M[0][0] * R.X + L.M[1][0] * R.Y + L.M[2][0] * R.Z,
+			    L.M[0][1] * R.X + L.M[1][1] * R.Y + L.M[2][1] * R.Z,
+			    L.M[0][2] * R.X + L.M[1][2] * R.Y + L.M[2][2] * R.Z);
 		}
 
-		inline FVec4 Multiply(const FMatrix44& LIn, const FVec4& R)
+		inline FVec4 Multiply(const FMatrix44& L, const FVec4& R)
 		{
-			// @todo(ccaulfield): optimize: remove transposes and use simd etc
-			FMatrix44 L = LIn.GetTransposed();
-
+			// @todo(chaos): optimize: use simd
 			return FVec4(
-				L.M[0][0] * R.X + L.M[0][1] * R.Y + L.M[0][2] * R.Z + L.M[0][3] * R.W,
-				L.M[1][0] * R.X + L.M[1][1] * R.Y + L.M[1][2] * R.Z + L.M[1][3] * R.W,
-				L.M[2][0] * R.X + L.M[2][1] * R.Y + L.M[2][2] * R.Z + L.M[2][3] * R.W,
-				L.M[3][0] * R.X + L.M[3][1] * R.Y + L.M[3][2] * R.Z + L.M[3][3] * R.W
+				L.M[0][0] * R.X + L.M[1][0] * R.Y + L.M[2][0] * R.Z + L.M[3][0] * R.W,
+				L.M[0][1] * R.X + L.M[1][1] * R.Y + L.M[2][1] * R.Z + L.M[3][1] * R.W,
+				L.M[0][2] * R.X + L.M[1][2] * R.Y + L.M[2][2] * R.Z + L.M[3][2] * R.W,
+				L.M[0][3] * R.X + L.M[1][3] * R.Y + L.M[2][3] * R.Z + L.M[3][3] * R.W
 				);
 		}
 
