Chaos - separate solver loops for friction/non-friction/incremental paths
- inline solver loops

diff --git a/Engine/Source/Runtime/Experimental/Chaos/Private/Chaos/Collision/PBDCollisionSolver.cpp b/Engine/Source/Runtime/Experimental/Chaos/Private/Chaos/Collision/PBDCollisionSolver.cpp
--- a/Engine/Source/Runtime/Experimental/Chaos/Private/Chaos/Collision/PBDCollisionSolver.cpp
+++ b/Engine/Source/Runtime/Experimental/Chaos/Private/Chaos/Collision/PBDCollisionSolver.cpp
@@ -15,12 +15,6 @@
 
 //PRAGMA_DISABLE_OPTIMIZATION
 
-// Set to 0 to use a linearized error calculation, and set to 1 to use a non-linear error calculation in collision detection. 
-// In principle nonlinear is more accurate when large rotation corrections occur, but this is not too important for collisions because 
-// when the bodies settle the corrections are small. The linearized version is significantly faster than the non-linear version because 
-// the non-linear version requires a quaternion multiply and renormalization whereas the linear version is just a cross product.
-#define CHAOS_NONLINEAR_COLLISIONS_ENABLED 0
-
 DEFINE_LOG_CATEGORY(LogChaosCollision);
 
 namespace Chaos
@@ -65,289 +59,6 @@ namespace Chaos
 	}
 	using namespace CVars;
 
-	inline FReal GetPositionInvMassScale(int32 It, int32 NumIts)
-	{
-		if (Chaos_PBDCollisionSolver_Position_ShockPropagationIterations > 0)
-		{
-			const int32 FirstShockIt = FMath::Max(NumIts - Chaos_PBDCollisionSolver_Position_ShockPropagationIterations, 0);
-			const FReal Interpolant = FMath::Clamp(FReal(It - FirstShockIt + 1) / FReal(NumIts - FirstShockIt), FReal(0), FReal(1));
-			return FMath::Lerp((FReal)1.0f, (FReal)Chaos_PBDCollisionSolver_Position_MinInvMassScale, Interpolant);
-		}
-		return 1.0f;
-	}
-
-	inline FReal GetVelocityInvMassScale(int32 It, int32 NumIts)
-	{
-		if (Chaos_PBDCollisionSolver_Velocity_ShockPropagationIterations > 0)
-		{
-			const int32 FirstShockIt = FMath::Max(NumIts - Chaos_PBDCollisionSolver_Velocity_ShockPropagationIterations, 0);
-			const FReal Interpolant = FMath::Clamp(FReal(It - FirstShockIt + 1) / FReal(NumIts - FirstShockIt), FReal(0), FReal(1));
-			return FMath::Lerp((FReal)1.0f, (FReal)Chaos_PBDCollisionSolver_Velocity_MinInvMassScale, Interpolant);
-		}
-		return 1.0f;
-	}
-
-	inline void CalculatePositionCorrectionWithoutFriction(
-		const FReal Stiffness,
-		const FReal ContactDeltaNormal,
-		const FVec3& ContactNormal,
-		const FReal ContactMassNormal,
-		FVec3& InOutNetPushOut,
-		FVec3& OutPushOut)
-	{
-		FVec3 PushOut = -(Stiffness * ContactDeltaNormal * ContactMassNormal) * ContactNormal;
-
-		// The total pushout so far this sub-step
-		// We allow negative incremental impulses, but not net negative impulses
-		const FVec3 NetPushOut = InOutNetPushOut + PushOut;
-		const FReal NetPushOutNormal = FVec3::DotProduct(NetPushOut, ContactNormal);
-		if (NetPushOutNormal < 0)
-		{
-			PushOut = -InOutNetPushOut;
-		}
-
-		InOutNetPushOut += PushOut;
-		OutPushOut = PushOut;
-	}
-
-
-	inline void CalculatePositionCorrectionWithFriction(
-		const FReal Stiffness,
-		const FVec3& ContactDelta,
-		const FReal ContactDeltaNormal,
-		const FVec3& ContactNormal,
-		const FMatrix33& ContactMass,
-		const FReal StaticFriction,
-		const FReal DynamicFriction,
-		FVec3& InOutNetPushOut,
-		FVec3& OutPushOut,
-		FReal& InOutStaticFrictionMax,
-		bool& bOutInsideStaticFrictionCone)
-	{
-		// If static friction is enabled, calculate the correction to move the contact point back to its
-		// original relative location on all axes.
-		// @todo(chaos): this should be moved to the ManifoldPoint error calculation?
-		FVec3 ModifiedContactError = -ContactDelta;
-		const FReal FrictionStiffness = Chaos_PBDCollisionSolver_Position_StaticFrictionStiffness;
-		if (FrictionStiffness < 1.0f)
-		{
-			const FVec3 ContactDeltaTangent = ContactDelta - ContactDeltaNormal * ContactNormal;
-			ModifiedContactError = -ContactDeltaNormal * ContactNormal - FrictionStiffness * ContactDeltaTangent;
-		}
-
-		FVec3 PushOut = Stiffness * ContactMass * ModifiedContactError;
-
-		// If we ended up with a negative normal pushout, disable friction
-		FVec3 NetPushOut = InOutNetPushOut + PushOut;
-		const FReal NetPushOutNormal = FVec3::DotProduct(NetPushOut, ContactNormal);
-		bool bInsideStaticFrictionCone = true;
-		if (NetPushOutNormal < FReal(SMALL_NUMBER))
-		{
-			bInsideStaticFrictionCone = false;
-		}
-
-		// Static friction limit: immediately increase maximum lateral correction, but smoothly decay maximum static friction limit. 
-		// This is so that small variations in position (jitter) and therefore NetPushOutNormal don't cause static friction to slip
-		// @todo(chaos): StaticFriction smoothing is iteration count depenendent - try to make it not so
-		const FReal StaticFrictionLerpRate = Chaos_PBDCollisionSolver_Position_StaticFrictionLerpRate;
-		const FReal StaticFrictionDest = FMath::Max(NetPushOutNormal, FReal(0));
-		FReal StaticFrictionMax = FMath::Lerp(FMath::Max(InOutStaticFrictionMax, StaticFrictionDest), StaticFrictionDest, StaticFrictionLerpRate);
-
-		// If we exceed the friction cone, stop adding frictional corrections (although
-		// any already added lateral corrections will not be undone)
-		// @todo(chaos): clamp to dynamic friction
-		if (bInsideStaticFrictionCone)
-		{
-			const FReal MaxPushOutTangentSq = FMath::Square(StaticFriction * StaticFrictionMax);
-			const FVec3 NetPushOutTangent = (NetPushOut - NetPushOutNormal * ContactNormal);
-			const FReal NetPushOutTangentSq = NetPushOutTangent.SizeSquared();
-			if (NetPushOutTangentSq > MaxPushOutTangentSq)
-			{
-				NetPushOut = NetPushOutNormal * ContactNormal + (StaticFriction * StaticFrictionMax) * NetPushOutTangent / FMath::Sqrt(NetPushOutTangentSq);
-				PushOut = NetPushOut - InOutNetPushOut;
-				bInsideStaticFrictionCone = false;
-			}
-		}
-
-		// If we leave the friction cone, we will flll through into the non-friction impulse calculation
-		// so do not export the results
-		if (bInsideStaticFrictionCone)
-		{
-			InOutNetPushOut = NetPushOut;
-			OutPushOut = PushOut;
-		}
-		else
-		{
-			OutPushOut = FVec3(0);
-		}
-
-		InOutStaticFrictionMax = StaticFrictionMax;
-		bOutInsideStaticFrictionCone = bInsideStaticFrictionCone;
-	}
-
-	inline void ApplyPositionCorrection(
-		const FReal Stiffness,
-		const FReal StaticFriction,
-		const FReal DynamicFriction,
-		const FVec3& ContactDelta,
-		const FReal ContactDeltaNormal,
-		FPBDCollisionSolverManifoldPoint& ManifoldPoint,
-		FConstraintSolverBody& Body0,
-		FConstraintSolverBody& Body1)
-	{
-		FVec3 PushOut = FVec3(0);
-
-		// If we want to add static friction...
-		// (note: we run a few iterations without friction by temporarily setting StaticFriction to 0)
-		if ((StaticFriction > 0) && ManifoldPoint.bInsideStaticFrictionCone)
-		{
-			CalculatePositionCorrectionWithFriction(
-				Stiffness,
-				ContactDelta,
-				ContactDeltaNormal,
-				ManifoldPoint.WorldContactNormal,
-				ManifoldPoint.WorldContactMass,
-				StaticFriction,
-				DynamicFriction,
-				ManifoldPoint.NetPushOut,					// Out
-				PushOut,									// Out
-				ManifoldPoint.StaticFrictionMax,			// Out
-				ManifoldPoint.bInsideStaticFrictionCone);	// Out
-		}
-		else
-		{
-			CalculatePositionCorrectionWithoutFriction(
-				Stiffness,
-				ContactDeltaNormal,
-				ManifoldPoint.WorldContactNormal,
-				ManifoldPoint.WorldContactMassNormal,
-				ManifoldPoint.NetPushOut,					// Out
-				PushOut);									// Out
-		}
-
-		// Update the particle state based on the pushout
-		if (Body0.IsDynamic())
-		{
-			const FVec3 AngularPushOut = FVec3::CrossProduct(ManifoldPoint.WorldContactPosition - Body0.P(), PushOut);
-			const FVec3 DX0 = Body0.InvM() * PushOut;
-			const FVec3 DR0 = Body0.InvI() * AngularPushOut;
-			//if (!DX0.IsNearlyZero(Chaos_PBDCollisionSolver_Position_PositionSolverTolerance))
-			{
-				Body0.ApplyPositionDelta(DX0);
-			}
-			//if (!DR0.IsNearlyZero(Chaos_PBDCollisionSolver_Position_RotationSolverTolerance))
-			{
-				Body0.ApplyRotationDelta(DR0);
-			}
-		}
-		if (Body1.IsDynamic())
-		{
-			const FVec3 AngularPushOut = FVec3::CrossProduct(ManifoldPoint.WorldContactPosition - Body1.P(), PushOut);
-			const FVec3 DX1 = -(Body1.InvM() * PushOut);
-			const FVec3 DR1 = -(Body1.InvI() * AngularPushOut);
-			//if (!DX1.IsNearlyZero(Chaos_PBDCollisionSolver_Position_PositionSolverTolerance))
-			{
-				Body1.ApplyPositionDelta(DX1);
-			}
-			//if (!DR1.IsNearlyZero(Chaos_PBDCollisionSolver_Position_RotationSolverTolerance))
-			{
-				Body1.ApplyRotationDelta(DR1);
-			}
-		}
-	}
-
-	inline void CalculateVelocityCorrectionImpulse(
-		const FReal Stiffness,
-		const FReal Dt,
-		const FReal DynamicFriction,
-		const FVec3& ContactNormal,
-		const FMatrix33& ContactMass,
-		const FReal ContactMassNormal,
-		const FVec3& ContactVelocityDelta,
-		const FReal ContactVelocityDeltaNormal,
-		const FVec3& NetPushOut,
-		FVec3& InOutNetImpulse,
-		FVec3& OutImpulse)
-	{
-		FVec3 Impulse = FVec3(0);
-
-		if ((ContactVelocityDeltaNormal > FReal(0)) && !bChaos_PBDCollisionSolver_Velocity_NegativeImpulseEnabled)
-		{
-			return;
-		}
-
-		// Tangential velocity (dynamic friction)
-		const bool bApplyFriction = (DynamicFriction > 0) && (Dt > 0);
-		if (bApplyFriction)
-		{
-			Impulse = -Stiffness * (ContactMass * ContactVelocityDelta);
-		}
-		else
-		{
-			Impulse = -(Stiffness * ContactMassNormal) * ContactVelocityDelta;
-		}
-
-		// Clamp the total impulse to be positive along the normal. We can apply negative velocity correction, 
-		// but only to correct the velocity that was added by pushout, or in this velocity solve step.
-		if (bChaos_PBDCollisionSolver_Velocity_ImpulseClampEnabled && (Dt > 0))
-		{
-			// @todo(chaos): cache max negative impulse
-			const FVec3 NetImpulse = InOutNetImpulse + Impulse;
-			const FReal PushOutImpulseNormal = FMath::Max(0.0f, FVec3::DotProduct(NetPushOut, ContactNormal) / Dt);
-			const FReal NetImpulseNormal = FVec3::DotProduct(NetImpulse, ContactNormal);
-			if (NetImpulseNormal < -PushOutImpulseNormal)
-			{
-				// We are trying to apply a negative impulse larger than one to counteract the effective pushout impulse
-				// so clamp the net impulse to be equal to minus the pushout impulse along the normal.
-				// NOTE: NetImpulseNormal is negative here
-				//const FVec3 NewNetImpulse = NetImpulse - NetImpulseNormal * ContactNormal - PushOutImpulseNormal * ContactNormal;
-				//const FVec3 NewImpulse = NewNetImpulse - InOutNetImpulse;
-				//Impulse = InOutNetImpulse + Impulse - NetImpulseNormal * ContactNormal - PushOutImpulseNormal * ContactNormal - InOutNetImpulse;
-				Impulse = Impulse - (NetImpulseNormal + PushOutImpulseNormal) * ContactNormal;
-			}
-		}
-
-		OutImpulse = Impulse;
-		InOutNetImpulse += Impulse;
-	}
-
-	inline void ApplyVelocityCorrection(
-		const FReal Stiffness,
-		const FReal Dt,
-		const FReal DynamicFriction,
-		const FVec3& ContactVelocityDelta,
-		const FReal ContactVelocityDeltaNormal,
-		FPBDCollisionSolverManifoldPoint& ManifoldPoint,
-		FConstraintSolverBody& Body0,
-		FConstraintSolverBody& Body1)
-	{
-		FVec3 Impulse = FVec3(0);
-
-		CalculateVelocityCorrectionImpulse(
-			Stiffness,
-			Dt,
-			DynamicFriction,
-			ManifoldPoint.WorldContactNormal,
-			ManifoldPoint.WorldContactMass,
-			ManifoldPoint.WorldContactMassNormal,
-			ContactVelocityDelta,
-			ContactVelocityDeltaNormal,
-			ManifoldPoint.NetPushOut,	// Out
-			ManifoldPoint.NetImpulse,	// Out
-			Impulse);					// Out
-
-		// Calculate the velocity deltas from the impulse
-		if (Body0.IsDynamic())
-		{
-			const FVec3 AngularImpulse = FVec3::CrossProduct(ManifoldPoint.WorldContactPosition - Body0.P(), Impulse);
-			Body0.ApplyVelocityDelta(Body0.InvM() * Impulse, Body0.InvI() * AngularImpulse);
-		}
-		if (Body1.IsDynamic())
-		{
-			const FVec3 AngularImpulse = FVec3::CrossProduct(ManifoldPoint.WorldContactPosition - Body1.P(), Impulse);
-			Body1.ApplyVelocityDelta(Body1.InvM() * -Impulse, Body1.InvI() * -AngularImpulse);
-		}
-	}
 
 	void FPBDCollisionSolverManifoldPoint::InitContact(
 		const FConstraintSolverBody& Body0,
@@ -422,68 +133,6 @@ namespace Chaos
 		return ContactVelocity0 - ContactVelocity1;
 	}
 
-	inline void FPBDCollisionSolverManifoldPoint::CalculateContactPositionError(const FConstraintSolverBody& Body0, const FConstraintSolverBody& Body1, const FReal MaxPushOut, FVec3& OutContactDelta, FReal& OutContactDeltaNormal) const
-	{
-		const FVec3 WorldRelativeContactPosition0 = WorldContactPosition - Body0.P();
-		const FVec3 WorldRelativeContactPosition1 = WorldContactPosition - Body1.P();
-#if CHAOS_NONLINEAR_COLLISIONS_ENABLED
-		// Non-linear version: calculate the contact delta after we have converted the current positional impulses into position and rotation corrections.
-		// We could precalculate and store the LocalContactPositions if we really want to use this nonlinear version
-		const FVec3 LocalContactPosition0 = Body0.Q().Inverse() * WorldRelativeContactPosition0;
-		const FVec3 LocalContactPosition1 = Body1.Q().Inverse() * WorldRelativeContactPosition1;
-		OutContactDelta = (Body0.CorrectedP() + Body0.CorrectedQ() * LocalContactPosition0) - (Body1.CorrectedP() + Body1.CorrectedQ() * LocalContactPosition1);
-#else
-		// Linear version: calculate the contact delta assuming linear motion after applying a positional impulse at the contact point. There will be an error that depends on the size of the rotation.
-		const FVec3 ContactDelta0 = Body0.DP() + FVec3::CrossProduct(Body0.DQ(), WorldRelativeContactPosition0);
-		const FVec3 ContactDelta1 = Body1.DP() + FVec3::CrossProduct(Body1.DQ(), WorldRelativeContactPosition1);
-		OutContactDelta = WorldContactDelta + ContactDelta0 - ContactDelta1;
-#endif
-		OutContactDeltaNormal = FVec3::DotProduct(OutContactDelta, WorldContactNormal);
-
-		// NOTE: OutContactDeltaNormal is negative for penetration
-		// NOTE: MaxPushOut == 0 disables the pushout limits
-		if ((MaxPushOut > 0) && (OutContactDeltaNormal < -MaxPushOut))
-		{
-			const FReal ClampedContactDeltaNormal = -MaxPushOut;
-			OutContactDelta += (ClampedContactDeltaNormal - OutContactDeltaNormal) * WorldContactNormal;
-			OutContactDeltaNormal = ClampedContactDeltaNormal;
-		}
-	}
-
-	inline void FPBDCollisionSolverManifoldPoint::CalculateContactVelocityError(const FConstraintSolverBody& Body0, const FConstraintSolverBody& Body1, const FReal DynamicFriction, const FReal Dt, FVec3& OutContactVelocityDelta, FReal& OutContactVelocityDeltaNormal) const
-	{
-		const FVec3 ContactVelocity0 = Body0.V() + FVec3::CrossProduct(Body0.W(), WorldContactPosition - Body0.P());
-		const FVec3 ContactVelocity1 = Body1.V() + FVec3::CrossProduct(Body1.W(), WorldContactPosition - Body1.P());
-		const FVec3 ContactVelocity = ContactVelocity0 - ContactVelocity1;
-		const FReal ContactVelocityNormal = FVec3::DotProduct(ContactVelocity, WorldContactNormal);
-
-		// Target normal velocity, including restitution
-		const FReal ContactVelocityTargetNormal = WorldContactVelocityTargetNormal;
-
-		// Add up the errors in the velocity (current velocity - desired velocity)
-		OutContactVelocityDeltaNormal = (ContactVelocityNormal - WorldContactVelocityTargetNormal);
-		OutContactVelocityDelta = (ContactVelocityNormal - ContactVelocityTargetNormal) * WorldContactNormal;
-
-		const bool bApplyFriction = (DynamicFriction > 0) && (Dt > 0);
-		if (bApplyFriction)
-		{
-			const FVec3 ContactVelocityTangential = ContactVelocity - ContactVelocityNormal * WorldContactNormal;
-			const FReal ContactVelocityTangentialLen = ContactVelocityTangential.Size();
-			if (ContactVelocityTangentialLen > SMALL_NUMBER)
-			{
-				// PushOut = ContactMass * DP, where DP is the contact positional correction
-				// Friction force is proportional to the normal force, so friction velocity correction
-				// is proprtional to normal velocity correction, or DVn = DPn/dt = PushOut.N / (ContactMass * dt);
-				const FReal PushOutNormal = FVec3::DotProduct(NetPushOut, WorldContactNormal);
-				const FReal DynamicFrictionVelocityError = PushOutNormal / (WorldContactMassNormal * Dt);
-				if (DynamicFrictionVelocityError > SMALL_NUMBER)
-				{
-					const FReal ContactVelocityErrorTangential = FMath::Min(DynamicFriction * DynamicFrictionVelocityError, ContactVelocityTangentialLen);
-					OutContactVelocityDelta += ContactVelocityTangential * (ContactVelocityErrorTangential / ContactVelocityTangentialLen);
-				}
-			}
-		}
-	}
 
 
 	//////////////////////////////////////////////////////////////////////////////////////////////////
@@ -584,68 +233,6 @@ namespace Chaos
 		State.ManifoldPoints[ManifoldPoiontIndex].UpdateContact(State.SolverBodies[0], State.SolverBodies[1], InWorldAnchorPoint0, InWorldAnchorPoint1, InWorldContactNormal);
 	}
 
-	bool FPBDCollisionSolver::SolvePosition(const FReal Dt, const FReal MaxPushOut, const bool bApplyStaticFriction)
-	{
-		if (!bChaos_PBDCollisionSolver_Position_SolveEnabled)
-		{
-			return false;
-		}
-
-		// SolverBody decorator used to add mass scaling
-		FConstraintSolverBody& Body0 = SolverBody0();
-		FConstraintSolverBody& Body1 = SolverBody1();
-
-		// Check for solved
-		if (State.bIsSolved)
-		{
-			if ((Body0.LastChangeEpoch() == State.BodyEpochs[0]) && (Body1.LastChangeEpoch() == State.BodyEpochs[1]))
-			{
-				// We did not apply a correction last iteration (within tolerance) and nothing else has moved the bodies
-				return false;
-			}
-		}
-
-		// The first few iterations have friction disabled. This allows us to solve the normal penetration at each contact point
-		// without adding spurious lateral impulses which get cancelled out by other contacts. The spurious lateral impulses
-		// cause problems when clipping to the friction cone on a per-contact basis.
-		FReal StaticFriction = (bApplyStaticFriction) ? State.StaticFriction : FReal(0);
-		FReal DynamicFriction = (bApplyStaticFriction) ? State.DynamicFriction : FReal(0);
-		const FReal Stiffness = State.Stiffness;
-
-		// Apply the position correction so that all contacts have zero separation
-		for (int32 PointIndex = 0; PointIndex < State.NumManifoldPoints; ++PointIndex)
-		{
-			FPBDCollisionSolverManifoldPoint& SolverManifoldPoint = State.ManifoldPoints[PointIndex];
-
-			FVec3 ContactDelta;
-			FReal ContactDeltaNormal;
-			SolverManifoldPoint.CalculateContactPositionError(Body0.SolverBody(), Body1.SolverBody(), MaxPushOut, ContactDelta, ContactDeltaNormal);
-
-			const bool bProcessManifoldPoint = (ContactDeltaNormal < FReal(0)) || !SolverManifoldPoint.NetPushOut.IsNearlyZero();
-			if (bProcessManifoldPoint)
-			{
-				ApplyPositionCorrection(
-					Stiffness,
-					StaticFriction,
-					DynamicFriction,
-					ContactDelta,
-					ContactDeltaNormal,
-					SolverManifoldPoint,
-					Body0,
-					Body1);
-			}
-		}
-
-		// We are solved if we did not move the bodies within some tolerance
-		// NOTE: we can't claim to be solved until we have done at least one friction iteration so we can't early out before friction has been applied.
-		// @todo(chaos): better early-out system
-		State.bIsSolved = bApplyStaticFriction && (Body0.LastChangeEpoch() == State.BodyEpochs[0]) && (Body1.LastChangeEpoch() == State.BodyEpochs[1]);
-		State.BodyEpochs[0] = Body0.LastChangeEpoch();
-		State.BodyEpochs[1] = Body1.LastChangeEpoch();
-		State.NumPositionSolves = State.NumPositionSolves + 1;
-
-		return !State.bIsSolved;
-	}
 
 	void FPBDCollisionSolver::SolveVelocityAverage(const FReal Dt)
 	{
@@ -714,11 +301,6 @@ namespace Chaos
 
 	bool FPBDCollisionSolver::SolveVelocity(const FReal Dt, const bool bApplyDynamicFriction)
 	{
-		if (!bChaos_PBDCollisionSolver_Velocity_SolveEnabled)
-		{
-			return false;
-		}
-
 		// Apply restitution at the average contact point
 		// This means we don't need to run as many iterations to get stable bouncing
 		if (State.bHaveRestitution && (State.NumManifoldPoints > 1))
diff --git a/Engine/Source/Runtime/Experimental/Chaos/Private/Chaos/Collision/PBDCollisionSolver.h b/Engine/Source/Runtime/Experimental/Chaos/Private/Chaos/Collision/PBDCollisionSolver.h
--- a/Engine/Source/Runtime/Experimental/Chaos/Private/Chaos/Collision/PBDCollisionSolver.h
+++ b/Engine/Source/Runtime/Experimental/Chaos/Private/Chaos/Collision/PBDCollisionSolver.h
@@ -8,6 +8,12 @@
 #include "Chaos/Defines.h"
 #include "Chaos/Evolution/SolverBody.h"
 
+// Set to 0 to use a linearized error calculation, and set to 1 to use a non-linear error calculation in collision detection. 
+// In principle nonlinear is more accurate when large rotation corrections occur, but this is not too important for collisions because 
+// when the bodies settle the corrections are small. The linearized version is significantly faster than the non-linear version because 
+// the non-linear version requires a quaternion multiply and renormalization whereas the linear version is just a cross product.
+#define CHAOS_NONLINEAR_COLLISIONS_ENABLED 0
+
 namespace Chaos
 {
 	namespace Collisions
@@ -19,6 +25,15 @@ namespace Chaos
 	class FPBDCollisionSolver;
 	class FPBDCollisionConstraint;
 
+	namespace CVars
+	{
+		extern float Chaos_PBDCollisionSolver_Position_StaticFrictionStiffness;
+		extern float Chaos_PBDCollisionSolver_Position_StaticFrictionLerpRate;
+		
+		extern bool bChaos_PBDCollisionSolver_Velocity_NegativeImpulseEnabled;
+		extern bool bChaos_PBDCollisionSolver_Velocity_ImpulseClampEnabled;
+	}
+
 	/**
 	 * @brief A single contact point in a FPBDCollisionSolver
 	*/
@@ -228,7 +243,8 @@ namespace Chaos
 		 * @brief Calculate and apply the position correction for this iteration
 		 * @return true if we need to run more iterations, false if we did not apply any correction
 		*/
-		bool SolvePosition(const FReal Dt, const FReal MaxPushOut, const bool bApplyStaticFriction);
+		bool SolvePositionWithFriction(const FReal Dt, const FReal MaxPushOut);
+		bool SolvePositionNoFriction(const FReal Dt, const FReal MaxPushOut);
 
 		/**
 		 * @brief Calculate and apply the velocity correction for this iteration
@@ -283,6 +299,491 @@ namespace Chaos
 
 		FState State;
 	};
+
+
+	//////////////////////////////////////////////////////////////////////////////////////////////////
+	//////////////////////////////////////////////////////////////////////////////////////////////////
+	//////////////////////////////////////////////////////////////////////////////////////////////////
+	//////////////////////////////////////////////////////////////////////////////////////////////////
+
+
+	FORCEINLINE_DEBUGGABLE void CalculatePositionCorrectionWithoutFriction(
+		const FReal Stiffness,
+		const FReal ContactDeltaNormal,
+		const FVec3& ContactNormal,
+		const FReal ContactMassNormal,
+		FVec3& InOutNetPushOut,
+		FVec3& OutPushOut)
+	{
+		FVec3 PushOut = -(Stiffness * ContactDeltaNormal * ContactMassNormal) * ContactNormal;
+
+		// The total pushout so far this sub-step
+		// We allow negative incremental impulses, but not net negative impulses
+		const FVec3 NetPushOut = InOutNetPushOut + PushOut;
+		const FReal NetPushOutNormal = FVec3::DotProduct(NetPushOut, ContactNormal);
+		if (NetPushOutNormal < 0)
+		{
+			PushOut = -InOutNetPushOut;
+		}
+
+		InOutNetPushOut += PushOut;
+		OutPushOut = PushOut;
+	}
+
+
+	FORCEINLINE_DEBUGGABLE void CalculatePositionCorrectionWithFriction(
+		const FReal Stiffness,
+		const FVec3& ContactDelta,
+		const FReal ContactDeltaNormal,
+		const FVec3& ContactNormal,
+		const FMatrix33& ContactMass,
+		const FReal StaticFriction,
+		const FReal DynamicFriction,
+		FVec3& InOutNetPushOut,
+		FVec3& OutPushOut,
+		FReal& InOutStaticFrictionMax,
+		bool& bOutInsideStaticFrictionCone)
+	{
+		// If static friction is enabled, calculate the correction to move the contact point back to its
+		// original relative location on all axes.
+		// @todo(chaos): this should be moved to the ManifoldPoint error calculation?
+		FVec3 ModifiedContactError = -ContactDelta;
+		const FReal FrictionStiffness = CVars::Chaos_PBDCollisionSolver_Position_StaticFrictionStiffness;
+		if (FrictionStiffness < 1.0f)
+		{
+			const FVec3 ContactDeltaTangent = ContactDelta - ContactDeltaNormal * ContactNormal;
+			ModifiedContactError = -ContactDeltaNormal * ContactNormal - FrictionStiffness * ContactDeltaTangent;
+		}
+
+		FVec3 PushOut = Stiffness * ContactMass * ModifiedContactError;
+
+		// If we ended up with a negative normal pushout, disable friction
+		FVec3 NetPushOut = InOutNetPushOut + PushOut;
+		const FReal NetPushOutNormal = FVec3::DotProduct(NetPushOut, ContactNormal);
+		bool bInsideStaticFrictionCone = true;
+		if (NetPushOutNormal < FReal(SMALL_NUMBER))
+		{
+			bInsideStaticFrictionCone = false;
+		}
+
+		// Static friction limit: immediately increase maximum lateral correction, but smoothly decay maximum static friction limit. 
+		// This is so that small variations in position (jitter) and therefore NetPushOutNormal don't cause static friction to slip
+		// @todo(chaos): StaticFriction smoothing is iteration count depenendent - try to make it not so
+		const FReal StaticFrictionLerpRate = CVars::Chaos_PBDCollisionSolver_Position_StaticFrictionLerpRate;
+		const FReal StaticFrictionDest = FMath::Max(NetPushOutNormal, FReal(0));
+		FReal StaticFrictionMax = FMath::Lerp(FMath::Max(InOutStaticFrictionMax, StaticFrictionDest), StaticFrictionDest, StaticFrictionLerpRate);
+
+		// If we exceed the friction cone, stop adding frictional corrections (although
+		// any already added lateral corrections will not be undone)
+		// @todo(chaos): clamp to dynamic friction
+		if (bInsideStaticFrictionCone)
+		{
+			const FReal MaxPushOutTangentSq = FMath::Square(StaticFriction * StaticFrictionMax);
+			const FVec3 NetPushOutTangent = (NetPushOut - NetPushOutNormal * ContactNormal);
+			const FReal NetPushOutTangentSq = NetPushOutTangent.SizeSquared();
+			if (NetPushOutTangentSq > MaxPushOutTangentSq)
+			{
+				NetPushOut = NetPushOutNormal * ContactNormal + (StaticFriction * StaticFrictionMax) * NetPushOutTangent / FMath::Sqrt(NetPushOutTangentSq);
+				PushOut = NetPushOut - InOutNetPushOut;
+				bInsideStaticFrictionCone = false;
+			}
+		}
+
+		// If we leave the friction cone, we will fall through into the non-friction impulse calculation
+		// so do not export the results
+		if (bInsideStaticFrictionCone)
+		{
+			InOutNetPushOut = NetPushOut;
+			OutPushOut = PushOut;
+		}
+		else
+		{
+			OutPushOut = FVec3(0);
+		}
+
+		InOutStaticFrictionMax = StaticFrictionMax;
+		bOutInsideStaticFrictionCone = bInsideStaticFrictionCone;
+	}
+
+	FORCEINLINE_DEBUGGABLE void ApplyPositionCorrectionWithFriction(
+		const FReal Stiffness,
+		const FReal StaticFriction,
+		const FReal DynamicFriction,
+		const FVec3& ContactDelta,
+		const FReal ContactDeltaNormal,
+		FPBDCollisionSolverManifoldPoint& ManifoldPoint,
+		FConstraintSolverBody& Body0,
+		FConstraintSolverBody& Body1)
+	{
+		FVec3 PushOut = FVec3(0);
+
+		// If we want to add static friction...
+		// (note: we run a few iterations without friction by temporarily setting StaticFriction to 0)
+		if ((StaticFriction > 0) && ManifoldPoint.bInsideStaticFrictionCone)
+		{
+			CalculatePositionCorrectionWithFriction(
+				Stiffness,
+				ContactDelta,
+				ContactDeltaNormal,
+				ManifoldPoint.WorldContactNormal,
+				ManifoldPoint.WorldContactMass,
+				StaticFriction,
+				DynamicFriction,
+				ManifoldPoint.NetPushOut,					// Out
+				PushOut,									// Out
+				ManifoldPoint.StaticFrictionMax,			// Out
+				ManifoldPoint.bInsideStaticFrictionCone);	// Out
+		}
+		else
+		{
+			CalculatePositionCorrectionWithoutFriction(
+				Stiffness,
+				ContactDeltaNormal,
+				ManifoldPoint.WorldContactNormal,
+				ManifoldPoint.WorldContactMassNormal,
+				ManifoldPoint.NetPushOut,					// Out
+				PushOut);									// Out
+		}
+
+		// Update the particle state based on the pushout
+		if (Body0.IsDynamic())
+		{
+			const FVec3 AngularPushOut = FVec3::CrossProduct(ManifoldPoint.WorldContactPosition - Body0.P(), PushOut);
+			const FVec3 DX0 = Body0.InvM() * PushOut;
+			const FVec3 DR0 = Body0.InvI() * AngularPushOut;
+			//if (!DX0.IsNearlyZero(CVars::Chaos_PBDCollisionSolver_Position_PositionSolverTolerance))
+			{
+				Body0.ApplyPositionDelta(DX0);
+			}
+			//if (!DR0.IsNearlyZero(CVars::Chaos_PBDCollisionSolver_Position_RotationSolverTolerance))
+			{
+				Body0.ApplyRotationDelta(DR0);
+			}
+		}
+		if (Body1.IsDynamic())
+		{
+			const FVec3 AngularPushOut = FVec3::CrossProduct(ManifoldPoint.WorldContactPosition - Body1.P(), PushOut);
+			const FVec3 DX1 = -(Body1.InvM() * PushOut);
+			const FVec3 DR1 = -(Body1.InvI() * AngularPushOut);
+			//if (!DX1.IsNearlyZero(CVars::Chaos_PBDCollisionSolver_Position_PositionSolverTolerance))
+			{
+				Body1.ApplyPositionDelta(DX1);
+			}
+			//if (!DR1.IsNearlyZero(CVars::Chaos_PBDCollisionSolver_Position_RotationSolverTolerance))
+			{
+				Body1.ApplyRotationDelta(DR1);
+			}
+		}
+	}
+
+	FORCEINLINE_DEBUGGABLE void ApplyPositionCorrectionNoFriction(
+		const FReal Stiffness,
+		const FVec3& ContactDelta,
+		const FReal ContactDeltaNormal,
+		FPBDCollisionSolverManifoldPoint& ManifoldPoint,
+		FConstraintSolverBody& Body0,
+		FConstraintSolverBody& Body1)
+	{
+		FVec3 PushOut = FVec3(0);
+
+		CalculatePositionCorrectionWithoutFriction(
+			Stiffness,
+			ContactDeltaNormal,
+			ManifoldPoint.WorldContactNormal,
+			ManifoldPoint.WorldContactMassNormal,
+			ManifoldPoint.NetPushOut,					// Out
+			PushOut);									// Out
+
+		// Update the particle state based on the pushout
+		if (Body0.IsDynamic())
+		{
+			const FVec3 AngularPushOut = FVec3::CrossProduct(ManifoldPoint.WorldContactPosition - Body0.P(), PushOut);
+			const FVec3 DX0 = Body0.InvM() * PushOut;
+			const FVec3 DR0 = Body0.InvI() * AngularPushOut;
+			//if (!DX0.IsNearlyZero(CVars::Chaos_PBDCollisionSolver_Position_PositionSolverTolerance))
+			{
+				Body0.ApplyPositionDelta(DX0);
+			}
+			//if (!DR0.IsNearlyZero(CVars::Chaos_PBDCollisionSolver_Position_RotationSolverTolerance))
+			{
+				Body0.ApplyRotationDelta(DR0);
+			}
+		}
+		if (Body1.IsDynamic())
+		{
+			const FVec3 AngularPushOut = FVec3::CrossProduct(ManifoldPoint.WorldContactPosition - Body1.P(), PushOut);
+			const FVec3 DX1 = -(Body1.InvM() * PushOut);
+			const FVec3 DR1 = -(Body1.InvI() * AngularPushOut);
+			//if (!DX1.IsNearlyZero(CVars::Chaos_PBDCollisionSolver_Position_PositionSolverTolerance))
+			{
+				Body1.ApplyPositionDelta(DX1);
+			}
+			//if (!DR1.IsNearlyZero(CVars::Chaos_PBDCollisionSolver_Position_RotationSolverTolerance))
+			{
+				Body1.ApplyRotationDelta(DR1);
+			}
+		}
+	}
+
+	FORCEINLINE_DEBUGGABLE void CalculateVelocityCorrectionImpulse(
+		const FReal Stiffness,
+		const FReal Dt,
+		const FReal DynamicFriction,
+		const FVec3& ContactNormal,
+		const FMatrix33& ContactMass,
+		const FReal ContactMassNormal,
+		const FVec3& ContactVelocityDelta,
+		const FReal ContactVelocityDeltaNormal,
+		const FVec3& NetPushOut,
+		FVec3& InOutNetImpulse,
+		FVec3& OutImpulse)
+	{
+		FVec3 Impulse = FVec3(0);
+
+		if ((ContactVelocityDeltaNormal > FReal(0)) && !CVars::bChaos_PBDCollisionSolver_Velocity_NegativeImpulseEnabled)
+		{
+			return;
+		}
+
+		// Tangential velocity (dynamic friction)
+		const bool bApplyFriction = (DynamicFriction > 0) && (Dt > 0);
+		if (bApplyFriction)
+		{
+			Impulse = -Stiffness * (ContactMass * ContactVelocityDelta);
+		}
+		else
+		{
+			Impulse = -(Stiffness * ContactMassNormal) * ContactVelocityDelta;
+		}
+
+		// Clamp the total impulse to be positive along the normal. We can apply negative velocity correction, 
+		// but only to correct the velocity that was added by pushout, or in this velocity solve step.
+		if (CVars::bChaos_PBDCollisionSolver_Velocity_ImpulseClampEnabled && (Dt > 0))
+		{
+			// @todo(chaos): cache max negative impulse
+			const FVec3 NetImpulse = InOutNetImpulse + Impulse;
+			const FReal PushOutImpulseNormal = FMath::Max(0.0f, FVec3::DotProduct(NetPushOut, ContactNormal) / Dt);
+			const FReal NetImpulseNormal = FVec3::DotProduct(NetImpulse, ContactNormal);
+			if (NetImpulseNormal < -PushOutImpulseNormal)
+			{
+				// We are trying to apply a negative impulse larger than one to counteract the effective pushout impulse
+				// so clamp the net impulse to be equal to minus the pushout impulse along the normal.
+				// NOTE: NetImpulseNormal is negative here
+				//const FVec3 NewNetImpulse = NetImpulse - NetImpulseNormal * ContactNormal - PushOutImpulseNormal * ContactNormal;
+				//const FVec3 NewImpulse = NewNetImpulse - InOutNetImpulse;
+				//Impulse = InOutNetImpulse + Impulse - NetImpulseNormal * ContactNormal - PushOutImpulseNormal * ContactNormal - InOutNetImpulse;
+				Impulse = Impulse - (NetImpulseNormal + PushOutImpulseNormal) * ContactNormal;
+			}
+		}
+
+		OutImpulse = Impulse;
+		InOutNetImpulse += Impulse;
+	}
+
+	FORCEINLINE_DEBUGGABLE void ApplyVelocityCorrection(
+		const FReal Stiffness,
+		const FReal Dt,
+		const FReal DynamicFriction,
+		const FVec3& ContactVelocityDelta,
+		const FReal ContactVelocityDeltaNormal,
+		FPBDCollisionSolverManifoldPoint& ManifoldPoint,
+		FConstraintSolverBody& Body0,
+		FConstraintSolverBody& Body1)
+	{
+		FVec3 Impulse = FVec3(0);
+
+		CalculateVelocityCorrectionImpulse(
+			Stiffness,
+			Dt,
+			DynamicFriction,
+			ManifoldPoint.WorldContactNormal,
+			ManifoldPoint.WorldContactMass,
+			ManifoldPoint.WorldContactMassNormal,
+			ContactVelocityDelta,
+			ContactVelocityDeltaNormal,
+			ManifoldPoint.NetPushOut,	// Out
+			ManifoldPoint.NetImpulse,	// Out
+			Impulse);					// Out
+
+		// Calculate the velocity deltas from the impulse
+		if (Body0.IsDynamic())
+		{
+			const FVec3 AngularImpulse = FVec3::CrossProduct(ManifoldPoint.WorldContactPosition - Body0.P(), Impulse);
+			Body0.ApplyVelocityDelta(Body0.InvM() * Impulse, Body0.InvI() * AngularImpulse);
+		}
+		if (Body1.IsDynamic())
+		{
+			const FVec3 AngularImpulse = FVec3::CrossProduct(ManifoldPoint.WorldContactPosition - Body1.P(), Impulse);
+			Body1.ApplyVelocityDelta(Body1.InvM() * -Impulse, Body1.InvI() * -AngularImpulse);
+		}
+	}
+
+
+	//////////////////////////////////////////////////////////////////////////////////////////////////
+	//////////////////////////////////////////////////////////////////////////////////////////////////
+	//////////////////////////////////////////////////////////////////////////////////////////////////
+	//////////////////////////////////////////////////////////////////////////////////////////////////
+
+
+	FORCEINLINE_DEBUGGABLE void FPBDCollisionSolverManifoldPoint::CalculateContactPositionError(const FConstraintSolverBody& Body0, const FConstraintSolverBody& Body1, const FReal MaxPushOut, FVec3& OutContactDelta, FReal& OutContactDeltaNormal) const
+	{
+		const FVec3 WorldRelativeContactPosition0 = WorldContactPosition - Body0.P();
+		const FVec3 WorldRelativeContactPosition1 = WorldContactPosition - Body1.P();
+#if CHAOS_NONLINEAR_COLLISIONS_ENABLED
+		// Non-linear version: calculate the contact delta after we have converted the current positional impulses into position and rotation corrections.
+		// We could precalculate and store the LocalContactPositions if we really want to use this nonlinear version
+		const FVec3 LocalContactPosition0 = Body0.Q().Inverse() * WorldRelativeContactPosition0;
+		const FVec3 LocalContactPosition1 = Body1.Q().Inverse() * WorldRelativeContactPosition1;
+		OutContactDelta = (Body0.CorrectedP() + Body0.CorrectedQ() * LocalContactPosition0) - (Body1.CorrectedP() + Body1.CorrectedQ() * LocalContactPosition1);
+#else
+		// Linear version: calculate the contact delta assuming linear motion after applying a positional impulse at the contact point. There will be an error that depends on the size of the rotation.
+		const FVec3 ContactDelta0 = Body0.DP() + FVec3::CrossProduct(Body0.DQ(), WorldRelativeContactPosition0);
+		const FVec3 ContactDelta1 = Body1.DP() + FVec3::CrossProduct(Body1.DQ(), WorldRelativeContactPosition1);
+		OutContactDelta = WorldContactDelta + ContactDelta0 - ContactDelta1;
+#endif
+		OutContactDeltaNormal = FVec3::DotProduct(OutContactDelta, WorldContactNormal);
+
+		// NOTE: OutContactDeltaNormal is negative for penetration
+		// NOTE: MaxPushOut == 0 disables the pushout limits
+		if ((MaxPushOut > 0) && (OutContactDeltaNormal < -MaxPushOut))
+		{
+			const FReal ClampedContactDeltaNormal = -MaxPushOut;
+			OutContactDelta += (ClampedContactDeltaNormal - OutContactDeltaNormal) * WorldContactNormal;
+			OutContactDeltaNormal = ClampedContactDeltaNormal;
+		}
+	}
+
+	FORCEINLINE_DEBUGGABLE void FPBDCollisionSolverManifoldPoint::CalculateContactVelocityError(const FConstraintSolverBody& Body0, const FConstraintSolverBody& Body1, const FReal DynamicFriction, const FReal Dt, FVec3& OutContactVelocityDelta, FReal& OutContactVelocityDeltaNormal) const
+	{
+		const FVec3 ContactVelocity0 = Body0.V() + FVec3::CrossProduct(Body0.W(), WorldContactPosition - Body0.P());
+		const FVec3 ContactVelocity1 = Body1.V() + FVec3::CrossProduct(Body1.W(), WorldContactPosition - Body1.P());
+		const FVec3 ContactVelocity = ContactVelocity0 - ContactVelocity1;
+		const FReal ContactVelocityNormal = FVec3::DotProduct(ContactVelocity, WorldContactNormal);
+
+		// Target normal velocity, including restitution
+		const FReal ContactVelocityTargetNormal = WorldContactVelocityTargetNormal;
+
+		// Add up the errors in the velocity (current velocity - desired velocity)
+		OutContactVelocityDeltaNormal = (ContactVelocityNormal - WorldContactVelocityTargetNormal);
+		OutContactVelocityDelta = (ContactVelocityNormal - ContactVelocityTargetNormal) * WorldContactNormal;
+
+		const bool bApplyFriction = (DynamicFriction > 0) && (Dt > 0);
+		if (bApplyFriction)
+		{
+			const FVec3 ContactVelocityTangential = ContactVelocity - ContactVelocityNormal * WorldContactNormal;
+			const FReal ContactVelocityTangentialLen = ContactVelocityTangential.Size();
+			if (ContactVelocityTangentialLen > SMALL_NUMBER)
+			{
+				// PushOut = ContactMass * DP, where DP is the contact positional correction
+				// Friction force is proportional to the normal force, so friction velocity correction
+				// is proprtional to normal velocity correction, or DVn = DPn/dt = PushOut.N / (ContactMass * dt);
+				const FReal PushOutNormal = FVec3::DotProduct(NetPushOut, WorldContactNormal);
+				const FReal DynamicFrictionVelocityError = PushOutNormal / (WorldContactMassNormal * Dt);
+				if (DynamicFrictionVelocityError > SMALL_NUMBER)
+				{
+					const FReal ContactVelocityErrorTangential = FMath::Min(DynamicFriction * DynamicFrictionVelocityError, ContactVelocityTangentialLen);
+					OutContactVelocityDelta += ContactVelocityTangential * (ContactVelocityErrorTangential / ContactVelocityTangentialLen);
+				}
+			}
+		}
+	}
+
+
+	//////////////////////////////////////////////////////////////////////////////////////////////////
+	//////////////////////////////////////////////////////////////////////////////////////////////////
+	//////////////////////////////////////////////////////////////////////////////////////////////////
+	//////////////////////////////////////////////////////////////////////////////////////////////////
+
+
+	FORCEINLINE_DEBUGGABLE bool FPBDCollisionSolver::SolvePositionWithFriction(const FReal Dt, const FReal MaxPushOut)
+	{
+		// SolverBody decorator used to add mass scaling
+		FConstraintSolverBody& Body0 = SolverBody0();
+		FConstraintSolverBody& Body1 = SolverBody1();
+
+		// Check for solved
+		if (State.bIsSolved)
+		{
+			if ((Body0.LastChangeEpoch() == State.BodyEpochs[0]) && (Body1.LastChangeEpoch() == State.BodyEpochs[1]))
+			{
+				// We did not apply a correction last iteration (within tolerance) and nothing else has moved the bodies
+				return false;
+			}
+		}
+
+		// Apply the position correction so that all contacts have zero separation
+		for (int32 PointIndex = 0; PointIndex < State.NumManifoldPoints; ++PointIndex)
+		{
+			FPBDCollisionSolverManifoldPoint& SolverManifoldPoint = State.ManifoldPoints[PointIndex];
+
+			FVec3 ContactDelta;
+			FReal ContactDeltaNormal;
+			SolverManifoldPoint.CalculateContactPositionError(Body0.SolverBody(), Body1.SolverBody(), MaxPushOut, ContactDelta, ContactDeltaNormal);
+
+			const bool bProcessManifoldPoint = (ContactDeltaNormal < FReal(0)) || !SolverManifoldPoint.NetPushOut.IsNearlyZero();
+			if (bProcessManifoldPoint)
+			{
+				ApplyPositionCorrectionWithFriction(
+					State.Stiffness,
+					State.StaticFriction,
+					State.DynamicFriction,
+					ContactDelta,
+					ContactDeltaNormal,
+					SolverManifoldPoint,
+					Body0,
+					Body1);
+			}
+		}
+
+		// We are solved if we did not move the bodies within some tolerance
+		// @todo(chaos): better early-out system
+		State.bIsSolved = (Body0.LastChangeEpoch() == State.BodyEpochs[0]) && (Body1.LastChangeEpoch() == State.BodyEpochs[1]);
+		State.BodyEpochs[0] = Body0.LastChangeEpoch();
+		State.BodyEpochs[1] = Body1.LastChangeEpoch();
+		State.NumPositionSolves = State.NumPositionSolves + 1;
+
+		return !State.bIsSolved;
+	}
+
+	FORCEINLINE_DEBUGGABLE bool FPBDCollisionSolver::SolvePositionNoFriction(const FReal Dt, const FReal MaxPushOut)
+	{
+		// SolverBody decorator used to add mass scaling
+		FConstraintSolverBody& Body0 = SolverBody0();
+		FConstraintSolverBody& Body1 = SolverBody1();
+
+		// Apply the position correction so that all contacts have zero separation
+		for (int32 PointIndex = 0; PointIndex < State.NumManifoldPoints; ++PointIndex)
+		{
+			FPBDCollisionSolverManifoldPoint& SolverManifoldPoint = State.ManifoldPoints[PointIndex];
+
+			FVec3 ContactDelta;
+			FReal ContactDeltaNormal;
+			SolverManifoldPoint.CalculateContactPositionError(Body0.SolverBody(), Body1.SolverBody(), MaxPushOut, ContactDelta, ContactDeltaNormal);
+
+			const bool bProcessManifoldPoint = (ContactDeltaNormal < FReal(0)) || !SolverManifoldPoint.NetPushOut.IsNearlyZero();
+			if (bProcessManifoldPoint)
+			{
+				ApplyPositionCorrectionNoFriction(
+					State.Stiffness,
+					ContactDelta,
+					ContactDeltaNormal,
+					SolverManifoldPoint,
+					Body0,
+					Body1);
+			}
+		}
+
+		// We are solved if we did not move the bodies within some tolerance
+		// NOTE: we can't claim to be solved until we have done at least one friction iteration so we can't early out before friction has been applied.
+		// @todo(chaos): better early-out system
+		State.bIsSolved = false;
+		State.BodyEpochs[0] = Body0.LastChangeEpoch();
+		State.BodyEpochs[1] = Body1.LastChangeEpoch();
+		State.NumPositionSolves = State.NumPositionSolves + 1;
+
+		return !State.bIsSolved;
+	}
+
 }
 
 CHAOS_API DECLARE_LOG_CATEGORY_EXTERN(LogChaosCollision, Log, All);
diff --git a/Engine/Source/Runtime/Experimental/Chaos/Private/Chaos/Collision/SolverCollisionContainer.cpp b/Engine/Source/Runtime/Experimental/Chaos/Private/Chaos/Collision/SolverCollisionContainer.cpp
--- a/Engine/Source/Runtime/Experimental/Chaos/Private/Chaos/Collision/SolverCollisionContainer.cpp
+++ b/Engine/Source/Runtime/Experimental/Chaos/Private/Chaos/Collision/SolverCollisionContainer.cpp
@@ -23,6 +23,9 @@ namespace Chaos
 		extern int32 Chaos_PBDCollisionSolver_Position_ShockPropagationIterations;
 
 		extern int32 Chaos_PBDCollisionSolver_Velocity_ShockPropagationIterations;
+
+		extern bool bChaos_PBDCollisionSolver_Position_SolveEnabled;
+		extern bool bChaos_PBDCollisionSolver_Velocity_SolveEnabled;
 	}
 	using namespace CVars;
 
@@ -161,6 +164,7 @@ namespace Chaos
 	FPBDCollisionSolverContainer::FPBDCollisionSolverContainer()
 		: FConstraintSolverContainer()
 		, MaxPushOutVelocity(0)
+		, bRequiresIncrementalCollisionDetection(false)
 	{
 	}
 
@@ -182,6 +186,8 @@ namespace Chaos
 		const int32 SolverCollisionIndex = CollisionSolvers.AddDefaulted();
 
 		CollisionSolvers[SolverCollisionIndex].GatherInput(Dt, Constraint, Particle0Level, Particle1Level, SolverBodyContainer);
+
+		bRequiresIncrementalCollisionDetection |= CollisionSolvers[SolverCollisionIndex].IsIncrementalManifold();
 	}
 
 	void FPBDCollisionSolverContainer::UpdatePositionShockPropagation(const FReal Dt, const int32 It, const int32 NumIts, const int32 BeginIndex, const int32 EndIndex)
@@ -240,9 +246,14 @@ namespace Chaos
 		return SolveVelocityImpl(Dt, It, NumIts, BeginIndex, EndIndex, true);
 	}
 
+	// @todo(chaos): parallel version of SolvePosition
 	bool FPBDCollisionSolverContainer::SolvePositionImpl(const FReal Dt, const int32 It, const int32 NumIts, const int32 BeginIndex, const int32 EndIndex, const bool bParallel)
 	{
 		SCOPE_CYCLE_COUNTER(STAT_Collisions_Apply);
+		if (!bChaos_PBDCollisionSolver_Position_SolveEnabled)
+		{
+			return false;
+		}
 
 		UpdatePositionShockPropagation(Dt, It, NumIts, BeginIndex, EndIndex);
 
@@ -251,28 +262,74 @@ namespace Chaos
 		const FReal MaxPushOut = (MaxPushOutVelocity > 0) ? (MaxPushOutVelocity * Dt) / FReal(NumIts) : 0;
 
 		// Apply the position correction
-		// @todo(chaos): parallel version of SolvePosition
+		if (bRequiresIncrementalCollisionDetection)
+		{
+			return SolvePositionIncrementalImpl(Dt, BeginIndex, EndIndex, MaxPushOut, bApplyStaticFriction);
+		}
+		else if (bApplyStaticFriction)
+		{
+			return SolvePositionWithFrictionImpl(Dt, BeginIndex, EndIndex, MaxPushOut);
+		}
+		else
+		{
+			return SolvePositionNoFrictionImpl(Dt, BeginIndex, EndIndex, MaxPushOut);
+		}
+	}
+
+	// Solve position including support for incremental collision detection
+	bool FPBDCollisionSolverContainer::SolvePositionIncrementalImpl(const FReal Dt, const int32 BeginIndex, const int32 EndIndex, const FReal MaxPushOut, const bool bApplyStaticFriction)
+	{
 		bool bNeedsAnotherIteration = false;
 		for (int32 SolverIndex = BeginIndex; SolverIndex < EndIndex; ++SolverIndex)
 		{
 			FPBDCollisionSolverAdapter& CollisionSolver = CollisionSolvers[SolverIndex];
-
-			// @todo(chaos): remove this when we no longer need to support incremental manifolds - this prevents us from fully optimizing the solver loop
 			if (CollisionSolver.IsIncrementalManifold())
 			{
 				Collisions::Update(*CollisionSolver.GetConstraint(), Dt);
 				CollisionSolver.GatherManifoldPoints(Dt);
 			}
+			if (bApplyStaticFriction)
+			{
+				bNeedsAnotherIteration |= CollisionSolver.GetSolver().SolvePositionWithFriction(Dt, MaxPushOut);
+			}
+			else
+			{
+				bNeedsAnotherIteration |= CollisionSolver.GetSolver().SolvePositionNoFriction(Dt, MaxPushOut);
+			}
 
-			bNeedsAnotherIteration |= CollisionSolvers[SolverIndex].GetSolver().SolvePosition(Dt, MaxPushOut, bApplyStaticFriction);
 		}
+		return bNeedsAnotherIteration;
+	}
 
+	// Solve position with friction (last few iterations each tick)
+	bool FPBDCollisionSolverContainer::SolvePositionWithFrictionImpl(const FReal Dt, const int32 BeginIndex, const int32 EndIndex, const FReal MaxPushOut)
+	{
+		bool bNeedsAnotherIteration = false;
+		for (int32 SolverIndex = BeginIndex; SolverIndex < EndIndex; ++SolverIndex)
+		{
+			bNeedsAnotherIteration |= CollisionSolvers[SolverIndex].GetSolver().SolvePositionWithFriction(Dt, MaxPushOut);
+		}
+		return bNeedsAnotherIteration;
+	}
+
+	// Solve position without friction (first few iterations each tick)
+	bool FPBDCollisionSolverContainer::SolvePositionNoFrictionImpl(const FReal Dt, const int32 BeginIndex, const int32 EndIndex, const FReal MaxPushOut)
+	{
+		bool bNeedsAnotherIteration = false;
+		for (int32 SolverIndex = BeginIndex; SolverIndex < EndIndex; ++SolverIndex)
+		{
+			bNeedsAnotherIteration |= CollisionSolvers[SolverIndex].GetSolver().SolvePositionNoFriction(Dt, MaxPushOut);
+		}
 		return bNeedsAnotherIteration;
 	}
 
 	bool FPBDCollisionSolverContainer::SolveVelocityImpl(const FReal Dt, const int32 It, const int32 NumIts, const int32 BeginIndex, const int32 EndIndex, const bool bParallel)
 	{
 		SCOPE_CYCLE_COUNTER(STAT_Collisions_ApplyPushOut);
+		if (!bChaos_PBDCollisionSolver_Velocity_SolveEnabled)
+		{
+			return false;
+		}
 
 		UpdateVelocityShockPropagation(Dt, It, NumIts, BeginIndex, EndIndex);
 
diff --git a/Engine/Source/Runtime/Experimental/Chaos/Public/Chaos/Collision/SolverCollisionContainer.h b/Engine/Source/Runtime/Experimental/Chaos/Public/Chaos/Collision/SolverCollisionContainer.h
--- a/Engine/Source/Runtime/Experimental/Chaos/Public/Chaos/Collision/SolverCollisionContainer.h
+++ b/Engine/Source/Runtime/Experimental/Chaos/Public/Chaos/Collision/SolverCollisionContainer.h
@@ -50,10 +50,14 @@ namespace Chaos
 		void UpdatePositionShockPropagation(const FReal Dt, const int32 It, const int32 NumIts, const int32 BeginIndex, const int32 EndIndex);
 		void UpdateVelocityShockPropagation(const FReal Dt, const int32 It, const int32 NumIts, const int32 BeginIndex, const int32 EndIndex);
 		bool SolvePositionImpl(const FReal Dt, const int32 It, const int32 NumIts, const int32 BeginIndex, const int32 EndIndex, const bool bParallel);
+		bool SolvePositionIncrementalImpl(const FReal Dt, const int32 BeginIndex, const int32 EndIndex, const FReal MaxPushOut, const bool bApplyStaticFriction);
+		bool SolvePositionWithFrictionImpl(const FReal Dt, const int32 BeginIndex, const int32 EndIndex, const FReal MaxPushOut);
+		bool SolvePositionNoFrictionImpl(const FReal Dt, const int32 BeginIndex, const int32 EndIndex, const FReal MaxPushOut);
 		bool SolveVelocityImpl(const FReal Dt, const int32 It, const int32 NumIts, const int32 BeginIndex, const int32 EndIndex, const bool bParallel);
 		void ScatterOutputImpl(const FReal Dt, const int32 BeginIndex, const int32 EndIndex, const bool bParallel);
 
 		TArray<FPBDCollisionSolverAdapter> CollisionSolvers;
 		FReal MaxPushOutVelocity;
+		bool bRequiresIncrementalCollisionDetection;
 	};
 }
