Chaos - fix impulse in static frinction path when net impulse goes negative

diff --git a/Engine/Source/Runtime/Experimental/Chaos/Private/Chaos/Collision/PBDCollisionSolver.h b/Engine/Source/Runtime/Experimental/Chaos/Private/Chaos/Collision/PBDCollisionSolver.h
--- a/Engine/Source/Runtime/Experimental/Chaos/Private/Chaos/Collision/PBDCollisionSolver.h
+++ b/Engine/Source/Runtime/Experimental/Chaos/Private/Chaos/Collision/PBDCollisionSolver.h
@@ -319,14 +319,15 @@ namespace Chaos
 
 		// The total pushout so far this sub-step
 		// We allow negative incremental impulses, but not net negative impulses
-		const FVec3 NetPushOut = InOutNetPushOut + PushOut;
+		FVec3 NetPushOut = InOutNetPushOut + PushOut;
 		const FReal NetPushOutNormal = FVec3::DotProduct(NetPushOut, ContactNormal);
 		if (NetPushOutNormal < 0)
 		{
 			PushOut = -InOutNetPushOut;
+			NetPushOut = FVec3(0);
 		}
 
-		InOutNetPushOut += PushOut;
+		InOutNetPushOut = NetPushOut;
 		OutPushOut = PushOut;
 	}
 
@@ -357,13 +358,15 @@ namespace Chaos
 
 		FVec3 PushOut = Stiffness * ContactMass * ModifiedContactError;
 
-		// If we ended up with a negative normal pushout, disable friction
+		// If we ended up with a negative normal pushout, remove all correction from this point
 		FVec3 NetPushOut = InOutNetPushOut + PushOut;
 		const FReal NetPushOutNormal = FVec3::DotProduct(NetPushOut, ContactNormal);
 		bool bInsideStaticFrictionCone = true;
 		if (NetPushOutNormal < FReal(SMALL_NUMBER))
 		{
 			bInsideStaticFrictionCone = false;
+			PushOut = -InOutNetPushOut;
+			NetPushOut = FVec3(0);
 		}
 
 		// Static friction limit: immediately increase maximum lateral correction, but smoothly decay maximum static friction limit. 
@@ -373,8 +376,7 @@ namespace Chaos
 		const FReal StaticFrictionDest = FMath::Max(NetPushOutNormal, FReal(0));
 		FReal StaticFrictionMax = FMath::Lerp(FMath::Max(InOutStaticFrictionMax, StaticFrictionDest), StaticFrictionDest, StaticFrictionLerpRate);
 
-		// If we exceed the friction cone, stop adding frictional corrections (although
-		// any already added lateral corrections will not be undone)
+		// If we exceed the friction cone, stop adding frictional corrections and clip correction to cone
 		// @todo(chaos): clamp to dynamic friction
 		if (bInsideStaticFrictionCone)
 		{
@@ -389,18 +391,8 @@ namespace Chaos
 			}
 		}
 
-		// If we leave the friction cone, we will fall through into the non-friction impulse calculation
-		// so do not export the results
-		if (bInsideStaticFrictionCone)
-		{
-			InOutNetPushOut = NetPushOut;
-			OutPushOut = PushOut;
-		}
-		else
-		{
-			OutPushOut = FVec3(0);
-		}
-
+		InOutNetPushOut = NetPushOut;
+		OutPushOut = PushOut;
 		InOutStaticFrictionMax = StaticFrictionMax;
 		bOutInsideStaticFrictionCone = bInsideStaticFrictionCone;
 	}
