Chaos - Constraint solver can now use floats even if doubles are the default number type (see CHAOS_CONSTRAINTSOLVER_LOWPRECISION)
- Remove some unnecessary data from collision solver
- Reorder some collision solver data

diff --git a/Engine/Source/Programs/HeadlessChaos/Private/HeadlessChaosTestRotation.cpp b/Engine/Source/Programs/HeadlessChaos/Private/HeadlessChaosTestRotation.cpp
--- a/Engine/Source/Programs/HeadlessChaos/Private/HeadlessChaosTestRotation.cpp
+++ b/Engine/Source/Programs/HeadlessChaos/Private/HeadlessChaosTestRotation.cpp
@@ -288,7 +288,7 @@ namespace ChaosTest {
 			EXPECT_NEAR(MatrixCrossP.Y, V3.Y, Tolerance);
 			EXPECT_NEAR(MatrixCrossP.Z, V3.Z, Tolerance);
 
-			FVec3 MatrixCrossN = Utilities::Multiply(V1M.GetTransposed(), V2);
+			FVec3 MatrixCrossN = Utilities::Multiply(FMatrix33(V1M.GetTransposed()), V2);
 			EXPECT_NEAR(MatrixCrossN.X, -V3.X, Tolerance);
 			EXPECT_NEAR(MatrixCrossN.Y, -V3.Y, Tolerance);
 			EXPECT_NEAR(MatrixCrossN.Z, -V3.Z, Tolerance);
diff --git a/Engine/Source/Runtime/Experimental/Chaos/Private/Chaos/ChaosDebugDraw.cpp b/Engine/Source/Runtime/Experimental/Chaos/Private/Chaos/ChaosDebugDraw.cpp
--- a/Engine/Source/Runtime/Experimental/Chaos/Private/Chaos/ChaosDebugDraw.cpp
+++ b/Engine/Source/Runtime/Experimental/Chaos/Private/Chaos/ChaosDebugDraw.cpp
@@ -923,15 +923,6 @@ namespace Chaos
 							const FReal BoxScale = Settings.DrawScale * Settings.ContactWidth;
 							FDebugDrawQueue::GetInstance().DrawDebugBox(WorldPlaneLocation, FVec3(BoxScale, BoxScale, FReal(0.01)), FRotation3(FRotationMatrix::MakeFromZ(WorldPlaneNormal)), FColor::Purple, false, KINDA_SMALL_NUMBER, Settings.DrawPriority, 0.5f * Settings.LineThickness);
 						}
-
-						// How many iterations (0 = green, 8 = red)
-						if (bChaosDebugDebugDrawContactIterations)
-						{
-							const FReal BoxScale = Settings.DrawScale * Settings.ContactWidth * FReal(1.5);
-							const float IterationColorScale = FMath::Clamp(float(Contact.GetNumActivePositionIterations()) / float(8), float(0), float(1));
-							const FColor IterationColor = FLinearColor::LerpUsingHSV(FColor::Green, FColor::Red, IterationColorScale).ToFColor(false);
-							FDebugDrawQueue::GetInstance().DrawDebugBox(WorldPlaneLocation, FVec3(BoxScale, BoxScale, FReal(0.01)), FRotation3(FRotationMatrix::MakeFromZ(WorldPlaneNormal)), IterationColor, false, KINDA_SMALL_NUMBER, Settings.DrawPriority, 0.5f * Settings.LineThickness);
-						}
 					}
 				}
 				else
diff --git a/Engine/Source/Runtime/Experimental/Chaos/Private/Chaos/Collision/PBDCollisionConstraint.cpp b/Engine/Source/Runtime/Experimental/Chaos/Private/Chaos/Collision/PBDCollisionConstraint.cpp
--- a/Engine/Source/Runtime/Experimental/Chaos/Private/Chaos/Collision/PBDCollisionConstraint.cpp
+++ b/Engine/Source/Runtime/Experimental/Chaos/Private/Chaos/Collision/PBDCollisionConstraint.cpp
@@ -158,7 +158,6 @@ namespace Chaos
 		, LastShapeWorldTransform0()
 		, LastShapeWorldTransform1()
 		, ExpectedNumManifoldPoints(0)
-		, NumActivePositionIterations(0)
 	{
 		Manifold.Implicit[0] = nullptr;
 		Manifold.Implicit[1] = nullptr;
@@ -195,7 +194,6 @@ namespace Chaos
 		, LastShapeWorldTransform0()
 		, LastShapeWorldTransform1()
 		, ExpectedNumManifoldPoints(0)
-		, NumActivePositionIterations(0)
 	{
 		Manifold.Implicit[0] = Implicit0;
 		Manifold.Implicit[1] = Implicit1;
diff --git a/Engine/Source/Runtime/Experimental/Chaos/Private/Chaos/Collision/PBDCollisionSolver.cpp b/Engine/Source/Runtime/Experimental/Chaos/Private/Chaos/Collision/PBDCollisionSolver.cpp
--- a/Engine/Source/Runtime/Experimental/Chaos/Private/Chaos/Collision/PBDCollisionSolver.cpp
+++ b/Engine/Source/Runtime/Experimental/Chaos/Private/Chaos/Collision/PBDCollisionSolver.cpp
@@ -65,10 +65,10 @@ namespace Chaos
 		const FConstraintSolverBody& Body1,
 		const FVec3& InWorldAnchorPoint0,
 		const FVec3& InWorldAnchorPoint1,
-		const FVec3& InWorldContactNormal)
+		const FVec3LP& InWorldContactNormal)
 	{
-		NetPushOut = FVec3(0);
-		NetImpulse = FVec3(0);
+		NetPushOut = FVec3LP(0);
+		NetImpulse = FVec3LP(0);
 		UpdateContact(Body0, Body1, InWorldAnchorPoint0, InWorldAnchorPoint1, InWorldContactNormal);
 		UpdateMass(Body0, Body1);
 	}
@@ -76,19 +76,19 @@ namespace Chaos
 	void FPBDCollisionSolverManifoldPoint::InitMaterial(
 		const FConstraintSolverBody& Body0,
 		const FConstraintSolverBody& Body1,
-		const FReal InRestitution,
-		const FReal InRestitutionVelocityThreshold,
+		const FRealLP InRestitution,
+		const FRealLP InRestitutionVelocityThreshold,
 		const bool bInEnableStaticFriction,
-		const FReal InStaticFrictionMax)
+		const FRealLP InStaticFrictionMax)
 	{
 		StaticFrictionMax = InStaticFrictionMax;
 		bInsideStaticFrictionCone = bInEnableStaticFriction;
-		WorldContactVelocityTargetNormal = FReal(0);
+		WorldContactVelocityTargetNormal = FRealLP(0);
 
-		if (InRestitution > FReal(0))
+		if (InRestitution > FRealLP(0))
 		{
-			const FVec3 ContactVelocity = CalculateContactVelocity(Body0, Body1);
-			const FReal ContactVelocityNormal = FVec3::DotProduct(ContactVelocity, WorldContactNormal);
+			const FVec3LP ContactVelocity = CalculateContactVelocity(Body0, Body1);
+			const FRealLP ContactVelocityNormal = FVec3LP::DotProduct(ContactVelocity, WorldContactNormal);
 			if (ContactVelocityNormal < -InRestitutionVelocityThreshold)
 			{
 				WorldContactVelocityTargetNormal = -InRestitution * ContactVelocityNormal;
@@ -102,10 +102,12 @@ namespace Chaos
 		const FConstraintSolverBody& Body1,
 		const FVec3& InWorldAnchorPoint0,
 		const FVec3& InWorldAnchorPoint1, 
-		const FVec3& InWorldContactNormal)
+		const FVec3LP& InWorldContactNormal)
 	{
 		// The world-space point where we apply impulses/corrections (same world-space point for momentum conservation)
-		WorldContactPosition = FReal(0.5) * (InWorldAnchorPoint0 + InWorldAnchorPoint1);
+		const FVec3 WorldContactPosition = FReal(0.5) * (InWorldAnchorPoint0 + InWorldAnchorPoint1);
+		RelativeContactPosition0 = FVec3LP(WorldContactPosition - Body0.P());
+		RelativeContactPosition1 = FVec3LP(WorldContactPosition - Body1.P());
 
 		// The world-space contact normal
 		WorldContactNormal = InWorldContactNormal;
@@ -116,20 +118,31 @@ namespace Chaos
 
 	inline void FPBDCollisionSolverManifoldPoint::UpdateMass(const FConstraintSolverBody& Body0, const FConstraintSolverBody& Body1)
 	{
-		const FMatrix33 ContactMassInv =
-			(Body0.IsDynamic() ? Collisions::ComputeFactorMatrix3(WorldContactPosition - Body0.P(), Body0.InvI(), Body0.InvM()) : FMatrix33(0)) +
-			(Body1.IsDynamic() ? Collisions::ComputeFactorMatrix3(WorldContactPosition - Body1.P(), Body1.InvI(), Body1.InvM()) : FMatrix33(0));
-		const FMatrix33 ContactMass = ContactMassInv.Inverse();
-		const FReal ContactMassInvNormal = FVec3::DotProduct(WorldContactNormal, Utilities::Multiply(ContactMassInv, WorldContactNormal));
-		const FReal ContactMassNormal = (ContactMassInvNormal > FReal(SMALL_NUMBER)) ? FReal(1) / ContactMassInvNormal : FReal(0);
+		const FMatrix33LP ContactMassInv =
+			(Body0.IsDynamic() ? Collisions::ComputeFactorMatrix3(RelativeContactPosition0, FMatrix33LP(Body0.InvI()), FRealLP(Body0.InvM())) : FMatrix33LP(0)) +
+			(Body1.IsDynamic() ? Collisions::ComputeFactorMatrix3(RelativeContactPosition1, FMatrix33LP(Body1.InvI()), FRealLP(Body1.InvM())) : FMatrix33LP(0));
+		const FMatrix33LP ContactMass = ContactMassInv.Inverse();
+		const FRealLP ContactMassInvNormal = FVec3LP::DotProduct(WorldContactNormal, Utilities::Multiply(ContactMassInv, WorldContactNormal));
+		const FRealLP ContactMassNormal = (ContactMassInvNormal > FRealLP(SMALL_NUMBER)) ? FRealLP(1) / ContactMassInvNormal : FRealLP(0);
 		WorldContactMass = ContactMass;
 		WorldContactMassNormal = ContactMassNormal;
+
+		WorldContactNormalAngular0 = FVec3LP(0);
+		WorldContactNormalAngular1 = FVec3LP(0);
+		if (Body0.IsDynamic())
+		{
+			WorldContactNormalAngular0 = Body0.InvI() * FVec3LP::CrossProduct(RelativeContactPosition0, WorldContactNormal);
+		}
+		if (Body1.IsDynamic())
+		{
+			WorldContactNormalAngular1 = Body1.InvI() * FVec3LP::CrossProduct(RelativeContactPosition1, WorldContactNormal);
+		}
 	}
 
 	FVec3 FPBDCollisionSolverManifoldPoint::CalculateContactVelocity(const FConstraintSolverBody& Body0, const FConstraintSolverBody& Body1) const
 	{
-		const FVec3 ContactVelocity0 = Body0.V() + FVec3::CrossProduct(Body0.W(), WorldContactPosition - Body0.P());
-		const FVec3 ContactVelocity1 = Body1.V() + FVec3::CrossProduct(Body1.W(), WorldContactPosition - Body1.P());
+		const FVec3LP ContactVelocity0 = Body0.V() + FVec3::CrossProduct(Body0.W(), RelativeContactPosition0);
+		const FVec3LP ContactVelocity1 = Body1.V() + FVec3::CrossProduct(Body1.W(), RelativeContactPosition1);
 		return ContactVelocity0 - ContactVelocity1;
 	}
 
@@ -161,7 +174,7 @@ namespace Chaos
 		SetShockPropagationInvMassScale(FReal(1));
 	}
 
-	void FPBDCollisionSolver::SetShockPropagationInvMassScale(const FReal InvMassScale)
+	void FPBDCollisionSolver::SetShockPropagationInvMassScale(const FRealLP InvMassScale)
 	{
 		FConstraintSolverBody& Body0 = SolverBody0();
 		FConstraintSolverBody& Body1 = SolverBody1();
@@ -208,7 +221,12 @@ namespace Chaos
 		const FVec3& InWorldAnchorPoint1,
 		const FVec3& InWorldContactNormal)
 	{
-		State.ManifoldPoints[ManifoldPoiontIndex].InitContact(State.SolverBodies[0], State.SolverBodies[1], InWorldAnchorPoint0, InWorldAnchorPoint1, InWorldContactNormal);
+		State.ManifoldPoints[ManifoldPoiontIndex].InitContact(
+			State.SolverBodies[0], 
+			State.SolverBodies[1], 
+			InWorldAnchorPoint0, 
+			InWorldAnchorPoint1, 
+			FVec3LP(InWorldContactNormal));
 	}
 
 	void FPBDCollisionSolver::InitMaterial(
@@ -218,7 +236,13 @@ namespace Chaos
 		const bool bInEnableStaticFriction,
 		const FReal InStaticFrictionMax)
 	{
-		State.ManifoldPoints[ManifoldPoiontIndex].InitMaterial(State.SolverBodies[0], State.SolverBodies[1], InRestitution, InRestitutionVelocityThreshold, bInEnableStaticFriction, InStaticFrictionMax);
+		State.ManifoldPoints[ManifoldPoiontIndex].InitMaterial(
+			State.SolverBodies[0], 
+			State.SolverBodies[1], 
+			FRealLP(InRestitution), 
+			FRealLP(InRestitutionVelocityThreshold),
+			bInEnableStaticFriction, 
+			FRealLP(InStaticFrictionMax));
 
 		// Track if any points have restitution enabled. See SolveVelocity
 		State.bHaveRestitution = State.bHaveRestitution || (State.ManifoldPoints[ManifoldPoiontIndex].WorldContactVelocityTargetNormal > 0);
@@ -230,27 +254,34 @@ namespace Chaos
 		const FVec3& InWorldAnchorPoint1,
 		const FVec3& InWorldContactNormal)
 	{
-		State.ManifoldPoints[ManifoldPoiontIndex].UpdateContact(State.SolverBodies[0], State.SolverBodies[1], InWorldAnchorPoint0, InWorldAnchorPoint1, InWorldContactNormal);
+		State.ManifoldPoints[ManifoldPoiontIndex].UpdateContact(
+			State.SolverBodies[0], 
+			State.SolverBodies[1], 
+			InWorldAnchorPoint0, 
+			InWorldAnchorPoint1, 
+			FVec3LP(InWorldContactNormal));
 	}
 
 
-	void FPBDCollisionSolver::SolveVelocityAverage(const FReal Dt)
+	void FPBDCollisionSolver::SolveVelocityAverage(const FRealLP Dt)
 	{
 		FConstraintSolverBody& Body0 = SolverBody0();
 		FConstraintSolverBody& Body1 = SolverBody1();
 
 		// Generate a new contact point at the average of all the active contacts
 		int32 NumActiveManifoldPoints = 0;
-		FVec3 WorldContactPosition = FVec3(0);
-		FVec3 WorldContactNormal = FVec3(0);
-		FVec3 NetPushOut = FVec3(0);
-		FReal WorldContactVelocityTargetNormal = FReal(0);
+		FVec3LP RelativeContactPosition0 = FVec3(0);
+		FVec3LP RelativeContactPosition1 = FVec3(0);
+		FVec3LP WorldContactNormal = FVec3LP(0);
+		FVec3LP NetPushOut = FVec3LP(0);
+		FRealLP WorldContactVelocityTargetNormal = FReal(0);
 		for (int32 PointIndex = 0; PointIndex < State.NumManifoldPoints; ++PointIndex)
 		{
 			FPBDCollisionSolverManifoldPoint& SolverManifoldPoint = State.ManifoldPoints[PointIndex];
 			if (!SolverManifoldPoint.NetPushOut.IsNearlyZero())
 			{
-				WorldContactPosition += SolverManifoldPoint.WorldContactPosition;
+				RelativeContactPosition0 += SolverManifoldPoint.RelativeContactPosition0;
+				RelativeContactPosition1 += SolverManifoldPoint.RelativeContactPosition1;
 				WorldContactVelocityTargetNormal += SolverManifoldPoint.WorldContactVelocityTargetNormal;
 				WorldContactNormal = SolverManifoldPoint.WorldContactNormal;	// Take last value - should all be similar
 				++NumActiveManifoldPoints;
@@ -261,30 +292,31 @@ namespace Chaos
 		// We only do this if we have multiple active contacts
 		if (NumActiveManifoldPoints > 1)
 		{
-			const FReal DynamicFriction = FReal(0);
-			const FReal Scale = FReal(1) / FReal(NumActiveManifoldPoints);
+			const FRealLP DynamicFriction = FRealLP(0);
+			const FRealLP Scale = FRealLP(1) / FRealLP(NumActiveManifoldPoints);
 
 			FPBDCollisionSolverManifoldPoint AverageManifoldPoint;
-			AverageManifoldPoint.WorldContactPosition = WorldContactPosition * Scale;
+			AverageManifoldPoint.RelativeContactPosition0 = RelativeContactPosition0 * Scale;
+			AverageManifoldPoint.RelativeContactPosition1 = RelativeContactPosition1 * Scale;
 			AverageManifoldPoint.WorldContactNormal = WorldContactNormal;
+			AverageManifoldPoint.WorldContactVelocityTargetNormal = WorldContactVelocityTargetNormal * Scale;
 
-			const FMatrix33 ContactMassInv =
-				(Body0.IsDynamic() ? Collisions::ComputeFactorMatrix3(AverageManifoldPoint.WorldContactPosition - Body0.P(), Body0.InvI(), Body0.InvM()) : FMatrix33(0)) +
-				(Body1.IsDynamic() ? Collisions::ComputeFactorMatrix3(AverageManifoldPoint.WorldContactPosition - Body1.P(), Body1.InvI(), Body1.InvM()) : FMatrix33(0));
-			const FReal ContactMassInvNormal = FVec3::DotProduct(AverageManifoldPoint.WorldContactNormal, Utilities::Multiply(ContactMassInv, AverageManifoldPoint.WorldContactNormal));
-			const FReal ContactMassNormal = (ContactMassInvNormal > FReal(SMALL_NUMBER)) ? FReal(1) / ContactMassInvNormal : FReal(0);
-			AverageManifoldPoint.WorldContactMass = FMatrix33(0);
+			const FMatrix33LP ContactMassInv =
+				(Body0.IsDynamic() ? Collisions::ComputeFactorMatrix3(AverageManifoldPoint.RelativeContactPosition0, Body0.InvI(), Body0.InvM()) : FMatrix33LP(0)) +
+				(Body1.IsDynamic() ? Collisions::ComputeFactorMatrix3(AverageManifoldPoint.RelativeContactPosition1, Body1.InvI(), Body1.InvM()) : FMatrix33LP(0));
+			const FRealLP ContactMassInvNormal = FVec3LP::DotProduct(AverageManifoldPoint.WorldContactNormal, Utilities::Multiply(ContactMassInv, AverageManifoldPoint.WorldContactNormal));
+			const FRealLP ContactMassNormal = (ContactMassInvNormal > FRealLP(SMALL_NUMBER)) ? FRealLP(1) / ContactMassInvNormal : FRealLP(0);
 			AverageManifoldPoint.WorldContactMassNormal = ContactMassNormal;
 
-			AverageManifoldPoint.WorldContactDelta = FVec3(0);
-			AverageManifoldPoint.WorldContactVelocityTargetNormal = WorldContactVelocityTargetNormal * Scale;
-			AverageManifoldPoint.NetPushOut = FVec3(0);
-			AverageManifoldPoint.NetImpulse = FVec3(0);
-			AverageManifoldPoint.StaticFrictionMax = FReal(0);
-			AverageManifoldPoint.bInsideStaticFrictionCone = false;
+			AverageManifoldPoint.WorldContactMass = FMatrix33LP(0);	// Unused
+			AverageManifoldPoint.WorldContactDelta = FVec3LP(0);	// Unused
+			AverageManifoldPoint.NetPushOut = FVec3LP(0);			// Unused
+			AverageManifoldPoint.NetImpulse = FVec3LP(0);			// Unused
+			AverageManifoldPoint.StaticFrictionMax = FRealLP(0);	// Unused
+			AverageManifoldPoint.bInsideStaticFrictionCone = false;	// Unused
 
-			FVec3 ContactVelocityDelta;
-			FReal ContactVelocityDeltaNormal;
+			FVec3LP ContactVelocityDelta;
+			FRealLP ContactVelocityDeltaNormal;
 			AverageManifoldPoint.CalculateContactVelocityError(Body0, Body1, DynamicFriction, Dt, ContactVelocityDelta, ContactVelocityDeltaNormal);
 
 			ApplyVelocityCorrection(
@@ -299,7 +331,7 @@ namespace Chaos
 		}
 	}
 
-	bool FPBDCollisionSolver::SolveVelocity(const FReal Dt, const bool bApplyDynamicFriction)
+	bool FPBDCollisionSolver::SolveVelocity(const FRealLP Dt, const bool bApplyDynamicFriction)
 	{
 		// Apply restitution at the average contact point
 		// This means we don't need to run as many iterations to get stable bouncing
@@ -313,15 +345,15 @@ namespace Chaos
 
 		// NOTE: this dynamic friction implementation is iteration-count sensitive
 		// @todo(chaos): fix iteration count dependence of dynamic friction
-		const FReal DynamicFriction = (bApplyDynamicFriction && bChaos_PBDCollisionSolver_Velocity_DynamicFrictionEnabled) ? State.DynamicFriction : (FReal)0.0f;
+		const FRealLP DynamicFriction = (bApplyDynamicFriction && bChaos_PBDCollisionSolver_Velocity_DynamicFrictionEnabled) ? State.DynamicFriction : FRealLP(0);
 
 		for (int32 PointIndex = 0; PointIndex < State.NumManifoldPoints; ++PointIndex)
 		{
 			FPBDCollisionSolverManifoldPoint& SolverManifoldPoint = State.ManifoldPoints[PointIndex];
 			if (!SolverManifoldPoint.NetPushOut.IsNearlyZero())
 			{
-				FVec3 ContactVelocityDelta;
-				FReal ContactVelocityDeltaNormal;
+				FVec3LP ContactVelocityDelta;
+				FRealLP ContactVelocityDeltaNormal;
 				SolverManifoldPoint.CalculateContactVelocityError(Body0, Body1, DynamicFriction, Dt, ContactVelocityDelta, ContactVelocityDeltaNormal);
 				ApplyVelocityCorrection(
 					State.Stiffness,
@@ -335,8 +367,6 @@ namespace Chaos
 			}
 		}
 
-		State.NumVelocitySolves = State.NumVelocitySolves + 1;
-
 		// Early-out support for the velocity solve is not currently very important because we
 		// only run one iteration in the velocity solve phase.
 		// @todo(chaos): support early-out in velocity solve if necessary
diff --git a/Engine/Source/Runtime/Experimental/Chaos/Private/Chaos/Collision/PBDCollisionSolver.h b/Engine/Source/Runtime/Experimental/Chaos/Private/Chaos/Collision/PBDCollisionSolver.h
--- a/Engine/Source/Runtime/Experimental/Chaos/Private/Chaos/Collision/PBDCollisionSolver.h
+++ b/Engine/Source/Runtime/Experimental/Chaos/Private/Chaos/Collision/PBDCollisionSolver.h
@@ -34,6 +34,7 @@ namespace Chaos
 		extern bool bChaos_PBDCollisionSolver_Velocity_ImpulseClampEnabled;
 	}
 
+
 	/**
 	 * @brief A single contact point in a FPBDCollisionSolver
 	*/
@@ -48,7 +49,7 @@ namespace Chaos
 			const FConstraintSolverBody& Body1,
 			const FVec3& InWorldAnchorPoint0,
 			const FVec3& InWorldAnchorPoint1,
-			const FVec3& InWorldContactNormal);
+			const FVec3LP& InWorldContactNormal);
 
 		/**
 		 * @brief Initialize the material related properties of the contact
@@ -56,10 +57,10 @@ namespace Chaos
 		void InitMaterial(
 			const FConstraintSolverBody& Body0,
 			const FConstraintSolverBody& Body1,
-			const FReal InRestitution,
-			const FReal InRestitutionVelocityThreshold,
+			const FRealLP InRestitution,
+			const FRealLP InRestitutionVelocityThreshold,
 			const bool bInEnableStaticFriction,
-			const FReal InStaticFrictionMax);
+			const FRealLP InStaticFrictionMax);
 
 		/**
 		 * @brief Update the world-space relative contact points based on current body transforms and body-space contact positions
@@ -69,7 +70,7 @@ namespace Chaos
 			const FConstraintSolverBody& Body1,
 			const FVec3& InWorldAnchorPoint0,
 			const FVec3& InWorldAnchorPoint1,
-			const FVec3& InWorldContactNormal);
+			const FVec3LP& InWorldContactNormal);
 
 		/**
 		 * @brief Update the cached mass properties based on the current body transforms
@@ -86,39 +87,46 @@ namespace Chaos
 		 * @brief Calculate the position error at the current transforms
 		 * @param MaxPushOut a limit on the position error for this iteration to prevent initial-penetration explosion (a common PBD problem)
 		*/
-		void CalculateContactPositionError(const FConstraintSolverBody& Body0, const FConstraintSolverBody& Body1, const FReal MaxPushOut, FVec3& OutContactDelta, FReal& OutContactDeltaNormal) const;
+		void CalculateContactPositionError(const FConstraintSolverBody& Body0, const FConstraintSolverBody& Body1, const FRealLP MaxPushOut, FVec3LP& OutContactDelta, FRealLP& OutContactDeltaNormal) const;
 
 		/**
 		 * @brief Calculate the velocity error at the current transforms
 		*/
-		void CalculateContactVelocityError(const FConstraintSolverBody& Body0, const FConstraintSolverBody& Body1, const FReal DynamicFriction, const FReal Dt, FVec3& OutContactVelocityDelta, FReal& OutContactVelocityDeltaNormal) const;
+		void CalculateContactVelocityError(const FConstraintSolverBody& Body0, const FConstraintSolverBody& Body1, const FRealLP DynamicFriction, const FRealLP Dt, FVec3LP& OutContactVelocityDelta, FRealLP& OutContactVelocityDeltaNormal) const;
 
 		// @todo(chaos): make private
 	public:
 		friend class FPBDCollisionSolver;
 
-		// World-space contact point
-		FVec3 WorldContactPosition;
+		// Contact mass (for friction)
+		FMatrix33LP WorldContactMass;
+
+		// World-space body-relative contact points
+		FVec3LP RelativeContactPosition0;
+		FVec3LP RelativeContactPosition1;
 
 		// World-space contact normal
-		FVec3 WorldContactNormal;
+		FVec3LP WorldContactNormal;
 
-		// Contact mass
-		FMatrix33 WorldContactMass;
-		FReal WorldContactMassNormal;
+		// I^-1.(R x N) for each body
+		FVec3LP WorldContactNormalAngular0;
+		FVec3LP WorldContactNormalAngular1;
 
-		// World-space contact separation that we are trying to correct
-		FVec3 WorldContactDelta;
+		// Contact mass (for non-friction)
+		FRealLP WorldContactMassNormal;
+
+		// Initial world-space contact separation that we are trying to correct
+		FVec3LP WorldContactDelta;
 
 		// Desired final normal velocity, taking Restitution into account
-		FReal WorldContactVelocityTargetNormal;
+		FRealLP WorldContactVelocityTargetNormal;
 
 		// Solver outputs
-		FVec3 NetPushOut;
-		FVec3 NetImpulse;
+		FVec3LP NetPushOut;
+		FVec3LP NetImpulse;
 
 		// A smoothed NetImpulse along the normal, used for clipping to the static friction cone
-		FReal StaticFrictionMax;
+		FRealLP StaticFrictionMax;
 
 		// Whether we are still in the static friction cone
 		bool bInsideStaticFrictionCone;
@@ -138,18 +146,18 @@ namespace Chaos
 
 		FPBDCollisionSolver();
 
-		FReal StaticFriction() const { return State.StaticFriction; }
-		FReal DynamicFriction() const { return State.DynamicFriction; }
+		FRealLP StaticFriction() const { return State.StaticFriction; }
+		FRealLP DynamicFriction() const { return State.DynamicFriction; }
 
 		void SetFriction(const FReal InStaticFriction, const FReal InDynamicFriction)
 		{
-			State.StaticFriction = InStaticFriction;
-			State.DynamicFriction = InDynamicFriction;
+			State.StaticFriction = FRealLP(InStaticFriction);
+			State.DynamicFriction = FRealLP(InDynamicFriction);
 		}
 
 		void SetStiffness(const FReal InStiffness)
 		{
-			State.Stiffness = InStiffness;
+			State.Stiffness = FRealLP(InStiffness);
 		}
 
 		void SetSolverBodies(FSolverBody* SolverBody0, FSolverBody* SolverBody1)
@@ -181,16 +189,6 @@ namespace Chaos
 			return State.ManifoldPoints[ManifoldPointIndex];
 		}
 
-		int32 NumPositionSolves() const
-		{
-			return State.NumPositionSolves;
-		}
-
-		int32 NumVelocitySolves() const
-		{
-			return State.NumVelocitySolves;
-		}
-
 		void InitContact(
 			const int32 ManifoldPoiontIndex,
 			const FVec3& InWorldAnchorPoint0,
@@ -243,14 +241,21 @@ namespace Chaos
 		 * @brief Calculate and apply the position correction for this iteration
 		 * @return true if we need to run more iterations, false if we did not apply any correction
 		*/
-		bool SolvePositionWithFriction(const FReal Dt, const FReal MaxPushOut);
-		bool SolvePositionNoFriction(const FReal Dt, const FReal MaxPushOut);
+		bool SolvePositionWithFriction(const FRealLP Dt, const FRealLP MaxPushOut);
+		bool SolvePositionNoFriction(const FRealLP Dt, const FRealLP MaxPushOut);
 
 		/**
 		 * @brief Calculate and apply the velocity correction for this iteration
 		 * @return true if we need to run more iterations, false if we did not apply any correction
 		*/
-		bool SolveVelocity(const FReal Dt, const bool bApplyDynamicFriction);
+		bool SolveVelocity(const FRealLP Dt, const bool bApplyDynamicFriction);
+
+	private:
+		/**
+		 * @brief Apply the inverse mass scale the body with the lower level
+		 * @param InvMassScale 
+		*/
+		void SetShockPropagationInvMassScale(const FRealLP InvMassScale);
 
 		/**
 		 * @brief Run a velocity solve on the average point from all the points that received a position impulse
@@ -258,43 +263,28 @@ namespace Chaos
 		 * adding more velocity iterations.
 		 * This will only perform work if there is more than one active contact.
 		*/
-		void SolveVelocityAverage(const FReal Dt);
-
-	private:
-		/**
-		 * @brief Apply the inverse mass scale the body with the lower level
-		 * @param InvMassScale 
-		*/
-		void SetShockPropagationInvMassScale(const FReal InvMassScale);
+		void SolveVelocityAverage(const FRealLP Dt);
 
 		struct FState
 		{
 			FState()
-				: SolverBodies()
-				, ManifoldPoints()
-				, NumManifoldPoints(0)
-				, StaticFriction(0)
+				: StaticFriction(0)
 				, DynamicFriction(0)
 				, Stiffness(1)
-				, BodyEpochs{ INDEX_NONE, INDEX_NONE }
-				, NumPositionSolves(0)
-				, NumVelocitySolves(0)
+				, NumManifoldPoints(0)
 				, bHaveRestitution(false)
-				, bIsSolved(false)
+				, SolverBodies()
+				, ManifoldPoints()
 			{
 			}
 
-			FConstraintSolverBody SolverBodies[MaxConstrainedBodies];
-			FPBDCollisionSolverManifoldPoint ManifoldPoints[MaxPointsPerConstraint];
+			FRealLP StaticFriction;
+			FRealLP DynamicFriction;
+			FRealLP Stiffness;
 			int32 NumManifoldPoints;
-			FReal StaticFriction;
-			FReal DynamicFriction;
-			FReal Stiffness;
-			int32 BodyEpochs[MaxConstrainedBodies];
-			int32 NumPositionSolves;
-			int32 NumVelocitySolves;
 			bool bHaveRestitution;
-			bool bIsSolved;
+			FConstraintSolverBody SolverBodies[MaxConstrainedBodies];
+			FPBDCollisionSolverManifoldPoint ManifoldPoints[MaxPointsPerConstraint];
 		};
 
 		FState State;
@@ -307,24 +297,24 @@ namespace Chaos
 	//////////////////////////////////////////////////////////////////////////////////////////////////
 
 
-	FORCEINLINE_DEBUGGABLE void CalculatePositionCorrectionWithoutFriction(
-		const FReal Stiffness,
-		const FReal ContactDeltaNormal,
-		const FVec3& ContactNormal,
-		const FReal ContactMassNormal,
-		FVec3& InOutNetPushOut,
-		FVec3& OutPushOut)
+	FORCEINLINE_DEBUGGABLE void CalculatePositionCorrectionWithoutFrictionV(
+		const FRealLP Stiffness,
+		const FRealLP ContactDeltaNormal,
+		const FVec3LP& ContactNormal,
+		const FRealLP ContactMassNormal,
+		FVec3LP& InOutNetPushOut,
+		FVec3LP& OutPushOut)
 	{
-		FVec3 PushOut = -(Stiffness * ContactDeltaNormal * ContactMassNormal) * ContactNormal;
+		FVec3LP PushOut = -(Stiffness * ContactDeltaNormal * ContactMassNormal) * ContactNormal;
 
 		// The total pushout so far this sub-step
 		// We allow negative incremental impulses, but not net negative impulses
-		FVec3 NetPushOut = InOutNetPushOut + PushOut;
-		const FReal NetPushOutNormal = FVec3::DotProduct(NetPushOut, ContactNormal);
+		FVec3LP NetPushOut = InOutNetPushOut + PushOut;
+		const FRealLP NetPushOutNormal = FVec3LP::DotProduct(NetPushOut, ContactNormal);
 		if (NetPushOutNormal < 0)
 		{
 			PushOut = -InOutNetPushOut;
-			NetPushOut = FVec3(0);
+			NetPushOut = FVec3LP(0);
 		}
 
 		InOutNetPushOut = NetPushOut;
@@ -332,60 +322,80 @@ namespace Chaos
 	}
 
 
+	FORCEINLINE_DEBUGGABLE void CalculatePositionCorrectionWithoutFrictionF(
+		const FRealLP Stiffness,
+		const FRealLP ContactDeltaNormal,
+		const FRealLP ContactMassNormal,
+		const FRealLP NetPushOutNormal,
+		FRealLP& OutPushOutNormal)
+	{
+		FRealLP PushOutNormal = -(Stiffness * ContactDeltaNormal * ContactMassNormal);
+
+		// The total pushout so far this sub-step
+		// We allow negative incremental impulses, but not net negative impulses
+		if (NetPushOutNormal + PushOutNormal < 0)
+		{
+			PushOutNormal = -NetPushOutNormal;
+		}
+
+		OutPushOutNormal = PushOutNormal;
+	}
+
+
 	FORCEINLINE_DEBUGGABLE void CalculatePositionCorrectionWithFriction(
-		const FReal Stiffness,
-		const FVec3& ContactDelta,
-		const FReal ContactDeltaNormal,
-		const FVec3& ContactNormal,
-		const FMatrix33& ContactMass,
-		const FReal StaticFriction,
-		const FReal DynamicFriction,
-		FVec3& InOutNetPushOut,
-		FVec3& OutPushOut,
-		FReal& InOutStaticFrictionMax,
+		const FRealLP Stiffness,
+		const FVec3LP& ContactDelta,
+		const FRealLP ContactDeltaNormal,
+		const FVec3LP& ContactNormal,
+		const FMatrix33LP& ContactMass,
+		const FRealLP StaticFriction,
+		const FRealLP DynamicFriction,
+		FVec3LP& InOutNetPushOut,
+		FVec3LP& OutPushOut,
+		FRealLP& InOutStaticFrictionMax,
 		bool& bOutInsideStaticFrictionCone)
 	{
 		// If static friction is enabled, calculate the correction to move the contact point back to its
 		// original relative location on all axes.
 		// @todo(chaos): this should be moved to the ManifoldPoint error calculation?
-		FVec3 ModifiedContactError = -ContactDelta;
-		const FReal FrictionStiffness = CVars::Chaos_PBDCollisionSolver_Position_StaticFrictionStiffness;
+		FVec3LP ModifiedContactError = -ContactDelta;
+		const FRealLP FrictionStiffness = CVars::Chaos_PBDCollisionSolver_Position_StaticFrictionStiffness;
 		if (FrictionStiffness < 1.0f)
 		{
-			const FVec3 ContactDeltaTangent = ContactDelta - ContactDeltaNormal * ContactNormal;
+			const FVec3LP ContactDeltaTangent = ContactDelta - ContactDeltaNormal * ContactNormal;
 			ModifiedContactError = -ContactDeltaNormal * ContactNormal - FrictionStiffness * ContactDeltaTangent;
 		}
 
-		FVec3 PushOut = Stiffness * ContactMass * ModifiedContactError;
+		FVec3LP PushOut = Stiffness * ContactMass * ModifiedContactError;
 
 		// If we ended up with a negative normal pushout, remove all correction from this point
-		FVec3 NetPushOut = InOutNetPushOut + PushOut;
-		const FReal NetPushOutNormal = FVec3::DotProduct(NetPushOut, ContactNormal);
+		FVec3LP NetPushOut = InOutNetPushOut + PushOut;
+		const FRealLP NetPushOutNormal = FVec3LP::DotProduct(NetPushOut, ContactNormal);
 		bool bInsideStaticFrictionCone = true;
-		if (NetPushOutNormal < FReal(SMALL_NUMBER))
+		if (NetPushOutNormal < FRealLP(SMALL_NUMBER))
 		{
 			bInsideStaticFrictionCone = false;
 			PushOut = -InOutNetPushOut;
-			NetPushOut = FVec3(0);
+			NetPushOut = FVec3LP(0);
 		}
 
 		// Static friction limit: immediately increase maximum lateral correction, but smoothly decay maximum static friction limit. 
 		// This is so that small variations in position (jitter) and therefore NetPushOutNormal don't cause static friction to slip
 		// @todo(chaos): StaticFriction smoothing is iteration count depenendent - try to make it not so
-		const FReal StaticFrictionLerpRate = CVars::Chaos_PBDCollisionSolver_Position_StaticFrictionLerpRate;
-		const FReal StaticFrictionDest = FMath::Max(NetPushOutNormal, FReal(0));
-		FReal StaticFrictionMax = FMath::Lerp(FMath::Max(InOutStaticFrictionMax, StaticFrictionDest), StaticFrictionDest, StaticFrictionLerpRate);
+		const FRealLP StaticFrictionLerpRate = CVars::Chaos_PBDCollisionSolver_Position_StaticFrictionLerpRate;
+		const FRealLP StaticFrictionDest = FMath::Max(NetPushOutNormal, FRealLP(0));
+		FRealLP StaticFrictionMax = FMath::Lerp(FMath::Max(InOutStaticFrictionMax, StaticFrictionDest), StaticFrictionDest, StaticFrictionLerpRate);
 
 		// If we exceed the friction cone, stop adding frictional corrections and clip correction to cone
 		// @todo(chaos): clamp to dynamic friction
 		if (bInsideStaticFrictionCone)
 		{
-			const FReal MaxPushOutTangentSq = FMath::Square(StaticFriction * StaticFrictionMax);
-			const FVec3 NetPushOutTangent = (NetPushOut - NetPushOutNormal * ContactNormal);
-			const FReal NetPushOutTangentSq = NetPushOutTangent.SizeSquared();
+			const FRealLP MaxPushOutTangentSq = FMath::Square(StaticFriction * StaticFrictionMax);
+			const FVec3LP NetPushOutTangent = (NetPushOut - NetPushOutNormal * ContactNormal);
+			const FRealLP NetPushOutTangentSq = NetPushOutTangent.SizeSquared();
 			if (NetPushOutTangentSq > MaxPushOutTangentSq)
 			{
-				NetPushOut = NetPushOutNormal * ContactNormal + (StaticFriction * StaticFrictionMax) * NetPushOutTangent / FMath::Sqrt(NetPushOutTangentSq);
+				NetPushOut = NetPushOutNormal * ContactNormal + (StaticFriction * StaticFrictionMax) * NetPushOutTangent * FMath::InvSqrt(NetPushOutTangentSq);
 				PushOut = NetPushOut - InOutNetPushOut;
 				bInsideStaticFrictionCone = false;
 			}
@@ -398,16 +408,16 @@ namespace Chaos
 	}
 
 	FORCEINLINE_DEBUGGABLE void ApplyPositionCorrectionWithFriction(
-		const FReal Stiffness,
-		const FReal StaticFriction,
-		const FReal DynamicFriction,
-		const FVec3& ContactDelta,
-		const FReal ContactDeltaNormal,
+		const FRealLP Stiffness,
+		const FRealLP StaticFriction,
+		const FRealLP DynamicFriction,
+		const FVec3LP& ContactDelta,
+		const FRealLP ContactDeltaNormal,
 		FPBDCollisionSolverManifoldPoint& ManifoldPoint,
 		FConstraintSolverBody& Body0,
 		FConstraintSolverBody& Body1)
 	{
-		FVec3 PushOut = FVec3(0);
+		FVec3LP PushOut = FVec3LP(0);
 
 		// If we want to add static friction...
 		// (note: we run a few iterations without friction by temporarily setting StaticFriction to 0)
@@ -428,111 +438,87 @@ namespace Chaos
 		}
 		else
 		{
-			CalculatePositionCorrectionWithoutFriction(
+			CalculatePositionCorrectionWithoutFrictionV(
 				Stiffness,
 				ContactDeltaNormal,
 				ManifoldPoint.WorldContactNormal,
 				ManifoldPoint.WorldContactMassNormal,
 				ManifoldPoint.NetPushOut,					// Out
 				PushOut);									// Out
+
 		}
 
 		// Update the particle state based on the pushout
 		if (Body0.IsDynamic())
 		{
-			const FVec3 AngularPushOut = FVec3::CrossProduct(ManifoldPoint.WorldContactPosition - Body0.P(), PushOut);
-			const FVec3 DX0 = Body0.InvM() * PushOut;
-			const FVec3 DR0 = Body0.InvI() * AngularPushOut;
-			//if (!DX0.IsNearlyZero(CVars::Chaos_PBDCollisionSolver_Position_PositionSolverTolerance))
-			{
-				Body0.ApplyPositionDelta(DX0);
-			}
-			//if (!DR0.IsNearlyZero(CVars::Chaos_PBDCollisionSolver_Position_RotationSolverTolerance))
-			{
-				Body0.ApplyRotationDelta(DR0);
-			}
+			const FVec3LP AngularPushOut = FVec3LP::CrossProduct(ManifoldPoint.RelativeContactPosition0, PushOut);
+			const FVec3LP DX0 = Body0.InvM() * PushOut;
+			const FVec3LP DR0 = Body0.InvI() * AngularPushOut;
+			Body0.ApplyPositionDelta(DX0);
+			Body0.ApplyRotationDelta(DR0);
 		}
 		if (Body1.IsDynamic())
 		{
-			const FVec3 AngularPushOut = FVec3::CrossProduct(ManifoldPoint.WorldContactPosition - Body1.P(), PushOut);
-			const FVec3 DX1 = -(Body1.InvM() * PushOut);
-			const FVec3 DR1 = -(Body1.InvI() * AngularPushOut);
-			//if (!DX1.IsNearlyZero(CVars::Chaos_PBDCollisionSolver_Position_PositionSolverTolerance))
-			{
-				Body1.ApplyPositionDelta(DX1);
-			}
-			//if (!DR1.IsNearlyZero(CVars::Chaos_PBDCollisionSolver_Position_RotationSolverTolerance))
-			{
-				Body1.ApplyRotationDelta(DR1);
-			}
+			const FVec3LP AngularPushOut = FVec3LP::CrossProduct(ManifoldPoint.RelativeContactPosition1, PushOut);
+			const FVec3LP DX1 = -(Body1.InvM() * PushOut);
+			const FVec3LP DR1 = -(Body1.InvI() * AngularPushOut);
+			Body1.ApplyPositionDelta(DX1);
+			Body1.ApplyRotationDelta(DR1);
 		}
 	}
 
 	FORCEINLINE_DEBUGGABLE void ApplyPositionCorrectionNoFriction(
-		const FReal Stiffness,
-		const FVec3& ContactDelta,
-		const FReal ContactDeltaNormal,
+		const FRealLP Stiffness,
+		const FRealLP ContactDeltaNormal,
 		FPBDCollisionSolverManifoldPoint& ManifoldPoint,
 		FConstraintSolverBody& Body0,
 		FConstraintSolverBody& Body1)
 	{
-		FVec3 PushOut = FVec3(0);
+		FRealLP PushOutNormal = FRealLP(0);
+		const FRealLP NetPushOutNormal = FVec3LP::DotProduct(ManifoldPoint.NetPushOut, ManifoldPoint.WorldContactNormal);
 
-		CalculatePositionCorrectionWithoutFriction(
+		CalculatePositionCorrectionWithoutFrictionF(
 			Stiffness,
 			ContactDeltaNormal,
-			ManifoldPoint.WorldContactNormal,
 			ManifoldPoint.WorldContactMassNormal,
-			ManifoldPoint.NetPushOut,					// Out
-			PushOut);									// Out
+			NetPushOutNormal,
+			PushOutNormal);									// Out
+
+		ManifoldPoint.NetPushOut += PushOutNormal * ManifoldPoint.WorldContactNormal;
 
 		// Update the particle state based on the pushout
 		if (Body0.IsDynamic())
 		{
-			const FVec3 AngularPushOut = FVec3::CrossProduct(ManifoldPoint.WorldContactPosition - Body0.P(), PushOut);
-			const FVec3 DX0 = Body0.InvM() * PushOut;
-			const FVec3 DR0 = Body0.InvI() * AngularPushOut;
-			//if (!DX0.IsNearlyZero(CVars::Chaos_PBDCollisionSolver_Position_PositionSolverTolerance))
-			{
-				Body0.ApplyPositionDelta(DX0);
-			}
-			//if (!DR0.IsNearlyZero(CVars::Chaos_PBDCollisionSolver_Position_RotationSolverTolerance))
-			{
-				Body0.ApplyRotationDelta(DR0);
-			}
+			const FVec3LP DX0 = (Body0.InvM() * PushOutNormal) * ManifoldPoint.WorldContactNormal;
+			const FVec3LP DR0 = ManifoldPoint.WorldContactNormalAngular0 * PushOutNormal;
+			Body0.ApplyPositionDelta(DX0);
+			Body0.ApplyRotationDelta(DR0);
 		}
 		if (Body1.IsDynamic())
 		{
-			const FVec3 AngularPushOut = FVec3::CrossProduct(ManifoldPoint.WorldContactPosition - Body1.P(), PushOut);
-			const FVec3 DX1 = -(Body1.InvM() * PushOut);
-			const FVec3 DR1 = -(Body1.InvI() * AngularPushOut);
-			//if (!DX1.IsNearlyZero(CVars::Chaos_PBDCollisionSolver_Position_PositionSolverTolerance))
-			{
-				Body1.ApplyPositionDelta(DX1);
-			}
-			//if (!DR1.IsNearlyZero(CVars::Chaos_PBDCollisionSolver_Position_RotationSolverTolerance))
-			{
-				Body1.ApplyRotationDelta(DR1);
-			}
+			const FVec3LP DX1 = (Body1.InvM() * -PushOutNormal) * ManifoldPoint.WorldContactNormal;
+			const FVec3LP DR1 = ManifoldPoint.WorldContactNormalAngular1 * -PushOutNormal;
+			Body1.ApplyPositionDelta(DX1);
+			Body1.ApplyRotationDelta(DR1);
 		}
 	}
 
 	FORCEINLINE_DEBUGGABLE void CalculateVelocityCorrectionImpulse(
-		const FReal Stiffness,
-		const FReal Dt,
-		const FReal DynamicFriction,
-		const FVec3& ContactNormal,
-		const FMatrix33& ContactMass,
-		const FReal ContactMassNormal,
-		const FVec3& ContactVelocityDelta,
-		const FReal ContactVelocityDeltaNormal,
-		const FVec3& NetPushOut,
-		FVec3& InOutNetImpulse,
-		FVec3& OutImpulse)
+		const FRealLP Stiffness,
+		const FRealLP Dt,
+		const FRealLP DynamicFriction,
+		const FVec3LP& ContactNormal,
+		const FMatrix33LP& ContactMass,
+		const FRealLP ContactMassNormal,
+		const FVec3LP& ContactVelocityDelta,
+		const FRealLP ContactVelocityDeltaNormal,
+		const FVec3LP& NetPushOut,
+		FVec3LP& InOutNetImpulse,
+		FVec3LP& OutImpulse)
 	{
-		FVec3 Impulse = FVec3(0);
+		FVec3LP Impulse = FVec3LP(0);
 
-		if ((ContactVelocityDeltaNormal > FReal(0)) && !CVars::bChaos_PBDCollisionSolver_Velocity_NegativeImpulseEnabled)
+		if ((ContactVelocityDeltaNormal > FRealLP(0)) && !CVars::bChaos_PBDCollisionSolver_Velocity_NegativeImpulseEnabled)
 		{
 			return;
 		}
@@ -553,9 +539,9 @@ namespace Chaos
 		if (CVars::bChaos_PBDCollisionSolver_Velocity_ImpulseClampEnabled && (Dt > 0))
 		{
 			// @todo(chaos): cache max negative impulse
-			const FVec3 NetImpulse = InOutNetImpulse + Impulse;
-			const FReal PushOutImpulseNormal = FMath::Max(0.0f, FVec3::DotProduct(NetPushOut, ContactNormal) / Dt);
-			const FReal NetImpulseNormal = FVec3::DotProduct(NetImpulse, ContactNormal);
+			const FVec3LP NetImpulse = InOutNetImpulse + Impulse;
+			const FRealLP PushOutImpulseNormal = FMath::Max(0.0f, FVec3LP::DotProduct(NetPushOut, ContactNormal) / Dt);
+			const FRealLP NetImpulseNormal = FVec3LP::DotProduct(NetImpulse, ContactNormal);
 			if (NetImpulseNormal < -PushOutImpulseNormal)
 			{
 				// We are trying to apply a negative impulse larger than one to counteract the effective pushout impulse
@@ -573,16 +559,16 @@ namespace Chaos
 	}
 
 	FORCEINLINE_DEBUGGABLE void ApplyVelocityCorrection(
-		const FReal Stiffness,
-		const FReal Dt,
-		const FReal DynamicFriction,
-		const FVec3& ContactVelocityDelta,
-		const FReal ContactVelocityDeltaNormal,
+		const FRealLP Stiffness,
+		const FRealLP Dt,
+		const FRealLP DynamicFriction,
+		const FVec3LP& ContactVelocityDelta,
+		const FRealLP ContactVelocityDeltaNormal,
 		FPBDCollisionSolverManifoldPoint& ManifoldPoint,
 		FConstraintSolverBody& Body0,
 		FConstraintSolverBody& Body1)
 	{
-		FVec3 Impulse = FVec3(0);
+		FVec3LP Impulse = FVec3LP(0);
 
 		CalculateVelocityCorrectionImpulse(
 			Stiffness,
@@ -600,12 +586,12 @@ namespace Chaos
 		// Calculate the velocity deltas from the impulse
 		if (Body0.IsDynamic())
 		{
-			const FVec3 AngularImpulse = FVec3::CrossProduct(ManifoldPoint.WorldContactPosition - Body0.P(), Impulse);
+			const FVec3LP AngularImpulse = FVec3LP::CrossProduct(ManifoldPoint.RelativeContactPosition0, Impulse);
 			Body0.ApplyVelocityDelta(Body0.InvM() * Impulse, Body0.InvI() * AngularImpulse);
 		}
 		if (Body1.IsDynamic())
 		{
-			const FVec3 AngularImpulse = FVec3::CrossProduct(ManifoldPoint.WorldContactPosition - Body1.P(), Impulse);
+			const FVec3LP AngularImpulse = FVec3LP::CrossProduct(ManifoldPoint.RelativeContactPosition1, Impulse);
 			Body1.ApplyVelocityDelta(Body1.InvM() * -Impulse, Body1.InvI() * -AngularImpulse);
 		}
 	}
@@ -617,10 +603,8 @@ namespace Chaos
 	//////////////////////////////////////////////////////////////////////////////////////////////////
 
 
-	FORCEINLINE_DEBUGGABLE void FPBDCollisionSolverManifoldPoint::CalculateContactPositionError(const FConstraintSolverBody& Body0, const FConstraintSolverBody& Body1, const FReal MaxPushOut, FVec3& OutContactDelta, FReal& OutContactDeltaNormal) const
+	FORCEINLINE_DEBUGGABLE void FPBDCollisionSolverManifoldPoint::CalculateContactPositionError(const FConstraintSolverBody& Body0, const FConstraintSolverBody& Body1, const FRealLP MaxPushOut, FVec3LP& OutContactDelta, FRealLP& OutContactDeltaNormal) const
 	{
-		const FVec3 WorldRelativeContactPosition0 = WorldContactPosition - Body0.P();
-		const FVec3 WorldRelativeContactPosition1 = WorldContactPosition - Body1.P();
 #if CHAOS_NONLINEAR_COLLISIONS_ENABLED
 		// Non-linear version: calculate the contact delta after we have converted the current positional impulses into position and rotation corrections.
 		// We could precalculate and store the LocalContactPositions if we really want to use this nonlinear version
@@ -629,31 +613,31 @@ namespace Chaos
 		OutContactDelta = (Body0.CorrectedP() + Body0.CorrectedQ() * LocalContactPosition0) - (Body1.CorrectedP() + Body1.CorrectedQ() * LocalContactPosition1);
 #else
 		// Linear version: calculate the contact delta assuming linear motion after applying a positional impulse at the contact point. There will be an error that depends on the size of the rotation.
-		const FVec3 ContactDelta0 = Body0.DP() + FVec3::CrossProduct(Body0.DQ(), WorldRelativeContactPosition0);
-		const FVec3 ContactDelta1 = Body1.DP() + FVec3::CrossProduct(Body1.DQ(), WorldRelativeContactPosition1);
+		const FVec3LP ContactDelta0 = FVec3LP(Body0.DP()) + FVec3LP::CrossProduct(Body0.DQ(), RelativeContactPosition0);
+		const FVec3LP ContactDelta1 = FVec3LP(Body1.DP()) + FVec3LP::CrossProduct(Body1.DQ(), RelativeContactPosition1);
 		OutContactDelta = WorldContactDelta + ContactDelta0 - ContactDelta1;
 #endif
-		OutContactDeltaNormal = FVec3::DotProduct(OutContactDelta, WorldContactNormal);
+		OutContactDeltaNormal = FVec3LP::DotProduct(OutContactDelta, WorldContactNormal);
 
 		// NOTE: OutContactDeltaNormal is negative for penetration
 		// NOTE: MaxPushOut == 0 disables the pushout limits
 		if ((MaxPushOut > 0) && (OutContactDeltaNormal < -MaxPushOut))
 		{
-			const FReal ClampedContactDeltaNormal = -MaxPushOut;
+			const FRealLP ClampedContactDeltaNormal = -MaxPushOut;
 			OutContactDelta += (ClampedContactDeltaNormal - OutContactDeltaNormal) * WorldContactNormal;
 			OutContactDeltaNormal = ClampedContactDeltaNormal;
 		}
 	}
 
-	FORCEINLINE_DEBUGGABLE void FPBDCollisionSolverManifoldPoint::CalculateContactVelocityError(const FConstraintSolverBody& Body0, const FConstraintSolverBody& Body1, const FReal DynamicFriction, const FReal Dt, FVec3& OutContactVelocityDelta, FReal& OutContactVelocityDeltaNormal) const
+	FORCEINLINE_DEBUGGABLE void FPBDCollisionSolverManifoldPoint::CalculateContactVelocityError(const FConstraintSolverBody& Body0, const FConstraintSolverBody& Body1, const FRealLP DynamicFriction, const FRealLP Dt, FVec3LP& OutContactVelocityDelta, FRealLP& OutContactVelocityDeltaNormal) const
 	{
-		const FVec3 ContactVelocity0 = Body0.V() + FVec3::CrossProduct(Body0.W(), WorldContactPosition - Body0.P());
-		const FVec3 ContactVelocity1 = Body1.V() + FVec3::CrossProduct(Body1.W(), WorldContactPosition - Body1.P());
-		const FVec3 ContactVelocity = ContactVelocity0 - ContactVelocity1;
-		const FReal ContactVelocityNormal = FVec3::DotProduct(ContactVelocity, WorldContactNormal);
+		const FVec3LP ContactVelocity0 = Body0.V() + FVec3LP::CrossProduct(Body0.W(), RelativeContactPosition0);
+		const FVec3LP ContactVelocity1 = Body1.V() + FVec3LP::CrossProduct(Body1.W(), RelativeContactPosition1);
+		const FVec3LP ContactVelocity = ContactVelocity0 - ContactVelocity1;
+		const FRealLP ContactVelocityNormal = FVec3LP::DotProduct(ContactVelocity, WorldContactNormal);
 
 		// Target normal velocity, including restitution
-		const FReal ContactVelocityTargetNormal = WorldContactVelocityTargetNormal;
+		const FRealLP ContactVelocityTargetNormal = WorldContactVelocityTargetNormal;
 
 		// Add up the errors in the velocity (current velocity - desired velocity)
 		OutContactVelocityDeltaNormal = (ContactVelocityNormal - WorldContactVelocityTargetNormal);
@@ -662,18 +646,18 @@ namespace Chaos
 		const bool bApplyFriction = (DynamicFriction > 0) && (Dt > 0);
 		if (bApplyFriction)
 		{
-			const FVec3 ContactVelocityTangential = ContactVelocity - ContactVelocityNormal * WorldContactNormal;
-			const FReal ContactVelocityTangentialLen = ContactVelocityTangential.Size();
+			const FVec3LP ContactVelocityTangential = ContactVelocity - ContactVelocityNormal * WorldContactNormal;
+			const FRealLP ContactVelocityTangentialLen = ContactVelocityTangential.Size();
 			if (ContactVelocityTangentialLen > SMALL_NUMBER)
 			{
 				// PushOut = ContactMass * DP, where DP is the contact positional correction
 				// Friction force is proportional to the normal force, so friction velocity correction
 				// is proprtional to normal velocity correction, or DVn = DPn/dt = PushOut.N / (ContactMass * dt);
-				const FReal PushOutNormal = FVec3::DotProduct(NetPushOut, WorldContactNormal);
-				const FReal DynamicFrictionVelocityError = PushOutNormal / (WorldContactMassNormal * Dt);
+				const FRealLP PushOutNormal = FVec3LP::DotProduct(NetPushOut, WorldContactNormal);
+				const FRealLP DynamicFrictionVelocityError = PushOutNormal / (WorldContactMassNormal * Dt);
 				if (DynamicFrictionVelocityError > SMALL_NUMBER)
 				{
-					const FReal ContactVelocityErrorTangential = FMath::Min(DynamicFriction * DynamicFrictionVelocityError, ContactVelocityTangentialLen);
+					const FRealLP ContactVelocityErrorTangential = FMath::Min(DynamicFriction * DynamicFrictionVelocityError, ContactVelocityTangentialLen);
 					OutContactVelocityDelta += ContactVelocityTangential * (ContactVelocityErrorTangential / ContactVelocityTangentialLen);
 				}
 			}
@@ -687,32 +671,22 @@ namespace Chaos
 	//////////////////////////////////////////////////////////////////////////////////////////////////
 
 
-	FORCEINLINE_DEBUGGABLE bool FPBDCollisionSolver::SolvePositionWithFriction(const FReal Dt, const FReal MaxPushOut)
+	FORCEINLINE_DEBUGGABLE bool FPBDCollisionSolver::SolvePositionWithFriction(const FRealLP Dt, const FRealLP MaxPushOut)
 	{
 		// SolverBody decorator used to add mass scaling
 		FConstraintSolverBody& Body0 = SolverBody0();
 		FConstraintSolverBody& Body1 = SolverBody1();
 
-		// Check for solved
-		if (State.bIsSolved)
-		{
-			if ((Body0.LastChangeEpoch() == State.BodyEpochs[0]) && (Body1.LastChangeEpoch() == State.BodyEpochs[1]))
-			{
-				// We did not apply a correction last iteration (within tolerance) and nothing else has moved the bodies
-				return false;
-			}
-		}
-
 		// Apply the position correction so that all contacts have zero separation
 		for (int32 PointIndex = 0; PointIndex < State.NumManifoldPoints; ++PointIndex)
 		{
 			FPBDCollisionSolverManifoldPoint& SolverManifoldPoint = State.ManifoldPoints[PointIndex];
 
-			FVec3 ContactDelta;
-			FReal ContactDeltaNormal;
+			FVec3LP ContactDelta;
+			FRealLP ContactDeltaNormal;
 			SolverManifoldPoint.CalculateContactPositionError(Body0.SolverBody(), Body1.SolverBody(), MaxPushOut, ContactDelta, ContactDeltaNormal);
 
-			const bool bProcessManifoldPoint = (ContactDeltaNormal < FReal(0)) || !SolverManifoldPoint.NetPushOut.IsNearlyZero();
+			const bool bProcessManifoldPoint = (ContactDeltaNormal < FRealLP(0)) || !SolverManifoldPoint.NetPushOut.IsNearlyZero();
 			if (bProcessManifoldPoint)
 			{
 				ApplyPositionCorrectionWithFriction(
@@ -727,17 +701,10 @@ namespace Chaos
 			}
 		}
 
-		// We are solved if we did not move the bodies within some tolerance
-		// @todo(chaos): better early-out system
-		State.bIsSolved = (Body0.LastChangeEpoch() == State.BodyEpochs[0]) && (Body1.LastChangeEpoch() == State.BodyEpochs[1]);
-		State.BodyEpochs[0] = Body0.LastChangeEpoch();
-		State.BodyEpochs[1] = Body1.LastChangeEpoch();
-		State.NumPositionSolves = State.NumPositionSolves + 1;
-
-		return !State.bIsSolved;
+		return false;
 	}
 
-	FORCEINLINE_DEBUGGABLE bool FPBDCollisionSolver::SolvePositionNoFriction(const FReal Dt, const FReal MaxPushOut)
+	FORCEINLINE_DEBUGGABLE bool FPBDCollisionSolver::SolvePositionNoFriction(const FRealLP Dt, const FRealLP MaxPushOut)
 	{
 		// SolverBody decorator used to add mass scaling
 		FConstraintSolverBody& Body0 = SolverBody0();
@@ -748,16 +715,15 @@ namespace Chaos
 		{
 			FPBDCollisionSolverManifoldPoint& SolverManifoldPoint = State.ManifoldPoints[PointIndex];
 
-			FVec3 ContactDelta;
-			FReal ContactDeltaNormal;
+			FVec3LP ContactDelta;
+			FRealLP ContactDeltaNormal;
 			SolverManifoldPoint.CalculateContactPositionError(Body0.SolverBody(), Body1.SolverBody(), MaxPushOut, ContactDelta, ContactDeltaNormal);
 
-			const bool bProcessManifoldPoint = (ContactDeltaNormal < FReal(0)) || !SolverManifoldPoint.NetPushOut.IsNearlyZero();
+			const bool bProcessManifoldPoint = (ContactDeltaNormal < FRealLP(0)) || !SolverManifoldPoint.NetPushOut.IsNearlyZero();
 			if (bProcessManifoldPoint)
 			{
 				ApplyPositionCorrectionNoFriction(
 					State.Stiffness,
-					ContactDelta,
 					ContactDeltaNormal,
 					SolverManifoldPoint,
 					Body0,
@@ -765,15 +731,7 @@ namespace Chaos
 			}
 		}
 
-		// We are solved if we did not move the bodies within some tolerance
-		// NOTE: we can't claim to be solved until we have done at least one friction iteration so we can't early out before friction has been applied.
-		// @todo(chaos): better early-out system
-		State.bIsSolved = false;
-		State.BodyEpochs[0] = Body0.LastChangeEpoch();
-		State.BodyEpochs[1] = Body1.LastChangeEpoch();
-		State.NumPositionSolves = State.NumPositionSolves + 1;
-
-		return !State.bIsSolved;
+		return false;
 	}
 
 }
diff --git a/Engine/Source/Runtime/Experimental/Chaos/Private/Chaos/Collision/SolverCollisionContainer.cpp b/Engine/Source/Runtime/Experimental/Chaos/Private/Chaos/Collision/SolverCollisionContainer.cpp
--- a/Engine/Source/Runtime/Experimental/Chaos/Private/Chaos/Collision/SolverCollisionContainer.cpp
+++ b/Engine/Source/Runtime/Experimental/Chaos/Private/Chaos/Collision/SolverCollisionContainer.cpp
@@ -129,7 +129,6 @@ namespace Chaos
 			FVec3 AccumulatedImpulse = FVec3(0);
 
 			Constraint->ResetSolverResults();
-			Constraint->SetNumActivePositionIterations(Solver.NumPositionSolves());
 
 			for (int32 PointIndex = 0; PointIndex < Solver.NumManifoldPoints(); ++PointIndex)
 			{
@@ -277,8 +276,11 @@ namespace Chaos
 	}
 
 	// Solve position including support for incremental collision detection
-	bool FPBDCollisionSolverContainer::SolvePositionIncrementalImpl(const FReal Dt, const int32 BeginIndex, const int32 EndIndex, const FReal MaxPushOut, const bool bApplyStaticFriction)
+	bool FPBDCollisionSolverContainer::SolvePositionIncrementalImpl(const FReal InDt, const int32 BeginIndex, const int32 EndIndex, const FReal InMaxPushOut, const bool bApplyStaticFriction)
 	{
+		const FRealLP Dt = FRealLP(InDt);
+		const FRealLP MaxPushOut = FRealLP(InMaxPushOut);
+
 		bool bNeedsAnotherIteration = false;
 		for (int32 SolverIndex = BeginIndex; SolverIndex < EndIndex; ++SolverIndex)
 		{
@@ -302,8 +304,11 @@ namespace Chaos
 	}
 
 	// Solve position with friction (last few iterations each tick)
-	bool FPBDCollisionSolverContainer::SolvePositionWithFrictionImpl(const FReal Dt, const int32 BeginIndex, const int32 EndIndex, const FReal MaxPushOut)
+	bool FPBDCollisionSolverContainer::SolvePositionWithFrictionImpl(const FReal InDt, const int32 BeginIndex, const int32 EndIndex, const FReal InMaxPushOut)
 	{
+		const FRealLP Dt = FRealLP(InDt);
+		const FRealLP MaxPushOut = FRealLP(InMaxPushOut);
+
 		bool bNeedsAnotherIteration = false;
 		for (int32 SolverIndex = BeginIndex; SolverIndex < EndIndex; ++SolverIndex)
 		{
@@ -313,8 +318,11 @@ namespace Chaos
 	}
 
 	// Solve position without friction (first few iterations each tick)
-	bool FPBDCollisionSolverContainer::SolvePositionNoFrictionImpl(const FReal Dt, const int32 BeginIndex, const int32 EndIndex, const FReal MaxPushOut)
+	bool FPBDCollisionSolverContainer::SolvePositionNoFrictionImpl(const FReal InDt, const int32 BeginIndex, const int32 EndIndex, const FReal InMaxPushOut)
 	{
+		const FRealLP Dt = FRealLP(InDt);
+		const FRealLP MaxPushOut = FRealLP(InMaxPushOut);
+
 		bool bNeedsAnotherIteration = false;
 		for (int32 SolverIndex = BeginIndex; SolverIndex < EndIndex; ++SolverIndex)
 		{
@@ -323,13 +331,14 @@ namespace Chaos
 		return bNeedsAnotherIteration;
 	}
 
-	bool FPBDCollisionSolverContainer::SolveVelocityImpl(const FReal Dt, const int32 It, const int32 NumIts, const int32 BeginIndex, const int32 EndIndex, const bool bParallel)
+	bool FPBDCollisionSolverContainer::SolveVelocityImpl(const FReal InDt, const int32 It, const int32 NumIts, const int32 BeginIndex, const int32 EndIndex, const bool bParallel)
 	{
 		SCOPE_CYCLE_COUNTER(STAT_Collisions_ApplyPushOut);
 		if (!bChaos_PBDCollisionSolver_Velocity_SolveEnabled)
 		{
 			return false;
 		}
+		const FRealLP Dt = FRealLP(InDt);
 
 		UpdateVelocityShockPropagation(Dt, It, NumIts, BeginIndex, EndIndex);
 
diff --git a/Engine/Source/Runtime/Experimental/Chaos/Private/Chaos/CollisionResolutionUtil.cpp b/Engine/Source/Runtime/Experimental/Chaos/Private/Chaos/CollisionResolutionUtil.cpp
--- a/Engine/Source/Runtime/Experimental/Chaos/Private/Chaos/CollisionResolutionUtil.cpp
+++ b/Engine/Source/Runtime/Experimental/Chaos/Private/Chaos/CollisionResolutionUtil.cpp
@@ -29,21 +29,6 @@ namespace Chaos
 			return FRigidTransform3(Generic->P(), Generic->Q());
 		}
 
-		// @todo(ccaulfield): This is duplicated in JointConstraints - move to a utility file
-		FMatrix33 ComputeFactorMatrix3(const FVec3& V, const FMatrix33& M, const FReal Im)
-		{
-			// Rigid objects rotational contribution to the impulse.
-			// Vx*M*VxT+Im
-			ensure(Im > FLT_MIN);
-				return FMatrix33(
-					-V[2] * (-V[2] * M.M[1][1] + V[1] * M.M[2][1]) + V[1] * (-V[2] * M.M[2][1] + V[1] * M.M[2][2]) + Im,
-					V[2] * (-V[2] * M.M[1][0] + V[1] * M.M[2][0]) - V[0] * (-V[2] * M.M[2][1] + V[1] * M.M[2][2]),
-					-V[1] * (-V[2] * M.M[1][0] + V[1] * M.M[2][0]) + V[0] * (-V[2] * M.M[1][1] + V[1] * M.M[2][1]),
-					V[2] * (V[2] * M.M[0][0] - V[0] * M.M[2][0]) - V[0] * (V[2] * M.M[2][0] - V[0] * M.M[2][2]) + Im,
-					-V[1] * (V[2] * M.M[0][0] - V[0] * M.M[2][0]) + V[0] * (V[2] * M.M[1][0] - V[0] * M.M[2][1]),
-					-V[1] * (-V[1] * M.M[0][0] + V[0] * M.M[1][0]) + V[0] * (-V[1] * M.M[1][0] + V[0] * M.M[1][1]) + Im);
-		}
-
 		// Reference: Energy Stability and Fracture for Frame Rate Rigid Body Simulations (Su et al.) (3.2. Clamping Impulses)
 		FVec3 GetEnergyClampedImpulse(const TPBDRigidParticleHandle<FReal, 3>* PBDRigid0, const TPBDRigidParticleHandle<FReal, 3>* PBDRigid1, const FVec3& Impulse, const FVec3& VectorToPoint1, const FVec3& VectorToPoint2, const FVec3& Velocity1, const FVec3& Velocity2)
 		{
diff --git a/Engine/Source/Runtime/Experimental/Chaos/Private/Chaos/Evolution/SolverBodyContainer.cpp b/Engine/Source/Runtime/Experimental/Chaos/Private/Chaos/Evolution/SolverBodyContainer.cpp
--- a/Engine/Source/Runtime/Experimental/Chaos/Private/Chaos/Evolution/SolverBodyContainer.cpp
+++ b/Engine/Source/Runtime/Experimental/Chaos/Private/Chaos/Evolution/SolverBodyContainer.cpp
@@ -39,17 +39,12 @@ namespace Chaos
 	{
 		if (Particle.IsValid())
 		{
-			if (SolverBody.IsDynamic() && SolverBody.IsModified())
+			if (SolverBody.IsDynamic())
 			{
 				// Set the particle state
 				FParticleUtilities::SetCoMWorldTransform(Particle, SolverBody.CorrectedP(), SolverBody.CorrectedQ());
 				Particle->SetV(SolverBody.V());
 				Particle->SetW(SolverBody.W());
-
-				if (SolverBody.HasActiveCollision())
-				{
-					//Particle->AuxilaryValue(*ParticleParameters.Collided) = true;
-				}
 			}
 
 			// Reset SolverBodyIndex cookie every step - it will be reassigned next step
diff --git a/Engine/Source/Runtime/Experimental/Chaos/Private/Chaos/PBDCollisionConstraintsContact.cpp b/Engine/Source/Runtime/Experimental/Chaos/Private/Chaos/PBDCollisionConstraintsContact.cpp
--- a/Engine/Source/Runtime/Experimental/Chaos/Private/Chaos/PBDCollisionConstraintsContact.cpp
+++ b/Engine/Source/Runtime/Experimental/Chaos/Private/Chaos/PBDCollisionConstraintsContact.cpp
@@ -112,10 +112,10 @@ namespace Chaos
 				// add the padding to the Phi (since it was from pre-padded collision detection).
 				if (bApplyResitution && !bHaveRestitutionPadding)
 				{
-					const FVec3 BodyV0 = Body0.V() + (Body0.DP() / IterationParameters.Dt);
-					const FVec3 BodyW0 = Body0.W() + (Body0.DQ() / IterationParameters.Dt);
-					const FVec3 BodyV1 = Body1.V() + (Body1.DP() / IterationParameters.Dt);
-					const FVec3 BodyW1 = Body1.W() + (Body1.DQ() / IterationParameters.Dt);
+					const FVec3 BodyV0 = Body0.V() + (FVec3(Body0.DP()) / IterationParameters.Dt);
+					const FVec3 BodyW0 = Body0.W() + (FVec3(Body0.DQ()) / IterationParameters.Dt);
+					const FVec3 BodyV1 = Body1.V() + (FVec3(Body1.DP()) / IterationParameters.Dt);
+					const FVec3 BodyW1 = Body1.W() + (FVec3(Body1.DQ()) / IterationParameters.Dt);
 					FVec3 CV0 = BodyV0 + FVec3::CrossProduct(BodyW0, VectorToPoint0);
 					FVec3 CV1 = BodyV1 + FVec3::CrossProduct(BodyW1, VectorToPoint1);
 					FVec3 CV = CV0 - CV1;
@@ -130,8 +130,8 @@ namespace Chaos
 				}
 			
 				FMatrix33 ContactInvI =
-					(Body0.IsDynamic() ? ComputeFactorMatrix3(VectorToPoint0, Body0.InvI(), Body0.InvM()) : FMatrix33(0)) +
-					(Body1.IsDynamic() ? ComputeFactorMatrix3(VectorToPoint1, Body1.InvI(), Body1.InvM()) : FMatrix33(0));
+					(Body0.IsDynamic() ? ComputeFactorMatrix3(VectorToPoint0, FMatrix33(Body0.InvI()), FReal(Body0.InvM())) : FMatrix33(0)) +
+					(Body1.IsDynamic() ? ComputeFactorMatrix3(VectorToPoint1, FMatrix33(Body1.InvI()), FReal(Body1.InvM())) : FMatrix33(0));
 
 				// Calculate the normal correction
 				FVec3 NormalError = Contact.Phi * Contact.Normal;
@@ -178,14 +178,14 @@ namespace Chaos
 				if (Body0.IsDynamic())
 				{
 					FVec3 DP0 = Body0.InvM() * DX;
-					FVec3 DR0 = Utilities::Multiply(Body0.InvI(), FVec3::CrossProduct(VectorToPoint0, DX));
+					FVec3 DR0 = Utilities::Multiply(FMatrix33(Body0.InvI()), FVec3::CrossProduct(VectorToPoint0, DX));
 					Body0.ApplyTransformDelta(DP0, DR0);
 					//Body0.UpdateRotationDependentState();
 				}
 				if (Body1.IsDynamic())
 				{
 					FVec3 DP1 = Body1.InvM() * -DX;
-					FVec3 DR1 = Utilities::Multiply(Body1.InvI(), FVec3::CrossProduct(VectorToPoint1, -DX));
+					FVec3 DR1 = Utilities::Multiply(FMatrix33(Body1.InvI()), FVec3::CrossProduct(VectorToPoint1, -DX));
 					Body1.ApplyTransformDelta(DP1, DR1);
 					//Body1.UpdateRotationDependentState();
 				}
diff --git a/Engine/Source/Runtime/Experimental/Chaos/Public/Chaos/Collision/PBDCollisionConstraint.h b/Engine/Source/Runtime/Experimental/Chaos/Public/Chaos/Collision/PBDCollisionConstraint.h
--- a/Engine/Source/Runtime/Experimental/Chaos/Public/Chaos/Collision/PBDCollisionConstraint.h
+++ b/Engine/Source/Runtime/Experimental/Chaos/Public/Chaos/Collision/PBDCollisionConstraint.h
@@ -385,9 +385,6 @@ namespace Chaos
 		
 		bool GetUseIncrementalCollisionDetection() const { return !Flags.bUseManifold || Flags.bUseIncrementalManifold; }
 
-		void SetNumActivePositionIterations(const int32 InNumActivePositionIterations) { NumActivePositionIterations = InNumActivePositionIterations; }
-		int32 GetNumActivePositionIterations() const { return NumActivePositionIterations; }
-
 		/**
 		 * @brief Clear the current and previous manifolds
 		*/
@@ -593,7 +590,6 @@ namespace Chaos
 		FRigidTransform3 LastShapeWorldTransform0;
 		FRigidTransform3 LastShapeWorldTransform1;
 		int32 ExpectedNumManifoldPoints;
-		int32 NumActivePositionIterations;
 	};
 
 
diff --git a/Engine/Source/Runtime/Experimental/Chaos/Public/Chaos/CollisionResolutionUtil.h b/Engine/Source/Runtime/Experimental/Chaos/Public/Chaos/CollisionResolutionUtil.h
--- a/Engine/Source/Runtime/Experimental/Chaos/Public/Chaos/CollisionResolutionUtil.h
+++ b/Engine/Source/Runtime/Experimental/Chaos/Public/Chaos/CollisionResolutionUtil.h
@@ -21,8 +21,20 @@ namespace Chaos
 		FRigidTransform3 
 		GetTransform(const TGeometryParticleHandle<FReal, 3>* Particle);
 
-		FMatrix33 
-		ComputeFactorMatrix3(const FVec3& V, const FMatrix33& M, const FReal Im);
+		template<typename TRealType>
+		inline TMatrix33<TRealType> ComputeFactorMatrix3(const TVec3<TRealType>& V, const TMatrix33<TRealType>& M, const TRealType Im)
+		{
+			// Rigid objects rotational contribution to the impulse.
+			// Vx*M*VxT+Im
+			check(Im > FLT_MIN);
+			return TMatrix33<TRealType>(
+				-V[2] * (-V[2] * M.M[1][1] + V[1] * M.M[2][1]) + V[1] * (-V[2] * M.M[2][1] + V[1] * M.M[2][2]) + Im,
+				V[2] * (-V[2] * M.M[1][0] + V[1] * M.M[2][0]) - V[0] * (-V[2] * M.M[2][1] + V[1] * M.M[2][2]),
+				-V[1] * (-V[2] * M.M[1][0] + V[1] * M.M[2][0]) + V[0] * (-V[2] * M.M[1][1] + V[1] * M.M[2][1]),
+				V[2] * (V[2] * M.M[0][0] - V[0] * M.M[2][0]) - V[0] * (V[2] * M.M[2][0] - V[0] * M.M[2][2]) + Im,
+				-V[1] * (V[2] * M.M[0][0] - V[0] * M.M[2][0]) + V[0] * (V[2] * M.M[1][0] - V[0] * M.M[2][1]),
+				-V[1] * (-V[1] * M.M[0][0] + V[0] * M.M[1][0]) + V[0] * (-V[1] * M.M[1][0] + V[0] * M.M[1][1]) + Im);
+		}
 
 		FVec3
 		GetEnergyClampedImpulse(const TPBDRigidParticleHandle<FReal, 3>* PBDRigid0, const TPBDRigidParticleHandle<FReal, 3>* PBDRigid1, const FVec3& Impulse, const FVec3& VectorToPoint1, const FVec3& VectorToPoint2, const FVec3& Velocity1, const FVec3& Velocity2);
diff --git a/Engine/Source/Runtime/Experimental/Chaos/Public/Chaos/Evolution/SolverBody.h b/Engine/Source/Runtime/Experimental/Chaos/Public/Chaos/Evolution/SolverBody.h
--- a/Engine/Source/Runtime/Experimental/Chaos/Public/Chaos/Evolution/SolverBody.h
+++ b/Engine/Source/Runtime/Experimental/Chaos/Public/Chaos/Evolution/SolverBody.h
@@ -4,8 +4,21 @@
 #include "Chaos/Core.h"
 #include "Math/Quat.h"
 
+// Set to 1 to force single precision in the constraint solver, even if the default numeric type is double
+#define CHAOS_CONSTRAINTSOLVER_LOWPRECISION 1
+
 namespace Chaos
 {
+	// Set the low precision (LP) math types
+#if CHAOS_CONSTRAINTSOLVER_LOWPRECISION
+	using FRealLP = FRealSingle;
+#else
+	using FRealLP = FReal;
+#endif
+	using FVec3LP = TVec3<FRealLP>;
+	using FMatrix33LP = TMatrix33<FRealLP>;
+
+
 	class FSolverBody;
 	class FSolverBodyContainer;
 
@@ -123,42 +136,44 @@ namespace Chaos
 		{
 			if (IsDynamic() && (Dt != FReal(0)))
 			{
-				State.V += State.DP / Dt;
-				State.W += State.DQ / Dt;
+				const FRealLP InvDt = FRealLP(1) / FRealLP(Dt);
+				State.V += State.DP * InvDt;
+				State.W += State.DQ * InvDt;
 			}
 		}
 
 		/**
 		 * @brief Get the inverse mass
 		*/
-		inline FReal InvM() const { return State.InvM; }
+		//inline FReal InvM() const { return State.InvM; }
+		inline FRealLP InvM() const { return State.InvM; }
 
 		/**
 		 * @brief Set the inverse mass
 		*/
-		inline void SetInvM(FReal InInvM) { State.InvM = InInvM; }
+		inline void SetInvM(FReal InInvM) { State.InvM = FRealLP(InInvM); }
 
 		/**
 		 * @brief Get the world-space inverse inertia
 		*/
-		inline const FMatrix33& InvI() const { return State.InvI; }
+		inline const FMatrix33LP& InvI() const { return State.InvI; }
 
 		/**
 		 * @brief Set the world-space inverse inertia
 		*/
-		inline void SetInvI(const FMatrix33& InInvI) { State.InvI = InInvI; }
+		inline void SetInvI(const FMatrix33& InInvI) { State.InvI = FMatrix33LP(InInvI); }
 
 		/**
 		 * @brief Get the local-space inverse inertia (diagonal elements)
 		*/
-		inline const FVec3& InvILocal() const { return State.InvILocal; }
+		inline const FVec3LP& InvILocal() const { return State.InvILocal; }
 
 		/**
 		 * @brief Set the local-space inverse inertia (diagonal elements)
 		*/
 		inline void SetInvILocal(const FVec3& InInvILocal)
 		{ 
-			State.InvILocal = InInvILocal; 
+			State.InvILocal = FVec3LP(InInvILocal); 
 			UpdateRotationDependentState();
 		}
 
@@ -216,24 +231,24 @@ namespace Chaos
 		/**
 		 * @brief Net world-space position correction applied by the constraints
 		*/
-		inline const FVec3& DP() const { return State.DP; }
+		inline const FVec3LP& DP() const { return State.DP; }
 
 		/**
 		 * @brief Net world-space rotation correction applied by the constraints (axis-angle vector equivalent to angular velocity but for position)
 		*/
-		inline const FVec3& DQ() const { return State.DQ; }
+		inline const FVec3LP& DQ() const { return State.DQ; }
 
 		/**
 		 * @brief World-space position after applying the net correction DP()
 		 * @note Calculated on demand from P() and DP() (only requires vector addition)
 		*/
-		inline FVec3 CorrectedP() const { return State.P + State.DP; }
+		inline FVec3 CorrectedP() const { return State.P + FVec3(State.DP); }
 
 		/**
 		 * @brief World-space rotation after applying the net correction DQ()
 		 * @note Calculated on demand from Q() and DQ() (requires quaternion multiply and normalization)
 		*/
-		inline FRotation3 CorrectedQ() const { return IsModified() ? FRotation3::IntegrateRotationWithAngularVelocity(State.Q, State.DQ, FReal(1)) : State.Q; }
+		inline FRotation3 CorrectedQ() const { return IsDynamic() ? FRotation3::IntegrateRotationWithAngularVelocity(State.Q, FVec3(State.DQ), FReal(1)) : State.Q; }
 
 		/**
 		 * @brief Apply the accumulated position and rotation corrections to the predicted P and Q
@@ -244,8 +259,8 @@ namespace Chaos
 		{
 			State.P = CorrectedP();
 			State.Q = CorrectedQ();
-			State.DP = FVec3(0);
-			State.DQ = FVec3(0);
+			State.DP = FVec3LP(0);
+			State.DQ = FVec3LP(0);
 		}
 
 		/**
@@ -266,12 +281,6 @@ namespace Chaos
 		inline int32 Level() const { return State.Level; }
 		inline void SetLevel(int32 InLevel) { State.Level = InLevel; }
 
-		/**
-		 * @brief Whether there were any active collision constraints on this body
-		*/
-		inline bool HasActiveCollision() const { return State.bHasActiveCollision; }
-		inline void SetHasActiveCollision(bool bInHasCollision) { State.bHasActiveCollision = bInHasCollision; }
-
 		/**
 		 * @brief Whether the body has a finite mass
 		 * @note This is based on the current inverse mass, so a "dynamic" particle with 0 inverse mass will return true here.
@@ -281,7 +290,7 @@ namespace Chaos
 		/**
 		 * @brief Apply a world-space position and rotation delta to the body center of mass, and update inverse mass
 		*/
-		inline void ApplyTransformDelta(const FVec3 DP, const FVec3 DR)
+		inline void ApplyTransformDelta(const FVec3LP& DP, const FVec3LP& DR)
 		{
 			ApplyPositionDelta(DP);
 			ApplyRotationDelta(DR);
@@ -290,25 +299,23 @@ namespace Chaos
 		/**
 		 * @brief Apply a world-space position delta to the solver body center of mass
 		*/
-		inline void ApplyPositionDelta(const FVec3 DP)
+		inline void ApplyPositionDelta(const FVec3LP& DP)
 		{
 			State.DP += DP;
-			++State.LastChangeEpoch;
 		}
 
 		/**
 		 * @brief Apply a world-space rotation delta to the solver body and update the inverse mass
 		*/
-		inline void ApplyRotationDelta(const FVec3 DR)
+		inline void ApplyRotationDelta(const FVec3LP& DR)
 		{
 			State.DQ += DR;
-			++State.LastChangeEpoch;
 		}
 
 		/**
 		 * @brief Apply a world-space velocity delta to the solver body
 		*/
-		inline void ApplyVelocityDelta(const FVec3& DV, const FVec3& DW)
+		inline void ApplyVelocityDelta(const FVec3LP& DV, const FVec3LP& DW)
 		{
 			ApplyLinearVelocityDelta(DV);
 			ApplyAngularVelocityDelta(DW);
@@ -317,19 +324,17 @@ namespace Chaos
 		/**
 		 * @brief Apply a world-space linear velocity delta to the solver body
 		*/
-		inline void ApplyLinearVelocityDelta(const FVec3& DV)
+		inline void ApplyLinearVelocityDelta(const FVec3LP& DV)
 		{
 			State.V += DV;
-			++State.LastChangeEpoch;
 		}
 
 		/**
 		 * @brief Apply an world-space angular velocity delta to the solver body
 		*/
-		inline void ApplyAngularVelocityDelta(const FVec3& DW)
+		inline void ApplyAngularVelocityDelta(const FVec3LP& DW)
 		{
 			State.W += DW;
-			++State.LastChangeEpoch;
 		}
 
 		/**
@@ -341,16 +346,6 @@ namespace Chaos
 			State.Q.EnforceShortestArcWith(InQ);
 		}
 
-		inline int32 LastChangeEpoch() const
-		{
-			return State.LastChangeEpoch;
-		}
-
-		inline bool IsModified() const
-		{
-			return State.LastChangeEpoch != 0;
-		}
-
 		/**
 		 * @brief Update cached state that depends on rotation (i.e., world space inertia)
 		*/
@@ -366,6 +361,9 @@ namespace Chaos
 				: InvILocal(0)
 				, InvM(0)
 				, InvI(0)
+				, DP(0)
+				, DQ(0)
+				, Level(0)
 				, X(0)
 				, R(FRotation3::Identity)
 				, P(0)
@@ -374,22 +372,26 @@ namespace Chaos
 				, W(0)
 				, CoM(0)
 				, RoM(FRotation3::Identity)
-				, DP(0)
-				, DQ(0)
-				, Level(0)
-				, LastChangeEpoch(0)
-				, bHasActiveCollision(false)
 			{}
 
 			// Local-space inverse inertia (diagonal, so only 3 elements)
-			FVec3 InvILocal;
+			FVec3LP InvILocal;
 
 			// Inverse mass
-			FReal InvM;
+			FRealLP InvM;
 
 			// World-space inverse inertia
 			// @todo(chaos): do we need this, or should we force all systems to use the FConstraintSolverBody decorator?
-			FMatrix33 InvI;
+			FMatrix33LP InvI;
+
+			// Net position delta applied by all constraints (constantly changing as we iterate over constraints)
+			FVec3LP DP;
+
+			// Net rotation delta applied by all constraints (constantly changing as we iterate over constraints)
+			FVec3LP DQ;
+
+			// Distance to a kinmatic body (through the contact graph). Used by collision shock propagation
+			int32 Level;
 
 			// World-space center of mass state at start of sub step
 			FVec3 X;
@@ -414,26 +416,8 @@ namespace Chaos
 
 			// Actor-space center of mass rotation
 			FRotation3 RoM;
-
-			// Net position delta applied by all constraints (constantly changing as we iterate over constraints)
-			FVec3 DP;
-
-			// Net rotation delta applied by all constraints (constantly changing as we iterate over constraints)
-			FVec3 DQ;
-
-			// Distance to a kinmatic body (through the contact graph). Used by collision shock propagation
-			int32 Level;
-
-			// A counter incremented every time the state changes. 
-			// Used by constraints to determine if some other constraint has modified the body for early exit logic
-			int32 LastChangeEpoch;
-
-			// Whether we had any active contacts this substep
-			// @todo(chaos): maybe make this a counter?
-			bool bHasActiveCollision;
 		};
 
-		//FGenericParticleHandle Particle;
 		FState State;
 	};
 
@@ -487,23 +471,23 @@ namespace Chaos
 		/**
 		 * @brief A scale applied to both inverse mass and inverse inertia
 		*/
-		inline FReal InvMScale() const { return State.InvMassScale; }
-		inline void SetInvMScale(FReal InInvMassScale) { State.InvMassScale = InInvMassScale; }
+		inline FRealLP InvMScale() const { return State.InvMassScale; }
+		inline void SetInvMScale(FReal InInvMassScale) { State.InvMassScale = FRealLP(InInvMassScale); }
 
 		/**
 		 * @brief The scaled inverse mass
 		*/
-		FReal InvM() const { return State.InvMassScale * Body->InvM(); }
+		FRealLP InvM() const { return State.InvMassScale * Body->InvM(); }
 
 		/**
 		 * @brief The scaled inverse inertia
 		*/
-		FMatrix33 InvI() const { return State.InvMassScale * Body->InvI(); }
+		FMatrix33LP InvI() const { return State.InvMassScale * Body->InvI(); }
 
 		/**
 		 * @brief The scaled local space inverse inertia
 		*/
-		FVec3 InvILocal() const { return State.InvMassScale * Body->InvILocal(); }
+		FVec3LP InvILocal() const { return State.InvMassScale * Body->InvILocal(); }
 
 		/**
 		 * @brief Whether the body is dynamic (i.e., has a finite mass) after InvMassScale is applied
@@ -525,30 +509,28 @@ namespace Chaos
 		inline const FVec3& V() const { return Body->V(); }
 		inline const FVec3& W() const { return Body->W(); }
 		inline int32 Level() const { return Body->Level(); }
-		inline const FVec3& DP() const { return Body->DP(); }
-		inline const FVec3& DQ() const { return Body->DQ(); }
+		inline const FVec3LP& DP() const { return Body->DP(); }
+		inline const FVec3LP& DQ() const { return Body->DQ(); }
 		inline FVec3 CorrectedP() const { return Body->CorrectedP(); }
 		inline FRotation3 CorrectedQ() const { return Body->CorrectedQ(); }
 
-		inline void ApplyTransformDelta(const FVec3 DP, const FVec3 DR) { Body->ApplyTransformDelta(DP, DR); }
-		inline void ApplyPositionDelta(const FVec3 DP) { Body->ApplyPositionDelta(DP); }
-		inline void ApplyRotationDelta(const FVec3 DR) { Body->ApplyRotationDelta(DR); }
-		inline void ApplyVelocityDelta(const FVec3& DV, const FVec3& DW) { Body->ApplyVelocityDelta(DV, DW); }
-		inline void ApplyLinearVelocityDelta(const FVec3& DV) { Body->ApplyLinearVelocityDelta(DV); }
-		inline void ApplyAngularVelocityDelta(const FVec3& DW) { Body->ApplyAngularVelocityDelta(DW); }
+		inline void ApplyTransformDelta(const FVec3LP& DP, const FVec3LP& DR) { Body->ApplyTransformDelta(DP, DR); }
+		inline void ApplyPositionDelta(const FVec3LP& DP) { Body->ApplyPositionDelta(DP); }
+		inline void ApplyRotationDelta(const FVec3LP& DR) { Body->ApplyRotationDelta(DR); }
+		inline void ApplyVelocityDelta(const FVec3LP& DV, const FVec3LP& DW) { Body->ApplyVelocityDelta(DV, DW); }
+		inline void ApplyLinearVelocityDelta(const FVec3LP& DV) { Body->ApplyLinearVelocityDelta(DV); }
+		inline void ApplyAngularVelocityDelta(const FVec3LP& DW) { Body->ApplyAngularVelocityDelta(DW); }
 		inline void EnforceShortestRotationTo(const FRotation3& InQ) { Body->EnforceShortestRotationTo(InQ); }
 		inline void UpdateRotationDependentState() { Body->UpdateRotationDependentState(); }
 
-		inline int32 LastChangeEpoch() const { return Body->LastChangeEpoch(); }
-
 	private:
 		// Struct is only so that we can use the same var names as function names
 		struct FState
 		{
 			FState() 
-				: InvMassScale(FReal(1)) 
+				: InvMassScale(FRealLP(1))
 			{}
-			FReal InvMassScale;
+			FRealLP InvMassScale;
 		};
 
 		// The body we decorate
diff --git a/Engine/Source/Runtime/Experimental/Chaos/Public/Chaos/Utilities.h b/Engine/Source/Runtime/Experimental/Chaos/Public/Chaos/Utilities.h
--- a/Engine/Source/Runtime/Experimental/Chaos/Public/Chaos/Utilities.h
+++ b/Engine/Source/Runtime/Experimental/Chaos/Public/Chaos/Utilities.h
@@ -244,6 +244,17 @@ namespace Chaos
 			    L.M[0][2] * R.X + L.M[1][2] * R.Y + L.M[2][2] * R.Z);
 		}
 
+#if !UE_LARGE_WORLD_COORDINATES_DISABLED
+		inline TVec3<FRealSingle> Multiply(const TMatrix33<FRealSingle>& L, const TVec3<FRealSingle>& R)
+		{
+			// @todo(chaos): optimize: use simd
+			return TVec3<FRealSingle>(
+				L.M[0][0] * R.X + L.M[1][0] * R.Y + L.M[2][0] * R.Z,
+				L.M[0][1] * R.X + L.M[1][1] * R.Y + L.M[2][1] * R.Z,
+				L.M[0][2] * R.X + L.M[1][2] * R.Y + L.M[2][2] * R.Z);
+		}
+#endif
+
 		inline FVec4 Multiply(const FMatrix44& L, const FVec4& R)
 		{
 			// @todo(chaos): optimize: use simd
diff --git a/Engine/Source/Runtime/Experimental/ChaosCore/Public/Chaos/Core.h b/Engine/Source/Runtime/Experimental/ChaosCore/Public/Chaos/Core.h
--- a/Engine/Source/Runtime/Experimental/ChaosCore/Public/Chaos/Core.h
+++ b/Engine/Source/Runtime/Experimental/ChaosCore/Public/Chaos/Core.h
@@ -32,6 +32,12 @@ namespace Chaos
 	template <typename T>
 	using TVec4 = TVector<T, 4>;
 
+	template<typename T>
+	using TRotation3 = TRotation<T, 3>;
+
+	template<typename T>
+	using TMatrix33 = PMatrix<T, 3, 3>;
+
 	// NOTE: if you get a merge conflict on the GUID, you must replace it with a new GUID - do not accept the source or target
 	// or you will likely get DDC version conflicts resulting in crashes during load.
 	// Core version string for Chaos data. Any DDC builder dependent on Chaos for serialization should depend on this version
diff --git a/Engine/Source/Runtime/Experimental/ChaosCore/Public/Chaos/Matrix.h b/Engine/Source/Runtime/Experimental/ChaosCore/Public/Chaos/Matrix.h
--- a/Engine/Source/Runtime/Experimental/ChaosCore/Public/Chaos/Matrix.h
+++ b/Engine/Source/Runtime/Experimental/ChaosCore/Public/Chaos/Matrix.h
@@ -560,4 +560,401 @@ namespace Chaos
 		static const PMatrix<FReal, 3, 3> Zero;
 		static const PMatrix<FReal, 3, 3> Identity;
 	};
+
+#if !UE_LARGE_WORLD_COORDINATES_DISABLED
+	template<>
+	class PMatrix<FRealSingle, 3, 3> : public UE::Math::TMatrix<FRealSingle>
+	{
+	public:
+		PMatrix()
+			: UE::Math::TMatrix<FRealSingle>() {}
+		PMatrix(UE::Math::TMatrix<FRealSingle>&& Other)
+			: UE::Math::TMatrix<FRealSingle>(MoveTemp(Other)) {}
+		PMatrix(const UE::Math::TMatrix<FRealSingle>& Other)
+			: UE::Math::TMatrix<FRealSingle>((UE::Math::TMatrix<FRealSingle>)Other) {}
+		PMatrix(const UE::Math::TMatrix<FRealDouble>& Other)
+			: UE::Math::TMatrix<FRealSingle>((UE::Math::TMatrix<FRealSingle>)Other) {}
+		PMatrix(const FRealSingle x00, const FRealSingle x11, const FRealSingle x22)
+			: UE::Math::TMatrix<FRealSingle>()
+		{
+			M[0][0] = x00;
+			M[1][0] = 0;
+			M[2][0] = 0;
+			M[0][1] = 0;
+			M[1][1] = x11;
+			M[2][1] = 0;
+			M[0][2] = 0;
+			M[1][2] = 0;
+			M[2][2] = x22;
+			// Fill in the remainder with reasonable values.
+			M[3][0] = 0;
+			M[3][1] = 0;
+			M[3][2] = 0;
+			M[3][3] = 1;
+			M[0][3] = 0;
+			M[1][3] = 0;
+			M[2][3] = 0;
+		}
+		PMatrix(const FRealSingle x00, const FRealSingle x10, const FRealSingle x20, const FRealSingle x11, const FRealSingle x21, const FRealSingle x22)
+			: UE::Math::TMatrix<FRealSingle>()
+		{
+			M[0][0] = x00;
+			M[1][0] = x10;
+			M[2][0] = x20;
+			M[0][1] = x10;
+			M[1][1] = x11;
+			M[2][1] = x21;
+			M[0][2] = x20;
+			M[1][2] = x21;
+			M[2][2] = x22;
+			// Fill in the remainder with reasonable values.
+			M[3][0] = 0;
+			M[3][1] = 0;
+			M[3][2] = 0;
+			M[3][3] = 1;
+			M[0][3] = 0;
+			M[1][3] = 0;
+			M[2][3] = 0;
+		}
+		PMatrix(const FRealSingle x00, const FRealSingle x10, const FRealSingle x20, const FRealSingle x01, const FRealSingle x11, const FRealSingle x21, const FRealSingle x02, const FRealSingle x12, const FRealSingle x22)
+			: UE::Math::TMatrix<FRealSingle>()
+		{
+			M[0][0] = x00;
+			M[1][0] = x10;
+			M[2][0] = x20;
+			M[0][1] = x01;
+			M[1][1] = x11;
+			M[2][1] = x21;
+			M[0][2] = x02;
+			M[1][2] = x12;
+			M[2][2] = x22;
+			// Fill in the remainder with reasonable values.
+			M[3][0] = 0;
+			M[3][1] = 0;
+			M[3][2] = 0;
+			M[3][3] = 1;
+			M[0][3] = 0;
+			M[1][3] = 0;
+			M[2][3] = 0;
+		}
+		PMatrix(const FRealSingle x)
+			: UE::Math::TMatrix<FRealSingle>()
+		{
+			M[0][0] = x;
+			M[1][0] = x;
+			M[2][0] = x;
+			M[0][1] = x;
+			M[1][1] = x;
+			M[2][1] = x;
+			M[0][2] = x;
+			M[1][2] = x;
+			M[2][2] = x;
+			// Fill in the remainder with reasonable values.
+			M[3][0] = 0;
+			M[3][1] = 0;
+			M[3][2] = 0;
+			M[3][3] = 1;
+			M[0][3] = 0;
+			M[1][3] = 0;
+			M[2][3] = 0;
+		}
+		PMatrix(const TVector<FRealSingle, 3>& C1, const TVector<FRealSingle, 3>& C2, const TVector<FRealSingle, 3>& C3)
+		{
+			M[0][0] = C1.X;
+			M[1][0] = C1.Y;
+			M[2][0] = C1.Z;
+			M[0][1] = C2.X;
+			M[1][1] = C2.Y;
+			M[2][1] = C2.Z;
+			M[0][2] = C3.X;
+			M[1][2] = C3.Y;
+			M[2][2] = C3.Z;
+			// Fill in the remainder with reasonable values.
+			M[3][0] = 0;
+			M[3][1] = 0;
+			M[3][2] = 0;
+			M[3][3] = 1;
+			M[0][3] = 0;
+			M[1][3] = 0;
+			M[2][3] = 0;
+		}
+#if COMPILE_WITHOUT_UNREAL_SUPPORT
+		PMatrix<FRealSingle, 3, 3> GetTransposed()
+		{
+			return PMatrix<FRealSingle, 3, 3>(M[0][0], M[0][1], M[0][2], M[1][0], M[1][1], M[1][2], M[2][0], M[2][1], M[2][2]);
+		}
+		FRealSingle Determinant()
+		{
+			return M[0][0] * (M[1][1] * M[2][2] - M[1][2] * M[2][1]) - M[0][1] * (M[1][0] * M[2][2] - M[1][2] * M[2][0]) + M[0][2] * (M[1][0] * M[2][1] - M[1][1] * M[2][0]);
+		}
+		PMatrix<FRealSingle, 3, 3>& operator+=(const PMatrix<FRealSingle, 3, 3>& Other)
+		{
+			M[0][0] += Other.M[0][0];
+			M[0][1] += Other.M[0][1];
+			M[0][2] += Other.M[0][2];
+			M[1][0] += Other.M[1][0];
+			M[1][1] += Other.M[1][1];
+			M[1][2] += Other.M[1][2];
+			M[2][0] += Other.M[2][0];
+			M[2][1] += Other.M[2][1];
+			M[2][2] += Other.M[2][2];
+			return *this;
+		}
+#endif
+		// TDOD(mlentine): This should really be a vector multiply and sum for each entry using sse
+		TVector<FRealSingle, 3> operator*(const TVector<FRealSingle, 3>& Other) const
+		{
+			return TVector<FRealSingle, 3>(
+				M[0][0] * Other[0] + M[0][1] * Other[1] + M[0][2] * Other[2],
+				M[1][0] * Other[0] + M[1][1] * Other[1] + M[1][2] * Other[2],
+				M[2][0] * Other[0] + M[2][1] * Other[1] + M[2][2] * Other[2]);
+		}
+		PMatrix<FRealSingle, 3, 3> operator+(const PMatrix<FRealSingle, 3, 3>& Other) const
+		{
+			return PMatrix<FRealSingle, 3, 3>(
+				M[0][0] + Other.M[0][0],
+				M[1][0] + Other.M[1][0],
+				M[2][0] + Other.M[2][0],
+				M[0][1] + Other.M[0][1],
+				M[1][1] + Other.M[1][1],
+				M[2][1] + Other.M[2][1],
+				M[0][2] + Other.M[0][2],
+				M[1][2] + Other.M[1][2],
+				M[2][2] + Other.M[2][2]);
+		}
+		friend PMatrix<FRealSingle, 3, 3> operator+(const PMatrix<FRealSingle, 3, 3>& Other)
+		{
+			return Other;
+		}
+		PMatrix<FRealSingle, 3, 3> operator-(const PMatrix<FRealSingle, 3, 3>& Other) const
+		{
+			return PMatrix<FRealSingle, 3, 3>(
+				M[0][0] - Other.M[0][0],
+				M[1][0] - Other.M[1][0],
+				M[2][0] - Other.M[2][0],
+				M[0][1] - Other.M[0][1],
+				M[1][1] - Other.M[1][1],
+				M[2][1] - Other.M[2][1],
+				M[0][2] - Other.M[0][2],
+				M[1][2] - Other.M[1][2],
+				M[2][2] - Other.M[2][2]);
+		}
+		friend PMatrix<FRealSingle, 3, 3> operator-(const PMatrix<FRealSingle, 3, 3>& Other)
+		{
+			return PMatrix<FRealSingle, 3, 3>(
+				-Other.M[0][0],
+				-Other.M[1][0],
+				-Other.M[2][0],
+				-Other.M[0][1],
+				-Other.M[1][1],
+				-Other.M[2][1],
+				-Other.M[0][2],
+				-Other.M[1][2],
+				-Other.M[2][2]);
+		}
+		PMatrix<FRealSingle, 3, 3> operator*(const PMatrix<FRealSingle, 3, 3>& Other) const
+		{
+			return static_cast<const UE::Math::TMatrix<FRealSingle>*>(this)->operator*(static_cast<const UE::Math::TMatrix<FRealSingle>&>(Other));
+		}
+		PMatrix<FRealSingle, 3, 3> operator*(const FRealSingle Other) const
+		{
+			return PMatrix<FRealSingle, 3, 3>(
+				M[0][0] * Other,
+				M[1][0] * Other,
+				M[2][0] * Other,
+				M[0][1] * Other,
+				M[1][1] * Other,
+				M[2][1] * Other,
+				M[0][2] * Other,
+				M[1][2] * Other,
+				M[2][2] * Other);
+		}
+		friend PMatrix<FRealSingle, 3, 3> operator*(const FRealSingle OtherF, const PMatrix<FRealSingle, 3, 3>& OtherM)
+		{
+			return OtherM * OtherF;
+		}
+		PMatrix<FRealSingle, 3, 3> SubtractDiagonal(const FRealSingle Scalar) const
+		{
+			return PMatrix<FRealSingle, 3, 3>(
+				M[0][0] - Scalar,
+				M[1][0],
+				M[2][0],
+				M[0][1],
+				M[1][1] - Scalar,
+				M[2][1],
+				M[0][2],
+				M[1][2],
+				M[2][2] - Scalar);
+		}
+		PMatrix<FRealSingle, 3, 3> SymmetricCofactorMatrix() const
+		{
+			return PMatrix<FRealSingle, 3, 3>(
+				M[1][1] * M[2][2] - M[2][1] * M[2][1],
+				M[2][1] * M[2][0] - M[1][0] * M[2][2],
+				M[1][0] * M[2][1] - M[1][1] * M[2][0],
+				M[0][0] * M[2][2] - M[2][0] * M[2][0],
+				M[1][0] * M[2][0] - M[0][0] * M[2][1],
+				M[0][0] * M[1][1] - M[1][0] * M[1][0]);
+		}
+		TVector<FRealSingle, 3> LargestColumnNormalized() const
+		{
+			FRealSingle m10 = M[1][0] * M[1][0];
+			FRealSingle m20 = M[2][0] * M[2][0];
+			FRealSingle m21 = M[2][1] * M[2][1];
+			FRealSingle c0 = M[0][0] * M[0][0] + m10 + m20;
+			FRealSingle c1 = m10 + M[1][1] * M[1][1] + m21;
+			FRealSingle c2 = m20 + m21 + M[2][2] * M[2][2];
+			if (c0 > c1 && c0 > c2)
+			{
+				return TVector<FRealSingle, 3>(M[0][0], M[1][0], M[2][0]) / sqrt(c0);
+			}
+			if (c1 > c2)
+			{
+				return TVector<FRealSingle, 3>(M[1][0], M[1][1], M[2][1]) / sqrt(c1);
+			}
+			if (c2 > 0)
+			{
+				return TVector<FRealSingle, 3>(M[2][0], M[2][1], M[2][2]) / sqrt(c2);
+			}
+			return TVector<FRealSingle, 3>(1, 0, 0);
+		}
+
+		/**
+		 * Get the specified axis (0-indexed, X,Y,Z).
+		 * @note: we are treating matrices as column major, so axis elements are sequential in memory
+		 */
+		FORCEINLINE TVector<FRealSingle, 3> GetAxis(int32 AxisIndex) const
+		{
+			return TVector<FRealSingle, 3>(M[AxisIndex][0], M[AxisIndex][1], M[AxisIndex][2]);
+		}
+
+		/**
+		 * Set the specified axis (0-indexed, X,Y,Z).
+		 * @note: we are treating matrices as column major, so axis elements are sequential in memory
+		 */
+		FORCEINLINE void SetAxis(int32 AxisIndex, const TVector<FRealSingle, 3>& Axis)
+		{
+			M[AxisIndex][0] = Axis.X;
+			M[AxisIndex][1] = Axis.Y;
+			M[AxisIndex][2] = Axis.Z;
+			M[AxisIndex][3] = 0;
+		}
+
+		/**
+		 * Get the specified row (0-indexed, X,Y,Z).
+		 * @note: we are treating matrices as column major, so rows are not sequential in memory
+		 * @seealso GetAxis, GetColumn
+		 */
+		FORCEINLINE TVector<FRealSingle, 3> GetRow(int32 RowIndex) const
+		{
+			return TVector<FRealSingle, 3>(M[0][RowIndex], M[1][RowIndex], M[2][RowIndex]);
+		}
+
+		/**
+		 * Set the specified row.
+		 * @note: we are treating matrices as column major, so axis elements are sequential in memory
+		 * @seealso SetAxis, SetColumn
+		 */
+		FORCEINLINE void SetRow(int32 RowIndex, const TVector<FRealSingle, 3>& V)
+		{
+			M[0][RowIndex] = V.X;
+			M[1][RowIndex] = V.Y;
+			M[2][RowIndex] = V.Z;
+			M[3][RowIndex] = 0;
+		}
+
+		/**
+		 * Get the specified column (0-indexed, X,Y,Z). Equivalent to GetAxis.
+		 * @note: we are treating matrices as column major, so columns are sequential in memory
+		 * @seealso GetAxis, GetRow
+		 */
+		FORCEINLINE TVector<FRealSingle, 3> GetColumn(int32 ColumnIndex) const
+		{
+			return GetAxis(ColumnIndex);
+		}
+
+		/**
+		 * Set the specified column. Equivalent to SetAxis.
+		 * @note: we are treating matrices as column major, so axis elements are sequential in memory
+		 * @seealso SetAxis, SetRow
+		 */
+		FORCEINLINE void SetColumn(int32 ColumnIndex, const TVector<FRealSingle, 3>& V)
+		{
+			SetAxis(ColumnIndex, V);
+		}
+
+		/**
+		 * Get the diagonal elements as a vector.
+		 */
+		FORCEINLINE TVector<FRealSingle, 3> GetDiagonal() const
+		{
+			return TVector<FRealSingle, 3>(M[0][0], M[1][1], M[2][2]);
+		}
+
+		FORCEINLINE FRealSingle GetAt(int32 RowIndex, int32 ColIndex) const
+		{
+			return M[ColIndex][RowIndex];
+		}
+
+		FORCEINLINE void SetAt(int32 RowIndex, int32 ColIndex, FRealSingle V)
+		{
+			M[ColIndex][RowIndex] = V;
+		}
+
+		/**
+		 * Return a diagonal matrix with the specified elements
+		 */
+		static PMatrix<FRealSingle, 3, 3> FromDiagonal(const TVector<FRealSingle, 3>& D)
+		{
+			return PMatrix<FRealSingle, 3, 3>(D.X, D.Y, D.Z);
+		}
+
+#if COMPILE_WITHOUT_UNREAL_SUPPORT
+		// TODO(mlentine): Document which one is row and which one is column
+		PMatrix<FRealSingle, 3, 3> operator*(const PMatrix<FRealSingle, 3, 3>& Other)
+		{
+			return PMatrix<FRealSingle, 3, 3>(
+				M[0][0] * Other.M[0][0] + M[0][1] * Other.M[1][0] + M[0][2] * Other.M[2][0],
+				M[1][0] * Other.M[0][0] + M[1][1] * Other.M[1][0] + M[1][2] * Other.M[2][0],
+				M[2][0] * Other.M[0][0] + M[2][1] * Other.M[1][0] + M[2][2] * Other.M[2][0],
+				M[0][0] * Other.M[0][1] + M[0][1] * Other.M[1][1] + M[0][2] * Other.M[2][1],
+				M[1][0] * Other.M[0][1] + M[1][1] * Other.M[1][1] + M[1][2] * Other.M[2][1],
+				M[2][0] * Other.M[0][1] + M[2][1] * Other.M[1][1] + M[2][2] * Other.M[2][1],
+				M[0][0] * Other.M[0][2] + M[0][1] * Other.M[1][2] + M[0][2] * Other.M[2][2],
+				M[1][0] * Other.M[0][2] + M[1][1] * Other.M[1][2] + M[1][2] * Other.M[2][2],
+				M[2][0] * Other.M[0][2] + M[2][1] * Other.M[1][2] + M[2][2] * Other.M[2][2]);
+		}
+		PMatrix<FRealSingle, 3, 3> operator*(const FRealSingle Scalar)
+		{
+			return PMatrix<FRealSingle, 3, 3>(
+				M[0][0] * Scalar + M[0][1] * Scalar + M[0][2] * Scalar,
+				M[1][0] * Scalar + M[1][1] * Scalar + M[1][2] * Scalar,
+				M[2][0] * Scalar + M[2][1] * Scalar + M[2][2] * Scalar,
+				M[0][0] * Scalar + M[0][1] * Scalar + M[0][2] * Scalar,
+				M[1][0] * Scalar + M[1][1] * Scalar + M[1][2] * Scalar,
+				M[2][0] * Scalar + M[2][1] * Scalar + M[2][2] * Scalar,
+				M[0][0] * Scalar + M[0][1] * Scalar + M[0][2] * Scalar,
+				M[1][0] * Scalar + M[1][1] * Scalar + M[1][2] * Scalar,
+				M[2][0] * Scalar + M[2][1] * Scalar + M[2][2] * Scalar);
+		}
+#endif
+		inline bool Equals(const PMatrix<FRealSingle, 3, 3>& Other, FRealSingle Tolerance = KINDA_SMALL_NUMBER) const
+		{
+			return true
+				&& (FMath::Abs(Other.M[0][0] - M[0][0]) <= Tolerance)
+				&& (FMath::Abs(Other.M[0][1] - M[0][1]) <= Tolerance)
+				&& (FMath::Abs(Other.M[0][2] - M[0][2]) <= Tolerance)
+				&& (FMath::Abs(Other.M[1][0] - M[1][0]) <= Tolerance)
+				&& (FMath::Abs(Other.M[1][1] - M[1][1]) <= Tolerance)
+				&& (FMath::Abs(Other.M[1][2] - M[1][2]) <= Tolerance)
+				&& (FMath::Abs(Other.M[2][0] - M[2][0]) <= Tolerance)
+				&& (FMath::Abs(Other.M[2][1] - M[2][1]) <= Tolerance)
+				&& (FMath::Abs(Other.M[2][2] - M[2][2]) <= Tolerance);
+		}
+
+
+		static const PMatrix<FRealSingle, 3, 3> Zero;
+		static const PMatrix<FRealSingle, 3, 3> Identity;
+	};
+#endif
 }
\ No newline at end of file
