Chaos : Clustering
- Removed template parameters from FRigidClustering

diff --git a/Engine/Source/Programs/HeadlessChaos/Private/GeometryCollection/GeometryCollectionTestClustering.cpp b/Engine/Source/Programs/HeadlessChaos/Private/GeometryCollection/GeometryCollectionTestClustering.cpp
--- a/Engine/Source/Programs/HeadlessChaos/Private/GeometryCollection/GeometryCollectionTestClustering.cpp
+++ b/Engine/Source/Programs/HeadlessChaos/Private/GeometryCollection/GeometryCollectionTestClustering.cpp
@@ -36,7 +36,7 @@ namespace GeometryCollectionTest
 {
 	using namespace ChaosTest;
 
-	bool ClusterMapContains(const Chaos::TPBDRigidClustering<FPBDRigidsEvolution, FPBDCollisionConstraints>::FClusterMap& ClusterMap, const FPBDRigidParticleHandle* InKey, TArray<FPBDRigidParticleHandle*> Elements)
+	bool ClusterMapContains(const Chaos::FRigidClustering::FClusterMap& ClusterMap, const FPBDRigidParticleHandle* InKey, TArray<FPBDRigidParticleHandle*> Elements)
 	{
 		if (ClusterMap.Num())
 		{
@@ -880,9 +880,8 @@ namespace GeometryCollectionTest
 		TArray<Chaos::TPBDRigidClusteredParticleHandle<FReal, 3>*>& ParticleHandles = Collection->PhysObject->GetSolverParticleHandles();
 		TArray<Chaos::TPBDRigidClusteredParticleHandle<FReal, 3>*>& ClusterHandles = Collection->PhysObject->GetSolverClusterHandles();
 
-		using FClustering = TPBDRigidClustering<FPBDRigidsEvolutionGBF, FPBDCollisionConstraints>;
-		FClustering& Clustering = UnitTest.Solver->GetEvolution()->GetRigidClustering();
-		const FClustering::FClusterMap& ClusterMap = Clustering.GetChildrenMap();
+		FRigidClustering& Clustering = UnitTest.Solver->GetEvolution()->GetRigidClustering();
+		const FRigidClustering::FClusterMap& ClusterMap = Clustering.GetChildrenMap();
 
 		// Verify that the parent-child relationship is reflected in the clustering hierarchy
 		// Tree should be:
@@ -1180,9 +1179,8 @@ namespace GeometryCollectionTest
 		TArray<Chaos::TPBDRigidClusteredParticleHandle<FReal, 3>*>& ParticleHandles = Collection->PhysObject->GetSolverParticleHandles();
 		TArray<Chaos::TPBDRigidClusteredParticleHandle<FReal, 3>*>& ClusterHandles = Collection->PhysObject->GetSolverClusterHandles();
 
-		using FClustering = TPBDRigidClustering<FPBDRigidsEvolutionGBF, FPBDCollisionConstraints>;
-		FClustering& Clustering = UnitTest.Solver->GetEvolution()->GetRigidClustering();
-		const FClustering::FClusterMap& ClusterMap = Clustering.GetChildrenMap();
+		FRigidClustering& Clustering = UnitTest.Solver->GetEvolution()->GetRigidClustering();
+		const FRigidClustering::FClusterMap& ClusterMap = Clustering.GetChildrenMap();
 
 		// Verify that the parent-child relationship is reflected in the clustering hierarchy
 		// Tree should be:
@@ -2063,9 +2061,8 @@ namespace GeometryCollectionTest
 
 		UnitTest.Advance();
 
-		using FClustering = TPBDRigidClustering<FPBDRigidsEvolutionGBF, FPBDCollisionConstraints>;
-		FClustering& Clustering = UnitTest.Solver->GetEvolution()->GetRigidClustering();
-		const FClustering::FClusterMap& ClusterMap = Clustering.GetChildrenMap();
+		FRigidClustering& Clustering = UnitTest.Solver->GetEvolution()->GetRigidClustering();
+		const FRigidClustering::FClusterMap& ClusterMap = Clustering.GetChildrenMap();
 		const Chaos::TArrayCollectionArray<Chaos::ClusterId>& ClusterIdsArray = Clustering.GetClusterIdsArray();
 
 		UnitTest.Solver->RegisterSimOneShotCallback([&]()
@@ -2255,7 +2252,7 @@ namespace GeometryCollectionTest
 		FalloffField->Falloff = EFieldFalloffType::Field_FallOff_None;
 
 		TArray<Chaos::TPBDRigidClusteredParticleHandle<FReal, 3>*>& ParticleHandles = Collection->PhysObject->GetSolverParticleHandles();
-		Chaos::TPBDRigidClustering<Chaos::FPBDRigidsEvolutionGBF, Chaos::FPBDCollisionConstraints>& Clustering = UnitTest.Solver->GetEvolution()->GetRigidClustering();
+		Chaos::FRigidClustering& Clustering = UnitTest.Solver->GetEvolution()->GetRigidClustering();
 		const auto& ClusterMap = Clustering.GetChildrenMap();
 		auto& x = UnitTest.Solver->GetEvolution()->GetRigidClustering().GetStrainArray();
 
diff --git a/Engine/Source/Runtime/Experimental/Chaos/Private/Chaos/PBDRigidClustering.cpp b/Engine/Source/Runtime/Experimental/Chaos/Private/Chaos/PBDRigidClustering.cpp
--- a/Engine/Source/Runtime/Experimental/Chaos/Private/Chaos/PBDRigidClustering.cpp
+++ b/Engine/Source/Runtime/Experimental/Chaos/Private/Chaos/PBDRigidClustering.cpp
@@ -358,8 +358,7 @@ namespace Chaos
 	// TPBDRigidClustering
 	//==========================================================================
 
-	template<class T_FPBDRigidsEvolution, class T_FPBDCollisionConstraint>
-	TPBDRigidClustering<T_FPBDRigidsEvolution, T_FPBDCollisionConstraint>::TPBDRigidClustering(T_FPBDRigidsEvolution& InEvolution, FPBDRigidClusteredParticles& InParticles)
+	FRigidClustering::FRigidClustering(FPBDRigidsEvolution& InEvolution, FPBDRigidClusteredParticles& InParticles)
 		: MEvolution(InEvolution)
 		, MParticles(InParticles)
 		, MCollisionImpulseArrayDirty(true)
@@ -368,13 +367,17 @@ namespace Chaos
 		, MClusterUnionConnectionType(FClusterCreationParameters::EConnectionMethod::DelaunayTriangulation)
 	{}
 
-	template<class T_FPBDRigidsEvolution, class T_FPBDCollisionConstraint>
-	TPBDRigidClustering<T_FPBDRigidsEvolution, T_FPBDCollisionConstraint>::~TPBDRigidClustering()
+	FRigidClustering::~FRigidClustering()
 	{}
 
 	DECLARE_CYCLE_STAT(TEXT("TPBDRigidClustering<>::CreateClusterParticle"), STAT_CreateClusterParticle, STATGROUP_Chaos);
-	template<class T_FPBDRigidsEvolution, class T_FPBDCollisionConstraint>
-	Chaos::FPBDRigidClusteredParticleHandle* TPBDRigidClustering<T_FPBDRigidsEvolution, T_FPBDCollisionConstraint>::CreateClusterParticle(const int32 ClusterGroupIndex, TArray<Chaos::FPBDRigidParticleHandle*>&& Children, const FClusterCreationParameters& Parameters, TSharedPtr<Chaos::FImplicitObject, ESPMode::ThreadSafe> ProxyGeometry, const FRigidTransform3* ForceMassOrientation, const FUniqueIdx* ExistingIndex)
+	Chaos::FPBDRigidClusteredParticleHandle* FRigidClustering::CreateClusterParticle(
+		const int32 ClusterGroupIndex, 
+		TArray<Chaos::FPBDRigidParticleHandle*>&& Children, 
+		const FClusterCreationParameters& Parameters, 
+		TSharedPtr<Chaos::FImplicitObject, ESPMode::ThreadSafe> ProxyGeometry, 
+		const FRigidTransform3* ForceMassOrientation, 
+		const FUniqueIdx* ExistingIndex)
 	{
 		SCOPE_CYCLE_COUNTER(STAT_CreateClusterParticle);
 
@@ -463,9 +466,8 @@ namespace Chaos
 	FAutoConsoleVariableRef CVarUnionsHaveCollisionParticles(TEXT("p.UnionsHaveCollisionParticles"), UnionsHaveCollisionParticles, TEXT(""));
 
 	DECLARE_CYCLE_STAT(TEXT("TPBDRigidClustering<>::CreateClusterParticleFromClusterChildren"), STAT_CreateClusterParticleFromClusterChildren, STATGROUP_Chaos);
-	template<class T_FPBDRigidsEvolution, class T_FPBDCollisionConstraint>
 	Chaos::FPBDRigidClusteredParticleHandle* 
-	TPBDRigidClustering<T_FPBDRigidsEvolution, T_FPBDCollisionConstraint>::CreateClusterParticleFromClusterChildren(
+	FRigidClustering::CreateClusterParticleFromClusterChildren(
 		TArray<FPBDRigidParticleHandle*>&& Children, 
 		FPBDRigidClusteredParticleHandle* Parent, 
 		const FRigidTransform3& ClusterWorldTM, 
@@ -540,8 +542,8 @@ namespace Chaos
 	}
 
 	DECLARE_CYCLE_STAT(TEXT("TPBDRigidClustering<>::UnionClusterGroups"), STAT_UnionClusterGroups, STATGROUP_Chaos);
-	template<class T_FPBDRigidsEvolution, class T_FPBDCollisionConstraint>
-	void TPBDRigidClustering<T_FPBDRigidsEvolution, T_FPBDCollisionConstraint>::UnionClusterGroups()
+	void 
+	FRigidClustering::UnionClusterGroups()
 	{
 		SCOPE_CYCLE_COUNTER(STAT_UnionClusterGroups);
 
@@ -615,8 +617,8 @@ namespace Chaos
 
 
 	DECLARE_CYCLE_STAT(TEXT("TPBDRigidClustering<>::DeactivateClusterParticle"), STAT_DeactivateClusterParticle, STATGROUP_Chaos);
-	template<class FPBDRigidsEvolution, class FPBDCollisionConstraint>
-	TSet<FPBDRigidParticleHandle*> TPBDRigidClustering<FPBDRigidsEvolution, FPBDCollisionConstraint>::DeactivateClusterParticle(
+	TSet<FPBDRigidParticleHandle*> 
+	FRigidClustering::DeactivateClusterParticle(
 		FPBDRigidClusteredParticleHandle* ClusteredParticle)
 	{
 		SCOPE_CYCLE_COUNTER(STAT_DeactivateClusterParticle);
@@ -631,9 +633,8 @@ namespace Chaos
 	}
 
 	DECLARE_CYCLE_STAT(TEXT("TPBDRigidClustering<>::ReleaseClusterParticles(STRAIN)"), STAT_ReleaseClusterParticles_STRAIN, STATGROUP_Chaos);
-	template<class T_FPBDRigidsEvolution, class T_FPBDCollisionConstraint>
 	TSet<FPBDRigidParticleHandle*> 
-	TPBDRigidClustering<T_FPBDRigidsEvolution, T_FPBDCollisionConstraint>::ReleaseClusterParticles(
+	FRigidClustering::ReleaseClusterParticles(
 		FPBDRigidClusteredParticleHandle* ClusteredParticle,
 		const TMap<FGeometryParticleHandle*, Chaos::FReal>* ExternalStrainMap,
 		bool bForceRelease)
@@ -903,12 +904,11 @@ namespace Chaos
 
 
 	DECLARE_CYCLE_STAT(TEXT("TPBDRigidClustering<>::ReleaseClusterParticlesNoInternalCluster"), STAT_ReleaseClusterParticlesNoInternalCluster, STATGROUP_Chaos);
-	template<class T_FPBDRigidsEvolution, class T_FPBDCollisionConstraint>
 	TSet<FPBDRigidParticleHandle*>
-		TPBDRigidClustering<T_FPBDRigidsEvolution, T_FPBDCollisionConstraint>::ReleaseClusterParticlesNoInternalCluster(
-			FPBDRigidClusteredParticleHandle* ClusteredParticle,
-			const TMap<FGeometryParticleHandle*, Chaos::FReal>* ExternalStrainMap,
-			bool bForceRelease)
+	FRigidClustering::ReleaseClusterParticlesNoInternalCluster(
+		FPBDRigidClusteredParticleHandle* ClusteredParticle,
+		const TMap<FGeometryParticleHandle*, Chaos::FReal>* ExternalStrainMap,
+		bool bForceRelease)
 	{
 		/* This is a near duplicate of the ReleaseClusterParticles() method with the internal cluster creation removed.
 		*  This method should be used exclusively by the GeometryCollectionComponentCacheAdaptor in order to implement
@@ -1130,8 +1130,8 @@ namespace Chaos
 
 
 	DECLARE_CYCLE_STAT(TEXT("TPBDRigidClustering<>::ReleaseClusterParticles(LIST)"), STAT_ReleaseClusterParticles_LIST, STATGROUP_Chaos);
-	template<class FPBDRigidsEvolution, class FPBDCollisionConstraint>
-	TSet<FPBDRigidParticleHandle*> TPBDRigidClustering<FPBDRigidsEvolution, FPBDCollisionConstraint>::ReleaseClusterParticles(
+	TSet<FPBDRigidParticleHandle*> 
+	FRigidClustering::ReleaseClusterParticles(
 		TArray<FPBDRigidParticleHandle*> ChildrenParticles)
 	{
 		SCOPE_CYCLE_COUNTER(STAT_ReleaseClusterParticles_LIST);
@@ -1178,11 +1178,10 @@ namespace Chaos
 	DECLARE_CYCLE_STAT(TEXT("TPBDRigidClustering<>::Update Impulse from Strain"), STAT_UpdateImpulseStrain, STATGROUP_Chaos);
 	DECLARE_CYCLE_STAT(TEXT("TPBDRigidClustering<>::Update Dirty Impulses"), STAT_UpdateDirtyImpulses, STATGROUP_Chaos);
 	DECLARE_CYCLE_STAT(TEXT("TPBDRigidClustering<>::Rewind"), STAT_ClusterRewind, STATGROUP_Chaos);
-	template<class T_FPBDRigidsEvolution, class T_FPBDCollisionConstraint>
 	void 
-	TPBDRigidClustering<T_FPBDRigidsEvolution, T_FPBDCollisionConstraint>::AdvanceClustering(
+	FRigidClustering::AdvanceClustering(
 		const FReal Dt, 
-		T_FPBDCollisionConstraint& CollisionRule)
+		FPBDCollisionConstraints& CollisionRule)
 	{
 		SCOPE_CYCLE_COUNTER(STAT_AdvanceClustering);
 		UE_LOG(LogChaos, Verbose, TEXT("START FRAME with Dt %f"), Dt);
@@ -1309,9 +1308,8 @@ namespace Chaos
 	}
 
 	DECLARE_CYCLE_STAT(TEXT("TPBDRigidClustering<>::BreakingModel()"), STAT_BreakingModel, STATGROUP_Chaos);
-	template<class T_FPBDRigidsEvolution, class T_FPBDCollisionConstraint>
 	TMap<FPBDRigidClusteredParticleHandle*, TSet<FPBDRigidParticleHandle*>> 
-	TPBDRigidClustering<T_FPBDRigidsEvolution, T_FPBDCollisionConstraint>::BreakingModel(
+	FRigidClustering::BreakingModel(
 		TMap<FGeometryParticleHandle*, Chaos::FReal>* ExternalStrainMap)
 	{
 		SCOPE_CYCLE_COUNTER(STAT_BreakingModel);
@@ -1373,8 +1371,8 @@ namespace Chaos
 	}
 
 	DECLARE_CYCLE_STAT(TEXT("TPBDRigidClustering<>::PromoteStrains()"), STAT_PromoteStrains, STATGROUP_Chaos);
-	template<class T_FPBDRigidsEvolution, class T_FPBDCollisionConstraint>
-	FReal TPBDRigidClustering<T_FPBDRigidsEvolution, T_FPBDCollisionConstraint>::PromoteStrains(
+	FReal 
+	FRigidClustering::PromoteStrains(
 		FPBDRigidParticleHandle* CurrentNode)
 	{
 		SCOPE_CYCLE_COUNTER(STAT_PromoteStrains);
@@ -1399,8 +1397,8 @@ namespace Chaos
 	}
 
 	DECLARE_CYCLE_STAT(TEXT("TPBDRigidClustering<>::UpdateKinematicProperties()"), STAT_UpdateKinematicProperties, STATGROUP_Chaos);
-	template<class T_FPBDRigidsEvolution, class T_FPBDCollisionConstraint>
-	void TPBDRigidClustering<T_FPBDRigidsEvolution, T_FPBDCollisionConstraint>::UpdateKinematicProperties(
+	void 
+	FRigidClustering::UpdateKinematicProperties(
 		Chaos::FPBDRigidParticleHandle* InParent)
 	{
 		SCOPE_CYCLE_COUNTER(STAT_UpdateKinematicProperties);
@@ -1450,49 +1448,9 @@ namespace Chaos
 		}
 	}
 
-	DECLARE_CYCLE_STAT(TEXT("TPBDRigidClustering<>::SwapBufferedData"), STAT_SwapBufferedData, STATGROUP_Chaos);
-	template<class T_FPBDRigidsEvolution, class T_FPBDCollisionConstraint>
-	void TPBDRigidClustering<T_FPBDRigidsEvolution, T_FPBDCollisionConstraint>::SwapBufferedData()
-	{
-		check(false);
-		// TODO: Ryan - this code currently uses MParticles as the only source of clustered particles.
-		// As things stand, clustered particles and geometry collection particles are both of clustered
-		// type, but are stored separately.  Geometry collection particles exist on both game and physics
-		// threads, cluster particles only exist on the physics thread.
-/*
-		SCOPE_CYCLE_COUNTER(STAT_SwapBufferedData);
-		ResourceLock.WriteLock();
-		//BufferResource.MChildren.Reset();
-		//BufferResource.ClusterParentTransforms.Reset();	//todo: once everything is atomic this should get reset
-		const TArray<TSerializablePtr<FImplicitObject>>& AllGeom = MParticles.GetAllGeometry();
-		BufferResource.GeometryPtrs.SetNum(AllGeom.Num());
-
-		const auto& NonDisabledClusteredParticles = MEvolution.GetNonDisabledClusteredArray();
-		for (auto& ClusteredParticle : NonDisabledClusteredParticles)
-		{
-			//const TArray<uint32>* Children = MChildren.Find(ClusteredParticle);
-			TArray<FPBDRigidParticleHandle*>* Children = MChildren.Find(ClusteredParticle);
-			//if (Children && ClusteredParticle->ClusterIds().Id == INDEX_NONE) //root cluster so copy children
-			if (Children && ClusteredParticle->ClusterIds().Id == nullptr) //root cluster so copy children
-			{
-				//TODO: record GT particle pointer instead
-				//BufferResource.MChildren.Add(ClusteredParticle->TransientParticleIndex(), *Children);
-				BufferResource.MChildren.Add(
-					ClusteredParticle, *Children);
-				//BufferResource.ClusterParentTransforms.Add(ClusteredParticle->TransientParticleIndex(), TRigidTransform<FReal, 3>(ClusteredParticle->X(), ClusteredParticle->R()));
-				BufferResource.ClusterParentTransforms.Add(
-					ClusteredParticle, TRigidTransform<FReal, 3>(ClusteredParticle->X(), ClusteredParticle->R()));
-			}
-		}
-
-		BufferResource.GeometryPtrs = AllGeom; //in future this should be sparse. SQ has fallback that relies on potentially all geom so can't do it yet
-		ResourceLock.WriteUnlock();
-*/
-	}
-
 	DECLARE_CYCLE_STAT(TEXT("TPBDRigidClustering<>::GetActiveClusterIndex"), STAT_GetActiveClusterIndex, STATGROUP_Chaos);
-	template<class T_FPBDRigidsEvolution, class T_FPBDCollisionConstraint>
-	FPBDRigidParticleHandle* TPBDRigidClustering<T_FPBDRigidsEvolution, T_FPBDCollisionConstraint>::GetActiveClusterIndex(
+	FPBDRigidParticleHandle* 
+	FRigidClustering::GetActiveClusterIndex(
 		FPBDRigidParticleHandle* Child)
 	{
 		SCOPE_CYCLE_COUNTER(STAT_GetActiveClusterIndex);
@@ -1504,8 +1462,8 @@ namespace Chaos
 	}
 
 	DECLARE_CYCLE_STAT(TEXT("TPBDRigidClustering<>::GenerateConnectionGraph"), STAT_GenerateConnectionGraph, STATGROUP_Chaos);
-	template<class T_FPBDRigidsEvolution, class T_FPBDCollisionConstraint>
-	void TPBDRigidClustering<T_FPBDRigidsEvolution, T_FPBDCollisionConstraint>::GenerateConnectionGraph(
+	void 
+	FRigidClustering::GenerateConnectionGraph(
 		Chaos::FPBDRigidClusteredParticleHandle* Parent,
 		const FClusterCreationParameters& Parameters)
 	{
@@ -1549,8 +1507,8 @@ namespace Chaos
 	}
 
 	DECLARE_CYCLE_STAT(TEXT("TPBDRigidClustering<>::UpdateMassProperties"), STAT_UpdateMassProperties, STATGROUP_Chaos);
-	template<class T_FPBDRigidsEvolution, class T_FPBDCollisionConstraint>
-	void TPBDRigidClustering<T_FPBDRigidsEvolution, T_FPBDCollisionConstraint>::UpdateMassProperties(
+	void 
+	FRigidClustering::UpdateMassProperties(
 		Chaos::FPBDRigidClusteredParticleHandle* Parent, 
 		TSet<FPBDRigidParticleHandle*>& Children, 
 		const FRigidTransform3* ForceMassOrientation)
@@ -1569,8 +1527,8 @@ namespace Chaos
 	DECLARE_CYCLE_STAT(TEXT("TPBDRigidClustering<>::UpdateGeometry_CopyPoints"), STAT_UpdateGeometry_CopyPoints, STATGROUP_Chaos);
 	DECLARE_CYCLE_STAT(TEXT("TPBDRigidClustering<>::UpdateGeometry_PointsBVH"), STAT_UpdateGeometry_PointsBVH, STATGROUP_Chaos);
 
-	template<class T_FPBDRigidsEvolution, class T_FPBDCollisionConstraint>
-	void TPBDRigidClustering<T_FPBDRigidsEvolution, T_FPBDCollisionConstraint>::UpdateGeometry(
+	void 
+	FRigidClustering::UpdateGeometry(
 		Chaos::FPBDRigidClusteredParticleHandle* Parent, 
 		const TSet<FPBDRigidParticleHandle*>& Children, 
 		TSharedPtr<Chaos::FImplicitObject, ESPMode::ThreadSafe> ProxyGeometry,
@@ -1828,9 +1786,9 @@ namespace Chaos
 	FAutoConsoleVariableRef CVarMinContactSpeedForStrainEval(TEXT("p.chaos.MinContactSpeedForStrainEval"), MinContactSpeedForStrainEval, TEXT("Minimum speed at the contact before accumulating for strain eval "));
 
 	DECLARE_CYCLE_STAT(TEXT("ComputeStrainFromCollision"), STAT_ComputeStrainFromCollision, STATGROUP_Chaos);
-	template<class T_FPBDRigidsEvolution, class T_FPBDCollisionConstraint>
-	void TPBDRigidClustering<T_FPBDRigidsEvolution, T_FPBDCollisionConstraint>::ComputeStrainFromCollision(
-		const T_FPBDCollisionConstraint& CollisionRule)
+	void 
+	FRigidClustering::ComputeStrainFromCollision(
+		const FPBDCollisionConstraints& CollisionRule)
 	{
 		SCOPE_CYCLE_COUNTER(STAT_ComputeStrainFromCollision);
 		FClusterMap& MParentToChildren = GetChildrenMap();
@@ -1907,8 +1865,8 @@ namespace Chaos
 	}
 
 	DECLARE_CYCLE_STAT(TEXT("ResetCollisionImpulseArray"), STAT_ResetCollisionImpulseArray, STATGROUP_Chaos);
-	template<class T_FPBDRigidsEvolution, class T_FPBDCollisionConstraint>
-	void TPBDRigidClustering<T_FPBDRigidsEvolution, T_FPBDCollisionConstraint>::ResetCollisionImpulseArray()
+	void 
+	FRigidClustering::ResetCollisionImpulseArray()
 	{
 		SCOPE_CYCLE_COUNTER(STAT_ResetCollisionImpulseArray);
 		if (MCollisionImpulseArrayDirty)
@@ -1920,8 +1878,8 @@ namespace Chaos
 		}
 	}
 
-	template<class T_FPBDRigidsEvolution, class T_FPBDCollisionConstraint>
-	void TPBDRigidClustering<T_FPBDRigidsEvolution, T_FPBDCollisionConstraint>::DisableCluster(
+	void 
+	FRigidClustering::DisableCluster(
 		FPBDRigidClusteredParticleHandle* ClusteredParticle)
 	{
 		// #note: we don't recursively descend to the children
@@ -1933,8 +1891,8 @@ namespace Chaos
 		MActiveRemovalIndices.Remove(ClusteredParticle);
 	}
 
-	template<class T_FPBDRigidsEvolution, class T_FPBDCollisionConstraint>
-	void TPBDRigidClustering<T_FPBDRigidsEvolution, T_FPBDCollisionConstraint>::DisableParticleWithBreakEvent(
+	void 
+	FRigidClustering::DisableParticleWithBreakEvent(
 		Chaos::FPBDRigidParticleHandle* Particle)
 	{
 		DisableCluster(Particle->CastToClustered());
@@ -1965,9 +1923,8 @@ namespace Chaos
 		}
 	}
 
-	template<class T_FPBDRigidsEvolution, class T_FPBDCollisionConstraint>
 	FPBDRigidClusteredParticleHandle*
-	TPBDRigidClustering<T_FPBDRigidsEvolution, T_FPBDCollisionConstraint>::DestroyClusterParticle(
+	FRigidClustering::DestroyClusterParticle(
 		FPBDRigidClusteredParticleHandle* ClusteredParticle,
 		const FClusterDestoryParameters& Parameters)
 	{
@@ -2040,8 +1997,8 @@ namespace Chaos
 
 
 	DECLARE_CYCLE_STAT(TEXT("TPBDRigidClustering<>::UpdateConnectivityGraphUsingPointImplicit"), STAT_UpdateConnectivityGraphUsingPointImplicit, STATGROUP_Chaos);
-	template<class T_FPBDRigidsEvolution, class T_FPBDCollisionConstraint>
-	void TPBDRigidClustering<T_FPBDRigidsEvolution, T_FPBDCollisionConstraint>::UpdateConnectivityGraphUsingPointImplicit(
+	void 
+	FRigidClustering::UpdateConnectivityGraphUsingPointImplicit(
 		Chaos::FPBDRigidClusteredParticleHandle* Parent,
 		const FClusterCreationParameters& Parameters)
 	{
@@ -2113,8 +2070,8 @@ namespace Chaos
 	}
 
 	DECLARE_CYCLE_STAT(TEXT("TPBDRigidClustering<>::FixConnectivityGraphUsingDelaunayTriangulation"), STAT_FixConnectivityGraphUsingDelaunayTriangulation, STATGROUP_Chaos);
-	template<class T_FPBDRigidsEvolution, class T_FPBDCollisionConstraint>
-	void TPBDRigidClustering<T_FPBDRigidsEvolution, T_FPBDCollisionConstraint>::FixConnectivityGraphUsingDelaunayTriangulation(
+	void 
+	FRigidClustering::FixConnectivityGraphUsingDelaunayTriangulation(
 		Chaos::FPBDRigidClusteredParticleHandle* Parent,
 		const FClusterCreationParameters& Parameters)
 	{
@@ -2267,9 +2224,10 @@ namespace Chaos
 	}
 
 	DECLARE_CYCLE_STAT(TEXT("TPBDRigidClustering<>::UpdateConnectivityGraphUsingDelaunayTriangulation"), STAT_UpdateConnectivityGraphUsingDelaunayTriangulation, STATGROUP_Chaos);
-
-	template<class T_FPBDRigidsEvolution, class T_FPBDCollisionConstraint>
-	void TPBDRigidClustering<T_FPBDRigidsEvolution, T_FPBDCollisionConstraint>::UpdateConnectivityGraphUsingDelaunayTriangulation(Chaos::FPBDRigidClusteredParticleHandle* Parent, const FClusterCreationParameters& Parameters)
+	void 
+	FRigidClustering::UpdateConnectivityGraphUsingDelaunayTriangulation(
+		Chaos::FPBDRigidClusteredParticleHandle* Parent, 
+		const FClusterCreationParameters& Parameters)
 	{
 		SCOPE_CYCLE_COUNTER(STAT_UpdateConnectivityGraphUsingDelaunayTriangulation);
 
@@ -2305,8 +2263,8 @@ namespace Chaos
 		}
 	}
 
-	template<class T_FPBDRigidsEvolution, class T_FPBDCollisionConstraint>
-	void TPBDRigidClustering<T_FPBDRigidsEvolution, T_FPBDCollisionConstraint>::ConnectNodes(
+	void 
+	FRigidClustering::ConnectNodes(
 		FPBDRigidParticleHandle* Child1,
 		FPBDRigidParticleHandle* Child2)
 	{
@@ -2316,8 +2274,8 @@ namespace Chaos
 		ConnectNodes(ClusteredChild1, ClusteredChild2);
 	}
 
-	template<class T_FPBDRigidsEvolution, class T_FPBDCollisionConstraint>
-	void TPBDRigidClustering<T_FPBDRigidsEvolution, T_FPBDCollisionConstraint>::ConnectNodes(
+	void 
+	FRigidClustering::ConnectNodes(
 		FPBDRigidClusteredParticleHandle* ClusteredChild1,
 		FPBDRigidClusteredParticleHandle* ClusteredChild2)
 	{
@@ -2339,16 +2297,16 @@ namespace Chaos
 		}
 	}
 
-	template<class T_FPBDRigidsEvolution, class T_FPBDCollisionConstraint>
-	void TPBDRigidClustering<T_FPBDRigidsEvolution, T_FPBDCollisionConstraint>::RemoveNodeConnections(
+	void 
+	FRigidClustering::RemoveNodeConnections(
 		FPBDRigidParticleHandle* Child)
 	{
 		RemoveNodeConnections(Child->CastToClustered());
 	}
 
 	DECLARE_CYCLE_STAT(TEXT("TPBDRigidClustering<>::RemoveNodeConnections"), STAT_RemoveNodeConnections, STATGROUP_Chaos);
-	template<class T_FPBDRigidsEvolution, class T_FPBDCollisionConstraint>
-	void TPBDRigidClustering<T_FPBDRigidsEvolution, T_FPBDCollisionConstraint>::RemoveNodeConnections(
+	void 
+	FRigidClustering::RemoveNodeConnections(
 		FPBDRigidClusteredParticleHandle* ClusteredChild)
 	{
 		SCOPE_CYCLE_COUNTER(STAT_RemoveNodeConnections);
@@ -2368,8 +2326,3 @@ namespace Chaos
 
 
 } // namespace Chaos
-
-using namespace Chaos;
-
-template class CHAOS_API Chaos::TPBDRigidClustering<Chaos::FPBDRigidsEvolutionGBF, FPBDCollisionConstraints>;
-
diff --git a/Engine/Source/Runtime/Experimental/Chaos/Private/EventDefaults.cpp b/Engine/Source/Runtime/Experimental/Chaos/Private/EventDefaults.cpp
--- a/Engine/Source/Runtime/Experimental/Chaos/Private/EventDefaults.cpp
+++ b/Engine/Source/Runtime/Experimental/Chaos/Private/EventDefaults.cpp
@@ -104,7 +104,7 @@ namespace Chaos
 #if TODO_REIMPLEMENT_RIGID_CLUSTERING
 			const Chaos::FPBDRigidsSolver::FClusteringType::FClusterMap& ParentToChildrenMap = Evolution->GetRigidClustering().GetChildrenMap();
 #endif
-			const typename Chaos::TPBDRigidClustering<typename FPBDRigidsSolver::FPBDRigidsEvolution, FPBDCollisionConstraints>::FClusterMap& ParentToChildrenMap = Evolution->GetRigidClustering().GetChildrenMap();
+			const typename Chaos::FRigidClustering::FClusterMap& ParentToChildrenMap = Evolution->GetRigidClustering().GetChildrenMap();
 
 			if(CollisionRule.NumConstraints() > 0)
 			{
diff --git a/Engine/Source/Runtime/Experimental/Chaos/Public/Chaos/PBDCollisionConstraints.h b/Engine/Source/Runtime/Experimental/Chaos/Public/Chaos/PBDCollisionConstraints.h
--- a/Engine/Source/Runtime/Experimental/Chaos/Public/Chaos/PBDCollisionConstraints.h
+++ b/Engine/Source/Runtime/Experimental/Chaos/Public/Chaos/PBDCollisionConstraints.h
@@ -50,7 +50,6 @@ public:
 	using FConstHandles = TArrayView<const FPBDCollisionConstraint* const>;
 
 	// For use by dependent types
-	using FPointContactConstraint = FPBDCollisionConstraint;				// Used by clustering (remove this?)
 	using FConstraintContainerHandle = FPBDCollisionConstraintHandle;		// Used by constraint rules
 	using FConstraintSolverContainerType = FPBDCollisionSolverContainer;	// Used by constraint rules
 	using FParticleHandle = TGeometryParticleHandle<FReal, 3>;
diff --git a/Engine/Source/Runtime/Experimental/Chaos/Public/Chaos/PBDRigidClustering.h b/Engine/Source/Runtime/Experimental/Chaos/Public/Chaos/PBDRigidClustering.h
--- a/Engine/Source/Runtime/Experimental/Chaos/Public/Chaos/PBDRigidClustering.h
+++ b/Engine/Source/Runtime/Experimental/Chaos/Public/Chaos/PBDRigidClustering.h
@@ -16,6 +16,7 @@
 namespace Chaos
 {
 	extern CHAOS_API FRealSingle ChaosClusteringChildrenInheritVelocity;
+	class FPBDCollisionConstraints;
 }
 
 namespace Chaos
@@ -41,12 +42,11 @@ struct CHAOS_API FClusterDestoryParameters {
 /* 
 * PDBRigidClustering
 */
-template<class T_FPBDRigidEvolution, class T_FPBDCollisionConstraint>
-class CHAOS_API TPBDRigidClustering
+class CHAOS_API FRigidClustering
 {
-	typedef typename T_FPBDCollisionConstraint::FPointContactConstraint FPointContactConstraint;
 public:
 	/** Parent to children */
+	typedef FPBDRigidsEvolutionGBF                      FRigidEvolution;
 
 	typedef FPBDRigidParticleHandle*					FRigidHandle;
 	typedef TArray<FRigidHandle>						FRigidHandleArray;
@@ -56,12 +56,10 @@ public:
 
 	typedef TPair<FClusterHandle, FRigidHandleArray>	FClusterMapEntry;
 	typedef TMap<FClusterHandle, FRigidHandleArray>	    FClusterMap;
+	///using FCollisionConstraintHandle = FPBDCollisionConstraintHandle;
 
-
-	using FCollisionConstraintHandle = FPBDCollisionConstraintHandle;
-
-	TPBDRigidClustering(T_FPBDRigidEvolution& InEvolution, FPBDRigidClusteredParticles& InParticles);
-	~TPBDRigidClustering();
+	FRigidClustering(FRigidEvolution& InEvolution, FPBDRigidClusteredParticles& InParticles);
+	~FRigidClustering();
 
 	//
 	// Initialization
@@ -159,7 +157,7 @@ public:
 	*   ... Release bodies based collision impulses.
 	*   ... Updating properties as necessary.
 	*/
-	void AdvanceClustering(const FReal dt, T_FPBDCollisionConstraint& CollisionRule);
+	void AdvanceClustering(const FReal dt, FPBDCollisionConstraints& CollisionRule);
 
 	/**
 	*  BreakingModel
@@ -193,8 +191,7 @@ public:
 	//  be accessed via the game thread.
 	//
 	const FClusterBuffer&  GetBufferedData() const { ResourceLock.ReadLock(); return BufferResource; } /* Secure access from game thread*/
-	void                   ReleaseBufferedData() const { ResourceLock.ReadUnlock(); }    /* Release access from game thread*/
-	void                   SwapBufferedData();                                         /* Managed by the PBDRigidSolver ONLY!*/
+	void                   ReleaseBufferedData() const { ResourceLock.ReadUnlock(); }				   /* Release access from game thread*/
 
 
 	/*
@@ -321,7 +318,7 @@ public:
 		TSharedPtr<Chaos::FImplicitObject, ESPMode::ThreadSafe> ProxyGeometry,
 		const FClusterCreationParameters& Parameters);
 
-	void ComputeStrainFromCollision(const T_FPBDCollisionConstraint& CollisionRule);
+	void ComputeStrainFromCollision(const FPBDCollisionConstraints& CollisionRule);
 	void ResetCollisionImpulseArray();
 	void DisableCluster(FPBDRigidClusteredParticleHandle* ClusteredParticle);
 	void DisableParticleWithBreakEvent(Chaos::FPBDRigidParticleHandle* Particle);
@@ -351,7 +348,7 @@ public:
 
 private:
 
-	T_FPBDRigidEvolution& MEvolution;
+	FRigidEvolution& MEvolution;
 	FPBDRigidClusteredParticles& MParticles;
 	TSet<Chaos::FPBDRigidClusteredParticleHandle*> TopLevelClusterParents;
 	TSet<Chaos::FPBDRigidParticleHandle*> MActiveRemovalIndices;
@@ -529,7 +526,4 @@ inline void CleanCollisionParticles(
 	}
 }
 
-template <typename T, int d>
-using TClusterBuffer UE_DEPRECATED(4.27, "Deprecated. this class is to be deleted, use FClusterBuffer instead") = FClusterBuffer;
-
 } // namespace Chaos
diff --git a/Engine/Source/Runtime/Experimental/Chaos/Public/Chaos/PBDRigidsEvolutionGBF.h b/Engine/Source/Runtime/Experimental/Chaos/Public/Chaos/PBDRigidsEvolutionGBF.h
--- a/Engine/Source/Runtime/Experimental/Chaos/Public/Chaos/PBDRigidsEvolutionGBF.h
+++ b/Engine/Source/Runtime/Experimental/Chaos/Public/Chaos/PBDRigidsEvolutionGBF.h
@@ -7,6 +7,7 @@
 #include "Chaos/Collision/SpatialAccelerationCollisionDetector.h"
 #include "Chaos/Evolution/SolverBodyContainer.h"
 #include "Chaos/PBDCollisionConstraints.h"
+#include "Chaos/PBDRigidClustering.h"
 #include "Chaos/PBDRigidsEvolution.h"
 #include "Chaos/PerParticleAddImpulses.h"
 #include "Chaos/PerParticleEtherDrag.h"
@@ -87,7 +88,6 @@ namespace Chaos
 		using FCollisionConstraintRule = TPBDConstraintColorRule<FCollisionConstraints>;
 		using FCollisionDetector = FSpatialAccelerationCollisionDetector;
 		using FExternalForces = FPerParticleExternalForces;
-		using FRigidClustering = TPBDRigidClustering<FPBDRigidsEvolutionGBF, FPBDCollisionConstraints>;
 		using FJointConstraintsRule = TPBDConstraintIslandRule<FPBDJointConstraints>;
 		using FSuspensionConstraintsRule = TPBDConstraintIslandRule<FPBDSuspensionConstraints>;
 		using FJointConstraints = FPBDJointConstraints;
@@ -164,8 +164,8 @@ namespace Chaos
 		FORCEINLINE FGravityForces& GetGravityForces() { return GravityForces; }
 		FORCEINLINE const FGravityForces& GetGravityForces() const { return GravityForces; }
 
-		FORCEINLINE const TPBDRigidClustering<FPBDRigidsEvolutionGBF, FPBDCollisionConstraints>& GetRigidClustering() const { return Clustering; }
-		FORCEINLINE TPBDRigidClustering<FPBDRigidsEvolutionGBF, FPBDCollisionConstraints>& GetRigidClustering() { return Clustering; }
+		FORCEINLINE const FRigidClustering& GetRigidClustering() const { return Clustering; }
+		FORCEINLINE FRigidClustering& GetRigidClustering() { return Clustering; }
 
 		FORCEINLINE FJointConstraints& GetJointConstraints() { return JointConstraints; }
 		FORCEINLINE const FJointConstraints& GetJointConstraints() const { return JointConstraints; }
@@ -345,7 +345,7 @@ namespace Chaos
 			return CurrentStepResimCacheImp;
 		}
 
-		TPBDRigidClustering<FPBDRigidsEvolutionGBF, FPBDCollisionConstraints> Clustering;
+		FRigidClustering Clustering;
 
 		FPBDJointConstraints JointConstraints;
 		TPBDConstraintIslandRule<FPBDJointConstraints> JointConstraintRule;
diff --git a/Engine/Source/Runtime/Experimental/GeometryCollectionEngine/Private/GeometryCollection/GeometryCollectionComponent.cpp b/Engine/Source/Runtime/Experimental/GeometryCollectionEngine/Private/GeometryCollection/GeometryCollectionComponent.cpp
--- a/Engine/Source/Runtime/Experimental/GeometryCollectionEngine/Private/GeometryCollection/GeometryCollectionComponent.cpp
+++ b/Engine/Source/Runtime/Experimental/GeometryCollectionEngine/Private/GeometryCollection/GeometryCollectionComponent.cpp
@@ -1185,7 +1185,7 @@ void UGeometryCollectionComponent::UpdateRemovalEventRegistration()
 	}
 }
 
-void ActivateClusters(Chaos::FPBDRigidsEvolution::FRigidClustering& Clustering, Chaos::TPBDRigidClusteredParticleHandle<Chaos::FReal, 3>* Cluster)
+void ActivateClusters(Chaos::FRigidClustering& Clustering, Chaos::TPBDRigidClusteredParticleHandle<Chaos::FReal, 3>* Cluster)
 {
 	if(!Cluster)
 	{
@@ -1236,7 +1236,7 @@ void UGeometryCollectionComponent::OnRep_RepData(const FGeometryCollectionRepDat
 				Chaos::FPhysicsSolver* Solver = Prox->GetSolver<Chaos::FPhysicsSolver>();
 				Chaos::FPBDRigidsEvolution* Evo = Solver->GetEvolution();
 				check(Evo);
-				Chaos::FPBDRigidsEvolution::FRigidClustering& Clustering = Evo->GetRigidClustering();
+				Chaos::FRigidClustering& Clustering = Evo->GetRigidClustering();
 				
 				// Set X/R/V/W for next sim step from the replicated state
 				Particle->SetX(SourcePose.Position);
@@ -1429,7 +1429,7 @@ void UGeometryCollectionComponent::NetAbandonCluster_Implementation(int32 Transf
 
 			Solver->RegisterSimOneShotCallback([Prox = PhysicsProxy, Strain, TransformIndex, Solver]()
 			{
-				Chaos::TPBDRigidClustering<Chaos::FPBDRigidsEvolution, Chaos::FPBDCollisionConstraints>& Clustering = Solver->GetEvolution()->GetRigidClustering();
+				Chaos::FRigidClustering& Clustering = Solver->GetEvolution()->GetRigidClustering();
 				Chaos::FPBDRigidClusteredParticleHandle* Parent = Prox->GetParticles()[TransformIndex];
 
 				if(!Parent->Disabled())
