Chaos : Clustering
- Removed stale code paths.

diff --git a/Engine/Source/Runtime/Experimental/Chaos/Private/Chaos/PBDRigidClustering.cpp b/Engine/Source/Runtime/Experimental/Chaos/Private/Chaos/PBDRigidClustering.cpp
--- a/Engine/Source/Runtime/Experimental/Chaos/Private/Chaos/PBDRigidClustering.cpp
+++ b/Engine/Source/Runtime/Experimental/Chaos/Private/Chaos/PBDRigidClustering.cpp
@@ -35,9 +35,6 @@ namespace Chaos
 	int32 UseConnectivity = 1;
 	FAutoConsoleVariableRef CVarUseConnectivity(TEXT("p.UseConnectivity"), UseConnectivity, TEXT("Whether to use connectivity graph when breaking up clusters"));
 
-	CHAOS_API FRealSingle ChaosClusteringChildrenInheritVelocity = 1.f;
-	FAutoConsoleVariableRef CVarChildrenInheritVelocity(TEXT("p.ChildrenInheritVelocity"), ChaosClusteringChildrenInheritVelocity, TEXT("Whether children inherit parent collision velocity when declustering. 0 has no impact velocity like glass, 1 has full impact velocity like brick"));
-
 	int32 ComputeClusterCollisionStrains = 1;
 	FAutoConsoleVariableRef CVarComputeClusterCollisionStrains(TEXT("p.ComputeClusterCollisionStrains"), ComputeClusterCollisionStrains, TEXT("Whether to use collision constraints when processing clustering."));
 
@@ -68,126 +65,6 @@ namespace Chaos
 	int32 DeactivateClusterChildren = 0;
 	FAutoConsoleVariableRef CVarDeactivateClusterChildren(TEXT("p.DeactivateClusterChildren"), DeactivateClusterChildren, TEXT("If children should be decativated when broken and put into another cluster."));
 
-	//==========================================================================
-	// Free helper functions
-	//==========================================================================
-
-	template<class T, int d>
-	FVec3 GetContactLocation(const FPBDCollisionConstraint& Contact)
-	{
-		return Contact.GetLocation();
-	}
-
-	template<class T, int d>
-	FVec3 GetContactLocation(const FRigidBodyContactConstraintPGS& Contact)
-	{
-		// @todo(mlentine): Does the exact point matter?
-		T MinPhi = FLT_MAX;
-		FVec3 MinLoc;
-		for (int32 i = 0; i < Contact.Phi.Num(); ++i)
-		{
-			if (Contact.Phi[i] < MinPhi)
-			{
-				MinPhi = Contact.Phi[i];
-				MinLoc = Contact.Location[i];
-			}
-		}
-		return MinLoc;
-	}
-
-	template<class T, int d>
-	T CalculatePseudoMomentum(const TPBDRigidClusteredParticles<T, d>& InParticles, const uint32 Index)
-	{
-		FVec3 LinearPseudoMomentum = (InParticles.X(Index) - InParticles.P(Index)) * InParticles.M(Index);
-		FRotation3 Delta = InParticles.R(Index) * InParticles.Q(Index).Inverse();
-		FVec3 Axis;
-		T Angle;
-		Delta.ToAxisAndAngle(Axis, Angle);
-		FVec3 AngularPseudoMomentum = InParticles.I(Index) * (Axis * Angle);
-		return LinearPseudoMomentum.Size() + AngularPseudoMomentum.Size();
-	}
-
-	DECLARE_CYCLE_STAT(TEXT("TPBDRigidClustering<>::RewindAndEvolve<BGF>()"), STAT_RewindAndEvolve_BGF, STATGROUP_Chaos);
-	template<typename T, int d>
-	void RewindAndEvolve(
-		FPBDRigidsEvolutionGBF& Evolution, 
-		TPBDRigidClusteredParticles<T, d>& InParticles, 
-		const TSet<int32>& IslandsToRecollide, 
-		const TSet<FPBDRigidParticleHandle*> AllActivatedChildren,
-		const T Dt, 
-		FPBDCollisionConstraints& CollisionRule)
-	{
-		SCOPE_CYCLE_COUNTER(STAT_RewindAndEvolve_BGF);
-		// Rewind active particles
-		const TArray<int32> IslandsToRecollideArray = IslandsToRecollide.Array();
-		PhysicsParallelFor(IslandsToRecollideArray.Num(), [&](int32 Idx) {
-			int32 Island = IslandsToRecollideArray[Idx];
-			auto Particles = Evolution.GetIslandParticles(Island); // copy
-			for (int32 ArrayIdx = Particles.Num() - 1; ArrayIdx >= 0; --ArrayIdx)
-			{
-				auto PBDRigid = Particles[ArrayIdx]->CastToRigidParticle();
-				if(PBDRigid && PBDRigid->ObjectState() == EObjectStateType::Dynamic)
-				{
-					if (!PBDRigid->Sleeping() && !PBDRigid->Disabled())
-					{
-						PBDRigid->P() = PBDRigid->X();
-						PBDRigid->Q() = PBDRigid->R();
-						PBDRigid->V() = PBDRigid->PreV();
-						PBDRigid->W() = PBDRigid->PreW();
-						continue;
-					}
-				}
-				Particles.RemoveAtSwap(ArrayIdx);
-			}
-			Evolution.Integrate(MakeHandleView(Particles), Dt);
-		});
-
-		TSet<FGeometryParticleHandle*> AllIslandParticles;
-		for (int32 Island = 0; Island < Evolution.NumIslands(); ++Island)
-		{
-			const auto& ParticleIndices = Evolution.GetIslandParticles(Island);
-			for (const auto Particle : ParticleIndices)
-			{
-				auto PBDRigid = Particle->CastToRigidParticle();
-				if(PBDRigid && PBDRigid->ObjectState() == EObjectStateType::Dynamic)
-				{
-					bool bDisabled = PBDRigid->Disabled();
-
-					// #TODO - Have to repeat checking out whether the particle is disabled matching the PFor above.
-					// Move these into shared array so we only process it once
-					if (!AllIslandParticles.Contains(Particle) && !bDisabled)
-					{
-						AllIslandParticles.Add(Particle);
-					}
-				}
-			}
-		}
-
-		const bool bRewindOnDeclusterSolve = ChaosClusteringChildrenInheritVelocity < 1.f;
-		if (bRewindOnDeclusterSolve)
-		{
-			// @todo(mlentine): We can precompute internal constraints which can filter some from the narrow phase tests but may not help much
-
-#if CHAOS_PARTICLEHANDLE_TODO
-			//CollisionRule.UpdateConstraints(InParticles, Evolution.GetNonDisabledIndices(), Dt, AllActivatedChildren, AllIslandParticles.Array());
-#else
-			//CollisionRule.UpdateConstraints(Dt, AllIslandParticles);	//this seems wrong
-#endif
-
-			Evolution.InitializeAccelerationStructures();
-
-			// Resolve collisions
-			PhysicsParallelFor(IslandsToRecollide.Num(), [&](int32 Island) {
-				// @todo(mlentine): This is heavy handed and probably can be simplified as we know only a little bit changed.
-				Evolution.UpdateAccelerationStructures(Dt, Island);
-				Evolution.ApplyConstraintsPhase1(Dt, Island);
-				// @todo(ccaulfield): should we also update velocities here? Evolution does...
-				Evolution.ApplyConstraintsPhase2(Dt, Island);
-				// @todo(ccaulfield): support sleep state update on evolution
-				//Evolution.UpdateSleepState(Island);
-			});
-		}
-	}
 
 	DECLARE_CYCLE_STAT(TEXT("TPBDRigidClustering<>::UpdateClusterMassProperties()"), STAT_UpdateClusterMassProperties, STATGROUP_Chaos);
 	void UpdateClusterMassProperties(
@@ -454,6 +331,21 @@ namespace Chaos
 
 		NewParticle->SetSleeping(bClusterIsAsleep);
 
+		auto AddToClusterUnion = [&](int32 ClusterID, FPBDRigidClusteredParticleHandle* Handle)
+		{
+			if (ClusterID <= 0)
+			{
+				return;
+			}
+
+			if (!ClusterUnionMap.Contains(ClusterID))
+			{
+				ClusterUnionMap.Add(ClusterID, TArray<FPBDRigidClusteredParticleHandle*>());
+			}
+
+			ClusterUnionMap[ClusterID].Add(Handle);
+		};
+
 		if(ClusterGroupIndex)
 		{
 			AddToClusterUnion(ClusterGroupIndex, NewParticle);
@@ -649,11 +541,7 @@ namespace Chaos
 		TArray<FPBDRigidParticleHandle*>& Children = MChildren[ClusteredParticle];
 
 		bool bChildrenChanged = false;
-		const bool bRewindOnDecluster = ChaosClusteringChildrenInheritVelocity < 1;
-		const FRigidTransform3 PreSolveTM = 
-			bRewindOnDecluster ? 
-			FRigidTransform3(ClusteredParticle->X(), ClusteredParticle->R()) : 
-			FRigidTransform3(ClusteredParticle->P(), ClusteredParticle->Q());
+		const FRigidTransform3 PreSolveTM = FRigidTransform3(ClusteredParticle->P(), ClusteredParticle->Q());
 
 		//@todo(ocohen): iterate with all the potential parents at once?
 		//find all children within some distance of contact point
@@ -671,11 +559,8 @@ namespace Chaos
 			Child->SetX(ChildFrame.GetTranslation());
 			Child->SetR(ChildFrame.GetRotation());
 
-			if (!bRewindOnDecluster)
-			{
-				Child->SetP(Child->X());
-				Child->SetQ(Child->R());
-			}
+			Child->SetP(Child->X());
+			Child->SetQ(Child->R());
 
 			//todo(ocohen): for now just inherit velocity at new COM. This isn't quite right for rotation
 			//todo(ocohen): in the presence of collisions, this will leave all children with the post-collision
@@ -925,11 +810,7 @@ namespace Chaos
 		TArray<FPBDRigidParticleHandle*>& Children = MChildren[ClusteredParticle];
 
 		bool bChildrenChanged = false;
-		const bool bRewindOnDecluster = ChaosClusteringChildrenInheritVelocity < 1;
-		const FRigidTransform3 PreSolveTM =
-			bRewindOnDecluster ?
-			FRigidTransform3(ClusteredParticle->X(), ClusteredParticle->R()) :
-			FRigidTransform3(ClusteredParticle->P(), ClusteredParticle->Q());
+		const FRigidTransform3 PreSolveTM = FRigidTransform3(ClusteredParticle->P(), ClusteredParticle->Q());
 
 		//@todo(ocohen): iterate with all the potential parents at once?
 		//find all children within some distance of contact point
@@ -947,11 +828,8 @@ namespace Chaos
 			Child->SetX(ChildFrame.GetTranslation());
 			Child->SetR(ChildFrame.GetRotation());
 
-			if (!bRewindOnDecluster)
-			{
-				Child->SetP(Child->X());
-				Child->SetQ(Child->R());
-			}
+			Child->SetP(Child->X());
+			Child->SetQ(Child->R());
 
 			//todo(ocohen): for now just inherit velocity at new COM. This isn't quite right for rotation
 			//todo(ocohen): in the presence of collisions, this will leave all children with the post-collision
@@ -1177,7 +1055,6 @@ namespace Chaos
 	DECLARE_CYCLE_STAT(TEXT("TPBDRigidClustering<>::AdvanceClustering"), STAT_AdvanceClustering, STATGROUP_Chaos);
 	DECLARE_CYCLE_STAT(TEXT("TPBDRigidClustering<>::Update Impulse from Strain"), STAT_UpdateImpulseStrain, STATGROUP_Chaos);
 	DECLARE_CYCLE_STAT(TEXT("TPBDRigidClustering<>::Update Dirty Impulses"), STAT_UpdateDirtyImpulses, STATGROUP_Chaos);
-	DECLARE_CYCLE_STAT(TEXT("TPBDRigidClustering<>::Rewind"), STAT_ClusterRewind, STATGROUP_Chaos);
 	void 
 	FRigidClustering::AdvanceClustering(
 		const FReal Dt, 
@@ -1225,83 +1102,36 @@ namespace Chaos
 			//
 			{
 				SCOPE_CYCLE_COUNTER(STAT_UpdateDirtyImpulses);
-				for (const auto& ActiveCluster : GetTopLevelClusterParents() )
+				for (const auto& ActiveCluster : GetTopLevelClusterParents())
 				{
 					if (!ActiveCluster->Disabled())
-				{
-					if (ActiveCluster->ClusterIds().NumChildren > 0) //active index is a cluster
 					{
+						if (ActiveCluster->ClusterIds().NumChildren > 0) //active index is a cluster
+						{
 							TArray<FRigidHandle>& ParentToChildren = MChildren[ActiveCluster];
 							for (FRigidHandle Child : ParentToChildren)
-						{
-								if (FClusterHandle ClusteredChild = Child->CastToClustered())
 							{
-								if (ClusteredChild->Strain() <= 0.f)
+								if (FClusterHandle ClusteredChild = Child->CastToClustered())
 								{
-									ClusteredChild->CollisionImpulse() = FLT_MAX;
-									MCollisionImpulseArrayDirty = true;
+									if (ClusteredChild->Strain() <= 0.f)
+									{
+										ClusteredChild->CollisionImpulse() = FLT_MAX;
+										MCollisionImpulseArrayDirty = true;
+									}
 								}
 							}
 						}
 					}
 				}
 			}
-			}
 
 			if (MCollisionImpulseArrayDirty)
 			{
-
 				SCOPE_CYCLE_COUNTER(STAT_UpdateDirtyImpulses);
 				TMap<FPBDRigidClusteredParticleHandle*, TSet<FPBDRigidParticleHandle*>> ClusterToActivatedChildren = 
 					BreakingModel();
-
-				TSet<FPBDRigidParticleHandle*> AllActivatedChildren;
-				TSet<int32> IslandsToRecollide;
-				for (auto Itr : ClusterToActivatedChildren)
-				{
-					//question: do we need to iterate all the children? Seems like island is known from cluster, but don't want to break anything at this point
-					TSet<FPBDRigidParticleHandle*>& ActivatedChildren = Itr.Value;
-					for (FPBDRigidParticleHandle* ActiveChild : ActivatedChildren)
-					{
-						if (ensure(!ActiveChild->Disabled()))
-						{
-							int32 Island = ActiveChild->IslandIndex();
-							if (!IslandsToRecollide.Contains(Island) && Island != INDEX_NONE) // todo ask mike
-							{
-								IslandsToRecollide.Add(Island);
-							}
-						}
-					}
-					AllActivatedChildren.Append(ActivatedChildren);
-				}
-
-				const bool bRewindOnDecluster = ChaosClusteringChildrenInheritVelocity < 1.f;
-				if (bRewindOnDecluster && AllActivatedChildren.Num())
-				{
-					SCOPE_CYCLE_COUNTER(STAT_ClusterRewind);
-
-					if (MEvolution.NumIslands())
-					{
-						RewindAndEvolve(MEvolution, MParticles, IslandsToRecollide, AllActivatedChildren, Dt, CollisionRule);
-					}
-
-					if (ChaosClusteringChildrenInheritVelocity > 0.f)
-					{
-						for (auto Itr : ClusterToActivatedChildren)
-						{
-							FPBDRigidClusteredParticleHandle* ClusteredParticle = Itr.Key;
-							TSet<FPBDRigidParticleHandle*>& ActivatedChildren = Itr.Value;
-							for (FPBDRigidParticleHandle* ActiveChild : ActivatedChildren)
-							{
-								ActiveChild->SetV(
-									ActiveChild->V() * (1.f - ChaosClusteringChildrenInheritVelocity) + ClusteredParticle->V() * ChaosClusteringChildrenInheritVelocity);
-								ActiveChild->SetW(
-									ActiveChild->W() * (1.f - ChaosClusteringChildrenInheritVelocity) + ClusteredParticle->W() * ChaosClusteringChildrenInheritVelocity);
-							}
-						}
-					}
-				}
 			} // end if MCollisionImpulseArrayDirty
+
 		} // end if MParticles.Size()
 		Timer.Stop();
 		UE_LOG(LogChaos, Verbose, TEXT("Cluster Break Update Time is %f"), Time);
@@ -1370,32 +1200,6 @@ namespace Chaos
 		return AllActivatedChildren;
 	}
 
-	DECLARE_CYCLE_STAT(TEXT("TPBDRigidClustering<>::PromoteStrains()"), STAT_PromoteStrains, STATGROUP_Chaos);
-	FReal 
-	FRigidClustering::PromoteStrains(
-		FPBDRigidParticleHandle* CurrentNode)
-	{
-		SCOPE_CYCLE_COUNTER(STAT_PromoteStrains);
-		if (FPBDRigidClusteredParticleHandle* ClusteredCurrentNode = CurrentNode->CastToClustered())
-		{
-			FReal ChildrenStrains = (FReal)0.;
-			if (MChildren.Contains(ClusteredCurrentNode))
-			{
-				for (FRigidHandle Child : MChildren[ClusteredCurrentNode])
-				{
-					ChildrenStrains += PromoteStrains(Child);
-				}
-			}
-			else
-			{
-				return ClusteredCurrentNode->Strains();
-			}
-			ClusteredCurrentNode->SetStrains(ClusteredCurrentNode->Strains() + ChildrenStrains);
-			return ClusteredCurrentNode->Strains();
-		}
-		return (FReal)0.;
-	}
-
 	DECLARE_CYCLE_STAT(TEXT("TPBDRigidClustering<>::UpdateKinematicProperties()"), STAT_UpdateKinematicProperties, STATGROUP_Chaos);
 	void 
 	FRigidClustering::UpdateKinematicProperties(
diff --git a/Engine/Source/Runtime/Experimental/Chaos/Public/Chaos/PBDRigidClustering.h b/Engine/Source/Runtime/Experimental/Chaos/Public/Chaos/PBDRigidClustering.h
--- a/Engine/Source/Runtime/Experimental/Chaos/Public/Chaos/PBDRigidClustering.h
+++ b/Engine/Source/Runtime/Experimental/Chaos/Public/Chaos/PBDRigidClustering.h
@@ -11,52 +11,53 @@
 #include "Chaos/GeometryParticlesfwd.h"
 #include "Framework/BufferedData.h"
 
-#define TODO_CONVERT_GEOMETRY_COLLECTION_PARTICLE_INDICES_TO_PARTICLE_POINTERS 0
-
 namespace Chaos
 {
-	extern CHAOS_API FRealSingle ChaosClusteringChildrenInheritVelocity;
 	class FPBDCollisionConstraints;
-}
-
-namespace Chaos
-{
-
-class CHAOS_API FClusterBuffer
-{
-public:
-	using FClusterChildrenMap = TMap<FPBDRigidParticleHandle*, TArray<FPBDRigidParticleHandle*>>;
-	using FClusterTransformMap = TMap<FPBDRigidParticleHandle*, FRigidTransform3>;
-
-	virtual ~FClusterBuffer() = default;
-
-	FClusterChildrenMap MChildren;
-	FClusterTransformMap ClusterParentTransforms;
-	TArray<Chaos::TSerializablePtr<FImplicitObject>> GeometryPtrs;
-};
 
 struct CHAOS_API FClusterDestoryParameters {
 	bool bReturnInternalOnly : true;
 };
 
-/* 
-* PDBRigidClustering
+/****
+*
+*   FRigidClustering
+* 
+*   The Chaos Destruction System allows artists to define exactly how geometry 
+*   will break and separate during the simulations. Artists construct the 
+*	simulation assets using pre-fractured geometry and utilize dynamically 
+*   generated rigid constraints to model the structural connections during the
+*   simulation. The resulting objects within the simulation can separate from 
+*   connected structures based on interactions with environmental elements, 
+*   like fields and collisions.
+*
+*   The destruction system relies on an internal clustering model 
+*   (aka Clustering) which controls how the rigidly attached geometry is 
+*   simulated. Clustering allows artists to initialize sets of geometry as 
+*   a single rigid body, then dynamically break the objects during the 
+*   simulation. At its core, the clustering system will simply join the mass
+*   and inertia of each connected element into one larger single rigid body.
+* 
+*   At the beginning of the simulation a connection graph is initialized 
+*   based on the rigid body’s nearest neighbors. Each connection between the
+*   bodies represents a rigid constraint within the cluster and is given 
+*   initial strain values. During the simulation, the strains within the 
+*   connection graph are evaluated. The connections can be broken when collision
+*   constraints, or field evaluations, impart an impulse on the rigid body that
+*   exceeds the connections limit. Fields can also be used to decrease the 
+*   internal strain values of the connections, resulting in a weakening of the
+*   internal structure
+*
 */
 class CHAOS_API FRigidClustering
 {
 public:
-	/** Parent to children */
-	typedef FPBDRigidsEvolutionGBF                      FRigidEvolution;
 
+	typedef FPBDRigidsEvolutionGBF                      FRigidEvolution;
 	typedef FPBDRigidParticleHandle*					FRigidHandle;
 	typedef TArray<FRigidHandle>						FRigidHandleArray;
-	typedef TSet<FPBDRigidParticleHandle>				FRigidHandleSet;
-
 	typedef FPBDRigidClusteredParticleHandle*			FClusterHandle;
-
-	typedef TPair<FClusterHandle, FRigidHandleArray>	FClusterMapEntry;
 	typedef TMap<FClusterHandle, FRigidHandleArray>	    FClusterMap;
-	///using FCollisionConstraintHandle = FPBDCollisionConstraintHandle;
 
 	FRigidClustering(FRigidEvolution& InEvolution, FPBDRigidClusteredParticles& InParticles);
 	~FRigidClustering();
@@ -169,15 +170,6 @@ public:
 	TMap<FPBDRigidClusteredParticleHandle*, TSet<FPBDRigidParticleHandle*>> BreakingModel(
 		TMap<FGeometryParticleHandle*, Chaos::FReal>* ExternalStrainMap = nullptr);
 
-	/**
-	*  PromoteStrains
-	*    Sums the strains based on the cluster hierarchy. For example
-	*    a cluster with two children that have strains {3,4} will have
-	*    a ExternalStrain entry of 7. Will only decent the current
-	*    node passed, and ignores the disabled flag.
-	*/
-	FReal PromoteStrains(FPBDRigidParticleHandle* CurrentNode);
-
 	/*
 	*  Process the kinematic state of the clusters. Because the leaf node geometry can
 	*  be changed by the solver, it is necessary to check all the sub clusters.
@@ -187,11 +179,6 @@ public:
 	//
 	// Access
 	//
-	//  The ClusterIds and ChildrenMap are shared resources that can
-	//  be accessed via the game thread.
-	//
-	const FClusterBuffer&  GetBufferedData() const { ResourceLock.ReadLock(); return BufferResource; } /* Secure access from game thread*/
-	void                   ReleaseBufferedData() const { ResourceLock.ReadUnlock(); }				   /* Release access from game thread*/
 
 
 	/*
@@ -250,21 +237,10 @@ public:
 	*/
 	TArrayCollectionArray<int32>& GetClusterGroupIndexArray() { return MParticles.ClusterGroupIndexArray(); }
 
-	void AddToClusterUnion(int32 ClusterID, FPBDRigidClusteredParticleHandle* Handle)
-	{
-		if(ClusterID <= 0)
-		{
-			return;
-		}
-
-		if(!ClusterUnionMap.Contains(ClusterID))
-		{
-			ClusterUnionMap.Add(ClusterID, TArray<FPBDRigidClusteredParticleHandle*>());
-		}
-
-		ClusterUnionMap[ClusterID].Add(Handle);
-	}
-
+	/*
+	*  Cluster Break Data
+	*     The cluster breaks can be used to seed particle emissions. 
+	*/
 	const TArray<FBreakingData>& GetAllClusterBreakings() const { return MAllClusterBreakings; }
 	void SetGenerateClusterBreaking(bool DoGenerate) { DoGenerateBreakingData = DoGenerate; }
 	void ResetAllClusterBreakings() { MAllClusterBreakings.Reset(); }
@@ -290,22 +266,6 @@ public:
 	const TSet<Chaos::FPBDRigidClusteredParticleHandle*>& GetTopLevelClusterParents() const { return TopLevelClusterParents; }
 	TSet<Chaos::FPBDRigidClusteredParticleHandle*>& GetTopLevelClusterParents() { return TopLevelClusterParents; }
 
-	/* Ryan - do we still need this?
-	void InitTopLevelClusterParents(const int32 StartIndex)
-	{
-		if (!StartIndex)
-		{
-			TopLevelClusterParents.Reset();
-		}
-		for (uint32 i = StartIndex; i < MParticles.Size(); ++i)
-		{
-			if (MParticles.ClusterIds(i).Id == INDEX_NONE && !MParticles.Disabled(i))
-			{
-				TopLevelClusterParents.Add(i);
-			}
-		}
-	}
-	*/
  protected:
 	void UpdateMassProperties(
 		Chaos::FPBDRigidClusteredParticleHandle* Parent, 
@@ -355,8 +315,6 @@ private:
 
 
 	// Cluster data
-	mutable FRWLock ResourceLock;
-	FClusterBuffer BufferResource;
 	FClusterMap MChildren;
 	TMap<int32, TArray<FPBDRigidClusteredParticleHandle*> > ClusterUnionMap;
 
