Fix for sleep collision bug

diff --git a/Engine/Source/Runtime/Experimental/Chaos/Public/Chaos/ParticleHandle.h b/Engine/Source/Runtime/Experimental/Chaos/Public/Chaos/ParticleHandle.h
--- a/Engine/Source/Runtime/Experimental/Chaos/Public/Chaos/ParticleHandle.h
+++ b/Engine/Source/Runtime/Experimental/Chaos/Public/Chaos/ParticleHandle.h
@@ -1375,7 +1375,7 @@ public:
 
 	bool HasCollisionConstraintFlag(const ECollisionConstraintFlags Flag)  const
 	{
-		if (MHandle->CastToRigidParticle() && MHandle->ObjectState() == EObjectStateType::Dynamic)
+		if (MHandle->CastToRigidParticle() && (MHandle->ObjectState() == EObjectStateType::Dynamic || MHandle->ObjectState() == EObjectStateType::Sleeping))
 		{
 			return MHandle->CastToRigidParticle()->HasCollisionConstraintFlag(Flag);
 		}
