Chaos - Clustering
- Move mass initilization into Algo

diff --git a/Engine/Source/Programs/HeadlessChaos/Private/HeadlessChaosTestSpatialHashing.cpp b/Engine/Source/Programs/HeadlessChaos/Private/HeadlessChaosTestSpatialHashing.cpp
--- a/Engine/Source/Programs/HeadlessChaos/Private/HeadlessChaosTestSpatialHashing.cpp
+++ b/Engine/Source/Programs/HeadlessChaos/Private/HeadlessChaosTestSpatialHashing.cpp
@@ -7,6 +7,7 @@
 #include "Modules/ModuleManager.h"
 #include "Chaos/Matrix.h"
 #include "Chaos/PBDRigidClustering.h"
+#include "Chaos/PBDRigidClusteringCollisionParticleAlgo.h"
 
 DEFINE_LOG_CATEGORY_STATIC(AHSP_Test, Verbose, All);
 
diff --git a/Engine/Source/Runtime/Experimental/Chaos/Private/Chaos/PBDRigidClustering.cpp b/Engine/Source/Runtime/Experimental/Chaos/Private/Chaos/PBDRigidClustering.cpp
--- a/Engine/Source/Runtime/Experimental/Chaos/Private/Chaos/PBDRigidClustering.cpp
+++ b/Engine/Source/Runtime/Experimental/Chaos/Private/Chaos/PBDRigidClustering.cpp
@@ -10,6 +10,7 @@
 #include "Chaos/PBDRigidsEvolution.h"
 #include "Chaos/PBDCollisionConstraints.h"
 #include "Chaos/PBDCollisionConstraintsPGS.h"
+#include "Chaos/PBDRigidClusteringAlgo.h"
 #include "Chaos/Sphere.h"
 #include "Chaos/UniformGrid.h"
 #include "ChaosStats.h"
@@ -38,199 +39,10 @@ namespace Chaos
 	int32 ComputeClusterCollisionStrains = 1;
 	FAutoConsoleVariableRef CVarComputeClusterCollisionStrains(TEXT("p.ComputeClusterCollisionStrains"), ComputeClusterCollisionStrains, TEXT("Whether to use collision constraints when processing clustering."));
 
-	//
-	// Update Geometry PVar
-	//
-	int32 MinLevelsetDimension = 4;
-	FAutoConsoleVariableRef CVarMinLevelsetDimension(TEXT("p.MinLevelsetDimension"), MinLevelsetDimension, TEXT("The minimum number of cells on a single level set axis"));
-
-	int32 MaxLevelsetDimension = 20;
-	FAutoConsoleVariableRef CVarMaxLevelsetDimension(TEXT("p.MaxLevelsetDimension"), MaxLevelsetDimension, TEXT("The maximum number of cells on a single level set axis"));
-
-	FRealSingle MinLevelsetSize = 50.f;
-	FAutoConsoleVariableRef CVarLevelSetResolution(TEXT("p.MinLevelsetSize"), MinLevelsetSize, TEXT("The minimum size on the smallest axis to use a level set"));
-
-	int32 UseLevelsetCollision = 0;
-	FAutoConsoleVariableRef CVarUseLevelsetCollision(TEXT("p.UseLevelsetCollision"), UseLevelsetCollision, TEXT("Whether unioned objects use levelsets"));
-
-	int32 LevelsetGhostCells = 1;
-	FAutoConsoleVariableRef CVarLevelsetGhostCells(TEXT("p.LevelsetGhostCells"), LevelsetGhostCells, TEXT("Increase the level set grid by this many ghost cells"));
-
-	FRealSingle ClusterSnapDistance = 1.f;
-	FAutoConsoleVariableRef CVarClusterSnapDistance(TEXT("p.ClusterSnapDistance"), ClusterSnapDistance, TEXT(""));
-
-	int32 MinCleanedPointsBeforeRemovingInternals = 10;
-	FAutoConsoleVariableRef CVarMinCleanedPointsBeforeRemovingInternals(TEXT("p.MinCleanedPointsBeforeRemovingInternals"), MinCleanedPointsBeforeRemovingInternals, TEXT("If we only have this many clean points, don't bother removing internal points as the object is likely very small"));
-
 	int32 DeactivateClusterChildren = 0;
 	FAutoConsoleVariableRef CVarDeactivateClusterChildren(TEXT("p.DeactivateClusterChildren"), DeactivateClusterChildren, TEXT("If children should be decativated when broken and put into another cluster."));
 
 
-	DECLARE_CYCLE_STAT(TEXT("TPBDRigidClustering<>::UpdateClusterMassProperties()"), STAT_UpdateClusterMassProperties, STATGROUP_Chaos);
-	void UpdateClusterMassProperties(
-		Chaos::FPBDRigidClusteredParticleHandle* Parent,
-		TSet<FPBDRigidParticleHandle*>& Children,
-		const FRigidTransform3* ForceMassOrientation)
-	{
-		SCOPE_CYCLE_COUNTER(STAT_UpdateClusterMassProperties);
-		check(Children.Num());
-
-		Parent->SetX(FVec3(0));
-		Parent->SetR(FRotation3(FQuat::MakeFromEuler(FVec3(0))));
-		Parent->SetV(FVec3(0));
-		Parent->SetW(FVec3(0));
-		Parent->SetM(0);
-		Parent->SetI(FMatrix33(0));
-
-		bool bHasChild = false;
-		bool bHasProxyChild = false;
-		for (FPBDRigidParticleHandle* OriginalChild : Children)
-		{
-			FPBDRigidParticleHandle* Child;
-			FVec3 ChildPosition;
-			FRotation3 ChildRotation;
-
-			Child = OriginalChild;
-			ChildPosition = Child->X();
-			ChildRotation = Child->R();
-
-			const FReal ChildMass = Child->M();
-			const FMatrix33 ChildWorldSpaceI =
-				(ChildRotation * FMatrix::Identity) * Child->I() * (ChildRotation * FMatrix::Identity).GetTransposed();
-			if (ChildWorldSpaceI.ContainsNaN())
-			{
-				continue;
-			}
-			bHasProxyChild = true;
-			bHasChild = true;
-			bHasProxyChild = true;
-			Parent->I() += ChildWorldSpaceI;
-			Parent->M() += ChildMass;
-			Parent->X() += ChildPosition * ChildMass;
-			Parent->V() += OriginalChild->V() * ChildMass; // Use orig child for vel because we don't sim the proxy
-			Parent->W() += OriginalChild->W() * ChildMass;
-		}
-		if (!ensure(bHasProxyChild))
-		{
-			for (FPBDRigidParticleHandle* OriginalChild : Children)
-			{
-				FPBDRigidParticleHandle* Child = OriginalChild;
-				const FVec3& ChildPosition = Child->X();
-				const FRotation3& ChildRotation = Child->R();
-				const FReal ChildMass = Child->M();
-
-				const FMatrix33 ChildWorldSpaceI =
-					(ChildRotation * FMatrix::Identity) * Child->I() * (ChildRotation * FMatrix::Identity).GetTransposed();
-				if (ChildWorldSpaceI.ContainsNaN())
-				{
-					continue;
-				}
-				bHasChild = true;
-				Parent->I() += ChildWorldSpaceI;
-				Parent->M() += ChildMass;
-				Parent->X() += ChildPosition * ChildMass;
-				Parent->V() += OriginalChild->V() * ChildMass; // Use orig child for vel because we don't sim the proxy
-				Parent->W() += OriginalChild->W() * ChildMass;
-			}
-		}
-		for (int32 i = 0; i < 3; i++)
-		{
-			const FMatrix33& InertiaTensor = Parent->I();
-			if (InertiaTensor.GetColumn(i)[i] < SMALL_NUMBER)
-			{
-				Parent->SetI(FMatrix33(1.f, 1.f, 1.f));
-				break;
-			}
-		}
-
-		if (!ensure(bHasChild) || !ensure(Parent->M() > SMALL_NUMBER))
-		{
-			Parent->M() = 1.0;
-			Parent->X() = FVec3(0);
-			Parent->V() = FVec3(0);
-			Parent->PreV() = Parent->V();
-			Parent->InvM() = 1;
-			Parent->P() = Parent->X();
-			Parent->W() = FVec3(0);
-			Parent->PreW() = Parent->W();
-			Parent->R() = FRotation3(FMatrix::Identity);
-			Parent->Q() = Parent->R();
-			Parent->I() = FMatrix::Identity;
-			Parent->InvI() = FMatrix::Identity;
-			return;
-		}
-
-		check(Parent->M() > SMALL_NUMBER);
-
-		Parent->X() /= Parent->M();
-		Parent->V() /= Parent->M();
-		Parent->PreV() = Parent->V();
-		Parent->InvM() = static_cast<FReal>(1.0) / Parent->M();
-		if (ForceMassOrientation)
-		{
-			Parent->X() = ForceMassOrientation->GetLocation();
-		}
-		Parent->P() = Parent->X();
-		for (FPBDRigidParticleHandle* OriginalChild : Children)
-		{
-			FPBDRigidParticleHandle* Child;
-			FVec3 ChildPosition;
-
-			Child = OriginalChild;
-			ChildPosition = Child->X();
-
-			FVec3 ParentToChild = ChildPosition - Parent->X();
-
-			const FReal ChildMass = Child->M();
-			// taking v from original child since we are not actually simulating the proxy child
-			Parent->W() +=
-				FVec3::CrossProduct(ParentToChild,
-					OriginalChild->V() * ChildMass);
-			{
-				const FReal& p0 = ParentToChild[0];
-				const FReal& p1 = ParentToChild[1];
-				const FReal& p2 = ParentToChild[2];
-				const FReal& m = ChildMass;
-				Parent->I() +=
-					FMatrix33(
-						m * (p1 * p1 + p2 * p2), -m * p1 * p0, -m * p2 * p0,
-						m * (p2 * p2 + p0 * p0), -m * p2 * p1, m * (p1 * p1 + p0 * p0));
-			}
-		}
-		FMatrix33& InertiaTensor = Parent->I();
-		if (Parent->I().ContainsNaN())
-		{
-			InertiaTensor = FMatrix33((FReal)1., (FReal)1., (FReal)1.);
-		}
-		else
-		{
-			for (int32 i = 0; i < 3; i++)
-			{
-				if (InertiaTensor.GetColumn(i)[i] < SMALL_NUMBER)
-				{
-					InertiaTensor = FMatrix33((FReal)1., (FReal)1., (FReal)1.);
-					break;
-				}
-			}
-		}
-		Parent->W() /= Parent->M();
-		Parent->PreW() = Parent->W();
-		if (ForceMassOrientation)
-		{
-			Parent->R() = ForceMassOrientation->GetRotation();
-		}
-
-		if (Parent->R().ContainsNaN())
-		{
-			InertiaTensor = PMatrix<FReal, 3, 3>(1.f, 1.f, 1.f);
-			Parent->R() = TRotation<FReal, 3>(FMatrix::Identity);
-		}
-
-		Parent->Q() = Parent->R();
-		Parent->InvI() = Parent->I().Inverse();
-	}
-
-
 	//==========================================================================
 	// TPBDRigidClustering
 	//==========================================================================
@@ -325,7 +137,8 @@ namespace Chaos
 
 		ensureMsgf(!ProxyGeometry || ForceMassOrientation, TEXT("If ProxyGeometry is passed, we must override the mass orientation as they are tied"));
 
-		UpdateMassProperties(NewParticle, ChildrenSet, ForceMassOrientation);
+		UpdateClusterMassProperties(NewParticle, ChildrenSet, ForceMassOrientation);
+		UpdateKinematicProperties(NewParticle, MChildren, MEvolution);
 		UpdateGeometry(NewParticle, ChildrenSet, ProxyGeometry, Parameters);
 		GenerateConnectionGraph(NewParticle, Parameters);
 
@@ -427,7 +240,9 @@ namespace Chaos
 		NoCleanParams.bCopyCollisionParticles = !!UnionsHaveCollisionParticles;
 
 		TSet<FPBDRigidParticleHandle*> ChildrenSet(ChildrenArray);
-		UpdateMassProperties(NewParticle, ChildrenSet, nullptr);
+		UpdateClusterMassProperties(NewParticle, ChildrenSet, nullptr);
+		UpdateKinematicProperties(NewParticle, MChildren, MEvolution);
+
 		UpdateGeometry(NewParticle, ChildrenSet, nullptr, NoCleanParams);
 
 		return NewParticle;
@@ -776,7 +591,7 @@ namespace Chaos
 
 			for (FPBDRigidParticleHandle* Child : ActivatedChildren)
 			{
-				UpdateKinematicProperties(Child);
+				UpdateKinematicProperties(Child, MChildren, MEvolution);
 			}
 
 			//disable cluster
@@ -994,7 +809,7 @@ namespace Chaos
 
 			for (FPBDRigidParticleHandle* Child : ActivatedChildren)
 			{
-				UpdateKinematicProperties(Child);
+				UpdateKinematicProperties(Child, MChildren, MEvolution);
 			}
 
 			//disable cluster
@@ -1200,58 +1015,6 @@ namespace Chaos
 		return AllActivatedChildren;
 	}
 
-	DECLARE_CYCLE_STAT(TEXT("TPBDRigidClustering<>::UpdateKinematicProperties()"), STAT_UpdateKinematicProperties, STATGROUP_Chaos);
-	void 
-	FRigidClustering::UpdateKinematicProperties(
-		Chaos::FPBDRigidParticleHandle* InParent)
-	{
-		SCOPE_CYCLE_COUNTER(STAT_UpdateKinematicProperties);
-
-		EObjectStateType ObjectState = EObjectStateType::Dynamic;
-		check(InParent != nullptr);
-		if (FClusterHandle ClusteredCurrentNode = InParent->CastToClustered())
-		{
-			if (MChildren.Contains(ClusteredCurrentNode) && MChildren[ClusteredCurrentNode].Num())
-		{
-			// TQueue is a linked list, which has no preallocator.
-			TQueue<Chaos::FPBDRigidParticleHandle*> Queue;
-				for (Chaos::FPBDRigidParticleHandle* Child : MChildren[ClusteredCurrentNode])
-			{
-				Queue.Enqueue(Child);
-			}
-
-			Chaos::FPBDRigidParticleHandle* CurrentHandle;
-			while (Queue.Dequeue(CurrentHandle) && ObjectState == EObjectStateType::Dynamic)
-			{
-					if (FClusterHandle CurrentClusterHandle = CurrentHandle->CastToClustered())
-					{
-				// @question : Maybe we should just store the leaf node bodies in a
-				// map, that will require Memory(n*log(n))
-						if (MChildren.Contains(CurrentClusterHandle))
-				{
-							for( Chaos::FPBDRigidParticleHandle* Child : MChildren[CurrentClusterHandle])
-					{
-						Queue.Enqueue(Child);
-					}
-				}
-					}
-
-				const EObjectStateType CurrState = CurrentHandle->ObjectState();
-				if (CurrState == EObjectStateType::Kinematic)
-				{
-					ObjectState = EObjectStateType::Kinematic;
-				}
-				else if (CurrState == EObjectStateType::Static)
-				{
-					ObjectState = EObjectStateType::Static;
-				}
-			}
-
-				MEvolution.SetParticleObjectState(ClusteredCurrentNode, ObjectState);
-			}
-		}
-	}
-
 	DECLARE_CYCLE_STAT(TEXT("TPBDRigidClustering<>::GetActiveClusterIndex"), STAT_GetActiveClusterIndex, STATGROUP_Chaos);
 	FPBDRigidParticleHandle* 
 	FRigidClustering::GetActiveClusterIndex(
@@ -1310,276 +1073,6 @@ namespace Chaos
 		}
 	}
 
-	DECLARE_CYCLE_STAT(TEXT("TPBDRigidClustering<>::UpdateMassProperties"), STAT_UpdateMassProperties, STATGROUP_Chaos);
-	void 
-	FRigidClustering::UpdateMassProperties(
-		Chaos::FPBDRigidClusteredParticleHandle* Parent, 
-		TSet<FPBDRigidParticleHandle*>& Children, 
-		const FRigidTransform3* ForceMassOrientation)
-	{
-		SCOPE_CYCLE_COUNTER(STAT_UpdateMassProperties);
-		UpdateClusterMassProperties(
-			Parent,
-			Children,
-			ForceMassOrientation);
-		UpdateKinematicProperties(Parent);
-	}
-
-	DECLARE_CYCLE_STAT(TEXT("TPBDRigidClustering<>::UpdateGeometry"), STAT_UpdateGeometry, STATGROUP_Chaos);
-	DECLARE_CYCLE_STAT(TEXT("TPBDRigidClustering<>::UpdateGeometry_GatherObjects"), STAT_UpdateGeometry_GatherObjects, STATGROUP_Chaos);
-	DECLARE_CYCLE_STAT(TEXT("TPBDRigidClustering<>::UpdateGeometry_GatherPoints"), STAT_UpdateGeometry_GatherPoints, STATGROUP_Chaos);
-	DECLARE_CYCLE_STAT(TEXT("TPBDRigidClustering<>::UpdateGeometry_CopyPoints"), STAT_UpdateGeometry_CopyPoints, STATGROUP_Chaos);
-	DECLARE_CYCLE_STAT(TEXT("TPBDRigidClustering<>::UpdateGeometry_PointsBVH"), STAT_UpdateGeometry_PointsBVH, STATGROUP_Chaos);
-
-	void 
-	FRigidClustering::UpdateGeometry(
-		Chaos::FPBDRigidClusteredParticleHandle* Parent, 
-		const TSet<FPBDRigidParticleHandle*>& Children, 
-		TSharedPtr<Chaos::FImplicitObject, ESPMode::ThreadSafe> ProxyGeometry,
-		const FClusterCreationParameters& Parameters)
-	{
-		SCOPE_CYCLE_COUNTER(STAT_UpdateGeometry);
-
-		TArray<TUniquePtr<FImplicitObject>> Objects;
-		TArray<TUniquePtr<FImplicitObject>> Objects2; //todo: find a better way to reuse this
-		Objects.Reserve(Children.Num());
-		Objects2.Reserve(Children.Num());
-
-		//we should never update existing geometry since this is used by SQ threads.
-		ensure(!Parent->Geometry());
-		ensure(!Parent->DynamicGeometry());
-
-		const FRigidTransform3 ClusterWorldTM(Parent->X(), Parent->R());
-
-		TArray<FVec3> OriginalPoints;
-		TArray<FPBDRigidParticleHandle*> ChildParticleHandles;
-		ChildParticleHandles.Reserve(Children.Num());
-
-		const bool bUseCollisionPoints = (ProxyGeometry || Parameters.bCopyCollisionParticles) && !Parameters.CollisionParticles;
-		bool bUseParticleImplicit = false;
-
-		// Need to extract a filter off one of the cluster children 
-		FCollisionFilterData Filter;
-		for(FPBDRigidParticleHandle* Child : Children)
-		{
-			bool bFilterValid = false;
-			for(const TUniquePtr<FPerShapeData>& Shape : Child->ShapesArray())
-			{
-				if(Shape)
-				{
-					Filter = Shape->GetSimData();
-					bFilterValid = Filter.Word0 != 0 || Filter.Word1 != 0 || Filter.Word2 != 0 || Filter.Word3 != 0;
-				}
-
-				if(bFilterValid)
-				{
-					break;
-				}
-			}
-
-			// Bail once we've found one
-			if(bFilterValid)
-			{
-				break;
-			}
-		}
-
-		{ // STAT_UpdateGeometry_GatherObjects
-			SCOPE_CYCLE_COUNTER(STAT_UpdateGeometry_GatherObjects);
-
-			if (bUseCollisionPoints)
-			{
-				uint32 NumPoints = 0;
-				for (FPBDRigidParticleHandle* Child : Children)
-				{
-					NumPoints += Child->CollisionParticlesSize();
-				}
-				OriginalPoints.Reserve(NumPoints);
-			}
-
-			for (FPBDRigidParticleHandle* Child : Children)
-			{
-				const FRigidTransform3 ChildWorldTM(Child->X(), Child->R());
-				FRigidTransform3 Frame = ChildWorldTM.GetRelativeTransform(ClusterWorldTM);
-				FPBDRigidParticleHandle* UsedGeomChild = Child;
-				if (Child->Geometry())
-				{
-					Objects.Add(TUniquePtr<FImplicitObject>(new TImplicitObjectTransformed<FReal, 3>(Child->Geometry(), Frame)));
-					Objects2.Add(TUniquePtr<FImplicitObject>(new TImplicitObjectTransformed<FReal, 3>(Child->Geometry(), Frame)));
-					ChildParticleHandles.Add(Child);
-				}
-
-				ensure(Child->Disabled() == true);
-				check(Child->CastToClustered()->ClusterIds().Id == Parent);
-
-				Child->CastToClustered()->SetChildToParent(Frame);
-
-				if (bUseCollisionPoints)
-				{
-					SCOPE_CYCLE_COUNTER(STAT_UpdateGeometry_GatherPoints);
-					if (const TUniquePtr<FBVHParticles>& CollisionParticles = Child->CollisionParticles())
-					{
-						for (uint32 i = 0; i < CollisionParticles->Size(); ++i)
-						{
-							OriginalPoints.Add(Frame.TransformPosition(CollisionParticles->X(i)));
-						}
-					}
-				}
-				if (Child->Geometry() && Child->Geometry()->GetType() == ImplicitObjectType::Unknown)
-				{
-					bUseParticleImplicit = true;
-				}
-			} // end for
-		} // STAT_UpdateGeometry_GatherObjects
-
-		{
-			QUICK_SCOPE_CYCLE_COUNTER(SpatialBVH);
-			TUniquePtr<FImplicitObjectUnionClustered>& ChildrenSpatial = Parent->ChildrenSpatial();
-			ChildrenSpatial.Reset(
-				Objects2.Num() ? 
-				new Chaos::FImplicitObjectUnionClustered(MoveTemp(Objects2), ChildParticleHandles) : 
-				nullptr);
-		}
-
-		TArray<FVec3> CleanedPoints;
-		if (!Parameters.CollisionParticles)
-		{
-			CleanedPoints = 
-				Parameters.bCleanCollisionParticles ? 
-				CleanCollisionParticles(OriginalPoints, ClusterSnapDistance) : 
-				OriginalPoints;
-		}
-
-		if (ProxyGeometry)
-		{
-			const FVector Scale = Parameters.Scale;
-			auto DeepCopyImplicit = [&Scale](const TSharedPtr<Chaos::FImplicitObject, ESPMode::ThreadSafe>& ImplicitToCopy) -> TUniquePtr<Chaos::FImplicitObject>
-			{
-				if (Scale.Equals(FVector::OneVector))
-				{
-					return ImplicitToCopy->DeepCopy();
-				}
-				else
-				{
-					return ImplicitToCopy->DeepCopyWithScale(Scale);
-				}
-			};
-			//ensureMsgf(false, TEXT("Checking usage with proxy"));
-			//@coverage {production}
-			Parent->SetSharedGeometry(TSharedPtr<FImplicitObject, ESPMode::ThreadSafe>(DeepCopyImplicit(ProxyGeometry).Release()));
-		}
-		else if (Objects.Num() == 0)
-		{
-			//ensureMsgf(false, TEXT("Checking usage with no proxy and no objects"));
-			//@coverage : {production}
-			Parent->SetGeometry(Chaos::TSerializablePtr<Chaos::FImplicitObject>());
-		}
-		else
-		{
-			if (UseLevelsetCollision)
-			{
-				ensureMsgf(false, TEXT("Checking usage with no proxy and multiple ojects with levelsets"));
-
-				FImplicitObjectUnionClustered UnionObject(MoveTemp(Objects));
-				FAABB3 Bounds = UnionObject.BoundingBox();
-				const FVec3 BoundsExtents = Bounds.Extents();
-				if (BoundsExtents.Min() >= MinLevelsetSize) //make sure the object is not too small
-				{
-					TVec3<int32> NumCells = Bounds.Extents() / MinLevelsetSize;
-					for (int i = 0; i < 3; ++i)
-					{
-						NumCells[i] = FMath::Clamp(NumCells[i], MinLevelsetDimension, MaxLevelsetDimension);
-					}
-
-					FErrorReporter ErrorReporter;
-					TUniformGrid<FReal, 3> Grid(Bounds.Min(), Bounds.Max(), NumCells, LevelsetGhostCells);
-					TUniquePtr<FLevelSet> LevelSet(new FLevelSet(ErrorReporter, Grid, UnionObject));
-
-					if (!Parameters.CollisionParticles)
-					{
-						const FReal MinDepthToSurface = Grid.Dx().Max();
-						for (int32 Idx = CleanedPoints.Num() - 1; Idx >= 0; --Idx)
-						{
-							if (CleanedPoints.Num() > MinCleanedPointsBeforeRemovingInternals) //todo(ocohen): this whole thing should really be refactored
-							{
-								const FVec3& CleanedCollision = CleanedPoints[Idx];
-								if (LevelSet->SignedDistance(CleanedCollision) < -MinDepthToSurface)
-								{
-									CleanedPoints.RemoveAtSwap(Idx);
-								}
-							}
-						}
-					}
-					Parent->SetDynamicGeometry(MoveTemp(LevelSet));
-				}
-				else
-				{
-					Parent->SetDynamicGeometry(
-						MakeUnique<TSphere<FReal, 3>>(FVec3(0), BoundsExtents.Size() * 0.5f));
-				}
-			}
-			else // !UseLevelsetCollision
-			{
-				if (Objects.Num() == 1)
-				{
-					Parent->SetDynamicGeometry(MoveTemp(Objects[0]));
-				}
-				else
-				{
-					Parent->SetDynamicGeometry(
-						MakeUnique<FImplicitObjectUnionClustered>(
-							MoveTemp(Objects), ChildParticleHandles));
-				}
-			}
-		}
-
-		//if children are ignore analytic and this is a dynamic geom, mark it too. todo(ocohen): clean this up
-		if (bUseParticleImplicit && Parent->DynamicGeometry()) 
-		{
-			Parent->DynamicGeometry()->SetDoCollide(false);
-		}
-
-		if (Parameters.CollisionParticles)
-		{
-			SCOPE_CYCLE_COUNTER(STAT_UpdateGeometry_CopyPoints);
-			Parent->CollisionParticles().Reset(Parameters.CollisionParticles);
-		}
-		else
-		{
-			{
-				SCOPE_CYCLE_COUNTER(STAT_UpdateGeometry_GatherPoints);
-				Parent->CollisionParticlesInitIfNeeded();
-				TUniquePtr<FBVHParticles>& CollisionParticles = Parent->CollisionParticles();
-				CollisionParticles->AddParticles(CleanedPoints.Num());
-				for (int32 i = 0; i < CleanedPoints.Num(); ++i)
-				{
-					CollisionParticles->X(i) = CleanedPoints[i];
-				}
-			}
-
-			if (bUseCollisionPoints)
-			{
-				SCOPE_CYCLE_COUNTER(STAT_UpdateGeometry_PointsBVH);
-				Parent->CollisionParticles()->UpdateAccelerationStructures();
-			}
-		}
-
-		if (TSerializablePtr<FImplicitObject> Implicit = Parent->Geometry())
-		{
-			// strange hacked initilization that seems misplaced and ill thought
-			Parent->SetHasBounds(true);
-			Parent->SetLocalBounds(Implicit->BoundingBox());
-			const Chaos::FRigidTransform3 Xf(Parent->X(), Parent->R());
-			Parent->UpdateWorldSpaceState(Xf, FVec3(0));
-		}
-
-		// Set the captured filter to our new shapes
-		for(const TUniquePtr<FPerShapeData>& Shape : Parent->ShapesArray())
-		{
-			Shape->SetSimData(Filter);
-		}
-	}
-
 	FRealSingle MinImpulseForStrainEval = 980 * 2 * 1.f / 30.f; //ignore impulses caused by just keeping object on ground. This is a total hack, we should not use accumulated impulse directly. Instead we need to look at delta v along constraint normal
 	FAutoConsoleVariableRef CVarMinImpulseForStrainEval(TEXT("p.chaos.MinImpulseForStrainEval"), MinImpulseForStrainEval, TEXT("Minimum accumulated impulse before accumulating for strain eval "));
 
diff --git a/Engine/Source/Runtime/Experimental/Chaos/Private/Chaos/PBDRigidClusteringAlgo.cpp b/Engine/Source/Runtime/Experimental/Chaos/Private/Chaos/PBDRigidClusteringAlgo.cpp
new file mode 100644
--- /dev/null
+++ b/Engine/Source/Runtime/Experimental/Chaos/Private/Chaos/PBDRigidClusteringAlgo.cpp
@@ -0,0 +1,514 @@
+// Copyright Epic Games, Inc. All Rights Reserved.
+
+#include "Chaos/PBDRigidClusteringAlgo.h"
+
+#include "Chaos/ErrorReporter.h"
+#include "Chaos/Levelset.h"
+#include "Chaos/PBDRigidClusteringCollisionParticleAlgo.h"
+#include "Chaos/PBDRigidsEvolutionGBF.h"
+
+namespace Chaos
+{
+	//
+    // Update Geometry PVar
+	//
+
+	int32 UseLevelsetCollision = 0;
+	FAutoConsoleVariableRef CVarUseLevelsetCollision2(TEXT("p.UseLevelsetCollision"), UseLevelsetCollision, TEXT("Whether unioned objects use levelsets"));
+
+	int32 MinLevelsetDimension = 4;
+	FAutoConsoleVariableRef CVarMinLevelsetDimension2(TEXT("p.MinLevelsetDimension"), MinLevelsetDimension, TEXT("The minimum number of cells on a single level set axis"));
+
+	int32 MaxLevelsetDimension = 20;
+	FAutoConsoleVariableRef CVarMaxLevelsetDimension2(TEXT("p.MaxLevelsetDimension"), MaxLevelsetDimension, TEXT("The maximum number of cells on a single level set axis"));
+
+	FRealSingle MinLevelsetSize = 50.f;
+	FAutoConsoleVariableRef CVarLevelSetResolution2(TEXT("p.MinLevelsetSize"), MinLevelsetSize, TEXT("The minimum size on the smallest axis to use a level set"));
+
+	int32 LevelsetGhostCells = 1;
+	FAutoConsoleVariableRef CVarLevelsetGhostCells2(TEXT("p.LevelsetGhostCells"), LevelsetGhostCells, TEXT("Increase the level set grid by this many ghost cells"));
+
+	int32 MinCleanedPointsBeforeRemovingInternals = 10;
+	FAutoConsoleVariableRef CVarMinCleanedPointsBeforeRemovingInternals2(TEXT("p.MinCleanedPointsBeforeRemovingInternals"), MinCleanedPointsBeforeRemovingInternals, TEXT("If we only have this many clean points, don't bother removing internal points as the object is likely very small"));
+
+	FRealSingle ClusterSnapDistance = 1.f;
+	FAutoConsoleVariableRef CVarClusterSnapDistance2(TEXT("p.ClusterSnapDistance"), ClusterSnapDistance, TEXT(""));
+
+	DECLARE_CYCLE_STAT(TEXT("TPBDRigidClustering<>::UpdateClusterMassProperties()"), STAT_UpdateClusterMassProperties, STATGROUP_Chaos);
+	void UpdateClusterMassProperties(
+		FPBDRigidClusteredParticleHandle* Parent,
+		TSet<FPBDRigidParticleHandle*>& Children,
+		const FRigidTransform3* ForceMassOrientation)
+	{
+		SCOPE_CYCLE_COUNTER(STAT_UpdateClusterMassProperties);
+		check(Children.Num());
+
+		Parent->SetX(FVec3(0));
+		Parent->SetR(FRotation3(FQuat::MakeFromEuler(FVec3(0))));
+		Parent->SetV(FVec3(0));
+		Parent->SetW(FVec3(0));
+		Parent->SetM(0);
+		Parent->SetI(FMatrix33(0));
+
+		bool bHasChild = false;
+		bool bHasProxyChild = false;
+		for (FPBDRigidParticleHandle* OriginalChild : Children)
+		{
+			FPBDRigidParticleHandle* Child;
+			FVec3 ChildPosition;
+			FRotation3 ChildRotation;
+
+			Child = OriginalChild;
+			ChildPosition = Child->X();
+			ChildRotation = Child->R();
+
+			const FReal ChildMass = Child->M();
+			const FMatrix33 ChildWorldSpaceI =
+				(ChildRotation * FMatrix::Identity) * Child->I() * (ChildRotation * FMatrix::Identity).GetTransposed();
+			if (ChildWorldSpaceI.ContainsNaN())
+			{
+				continue;
+			}
+			bHasProxyChild = true;
+			bHasChild = true;
+			bHasProxyChild = true;
+			Parent->I() += ChildWorldSpaceI;
+			Parent->M() += ChildMass;
+			Parent->X() += ChildPosition * ChildMass;
+			Parent->V() += OriginalChild->V() * ChildMass; // Use orig child for vel because we don't sim the proxy
+			Parent->W() += OriginalChild->W() * ChildMass;
+		}
+		if (!ensure(bHasProxyChild))
+		{
+			for (FPBDRigidParticleHandle* OriginalChild : Children)
+			{
+				FPBDRigidParticleHandle* Child = OriginalChild;
+				const FVec3& ChildPosition = Child->X();
+				const FRotation3& ChildRotation = Child->R();
+				const FReal ChildMass = Child->M();
+
+				const FMatrix33 ChildWorldSpaceI =
+					(ChildRotation * FMatrix::Identity) * Child->I() * (ChildRotation * FMatrix::Identity).GetTransposed();
+				if (ChildWorldSpaceI.ContainsNaN())
+				{
+					continue;
+				}
+				bHasChild = true;
+				Parent->I() += ChildWorldSpaceI;
+				Parent->M() += ChildMass;
+				Parent->X() += ChildPosition * ChildMass;
+				Parent->V() += OriginalChild->V() * ChildMass; // Use orig child for vel because we don't sim the proxy
+				Parent->W() += OriginalChild->W() * ChildMass;
+			}
+		}
+		for (int32 i = 0; i < 3; i++)
+		{
+			const FMatrix33& InertiaTensor = Parent->I();
+			if (InertiaTensor.GetColumn(i)[i] < SMALL_NUMBER)
+			{
+				Parent->SetI(FMatrix33(1.f, 1.f, 1.f));
+				break;
+			}
+		}
+
+		if (!ensure(bHasChild) || !ensure(Parent->M() > SMALL_NUMBER))
+		{
+			Parent->M() = 1.0;
+			Parent->X() = FVec3(0);
+			Parent->V() = FVec3(0);
+			Parent->PreV() = Parent->V();
+			Parent->InvM() = 1;
+			Parent->P() = Parent->X();
+			Parent->W() = FVec3(0);
+			Parent->PreW() = Parent->W();
+			Parent->R() = FRotation3(FMatrix::Identity);
+			Parent->Q() = Parent->R();
+			Parent->I() = FMatrix::Identity;
+			Parent->InvI() = FMatrix::Identity;
+			return;
+		}
+
+		check(Parent->M() > SMALL_NUMBER);
+
+		Parent->X() /= Parent->M();
+		Parent->V() /= Parent->M();
+		Parent->PreV() = Parent->V();
+		Parent->InvM() = static_cast<FReal>(1.0) / Parent->M();
+		if (ForceMassOrientation)
+		{
+			Parent->X() = ForceMassOrientation->GetLocation();
+		}
+		Parent->P() = Parent->X();
+		for (FPBDRigidParticleHandle* OriginalChild : Children)
+		{
+			FPBDRigidParticleHandle* Child;
+			FVec3 ChildPosition;
+
+			Child = OriginalChild;
+			ChildPosition = Child->X();
+
+			FVec3 ParentToChild = ChildPosition - Parent->X();
+
+			const FReal ChildMass = Child->M();
+			// taking v from original child since we are not actually simulating the proxy child
+			Parent->W() +=
+				FVec3::CrossProduct(ParentToChild,
+					OriginalChild->V() * ChildMass);
+			{
+				const FReal& p0 = ParentToChild[0];
+				const FReal& p1 = ParentToChild[1];
+				const FReal& p2 = ParentToChild[2];
+				const FReal& m = ChildMass;
+				Parent->I() +=
+					FMatrix33(
+						m * (p1 * p1 + p2 * p2), -m * p1 * p0, -m * p2 * p0,
+						m * (p2 * p2 + p0 * p0), -m * p2 * p1, m * (p1 * p1 + p0 * p0));
+			}
+		}
+		FMatrix33& InertiaTensor = Parent->I();
+		if (Parent->I().ContainsNaN())
+		{
+			InertiaTensor = FMatrix33((FReal)1., (FReal)1., (FReal)1.);
+		}
+		else
+		{
+			for (int32 i = 0; i < 3; i++)
+			{
+				if (InertiaTensor.GetColumn(i)[i] < SMALL_NUMBER)
+				{
+					InertiaTensor = FMatrix33((FReal)1., (FReal)1., (FReal)1.);
+					break;
+				}
+			}
+		}
+		Parent->W() /= Parent->M();
+		Parent->PreW() = Parent->W();
+		if (ForceMassOrientation)
+		{
+			Parent->R() = ForceMassOrientation->GetRotation();
+		}
+
+		if (Parent->R().ContainsNaN())
+		{
+			InertiaTensor = PMatrix<FReal, 3, 3>(1.f, 1.f, 1.f);
+			Parent->R() = TRotation<FReal, 3>(FMatrix::Identity);
+		}
+
+		Parent->Q() = Parent->R();
+		Parent->InvI() = Parent->I().Inverse();
+	}
+
+
+
+	DECLARE_CYCLE_STAT(TEXT("TPBDRigidClustering<>::UpdateKinematicProperties()"), STAT_UpdateKinematicProperties, STATGROUP_Chaos);
+	void 
+		UpdateKinematicProperties(
+			Chaos::FPBDRigidParticleHandle* InParent,
+			FRigidClustering::FClusterMap& MChildren,
+			FRigidClustering::FRigidEvolution& MEvolution)
+	{
+		SCOPE_CYCLE_COUNTER(STAT_UpdateKinematicProperties);
+		typedef FPBDRigidClusteredParticleHandle* FClusterHandle;
+
+		EObjectStateType ObjectState = EObjectStateType::Dynamic;
+		check(InParent != nullptr);
+		if (FClusterHandle ClusteredCurrentNode = InParent->CastToClustered())
+		{
+			if (MChildren.Contains(ClusteredCurrentNode) && MChildren[ClusteredCurrentNode].Num())
+			{
+				// TQueue is a linked list, which has no preallocator.
+				TQueue<Chaos::FPBDRigidParticleHandle*> Queue;
+				for (Chaos::FPBDRigidParticleHandle* Child : MChildren[ClusteredCurrentNode])
+				{
+					Queue.Enqueue(Child);
+				}
+
+				Chaos::FPBDRigidParticleHandle* CurrentHandle;
+				while (Queue.Dequeue(CurrentHandle) && ObjectState == EObjectStateType::Dynamic)
+				{
+					if (FClusterHandle CurrentClusterHandle = CurrentHandle->CastToClustered())
+					{
+						// @question : Maybe we should just store the leaf node bodies in a
+						// map, that will require Memory(n*log(n))
+						if (MChildren.Contains(CurrentClusterHandle))
+						{
+							for (Chaos::FPBDRigidParticleHandle* Child : MChildren[CurrentClusterHandle])
+							{
+								Queue.Enqueue(Child);
+							}
+						}
+					}
+
+					const EObjectStateType CurrState = CurrentHandle->ObjectState();
+					if (CurrState == EObjectStateType::Kinematic)
+					{
+						ObjectState = EObjectStateType::Kinematic;
+					}
+					else if (CurrState == EObjectStateType::Static)
+					{
+						ObjectState = EObjectStateType::Static;
+					}
+				}
+
+				MEvolution.SetParticleObjectState(ClusteredCurrentNode, ObjectState);
+			}
+		}
+	}
+	
+
+	DECLARE_CYCLE_STAT(TEXT("TPBDRigidClustering<>::UpdateGeometry"), STAT_UpdateGeometry, STATGROUP_Chaos);
+	DECLARE_CYCLE_STAT(TEXT("TPBDRigidClustering<>::UpdateGeometry_GatherObjects"), STAT_UpdateGeometry_GatherObjects, STATGROUP_Chaos);
+	DECLARE_CYCLE_STAT(TEXT("TPBDRigidClustering<>::UpdateGeometry_GatherPoints"), STAT_UpdateGeometry_GatherPoints, STATGROUP_Chaos);
+	DECLARE_CYCLE_STAT(TEXT("TPBDRigidClustering<>::UpdateGeometry_CopyPoints"), STAT_UpdateGeometry_CopyPoints, STATGROUP_Chaos);
+	DECLARE_CYCLE_STAT(TEXT("TPBDRigidClustering<>::UpdateGeometry_PointsBVH"), STAT_UpdateGeometry_PointsBVH, STATGROUP_Chaos);
+
+	void
+	UpdateGeometry(
+		Chaos::FPBDRigidClusteredParticleHandle* Parent,
+		const TSet<FPBDRigidParticleHandle*>& Children,
+		TSharedPtr<Chaos::FImplicitObject, ESPMode::ThreadSafe> ProxyGeometry,
+		const FClusterCreationParameters& Parameters)
+	{
+		SCOPE_CYCLE_COUNTER(STAT_UpdateGeometry);
+
+		TArray<TUniquePtr<FImplicitObject>> Objects;
+		TArray<TUniquePtr<FImplicitObject>> Objects2; //todo: find a better way to reuse this
+		Objects.Reserve(Children.Num());
+		Objects2.Reserve(Children.Num());
+
+		//we should never update existing geometry since this is used by SQ threads.
+		ensure(!Parent->Geometry());
+		ensure(!Parent->DynamicGeometry());
+
+		const FRigidTransform3 ClusterWorldTM(Parent->X(), Parent->R());
+
+		TArray<FVec3> OriginalPoints;
+		TArray<FPBDRigidParticleHandle*> ChildParticleHandles;
+		ChildParticleHandles.Reserve(Children.Num());
+
+		const bool bUseCollisionPoints = (ProxyGeometry || Parameters.bCopyCollisionParticles) && !Parameters.CollisionParticles;
+		bool bUseParticleImplicit = false;
+
+		// Need to extract a filter off one of the cluster children 
+		FCollisionFilterData Filter;
+		for (FPBDRigidParticleHandle* Child : Children)
+		{
+			bool bFilterValid = false;
+			for (const TUniquePtr<FPerShapeData>& Shape : Child->ShapesArray())
+			{
+				if (Shape)
+				{
+					Filter = Shape->GetSimData();
+					bFilterValid = Filter.Word0 != 0 || Filter.Word1 != 0 || Filter.Word2 != 0 || Filter.Word3 != 0;
+				}
+
+				if (bFilterValid)
+				{
+					break;
+				}
+			}
+
+			// Bail once we've found one
+			if (bFilterValid)
+			{
+				break;
+			}
+		}
+
+		{ // STAT_UpdateGeometry_GatherObjects
+			SCOPE_CYCLE_COUNTER(STAT_UpdateGeometry_GatherObjects);
+
+			if (bUseCollisionPoints)
+			{
+				uint32 NumPoints = 0;
+				for (FPBDRigidParticleHandle* Child : Children)
+				{
+					NumPoints += Child->CollisionParticlesSize();
+				}
+				OriginalPoints.Reserve(NumPoints);
+			}
+
+			for (FPBDRigidParticleHandle* Child : Children)
+			{
+				const FRigidTransform3 ChildWorldTM(Child->X(), Child->R());
+				FRigidTransform3 Frame = ChildWorldTM.GetRelativeTransform(ClusterWorldTM);
+				FPBDRigidParticleHandle* UsedGeomChild = Child;
+				if (Child->Geometry())
+				{
+					Objects.Add(TUniquePtr<FImplicitObject>(new TImplicitObjectTransformed<FReal, 3>(Child->Geometry(), Frame)));
+					Objects2.Add(TUniquePtr<FImplicitObject>(new TImplicitObjectTransformed<FReal, 3>(Child->Geometry(), Frame)));
+					ChildParticleHandles.Add(Child);
+				}
+
+				ensure(Child->Disabled() == true);
+				check(Child->CastToClustered()->ClusterIds().Id == Parent);
+
+				Child->CastToClustered()->SetChildToParent(Frame);
+
+				if (bUseCollisionPoints)
+				{
+					SCOPE_CYCLE_COUNTER(STAT_UpdateGeometry_GatherPoints);
+					if (const TUniquePtr<FBVHParticles>& CollisionParticles = Child->CollisionParticles())
+					{
+						for (uint32 i = 0; i < CollisionParticles->Size(); ++i)
+						{
+							OriginalPoints.Add(Frame.TransformPosition(CollisionParticles->X(i)));
+						}
+					}
+				}
+				if (Child->Geometry() && Child->Geometry()->GetType() == ImplicitObjectType::Unknown)
+				{
+					bUseParticleImplicit = true;
+				}
+			} // end for
+		} // STAT_UpdateGeometry_GatherObjects
+
+		{
+			QUICK_SCOPE_CYCLE_COUNTER(SpatialBVH);
+			TUniquePtr<FImplicitObjectUnionClustered>& ChildrenSpatial = Parent->ChildrenSpatial();
+			ChildrenSpatial.Reset(
+				Objects2.Num() ?
+				new Chaos::FImplicitObjectUnionClustered(MoveTemp(Objects2), ChildParticleHandles) :
+				nullptr);
+		}
+
+		TArray<FVec3> CleanedPoints;
+		if (!Parameters.CollisionParticles)
+		{
+			CleanedPoints =
+				Parameters.bCleanCollisionParticles ?
+				CleanCollisionParticles(OriginalPoints, ClusterSnapDistance) :
+				OriginalPoints;
+		}
+
+		if (ProxyGeometry)
+		{
+			const FVector Scale = Parameters.Scale;
+			auto DeepCopyImplicit = [&Scale](const TSharedPtr<Chaos::FImplicitObject, ESPMode::ThreadSafe>& ImplicitToCopy) -> TUniquePtr<Chaos::FImplicitObject>
+			{
+				if (Scale.Equals(FVector::OneVector))
+				{
+					return ImplicitToCopy->DeepCopy();
+				}
+				else
+				{
+					return ImplicitToCopy->DeepCopyWithScale(Scale);
+				}
+			};
+			//ensureMsgf(false, TEXT("Checking usage with proxy"));
+			//@coverage {production}
+			Parent->SetSharedGeometry(TSharedPtr<FImplicitObject, ESPMode::ThreadSafe>(DeepCopyImplicit(ProxyGeometry).Release()));
+		}
+		else if (Objects.Num() == 0)
+		{
+			//ensureMsgf(false, TEXT("Checking usage with no proxy and no objects"));
+			//@coverage : {production}
+			Parent->SetGeometry(Chaos::TSerializablePtr<Chaos::FImplicitObject>());
+		}
+		else
+		{
+			if (UseLevelsetCollision)
+			{
+				ensureMsgf(false, TEXT("Checking usage with no proxy and multiple ojects with levelsets"));
+
+				FImplicitObjectUnionClustered UnionObject(MoveTemp(Objects));
+				FAABB3 Bounds = UnionObject.BoundingBox();
+				const FVec3 BoundsExtents = Bounds.Extents();
+				if (BoundsExtents.Min() >= MinLevelsetSize) //make sure the object is not too small
+				{
+					TVec3<int32> NumCells = Bounds.Extents() / MinLevelsetSize;
+					for (int i = 0; i < 3; ++i)
+					{
+						NumCells[i] = FMath::Clamp(NumCells[i], MinLevelsetDimension, MaxLevelsetDimension);
+					}
+
+					FErrorReporter ErrorReporter;
+					TUniformGrid<FReal, 3> Grid(Bounds.Min(), Bounds.Max(), NumCells, LevelsetGhostCells);
+					TUniquePtr<FLevelSet> LevelSet(new FLevelSet(ErrorReporter, Grid, UnionObject));
+
+					if (!Parameters.CollisionParticles)
+					{
+						const FReal MinDepthToSurface = Grid.Dx().Max();
+						for (int32 Idx = CleanedPoints.Num() - 1; Idx >= 0; --Idx)
+						{
+							if (CleanedPoints.Num() > MinCleanedPointsBeforeRemovingInternals) //todo(ocohen): this whole thing should really be refactored
+							{
+								const FVec3& CleanedCollision = CleanedPoints[Idx];
+								if (LevelSet->SignedDistance(CleanedCollision) < -MinDepthToSurface)
+								{
+									CleanedPoints.RemoveAtSwap(Idx);
+								}
+							}
+						}
+					}
+					Parent->SetDynamicGeometry(MoveTemp(LevelSet));
+				}
+				else
+				{
+					Parent->SetDynamicGeometry(
+						MakeUnique<TSphere<FReal, 3>>(FVec3(0), BoundsExtents.Size() * 0.5f));
+				}
+			}
+			else // !UseLevelsetCollision
+			{
+				if (Objects.Num() == 1)
+				{
+					Parent->SetDynamicGeometry(MoveTemp(Objects[0]));
+				}
+				else
+				{
+					Parent->SetDynamicGeometry(
+						MakeUnique<FImplicitObjectUnionClustered>(
+							MoveTemp(Objects), ChildParticleHandles));
+				}
+			}
+		}
+
+		//if children are ignore analytic and this is a dynamic geom, mark it too. todo(ocohen): clean this up
+		if (bUseParticleImplicit && Parent->DynamicGeometry())
+		{
+			Parent->DynamicGeometry()->SetDoCollide(false);
+		}
+
+		if (Parameters.CollisionParticles)
+		{
+			SCOPE_CYCLE_COUNTER(STAT_UpdateGeometry_CopyPoints);
+			Parent->CollisionParticles().Reset(Parameters.CollisionParticles);
+		}
+		else
+		{
+			{
+				SCOPE_CYCLE_COUNTER(STAT_UpdateGeometry_GatherPoints);
+				Parent->CollisionParticlesInitIfNeeded();
+				TUniquePtr<FBVHParticles>& CollisionParticles = Parent->CollisionParticles();
+				CollisionParticles->AddParticles(CleanedPoints.Num());
+				for (int32 i = 0; i < CleanedPoints.Num(); ++i)
+				{
+					CollisionParticles->X(i) = CleanedPoints[i];
+				}
+			}
+
+			if (bUseCollisionPoints)
+			{
+				SCOPE_CYCLE_COUNTER(STAT_UpdateGeometry_PointsBVH);
+				Parent->CollisionParticles()->UpdateAccelerationStructures();
+			}
+		}
+
+		if (TSerializablePtr<FImplicitObject> Implicit = Parent->Geometry())
+		{
+			// strange hacked initilization that seems misplaced and ill thought
+			Parent->SetHasBounds(true);
+			Parent->SetLocalBounds(Implicit->BoundingBox());
+			const Chaos::FRigidTransform3 Xf(Parent->X(), Parent->R());
+			Parent->UpdateWorldSpaceState(Xf, FVec3(0));
+		}
+
+		// Set the captured filter to our new shapes
+		for (const TUniquePtr<FPerShapeData>& Shape : Parent->ShapesArray())
+		{
+			Shape->SetSimData(Filter);
+		}
+	}
+
+} // namespace Chaos
diff --git a/Engine/Source/Runtime/Experimental/Chaos/Private/GeometryCollection/GeometryCollectionCollisionStructureManager.cpp b/Engine/Source/Runtime/Experimental/Chaos/Private/GeometryCollection/GeometryCollectionCollisionStructureManager.cpp
--- a/Engine/Source/Runtime/Experimental/Chaos/Private/GeometryCollection/GeometryCollectionCollisionStructureManager.cpp
+++ b/Engine/Source/Runtime/Experimental/Chaos/Private/GeometryCollection/GeometryCollectionCollisionStructureManager.cpp
@@ -13,6 +13,7 @@
 #include "HAL/IConsoleManager.h"
 #include "Chaos/TriangleMesh.h"
 #include "Chaos/PBDRigidClustering.h"
+#include "Chaos/PBDRigidClusteringCollisionParticleAlgo.h"
 
 DEFINE_LOG_CATEGORY_STATIC(GCS_Log, NoLogging, All);
 
diff --git a/Engine/Source/Runtime/Experimental/Chaos/Public/Chaos/PBDRigidClustering.h b/Engine/Source/Runtime/Experimental/Chaos/Public/Chaos/PBDRigidClustering.h
--- a/Engine/Source/Runtime/Experimental/Chaos/Public/Chaos/PBDRigidClustering.h
+++ b/Engine/Source/Runtime/Experimental/Chaos/Public/Chaos/PBDRigidClustering.h
@@ -170,12 +170,6 @@ public:
 	TMap<FPBDRigidClusteredParticleHandle*, TSet<FPBDRigidParticleHandle*>> BreakingModel(
 		TMap<FGeometryParticleHandle*, Chaos::FReal>* ExternalStrainMap = nullptr);
 
-	/*
-	*  Process the kinematic state of the clusters. Because the leaf node geometry can
-	*  be changed by the solver, it is necessary to check all the sub clusters.
-	*/
-	void UpdateKinematicProperties(FPBDRigidParticleHandle* Parent);
-
 	//
 	// Access
 	//
@@ -195,8 +189,8 @@ public:
 	*    active id, see the GetActiveClusterIndex to find the active cluster.
 	*    INDEX_NONE represents a non-clustered body.
 	*/
-	TArrayCollectionArray<ClusterId>&					GetClusterIdsArray() { return MParticles.ClusterIdsArray(); }
-	const TArrayCollectionArray<ClusterId>&             GetClusterIdsArray() const { return MParticles.ClusterIdsArray(); }
+	TArrayCollectionArray<ClusterId>& GetClusterIdsArray() { return MParticles.ClusterIdsArray(); }
+	const TArrayCollectionArray<ClusterId>& GetClusterIdsArray() const { return MParticles.ClusterIdsArray(); }
 
 	/*
 	*  GetInternalClusterArray
@@ -225,7 +219,7 @@ public:
 	*    the keys of the map. The value of the map is a pointer to an array  constrained
 	*    rigid bodies.
 	*/
-	FClusterMap &       GetChildrenMap() { return MChildren; }
+	FClusterMap & GetChildrenMap() { return MChildren; }
 	const FClusterMap & GetChildrenMap() const { return MChildren; }
 
 	/*
@@ -267,16 +261,6 @@ public:
 	TSet<Chaos::FPBDRigidClusteredParticleHandle*>& GetTopLevelClusterParents() { return TopLevelClusterParents; }
 
  protected:
-	void UpdateMassProperties(
-		Chaos::FPBDRigidClusteredParticleHandle* Parent, 
-		TSet<FPBDRigidParticleHandle*>& Children, 
-		const FRigidTransform3* ForceMassOrientation);
-
-	void UpdateGeometry(
-		Chaos::FPBDRigidClusteredParticleHandle* Parent, 
-		const TSet<FPBDRigidParticleHandle*>& Children, 
-		TSharedPtr<Chaos::FImplicitObject, ESPMode::ThreadSafe> ProxyGeometry,
-		const FClusterCreationParameters& Parameters);
 
 	void ComputeStrainFromCollision(const FPBDCollisionConstraints& CollisionRule);
 	void ResetCollisionImpulseArray();
@@ -330,158 +314,4 @@ private:
 	FClusterCreationParameters::EConnectionMethod MClusterUnionConnectionType;
 };
 
-void UpdateClusterMassProperties(
-	Chaos::FPBDRigidClusteredParticleHandle* Parent,
-	TSet<FPBDRigidParticleHandle*>& Children,
-	const FRigidTransform3* ForceMassOrientation = nullptr);
-
-inline TArray<FVec3> CleanCollisionParticles(
-	const TArray<FVec3>& Vertices,
-	FAABB3 BBox, 
-	const FReal SnapDistance=(FReal)0.01)
-{
-	const int32 NumPoints = Vertices.Num();
-	if (NumPoints <= 1)
-		return TArray<FVec3>(Vertices);
-
-	FReal MaxBBoxDim = BBox.Extents().Max();
-	if (MaxBBoxDim < SnapDistance)
-		return TArray<FVec3>(&Vertices[0], 1);
-
-	BBox.Thicken(FMath::Max(SnapDistance/10, (FReal)(KINDA_SMALL_NUMBER*10))); // 0.001
-	MaxBBoxDim = BBox.Extents().Max();
-
-	const FVec3 PointsCenter = BBox.Center();
-	TArray<FVec3> Points(Vertices);
-
-	// Find coincident vertices.  We hash to a grid of fine enough resolution such
-	// that if 2 particles hash to the same cell, then we're going to consider them
-	// coincident.
-	TSet<int64> OccupiedCells;
-	OccupiedCells.Reserve(NumPoints);
-
-	TArray<int32> Redundant;
-	Redundant.Reserve(NumPoints); // Excessive, but ensures consistent performance.
-
-	int32 NumCoincident = 0;
-	const int64 Resolution = static_cast<int64>(floor(MaxBBoxDim / FMath::Max(SnapDistance,(FReal)KINDA_SMALL_NUMBER)));
-	const FReal CellSize = static_cast<FReal>(static_cast<double>(MaxBBoxDim) / static_cast<double>(Resolution));
-	for (int32 i = 0; i < 2; i++)
-	{
-		Redundant.Reset();
-		OccupiedCells.Reset();
-		// Shift the grid by 1/2 a grid cell the second iteration so that
-		// we don't miss slightly adjacent coincident points across cell
-		// boundaries.
-		const FVec3 GridCenter = FVec3(0) - FVec3(static_cast<FReal>(i) * CellSize / 2);
-		for (int32 j = 0; j < Points.Num(); j++)
-		{
-			const FVec3 Pos = Points[j] - PointsCenter; // Centered at the origin
-			const TVec3<int64> Coord(
-				static_cast<int64>(FMath::Floor((Pos[0] - GridCenter[0]) / CellSize + static_cast<double>(Resolution) / 2)),
-				static_cast<int64>(FMath::Floor((Pos[1] - GridCenter[1]) / CellSize + static_cast<double>(Resolution) / 2)),
-				static_cast<int64>(FMath::Floor((Pos[2] - GridCenter[2]) / CellSize + static_cast<double>(Resolution) / 2)));
-			const int64 FlatIdx =
-				((Coord[0] * Resolution + Coord[1]) * Resolution) + Coord[2];
-
-			bool AlreadyInSet = false;
-			OccupiedCells.Add(FlatIdx, &AlreadyInSet);
-			if (AlreadyInSet)
-				Redundant.Add(j);
-		}
-
-		for (int32 j = Redundant.Num(); j--;)
-		{
-			Points.RemoveAt(Redundant[j]);
-		}
-	}
-
-	// Shrink the array, if appropriate
-	Points.SetNum(Points.Num(), true);
-	return Points;
-}
-
-inline TArray<FVec3> CleanCollisionParticles(
-	const TArray<FVec3>& Vertices, 
-	const FReal SnapDistance=(FReal)0.01)
-{
-	if (!Vertices.Num())
-	{
-		return TArray<FVec3>();
-	}
-	FAABB3 BBox(FAABB3::EmptyAABB());
-	for (const FVec3& Pt : Vertices)
-	{
-		BBox.GrowToInclude(Pt);
-	}
-	return CleanCollisionParticles(Vertices, BBox, SnapDistance);
-}
-
-inline TArray<FVec3> CleanCollisionParticles(
-	FTriangleMesh &TriMesh, 
-	const TArrayView<const FVec3>& Vertices, 
-	const FReal Fraction)
-{
-	TArray<FVec3> CollisionVertices;
-	if (Fraction <= 0.0)
-		return CollisionVertices;
-
-	// If the tri mesh has any open boundaries, see if we can merge any coincident
-	// vertices on the boundary.  This makes the importance ordering work much better
-	// as we need the curvature at each edge of the tri mesh, and we can't calculate
-	// curvature on discontiguous triangles.
-	TSet<int32> BoundaryPoints = TriMesh.GetBoundaryPoints();
-	if (BoundaryPoints.Num())
-	{
-		TMap<int32, int32> Remapping =
-			TriMesh.FindCoincidentVertexRemappings(BoundaryPoints.Array(), Vertices);
-		TriMesh.RemapVertices(Remapping);
-	}
-
-	// Get the importance vertex ordering, from most to least.  Reorder the 
-	// particles accordingly.
-	TArray<int32> CoincidentVertices;
-	const TArray<int32> Ordering = TriMesh.GetVertexImportanceOrdering(Vertices, &CoincidentVertices, true);
-
-	// Particles are ordered from most important to least, with coincident 
-	// vertices at the very end.
-	const int32 NumGoodPoints = Ordering.Num() - CoincidentVertices.Num();
-
-#if DO_GUARD_SLOW
-	for (int i = NumGoodPoints; i < Ordering.Num(); ++i)
-	{
-		ensure(CoincidentVertices.Contains(Ordering[i]));	//make sure all coincident vertices are at the back
-	}
-#endif
-
-	CollisionVertices.AddUninitialized(FMath::Min(NumGoodPoints, static_cast<int32>(ceil(static_cast<FReal>(NumGoodPoints) * Fraction))));
-	for (int i = 0; i < CollisionVertices.Num(); i++)
-	{
-		CollisionVertices[i] = Vertices[Ordering[i]];
-	}
-	return CollisionVertices;
-}
-
-inline void CleanCollisionParticles(
-	FTriangleMesh &TriMesh, 
-	const TArrayView<const FVec3>& Vertices, 
-	const FReal Fraction,
-	TSet<int32>& ResultingIndices)
-{
-	ResultingIndices.Reset();
-	if (Fraction <= 0.0)
-		return;
-
-	TArray<int32> CoincidentVertices;
-	const TArray<int32> Ordering = TriMesh.GetVertexImportanceOrdering(Vertices, &CoincidentVertices, true);
-	int32 NumGoodPoints = Ordering.Num() - CoincidentVertices.Num();
-	NumGoodPoints = FMath::Min(NumGoodPoints, static_cast<int32>(ceil(static_cast<FReal>(NumGoodPoints) * Fraction)));
-
-	ResultingIndices.Reserve(NumGoodPoints);
-	for (int32 i = 0; i < NumGoodPoints; i++)
-	{
-		ResultingIndices.Add(Ordering[i]);
-	}
-}
-
 } // namespace Chaos
diff --git a/Engine/Source/Runtime/Experimental/Chaos/Public/Chaos/PBDRigidClusteringAlgo.h b/Engine/Source/Runtime/Experimental/Chaos/Public/Chaos/PBDRigidClusteringAlgo.h
new file mode 100644
--- /dev/null
+++ b/Engine/Source/Runtime/Experimental/Chaos/Public/Chaos/PBDRigidClusteringAlgo.h
@@ -0,0 +1,24 @@
+// Copyright Epic Games, Inc. All Rights Reserved.
+#pragma once
+
+#include "Chaos/PBDRigidClustering.h"
+
+namespace Chaos
+{
+	void UpdateClusterMassProperties(
+		FPBDRigidClusteredParticleHandle* Parent,
+		TSet<FPBDRigidParticleHandle*>& Children,
+		const FRigidTransform3* ForceMassOrientation = nullptr);
+
+	void CHAOS_API UpdateKinematicProperties(
+		FPBDRigidParticleHandle* Parent,
+		FRigidClustering::FClusterMap&,
+		FRigidClustering::FRigidEvolution&);
+
+	void CHAOS_API UpdateGeometry(
+		Chaos::FPBDRigidClusteredParticleHandle* Parent,
+		const TSet<FPBDRigidParticleHandle*>& Children,
+		TSharedPtr<Chaos::FImplicitObject, ESPMode::ThreadSafe> ProxyGeometry,
+		const FClusterCreationParameters& Parameters);
+
+} // namespace Chaos
diff --git a/Engine/Source/Runtime/Experimental/Chaos/Public/Chaos/PBDRigidClusteringCollisionParticleAlgo.h b/Engine/Source/Runtime/Experimental/Chaos/Public/Chaos/PBDRigidClusteringCollisionParticleAlgo.h
new file mode 100644
--- /dev/null
+++ b/Engine/Source/Runtime/Experimental/Chaos/Public/Chaos/PBDRigidClusteringCollisionParticleAlgo.h
@@ -0,0 +1,158 @@
+// Copyright Epic Games, Inc. All Rights Reserved.
+#pragma once
+
+#include "Chaos/PBDRigidClustering.h"
+
+namespace Chaos
+{
+
+inline TArray<FVec3> CleanCollisionParticles(
+	const TArray<FVec3>& Vertices,
+	FAABB3 BBox, 
+	const FReal SnapDistance=(FReal)0.01)
+{
+	const int32 NumPoints = Vertices.Num();
+	if (NumPoints <= 1)
+		return TArray<FVec3>(Vertices);
+
+	FReal MaxBBoxDim = BBox.Extents().Max();
+	if (MaxBBoxDim < SnapDistance)
+		return TArray<FVec3>(&Vertices[0], 1);
+
+	BBox.Thicken(FMath::Max(SnapDistance/10, (FReal)(KINDA_SMALL_NUMBER*10))); // 0.001
+	MaxBBoxDim = BBox.Extents().Max();
+
+	const FVec3 PointsCenter = BBox.Center();
+	TArray<FVec3> Points(Vertices);
+
+	// Find coincident vertices.  We hash to a grid of fine enough resolution such
+	// that if 2 particles hash to the same cell, then we're going to consider them
+	// coincident.
+	TSet<int64> OccupiedCells;
+	OccupiedCells.Reserve(NumPoints);
+
+	TArray<int32> Redundant;
+	Redundant.Reserve(NumPoints); // Excessive, but ensures consistent performance.
+
+	int32 NumCoincident = 0;
+	const int64 Resolution = static_cast<int64>(floor(MaxBBoxDim / FMath::Max(SnapDistance,(FReal)KINDA_SMALL_NUMBER)));
+	const FReal CellSize = static_cast<FReal>(static_cast<double>(MaxBBoxDim) / static_cast<double>(Resolution));
+	for (int32 i = 0; i < 2; i++)
+	{
+		Redundant.Reset();
+		OccupiedCells.Reset();
+		// Shift the grid by 1/2 a grid cell the second iteration so that
+		// we don't miss slightly adjacent coincident points across cell
+		// boundaries.
+		const FVec3 GridCenter = FVec3(0) - FVec3(static_cast<FReal>(i) * CellSize / 2);
+		for (int32 j = 0; j < Points.Num(); j++)
+		{
+			const FVec3 Pos = Points[j] - PointsCenter; // Centered at the origin
+			const TVec3<int64> Coord(
+				static_cast<int64>(FMath::Floor((Pos[0] - GridCenter[0]) / CellSize + static_cast<double>(Resolution) / 2)),
+				static_cast<int64>(FMath::Floor((Pos[1] - GridCenter[1]) / CellSize + static_cast<double>(Resolution) / 2)),
+				static_cast<int64>(FMath::Floor((Pos[2] - GridCenter[2]) / CellSize + static_cast<double>(Resolution) / 2)));
+			const int64 FlatIdx =
+				((Coord[0] * Resolution + Coord[1]) * Resolution) + Coord[2];
+
+			bool AlreadyInSet = false;
+			OccupiedCells.Add(FlatIdx, &AlreadyInSet);
+			if (AlreadyInSet)
+				Redundant.Add(j);
+		}
+
+		for (int32 j = Redundant.Num(); j--;)
+		{
+			Points.RemoveAt(Redundant[j]);
+		}
+	}
+
+	// Shrink the array, if appropriate
+	Points.SetNum(Points.Num(), true);
+	return Points;
+}
+
+inline TArray<FVec3> CleanCollisionParticles(
+	const TArray<FVec3>& Vertices, 
+	const FReal SnapDistance=(FReal)0.01)
+{
+	if (!Vertices.Num())
+	{
+		return TArray<FVec3>();
+	}
+	FAABB3 BBox(FAABB3::EmptyAABB());
+	for (const FVec3& Pt : Vertices)
+	{
+		BBox.GrowToInclude(Pt);
+	}
+	return CleanCollisionParticles(Vertices, BBox, SnapDistance);
+}
+
+inline TArray<FVec3> CleanCollisionParticles(
+	FTriangleMesh &TriMesh, 
+	const TArrayView<const FVec3>& Vertices, 
+	const FReal Fraction)
+{
+	TArray<FVec3> CollisionVertices;
+	if (Fraction <= 0.0)
+		return CollisionVertices;
+
+	// If the tri mesh has any open boundaries, see if we can merge any coincident
+	// vertices on the boundary.  This makes the importance ordering work much better
+	// as we need the curvature at each edge of the tri mesh, and we can't calculate
+	// curvature on discontiguous triangles.
+	TSet<int32> BoundaryPoints = TriMesh.GetBoundaryPoints();
+	if (BoundaryPoints.Num())
+	{
+		TMap<int32, int32> Remapping =
+			TriMesh.FindCoincidentVertexRemappings(BoundaryPoints.Array(), Vertices);
+		TriMesh.RemapVertices(Remapping);
+	}
+
+	// Get the importance vertex ordering, from most to least.  Reorder the 
+	// particles accordingly.
+	TArray<int32> CoincidentVertices;
+	const TArray<int32> Ordering = TriMesh.GetVertexImportanceOrdering(Vertices, &CoincidentVertices, true);
+
+	// Particles are ordered from most important to least, with coincident 
+	// vertices at the very end.
+	const int32 NumGoodPoints = Ordering.Num() - CoincidentVertices.Num();
+
+#if DO_GUARD_SLOW
+	for (int i = NumGoodPoints; i < Ordering.Num(); ++i)
+	{
+		ensure(CoincidentVertices.Contains(Ordering[i]));	//make sure all coincident vertices are at the back
+	}
+#endif
+
+	CollisionVertices.AddUninitialized(FMath::Min(NumGoodPoints, static_cast<int32>(ceil(static_cast<FReal>(NumGoodPoints) * Fraction))));
+	for (int i = 0; i < CollisionVertices.Num(); i++)
+	{
+		CollisionVertices[i] = Vertices[Ordering[i]];
+	}
+	return CollisionVertices;
+}
+
+inline void CleanCollisionParticles(
+	FTriangleMesh &TriMesh, 
+	const TArrayView<const FVec3>& Vertices, 
+	const FReal Fraction,
+	TSet<int32>& ResultingIndices)
+{
+	ResultingIndices.Reset();
+	if (Fraction <= 0.0)
+		return;
+
+	TArray<int32> CoincidentVertices;
+	const TArray<int32> Ordering = TriMesh.GetVertexImportanceOrdering(Vertices, &CoincidentVertices, true);
+	int32 NumGoodPoints = Ordering.Num() - CoincidentVertices.Num();
+	NumGoodPoints = FMath::Min(NumGoodPoints, static_cast<int32>(ceil(static_cast<FReal>(NumGoodPoints) * Fraction)));
+
+	ResultingIndices.Reserve(NumGoodPoints);
+	for (int32 i = 0; i < NumGoodPoints; i++)
+	{
+		ResultingIndices.Add(Ordering[i]);
+	}
+}
+
+} // namespace Chaos
diff --git a/Engine/Source/Runtime/Experimental/Chaos/Public/PhysicsProxy/FieldSystemProxyHelper.h b/Engine/Source/Runtime/Experimental/Chaos/Public/PhysicsProxy/FieldSystemProxyHelper.h
--- a/Engine/Source/Runtime/Experimental/Chaos/Public/PhysicsProxy/FieldSystemProxyHelper.h
+++ b/Engine/Source/Runtime/Experimental/Chaos/Public/PhysicsProxy/FieldSystemProxyHelper.h
@@ -7,6 +7,7 @@
 #include "Chaos/ConstraintHandle.h"
 #include "Chaos/ParticleHandle.h"
 #include "Chaos/Defines.h"
+#include "Chaos/PBDRigidClusteringAlgo.h"
 #include "Chaos/Particles.h"
 #include <limits>
 
@@ -242,7 +243,7 @@ namespace Chaos
 				{
 					Chaos::FPBDRigidParticleHandle* RigidHandle = ClusteredHandle->CastToRigidParticle();
 					check(RigidHandle);
-					Clustering.UpdateKinematicProperties(RigidHandle);
+					UpdateKinematicProperties(RigidHandle, Clustering.GetChildrenMap(), *RigidSolver->GetEvolution());
 				}
 			}
 		}
