Chaos - Fix issue where geometry collection GT and PT particles have different spatial indices, causing potential crashes when asynchronous trace is used while a GC is destroyed

diff --git a/Engine/Source/Runtime/Experimental/Chaos/Private/PhysicsProxy/GeometryCollectionPhysicsProxy.cpp b/Engine/Source/Runtime/Experimental/Chaos/Private/PhysicsProxy/GeometryCollectionPhysicsProxy.cpp
--- a/Engine/Source/Runtime/Experimental/Chaos/Private/PhysicsProxy/GeometryCollectionPhysicsProxy.cpp
+++ b/Engine/Source/Runtime/Experimental/Chaos/Private/PhysicsProxy/GeometryCollectionPhysicsProxy.cpp
@@ -434,6 +434,11 @@ void FGeometryCollectionPhysicsProxy::Initialize(Chaos::FPBDRigidsEvolutionBase
 			P->SetProxy(this);
 			P->SetGeometry(GameThreadCollection.Implicits[Index]);
 
+			// IMPORTANT: we need to set eth right spatial index because GT particle is static and PT particle is rigid
+			// this is causing a mismatch when using the separate acceleration structures optimization which can cause crashes when destroying the particle while async tracing 
+			// todo(chaos) we should eventually refactor this code to use rigid particles on the GT side for geometry collection  
+			P->SetSpatialIdx(Chaos::FSpatialAccelerationIdx{ 0,0 });
+
 			const Chaos::FShapesArray& Shapes = P->ShapesArray();
 			const int32 NumShapes = Shapes.Num();
 			for(int32 ShapeIndex = 0; ShapeIndex < NumShapes; ++ShapeIndex)
