Chaos : Improved read/write lock

diff --git a/Engine/Source/Runtime/Experimental/Chaos/Private/Chaos/Framework/PhysicsSolverBase.cpp b/Engine/Source/Runtime/Experimental/Chaos/Private/Chaos/Framework/PhysicsSolverBase.cpp
--- a/Engine/Source/Runtime/Experimental/Chaos/Private/Chaos/Framework/PhysicsSolverBase.cpp
+++ b/Engine/Source/Runtime/Experimental/Chaos/Private/Chaos/Framework/PhysicsSolverBase.cpp
@@ -146,7 +146,7 @@ namespace Chaos
 		, bUseCollisionResimCache(false)
 		, bPaused_External(false)
 		, Owner(InOwner)
-		, ExternalDataLock_External(new FPhysicsSceneGuard())
+		, ExternalDataLock_External(new FPhysSceneLock())
 		, bIsShuttingDown(false)
 		, AsyncDt(InAsyncDt)
 		, AccumulatedTime(0)
diff --git a/Engine/Source/Runtime/Experimental/Chaos/Private/Framework/Threading.cpp b/Engine/Source/Runtime/Experimental/Chaos/Private/Framework/Threading.cpp
new file mode 100644
--- /dev/null
+++ b/Engine/Source/Runtime/Experimental/Chaos/Private/Framework/Threading.cpp
@@ -0,0 +1,114 @@
+// Copyright Epic Games, Inc. All Rights Reserved.
+
+#include "Framework/Threading.h"
+#include "ChaosCheck.h"
+
+namespace Chaos
+{
+	// uint32 pair for tracking reads and writes per thread.
+	struct FLockDepthData
+	{
+		uint32 ReadDepth = 0;
+		uint32 WriteDepth = 0;
+
+		bool HasLocks() const 
+		{
+			return (ReadDepth + WriteDepth) > 0;
+		}
+	};
+	
+	struct FMultiInstanceLockDepthData
+	{
+		FLockDepthData& Get(void* InstancePtr)
+		{
+			return InstanceDepths.FindOrAdd(InstancePtr);
+		}
+
+		void IncRead(void* InstancePtr)
+		{
+			++(Get(InstancePtr).ReadDepth);
+		}
+
+		void IncWrite(void* InstancePtr)
+		{
+			++(Get(InstancePtr).WriteDepth);
+		}
+
+		void DecRead(void* InstancePtr)
+		{
+			FLockDepthData& Data = Get(InstancePtr);
+			--Data.ReadDepth;
+
+			if(!Data.HasLocks())
+			{
+				InstanceDepths.Remove(InstancePtr);
+			}
+		}
+
+		void DecWrite(void* InstancePtr)
+		{
+			FLockDepthData& Data = Get(InstancePtr);
+			--Data.WriteDepth;
+
+			if(!Data.HasLocks())
+			{
+				InstanceDepths.Remove(InstancePtr);
+			}
+		}
+
+	private:
+		TMap<void*, FLockDepthData> InstanceDepths;
+	};
+	thread_local static FMultiInstanceLockDepthData GThreadLockCheckData;
+
+	namespace ThreadingPrivate
+	{
+#if CHAOS_SCENE_LOCK_CHECKS
+		void CheckLockReadAssumption(const TCHAR* Context)
+		{
+			// Holding either a read lock or a write lock enables reading, so if both are zero we emit an error
+			if(GThreadLockCheckData.ReadDepth == 0 && GThreadLockCheckData.WriteDepth == 0)
+			{
+				UE_LOG(LogChaos, Error, TEXT("Thread %u called %s without taking a read lock. Wrap the function call in a scene read lock to safely read data."), FPlatformTLS::GetCurrentThreadId(), Context);
+				CHAOS_ENSURE(false);
+			}
+		}
+
+		void CheckLockWriteAssumption(const TCHAR* Context)
+		{
+			// Must have at least one write scope to pass the assumption
+			if(GThreadLockCheckData.WriteDepth == 0)
+			{
+				UE_LOG(LogChaos, Error, TEXT("Thread %u called %s without taking a write lock. Wrap the function call in a scene write lock to safely write data."), FPlatformTLS::GetCurrentThreadId(), Context);
+				CHAOS_ENSURE(false);
+			}
+		}
+#endif
+
+		void IncReadDepth(void* Instance)
+		{
+			GThreadLockCheckData.IncRead(Instance);
+		}
+
+		void IncWriteDepth(void* Instance)
+		{
+			GThreadLockCheckData.IncWrite(Instance);
+		}
+
+		void DecReadDepth(void* Instance)
+		{
+			GThreadLockCheckData.DecRead(Instance);
+		}
+
+		void DecWriteDepth(void* Instance)
+		{
+			GThreadLockCheckData.DecWrite(Instance);
+		}
+
+		uint32 GetThreadReadDepth(void* Instance)
+		{
+			return GThreadLockCheckData.Get(Instance).ReadDepth;
+		}
+
+	}
+}
\ No newline at end of file
diff --git a/Engine/Source/Runtime/Experimental/Chaos/Public/Chaos/Framework/PhysicsSolverBase.h b/Engine/Source/Runtime/Experimental/Chaos/Public/Chaos/Framework/PhysicsSolverBase.h
--- a/Engine/Source/Runtime/Experimental/Chaos/Public/Chaos/Framework/PhysicsSolverBase.h
+++ b/Engine/Source/Runtime/Experimental/Chaos/Public/Chaos/Framework/PhysicsSolverBase.h
@@ -6,6 +6,7 @@
 #include "Misc/ScopeLock.h"
 #include "ChaosLog.h"
 #include "Chaos/Framework/PhysicsProxyBase.h"
+#include "Framework/Threading.h"
 #include "Chaos/ParticleDirtyFlags.h"
 #include "Async/ParallelFor.h"
 #include "Containers/Queue.h"
@@ -23,7 +24,6 @@ namespace Chaos
 {
 	class FPhysicsSolverBase;
 	struct FPendingSpatialDataQueue;
-	class FPhysicsSceneGuard;
 	class FChaosResultsManager;
 	class FRewindData;
 	class IRewindCallback;
@@ -477,7 +477,7 @@ namespace Chaos
 		FRWLock SimMaterialLock;
 		
 		/** Scene lock object for external threads (non-physics) */
-		TUniquePtr<FPhysicsSceneGuard> ExternalDataLock_External;
+		TUniquePtr<FPhysSceneLock> ExternalDataLock_External;
 
 		friend FChaosSolversModule;
 		friend FPhysicsSolverAdvanceTask;
@@ -515,7 +515,7 @@ namespace Chaos
 
 		/** Get the lock used for external data manipulation. A better API would be to use scoped locks so that getting a write lock is non-const */
 		//NOTE: this is a const operation so that you can acquire a read lock on a const solver. The assumption is that non-const write operations are already marked non-const
-		FPhysicsSceneGuard& GetExternalDataLock_External() const { return *ExternalDataLock_External; }
+		FPhysSceneLock& GetExternalDataLock_External() const { return *ExternalDataLock_External; }
 
 	protected:
 		/** Storage for events, see the Add/Remove pairs above for event timings */
diff --git a/Engine/Source/Runtime/Experimental/Chaos/Public/Framework/Threading.h b/Engine/Source/Runtime/Experimental/Chaos/Public/Framework/Threading.h
--- a/Engine/Source/Runtime/Experimental/Chaos/Public/Framework/Threading.h
+++ b/Engine/Source/Runtime/Experimental/Chaos/Public/Framework/Threading.h
@@ -6,7 +6,9 @@
 #include "UObject/ObjectMacros.h"
 #include "PhysicsCoreTypes.h"
 #include "ChaosLog.h"
+#include "ProfilingDebugging/CsvProfiler.h"
 #include "Async/ParallelFor.h"
+#include <atomic>
 
 #ifndef PHYSICS_THREAD_CONTEXT
 	#if (!UE_BUILD_SHIPPING && !UE_BUILD_TEST)
@@ -16,8 +18,123 @@
 	#endif
 #endif
 
+/** Controls whether we use a FIFO RW lock or an unfair RW lock for scene locking in Chaos */
+#define CHAOS_SCENE_LOCK_RWFIFO 1
+
+/**
+ * Controls whether we use a non-yielding lock for the FIFO lock if CHAOS_SCENE_LOCK_RWFIFO is enabled
+ * @see CHAOS_SCENE_LOCK_RWFIFO
+ */ 
+#define CHAOS_SCENE_LOCK_RWFIFO_SPINLOCK 0
+
+/**
+ * \def CHAOS_SCENE_LOCK_CHECKS
+ * Controls whether the runtime will check and emit errors when a read or write operation is attempted but an
+ * appropriate read or write lock has not been taken by the caller
+ * NOTE: Disable currently until this can be made to check with the per-instance thread counts.
+ */
+#ifndef CHAOS_SCENE_LOCK_CHECKS
+	#if (!UE_BUILD_SHIPPING && !UE_BUILD_TEST)
+		#define CHAOS_SCENE_LOCK_CHECKS 0
+	#else
+		#define CHAOS_SCENE_LOCK_CHECKS 0
+	#endif
+#endif
+
 namespace Chaos
 {
+	// Not intended for external callers, provided here to allow the locks below to record depths
+	namespace ThreadingPrivate
+	{
+		// Control the current thread read/write depths
+		CHAOS_API void IncReadDepth(void* Instance);
+		CHAOS_API void IncWriteDepth(void* Instance);
+		CHAOS_API void DecReadDepth(void* Instance);
+		CHAOS_API void DecWriteDepth(void* Instance);
+
+		// Get the calling thread's current read depth
+		CHAOS_API uint32 GetThreadReadDepth(void* Instance);
+	}
+
+#if CHAOS_SCENE_LOCK_CHECKS
+
+	// Not intended for external callers, provided here to allow the below check macros to function
+	namespace ThreadingPrivate
+	{
+		// Checks assumptions for functions marked _AssumesLocked in the interface.
+		CHAOS_API void CheckLockReadAssumption(const TCHAR* Context);
+		CHAOS_API void CheckLockWriteAssumption(const TCHAR* Context);
+	}
+
+	/** Checks that the caller currently has a read or write lock open, emits an error if not locked */
+	#define CHAOS_CHECK_READ_ASSUMPTION Chaos::ThreadingPrivate::CheckLockReadAssumption(ANSI_TO_TCHAR(__FUNCTION__))
+
+	/** Checks that the caller currently has a write lock open, emits an error if not locked */
+	#define CHAOS_CHECK_WRITE_ASSUMPTION Chaos::ThreadingPrivate::CheckLockWriteAssumption(ANSI_TO_TCHAR(__FUNCTION__))
+
+	/** 
+	 * Checks that the caller currently has a read or write lock open if an actor is currently bound to a solver.
+	 * The actor is a derived child of IPhysicsProxyBase which holds a solver pointer which is set on the main
+	 * thread during scene registration - if that isn't currently set the actor isn't under the control of the
+	 * physics thread and can be operated on to initialize it without a lock.
+	 * @see FPBDRigidsSolver::RegisterObject
+	 */
+	#define CHAOS_CHECK_READ_ASSUMPTION_ACTOR(Actor) if(Actor && Actor->GetSolverBase()) {Chaos::ThreadingPrivate::CheckLockReadAssumption(ANSI_TO_TCHAR(__FUNCTION__));}
+
+	/**
+	 * Checks that the caller currently has a write lock open if an actor is currently bound to a solver.
+	 * The actor is a derived child of IPhysicsProxyBase which holds a solver pointer which is set on the main
+	 * thread during scene registration - if that isn't currently set the actor isn't under the control of the
+	 * physics thread and can be operated on to initialize it without a lock.
+	 * @see FPBDRigidsSolver::RegisterObject
+	 */
+	#define CHAOS_CHECK_WRITE_ASSUMPTION_ACTOR(Actor) if(Actor && Actor->GetSolverBase()) {Chaos::ThreadingPrivate::CheckLockWriteAssumption(ANSI_TO_TCHAR(__FUNCTION__));}
+	
+	/**
+	 * Checks that the caller currently has a read or write lock open when reading constraint properties for
+	 * actors that are bound to a solver.
+	 * @see CHAOS_CHECK_READ_ASSUMPTION_ACTOR
+	 */
+	#define CHAOS_CHECK_READ_ASSUMPTION_CONSTRAINT(Handle)  \
+	if(Handle.Constraint && \
+		((Handle.Constraint->GetParticleProxies()[0] && Handle.Constraint->GetParticleProxies()[0]->GetSolverBase()) || \
+		(Handle.Constraint->GetParticleProxies()[1] && Handle.Constraint->GetParticleProxies()[1]->GetSolverBase()))) \
+	{Chaos::ThreadingPrivate::CheckLockReadAssumption(ANSI_TO_TCHAR(__FUNCTION__));}
+
+	/**
+	 * Checks that the caller currently has a read or write lock open when reading constraint properties for
+	 * actors that are bound to a solver.
+	 * @see CHAOS_CHECK_WRITE_ASSUMPTION_ACTOR
+	 */
+	#define CHAOS_CHECK_WRITE_ASSUMPTION_CONSTRAINT(Handle)  \
+	if(Handle.Constraint && \
+		((Handle.Constraint->GetParticleProxies()[0] && Handle.Constraint->GetParticleProxies()[0]->GetSolverBase()) || \
+		(Handle.Constraint->GetParticleProxies()[1] && Handle.Constraint->GetParticleProxies()[1]->GetSolverBase()))) \
+	{Chaos::ThreadingPrivate::CheckLockWriteAssumption(ANSI_TO_TCHAR(__FUNCTION__));}
+
+#else
+
+	// Empty when not compiled in
+	#define CHAOS_CHECK_READ_ASSUMPTION
+	#define CHAOS_CHECK_WRITE_ASSUMPTION
+	#define CHAOS_CHECK_READ_ASSUMPTION_ACTOR
+	#define CHAOS_CHECK_WRITE_ASSUMPTION_ACTOR
+	#define CHAOS_CHECK_READ_ASSUMPTION_CONSTRAINT
+	#define CHAOS_CHECK_WRITE_ASSUMPTION_CONSTRAINT
+
+#endif
+
+/** Signals that we have entered a read lock to control the checks above */
+#define CHAOS_RECORD_ENTER_READ_LOCK Chaos::ThreadingPrivate::IncReadDepth(this);
+
+/** Signals that we have entered a write lock to control the checks above */
+#define CHAOS_RECORD_ENTER_WRITE_LOCK Chaos::ThreadingPrivate::IncWriteDepth(this);
+
+/** Signals that we have left a read lock to control the checks above */
+#define CHAOS_RECORD_LEAVE_READ_LOCK Chaos::ThreadingPrivate::DecReadDepth(this);
+
+/** Signals that we have left a write lock to control the checks above */
+#define CHAOS_RECORD_LEAVE_WRITE_LOCK Chaos::ThreadingPrivate::DecWriteDepth(this);
 
 #if PHYSICS_THREAD_CONTEXT
 /** Debug helper to ensure threading mistakes are caught. Do not use for ship */
@@ -142,10 +259,10 @@ FORCEINLINE void EnsureIsInGameThreadContext()
 	 * Recursive Read/Write lock object for protecting external data accesses for physics scenes.
 	 * This is a fairly heavy lock designed to allow scene queries and user code to safely access
 	 * external physics data.
-	 * 
+	 *
 	 * The lock also allows a thread to recursively lock data to avoid deadlocks on repeated writes
 	 * or undefined behavior for nesting read locks.
-	 * 
+	 *
 	 * Fairness is determined by the underlying platform FRWLock type as this lock uses FRWLock
 	 * as it's internal primitive
 	 */
@@ -198,7 +315,7 @@ FORCEINLINE void EnsureIsInGameThreadContext()
 			ModifyTls([](FSceneLockTls& ThreadDataInner) {ThreadDataInner.WriteDepth++; });
 
 			const uint32 ThisThreadId = FPlatformTLS::GetCurrentThreadId();
-			
+
 			if(CurrentWriterThreadId.Load() != ThisThreadId)
 			{
 				InnerLock.WriteLock();
@@ -214,20 +331,21 @@ FORCEINLINE void EnsureIsInGameThreadContext()
 		void ReadUnlock()
 		{
 			const FSceneLockTls ThreadData = ModifyTls([](FSceneLockTls& ThreadDataInner)
-			{
-				if(ThreadDataInner.ReadDepth > 0)
-				{
-					ThreadDataInner.ReadDepth--;
-				}
-				else
 				{
+				if(ThreadDataInner.ReadDepth > 0)
+					{
+						ThreadDataInner.ReadDepth--;
+					}
+					else
+					{
 #if CHAOS_CHECKED
-					ensureMsgf(false, TEXT("ReadUnlock called on physics scene guard when the thread does not hold the lock"));
+						ensureMsgf(false, TEXT("ReadUnlock called on physics scene guard when the thread does not hold the lock"));
 #else
-					UE_LOG(LogChaos, Warning, TEXT("ReadUnlock called on physics scene guard when the thread does not hold the lock"))
+						UE_LOG(LogChaos, Warning, TEXT("ReadUnlock called on physics scene guard when the thread does not hold the lock"))
 #endif
-				}
-			});
+					}
+
+				});
 
 			const uint32 ThisThreadId = FPlatformTLS::GetCurrentThreadId();
 
@@ -287,7 +405,7 @@ FORCEINLINE void EnsureIsInGameThreadContext()
 
 			union
 			{
-				struct  
+				struct
 				{
 					uint16 WriteDepth;
 					uint16 ReadDepth;
@@ -320,37 +438,298 @@ FORCEINLINE void EnsureIsInGameThreadContext()
 		FRWLock InnerLock;
 	};
 
-	class FPhysicsSceneGuardScopedWrite
+	/**
+	 * Templated RAII scope lock around generic mutex type
+	 */
+	template<typename MutexType>
+	class TMutexScopeLock
+	{
+	public:
+		TMutexScopeLock(MutexType& InMutex)
+			: Mutex(InMutex)
+		{
+			Mutex.Lock();
+		}
+
+		~TMutexScopeLock()
+		{
+			Mutex.Unlock();
+		}
+
+	private:
+
+		// No default, copy or move construction
+		TMutexScopeLock() = delete;
+		TMutexScopeLock(const TMutexScopeLock&) = delete;
+		TMutexScopeLock(TMutexScopeLock&&) = delete;
+		TMutexScopeLock& operator=(const TMutexScopeLock&) = delete;
+		TMutexScopeLock& operator=(TMutexScopeLock&&) = delete;
+
+		MutexType& Mutex;
+	};
+
+	/**
+	 * A first-in, first-out "fair" read-write lock around a generic mutex
+	 * Given a mutex (either just FCriticalSection or some custom lock) this class implements a fair lock.
+	 * Any number of readers can enter the lock but as soon as a writer attempts to enter the lock all 
+	 * subsequent readers are forced to wait until the current readers leave the lock and the writer gets a chance
+	 * to perform its operation. Once the write is completed the waiting readers are able to resume.
+	 * This ensures we do not end up in a situation where we have a write waiting but many reads end up constantly
+	 * forcing the write to wait. In a physics context a write on the game thread is time-critical and we want
+	 * that thread to resume as soon as possible by pausing any reads (scene queries) until the write is finished
+	 */
+	template<typename MutexType>
+	class TRwFifoLock
+	{
+	public:
+
+		TRwFifoLock()
+			: NumReaders(0)
+		{
+			//ThreadingPrivate::CreateLockThreadData(this);
+		}
+
+		~TRwFifoLock()
+		{
+			//ThreadingPrivate::DestroyLockThreadData(this);
+		}
+
+		void ReadLock()
+		{
+			if(ThreadingPrivate::GetThreadReadDepth(this) == 0)
+			{
+				TMutexScopeLock<MutexType> Guard(Mutex);
+
+				// We lock for this increment to halt if there's a writer waiting to enter the lock
+				// In this case we will be forced to wait till the write completes
+				++NumReaders;
+			}
+			else
+			{
+				// Only require a lock on the first acquisition. this allows recursive reads even while a
+				// writer is holding the lock waiting to enter. The writer will be allowed to proceed when
+				// all write scopes end
+			++NumReaders;
+			}
+
+#if PHYSICS_THREAD_CONTEXT
+			// Read lock means we can access game thread data, so set the right context
+			FPhysicsThreadContext::Get().IncGameThreadContext();
+#endif
+
+			CHAOS_RECORD_ENTER_READ_LOCK;
+		}
+
+		void WriteLock()
+		{
+#if CHAOS_SCENE_LOCK_CHECKS
+			if(ThreadingPrivate::GetThreadReadDepth(this) > 0)
+			{
+				ensureMsgf(false, TEXT("A thread holding a read lock on the physics scene attempted to upgrade to a write lock - this is not supported, performing an unsafe write."));
+
+				// Still need to increment the context when we hit this case or we'll just crash later
+#if PHYSICS_THREAD_CONTEXT
+				// Write lock means we can access game thread data, so set the right context
+				FPhysicsThreadContext::Get().IncGameThreadContext();
+#endif
+
+				return;
+			}
+#endif
+
+			Mutex.Lock();
+
+			// Spin until all readers are finished
+			for(;;)
+			{
+				if(NumReaders.load() == 0)
+				{
+					// All readers now finished - writer can enter the lock properly (pass back to caller)
+					break;
+				}
+
+				// Issue pause instruction - architecture dependent instruction to better handle
+				// a spin lock not interfering with other threads on this core, this doesn't
+				// actually yield the thread
+				FPlatformProcess::Yield();
+			}
+
+			CHAOS_RECORD_ENTER_WRITE_LOCK;
+
+#if PHYSICS_THREAD_CONTEXT
+			// Write lock means we can access game thread data, so set the right context
+			FPhysicsThreadContext::Get().IncGameThreadContext();
+#endif
+		}
+
+		void ReadUnlock()
+		{
+			CHAOS_RECORD_LEAVE_READ_LOCK;
+
+			// No locking here, just decrement atomic reader count
+			--NumReaders;
+
+#if PHYSICS_THREAD_CONTEXT
+			// Read lock is released, the gamethread context is gone
+			FPhysicsThreadContext::Get().DecGameThreadContext();
+#endif
+		}
+
+		void WriteUnlock()
+		{
+			CHAOS_RECORD_LEAVE_WRITE_LOCK;
+
+			Mutex.Unlock();
+#if PHYSICS_THREAD_CONTEXT
+			// Write lock is released, the gamethread context is gone
+			FPhysicsThreadContext::Get().DecGameThreadContext();
+#endif
+		}
+
+	private:
+		MutexType Mutex;
+		std::atomic<uint32> NumReaders;
+	};
+
+	/**
+	 * A non-yielding, recursive spin lock
+	 * Implements a first-in, first-out lock / mutex that won't yield back to the system.
+	 * Intended for applications that must wake / resume at the earliest opportunity.
+	 * Each thread attempting a write gets an atomically controlled counter to wait on so the lock is fair in that 
+	 * the locks will be ordered according to the order Lock was called.
+	 */
+	class FPhysSpinLock
 	{
 	public:
-		FPhysicsSceneGuardScopedWrite(FPhysicsSceneGuard& InGuard)
-		: Guard(InGuard)
+		FPhysSpinLock()
+			: Next(0)
+			, Current(0)
+			, WriterId(0)
+			, Count()
+		{}
+
+		void Lock()
 		{
-			Guard.WriteLock();
+			// Support recursive locking
+			if(WriterId.load() == FPlatformTLS::GetCurrentThreadId())
+			{
+				Count++;
+				return;
+			}
+
+			// Get the wait value - acquire operation so Current.load can't be reordered before this
+			uint32 WaitFor = Next.fetch_add(1, std::memory_order_acquire);
+			for(;;)
+			{
+				if(WaitFor == Current.load())
+				{
+					break;
+				}
+
+				// Issue pause instruction - architecture dependent instruction to better handle
+				// a spin lock not interfering with other threads on this core, this doesn't
+				// actually yield the thread
+				FPlatformProcess::Yield();
+			}
+
+			// Lock acquired, store the thread ID for recursive locking
+			WriterId.store(FPlatformTLS::GetCurrentThreadId());
+			Count++;
 		}
 
-		~FPhysicsSceneGuardScopedWrite()
+		void Unlock()
 		{
-			Guard.WriteUnlock();
+			checkf(WriterId.load() == FPlatformTLS::GetCurrentThreadId(), TEXT("A thread unlocked without owning the lock (calling Lock first)"));
+			checkf(Count > 0, TEXT("A thread unlocked a lock that had no outstanding lock scopes"));
+
+			// Once all recursive locks are dropped, increment Current to allow the next thread in
+			if(--Count == 0)
+			{
+				// Clear the lock owner
+				WriterId.store(0);
+
+				// Release the next thread - this must be the last operation as immediately
+				// the next user of the lock will be allowed to take ownership
+				++Current;
+			}
 		}
+
 	private:
-		FPhysicsSceneGuard& Guard;
+		std::atomic<uint32> Next;
+		std::atomic<uint32> Current;
+		std::atomic<uint32> WriterId;
+		uint32 Count;
 	};
 
-	class FPhysicsSceneGuardScopedRead
+	/**
+	 * Implements a RAII scoped write lock around a generic mutex.
+	 */
+	template<typename MutexType>
+	class TPhysicsSceneGuardScopedWrite
 	{
 	public:
-		FPhysicsSceneGuardScopedRead(FPhysicsSceneGuard& InGuard)
-			: Guard(InGuard)
+		TPhysicsSceneGuardScopedWrite(MutexType& InMutex)
+			: Mutex(InMutex)
 		{
-			Guard.ReadLock();
+			CSV_SCOPED_TIMING_STAT_EXCLUSIVE(AcquireSceneWriteLock);
+			Mutex.WriteLock();
 		}
 
-		~FPhysicsSceneGuardScopedRead()
+		~TPhysicsSceneGuardScopedWrite()
 		{
-			Guard.ReadUnlock();
+			Mutex.WriteUnlock();
 		}
+
 	private:
-		FPhysicsSceneGuard& Guard;
+		TPhysicsSceneGuardScopedWrite() = delete;
+		TPhysicsSceneGuardScopedWrite(const TPhysicsSceneGuardScopedWrite&) = delete;
+		TPhysicsSceneGuardScopedWrite(TPhysicsSceneGuardScopedWrite&&) = delete;
+		TPhysicsSceneGuardScopedWrite& operator=(const TPhysicsSceneGuardScopedWrite&) = delete;
+		TPhysicsSceneGuardScopedWrite& operator=(TPhysicsSceneGuardScopedWrite&&) = delete;
+
+		MutexType& Mutex;
 	};
+
+	/**
+	 * Implements a RAII scoped read lock around a generic mutex.
+	 */
+	template<typename MutexType>
+	class TPhysicsSceneGuardScopedRead
+	{
+	public:
+		TPhysicsSceneGuardScopedRead(MutexType& InMutex)
+			: Mutex(InMutex)
+		{
+			CSV_SCOPED_TIMING_STAT_EXCLUSIVE(AcquireSceneReadLock);
+			Mutex.ReadLock();
+		}
+
+		~TPhysicsSceneGuardScopedRead()
+		{
+			Mutex.ReadUnlock();
+		}
+
+	private:
+		TPhysicsSceneGuardScopedRead() = delete;
+		TPhysicsSceneGuardScopedRead(const TPhysicsSceneGuardScopedRead&) = delete;
+		TPhysicsSceneGuardScopedRead(TPhysicsSceneGuardScopedRead&&) = delete;
+		TPhysicsSceneGuardScopedRead& operator=(const TPhysicsSceneGuardScopedRead&) = delete;
+		TPhysicsSceneGuardScopedRead& operator=(TPhysicsSceneGuardScopedRead&&) = delete;
+
+		MutexType& Mutex;
+	};
+
+#if !CHAOS_SCENE_LOCK_RWFIFO
+	using FPhysSceneLock = FPhysicsSceneGuard;
+#else
+#if CHAOS_SCENE_LOCK_RWFIFO_SPINLOCK
+	using FPhysSceneLock = TRwFifoLock<FPhysSpinLock>;
+#else
+	using FPhysSceneLock = TRwFifoLock<FCriticalSection>;
+#endif
+#endif
+
+	// Stable types to use in calling code configured by the compiler switches above
+	using FPhysicsSceneGuardScopedWrite = TPhysicsSceneGuardScopedWrite<FPhysSceneLock>;
+	using FPhysicsSceneGuardScopedRead = TPhysicsSceneGuardScopedRead<FPhysSceneLock>;
 }
