Chaos : Fix potential crash caused by particle geometry being updated but the particle to not be properly disabled in the constraint graph

diff --git a/Engine/Source/Runtime/Experimental/Chaos/Private/PhysicsProxy/SingleParticlePhysicsProxy.cpp b/Engine/Source/Runtime/Experimental/Chaos/Private/PhysicsProxy/SingleParticlePhysicsProxy.cpp
--- a/Engine/Source/Runtime/Experimental/Chaos/Private/PhysicsProxy/SingleParticlePhysicsProxy.cpp
+++ b/Engine/Source/Runtime/Experimental/Chaos/Private/PhysicsProxy/SingleParticlePhysicsProxy.cpp
@@ -68,6 +68,9 @@ void PushToPhysicsStateImp(const Chaos::FDirtyPropertiesManager& Manager, Chaos:
 		if(NewNonFrequentData)
 		{
 			Handle->SetNonFrequentData(*NewNonFrequentData);
+
+			// geometry may have changed, we need to invalidate the particle so it can be removed from caching structures in the evolution
+			Evolution.InvalidateParticle(Handle);
 		}
 
 		auto NewVelocities = bHasKinematicData ? ParticleData.FindVelocities(Manager, DataIdx) : nullptr;
diff --git a/Engine/Source/Runtime/Experimental/Chaos/Public/Chaos/PBDRigidsEvolution.h b/Engine/Source/Runtime/Experimental/Chaos/Public/Chaos/PBDRigidsEvolution.h
--- a/Engine/Source/Runtime/Experimental/Chaos/Public/Chaos/PBDRigidsEvolution.h
+++ b/Engine/Source/Runtime/Experimental/Chaos/Public/Chaos/PBDRigidsEvolution.h
@@ -368,6 +368,11 @@ public:
 		DisableConstraints(TSet<FGeometryParticleHandle*>({ Particle }));
 	}
 
+	CHAOS_API void InvalidateParticle(FGeometryParticleHandle* Particle)
+	{
+		ConstraintGraph.DisableParticle(Particle);
+	}
+	
 	CHAOS_API void FlushExternalAccelerationQueue(FAccelerationStructure& Acceleration,FPendingSpatialDataQueue& ExternalQueue);
 
 	CHAOS_API void DisableParticles(TSet<FGeometryParticleHandle*> &ParticlesIn)
