Fix bug that do not pick the default pipeline when changing it in the pipeline configuration dialog.
#rnx

diff --git a/Engine/Source/Runtime/Experimental/Interchange/Engine/Private/InterchangeManager.cpp b/Engine/Source/Runtime/Experimental/Interchange/Engine/Private/InterchangeManager.cpp
--- a/Engine/Source/Runtime/Experimental/Interchange/Engine/Private/InterchangeManager.cpp
+++ b/Engine/Source/Runtime/Experimental/Interchange/Engine/Private/InterchangeManager.cpp
@@ -674,7 +674,7 @@ UInterchangeManager::ImportInternal(const FString& ContentPath, const UInterchan
 															
 
 
-		const FName DefaultPipelineStackName = [ImportType, &InterchangeProjectSettings]()
+		auto GetDefaultPipelineStackName = [ImportType, &InterchangeProjectSettings]()->FName
 		{
 			if (ImportType == UE::Interchange::EImportType::ImportType_Scene)
 			{
@@ -684,7 +684,7 @@ UInterchangeManager::ImportInternal(const FString& ContentPath, const UInterchan
 			{
 				return InterchangeProjectSettings->DefaultPipelineStack;
 			}
-		}();
+		};
 			
 		const TMap<FName, FInterchangePipelineStack>& DefaultPipelineStacks = InterchangeProjectSettings->PipelineStacks;
 
@@ -723,7 +723,7 @@ UInterchangeManager::ImportInternal(const FString& ContentPath, const UInterchan
 		}
 		else
 		{
-			FName PipelineStackName = DefaultPipelineStackName;
+			FName PipelineStackName = GetDefaultPipelineStackName();
 			if (RegisteredPipelineConfiguration && (bShowPipelineStacksConfigurationDialog || (!DefaultPipelineStacks.Contains(PipelineStackName) && !bIsUnattended)))
 			{
 				//Show the dialog, a plugin should have register this dialog. We use a plugin to be able to use editor code when doing UI
@@ -736,6 +736,7 @@ UInterchangeManager::ImportInternal(const FString& ContentPath, const UInterchan
 				{
 					bImportAllWithDefault = true;
 				}
+				PipelineStackName = GetDefaultPipelineStackName();
 			}
 			if (!bImportCancel)
 			{
@@ -743,13 +744,13 @@ UInterchangeManager::ImportInternal(const FString& ContentPath, const UInterchan
 				const TMap<FName, FInterchangePipelineStack>& PipelineStacks = InterchangeProjectSettings->PipelineStacks;
 				if (!PipelineStacks.Contains(PipelineStackName))
 				{
-					if (PipelineStacks.Contains(DefaultPipelineStackName))
+					if (PipelineStacks.Contains(GetDefaultPipelineStackName()))
 					{
-						PipelineStackName = DefaultPipelineStackName;
+						PipelineStackName = GetDefaultPipelineStackName();
 					}
 					else
 					{
-						//Log an error, we cannot import asset without a valid pipeline, we will use the first pipeline
+						//Log an error, we cannot import asset without a valid pipeline, we will use the first available pipeline
 						for (const TPair<FName, FInterchangePipelineStack>& PipelineStack : PipelineStacks)
 						{
 							PipelineStackName = PipelineStack.Key;
