Chaos : Fix type in dirty element utils

diff --git a/Engine/Source/Runtime/Experimental/Chaos/Public/Chaos/AABBTreeDirtyGridUtils.h b/Engine/Source/Runtime/Experimental/Chaos/Public/Chaos/AABBTreeDirtyGridUtils.h
--- a/Engine/Source/Runtime/Experimental/Chaos/Public/Chaos/AABBTreeDirtyGridUtils.h
+++ b/Engine/Source/Runtime/Experimental/Chaos/Public/Chaos/AABBTreeDirtyGridUtils.h
@@ -44,10 +44,10 @@ namespace Chaos
 	template <typename T>
 	FORCEINLINE_DEBUGGABLE bool TooManyOverlapQueryCells(const TAABB<T, 3>& AABB, FReal DirtyElementGridCellSizeInv, int32 MaximumOverlap)
 	{
-		int64 XsampleCount = GetDirtyCellIndexFromWorldCoordinate((FReal)AABB.Max().X, DirtyElementGridCellSizeInv) - GetDirtyCellIndexFromWorldCoordinate((FReal)AABB.Min().X, DirtyElementGridCellSizeInv) + 1;
-		int64 YsampleCount = GetDirtyCellIndexFromWorldCoordinate((FReal)AABB.Max().Y, DirtyElementGridCellSizeInv) - GetDirtyCellIndexFromWorldCoordinate((FReal)AABB.Min().Y, DirtyElementGridCellSizeInv) + 1;
+		uint64 XsampleCount = GetDirtyCellIndexFromWorldCoordinate((FReal)AABB.Max().X, DirtyElementGridCellSizeInv) - GetDirtyCellIndexFromWorldCoordinate((FReal)AABB.Min().X, DirtyElementGridCellSizeInv) + 1;
+		uint64 YsampleCount = GetDirtyCellIndexFromWorldCoordinate((FReal)AABB.Max().Y, DirtyElementGridCellSizeInv) - GetDirtyCellIndexFromWorldCoordinate((FReal)AABB.Min().Y, DirtyElementGridCellSizeInv) + 1;
 
-		if (XsampleCount * YsampleCount <= (int64)MaximumOverlap)
+		if (XsampleCount * YsampleCount <= (uint64)MaximumOverlap)
 		{
 			return false;
 		}
