Chaos - Collision Solver solver 3x1D problems with 2D clamping rather than a 3D problem

diff --git a/Engine/Source/Programs/HeadlessChaos/Private/HeadlessChaosTestCollisions.cpp b/Engine/Source/Programs/HeadlessChaos/Private/HeadlessChaosTestCollisions.cpp
--- a/Engine/Source/Programs/HeadlessChaos/Private/HeadlessChaosTestCollisions.cpp
+++ b/Engine/Source/Programs/HeadlessChaos/Private/HeadlessChaosTestCollisions.cpp
@@ -457,7 +457,6 @@ namespace ChaosTest {
 
 	// This test will make sure that a dynamic cube colliding with a static floor will have the correct bounce velocity
 	// for a restitution of 0.5
-	// The dynamic cube will collide with one of its vertices onto a face of the static cube
 	void CollisionCubeCubeRestitution()
 	{
 		TArrayCollectionArray<bool> Collided;
@@ -476,13 +475,15 @@ namespace ChaosTest {
 
 		FGeometryParticleHandle* StaticCube = AppendStaticParticleBox(Particles, FVec3(100.0f));
 		StaticCube->X() = FVec3(0, 0, -50.0f);
+		StaticCube->UpdateWorldSpaceState(FRigidTransform3(StaticCube->X(), StaticCube->R()), FVec3(0));
+
 		FPBDRigidParticleHandle* DynamicCube = AppendDynamicParticleBox(Particles, FVec3(100.0f));
-		DynamicCube->P() = FVec3(0, 0, 80); // Penetrating by about 5cm
-		DynamicCube->Q() = FRotation3::FromElements( 0.27059805f, 0.27059805f, 0.0f, 0.923879532f ); // Rotate so that vertex collide
+		DynamicCube->X() = FVec3(0, 0, 50);
+		DynamicCube->R() = FRotation3::FromIdentity();
 		DynamicCube->V() = FVec3(0, 0, -100);
 		DynamicCube->PreV() = DynamicCube->V();
-		DynamicCube->X() = DynamicCube->P() - DynamicCube->V() * Dt;
-		DynamicCube->R() = DynamicCube->Q();
+		DynamicCube->P() = DynamicCube->X() + DynamicCube->V() * Dt;
+		DynamicCube->Q() = DynamicCube->R();
 		DynamicCube->AuxilaryValue(PhysicsMaterials) = MakeSerializable(PhysicsMaterial);
 		DynamicCube->UpdateWorldSpaceState(FRigidTransform3(DynamicCube->P(), DynamicCube->Q()), FVec3(0));
 
@@ -520,8 +521,7 @@ namespace ChaosTest {
 		Collisions.ScatterOutput(Dt);
 
 		// This test's tolerances are set to be very crude as to not be over sensitive (for now)
-		EXPECT_TRUE(DynamicCube->V().Z > 10.0f);  // restitution not too low
-		EXPECT_TRUE(DynamicCube->V().Z < 70.0f);  // restitution not too high
+		EXPECT_NEAR(DynamicCube->V().Z, 50.0f, 5.0f);  // restitution not too low
 		EXPECT_TRUE(FMath::Abs(DynamicCube->V().X) < 1.0f);
 		EXPECT_TRUE(FMath::Abs(DynamicCube->V().Y) < 1.0f);
 	}
diff --git a/Engine/Source/Programs/HeadlessChaos/Private/HeadlessChaosTestContactModification.cpp b/Engine/Source/Programs/HeadlessChaos/Private/HeadlessChaosTestContactModification.cpp
--- a/Engine/Source/Programs/HeadlessChaos/Private/HeadlessChaosTestContactModification.cpp
+++ b/Engine/Source/Programs/HeadlessChaos/Private/HeadlessChaosTestContactModification.cpp
@@ -432,26 +432,14 @@ namespace ChaosTest
 		Module->DestroySolver(Solver);
 	}
 
-	GTEST_TEST(AllTraits, ContactModification_ModifyRestitution)
+	GTEST_TEST(AllTraits, ContactModification_ModifyRestitution_NoBounce)
 	{
 		FChaosSolversModule* Module = FChaosSolversModule::GetModule();
 		auto* Solver = Module->CreateSolver(nullptr, /*AsyncDt=*/-1);
 		InitSolverSettings(Solver);
 		Solver->SetThreadingMode_External(EThreadingModeTemp::SingleThread);
 
-		// Collide two cubes with static floor, one with restitution modified to 1, one with 0.
-
-		// simulated cube with downward velocity, should bounce on floor and end up with upward velocity.
-		FSingleParticlePhysicsProxy* BounceCubeProxy = FSingleParticlePhysicsProxy::Create(Chaos::FPBDRigidParticle::CreateParticle());
-		auto& BounceCubeParticle = BounceCubeProxy->GetGameThreadAPI();
-		auto BounceCubeGeom = TSharedPtr<FImplicitObject, ESPMode::ThreadSafe>(new TBox<FReal, 3>(FVec3(-100), FVec3(100)));
-		BounceCubeParticle.SetGeometry(BounceCubeGeom);
-		Solver->RegisterObject(BounceCubeProxy);
-		BounceCubeParticle.SetGravityEnabled(false);
-		BounceCubeParticle.SetX(FVec3(200, 0, 200));
-		SetCubeInertiaTensor(BounceCubeParticle, /*Dimension=*/200, /*Mass=*/1);
-		ChaosTest::SetParticleSimDataToCollide({ BounceCubeProxy->GetParticle_LowLevel() });
-		BounceCubeParticle.SetV(FVec3(0, 0, -100));
+		// Collide cube with static floor with restitution modified to 0
 
 		// Simulated cube with downawrd velocity, should not bounce, and end up with zero velocity.
 		FSingleParticlePhysicsProxy* NoBounceCubeProxy = FSingleParticlePhysicsProxy::Create(Chaos::FPBDRigidParticle::CreateParticle());
@@ -476,10 +464,9 @@ namespace ChaosTest
 
 		// Save Unique indices of floor and modified cube to disable in contact mod.
 		TVec2<FUniqueIdx> NoBounceUniqueIndices({ NoBounceCubeParticle.UniqueIdx(), FloorParticle.UniqueIdx() });
-		TVec2<FUniqueIdx> BounceUniqueIndices({ BounceCubeParticle.UniqueIdx(), FloorParticle.UniqueIdx() });
 
 		FContactModificationTestCallback* Callback = Solver->CreateAndRegisterSimCallbackObject_External<FContactModificationTestCallback>(true);
-		Callback->TestLambda = [BounceUniqueIndices, NoBounceUniqueIndices](Chaos::FCollisionContactModifier& Modifier)
+		Callback->TestLambda = [NoBounceUniqueIndices](Chaos::FCollisionContactModifier& Modifier)
 		{
 			for (FContactPairModifier& PairModifier : Modifier)
 			{
@@ -497,15 +484,79 @@ namespace ChaosTest
 						PairModifier.ModifyRestitution(0);
 					}
 				}
-				else if ((BounceUniqueIndices[0] == Idx0 && BounceUniqueIndices[1] == Idx1) ||
-						 (BounceUniqueIndices[0] == Idx1 && BounceUniqueIndices[1] == Idx0))
+			}
+		};
+
+
+		const float Dt = 0.1f;
+		const int32 Steps = 30;
+		for (int Step = 0; Step < Steps; ++Step)
+		{
+			Solver->AdvanceAndDispatch_External(Dt);
+			Solver->UpdateGameThreadStructures();
+		}
+
+		
+		EXPECT_NEAR(NoBounceCubeParticle.V().Z, 0.f, KINDA_SMALL_NUMBER);
+
+		Solver->UnregisterAndFreeSimCallbackObject_External(Callback);
+		Solver->UnregisterObject(NoBounceCubeProxy);
+		Solver->UnregisterObject(FloorProxy);
+		Module->DestroySolver(Solver);
+	}	
+
+
+	GTEST_TEST(AllTraits, ContactModification_ModifyRestitution_Bounce)
+	{
+		FChaosSolversModule* Module = FChaosSolversModule::GetModule();
+		auto* Solver = Module->CreateSolver(nullptr, /*AsyncDt=*/-1);
+		InitSolverSettings(Solver);
+		Solver->SetThreadingMode_External(EThreadingModeTemp::SingleThread);
+
+		// Collide cubes with static floor  with restitution modified to 1
+
+		// simulated cube with downward velocity, should bounce on floor and end up with upward velocity.
+		FSingleParticlePhysicsProxy* BounceCubeProxy = FSingleParticlePhysicsProxy::Create(Chaos::FPBDRigidParticle::CreateParticle());
+		auto& BounceCubeParticle = BounceCubeProxy->GetGameThreadAPI();
+		auto BounceCubeGeom = TSharedPtr<FImplicitObject, ESPMode::ThreadSafe>(new TBox<FReal, 3>(FVec3(-100), FVec3(100)));
+		BounceCubeParticle.SetGeometry(BounceCubeGeom);
+		Solver->RegisterObject(BounceCubeProxy);
+		BounceCubeParticle.SetGravityEnabled(false);
+		BounceCubeParticle.SetX(FVec3(200, 0, 200));
+		SetCubeInertiaTensor(BounceCubeParticle, /*Dimension=*/200, /*Mass=*/1);
+		ChaosTest::SetParticleSimDataToCollide({ BounceCubeProxy->GetParticle_LowLevel() });
+		BounceCubeParticle.SetV(FVec3(0, 0, -100));
+
+		// static floor at origin, occupying Z = [-100,0]
+		FSingleParticlePhysicsProxy* FloorProxy = FSingleParticlePhysicsProxy::Create(Chaos::FGeometryParticle::CreateParticle());
+		auto& FloorParticle = FloorProxy->GetGameThreadAPI();
+		auto FloorGeom = TSharedPtr<FImplicitObject, ESPMode::ThreadSafe>(new TBox<FReal, 3>(FVec3(-500, -500, -100), FVec3(500, 500, 0)));
+		FloorParticle.SetGeometry(FloorGeom);
+		Solver->RegisterObject(FloorProxy);
+		FloorParticle.SetX(FVec3(0, 0, 0));
+		ChaosTest::SetParticleSimDataToCollide({ FloorProxy->GetParticle_LowLevel() });
+
+		// Save Unique indices of floor and modified cube to disable in contact mod.
+		TVec2<FUniqueIdx> BounceUniqueIndices({ BounceCubeParticle.UniqueIdx(), FloorParticle.UniqueIdx() });
+
+		FContactModificationTestCallback* Callback = Solver->CreateAndRegisterSimCallbackObject_External<FContactModificationTestCallback>(true);
+		Callback->TestLambda = [BounceUniqueIndices](Chaos::FCollisionContactModifier& Modifier)
+		{
+			for (FContactPairModifier& PairModifier : Modifier)
+			{
+				TVec2<FGeometryParticleHandle*> Particles = PairModifier.GetParticlePair();
+				FUniqueIdx Idx0 = Particles[0]->UniqueIdx();
+				FUniqueIdx Idx1 = Particles[1]->UniqueIdx();
+
+				if ((BounceUniqueIndices[0] == Idx0 && BounceUniqueIndices[1] == Idx1) ||
+					(BounceUniqueIndices[0] == Idx1 && BounceUniqueIndices[1] == Idx0))
 				{
 					int32 NumContacts = PairModifier.GetNumContacts();
 					for (int32 PointIdx = 0; PointIdx < NumContacts; ++PointIdx)
 					{
 						// set restitution
 						PairModifier.ModifyRestitution(1);
-						
+
 						// Our object is slow enough restitution will not be applied with default threshold.
 						PairModifier.ModifyRestitutionThreshold(99);
 					}
@@ -522,18 +573,13 @@ namespace ChaosTest
 			Solver->UpdateGameThreadStructures();
 		}
 
-		
-		EXPECT_NEAR(NoBounceCubeParticle.V().Z, 0.f, KINDA_SMALL_NUMBER);
-
-		// 50 is a somewhat arbritrary threshold, would expect 100 but solver seems to lose some energy with rest = 1.
-		EXPECT_GT(BounceCubeParticle.V().Z, 50.f);
+		EXPECT_NEAR(BounceCubeParticle.V().Z, 100.f, KINDA_SMALL_NUMBER);
 
 		Solver->UnregisterAndFreeSimCallbackObject_External(Callback);
 		Solver->UnregisterObject(BounceCubeProxy);
-		Solver->UnregisterObject(NoBounceCubeProxy);
 		Solver->UnregisterObject(FloorProxy);
 		Module->DestroySolver(Solver);
-	}	
+	}
 
 	GTEST_TEST(AllTraits, ContactModification_ModifyFriction)
 	{
@@ -572,7 +618,7 @@ namespace ChaosTest
 		FloorParticle.SetGeometry(FloorGeom);
 		Solver->RegisterObject(FloorProxy);
 		FloorParticle.SetX(FVec3(0, 0, 0));
-		FloorParticle.SetR(FQuat::MakeFromEuler(FVec3(0, 30, 0)));
+		FloorParticle.SetR(FQuat::MakeFromEuler(FVec3(0, 20, 0)));
 		ChaosTest::SetParticleSimDataToCollide({ FloorProxy->GetParticle_LowLevel() });
 
 
diff --git a/Engine/Source/Runtime/Experimental/Chaos/Private/Chaos/Collision/PBDCollisionSolver.cpp b/Engine/Source/Runtime/Experimental/Chaos/Private/Chaos/Collision/PBDCollisionSolver.cpp
--- a/Engine/Source/Runtime/Experimental/Chaos/Private/Chaos/Collision/PBDCollisionSolver.cpp
+++ b/Engine/Source/Runtime/Experimental/Chaos/Private/Chaos/Collision/PBDCollisionSolver.cpp
@@ -67,9 +67,15 @@ namespace Chaos
 		const FVec3& InWorldAnchorPoint1,
 		const FVec3LP& InWorldContactNormal)
 	{
-		NetPushOut = FVec3LP(0);
-		NetImpulse = FVec3LP(0);
+		NetPushOutNormal = FRealLP(0);
+		NetPushOutTangentU = FRealLP(0);
+		NetPushOutTangentV = FRealLP(0);
+		NetImpulseNormal = FRealLP(0);
+		NetImpulseTangentU = FRealLP(0);
+		NetImpulseTangentV = FRealLP(0);
+
 		UpdateContact(Body0, Body1, InWorldAnchorPoint0, InWorldAnchorPoint1, InWorldContactNormal);
+
 		UpdateMass(Body0, Body1);
 	}
 
@@ -79,9 +85,9 @@ namespace Chaos
 		const FRealLP InRestitution,
 		const FRealLP InRestitutionVelocityThreshold,
 		const bool bInEnableStaticFriction,
-		const FRealLP InStaticFrictionMax)
+		const FRealLP InFrictionNetPushOutNormal)
 	{
-		StaticFrictionMax = InStaticFrictionMax;
+		FrictionNetPushOutNormal = InFrictionNetPushOutNormal;
 		bInsideStaticFrictionCone = bInEnableStaticFriction;
 		WorldContactVelocityTargetNormal = FRealLP(0);
 
@@ -112,34 +118,76 @@ namespace Chaos
 		// The world-space contact normal
 		WorldContactNormal = InWorldContactNormal;
 
+		// World-space contact tangents. We are treating the normal as the constraint-space Z axis
+		// and the Tangent U and V as the constraint-space X and Y axes resepctively
+		WorldContactTangentU = FVec3LP::CrossProduct(FVec3LP(0, 1, 0), InWorldContactNormal);
+		if (!WorldContactTangentU.Normalize(FRealLP(KINDA_SMALL_NUMBER)))
+		{
+			WorldContactTangentU = FVec3LP::CrossProduct(FVec3LP(1, 0, 0), InWorldContactNormal);
+			WorldContactTangentU = WorldContactTangentU.GetUnsafeNormal();
+		}
+		WorldContactTangentV = FVec3LP::CrossProduct(InWorldContactNormal, WorldContactTangentU);
+
 		// The contact point error we are trying to correct in this solver
-		WorldContactDelta = InWorldAnchorPoint0 - InWorldAnchorPoint1;
+		const FVec3LP WorldContactDelta = InWorldAnchorPoint0 - InWorldAnchorPoint1;
+		WorldContactDeltaNormal = FVec3LP::DotProduct(WorldContactDelta, WorldContactNormal);
+		WorldContactDeltaTangent0 = FVec3LP::DotProduct(WorldContactDelta, WorldContactTangentU);
+		WorldContactDeltaTangent1 = FVec3LP::DotProduct(WorldContactDelta, WorldContactTangentV);
 	}
 
 	inline void FPBDCollisionSolverManifoldPoint::UpdateMass(const FConstraintSolverBody& Body0, const FConstraintSolverBody& Body1)
 	{
-		const FMatrix33LP ContactMassInv =
-			(Body0.IsDynamic() ? Collisions::ComputeFactorMatrix3(RelativeContactPosition0, FMatrix33LP(Body0.InvI()), FRealLP(Body0.InvM())) : FMatrix33LP(0)) +
-			(Body1.IsDynamic() ? Collisions::ComputeFactorMatrix3(RelativeContactPosition1, FMatrix33LP(Body1.InvI()), FRealLP(Body1.InvM())) : FMatrix33LP(0));
-		const FMatrix33LP ContactMass = ContactMassInv.Inverse();
-		const FRealLP ContactMassInvNormal = FVec3LP::DotProduct(WorldContactNormal, Utilities::Multiply(ContactMassInv, WorldContactNormal));
-		const FRealLP ContactMassNormal = (ContactMassInvNormal > FRealLP(SMALL_NUMBER)) ? FRealLP(1) / ContactMassInvNormal : FRealLP(0);
-		WorldContactMass = ContactMass;
-		WorldContactMassNormal = ContactMassNormal;
-
-		WorldContactNormalAngular0 = FVec3LP(0);
-		WorldContactNormalAngular1 = FVec3LP(0);
+		FRealLP ContactMassInvNormal = FRealLP(0);
+		FRealLP ContactMassInvTangentU = FRealLP(0);
+		FRealLP ContactMassInvTangentV = FRealLP(0);
+
+		// These are not used if not initialized below so no need to clear
+		//WorldContactNormalAngular0 = FVec3LP(0);
+		//WorldContactTangentUAngular0 = FVec3LP(0);
+		//WorldContactTangentVAngular0 = FVec3LP(0);
+		//WorldContactNormalAngular1 = FVec3LP(0);
+		//WorldContactTangentUAngular1 = FVec3LP(0);
+		//WorldContactTangentVAngular1 = FVec3LP(0);
+
 		if (Body0.IsDynamic())
 		{
-			WorldContactNormalAngular0 = Body0.InvI() * FVec3LP::CrossProduct(RelativeContactPosition0, WorldContactNormal);
+			const FVec3LP R0xN = FVec3LP::CrossProduct(RelativeContactPosition0, WorldContactNormal);
+			const FVec3LP R0xU = FVec3LP::CrossProduct(RelativeContactPosition0, WorldContactTangentU);
+			const FVec3LP R0xV = FVec3LP::CrossProduct(RelativeContactPosition0, WorldContactTangentV);
+
+			const FMatrix33LP InvI0 = Body0.InvI();
+
+			WorldContactNormalAngular0 = InvI0 * R0xN;
+			WorldContactTangentUAngular0 = InvI0 * R0xU;
+			WorldContactTangentVAngular0 = InvI0 * R0xV;
+
+			ContactMassInvNormal += FVec3LP::DotProduct(R0xN, WorldContactNormalAngular0) + Body0.InvM();
+			ContactMassInvTangentU += FVec3LP::DotProduct(R0xU, WorldContactTangentUAngular0) + Body0.InvM();
+			ContactMassInvTangentV += FVec3LP::DotProduct(R0xV, WorldContactTangentVAngular0) + Body0.InvM();
 		}
 		if (Body1.IsDynamic())
 		{
-			WorldContactNormalAngular1 = Body1.InvI() * FVec3LP::CrossProduct(RelativeContactPosition1, WorldContactNormal);
+			const FVec3LP R1xN = FVec3LP::CrossProduct(RelativeContactPosition1, WorldContactNormal);
+			const FVec3LP R1xU = FVec3LP::CrossProduct(RelativeContactPosition1, WorldContactTangentU);
+			const FVec3LP R1xV = FVec3LP::CrossProduct(RelativeContactPosition1, WorldContactTangentV);
+
+			const FMatrix33LP InvI1 = Body1.InvI();
+
+			WorldContactNormalAngular1 = InvI1 * R1xN;
+			WorldContactTangentUAngular1 = InvI1 * R1xU;
+			WorldContactTangentVAngular1 = InvI1 * R1xV;
+
+			ContactMassInvNormal += FVec3LP::DotProduct(R1xN, WorldContactNormalAngular1) + Body1.InvM();
+			ContactMassInvTangentU += FVec3LP::DotProduct(R1xU, WorldContactTangentUAngular1) + Body1.InvM();
+			ContactMassInvTangentV += FVec3LP::DotProduct(R1xV, WorldContactTangentVAngular1) + Body1.InvM();
 		}
+
+		ContactMassNormal = (ContactMassInvNormal > FRealLP(SMALL_NUMBER)) ? FRealLP(1) / ContactMassInvNormal : FRealLP(0);
+		ContactMassTangentU = (ContactMassInvTangentU > FRealLP(SMALL_NUMBER)) ? FRealLP(1) / ContactMassInvTangentU : FRealLP(0);
+		ContactMassTangentV = (ContactMassInvTangentV > FRealLP(SMALL_NUMBER)) ? FRealLP(1) / ContactMassInvTangentV : FRealLP(0);
 	}
 
-	FVec3 FPBDCollisionSolverManifoldPoint::CalculateContactVelocity(const FConstraintSolverBody& Body0, const FConstraintSolverBody& Body1) const
+	FVec3LP FPBDCollisionSolverManifoldPoint::CalculateContactVelocity(const FConstraintSolverBody& Body0, const FConstraintSolverBody& Body1) const
 	{
 		const FVec3LP ContactVelocity0 = Body0.V() + FVec3::CrossProduct(Body0.W(), RelativeContactPosition0);
 		const FVec3LP ContactVelocity1 = Body1.V() + FVec3::CrossProduct(Body1.W(), RelativeContactPosition1);
@@ -234,7 +282,7 @@ namespace Chaos
 		const FReal InRestitution,
 		const FReal InRestitutionVelocityThreshold,
 		const bool bInEnableStaticFriction,
-		const FReal InStaticFrictionMax)
+		const FReal InFrictionNetPushOutNormal)
 	{
 		State.ManifoldPoints[ManifoldPoiontIndex].InitMaterial(
 			State.SolverBodies[0], 
@@ -242,7 +290,7 @@ namespace Chaos
 			FRealLP(InRestitution), 
 			FRealLP(InRestitutionVelocityThreshold),
 			bInEnableStaticFriction, 
-			FRealLP(InStaticFrictionMax));
+			FRealLP(InFrictionNetPushOutNormal));
 
 		// Track if any points have restitution enabled. See SolveVelocity
 		State.bHaveRestitution = State.bHaveRestitution || (State.ManifoldPoints[ManifoldPoiontIndex].WorldContactVelocityTargetNormal > 0);
@@ -273,23 +321,28 @@ namespace Chaos
 		FVec3LP RelativeContactPosition0 = FVec3(0);
 		FVec3LP RelativeContactPosition1 = FVec3(0);
 		FVec3LP WorldContactNormal = FVec3LP(0);
-		FVec3LP NetPushOut = FVec3LP(0);
+		FRealLP NetPushOutNormal = FRealLP(0);
 		FRealLP WorldContactVelocityTargetNormal = FReal(0);
 		for (int32 PointIndex = 0; PointIndex < State.NumManifoldPoints; ++PointIndex)
 		{
 			FPBDCollisionSolverManifoldPoint& SolverManifoldPoint = State.ManifoldPoints[PointIndex];
-			if (!SolverManifoldPoint.NetPushOut.IsNearlyZero())
+			if (!FMath::IsNearlyZero(SolverManifoldPoint.NetPushOutNormal))
 			{
 				RelativeContactPosition0 += SolverManifoldPoint.RelativeContactPosition0;
 				RelativeContactPosition1 += SolverManifoldPoint.RelativeContactPosition1;
 				WorldContactVelocityTargetNormal += SolverManifoldPoint.WorldContactVelocityTargetNormal;
 				WorldContactNormal = SolverManifoldPoint.WorldContactNormal;	// Take last value - should all be similar
+				NetPushOutNormal += SolverManifoldPoint.NetPushOutNormal;
 				++NumActiveManifoldPoints;
 			}
 		}
 
 		// Solve for velocity at the new conatct point
 		// We only do this if we have multiple active contacts
+		// NOTE: this average contact isn't really correct, especially when all contacts do not equally
+		// contribute to the pushout (which is normal even for a simple box on a plane). E.g., the WorldContactVelocityTargetNormal
+		// and NetPushOutNormal will not be right. The goal here though is just to get close to the correct result in frewer iterations
+		// than we would have if we just solved the corners of the box in sequence.
 		if (NumActiveManifoldPoints > 1)
 		{
 			const FRealLP DynamicFriction = FRealLP(0);
@@ -300,31 +353,42 @@ namespace Chaos
 			AverageManifoldPoint.RelativeContactPosition1 = RelativeContactPosition1 * Scale;
 			AverageManifoldPoint.WorldContactNormal = WorldContactNormal;
 			AverageManifoldPoint.WorldContactVelocityTargetNormal = WorldContactVelocityTargetNormal * Scale;
+			AverageManifoldPoint.NetPushOutNormal = NetPushOutNormal * Scale;
 
 			const FMatrix33LP ContactMassInv =
 				(Body0.IsDynamic() ? Collisions::ComputeFactorMatrix3(AverageManifoldPoint.RelativeContactPosition0, Body0.InvI(), Body0.InvM()) : FMatrix33LP(0)) +
 				(Body1.IsDynamic() ? Collisions::ComputeFactorMatrix3(AverageManifoldPoint.RelativeContactPosition1, Body1.InvI(), Body1.InvM()) : FMatrix33LP(0));
 			const FRealLP ContactMassInvNormal = FVec3LP::DotProduct(AverageManifoldPoint.WorldContactNormal, Utilities::Multiply(ContactMassInv, AverageManifoldPoint.WorldContactNormal));
 			const FRealLP ContactMassNormal = (ContactMassInvNormal > FRealLP(SMALL_NUMBER)) ? FRealLP(1) / ContactMassInvNormal : FRealLP(0);
-			AverageManifoldPoint.WorldContactMassNormal = ContactMassNormal;
-
-			AverageManifoldPoint.WorldContactMass = FMatrix33LP(0);	// Unused
-			AverageManifoldPoint.WorldContactDelta = FVec3LP(0);	// Unused
-			AverageManifoldPoint.NetPushOut = FVec3LP(0);			// Unused
-			AverageManifoldPoint.NetImpulse = FVec3LP(0);			// Unused
-			AverageManifoldPoint.StaticFrictionMax = FRealLP(0);	// Unused
-			AverageManifoldPoint.bInsideStaticFrictionCone = false;	// Unused
-
-			FVec3LP ContactVelocityDelta;
-			FRealLP ContactVelocityDeltaNormal;
-			AverageManifoldPoint.CalculateContactVelocityError(Body0, Body1, DynamicFriction, Dt, ContactVelocityDelta, ContactVelocityDeltaNormal);
+			AverageManifoldPoint.ContactMassNormal = ContactMassNormal;
+
+			// @todo(chaos): All unused for the average contact solve - optimize
+			AverageManifoldPoint.ContactMassTangentU = FRealLP(0);
+			AverageManifoldPoint.ContactMassTangentV = FRealLP(0);
+			AverageManifoldPoint.WorldContactTangentU = FVec3LP(0);
+			AverageManifoldPoint.WorldContactTangentV = FVec3LP(0);
+			AverageManifoldPoint.WorldContactDeltaNormal = FRealLP(0);
+			AverageManifoldPoint.WorldContactDeltaTangent0 = FRealLP(0);
+			AverageManifoldPoint.WorldContactDeltaTangent1 = FRealLP(0);
+			AverageManifoldPoint.NetPushOutTangentU = FRealLP(0);
+			AverageManifoldPoint.NetPushOutTangentV = FRealLP(0);
+			AverageManifoldPoint.NetImpulseNormal = FRealLP(0);
+			AverageManifoldPoint.NetImpulseTangentU = FRealLP(0);
+			AverageManifoldPoint.NetImpulseTangentV = FRealLP(0);
+			AverageManifoldPoint.FrictionNetPushOutNormal = FRealLP(0);
+			AverageManifoldPoint.bInsideStaticFrictionCone = false;
+
+			// @todo(chaos): don't need tangential velocity delta here since we are setting DynamicFriction to 0
+			FRealLP ContactVelocityDeltaNormal, ContactVelocityDeltaTangentU, ContactVelocityDeltaTangentV;
+			AverageManifoldPoint.CalculateContactVelocityError(Body0, Body1, DynamicFriction, Dt, ContactVelocityDeltaNormal, ContactVelocityDeltaTangentU, ContactVelocityDeltaTangentV);
 
 			ApplyVelocityCorrection(
 				State.Stiffness,
 				Dt,
 				DynamicFriction,
-				ContactVelocityDelta,
 				ContactVelocityDeltaNormal,
+				ContactVelocityDeltaTangentU,
+				ContactVelocityDeltaTangentV,
 				AverageManifoldPoint,
 				Body0,
 				Body1);
@@ -335,7 +399,8 @@ namespace Chaos
 	{
 		// Apply restitution at the average contact point
 		// This means we don't need to run as many iterations to get stable bouncing
-		if (State.bHaveRestitution && (State.NumManifoldPoints > 1))
+		// It also helps with zero restitution to counter any velocioty added by the PBD solve
+		if (State.NumManifoldPoints > 1)
 		{
 			SolveVelocityAverage(Dt);
 		}
@@ -350,17 +415,18 @@ namespace Chaos
 		for (int32 PointIndex = 0; PointIndex < State.NumManifoldPoints; ++PointIndex)
 		{
 			FPBDCollisionSolverManifoldPoint& SolverManifoldPoint = State.ManifoldPoints[PointIndex];
-			if (!SolverManifoldPoint.NetPushOut.IsNearlyZero())
+			if (!FMath::IsNearlyZero(SolverManifoldPoint.NetPushOutNormal))
 			{
-				FVec3LP ContactVelocityDelta;
-				FRealLP ContactVelocityDeltaNormal;
-				SolverManifoldPoint.CalculateContactVelocityError(Body0, Body1, DynamicFriction, Dt, ContactVelocityDelta, ContactVelocityDeltaNormal);
+				FRealLP ContactVelocityDeltaNormal, ContactVelocityDeltaTangentU, ContactVelocityDeltaTangentV;
+				SolverManifoldPoint.CalculateContactVelocityError(Body0, Body1, DynamicFriction, Dt, ContactVelocityDeltaNormal, ContactVelocityDeltaTangentU, ContactVelocityDeltaTangentV);
+
 				ApplyVelocityCorrection(
 					State.Stiffness,
 					Dt,
 					DynamicFriction,
-					ContactVelocityDelta,
 					ContactVelocityDeltaNormal,
+					ContactVelocityDeltaTangentU,
+					ContactVelocityDeltaTangentV,
 					SolverManifoldPoint,
 					Body0,
 					Body1);
diff --git a/Engine/Source/Runtime/Experimental/Chaos/Private/Chaos/Collision/PBDCollisionSolver.h b/Engine/Source/Runtime/Experimental/Chaos/Private/Chaos/Collision/PBDCollisionSolver.h
--- a/Engine/Source/Runtime/Experimental/Chaos/Private/Chaos/Collision/PBDCollisionSolver.h
+++ b/Engine/Source/Runtime/Experimental/Chaos/Private/Chaos/Collision/PBDCollisionSolver.h
@@ -60,7 +60,7 @@ namespace Chaos
 			const FRealLP InRestitution,
 			const FRealLP InRestitutionVelocityThreshold,
 			const bool bInEnableStaticFriction,
-			const FRealLP InStaticFrictionMax);
+			const FRealLP InFrictionNetPushOutNormal);
 
 		/**
 		 * @brief Update the world-space relative contact points based on current body transforms and body-space contact positions
@@ -77,56 +77,68 @@ namespace Chaos
 		*/
 		void UpdateMass(const FConstraintSolverBody& Body0, const FConstraintSolverBody& Body1);
 
-		/**
-		 * @brief Calculate the relative velocity at the contact point
-		 * @note InitContact must be called before calling this function
-		*/
-		FVec3 CalculateContactVelocity(const FConstraintSolverBody& Body0, const FConstraintSolverBody& Body1) const;
-
 		/**
 		 * @brief Calculate the position error at the current transforms
 		 * @param MaxPushOut a limit on the position error for this iteration to prevent initial-penetration explosion (a common PBD problem)
 		*/
-		void CalculateContactPositionError(const FConstraintSolverBody& Body0, const FConstraintSolverBody& Body1, const FRealLP MaxPushOut, FVec3LP& OutContactDelta, FRealLP& OutContactDeltaNormal) const;
+		void CalculateContactPositionErrorNormal(const FConstraintSolverBody& Body0, const FConstraintSolverBody& Body1, const FRealLP MaxPushOut, FRealLP& OutContactDeltaNormal) const;
+		void CalculateContactPositionError(const FConstraintSolverBody& Body0, const FConstraintSolverBody& Body1, const FRealLP MaxPushOut, FRealLP& OutContactDeltaNormal, FRealLP& OutContactDeltaTanget0, FRealLP& OutContactDeltaTangent1) const;
 
 		/**
 		 * @brief Calculate the velocity error at the current transforms
 		*/
-		void CalculateContactVelocityError(const FConstraintSolverBody& Body0, const FConstraintSolverBody& Body1, const FRealLP DynamicFriction, const FRealLP Dt, FVec3LP& OutContactVelocityDelta, FRealLP& OutContactVelocityDeltaNormal) const;
+		void CalculateContactVelocityError(const FConstraintSolverBody& Body0, const FConstraintSolverBody& Body1, const FRealLP DynamicFriction, const FRealLP Dt, FRealLP& OutContactVelocityDeltaNormal, FRealLP& OutContactVelocityDeltaTangent0, FRealLP& OutContactVelocityDeltaTangent1) const;
 
 		// @todo(chaos): make private
 	public:
 		friend class FPBDCollisionSolver;
 
-		// Contact mass (for friction)
-		FMatrix33LP WorldContactMass;
+		/**
+		 * @brief Calculate the relative velocity at the contact point
+		 * @note InitContact must be called before calling this function
+		*/
+		FVec3LP CalculateContactVelocity(const FConstraintSolverBody& Body0, const FConstraintSolverBody& Body1) const;
 
 		// World-space body-relative contact points
 		FVec3LP RelativeContactPosition0;
 		FVec3LP RelativeContactPosition1;
 
-		// World-space contact normal
+		// World-space contact axes (normal and 2 tangents)
 		FVec3LP WorldContactNormal;
+		FVec3LP WorldContactTangentU;
+		FVec3LP WorldContactTangentV;
 
-		// I^-1.(R x N) for each body
+		// I^-1.(R x A) for each body where A is each axis (Normal, TangentU, TangentV)
 		FVec3LP WorldContactNormalAngular0;
+		FVec3LP WorldContactTangentUAngular0;
+		FVec3LP WorldContactTangentVAngular0;
 		FVec3LP WorldContactNormalAngular1;
+		FVec3LP WorldContactTangentUAngular1;
+		FVec3LP WorldContactTangentVAngular1;
 
 		// Contact mass (for non-friction)
-		FRealLP WorldContactMassNormal;
+		FRealLP ContactMassNormal;
+		FRealLP ContactMassTangentU;
+		FRealLP ContactMassTangentV;
 
 		// Initial world-space contact separation that we are trying to correct
-		FVec3LP WorldContactDelta;
+		FRealLP WorldContactDeltaNormal;
+		FRealLP WorldContactDeltaTangent0;
+		FRealLP WorldContactDeltaTangent1;
 
 		// Desired final normal velocity, taking Restitution into account
 		FRealLP WorldContactVelocityTargetNormal;
 
 		// Solver outputs
-		FVec3LP NetPushOut;
-		FVec3LP NetImpulse;
+		FRealLP NetPushOutNormal;
+		FRealLP NetPushOutTangentU;
+		FRealLP NetPushOutTangentV;
+		FRealLP NetImpulseNormal;
+		FRealLP NetImpulseTangentU;
+		FRealLP NetImpulseTangentV;
 
-		// A smoothed NetImpulse along the normal, used for clipping to the static friction cone
-		FRealLP StaticFrictionMax;
+		// A smoothed NetPushOutNormal along the normal, used for clipping to the static friction cone
+		FRealLP FrictionNetPushOutNormal;
 
 		// Whether we are still in the static friction cone
 		bool bInsideStaticFrictionCone;
@@ -200,7 +212,7 @@ namespace Chaos
 			const FReal InRestitution,
 			const FReal InRestitutionVelocityThreshold,
 			const bool bInEnableStaticFriction,
-			const FReal InStaticFrictionMax);
+			const FReal InFrictionNetPushOutNormal);
 
 		void UpdateContact(
 			const int32 ManifoldPoiontIndex,
@@ -297,171 +309,158 @@ namespace Chaos
 	//////////////////////////////////////////////////////////////////////////////////////////////////
 
 
-	FORCEINLINE_DEBUGGABLE void CalculatePositionCorrectionWithoutFrictionV(
+	FORCEINLINE_DEBUGGABLE void CalculatePositionCorrectionNormal(
 		const FRealLP Stiffness,
 		const FRealLP ContactDeltaNormal,
-		const FVec3LP& ContactNormal,
 		const FRealLP ContactMassNormal,
-		FVec3LP& InOutNetPushOut,
-		FVec3LP& OutPushOut)
+		const FRealLP NetPushOutNormal,
+		FRealLP& OutPushOutNormal)
 	{
-		FVec3LP PushOut = -(Stiffness * ContactDeltaNormal * ContactMassNormal) * ContactNormal;
+		const FRealLP PushOutNormal = -Stiffness * ContactDeltaNormal * ContactMassNormal;
 
 		// The total pushout so far this sub-step
-		// We allow negative incremental impulses, but not net negative impulses
-		FVec3LP NetPushOut = InOutNetPushOut + PushOut;
-		const FRealLP NetPushOutNormal = FVec3LP::DotProduct(NetPushOut, ContactNormal);
-		if (NetPushOutNormal < 0)
+		// Unilateral constraint: Net-negative impulses not allowed (negative incremental impulses are allowed as long as the net is positive)
+		if ((NetPushOutNormal + PushOutNormal) > FRealLP(0))
 		{
-			PushOut = -InOutNetPushOut;
-			NetPushOut = FVec3LP(0);
+			OutPushOutNormal = PushOutNormal;
+		}
+		else
+		{
+			OutPushOutNormal = -NetPushOutNormal;
 		}
-
-		InOutNetPushOut = NetPushOut;
-		OutPushOut = PushOut;
 	}
 
 
-	FORCEINLINE_DEBUGGABLE void CalculatePositionCorrectionWithoutFrictionF(
+	FORCEINLINE_DEBUGGABLE void CalculatePositionCorrectionTangent(
 		const FRealLP Stiffness,
-		const FRealLP ContactDeltaNormal,
-		const FRealLP ContactMassNormal,
-		const FRealLP NetPushOutNormal,
-		FRealLP& OutPushOutNormal)
+		const FRealLP ContactDeltaTangent,
+		const FRealLP ContactMassTangent,
+		FRealLP& OutPushOutTangent)
 	{
-		FRealLP PushOutNormal = -(Stiffness * ContactDeltaNormal * ContactMassNormal);
-
-		// The total pushout so far this sub-step
-		// We allow negative incremental impulses, but not net negative impulses
-		if (NetPushOutNormal + PushOutNormal < 0)
-		{
-			PushOutNormal = -NetPushOutNormal;
-		}
+		const FRealLP FrictionStiffness = CVars::Chaos_PBDCollisionSolver_Position_StaticFrictionStiffness;
 
-		OutPushOutNormal = PushOutNormal;
+		// Bilateral constraint - negative values allowed (unlike the normal correction)
+		OutPushOutTangent = -Stiffness * FrictionStiffness * ContactMassTangent * ContactDeltaTangent;
 	}
 
 
-	FORCEINLINE_DEBUGGABLE void CalculatePositionCorrectionWithFriction(
-		const FRealLP Stiffness,
-		const FVec3LP& ContactDelta,
-		const FRealLP ContactDeltaNormal,
-		const FVec3LP& ContactNormal,
-		const FMatrix33LP& ContactMass,
+	FORCEINLINE_DEBUGGABLE void ApplyFrictionCone(
 		const FRealLP StaticFriction,
 		const FRealLP DynamicFriction,
-		FVec3LP& InOutNetPushOut,
-		FVec3LP& OutPushOut,
-		FRealLP& InOutStaticFrictionMax,
-		bool& bOutInsideStaticFrictionCone)
+		const FRealLP NetPushOutNormal,
+		FRealLP& InOutPushOutTangentU,
+		FRealLP& InOutPushOutTangentV,
+		FRealLP& InOutNetPushOutTangentU,
+		FRealLP& InOutNetPushOutTangentV,
+		FRealLP& InOutFrictionNetPushOutNormal,
+		bool& bInOutInsideStaticFrictionCone)
 	{
-		// If static friction is enabled, calculate the correction to move the contact point back to its
-		// original relative location on all axes.
-		// @todo(chaos): this should be moved to the ManifoldPoint error calculation?
-		FVec3LP ModifiedContactError = -ContactDelta;
-		const FRealLP FrictionStiffness = CVars::Chaos_PBDCollisionSolver_Position_StaticFrictionStiffness;
-		if (FrictionStiffness < 1.0f)
-		{
-			const FVec3LP ContactDeltaTangent = ContactDelta - ContactDeltaNormal * ContactNormal;
-			ModifiedContactError = -ContactDeltaNormal * ContactNormal - FrictionStiffness * ContactDeltaTangent;
-		}
-
-		FVec3LP PushOut = Stiffness * ContactMass * ModifiedContactError;
-
-		// If we ended up with a negative normal pushout, remove all correction from this point
-		FVec3LP NetPushOut = InOutNetPushOut + PushOut;
-		const FRealLP NetPushOutNormal = FVec3LP::DotProduct(NetPushOut, ContactNormal);
-		bool bInsideStaticFrictionCone = true;
-		if (NetPushOutNormal < FRealLP(SMALL_NUMBER))
-		{
-			bInsideStaticFrictionCone = false;
-			PushOut = -InOutNetPushOut;
-			NetPushOut = FVec3LP(0);
-		}
+		bInOutInsideStaticFrictionCone = true;
 
 		// Static friction limit: immediately increase maximum lateral correction, but smoothly decay maximum static friction limit. 
 		// This is so that small variations in position (jitter) and therefore NetPushOutNormal don't cause static friction to slip
 		// @todo(chaos): StaticFriction smoothing is iteration count depenendent - try to make it not so
 		const FRealLP StaticFrictionLerpRate = CVars::Chaos_PBDCollisionSolver_Position_StaticFrictionLerpRate;
-		const FRealLP StaticFrictionDest = FMath::Max(NetPushOutNormal, FRealLP(0));
-		FRealLP StaticFrictionMax = FMath::Lerp(FMath::Max(InOutStaticFrictionMax, StaticFrictionDest), StaticFrictionDest, StaticFrictionLerpRate);
+		const FRealLP StaticFrictionTarget = FMath::Max(NetPushOutNormal, FRealLP(0));
+		InOutFrictionNetPushOutNormal = FMath::Lerp(FMath::Max(InOutFrictionNetPushOutNormal, StaticFrictionTarget), StaticFrictionTarget, StaticFrictionLerpRate);
 
-		// If we exceed the friction cone, stop adding frictional corrections and clip correction to cone
-		// @todo(chaos): clamp to dynamic friction
-		if (bInsideStaticFrictionCone)
+		if (NetPushOutNormal < FRealLP(SMALL_NUMBER))
+		{
+			// Note: we have already added the current iteration's PushOut to the NetPushOut but it has not been applied to the body
+			// so we must subtract it again to calculate the actual pushout we want to undo (i.e., the net pushout that has been applied to the body so far)
+			InOutPushOutTangentU = -(InOutNetPushOutTangentU - InOutPushOutTangentU);
+			InOutPushOutTangentV = -(InOutNetPushOutTangentV - InOutPushOutTangentV);
+			InOutNetPushOutTangentU = FReal(0);
+			InOutNetPushOutTangentV = FReal(0);
+			bInOutInsideStaticFrictionCone = false;
+		}
+
+		// If we exceed the friction cone, clip to cone and stop adding frictional corrections
+		if (bInOutInsideStaticFrictionCone)
 		{
-			const FRealLP MaxPushOutTangentSq = FMath::Square(StaticFriction * StaticFrictionMax);
-			const FVec3LP NetPushOutTangent = (NetPushOut - NetPushOutNormal * ContactNormal);
-			const FRealLP NetPushOutTangentSq = NetPushOutTangent.SizeSquared();
+			const FRealLP MaxPushOutTangent = StaticFriction * InOutFrictionNetPushOutNormal;
+			const FRealLP MaxPushOutTangentSq = FMath::Square(MaxPushOutTangent);
+			const FRealLP NetPushOutTangentSq = FMath::Square(InOutNetPushOutTangentU) + FMath::Square(InOutNetPushOutTangentV);
 			if (NetPushOutTangentSq > MaxPushOutTangentSq)
 			{
-				NetPushOut = NetPushOutNormal * ContactNormal + (StaticFriction * StaticFrictionMax) * NetPushOutTangent * FMath::InvSqrt(NetPushOutTangentSq);
-				PushOut = NetPushOut - InOutNetPushOut;
-				bInsideStaticFrictionCone = false;
+				// @todo(chaos): clamp to dynamic friction rather than static
+				const FRealLP ClippedMaxPushOutTangent = MaxPushOutTangent;
+				const FRealLP InvNetPushOutTangent = FMath::InvSqrt(NetPushOutTangentSq);
+				const FRealLP NetPushOutTangentU = ClippedMaxPushOutTangent * InOutNetPushOutTangentU * InvNetPushOutTangent;
+				const FRealLP NetPushOutTangentV = ClippedMaxPushOutTangent * InOutNetPushOutTangentV * InvNetPushOutTangent;
+				InOutPushOutTangentU = NetPushOutTangentU - (InOutNetPushOutTangentU - InOutPushOutTangentU);
+				InOutPushOutTangentV = NetPushOutTangentV - (InOutNetPushOutTangentV - InOutPushOutTangentV);
+				InOutNetPushOutTangentU = NetPushOutTangentU;
+				InOutNetPushOutTangentV = NetPushOutTangentV;
+				bInOutInsideStaticFrictionCone = false;
 			}
 		}
-
-		InOutNetPushOut = NetPushOut;
-		OutPushOut = PushOut;
-		InOutStaticFrictionMax = StaticFrictionMax;
-		bOutInsideStaticFrictionCone = bInsideStaticFrictionCone;
 	}
 
+
 	FORCEINLINE_DEBUGGABLE void ApplyPositionCorrectionWithFriction(
 		const FRealLP Stiffness,
 		const FRealLP StaticFriction,
 		const FRealLP DynamicFriction,
-		const FVec3LP& ContactDelta,
 		const FRealLP ContactDeltaNormal,
+		const FRealLP ContactDeltaTangent0,
+		const FRealLP ContactDeltaTangent1,
 		FPBDCollisionSolverManifoldPoint& ManifoldPoint,
 		FConstraintSolverBody& Body0,
 		FConstraintSolverBody& Body1)
 	{
-		FVec3LP PushOut = FVec3LP(0);
+		FRealLP PushOutNormal = FRealLP(0);
+		FRealLP PushOutTangentU = FRealLP(0);
+		FRealLP PushOutTangentV = FRealLP(0);
 
-		// If we want to add static friction...
-		// (note: we run a few iterations without friction by temporarily setting StaticFriction to 0)
-		if ((StaticFriction > 0) && ManifoldPoint.bInsideStaticFrictionCone)
-		{
-			CalculatePositionCorrectionWithFriction(
-				Stiffness,
-				ContactDelta,
-				ContactDeltaNormal,
-				ManifoldPoint.WorldContactNormal,
-				ManifoldPoint.WorldContactMass,
-				StaticFriction,
-				DynamicFriction,
-				ManifoldPoint.NetPushOut,					// Out
-				PushOut,									// Out
-				ManifoldPoint.StaticFrictionMax,			// Out
-				ManifoldPoint.bInsideStaticFrictionCone);	// Out
-		}
-		else
-		{
-			CalculatePositionCorrectionWithoutFrictionV(
-				Stiffness,
-				ContactDeltaNormal,
-				ManifoldPoint.WorldContactNormal,
-				ManifoldPoint.WorldContactMassNormal,
-				ManifoldPoint.NetPushOut,					// Out
-				PushOut);									// Out
+		CalculatePositionCorrectionNormal(
+			Stiffness,
+			ContactDeltaNormal,
+			ManifoldPoint.ContactMassNormal,
+			ManifoldPoint.NetPushOutNormal,
+			PushOutNormal);						// Out
 
-		}
+		CalculatePositionCorrectionTangent(
+			Stiffness,
+			ContactDeltaTangent0,
+			ManifoldPoint.ContactMassTangentU,
+			PushOutTangentU);					// Out
+
+		CalculatePositionCorrectionTangent(
+			Stiffness,
+			ContactDeltaTangent1,
+			ManifoldPoint.ContactMassTangentV,
+			PushOutTangentV);					// Out
+
+		ManifoldPoint.NetPushOutNormal += PushOutNormal;
+		ManifoldPoint.NetPushOutTangentU += PushOutTangentU;
+		ManifoldPoint.NetPushOutTangentV += PushOutTangentV;
+
+		ApplyFrictionCone(
+			StaticFriction,
+			DynamicFriction,
+			ManifoldPoint.NetPushOutNormal,
+			PushOutTangentU,									// Out
+			PushOutTangentV,									// Out
+			ManifoldPoint.NetPushOutTangentU,					// Out
+			ManifoldPoint.NetPushOutTangentV,					// Out
+			ManifoldPoint.FrictionNetPushOutNormal,				// Out
+			ManifoldPoint.bInsideStaticFrictionCone);			// Out
+
+		const FVec3LP PushOut = PushOutNormal * ManifoldPoint.WorldContactNormal + PushOutTangentU * ManifoldPoint.WorldContactTangentU + PushOutTangentV * ManifoldPoint.WorldContactTangentV;
 
 		// Update the particle state based on the pushout
 		if (Body0.IsDynamic())
 		{
-			const FVec3LP AngularPushOut = FVec3LP::CrossProduct(ManifoldPoint.RelativeContactPosition0, PushOut);
 			const FVec3LP DX0 = Body0.InvM() * PushOut;
-			const FVec3LP DR0 = Body0.InvI() * AngularPushOut;
+			const FVec3LP DR0 = ManifoldPoint.WorldContactNormalAngular0 * PushOutNormal + ManifoldPoint.WorldContactTangentUAngular0 * PushOutTangentU + ManifoldPoint.WorldContactTangentVAngular0 * PushOutTangentV;
 			Body0.ApplyPositionDelta(DX0);
 			Body0.ApplyRotationDelta(DR0);
 		}
 		if (Body1.IsDynamic())
 		{
-			const FVec3LP AngularPushOut = FVec3LP::CrossProduct(ManifoldPoint.RelativeContactPosition1, PushOut);
 			const FVec3LP DX1 = -(Body1.InvM() * PushOut);
-			const FVec3LP DR1 = -(Body1.InvI() * AngularPushOut);
+			const FVec3LP DR1 = ManifoldPoint.WorldContactNormalAngular1 * -PushOutNormal + ManifoldPoint.WorldContactTangentUAngular1 * -PushOutTangentU + ManifoldPoint.WorldContactTangentVAngular1 * -PushOutTangentV;
 			Body1.ApplyPositionDelta(DX1);
 			Body1.ApplyRotationDelta(DR1);
 		}
@@ -475,16 +474,15 @@ namespace Chaos
 		FConstraintSolverBody& Body1)
 	{
 		FRealLP PushOutNormal = FRealLP(0);
-		const FRealLP NetPushOutNormal = FVec3LP::DotProduct(ManifoldPoint.NetPushOut, ManifoldPoint.WorldContactNormal);
 
-		CalculatePositionCorrectionWithoutFrictionF(
+		CalculatePositionCorrectionNormal(
 			Stiffness,
 			ContactDeltaNormal,
-			ManifoldPoint.WorldContactMassNormal,
-			NetPushOutNormal,
-			PushOutNormal);									// Out
+			ManifoldPoint.ContactMassNormal,
+			ManifoldPoint.NetPushOutNormal,
+			PushOutNormal);						// Out
 
-		ManifoldPoint.NetPushOut += PushOutNormal * ManifoldPoint.WorldContactNormal;
+		ManifoldPoint.NetPushOutNormal += PushOutNormal;
 
 		// Update the particle state based on the pushout
 		if (Body0.IsDynamic())
@@ -503,85 +501,46 @@ namespace Chaos
 		}
 	}
 
-	FORCEINLINE_DEBUGGABLE void CalculateVelocityCorrectionImpulse(
+
+	FORCEINLINE_DEBUGGABLE void ApplyVelocityCorrection(
 		const FRealLP Stiffness,
 		const FRealLP Dt,
 		const FRealLP DynamicFriction,
-		const FVec3LP& ContactNormal,
-		const FMatrix33LP& ContactMass,
-		const FRealLP ContactMassNormal,
-		const FVec3LP& ContactVelocityDelta,
 		const FRealLP ContactVelocityDeltaNormal,
-		const FVec3LP& NetPushOut,
-		FVec3LP& InOutNetImpulse,
-		FVec3LP& OutImpulse)
+		const FRealLP ContactVelocityDeltaTangent0,
+		const FRealLP ContactVelocityDeltaTangent1,
+		FPBDCollisionSolverManifoldPoint& ManifoldPoint,
+		FConstraintSolverBody& Body0,
+		FConstraintSolverBody& Body1)
 	{
-		FVec3LP Impulse = FVec3LP(0);
-
-		if ((ContactVelocityDeltaNormal > FRealLP(0)) && !CVars::bChaos_PBDCollisionSolver_Velocity_NegativeImpulseEnabled)
+		FRealLP ImpulseNormal = -(Stiffness * ManifoldPoint.ContactMassNormal) * ContactVelocityDeltaNormal;
+		
+		// Clamp the total impulse to be positive along the normal. We can apply a net negative impulse, 
+		// but only to correct the velocity that was added by pushout.
+		FRealLP NetImpulseNormal = ManifoldPoint.NetImpulseNormal + ImpulseNormal;
+		const FRealLP PushOutImpulseNormal = FMath::Max(0.0f, ManifoldPoint.NetPushOutNormal / Dt);
+		if (NetImpulseNormal < -PushOutImpulseNormal)
 		{
-			return;
+			// We are trying to apply a net negative impulse larger than one to counteract the effective pushout impulse
+			// so clamp the net impulse to be equal to minus the pushout impulse along the normal.
+			// NOTE: NetImpulseNormal is negative here
+			ImpulseNormal = ImpulseNormal - (NetImpulseNormal + PushOutImpulseNormal);
 		}
 
-		// Tangential velocity (dynamic friction)
-		const bool bApplyFriction = (DynamicFriction > 0) && (Dt > 0);
+		FRealLP ImpulseTangentU = FRealLP(0);
+		FRealLP ImpulseTangentV = FRealLP(0);
+		const bool bApplyFriction = (DynamicFriction > 0) && !FMath::IsNearlyZero(ImpulseNormal) && (Dt > 0);
 		if (bApplyFriction)
 		{
-			Impulse = -Stiffness * (ContactMass * ContactVelocityDelta);
-		}
-		else
-		{
-			Impulse = -(Stiffness * ContactMassNormal) * ContactVelocityDelta;
-		}
-
-		// Clamp the total impulse to be positive along the normal. We can apply negative velocity correction, 
-		// but only to correct the velocity that was added by pushout, or in this velocity solve step.
-		if (CVars::bChaos_PBDCollisionSolver_Velocity_ImpulseClampEnabled && (Dt > 0))
-		{
-			// @todo(chaos): cache max negative impulse
-			const FVec3LP NetImpulse = InOutNetImpulse + Impulse;
-			const FRealLP PushOutImpulseNormal = FMath::Max(0.0f, FVec3LP::DotProduct(NetPushOut, ContactNormal) / Dt);
-			const FRealLP NetImpulseNormal = FVec3LP::DotProduct(NetImpulse, ContactNormal);
-			if (NetImpulseNormal < -PushOutImpulseNormal)
-			{
-				// We are trying to apply a negative impulse larger than one to counteract the effective pushout impulse
-				// so clamp the net impulse to be equal to minus the pushout impulse along the normal.
-				// NOTE: NetImpulseNormal is negative here
-				//const FVec3 NewNetImpulse = NetImpulse - NetImpulseNormal * ContactNormal - PushOutImpulseNormal * ContactNormal;
-				//const FVec3 NewImpulse = NewNetImpulse - InOutNetImpulse;
-				//Impulse = InOutNetImpulse + Impulse - NetImpulseNormal * ContactNormal - PushOutImpulseNormal * ContactNormal - InOutNetImpulse;
-				Impulse = Impulse - (NetImpulseNormal + PushOutImpulseNormal) * ContactNormal;
-			}
+			ImpulseTangentU = -(Stiffness * ManifoldPoint.ContactMassTangentU) * ContactVelocityDeltaTangent0;
+			ImpulseTangentV = -(Stiffness * ManifoldPoint.ContactMassTangentV) * ContactVelocityDeltaTangent1;
 		}
 
-		OutImpulse = Impulse;
-		InOutNetImpulse += Impulse;
-	}
-
-	FORCEINLINE_DEBUGGABLE void ApplyVelocityCorrection(
-		const FRealLP Stiffness,
-		const FRealLP Dt,
-		const FRealLP DynamicFriction,
-		const FVec3LP& ContactVelocityDelta,
-		const FRealLP ContactVelocityDeltaNormal,
-		FPBDCollisionSolverManifoldPoint& ManifoldPoint,
-		FConstraintSolverBody& Body0,
-		FConstraintSolverBody& Body1)
-	{
-		FVec3LP Impulse = FVec3LP(0);
+		ManifoldPoint.NetImpulseNormal += ImpulseNormal;
+		ManifoldPoint.NetImpulseTangentU += ImpulseTangentU;
+		ManifoldPoint.NetImpulseTangentV += ImpulseTangentV;
 
-		CalculateVelocityCorrectionImpulse(
-			Stiffness,
-			Dt,
-			DynamicFriction,
-			ManifoldPoint.WorldContactNormal,
-			ManifoldPoint.WorldContactMass,
-			ManifoldPoint.WorldContactMassNormal,
-			ContactVelocityDelta,
-			ContactVelocityDeltaNormal,
-			ManifoldPoint.NetPushOut,	// Out
-			ManifoldPoint.NetImpulse,	// Out
-			Impulse);					// Out
+		FVec3LP Impulse = ImpulseNormal * ManifoldPoint.WorldContactNormal + ImpulseTangentU * ManifoldPoint.WorldContactTangentU + ImpulseTangentV * ManifoldPoint.WorldContactTangentV;
 
 		// Calculate the velocity deltas from the impulse
 		if (Body0.IsDynamic())
@@ -602,34 +561,51 @@ namespace Chaos
 	//////////////////////////////////////////////////////////////////////////////////////////////////
 	//////////////////////////////////////////////////////////////////////////////////////////////////
 
+	FORCEINLINE_DEBUGGABLE void FPBDCollisionSolverManifoldPoint::CalculateContactPositionErrorNormal(const FConstraintSolverBody& Body0, const FConstraintSolverBody& Body1, const FRealLP MaxPushOut, FRealLP& OutContactDeltaNormal) const
+	{
+		// Linear version: calculate the contact delta assuming linear motion after applying a positional impulse at the contact point. There will be an error that depends on the size of the rotation.
+		const FVec3LP ContactDelta0 = FVec3LP(Body0.DP()) + FVec3LP::CrossProduct(Body0.DQ(), RelativeContactPosition0);
+		const FVec3LP ContactDelta1 = FVec3LP(Body1.DP()) + FVec3LP::CrossProduct(Body1.DQ(), RelativeContactPosition1);
+		const FVec3LP ContactDelta = ContactDelta0 - ContactDelta1;
+		OutContactDeltaNormal = WorldContactDeltaNormal + FVec3LP::DotProduct(ContactDelta, WorldContactNormal);
+
+		// NOTE: OutContactDeltaNormal is negative for penetration
+		// NOTE: MaxPushOut == 0 disables the pushout limits
+		if ((MaxPushOut > 0) && (OutContactDeltaNormal < -MaxPushOut))
+		{
+			const FRealLP ClampedContactDeltaNormal = -MaxPushOut;
+			OutContactDeltaNormal = ClampedContactDeltaNormal;
+		}
+	}
 
-	FORCEINLINE_DEBUGGABLE void FPBDCollisionSolverManifoldPoint::CalculateContactPositionError(const FConstraintSolverBody& Body0, const FConstraintSolverBody& Body1, const FRealLP MaxPushOut, FVec3LP& OutContactDelta, FRealLP& OutContactDeltaNormal) const
+	FORCEINLINE_DEBUGGABLE void FPBDCollisionSolverManifoldPoint::CalculateContactPositionError(const FConstraintSolverBody& Body0, const FConstraintSolverBody& Body1, const FRealLP MaxPushOut, FRealLP& OutContactDeltaNormal, FRealLP& OutContactDeltaTangent0, FRealLP& OutContactDeltaTangent1) const
 	{
 #if CHAOS_NONLINEAR_COLLISIONS_ENABLED
+#error "Non-linear collision solver no longer functional - fix it"
 		// Non-linear version: calculate the contact delta after we have converted the current positional impulses into position and rotation corrections.
 		// We could precalculate and store the LocalContactPositions if we really want to use this nonlinear version
-		const FVec3 LocalContactPosition0 = Body0.Q().Inverse() * WorldRelativeContactPosition0;
-		const FVec3 LocalContactPosition1 = Body1.Q().Inverse() * WorldRelativeContactPosition1;
-		OutContactDelta = (Body0.CorrectedP() + Body0.CorrectedQ() * LocalContactPosition0) - (Body1.CorrectedP() + Body1.CorrectedQ() * LocalContactPosition1);
+		const FVec3LP LocalContactPosition0 = Body0.Q().Inverse() * WorldRelativeContactPosition0;
+		const FVec3LP LocalContactPosition1 = Body1.Q().Inverse() * WorldRelativeContactPosition1;
+		const FVec3LP ContactDelta = (Body0.CorrectedP() + Body0.CorrectedQ() * LocalContactPosition0) - (Body1.CorrectedP() + Body1.CorrectedQ() * LocalContactPosition1);
 #else
 		// Linear version: calculate the contact delta assuming linear motion after applying a positional impulse at the contact point. There will be an error that depends on the size of the rotation.
 		const FVec3LP ContactDelta0 = FVec3LP(Body0.DP()) + FVec3LP::CrossProduct(Body0.DQ(), RelativeContactPosition0);
 		const FVec3LP ContactDelta1 = FVec3LP(Body1.DP()) + FVec3LP::CrossProduct(Body1.DQ(), RelativeContactPosition1);
-		OutContactDelta = WorldContactDelta + ContactDelta0 - ContactDelta1;
+		const FVec3LP ContactDelta = ContactDelta0 - ContactDelta1;
+		OutContactDeltaNormal = WorldContactDeltaNormal + FVec3LP::DotProduct(ContactDelta, WorldContactNormal);
+		OutContactDeltaTangent0 = WorldContactDeltaTangent0 + FVec3LP::DotProduct(ContactDelta, WorldContactTangentU);
+		OutContactDeltaTangent1 = WorldContactDeltaTangent1 + FVec3LP::DotProduct(ContactDelta, WorldContactTangentV);
 #endif
-		OutContactDeltaNormal = FVec3LP::DotProduct(OutContactDelta, WorldContactNormal);
 
 		// NOTE: OutContactDeltaNormal is negative for penetration
 		// NOTE: MaxPushOut == 0 disables the pushout limits
 		if ((MaxPushOut > 0) && (OutContactDeltaNormal < -MaxPushOut))
 		{
-			const FRealLP ClampedContactDeltaNormal = -MaxPushOut;
-			OutContactDelta += (ClampedContactDeltaNormal - OutContactDeltaNormal) * WorldContactNormal;
-			OutContactDeltaNormal = ClampedContactDeltaNormal;
+			OutContactDeltaNormal = -MaxPushOut;
 		}
 	}
 
-	FORCEINLINE_DEBUGGABLE void FPBDCollisionSolverManifoldPoint::CalculateContactVelocityError(const FConstraintSolverBody& Body0, const FConstraintSolverBody& Body1, const FRealLP DynamicFriction, const FRealLP Dt, FVec3LP& OutContactVelocityDelta, FRealLP& OutContactVelocityDeltaNormal) const
+	FORCEINLINE_DEBUGGABLE void FPBDCollisionSolverManifoldPoint::CalculateContactVelocityError(const FConstraintSolverBody& Body0, const FConstraintSolverBody& Body1, const FRealLP DynamicFriction, const FRealLP Dt, FRealLP& OutContactVelocityDeltaNormal, FRealLP& OutContactVelocityDeltaTangent0, FRealLP& OutContactVelocityDeltaTangent1) const
 	{
 		const FVec3LP ContactVelocity0 = Body0.V() + FVec3LP::CrossProduct(Body0.W(), RelativeContactPosition0);
 		const FVec3LP ContactVelocity1 = Body1.V() + FVec3LP::CrossProduct(Body1.W(), RelativeContactPosition1);
@@ -641,24 +617,29 @@ namespace Chaos
 
 		// Add up the errors in the velocity (current velocity - desired velocity)
 		OutContactVelocityDeltaNormal = (ContactVelocityNormal - WorldContactVelocityTargetNormal);
-		OutContactVelocityDelta = (ContactVelocityNormal - ContactVelocityTargetNormal) * WorldContactNormal;
+		OutContactVelocityDeltaTangent0 = FRealLP(0);
+		OutContactVelocityDeltaTangent1 = FRealLP(0);
 
 		const bool bApplyFriction = (DynamicFriction > 0) && (Dt > 0);
 		if (bApplyFriction)
 		{
-			const FVec3LP ContactVelocityTangential = ContactVelocity - ContactVelocityNormal * WorldContactNormal;
-			const FRealLP ContactVelocityTangentialLen = ContactVelocityTangential.Size();
-			if (ContactVelocityTangentialLen > SMALL_NUMBER)
+			const FRealLP ContactVelocityTangent0 = FVec3LP::DotProduct(ContactVelocity, WorldContactTangentU);
+			const FRealLP ContactVelocityTangent1 = FVec3LP::DotProduct(ContactVelocity, WorldContactTangentV);
+			OutContactVelocityDeltaTangent0 = ContactVelocityTangent0;
+			OutContactVelocityDeltaTangent1 = ContactVelocityTangent1;
+
+			const FRealLP ContactVelocityTangentialSq = FMath::Square(ContactVelocityTangent0) + FMath::Square(ContactVelocityTangent1);
+			if (ContactVelocityTangentialSq > FRealLP(SMALL_NUMBER))
 			{
 				// PushOut = ContactMass * DP, where DP is the contact positional correction
 				// Friction force is proportional to the normal force, so friction velocity correction
 				// is proprtional to normal velocity correction, or DVn = DPn/dt = PushOut.N / (ContactMass * dt);
-				const FRealLP PushOutNormal = FVec3LP::DotProduct(NetPushOut, WorldContactNormal);
-				const FRealLP DynamicFrictionVelocityError = PushOutNormal / (WorldContactMassNormal * Dt);
-				if (DynamicFrictionVelocityError > SMALL_NUMBER)
+				const FRealLP MaxContactVelocityTangential = DynamicFriction * NetPushOutNormal / (ContactMassNormal * Dt);
+				if (ContactVelocityTangentialSq > FMath::Square(MaxContactVelocityTangential))
 				{
-					const FRealLP ContactVelocityErrorTangential = FMath::Min(DynamicFriction * DynamicFrictionVelocityError, ContactVelocityTangentialLen);
-					OutContactVelocityDelta += ContactVelocityTangential * (ContactVelocityErrorTangential / ContactVelocityTangentialLen);
+					const FRealLP ContactVelocityTangentialScale = MaxContactVelocityTangential * FMath::InvSqrt(ContactVelocityTangentialSq);
+					OutContactVelocityDeltaTangent0 *= ContactVelocityTangentialScale;
+					OutContactVelocityDeltaTangent1 *= ContactVelocityTangentialScale;
 				}
 			}
 		}
@@ -682,22 +663,37 @@ namespace Chaos
 		{
 			FPBDCollisionSolverManifoldPoint& SolverManifoldPoint = State.ManifoldPoints[PointIndex];
 
-			FVec3LP ContactDelta;
-			FRealLP ContactDeltaNormal;
-			SolverManifoldPoint.CalculateContactPositionError(Body0.SolverBody(), Body1.SolverBody(), MaxPushOut, ContactDelta, ContactDeltaNormal);
+			FRealLP ContactDeltaNormal, ContactDeltaTangent0, ContactDeltaTangent1;
+			SolverManifoldPoint.CalculateContactPositionError(Body0.SolverBody(), Body1.SolverBody(), MaxPushOut, ContactDeltaNormal, ContactDeltaTangent0, ContactDeltaTangent1);
 
-			const bool bProcessManifoldPoint = (ContactDeltaNormal < FRealLP(0)) || !SolverManifoldPoint.NetPushOut.IsNearlyZero();
+			const bool bProcessManifoldPoint = (ContactDeltaNormal < FRealLP(0)) || !FMath::IsNearlyZero(SolverManifoldPoint.NetPushOutNormal);
 			if (bProcessManifoldPoint)
 			{
-				ApplyPositionCorrectionWithFriction(
-					State.Stiffness,
-					State.StaticFriction,
-					State.DynamicFriction,
-					ContactDelta,
-					ContactDeltaNormal,
-					SolverManifoldPoint,
-					Body0,
-					Body1);
+				// If we want to add static friction...
+				// (note: we run a few iterations without friction by temporarily setting StaticFriction to 0)
+				const bool bApplyFriction = (State.StaticFriction > 0) && SolverManifoldPoint.bInsideStaticFrictionCone;
+				if (bApplyFriction)
+				{
+					ApplyPositionCorrectionWithFriction(
+						State.Stiffness,
+						State.StaticFriction,
+						State.DynamicFriction,
+						ContactDeltaNormal,
+						ContactDeltaTangent0,
+						ContactDeltaTangent1,
+						SolverManifoldPoint,
+						Body0,
+						Body1);
+				}
+				else
+				{
+					ApplyPositionCorrectionNoFriction(
+						State.Stiffness, 
+						ContactDeltaNormal, 
+						SolverManifoldPoint, 
+						Body0, 
+						Body1);
+				}
 			}
 		}
 
@@ -715,11 +711,10 @@ namespace Chaos
 		{
 			FPBDCollisionSolverManifoldPoint& SolverManifoldPoint = State.ManifoldPoints[PointIndex];
 
-			FVec3LP ContactDelta;
 			FRealLP ContactDeltaNormal;
-			SolverManifoldPoint.CalculateContactPositionError(Body0.SolverBody(), Body1.SolverBody(), MaxPushOut, ContactDelta, ContactDeltaNormal);
+			SolverManifoldPoint.CalculateContactPositionErrorNormal(Body0.SolverBody(), Body1.SolverBody(), MaxPushOut, ContactDeltaNormal);
 
-			const bool bProcessManifoldPoint = (ContactDeltaNormal < FRealLP(0)) || !SolverManifoldPoint.NetPushOut.IsNearlyZero();
+			const bool bProcessManifoldPoint = (ContactDeltaNormal < FRealLP(0)) || !FMath::IsNearlyZero(SolverManifoldPoint.NetPushOutNormal);
 			if (bProcessManifoldPoint)
 			{
 				ApplyPositionCorrectionNoFriction(
diff --git a/Engine/Source/Runtime/Experimental/Chaos/Private/Chaos/Collision/ParticlePairMidPhase.cpp b/Engine/Source/Runtime/Experimental/Chaos/Private/Chaos/Collision/ParticlePairMidPhase.cpp
--- a/Engine/Source/Runtime/Experimental/Chaos/Private/Chaos/Collision/ParticlePairMidPhase.cpp
+++ b/Engine/Source/Runtime/Experimental/Chaos/Private/Chaos/Collision/ParticlePairMidPhase.cpp
@@ -287,7 +287,7 @@ namespace Chaos
 			// Otherwise we will try to reuse it below
 			if (!bWasUpdatedLastTick || bUseCCD)
 			{
-				// Clear all manifold data
+				// Clear all manifold data including saved contact data
 				Constraint->ResetManifold();
 			}
 
diff --git a/Engine/Source/Runtime/Experimental/Chaos/Private/Chaos/Collision/SolverCollisionContainer.cpp b/Engine/Source/Runtime/Experimental/Chaos/Private/Chaos/Collision/SolverCollisionContainer.cpp
--- a/Engine/Source/Runtime/Experimental/Chaos/Private/Chaos/Collision/SolverCollisionContainer.cpp
+++ b/Engine/Source/Runtime/Experimental/Chaos/Private/Chaos/Collision/SolverCollisionContainer.cpp
@@ -133,11 +133,22 @@ namespace Chaos
 			for (int32 PointIndex = 0; PointIndex < Solver.NumManifoldPoints(); ++PointIndex)
 			{
 				const FPBDCollisionSolverManifoldPoint& SolverManifoldPoint = Solver.GetManifoldPoint(PointIndex);
+
+				const FVec3LP NetPushOut = 
+					SolverManifoldPoint.NetPushOutNormal * SolverManifoldPoint.WorldContactNormal +
+					SolverManifoldPoint.NetPushOutTangentU * SolverManifoldPoint.WorldContactTangentU +
+					SolverManifoldPoint.NetPushOutTangentV * SolverManifoldPoint.WorldContactTangentV;
+
+				const FVec3LP NetImpulse =
+					SolverManifoldPoint.NetImpulseNormal * SolverManifoldPoint.WorldContactNormal +
+					SolverManifoldPoint.NetImpulseTangentU * SolverManifoldPoint.WorldContactTangentU +
+					SolverManifoldPoint.NetImpulseTangentV * SolverManifoldPoint.WorldContactTangentV;
+
 				Constraint->SetSolverResults(PointIndex, 
-					SolverManifoldPoint.NetPushOut, 
-					SolverManifoldPoint.NetImpulse, 
+					NetPushOut, 
+					NetImpulse, 
 					SolverManifoldPoint.bInsideStaticFrictionCone,
-					SolverManifoldPoint.StaticFrictionMax,
+					SolverManifoldPoint.FrictionNetPushOutNormal,
 					Dt);
 			}
 
diff --git a/Engine/Source/Runtime/Experimental/Chaos/Private/Chaos/CollisionOneShotManifoldsMiscShapes.cpp b/Engine/Source/Runtime/Experimental/Chaos/Private/Chaos/CollisionOneShotManifoldsMiscShapes.cpp
--- a/Engine/Source/Runtime/Experimental/Chaos/Private/Chaos/CollisionOneShotManifoldsMiscShapes.cpp
+++ b/Engine/Source/Runtime/Experimental/Chaos/Private/Chaos/CollisionOneShotManifoldsMiscShapes.cpp
@@ -41,7 +41,7 @@ namespace Chaos
 			ensure(SphereBTransform.GetScale3D() == FVec3(1.0f, 1.0f, 1.0f));
 
 			// @todo(chaos): support manifold maintenance
-			Constraint.ResetManifold();
+			Constraint.ResetActiveManifoldContacts();
 
 			FContactPoint ContactPoint = SphereSphereContactPoint(SphereA, SphereATransform, SphereB, SphereBTransform, Constraint.GetCullDistance(), Constraint.Manifold.RestitutionPadding);
 			if (ContactPoint.Phi < Constraint.GetCullDistance())
@@ -64,7 +64,7 @@ namespace Chaos
 			ensure(PlaneTransform.GetScale3D() == FVec3(1.0f, 1.0f, 1.0f));
 
 			// @todo(chaos): support manifold maintenance
-			Constraint.ResetManifold();
+			Constraint.ResetActiveManifoldContacts();
 
 			FContactPoint ContactPoint = SpherePlaneContactPoint(Sphere, SphereTransform, Plane, PlaneTransform, Constraint.Manifold.RestitutionPadding);
 			if (ContactPoint.Phi < Constraint.GetCullDistance())
@@ -81,7 +81,7 @@ namespace Chaos
 			ensure(BoxTransform.GetScale3D() == FVec3(1.0f, 1.0f, 1.0f));
 
 			// @todo(chaos): support manifold maintenance
-			Constraint.ResetManifold();
+			Constraint.ResetActiveManifoldContacts();
 
 			FContactPoint ContactPoint = SphereBoxContactPoint(Sphere, SphereTransform, Box, BoxTransform, Constraint.Manifold.RestitutionPadding);
 			if (ContactPoint.Phi < Constraint.GetCullDistance())
@@ -98,7 +98,7 @@ namespace Chaos
 			ensure(CapsuleTransform.GetScale3D() == FVec3(1.0f, 1.0f, 1.0f));
 
 			// @todo(chaos): support manifold maintenance
-			Constraint.ResetManifold();
+			Constraint.ResetActiveManifoldContacts();
 
 			FContactPoint ContactPoint = SphereCapsuleContactPoint(Sphere, SphereTransform, Capsule, CapsuleTransform, Constraint.Manifold.RestitutionPadding);
 			if (ContactPoint.Phi < Constraint.GetCullDistance())
@@ -115,7 +115,7 @@ namespace Chaos
 			ensure(ConvexTransform.GetScale3D() == FVec3(1.0f, 1.0f, 1.0f));
 
 			// @todo(chaos): support manifold maintenance
-			Constraint.ResetManifold();
+			Constraint.ResetActiveManifoldContacts();
 
 			FContactPoint ContactPoint = SphereConvexContactPoint(Sphere, SphereTransform, Convex, ConvexTransform);
 			if (ContactPoint.Phi < Constraint.GetCullDistance())
@@ -133,7 +133,7 @@ namespace Chaos
 			ensure(TriMeshWorldTransform.GetScale3D() == FVec3(1.0f, 1.0f, 1.0f));
 
 			// @todo(chaos): support manifold maintenance
-			Constraint.ResetManifold();
+			Constraint.ResetActiveManifoldContacts();
 
 			FContactPoint ContactPoint = SphereTriangleMeshContactPoint(Sphere, SphereWorldTransform, TriangleMesh, TriMeshWorldTransform, Constraint.GetCullDistance(), 0.0f);
 			if (ContactPoint.Phi < Constraint.GetCullDistance())
@@ -166,7 +166,7 @@ namespace Chaos
 			ensure(CapsuleBTransform.GetScale3D() == FVec3(1.0f, 1.0f, 1.0f));
 
 			// @todo(chaos): support manifold maintenance
-			Constraint.ResetManifold();
+			Constraint.ResetActiveManifoldContacts();
 
 			FVec3 CapsuleADirection(CapsuleATransform.TransformVector(CapsuleA.GetSegment().GetAxis()));
 			const FVec3 CapsuleBDirection(CapsuleBTransform.TransformVector(CapsuleB.GetSegment().GetAxis()));
@@ -272,7 +272,7 @@ namespace Chaos
 			ensure(TriMeshWorldTransform.GetScale3D() == FVec3(1.0f, 1.0f, 1.0f));
 
 			// @todo(chaos): support manifold maintenance
-			Constraint.ResetManifold();
+			Constraint.ResetActiveManifoldContacts();
 
 			TArray<FContactPoint> ContactPoints;
 			GJKImplicitManifold<FCapsule, TriMeshType>(Capsule, CapsuleWorldTransform, TriangleMesh, TriMeshWorldTransform, Constraint.GetCullDistance(), 0.0f, ContactPoints);
@@ -293,7 +293,7 @@ namespace Chaos
 			ensure(HeightFieldTransform.GetScale3D() == FVec3(1.0f, 1.0f, 1.0f));
 
 			// @todo(chaos): support manifold maintenance
-			Constraint.ResetManifold();
+			Constraint.ResetActiveManifoldContacts();
 
 			TArray<FContactPoint> ContactPoints;
 			GJKImplicitManifold<FCapsule>(Capsule, CapsuleTransform, HeightField, HeightFieldTransform, Constraint.GetCullDistance(), 0.0f, ContactPoints);
@@ -315,7 +315,7 @@ namespace Chaos
 			ensure(HeightFieldTransform.GetScale3D() == FVec3(1.0f, 1.0f, 1.0f));
 
 			// @todo(chaos): support manifold maintenance
-			Constraint.ResetManifold();
+			Constraint.ResetActiveManifoldContacts();
 
 			//FContactPoint ContactPoint = ConvexHeightFieldContactPoint(Convex, ConvexTransform, HeightField, HeightFieldTransform, Constraint.GetCullDistance(), 0.0f);
 
@@ -339,7 +339,7 @@ namespace Chaos
 			ensure(TriMeshTransform.GetScale3D() == FVec3(1.0f, 1.0f, 1.0f));
 
 			// @todo(chaos): support manifold maintenance
-			Constraint.ResetManifold();
+			Constraint.ResetActiveManifoldContacts();
 
 			TArray<FContactPoint> ContactPoints;
 			GJKImplicitManifold<ConvexType, TriMeshType>(Convex, ConvexTransform, TriangleMesh, TriMeshTransform, Constraint.GetCullDistance(), 0.0f, ContactPoints);
diff --git a/Engine/Source/Runtime/Experimental/Chaos/Private/Chaos/CollisionResolution.cpp b/Engine/Source/Runtime/Experimental/Chaos/Private/Chaos/CollisionResolution.cpp
--- a/Engine/Source/Runtime/Experimental/Chaos/Private/Chaos/CollisionResolution.cpp
+++ b/Engine/Source/Runtime/Experimental/Chaos/Private/Chaos/CollisionResolution.cpp
@@ -1615,7 +1615,7 @@ namespace Chaos
 			FPBDCollisionConstraint* Constraint = Context.CollisionAllocator->FindOrCreateConstraint(Particle0, Implicit0, Simplicial0, LocalTransform0, Particle1, Implicit1, Simplicial1, LocalTransform1, CullDistance, EContactShapesType::LevelSetLevelSet, bAllowManifolds);
 
 			// @todo(chaos): support manifold restoration for Levelsets and remove this
-			Constraint->ResetManifold();
+			Constraint->ResetActiveManifoldContacts();
 
 			// @todo(chaos): the order test should be done befor we create the constraint so we can do it right first time, rather than patching it after
 			bool bIsParticleDynamic0 = Particle0->CastToRigidParticle() && Particle0->ObjectState() == EObjectStateType::Dynamic;
