Chaos - rename FRealLP etc to FSolverReal etc

diff --git a/Engine/Source/Runtime/Experimental/Chaos/Private/Chaos/Collision/PBDCollisionSolver.cpp b/Engine/Source/Runtime/Experimental/Chaos/Private/Chaos/Collision/PBDCollisionSolver.cpp
--- a/Engine/Source/Runtime/Experimental/Chaos/Private/Chaos/Collision/PBDCollisionSolver.cpp
+++ b/Engine/Source/Runtime/Experimental/Chaos/Private/Chaos/Collision/PBDCollisionSolver.cpp
@@ -65,14 +65,14 @@ namespace Chaos
 		const FConstraintSolverBody& Body1,
 		const FVec3& InWorldAnchorPoint0,
 		const FVec3& InWorldAnchorPoint1,
-		const FVec3LP& InWorldContactNormal)
+		const FSolverVec3& InWorldContactNormal)
 	{
-		NetPushOutNormal = FRealLP(0);
-		NetPushOutTangentU = FRealLP(0);
-		NetPushOutTangentV = FRealLP(0);
-		NetImpulseNormal = FRealLP(0);
-		NetImpulseTangentU = FRealLP(0);
-		NetImpulseTangentV = FRealLP(0);
+		NetPushOutNormal = FSolverReal(0);
+		NetPushOutTangentU = FSolverReal(0);
+		NetPushOutTangentV = FSolverReal(0);
+		NetImpulseNormal = FSolverReal(0);
+		NetImpulseTangentU = FSolverReal(0);
+		NetImpulseTangentV = FSolverReal(0);
 
 		UpdateContact(Body0, Body1, InWorldAnchorPoint0, InWorldAnchorPoint1, InWorldContactNormal);
 
@@ -82,19 +82,19 @@ namespace Chaos
 	void FPBDCollisionSolverManifoldPoint::InitMaterial(
 		const FConstraintSolverBody& Body0,
 		const FConstraintSolverBody& Body1,
-		const FRealLP InRestitution,
-		const FRealLP InRestitutionVelocityThreshold,
+		const FSolverReal InRestitution,
+		const FSolverReal InRestitutionVelocityThreshold,
 		const bool bInEnableStaticFriction,
-		const FRealLP InFrictionNetPushOutNormal)
+		const FSolverReal InFrictionNetPushOutNormal)
 	{
 		FrictionNetPushOutNormal = InFrictionNetPushOutNormal;
 		bInsideStaticFrictionCone = bInEnableStaticFriction;
-		WorldContactVelocityTargetNormal = FRealLP(0);
+		WorldContactVelocityTargetNormal = FSolverReal(0);
 
-		if (InRestitution > FRealLP(0))
+		if (InRestitution > FSolverReal(0))
 		{
-			const FVec3LP ContactVelocity = CalculateContactVelocity(Body0, Body1);
-			const FRealLP ContactVelocityNormal = FVec3LP::DotProduct(ContactVelocity, WorldContactNormal);
+			const FSolverVec3 ContactVelocity = CalculateContactVelocity(Body0, Body1);
+			const FSolverReal ContactVelocityNormal = FSolverVec3::DotProduct(ContactVelocity, WorldContactNormal);
 			if (ContactVelocityNormal < -InRestitutionVelocityThreshold)
 			{
 				WorldContactVelocityTargetNormal = -InRestitution * ContactVelocityNormal;
@@ -108,89 +108,89 @@ namespace Chaos
 		const FConstraintSolverBody& Body1,
 		const FVec3& InWorldAnchorPoint0,
 		const FVec3& InWorldAnchorPoint1, 
-		const FVec3LP& InWorldContactNormal)
+		const FSolverVec3& InWorldContactNormal)
 	{
 		// The world-space point where we apply impulses/corrections (same world-space point for momentum conservation)
 		const FVec3 WorldContactPosition = FReal(0.5) * (InWorldAnchorPoint0 + InWorldAnchorPoint1);
-		RelativeContactPosition0 = FVec3LP(WorldContactPosition - Body0.P());
-		RelativeContactPosition1 = FVec3LP(WorldContactPosition - Body1.P());
+		RelativeContactPosition0 = FSolverVec3(WorldContactPosition - Body0.P());
+		RelativeContactPosition1 = FSolverVec3(WorldContactPosition - Body1.P());
 
 		// The world-space contact normal
 		WorldContactNormal = InWorldContactNormal;
 
 		// World-space contact tangents. We are treating the normal as the constraint-space Z axis
 		// and the Tangent U and V as the constraint-space X and Y axes resepctively
-		WorldContactTangentU = FVec3LP::CrossProduct(FVec3LP(0, 1, 0), InWorldContactNormal);
-		if (!WorldContactTangentU.Normalize(FRealLP(KINDA_SMALL_NUMBER)))
+		WorldContactTangentU = FSolverVec3::CrossProduct(FSolverVec3(0, 1, 0), InWorldContactNormal);
+		if (!WorldContactTangentU.Normalize(FSolverReal(KINDA_SMALL_NUMBER)))
 		{
-			WorldContactTangentU = FVec3LP::CrossProduct(FVec3LP(1, 0, 0), InWorldContactNormal);
+			WorldContactTangentU = FSolverVec3::CrossProduct(FSolverVec3(1, 0, 0), InWorldContactNormal);
 			WorldContactTangentU = WorldContactTangentU.GetUnsafeNormal();
 		}
-		WorldContactTangentV = FVec3LP::CrossProduct(InWorldContactNormal, WorldContactTangentU);
+		WorldContactTangentV = FSolverVec3::CrossProduct(InWorldContactNormal, WorldContactTangentU);
 
 		// The contact point error we are trying to correct in this solver
-		const FVec3LP WorldContactDelta = InWorldAnchorPoint0 - InWorldAnchorPoint1;
-		WorldContactDeltaNormal = FVec3LP::DotProduct(WorldContactDelta, WorldContactNormal);
-		WorldContactDeltaTangent0 = FVec3LP::DotProduct(WorldContactDelta, WorldContactTangentU);
-		WorldContactDeltaTangent1 = FVec3LP::DotProduct(WorldContactDelta, WorldContactTangentV);
+		const FSolverVec3 WorldContactDelta = InWorldAnchorPoint0 - InWorldAnchorPoint1;
+		WorldContactDeltaNormal = FSolverVec3::DotProduct(WorldContactDelta, WorldContactNormal);
+		WorldContactDeltaTangent0 = FSolverVec3::DotProduct(WorldContactDelta, WorldContactTangentU);
+		WorldContactDeltaTangent1 = FSolverVec3::DotProduct(WorldContactDelta, WorldContactTangentV);
 	}
 
 	inline void FPBDCollisionSolverManifoldPoint::UpdateMass(const FConstraintSolverBody& Body0, const FConstraintSolverBody& Body1)
 	{
-		FRealLP ContactMassInvNormal = FRealLP(0);
-		FRealLP ContactMassInvTangentU = FRealLP(0);
-		FRealLP ContactMassInvTangentV = FRealLP(0);
+		FSolverReal ContactMassInvNormal = FSolverReal(0);
+		FSolverReal ContactMassInvTangentU = FSolverReal(0);
+		FSolverReal ContactMassInvTangentV = FSolverReal(0);
 
 		// These are not used if not initialized below so no need to clear
-		//WorldContactNormalAngular0 = FVec3LP(0);
-		//WorldContactTangentUAngular0 = FVec3LP(0);
-		//WorldContactTangentVAngular0 = FVec3LP(0);
-		//WorldContactNormalAngular1 = FVec3LP(0);
-		//WorldContactTangentUAngular1 = FVec3LP(0);
-		//WorldContactTangentVAngular1 = FVec3LP(0);
+		//WorldContactNormalAngular0 = FSolverVec3(0);
+		//WorldContactTangentUAngular0 = FSolverVec3(0);
+		//WorldContactTangentVAngular0 = FSolverVec3(0);
+		//WorldContactNormalAngular1 = FSolverVec3(0);
+		//WorldContactTangentUAngular1 = FSolverVec3(0);
+		//WorldContactTangentVAngular1 = FSolverVec3(0);
 
 		if (Body0.IsDynamic())
 		{
-			const FVec3LP R0xN = FVec3LP::CrossProduct(RelativeContactPosition0, WorldContactNormal);
-			const FVec3LP R0xU = FVec3LP::CrossProduct(RelativeContactPosition0, WorldContactTangentU);
-			const FVec3LP R0xV = FVec3LP::CrossProduct(RelativeContactPosition0, WorldContactTangentV);
+			const FSolverVec3 R0xN = FSolverVec3::CrossProduct(RelativeContactPosition0, WorldContactNormal);
+			const FSolverVec3 R0xU = FSolverVec3::CrossProduct(RelativeContactPosition0, WorldContactTangentU);
+			const FSolverVec3 R0xV = FSolverVec3::CrossProduct(RelativeContactPosition0, WorldContactTangentV);
 
-			const FMatrix33LP InvI0 = Body0.InvI();
+			const FSolverMatrix33 InvI0 = Body0.InvI();
 
 			WorldContactNormalAngular0 = InvI0 * R0xN;
 			WorldContactTangentUAngular0 = InvI0 * R0xU;
 			WorldContactTangentVAngular0 = InvI0 * R0xV;
 
-			ContactMassInvNormal += FVec3LP::DotProduct(R0xN, WorldContactNormalAngular0) + Body0.InvM();
-			ContactMassInvTangentU += FVec3LP::DotProduct(R0xU, WorldContactTangentUAngular0) + Body0.InvM();
-			ContactMassInvTangentV += FVec3LP::DotProduct(R0xV, WorldContactTangentVAngular0) + Body0.InvM();
+			ContactMassInvNormal += FSolverVec3::DotProduct(R0xN, WorldContactNormalAngular0) + Body0.InvM();
+			ContactMassInvTangentU += FSolverVec3::DotProduct(R0xU, WorldContactTangentUAngular0) + Body0.InvM();
+			ContactMassInvTangentV += FSolverVec3::DotProduct(R0xV, WorldContactTangentVAngular0) + Body0.InvM();
 		}
 		if (Body1.IsDynamic())
 		{
-			const FVec3LP R1xN = FVec3LP::CrossProduct(RelativeContactPosition1, WorldContactNormal);
-			const FVec3LP R1xU = FVec3LP::CrossProduct(RelativeContactPosition1, WorldContactTangentU);
-			const FVec3LP R1xV = FVec3LP::CrossProduct(RelativeContactPosition1, WorldContactTangentV);
+			const FSolverVec3 R1xN = FSolverVec3::CrossProduct(RelativeContactPosition1, WorldContactNormal);
+			const FSolverVec3 R1xU = FSolverVec3::CrossProduct(RelativeContactPosition1, WorldContactTangentU);
+			const FSolverVec3 R1xV = FSolverVec3::CrossProduct(RelativeContactPosition1, WorldContactTangentV);
 
-			const FMatrix33LP InvI1 = Body1.InvI();
+			const FSolverMatrix33 InvI1 = Body1.InvI();
 
 			WorldContactNormalAngular1 = InvI1 * R1xN;
 			WorldContactTangentUAngular1 = InvI1 * R1xU;
 			WorldContactTangentVAngular1 = InvI1 * R1xV;
 
-			ContactMassInvNormal += FVec3LP::DotProduct(R1xN, WorldContactNormalAngular1) + Body1.InvM();
-			ContactMassInvTangentU += FVec3LP::DotProduct(R1xU, WorldContactTangentUAngular1) + Body1.InvM();
-			ContactMassInvTangentV += FVec3LP::DotProduct(R1xV, WorldContactTangentVAngular1) + Body1.InvM();
+			ContactMassInvNormal += FSolverVec3::DotProduct(R1xN, WorldContactNormalAngular1) + Body1.InvM();
+			ContactMassInvTangentU += FSolverVec3::DotProduct(R1xU, WorldContactTangentUAngular1) + Body1.InvM();
+			ContactMassInvTangentV += FSolverVec3::DotProduct(R1xV, WorldContactTangentVAngular1) + Body1.InvM();
 		}
 
-		ContactMassNormal = (ContactMassInvNormal > FRealLP(SMALL_NUMBER)) ? FRealLP(1) / ContactMassInvNormal : FRealLP(0);
-		ContactMassTangentU = (ContactMassInvTangentU > FRealLP(SMALL_NUMBER)) ? FRealLP(1) / ContactMassInvTangentU : FRealLP(0);
-		ContactMassTangentV = (ContactMassInvTangentV > FRealLP(SMALL_NUMBER)) ? FRealLP(1) / ContactMassInvTangentV : FRealLP(0);
+		ContactMassNormal = (ContactMassInvNormal > FSolverReal(SMALL_NUMBER)) ? FSolverReal(1) / ContactMassInvNormal : FSolverReal(0);
+		ContactMassTangentU = (ContactMassInvTangentU > FSolverReal(SMALL_NUMBER)) ? FSolverReal(1) / ContactMassInvTangentU : FSolverReal(0);
+		ContactMassTangentV = (ContactMassInvTangentV > FSolverReal(SMALL_NUMBER)) ? FSolverReal(1) / ContactMassInvTangentV : FSolverReal(0);
 	}
 
-	FVec3LP FPBDCollisionSolverManifoldPoint::CalculateContactVelocity(const FConstraintSolverBody& Body0, const FConstraintSolverBody& Body1) const
+	FSolverVec3 FPBDCollisionSolverManifoldPoint::CalculateContactVelocity(const FConstraintSolverBody& Body0, const FConstraintSolverBody& Body1) const
 	{
-		const FVec3LP ContactVelocity0 = Body0.V() + FVec3::CrossProduct(Body0.W(), RelativeContactPosition0);
-		const FVec3LP ContactVelocity1 = Body1.V() + FVec3::CrossProduct(Body1.W(), RelativeContactPosition1);
+		const FSolverVec3 ContactVelocity0 = Body0.V() + FVec3::CrossProduct(Body0.W(), RelativeContactPosition0);
+		const FSolverVec3 ContactVelocity1 = Body1.V() + FVec3::CrossProduct(Body1.W(), RelativeContactPosition1);
 		return ContactVelocity0 - ContactVelocity1;
 	}
 
@@ -222,7 +222,7 @@ namespace Chaos
 		SetShockPropagationInvMassScale(FReal(1));
 	}
 
-	void FPBDCollisionSolver::SetShockPropagationInvMassScale(const FRealLP InvMassScale)
+	void FPBDCollisionSolver::SetShockPropagationInvMassScale(const FSolverReal InvMassScale)
 	{
 		FConstraintSolverBody& Body0 = SolverBody0();
 		FConstraintSolverBody& Body1 = SolverBody1();
@@ -274,7 +274,7 @@ namespace Chaos
 			State.SolverBodies[1], 
 			InWorldAnchorPoint0, 
 			InWorldAnchorPoint1, 
-			FVec3LP(InWorldContactNormal));
+			FSolverVec3(InWorldContactNormal));
 	}
 
 	void FPBDCollisionSolver::InitMaterial(
@@ -287,10 +287,10 @@ namespace Chaos
 		State.ManifoldPoints[ManifoldPoiontIndex].InitMaterial(
 			State.SolverBodies[0], 
 			State.SolverBodies[1], 
-			FRealLP(InRestitution), 
-			FRealLP(InRestitutionVelocityThreshold),
+			FSolverReal(InRestitution), 
+			FSolverReal(InRestitutionVelocityThreshold),
 			bInEnableStaticFriction, 
-			FRealLP(InFrictionNetPushOutNormal));
+			FSolverReal(InFrictionNetPushOutNormal));
 
 		// Track if any points have restitution enabled. See SolveVelocity
 		State.bHaveRestitution = State.bHaveRestitution || (State.ManifoldPoints[ManifoldPoiontIndex].WorldContactVelocityTargetNormal > 0);
@@ -307,22 +307,22 @@ namespace Chaos
 			State.SolverBodies[1], 
 			InWorldAnchorPoint0, 
 			InWorldAnchorPoint1, 
-			FVec3LP(InWorldContactNormal));
+			FSolverVec3(InWorldContactNormal));
 	}
 
 
-	void FPBDCollisionSolver::SolveVelocityAverage(const FRealLP Dt)
+	void FPBDCollisionSolver::SolveVelocityAverage(const FSolverReal Dt)
 	{
 		FConstraintSolverBody& Body0 = SolverBody0();
 		FConstraintSolverBody& Body1 = SolverBody1();
 
 		// Generate a new contact point at the average of all the active contacts
 		int32 NumActiveManifoldPoints = 0;
-		FVec3LP RelativeContactPosition0 = FVec3(0);
-		FVec3LP RelativeContactPosition1 = FVec3(0);
-		FVec3LP WorldContactNormal = FVec3LP(0);
-		FRealLP NetPushOutNormal = FRealLP(0);
-		FRealLP WorldContactVelocityTargetNormal = FReal(0);
+		FSolverVec3 RelativeContactPosition0 = FVec3(0);
+		FSolverVec3 RelativeContactPosition1 = FVec3(0);
+		FSolverVec3 WorldContactNormal = FSolverVec3(0);
+		FSolverReal NetPushOutNormal = FSolverReal(0);
+		FSolverReal WorldContactVelocityTargetNormal = FReal(0);
 		for (int32 PointIndex = 0; PointIndex < State.NumManifoldPoints; ++PointIndex)
 		{
 			FPBDCollisionSolverManifoldPoint& SolverManifoldPoint = State.ManifoldPoints[PointIndex];
@@ -345,8 +345,8 @@ namespace Chaos
 		// than we would have if we just solved the corners of the box in sequence.
 		if (NumActiveManifoldPoints > 1)
 		{
-			const FRealLP DynamicFriction = FRealLP(0);
-			const FRealLP Scale = FRealLP(1) / FRealLP(NumActiveManifoldPoints);
+			const FSolverReal DynamicFriction = FSolverReal(0);
+			const FSolverReal Scale = FSolverReal(1) / FSolverReal(NumActiveManifoldPoints);
 
 			FPBDCollisionSolverManifoldPoint AverageManifoldPoint;
 			AverageManifoldPoint.RelativeContactPosition0 = RelativeContactPosition0 * Scale;
@@ -355,31 +355,31 @@ namespace Chaos
 			AverageManifoldPoint.WorldContactVelocityTargetNormal = WorldContactVelocityTargetNormal * Scale;
 			AverageManifoldPoint.NetPushOutNormal = NetPushOutNormal * Scale;
 
-			const FMatrix33LP ContactMassInv =
-				(Body0.IsDynamic() ? Collisions::ComputeFactorMatrix3(AverageManifoldPoint.RelativeContactPosition0, Body0.InvI(), Body0.InvM()) : FMatrix33LP(0)) +
-				(Body1.IsDynamic() ? Collisions::ComputeFactorMatrix3(AverageManifoldPoint.RelativeContactPosition1, Body1.InvI(), Body1.InvM()) : FMatrix33LP(0));
-			const FRealLP ContactMassInvNormal = FVec3LP::DotProduct(AverageManifoldPoint.WorldContactNormal, Utilities::Multiply(ContactMassInv, AverageManifoldPoint.WorldContactNormal));
-			const FRealLP ContactMassNormal = (ContactMassInvNormal > FRealLP(SMALL_NUMBER)) ? FRealLP(1) / ContactMassInvNormal : FRealLP(0);
+			const FSolverMatrix33 ContactMassInv =
+				(Body0.IsDynamic() ? Collisions::ComputeFactorMatrix3(AverageManifoldPoint.RelativeContactPosition0, Body0.InvI(), Body0.InvM()) : FSolverMatrix33(0)) +
+				(Body1.IsDynamic() ? Collisions::ComputeFactorMatrix3(AverageManifoldPoint.RelativeContactPosition1, Body1.InvI(), Body1.InvM()) : FSolverMatrix33(0));
+			const FSolverReal ContactMassInvNormal = FSolverVec3::DotProduct(AverageManifoldPoint.WorldContactNormal, Utilities::Multiply(ContactMassInv, AverageManifoldPoint.WorldContactNormal));
+			const FSolverReal ContactMassNormal = (ContactMassInvNormal > FSolverReal(SMALL_NUMBER)) ? FSolverReal(1) / ContactMassInvNormal : FSolverReal(0);
 			AverageManifoldPoint.ContactMassNormal = ContactMassNormal;
 
 			// @todo(chaos): All unused for the average contact solve - optimize
-			AverageManifoldPoint.ContactMassTangentU = FRealLP(0);
-			AverageManifoldPoint.ContactMassTangentV = FRealLP(0);
-			AverageManifoldPoint.WorldContactTangentU = FVec3LP(0);
-			AverageManifoldPoint.WorldContactTangentV = FVec3LP(0);
-			AverageManifoldPoint.WorldContactDeltaNormal = FRealLP(0);
-			AverageManifoldPoint.WorldContactDeltaTangent0 = FRealLP(0);
-			AverageManifoldPoint.WorldContactDeltaTangent1 = FRealLP(0);
-			AverageManifoldPoint.NetPushOutTangentU = FRealLP(0);
-			AverageManifoldPoint.NetPushOutTangentV = FRealLP(0);
-			AverageManifoldPoint.NetImpulseNormal = FRealLP(0);
-			AverageManifoldPoint.NetImpulseTangentU = FRealLP(0);
-			AverageManifoldPoint.NetImpulseTangentV = FRealLP(0);
-			AverageManifoldPoint.FrictionNetPushOutNormal = FRealLP(0);
+			AverageManifoldPoint.ContactMassTangentU = FSolverReal(0);
+			AverageManifoldPoint.ContactMassTangentV = FSolverReal(0);
+			AverageManifoldPoint.WorldContactTangentU = FSolverVec3(0);
+			AverageManifoldPoint.WorldContactTangentV = FSolverVec3(0);
+			AverageManifoldPoint.WorldContactDeltaNormal = FSolverReal(0);
+			AverageManifoldPoint.WorldContactDeltaTangent0 = FSolverReal(0);
+			AverageManifoldPoint.WorldContactDeltaTangent1 = FSolverReal(0);
+			AverageManifoldPoint.NetPushOutTangentU = FSolverReal(0);
+			AverageManifoldPoint.NetPushOutTangentV = FSolverReal(0);
+			AverageManifoldPoint.NetImpulseNormal = FSolverReal(0);
+			AverageManifoldPoint.NetImpulseTangentU = FSolverReal(0);
+			AverageManifoldPoint.NetImpulseTangentV = FSolverReal(0);
+			AverageManifoldPoint.FrictionNetPushOutNormal = FSolverReal(0);
 			AverageManifoldPoint.bInsideStaticFrictionCone = false;
 
 			// @todo(chaos): don't need tangential velocity delta here since we are setting DynamicFriction to 0
-			FRealLP ContactVelocityDeltaNormal, ContactVelocityDeltaTangentU, ContactVelocityDeltaTangentV;
+			FSolverReal ContactVelocityDeltaNormal, ContactVelocityDeltaTangentU, ContactVelocityDeltaTangentV;
 			AverageManifoldPoint.CalculateContactVelocityError(Body0, Body1, DynamicFriction, Dt, ContactVelocityDeltaNormal, ContactVelocityDeltaTangentU, ContactVelocityDeltaTangentV);
 
 			ApplyVelocityCorrection(
@@ -395,7 +395,7 @@ namespace Chaos
 		}
 	}
 
-	bool FPBDCollisionSolver::SolveVelocity(const FRealLP Dt, const bool bApplyDynamicFriction)
+	bool FPBDCollisionSolver::SolveVelocity(const FSolverReal Dt, const bool bApplyDynamicFriction)
 	{
 		// Apply restitution at the average contact point
 		// This means we don't need to run as many iterations to get stable bouncing
@@ -410,14 +410,14 @@ namespace Chaos
 
 		// NOTE: this dynamic friction implementation is iteration-count sensitive
 		// @todo(chaos): fix iteration count dependence of dynamic friction
-		const FRealLP DynamicFriction = (bApplyDynamicFriction && bChaos_PBDCollisionSolver_Velocity_DynamicFrictionEnabled) ? State.DynamicFriction : FRealLP(0);
+		const FSolverReal DynamicFriction = (bApplyDynamicFriction && bChaos_PBDCollisionSolver_Velocity_DynamicFrictionEnabled) ? State.DynamicFriction : FSolverReal(0);
 
 		for (int32 PointIndex = 0; PointIndex < State.NumManifoldPoints; ++PointIndex)
 		{
 			FPBDCollisionSolverManifoldPoint& SolverManifoldPoint = State.ManifoldPoints[PointIndex];
 			if (!FMath::IsNearlyZero(SolverManifoldPoint.NetPushOutNormal))
 			{
-				FRealLP ContactVelocityDeltaNormal, ContactVelocityDeltaTangentU, ContactVelocityDeltaTangentV;
+				FSolverReal ContactVelocityDeltaNormal, ContactVelocityDeltaTangentU, ContactVelocityDeltaTangentV;
 				SolverManifoldPoint.CalculateContactVelocityError(Body0, Body1, DynamicFriction, Dt, ContactVelocityDeltaNormal, ContactVelocityDeltaTangentU, ContactVelocityDeltaTangentV);
 
 				ApplyVelocityCorrection(
diff --git a/Engine/Source/Runtime/Experimental/Chaos/Private/Chaos/Collision/PBDCollisionSolver.h b/Engine/Source/Runtime/Experimental/Chaos/Private/Chaos/Collision/PBDCollisionSolver.h
--- a/Engine/Source/Runtime/Experimental/Chaos/Private/Chaos/Collision/PBDCollisionSolver.h
+++ b/Engine/Source/Runtime/Experimental/Chaos/Private/Chaos/Collision/PBDCollisionSolver.h
@@ -49,7 +49,7 @@ namespace Chaos
 			const FConstraintSolverBody& Body1,
 			const FVec3& InWorldAnchorPoint0,
 			const FVec3& InWorldAnchorPoint1,
-			const FVec3LP& InWorldContactNormal);
+			const FSolverVec3& InWorldContactNormal);
 
 		/**
 		 * @brief Initialize the material related properties of the contact
@@ -57,10 +57,10 @@ namespace Chaos
 		void InitMaterial(
 			const FConstraintSolverBody& Body0,
 			const FConstraintSolverBody& Body1,
-			const FRealLP InRestitution,
-			const FRealLP InRestitutionVelocityThreshold,
+			const FSolverReal InRestitution,
+			const FSolverReal InRestitutionVelocityThreshold,
 			const bool bInEnableStaticFriction,
-			const FRealLP InFrictionNetPushOutNormal);
+			const FSolverReal InFrictionNetPushOutNormal);
 
 		/**
 		 * @brief Update the world-space relative contact points based on current body transforms and body-space contact positions
@@ -70,7 +70,7 @@ namespace Chaos
 			const FConstraintSolverBody& Body1,
 			const FVec3& InWorldAnchorPoint0,
 			const FVec3& InWorldAnchorPoint1,
-			const FVec3LP& InWorldContactNormal);
+			const FSolverVec3& InWorldContactNormal);
 
 		/**
 		 * @brief Update the cached mass properties based on the current body transforms
@@ -81,13 +81,13 @@ namespace Chaos
 		 * @brief Calculate the position error at the current transforms
 		 * @param MaxPushOut a limit on the position error for this iteration to prevent initial-penetration explosion (a common PBD problem)
 		*/
-		void CalculateContactPositionErrorNormal(const FConstraintSolverBody& Body0, const FConstraintSolverBody& Body1, const FRealLP MaxPushOut, FRealLP& OutContactDeltaNormal) const;
-		void CalculateContactPositionError(const FConstraintSolverBody& Body0, const FConstraintSolverBody& Body1, const FRealLP MaxPushOut, FRealLP& OutContactDeltaNormal, FRealLP& OutContactDeltaTanget0, FRealLP& OutContactDeltaTangent1) const;
+		void CalculateContactPositionErrorNormal(const FConstraintSolverBody& Body0, const FConstraintSolverBody& Body1, const FSolverReal MaxPushOut, FSolverReal& OutContactDeltaNormal) const;
+		void CalculateContactPositionError(const FConstraintSolverBody& Body0, const FConstraintSolverBody& Body1, const FSolverReal MaxPushOut, FSolverReal& OutContactDeltaNormal, FSolverReal& OutContactDeltaTanget0, FSolverReal& OutContactDeltaTangent1) const;
 
 		/**
 		 * @brief Calculate the velocity error at the current transforms
 		*/
-		void CalculateContactVelocityError(const FConstraintSolverBody& Body0, const FConstraintSolverBody& Body1, const FRealLP DynamicFriction, const FRealLP Dt, FRealLP& OutContactVelocityDeltaNormal, FRealLP& OutContactVelocityDeltaTangent0, FRealLP& OutContactVelocityDeltaTangent1) const;
+		void CalculateContactVelocityError(const FConstraintSolverBody& Body0, const FConstraintSolverBody& Body1, const FSolverReal DynamicFriction, const FSolverReal Dt, FSolverReal& OutContactVelocityDeltaNormal, FSolverReal& OutContactVelocityDeltaTangent0, FSolverReal& OutContactVelocityDeltaTangent1) const;
 
 		// @todo(chaos): make private
 	public:
@@ -97,48 +97,48 @@ namespace Chaos
 		 * @brief Calculate the relative velocity at the contact point
 		 * @note InitContact must be called before calling this function
 		*/
-		FVec3LP CalculateContactVelocity(const FConstraintSolverBody& Body0, const FConstraintSolverBody& Body1) const;
+		FSolverVec3 CalculateContactVelocity(const FConstraintSolverBody& Body0, const FConstraintSolverBody& Body1) const;
 
 		// World-space body-relative contact points
-		FVec3LP RelativeContactPosition0;
-		FVec3LP RelativeContactPosition1;
+		FSolverVec3 RelativeContactPosition0;
+		FSolverVec3 RelativeContactPosition1;
 
 		// World-space contact axes (normal and 2 tangents)
-		FVec3LP WorldContactNormal;
-		FVec3LP WorldContactTangentU;
-		FVec3LP WorldContactTangentV;
+		FSolverVec3 WorldContactNormal;
+		FSolverVec3 WorldContactTangentU;
+		FSolverVec3 WorldContactTangentV;
 
 		// I^-1.(R x A) for each body where A is each axis (Normal, TangentU, TangentV)
-		FVec3LP WorldContactNormalAngular0;
-		FVec3LP WorldContactTangentUAngular0;
-		FVec3LP WorldContactTangentVAngular0;
-		FVec3LP WorldContactNormalAngular1;
-		FVec3LP WorldContactTangentUAngular1;
-		FVec3LP WorldContactTangentVAngular1;
+		FSolverVec3 WorldContactNormalAngular0;
+		FSolverVec3 WorldContactTangentUAngular0;
+		FSolverVec3 WorldContactTangentVAngular0;
+		FSolverVec3 WorldContactNormalAngular1;
+		FSolverVec3 WorldContactTangentUAngular1;
+		FSolverVec3 WorldContactTangentVAngular1;
 
 		// Contact mass (for non-friction)
-		FRealLP ContactMassNormal;
-		FRealLP ContactMassTangentU;
-		FRealLP ContactMassTangentV;
+		FSolverReal ContactMassNormal;
+		FSolverReal ContactMassTangentU;
+		FSolverReal ContactMassTangentV;
 
 		// Initial world-space contact separation that we are trying to correct
-		FRealLP WorldContactDeltaNormal;
-		FRealLP WorldContactDeltaTangent0;
-		FRealLP WorldContactDeltaTangent1;
+		FSolverReal WorldContactDeltaNormal;
+		FSolverReal WorldContactDeltaTangent0;
+		FSolverReal WorldContactDeltaTangent1;
 
 		// Desired final normal velocity, taking Restitution into account
-		FRealLP WorldContactVelocityTargetNormal;
+		FSolverReal WorldContactVelocityTargetNormal;
 
 		// Solver outputs
-		FRealLP NetPushOutNormal;
-		FRealLP NetPushOutTangentU;
-		FRealLP NetPushOutTangentV;
-		FRealLP NetImpulseNormal;
-		FRealLP NetImpulseTangentU;
-		FRealLP NetImpulseTangentV;
+		FSolverReal NetPushOutNormal;
+		FSolverReal NetPushOutTangentU;
+		FSolverReal NetPushOutTangentV;
+		FSolverReal NetImpulseNormal;
+		FSolverReal NetImpulseTangentU;
+		FSolverReal NetImpulseTangentV;
 
 		// A smoothed NetPushOutNormal along the normal, used for clipping to the static friction cone
-		FRealLP FrictionNetPushOutNormal;
+		FSolverReal FrictionNetPushOutNormal;
 
 		// Whether we are still in the static friction cone
 		bool bInsideStaticFrictionCone;
@@ -158,18 +158,18 @@ namespace Chaos
 
 		FPBDCollisionSolver();
 
-		FRealLP StaticFriction() const { return State.StaticFriction; }
-		FRealLP DynamicFriction() const { return State.DynamicFriction; }
+		FSolverReal StaticFriction() const { return State.StaticFriction; }
+		FSolverReal DynamicFriction() const { return State.DynamicFriction; }
 
 		void SetFriction(const FReal InStaticFriction, const FReal InDynamicFriction)
 		{
-			State.StaticFriction = FRealLP(InStaticFriction);
-			State.DynamicFriction = FRealLP(InDynamicFriction);
+			State.StaticFriction = FSolverReal(InStaticFriction);
+			State.DynamicFriction = FSolverReal(InDynamicFriction);
 		}
 
 		void SetStiffness(const FReal InStiffness)
 		{
-			State.Stiffness = FRealLP(InStiffness);
+			State.Stiffness = FSolverReal(InStiffness);
 		}
 
 		void SetSolverBodies(FSolverBody* SolverBody0, FSolverBody* SolverBody1)
@@ -253,21 +253,21 @@ namespace Chaos
 		 * @brief Calculate and apply the position correction for this iteration
 		 * @return true if we need to run more iterations, false if we did not apply any correction
 		*/
-		bool SolvePositionWithFriction(const FRealLP Dt, const FRealLP MaxPushOut);
-		bool SolvePositionNoFriction(const FRealLP Dt, const FRealLP MaxPushOut);
+		bool SolvePositionWithFriction(const FSolverReal Dt, const FSolverReal MaxPushOut);
+		bool SolvePositionNoFriction(const FSolverReal Dt, const FSolverReal MaxPushOut);
 
 		/**
 		 * @brief Calculate and apply the velocity correction for this iteration
 		 * @return true if we need to run more iterations, false if we did not apply any correction
 		*/
-		bool SolveVelocity(const FRealLP Dt, const bool bApplyDynamicFriction);
+		bool SolveVelocity(const FSolverReal Dt, const bool bApplyDynamicFriction);
 
 	private:
 		/**
 		 * @brief Apply the inverse mass scale the body with the lower level
 		 * @param InvMassScale 
 		*/
-		void SetShockPropagationInvMassScale(const FRealLP InvMassScale);
+		void SetShockPropagationInvMassScale(const FSolverReal InvMassScale);
 
 		/**
 		 * @brief Run a velocity solve on the average point from all the points that received a position impulse
@@ -275,7 +275,7 @@ namespace Chaos
 		 * adding more velocity iterations.
 		 * This will only perform work if there is more than one active contact.
 		*/
-		void SolveVelocityAverage(const FRealLP Dt);
+		void SolveVelocityAverage(const FSolverReal Dt);
 
 		struct FState
 		{
@@ -290,9 +290,9 @@ namespace Chaos
 			{
 			}
 
-			FRealLP StaticFriction;
-			FRealLP DynamicFriction;
-			FRealLP Stiffness;
+			FSolverReal StaticFriction;
+			FSolverReal DynamicFriction;
+			FSolverReal Stiffness;
 			int32 NumManifoldPoints;
 			bool bHaveRestitution;
 			FConstraintSolverBody SolverBodies[MaxConstrainedBodies];
@@ -310,17 +310,17 @@ namespace Chaos
 
 
 	FORCEINLINE_DEBUGGABLE void CalculatePositionCorrectionNormal(
-		const FRealLP Stiffness,
-		const FRealLP ContactDeltaNormal,
-		const FRealLP ContactMassNormal,
-		const FRealLP NetPushOutNormal,
-		FRealLP& OutPushOutNormal)
+		const FSolverReal Stiffness,
+		const FSolverReal ContactDeltaNormal,
+		const FSolverReal ContactMassNormal,
+		const FSolverReal NetPushOutNormal,
+		FSolverReal& OutPushOutNormal)
 	{
-		const FRealLP PushOutNormal = -Stiffness * ContactDeltaNormal * ContactMassNormal;
+		const FSolverReal PushOutNormal = -Stiffness * ContactDeltaNormal * ContactMassNormal;
 
 		// The total pushout so far this sub-step
 		// Unilateral constraint: Net-negative impulses not allowed (negative incremental impulses are allowed as long as the net is positive)
-		if ((NetPushOutNormal + PushOutNormal) > FRealLP(0))
+		if ((NetPushOutNormal + PushOutNormal) > FSolverReal(0))
 		{
 			OutPushOutNormal = PushOutNormal;
 		}
@@ -332,12 +332,12 @@ namespace Chaos
 
 
 	FORCEINLINE_DEBUGGABLE void CalculatePositionCorrectionTangent(
-		const FRealLP Stiffness,
-		const FRealLP ContactDeltaTangent,
-		const FRealLP ContactMassTangent,
-		FRealLP& OutPushOutTangent)
+		const FSolverReal Stiffness,
+		const FSolverReal ContactDeltaTangent,
+		const FSolverReal ContactMassTangent,
+		FSolverReal& OutPushOutTangent)
 	{
-		const FRealLP FrictionStiffness = CVars::Chaos_PBDCollisionSolver_Position_StaticFrictionStiffness;
+		const FSolverReal FrictionStiffness = CVars::Chaos_PBDCollisionSolver_Position_StaticFrictionStiffness;
 
 		// Bilateral constraint - negative values allowed (unlike the normal correction)
 		OutPushOutTangent = -Stiffness * FrictionStiffness * ContactMassTangent * ContactDeltaTangent;
@@ -345,14 +345,14 @@ namespace Chaos
 
 
 	FORCEINLINE_DEBUGGABLE void ApplyFrictionCone(
-		const FRealLP StaticFriction,
-		const FRealLP DynamicFriction,
-		const FRealLP NetPushOutNormal,
-		FRealLP& InOutPushOutTangentU,
-		FRealLP& InOutPushOutTangentV,
-		FRealLP& InOutNetPushOutTangentU,
-		FRealLP& InOutNetPushOutTangentV,
-		FRealLP& InOutFrictionNetPushOutNormal,
+		const FSolverReal StaticFriction,
+		const FSolverReal DynamicFriction,
+		const FSolverReal NetPushOutNormal,
+		FSolverReal& InOutPushOutTangentU,
+		FSolverReal& InOutPushOutTangentV,
+		FSolverReal& InOutNetPushOutTangentU,
+		FSolverReal& InOutNetPushOutTangentV,
+		FSolverReal& InOutFrictionNetPushOutNormal,
 		bool& bInOutInsideStaticFrictionCone)
 	{
 		bInOutInsideStaticFrictionCone = true;
@@ -360,11 +360,11 @@ namespace Chaos
 		// Static friction limit: immediately increase maximum lateral correction, but smoothly decay maximum static friction limit. 
 		// This is so that small variations in position (jitter) and therefore NetPushOutNormal don't cause static friction to slip
 		// @todo(chaos): StaticFriction smoothing is iteration count depenendent - try to make it not so
-		const FRealLP StaticFrictionLerpRate = CVars::Chaos_PBDCollisionSolver_Position_StaticFrictionLerpRate;
-		const FRealLP StaticFrictionTarget = FMath::Max(NetPushOutNormal, FRealLP(0));
+		const FSolverReal StaticFrictionLerpRate = CVars::Chaos_PBDCollisionSolver_Position_StaticFrictionLerpRate;
+		const FSolverReal StaticFrictionTarget = FMath::Max(NetPushOutNormal, FSolverReal(0));
 		InOutFrictionNetPushOutNormal = FMath::Lerp(FMath::Max(InOutFrictionNetPushOutNormal, StaticFrictionTarget), StaticFrictionTarget, StaticFrictionLerpRate);
 
-		if (NetPushOutNormal < FRealLP(SMALL_NUMBER))
+		if (NetPushOutNormal < FSolverReal(SMALL_NUMBER))
 		{
 			// Note: we have already added the current iteration's PushOut to the NetPushOut but it has not been applied to the body
 			// so we must subtract it again to calculate the actual pushout we want to undo (i.e., the net pushout that has been applied to the body so far)
@@ -378,16 +378,16 @@ namespace Chaos
 		// If we exceed the friction cone, clip to cone and stop adding frictional corrections
 		if (bInOutInsideStaticFrictionCone)
 		{
-			const FRealLP MaxPushOutTangent = StaticFriction * InOutFrictionNetPushOutNormal;
-			const FRealLP MaxPushOutTangentSq = FMath::Square(MaxPushOutTangent);
-			const FRealLP NetPushOutTangentSq = FMath::Square(InOutNetPushOutTangentU) + FMath::Square(InOutNetPushOutTangentV);
+			const FSolverReal MaxPushOutTangent = StaticFriction * InOutFrictionNetPushOutNormal;
+			const FSolverReal MaxPushOutTangentSq = FMath::Square(MaxPushOutTangent);
+			const FSolverReal NetPushOutTangentSq = FMath::Square(InOutNetPushOutTangentU) + FMath::Square(InOutNetPushOutTangentV);
 			if (NetPushOutTangentSq > MaxPushOutTangentSq)
 			{
 				// @todo(chaos): clamp to dynamic friction rather than static
-				const FRealLP ClippedMaxPushOutTangent = MaxPushOutTangent;
-				const FRealLP InvNetPushOutTangent = FMath::InvSqrt(NetPushOutTangentSq);
-				const FRealLP NetPushOutTangentU = ClippedMaxPushOutTangent * InOutNetPushOutTangentU * InvNetPushOutTangent;
-				const FRealLP NetPushOutTangentV = ClippedMaxPushOutTangent * InOutNetPushOutTangentV * InvNetPushOutTangent;
+				const FSolverReal ClippedMaxPushOutTangent = MaxPushOutTangent;
+				const FSolverReal InvNetPushOutTangent = FMath::InvSqrt(NetPushOutTangentSq);
+				const FSolverReal NetPushOutTangentU = ClippedMaxPushOutTangent * InOutNetPushOutTangentU * InvNetPushOutTangent;
+				const FSolverReal NetPushOutTangentV = ClippedMaxPushOutTangent * InOutNetPushOutTangentV * InvNetPushOutTangent;
 				InOutPushOutTangentU = NetPushOutTangentU - (InOutNetPushOutTangentU - InOutPushOutTangentU);
 				InOutPushOutTangentV = NetPushOutTangentV - (InOutNetPushOutTangentV - InOutPushOutTangentV);
 				InOutNetPushOutTangentU = NetPushOutTangentU;
@@ -399,19 +399,19 @@ namespace Chaos
 
 
 	FORCEINLINE_DEBUGGABLE void ApplyPositionCorrectionWithFriction(
-		const FRealLP Stiffness,
-		const FRealLP StaticFriction,
-		const FRealLP DynamicFriction,
-		const FRealLP ContactDeltaNormal,
-		const FRealLP ContactDeltaTangent0,
-		const FRealLP ContactDeltaTangent1,
+		const FSolverReal Stiffness,
+		const FSolverReal StaticFriction,
+		const FSolverReal DynamicFriction,
+		const FSolverReal ContactDeltaNormal,
+		const FSolverReal ContactDeltaTangent0,
+		const FSolverReal ContactDeltaTangent1,
 		FPBDCollisionSolverManifoldPoint& ManifoldPoint,
 		FConstraintSolverBody& Body0,
 		FConstraintSolverBody& Body1)
 	{
-		FRealLP PushOutNormal = FRealLP(0);
-		FRealLP PushOutTangentU = FRealLP(0);
-		FRealLP PushOutTangentV = FRealLP(0);
+		FSolverReal PushOutNormal = FSolverReal(0);
+		FSolverReal PushOutTangentU = FSolverReal(0);
+		FSolverReal PushOutTangentV = FSolverReal(0);
 
 		CalculatePositionCorrectionNormal(
 			Stiffness,
@@ -447,33 +447,33 @@ namespace Chaos
 			ManifoldPoint.FrictionNetPushOutNormal,				// Out
 			ManifoldPoint.bInsideStaticFrictionCone);			// Out
 
-		const FVec3LP PushOut = PushOutNormal * ManifoldPoint.WorldContactNormal + PushOutTangentU * ManifoldPoint.WorldContactTangentU + PushOutTangentV * ManifoldPoint.WorldContactTangentV;
+		const FSolverVec3 PushOut = PushOutNormal * ManifoldPoint.WorldContactNormal + PushOutTangentU * ManifoldPoint.WorldContactTangentU + PushOutTangentV * ManifoldPoint.WorldContactTangentV;
 
 		// Update the particle state based on the pushout
 		if (Body0.IsDynamic())
 		{
-			const FVec3LP DX0 = Body0.InvM() * PushOut;
-			const FVec3LP DR0 = ManifoldPoint.WorldContactNormalAngular0 * PushOutNormal + ManifoldPoint.WorldContactTangentUAngular0 * PushOutTangentU + ManifoldPoint.WorldContactTangentVAngular0 * PushOutTangentV;
+			const FSolverVec3 DX0 = Body0.InvM() * PushOut;
+			const FSolverVec3 DR0 = ManifoldPoint.WorldContactNormalAngular0 * PushOutNormal + ManifoldPoint.WorldContactTangentUAngular0 * PushOutTangentU + ManifoldPoint.WorldContactTangentVAngular0 * PushOutTangentV;
 			Body0.ApplyPositionDelta(DX0);
 			Body0.ApplyRotationDelta(DR0);
 		}
 		if (Body1.IsDynamic())
 		{
-			const FVec3LP DX1 = -(Body1.InvM() * PushOut);
-			const FVec3LP DR1 = ManifoldPoint.WorldContactNormalAngular1 * -PushOutNormal + ManifoldPoint.WorldContactTangentUAngular1 * -PushOutTangentU + ManifoldPoint.WorldContactTangentVAngular1 * -PushOutTangentV;
+			const FSolverVec3 DX1 = -(Body1.InvM() * PushOut);
+			const FSolverVec3 DR1 = ManifoldPoint.WorldContactNormalAngular1 * -PushOutNormal + ManifoldPoint.WorldContactTangentUAngular1 * -PushOutTangentU + ManifoldPoint.WorldContactTangentVAngular1 * -PushOutTangentV;
 			Body1.ApplyPositionDelta(DX1);
 			Body1.ApplyRotationDelta(DR1);
 		}
 	}
 
 	FORCEINLINE_DEBUGGABLE void ApplyPositionCorrectionNoFriction(
-		const FRealLP Stiffness,
-		const FRealLP ContactDeltaNormal,
+		const FSolverReal Stiffness,
+		const FSolverReal ContactDeltaNormal,
 		FPBDCollisionSolverManifoldPoint& ManifoldPoint,
 		FConstraintSolverBody& Body0,
 		FConstraintSolverBody& Body1)
 	{
-		FRealLP PushOutNormal = FRealLP(0);
+		FSolverReal PushOutNormal = FSolverReal(0);
 
 		CalculatePositionCorrectionNormal(
 			Stiffness,
@@ -487,15 +487,15 @@ namespace Chaos
 		// Update the particle state based on the pushout
 		if (Body0.IsDynamic())
 		{
-			const FVec3LP DX0 = (Body0.InvM() * PushOutNormal) * ManifoldPoint.WorldContactNormal;
-			const FVec3LP DR0 = ManifoldPoint.WorldContactNormalAngular0 * PushOutNormal;
+			const FSolverVec3 DX0 = (Body0.InvM() * PushOutNormal) * ManifoldPoint.WorldContactNormal;
+			const FSolverVec3 DR0 = ManifoldPoint.WorldContactNormalAngular0 * PushOutNormal;
 			Body0.ApplyPositionDelta(DX0);
 			Body0.ApplyRotationDelta(DR0);
 		}
 		if (Body1.IsDynamic())
 		{
-			const FVec3LP DX1 = (Body1.InvM() * -PushOutNormal) * ManifoldPoint.WorldContactNormal;
-			const FVec3LP DR1 = ManifoldPoint.WorldContactNormalAngular1 * -PushOutNormal;
+			const FSolverVec3 DX1 = (Body1.InvM() * -PushOutNormal) * ManifoldPoint.WorldContactNormal;
+			const FSolverVec3 DR1 = ManifoldPoint.WorldContactNormalAngular1 * -PushOutNormal;
 			Body1.ApplyPositionDelta(DX1);
 			Body1.ApplyRotationDelta(DR1);
 		}
@@ -503,22 +503,22 @@ namespace Chaos
 
 
 	FORCEINLINE_DEBUGGABLE void ApplyVelocityCorrection(
-		const FRealLP Stiffness,
-		const FRealLP Dt,
-		const FRealLP DynamicFriction,
-		const FRealLP ContactVelocityDeltaNormal,
-		const FRealLP ContactVelocityDeltaTangent0,
-		const FRealLP ContactVelocityDeltaTangent1,
+		const FSolverReal Stiffness,
+		const FSolverReal Dt,
+		const FSolverReal DynamicFriction,
+		const FSolverReal ContactVelocityDeltaNormal,
+		const FSolverReal ContactVelocityDeltaTangent0,
+		const FSolverReal ContactVelocityDeltaTangent1,
 		FPBDCollisionSolverManifoldPoint& ManifoldPoint,
 		FConstraintSolverBody& Body0,
 		FConstraintSolverBody& Body1)
 	{
-		FRealLP ImpulseNormal = -(Stiffness * ManifoldPoint.ContactMassNormal) * ContactVelocityDeltaNormal;
+		FSolverReal ImpulseNormal = -(Stiffness * ManifoldPoint.ContactMassNormal) * ContactVelocityDeltaNormal;
 		
 		// Clamp the total impulse to be positive along the normal. We can apply a net negative impulse, 
 		// but only to correct the velocity that was added by pushout.
-		FRealLP NetImpulseNormal = ManifoldPoint.NetImpulseNormal + ImpulseNormal;
-		const FRealLP PushOutImpulseNormal = FMath::Max(0.0f, ManifoldPoint.NetPushOutNormal / Dt);
+		FSolverReal NetImpulseNormal = ManifoldPoint.NetImpulseNormal + ImpulseNormal;
+		const FSolverReal PushOutImpulseNormal = FMath::Max(0.0f, ManifoldPoint.NetPushOutNormal / Dt);
 		if (NetImpulseNormal < -PushOutImpulseNormal)
 		{
 			// We are trying to apply a net negative impulse larger than one to counteract the effective pushout impulse
@@ -527,8 +527,8 @@ namespace Chaos
 			ImpulseNormal = ImpulseNormal - (NetImpulseNormal + PushOutImpulseNormal);
 		}
 
-		FRealLP ImpulseTangentU = FRealLP(0);
-		FRealLP ImpulseTangentV = FRealLP(0);
+		FSolverReal ImpulseTangentU = FSolverReal(0);
+		FSolverReal ImpulseTangentV = FSolverReal(0);
 		const bool bApplyFriction = (DynamicFriction > 0) && !FMath::IsNearlyZero(ImpulseNormal) && (Dt > 0);
 		if (bApplyFriction)
 		{
@@ -540,17 +540,17 @@ namespace Chaos
 		ManifoldPoint.NetImpulseTangentU += ImpulseTangentU;
 		ManifoldPoint.NetImpulseTangentV += ImpulseTangentV;
 
-		FVec3LP Impulse = ImpulseNormal * ManifoldPoint.WorldContactNormal + ImpulseTangentU * ManifoldPoint.WorldContactTangentU + ImpulseTangentV * ManifoldPoint.WorldContactTangentV;
+		FSolverVec3 Impulse = ImpulseNormal * ManifoldPoint.WorldContactNormal + ImpulseTangentU * ManifoldPoint.WorldContactTangentU + ImpulseTangentV * ManifoldPoint.WorldContactTangentV;
 
 		// Calculate the velocity deltas from the impulse
 		if (Body0.IsDynamic())
 		{
-			const FVec3LP AngularImpulse = FVec3LP::CrossProduct(ManifoldPoint.RelativeContactPosition0, Impulse);
+			const FSolverVec3 AngularImpulse = FSolverVec3::CrossProduct(ManifoldPoint.RelativeContactPosition0, Impulse);
 			Body0.ApplyVelocityDelta(Body0.InvM() * Impulse, Body0.InvI() * AngularImpulse);
 		}
 		if (Body1.IsDynamic())
 		{
-			const FVec3LP AngularImpulse = FVec3LP::CrossProduct(ManifoldPoint.RelativeContactPosition1, Impulse);
+			const FSolverVec3 AngularImpulse = FSolverVec3::CrossProduct(ManifoldPoint.RelativeContactPosition1, Impulse);
 			Body1.ApplyVelocityDelta(Body1.InvM() * -Impulse, Body1.InvI() * -AngularImpulse);
 		}
 	}
@@ -561,40 +561,40 @@ namespace Chaos
 	//////////////////////////////////////////////////////////////////////////////////////////////////
 	//////////////////////////////////////////////////////////////////////////////////////////////////
 
-	FORCEINLINE_DEBUGGABLE void FPBDCollisionSolverManifoldPoint::CalculateContactPositionErrorNormal(const FConstraintSolverBody& Body0, const FConstraintSolverBody& Body1, const FRealLP MaxPushOut, FRealLP& OutContactDeltaNormal) const
+	FORCEINLINE_DEBUGGABLE void FPBDCollisionSolverManifoldPoint::CalculateContactPositionErrorNormal(const FConstraintSolverBody& Body0, const FConstraintSolverBody& Body1, const FSolverReal MaxPushOut, FSolverReal& OutContactDeltaNormal) const
 	{
 		// Linear version: calculate the contact delta assuming linear motion after applying a positional impulse at the contact point. There will be an error that depends on the size of the rotation.
-		const FVec3LP ContactDelta0 = FVec3LP(Body0.DP()) + FVec3LP::CrossProduct(Body0.DQ(), RelativeContactPosition0);
-		const FVec3LP ContactDelta1 = FVec3LP(Body1.DP()) + FVec3LP::CrossProduct(Body1.DQ(), RelativeContactPosition1);
-		const FVec3LP ContactDelta = ContactDelta0 - ContactDelta1;
-		OutContactDeltaNormal = WorldContactDeltaNormal + FVec3LP::DotProduct(ContactDelta, WorldContactNormal);
+		const FSolverVec3 ContactDelta0 = FSolverVec3(Body0.DP()) + FSolverVec3::CrossProduct(Body0.DQ(), RelativeContactPosition0);
+		const FSolverVec3 ContactDelta1 = FSolverVec3(Body1.DP()) + FSolverVec3::CrossProduct(Body1.DQ(), RelativeContactPosition1);
+		const FSolverVec3 ContactDelta = ContactDelta0 - ContactDelta1;
+		OutContactDeltaNormal = WorldContactDeltaNormal + FSolverVec3::DotProduct(ContactDelta, WorldContactNormal);
 
 		// NOTE: OutContactDeltaNormal is negative for penetration
 		// NOTE: MaxPushOut == 0 disables the pushout limits
 		if ((MaxPushOut > 0) && (OutContactDeltaNormal < -MaxPushOut))
 		{
-			const FRealLP ClampedContactDeltaNormal = -MaxPushOut;
+			const FSolverReal ClampedContactDeltaNormal = -MaxPushOut;
 			OutContactDeltaNormal = ClampedContactDeltaNormal;
 		}
 	}
 
-	FORCEINLINE_DEBUGGABLE void FPBDCollisionSolverManifoldPoint::CalculateContactPositionError(const FConstraintSolverBody& Body0, const FConstraintSolverBody& Body1, const FRealLP MaxPushOut, FRealLP& OutContactDeltaNormal, FRealLP& OutContactDeltaTangent0, FRealLP& OutContactDeltaTangent1) const
+	FORCEINLINE_DEBUGGABLE void FPBDCollisionSolverManifoldPoint::CalculateContactPositionError(const FConstraintSolverBody& Body0, const FConstraintSolverBody& Body1, const FSolverReal MaxPushOut, FSolverReal& OutContactDeltaNormal, FSolverReal& OutContactDeltaTangent0, FSolverReal& OutContactDeltaTangent1) const
 	{
 #if CHAOS_NONLINEAR_COLLISIONS_ENABLED
 #error "Non-linear collision solver no longer functional - fix it"
 		// Non-linear version: calculate the contact delta after we have converted the current positional impulses into position and rotation corrections.
 		// We could precalculate and store the LocalContactPositions if we really want to use this nonlinear version
-		const FVec3LP LocalContactPosition0 = Body0.Q().Inverse() * WorldRelativeContactPosition0;
-		const FVec3LP LocalContactPosition1 = Body1.Q().Inverse() * WorldRelativeContactPosition1;
-		const FVec3LP ContactDelta = (Body0.CorrectedP() + Body0.CorrectedQ() * LocalContactPosition0) - (Body1.CorrectedP() + Body1.CorrectedQ() * LocalContactPosition1);
+		const FSolverVec3 LocalContactPosition0 = Body0.Q().Inverse() * WorldRelativeContactPosition0;
+		const FSolverVec3 LocalContactPosition1 = Body1.Q().Inverse() * WorldRelativeContactPosition1;
+		const FSolverVec3 ContactDelta = (Body0.CorrectedP() + Body0.CorrectedQ() * LocalContactPosition0) - (Body1.CorrectedP() + Body1.CorrectedQ() * LocalContactPosition1);
 #else
 		// Linear version: calculate the contact delta assuming linear motion after applying a positional impulse at the contact point. There will be an error that depends on the size of the rotation.
-		const FVec3LP ContactDelta0 = FVec3LP(Body0.DP()) + FVec3LP::CrossProduct(Body0.DQ(), RelativeContactPosition0);
-		const FVec3LP ContactDelta1 = FVec3LP(Body1.DP()) + FVec3LP::CrossProduct(Body1.DQ(), RelativeContactPosition1);
-		const FVec3LP ContactDelta = ContactDelta0 - ContactDelta1;
-		OutContactDeltaNormal = WorldContactDeltaNormal + FVec3LP::DotProduct(ContactDelta, WorldContactNormal);
-		OutContactDeltaTangent0 = WorldContactDeltaTangent0 + FVec3LP::DotProduct(ContactDelta, WorldContactTangentU);
-		OutContactDeltaTangent1 = WorldContactDeltaTangent1 + FVec3LP::DotProduct(ContactDelta, WorldContactTangentV);
+		const FSolverVec3 ContactDelta0 = FSolverVec3(Body0.DP()) + FSolverVec3::CrossProduct(Body0.DQ(), RelativeContactPosition0);
+		const FSolverVec3 ContactDelta1 = FSolverVec3(Body1.DP()) + FSolverVec3::CrossProduct(Body1.DQ(), RelativeContactPosition1);
+		const FSolverVec3 ContactDelta = ContactDelta0 - ContactDelta1;
+		OutContactDeltaNormal = WorldContactDeltaNormal + FSolverVec3::DotProduct(ContactDelta, WorldContactNormal);
+		OutContactDeltaTangent0 = WorldContactDeltaTangent0 + FSolverVec3::DotProduct(ContactDelta, WorldContactTangentU);
+		OutContactDeltaTangent1 = WorldContactDeltaTangent1 + FSolverVec3::DotProduct(ContactDelta, WorldContactTangentV);
 #endif
 
 		// NOTE: OutContactDeltaNormal is negative for penetration
@@ -605,39 +605,39 @@ namespace Chaos
 		}
 	}
 
-	FORCEINLINE_DEBUGGABLE void FPBDCollisionSolverManifoldPoint::CalculateContactVelocityError(const FConstraintSolverBody& Body0, const FConstraintSolverBody& Body1, const FRealLP DynamicFriction, const FRealLP Dt, FRealLP& OutContactVelocityDeltaNormal, FRealLP& OutContactVelocityDeltaTangent0, FRealLP& OutContactVelocityDeltaTangent1) const
+	FORCEINLINE_DEBUGGABLE void FPBDCollisionSolverManifoldPoint::CalculateContactVelocityError(const FConstraintSolverBody& Body0, const FConstraintSolverBody& Body1, const FSolverReal DynamicFriction, const FSolverReal Dt, FSolverReal& OutContactVelocityDeltaNormal, FSolverReal& OutContactVelocityDeltaTangent0, FSolverReal& OutContactVelocityDeltaTangent1) const
 	{
-		const FVec3LP ContactVelocity0 = Body0.V() + FVec3LP::CrossProduct(Body0.W(), RelativeContactPosition0);
-		const FVec3LP ContactVelocity1 = Body1.V() + FVec3LP::CrossProduct(Body1.W(), RelativeContactPosition1);
-		const FVec3LP ContactVelocity = ContactVelocity0 - ContactVelocity1;
-		const FRealLP ContactVelocityNormal = FVec3LP::DotProduct(ContactVelocity, WorldContactNormal);
+		const FSolverVec3 ContactVelocity0 = Body0.V() + FSolverVec3::CrossProduct(Body0.W(), RelativeContactPosition0);
+		const FSolverVec3 ContactVelocity1 = Body1.V() + FSolverVec3::CrossProduct(Body1.W(), RelativeContactPosition1);
+		const FSolverVec3 ContactVelocity = ContactVelocity0 - ContactVelocity1;
+		const FSolverReal ContactVelocityNormal = FSolverVec3::DotProduct(ContactVelocity, WorldContactNormal);
 
 		// Target normal velocity, including restitution
-		const FRealLP ContactVelocityTargetNormal = WorldContactVelocityTargetNormal;
+		const FSolverReal ContactVelocityTargetNormal = WorldContactVelocityTargetNormal;
 
 		// Add up the errors in the velocity (current velocity - desired velocity)
 		OutContactVelocityDeltaNormal = (ContactVelocityNormal - WorldContactVelocityTargetNormal);
-		OutContactVelocityDeltaTangent0 = FRealLP(0);
-		OutContactVelocityDeltaTangent1 = FRealLP(0);
+		OutContactVelocityDeltaTangent0 = FSolverReal(0);
+		OutContactVelocityDeltaTangent1 = FSolverReal(0);
 
 		const bool bApplyFriction = (DynamicFriction > 0) && (Dt > 0);
 		if (bApplyFriction)
 		{
-			const FRealLP ContactVelocityTangent0 = FVec3LP::DotProduct(ContactVelocity, WorldContactTangentU);
-			const FRealLP ContactVelocityTangent1 = FVec3LP::DotProduct(ContactVelocity, WorldContactTangentV);
+			const FSolverReal ContactVelocityTangent0 = FSolverVec3::DotProduct(ContactVelocity, WorldContactTangentU);
+			const FSolverReal ContactVelocityTangent1 = FSolverVec3::DotProduct(ContactVelocity, WorldContactTangentV);
 			OutContactVelocityDeltaTangent0 = ContactVelocityTangent0;
 			OutContactVelocityDeltaTangent1 = ContactVelocityTangent1;
 
-			const FRealLP ContactVelocityTangentialSq = FMath::Square(ContactVelocityTangent0) + FMath::Square(ContactVelocityTangent1);
-			if (ContactVelocityTangentialSq > FRealLP(SMALL_NUMBER))
+			const FSolverReal ContactVelocityTangentialSq = FMath::Square(ContactVelocityTangent0) + FMath::Square(ContactVelocityTangent1);
+			if (ContactVelocityTangentialSq > FSolverReal(SMALL_NUMBER))
 			{
 				// PushOut = ContactMass * DP, where DP is the contact positional correction
 				// Friction force is proportional to the normal force, so friction velocity correction
 				// is proprtional to normal velocity correction, or DVn = DPn/dt = PushOut.N / (ContactMass * dt);
-				const FRealLP MaxContactVelocityTangential = DynamicFriction * NetPushOutNormal / (ContactMassNormal * Dt);
+				const FSolverReal MaxContactVelocityTangential = DynamicFriction * NetPushOutNormal / (ContactMassNormal * Dt);
 				if (ContactVelocityTangentialSq > FMath::Square(MaxContactVelocityTangential))
 				{
-					const FRealLP ContactVelocityTangentialScale = MaxContactVelocityTangential * FMath::InvSqrt(ContactVelocityTangentialSq);
+					const FSolverReal ContactVelocityTangentialScale = MaxContactVelocityTangential * FMath::InvSqrt(ContactVelocityTangentialSq);
 					OutContactVelocityDeltaTangent0 *= ContactVelocityTangentialScale;
 					OutContactVelocityDeltaTangent1 *= ContactVelocityTangentialScale;
 				}
@@ -652,7 +652,7 @@ namespace Chaos
 	//////////////////////////////////////////////////////////////////////////////////////////////////
 
 
-	FORCEINLINE_DEBUGGABLE bool FPBDCollisionSolver::SolvePositionWithFriction(const FRealLP Dt, const FRealLP MaxPushOut)
+	FORCEINLINE_DEBUGGABLE bool FPBDCollisionSolver::SolvePositionWithFriction(const FSolverReal Dt, const FSolverReal MaxPushOut)
 	{
 		// SolverBody decorator used to add mass scaling
 		FConstraintSolverBody& Body0 = SolverBody0();
@@ -663,10 +663,10 @@ namespace Chaos
 		{
 			FPBDCollisionSolverManifoldPoint& SolverManifoldPoint = State.ManifoldPoints[PointIndex];
 
-			FRealLP ContactDeltaNormal, ContactDeltaTangent0, ContactDeltaTangent1;
+			FSolverReal ContactDeltaNormal, ContactDeltaTangent0, ContactDeltaTangent1;
 			SolverManifoldPoint.CalculateContactPositionError(Body0.SolverBody(), Body1.SolverBody(), MaxPushOut, ContactDeltaNormal, ContactDeltaTangent0, ContactDeltaTangent1);
 
-			const bool bProcessManifoldPoint = (ContactDeltaNormal < FRealLP(0)) || !FMath::IsNearlyZero(SolverManifoldPoint.NetPushOutNormal);
+			const bool bProcessManifoldPoint = (ContactDeltaNormal < FSolverReal(0)) || !FMath::IsNearlyZero(SolverManifoldPoint.NetPushOutNormal);
 			if (bProcessManifoldPoint)
 			{
 				// If we want to add static friction...
@@ -700,7 +700,7 @@ namespace Chaos
 		return false;
 	}
 
-	FORCEINLINE_DEBUGGABLE bool FPBDCollisionSolver::SolvePositionNoFriction(const FRealLP Dt, const FRealLP MaxPushOut)
+	FORCEINLINE_DEBUGGABLE bool FPBDCollisionSolver::SolvePositionNoFriction(const FSolverReal Dt, const FSolverReal MaxPushOut)
 	{
 		// SolverBody decorator used to add mass scaling
 		FConstraintSolverBody& Body0 = SolverBody0();
@@ -711,10 +711,10 @@ namespace Chaos
 		{
 			FPBDCollisionSolverManifoldPoint& SolverManifoldPoint = State.ManifoldPoints[PointIndex];
 
-			FRealLP ContactDeltaNormal;
+			FSolverReal ContactDeltaNormal;
 			SolverManifoldPoint.CalculateContactPositionErrorNormal(Body0.SolverBody(), Body1.SolverBody(), MaxPushOut, ContactDeltaNormal);
 
-			const bool bProcessManifoldPoint = (ContactDeltaNormal < FRealLP(0)) || !FMath::IsNearlyZero(SolverManifoldPoint.NetPushOutNormal);
+			const bool bProcessManifoldPoint = (ContactDeltaNormal < FSolverReal(0)) || !FMath::IsNearlyZero(SolverManifoldPoint.NetPushOutNormal);
 			if (bProcessManifoldPoint)
 			{
 				ApplyPositionCorrectionNoFriction(
diff --git a/Engine/Source/Runtime/Experimental/Chaos/Private/Chaos/Collision/SolverCollisionContainer.cpp b/Engine/Source/Runtime/Experimental/Chaos/Private/Chaos/Collision/SolverCollisionContainer.cpp
--- a/Engine/Source/Runtime/Experimental/Chaos/Private/Chaos/Collision/SolverCollisionContainer.cpp
+++ b/Engine/Source/Runtime/Experimental/Chaos/Private/Chaos/Collision/SolverCollisionContainer.cpp
@@ -134,12 +134,12 @@ namespace Chaos
 			{
 				const FPBDCollisionSolverManifoldPoint& SolverManifoldPoint = Solver.GetManifoldPoint(PointIndex);
 
-				const FVec3LP NetPushOut = 
+				const FSolverVec3 NetPushOut = 
 					SolverManifoldPoint.NetPushOutNormal * SolverManifoldPoint.WorldContactNormal +
 					SolverManifoldPoint.NetPushOutTangentU * SolverManifoldPoint.WorldContactTangentU +
 					SolverManifoldPoint.NetPushOutTangentV * SolverManifoldPoint.WorldContactTangentV;
 
-				const FVec3LP NetImpulse =
+				const FSolverVec3 NetImpulse =
 					SolverManifoldPoint.NetImpulseNormal * SolverManifoldPoint.WorldContactNormal +
 					SolverManifoldPoint.NetImpulseTangentU * SolverManifoldPoint.WorldContactTangentU +
 					SolverManifoldPoint.NetImpulseTangentV * SolverManifoldPoint.WorldContactTangentV;
@@ -289,8 +289,8 @@ namespace Chaos
 	// Solve position including support for incremental collision detection
 	bool FPBDCollisionSolverContainer::SolvePositionIncrementalImpl(const FReal InDt, const int32 BeginIndex, const int32 EndIndex, const FReal InMaxPushOut, const bool bApplyStaticFriction)
 	{
-		const FRealLP Dt = FRealLP(InDt);
-		const FRealLP MaxPushOut = FRealLP(InMaxPushOut);
+		const FSolverReal Dt = FSolverReal(InDt);
+		const FSolverReal MaxPushOut = FSolverReal(InMaxPushOut);
 
 		bool bNeedsAnotherIteration = false;
 		for (int32 SolverIndex = BeginIndex; SolverIndex < EndIndex; ++SolverIndex)
@@ -317,8 +317,8 @@ namespace Chaos
 	// Solve position with friction (last few iterations each tick)
 	bool FPBDCollisionSolverContainer::SolvePositionWithFrictionImpl(const FReal InDt, const int32 BeginIndex, const int32 EndIndex, const FReal InMaxPushOut)
 	{
-		const FRealLP Dt = FRealLP(InDt);
-		const FRealLP MaxPushOut = FRealLP(InMaxPushOut);
+		const FSolverReal Dt = FSolverReal(InDt);
+		const FSolverReal MaxPushOut = FSolverReal(InMaxPushOut);
 
 		bool bNeedsAnotherIteration = false;
 		for (int32 SolverIndex = BeginIndex; SolverIndex < EndIndex; ++SolverIndex)
@@ -331,8 +331,8 @@ namespace Chaos
 	// Solve position without friction (first few iterations each tick)
 	bool FPBDCollisionSolverContainer::SolvePositionNoFrictionImpl(const FReal InDt, const int32 BeginIndex, const int32 EndIndex, const FReal InMaxPushOut)
 	{
-		const FRealLP Dt = FRealLP(InDt);
-		const FRealLP MaxPushOut = FRealLP(InMaxPushOut);
+		const FSolverReal Dt = FSolverReal(InDt);
+		const FSolverReal MaxPushOut = FSolverReal(InMaxPushOut);
 
 		bool bNeedsAnotherIteration = false;
 		for (int32 SolverIndex = BeginIndex; SolverIndex < EndIndex; ++SolverIndex)
@@ -349,7 +349,7 @@ namespace Chaos
 		{
 			return false;
 		}
-		const FRealLP Dt = FRealLP(InDt);
+		const FSolverReal Dt = FSolverReal(InDt);
 
 		UpdateVelocityShockPropagation(Dt, It, NumIts, BeginIndex, EndIndex);
 
diff --git a/Engine/Source/Runtime/Experimental/Chaos/Public/Chaos/Evolution/SolverBody.h b/Engine/Source/Runtime/Experimental/Chaos/Public/Chaos/Evolution/SolverBody.h
--- a/Engine/Source/Runtime/Experimental/Chaos/Public/Chaos/Evolution/SolverBody.h
+++ b/Engine/Source/Runtime/Experimental/Chaos/Public/Chaos/Evolution/SolverBody.h
@@ -9,14 +9,17 @@
 
 namespace Chaos
 {
-	// Set the low precision (LP) math types
+	// Set the math types used by the constraint solvers when single precision is acceptable.
+	// NOTE: public APIs should still use FReal, FVec3 etc. whereas FSolverReal and associated
+	// types are for use internally where double precision is unnecessary and causes performance
+	// issues and/or memory bloat.
 #if CHAOS_CONSTRAINTSOLVER_LOWPRECISION
-	using FRealLP = FRealSingle;
+	using FSolverReal = FRealSingle;
 #else
-	using FRealLP = FReal;
+	using FSolverReal = FReal;
 #endif
-	using FVec3LP = TVec3<FRealLP>;
-	using FMatrix33LP = TMatrix33<FRealLP>;
+	using FSolverVec3 = TVec3<FSolverReal>;
+	using FSolverMatrix33 = TMatrix33<FSolverReal>;
 
 
 	class FSolverBody;
@@ -136,7 +139,7 @@ namespace Chaos
 		{
 			if (IsDynamic() && (Dt != FReal(0)))
 			{
-				const FRealLP InvDt = FRealLP(1) / FRealLP(Dt);
+				const FSolverReal InvDt = FSolverReal(1) / FSolverReal(Dt);
 				State.V += State.DP * InvDt;
 				State.W += State.DQ * InvDt;
 			}
@@ -146,34 +149,34 @@ namespace Chaos
 		 * @brief Get the inverse mass
 		*/
 		//inline FReal InvM() const { return State.InvM; }
-		inline FRealLP InvM() const { return State.InvM; }
+		inline FSolverReal InvM() const { return State.InvM; }
 
 		/**
 		 * @brief Set the inverse mass
 		*/
-		inline void SetInvM(FReal InInvM) { State.InvM = FRealLP(InInvM); }
+		inline void SetInvM(FReal InInvM) { State.InvM = FSolverReal(InInvM); }
 
 		/**
 		 * @brief Get the world-space inverse inertia
 		*/
-		inline const FMatrix33LP& InvI() const { return State.InvI; }
+		inline const FSolverMatrix33& InvI() const { return State.InvI; }
 
 		/**
 		 * @brief Set the world-space inverse inertia
 		*/
-		inline void SetInvI(const FMatrix33& InInvI) { State.InvI = FMatrix33LP(InInvI); }
+		inline void SetInvI(const FMatrix33& InInvI) { State.InvI = FSolverMatrix33(InInvI); }
 
 		/**
 		 * @brief Get the local-space inverse inertia (diagonal elements)
 		*/
-		inline const FVec3LP& InvILocal() const { return State.InvILocal; }
+		inline const FSolverVec3& InvILocal() const { return State.InvILocal; }
 
 		/**
 		 * @brief Set the local-space inverse inertia (diagonal elements)
 		*/
 		inline void SetInvILocal(const FVec3& InInvILocal)
 		{ 
-			State.InvILocal = FVec3LP(InInvILocal); 
+			State.InvILocal = FSolverVec3(InInvILocal); 
 			UpdateRotationDependentState();
 		}
 
@@ -231,12 +234,12 @@ namespace Chaos
 		/**
 		 * @brief Net world-space position correction applied by the constraints
 		*/
-		inline const FVec3LP& DP() const { return State.DP; }
+		inline const FSolverVec3& DP() const { return State.DP; }
 
 		/**
 		 * @brief Net world-space rotation correction applied by the constraints (axis-angle vector equivalent to angular velocity but for position)
 		*/
-		inline const FVec3LP& DQ() const { return State.DQ; }
+		inline const FSolverVec3& DQ() const { return State.DQ; }
 
 		/**
 		 * @brief World-space position after applying the net correction DP()
@@ -259,8 +262,8 @@ namespace Chaos
 		{
 			State.P = CorrectedP();
 			State.Q = CorrectedQ();
-			State.DP = FVec3LP(0);
-			State.DQ = FVec3LP(0);
+			State.DP = FSolverVec3(0);
+			State.DQ = FSolverVec3(0);
 		}
 
 		/**
@@ -290,7 +293,7 @@ namespace Chaos
 		/**
 		 * @brief Apply a world-space position and rotation delta to the body center of mass, and update inverse mass
 		*/
-		inline void ApplyTransformDelta(const FVec3LP& DP, const FVec3LP& DR)
+		inline void ApplyTransformDelta(const FSolverVec3& DP, const FSolverVec3& DR)
 		{
 			ApplyPositionDelta(DP);
 			ApplyRotationDelta(DR);
@@ -299,7 +302,7 @@ namespace Chaos
 		/**
 		 * @brief Apply a world-space position delta to the solver body center of mass
 		*/
-		inline void ApplyPositionDelta(const FVec3LP& DP)
+		inline void ApplyPositionDelta(const FSolverVec3& DP)
 		{
 			State.DP += DP;
 		}
@@ -307,7 +310,7 @@ namespace Chaos
 		/**
 		 * @brief Apply a world-space rotation delta to the solver body and update the inverse mass
 		*/
-		inline void ApplyRotationDelta(const FVec3LP& DR)
+		inline void ApplyRotationDelta(const FSolverVec3& DR)
 		{
 			State.DQ += DR;
 		}
@@ -315,7 +318,7 @@ namespace Chaos
 		/**
 		 * @brief Apply a world-space velocity delta to the solver body
 		*/
-		inline void ApplyVelocityDelta(const FVec3LP& DV, const FVec3LP& DW)
+		inline void ApplyVelocityDelta(const FSolverVec3& DV, const FSolverVec3& DW)
 		{
 			ApplyLinearVelocityDelta(DV);
 			ApplyAngularVelocityDelta(DW);
@@ -324,7 +327,7 @@ namespace Chaos
 		/**
 		 * @brief Apply a world-space linear velocity delta to the solver body
 		*/
-		inline void ApplyLinearVelocityDelta(const FVec3LP& DV)
+		inline void ApplyLinearVelocityDelta(const FSolverVec3& DV)
 		{
 			State.V += DV;
 		}
@@ -332,7 +335,7 @@ namespace Chaos
 		/**
 		 * @brief Apply an world-space angular velocity delta to the solver body
 		*/
-		inline void ApplyAngularVelocityDelta(const FVec3LP& DW)
+		inline void ApplyAngularVelocityDelta(const FSolverVec3& DW)
 		{
 			State.W += DW;
 		}
@@ -375,20 +378,20 @@ namespace Chaos
 			{}
 
 			// Local-space inverse inertia (diagonal, so only 3 elements)
-			FVec3LP InvILocal;
+			FSolverVec3 InvILocal;
 
 			// Inverse mass
-			FRealLP InvM;
+			FSolverReal InvM;
 
 			// World-space inverse inertia
 			// @todo(chaos): do we need this, or should we force all systems to use the FConstraintSolverBody decorator?
-			FMatrix33LP InvI;
+			FSolverMatrix33 InvI;
 
 			// Net position delta applied by all constraints (constantly changing as we iterate over constraints)
-			FVec3LP DP;
+			FSolverVec3 DP;
 
 			// Net rotation delta applied by all constraints (constantly changing as we iterate over constraints)
-			FVec3LP DQ;
+			FSolverVec3 DQ;
 
 			// Distance to a kinmatic body (through the contact graph). Used by collision shock propagation
 			int32 Level;
@@ -471,23 +474,23 @@ namespace Chaos
 		/**
 		 * @brief A scale applied to both inverse mass and inverse inertia
 		*/
-		inline FRealLP InvMScale() const { return State.InvMassScale; }
-		inline void SetInvMScale(FReal InInvMassScale) { State.InvMassScale = FRealLP(InInvMassScale); }
+		inline FSolverReal InvMScale() const { return State.InvMassScale; }
+		inline void SetInvMScale(FReal InInvMassScale) { State.InvMassScale = FSolverReal(InInvMassScale); }
 
 		/**
 		 * @brief The scaled inverse mass
 		*/
-		FRealLP InvM() const { return State.InvMassScale * Body->InvM(); }
+		FSolverReal InvM() const { return State.InvMassScale * Body->InvM(); }
 
 		/**
 		 * @brief The scaled inverse inertia
 		*/
-		FMatrix33LP InvI() const { return State.InvMassScale * Body->InvI(); }
+		FSolverMatrix33 InvI() const { return State.InvMassScale * Body->InvI(); }
 
 		/**
 		 * @brief The scaled local space inverse inertia
 		*/
-		FVec3LP InvILocal() const { return State.InvMassScale * Body->InvILocal(); }
+		FSolverVec3 InvILocal() const { return State.InvMassScale * Body->InvILocal(); }
 
 		/**
 		 * @brief Whether the body is dynamic (i.e., has a finite mass) after InvMassScale is applied
@@ -509,17 +512,17 @@ namespace Chaos
 		inline const FVec3& V() const { return Body->V(); }
 		inline const FVec3& W() const { return Body->W(); }
 		inline int32 Level() const { return Body->Level(); }
-		inline const FVec3LP& DP() const { return Body->DP(); }
-		inline const FVec3LP& DQ() const { return Body->DQ(); }
+		inline const FSolverVec3& DP() const { return Body->DP(); }
+		inline const FSolverVec3& DQ() const { return Body->DQ(); }
 		inline FVec3 CorrectedP() const { return Body->CorrectedP(); }
 		inline FRotation3 CorrectedQ() const { return Body->CorrectedQ(); }
 
-		inline void ApplyTransformDelta(const FVec3LP& DP, const FVec3LP& DR) { Body->ApplyTransformDelta(DP, DR); }
-		inline void ApplyPositionDelta(const FVec3LP& DP) { Body->ApplyPositionDelta(DP); }
-		inline void ApplyRotationDelta(const FVec3LP& DR) { Body->ApplyRotationDelta(DR); }
-		inline void ApplyVelocityDelta(const FVec3LP& DV, const FVec3LP& DW) { Body->ApplyVelocityDelta(DV, DW); }
-		inline void ApplyLinearVelocityDelta(const FVec3LP& DV) { Body->ApplyLinearVelocityDelta(DV); }
-		inline void ApplyAngularVelocityDelta(const FVec3LP& DW) { Body->ApplyAngularVelocityDelta(DW); }
+		inline void ApplyTransformDelta(const FSolverVec3& DP, const FSolverVec3& DR) { Body->ApplyTransformDelta(DP, DR); }
+		inline void ApplyPositionDelta(const FSolverVec3& DP) { Body->ApplyPositionDelta(DP); }
+		inline void ApplyRotationDelta(const FSolverVec3& DR) { Body->ApplyRotationDelta(DR); }
+		inline void ApplyVelocityDelta(const FSolverVec3& DV, const FSolverVec3& DW) { Body->ApplyVelocityDelta(DV, DW); }
+		inline void ApplyLinearVelocityDelta(const FSolverVec3& DV) { Body->ApplyLinearVelocityDelta(DV); }
+		inline void ApplyAngularVelocityDelta(const FSolverVec3& DW) { Body->ApplyAngularVelocityDelta(DW); }
 		inline void EnforceShortestRotationTo(const FRotation3& InQ) { Body->EnforceShortestRotationTo(InQ); }
 		inline void UpdateRotationDependentState() { Body->UpdateRotationDependentState(); }
 
@@ -528,9 +531,9 @@ namespace Chaos
 		struct FState
 		{
 			FState() 
-				: InvMassScale(FRealLP(1))
+				: InvMassScale(FSolverReal(1))
 			{}
-			FRealLP InvMassScale;
+			FSolverReal InvMassScale;
 		};
 
 		// The body we decorate
