Optimisation of the gather

diff --git a/Engine/Source/Runtime/Experimental/Chaos/Private/Chaos/Collision/PBDCollisionSolver.h b/Engine/Source/Runtime/Experimental/Chaos/Private/Chaos/Collision/PBDCollisionSolver.h
--- a/Engine/Source/Runtime/Experimental/Chaos/Private/Chaos/Collision/PBDCollisionSolver.h
+++ b/Engine/Source/Runtime/Experimental/Chaos/Private/Chaos/Collision/PBDCollisionSolver.h
@@ -158,6 +158,12 @@ namespace Chaos
 
 		FPBDCollisionSolver();
 
+		/** Reset the state of the collision solver */
+		void Reset()
+		{
+			State.Reset();
+		}
+
 		FSolverReal StaticFriction() const { return State.StaticFriction; }
 		FSolverReal DynamicFriction() const { return State.DynamicFriction; }
 
@@ -290,6 +296,16 @@ namespace Chaos
 			{
 			}
 
+			/** Reset the state struct members to its default values */
+			void Reset()
+			{
+				StaticFriction = 0;
+				DynamicFriction = 0;
+				Stiffness = 1;
+				bHaveRestitution = false;
+				NumManifoldPoints = 0;
+			}
+
 			FSolverReal StaticFriction;
 			FSolverReal DynamicFriction;
 			FSolverReal Stiffness;
@@ -356,7 +372,6 @@ namespace Chaos
 		bool& bInOutInsideStaticFrictionCone)
 	{
 		bInOutInsideStaticFrictionCone = true;
-
 		// Static friction limit: immediately increase maximum lateral correction, but smoothly decay maximum static friction limit. 
 		// This is so that small variations in position (jitter) and therefore NetPushOutNormal don't cause static friction to slip
 		// @todo(chaos): StaticFriction smoothing is iteration count depenendent - try to make it not so
diff --git a/Engine/Source/Runtime/Experimental/Chaos/Private/Chaos/Collision/SolverCollisionContainer.cpp b/Engine/Source/Runtime/Experimental/Chaos/Private/Chaos/Collision/SolverCollisionContainer.cpp
--- a/Engine/Source/Runtime/Experimental/Chaos/Private/Chaos/Collision/SolverCollisionContainer.cpp
+++ b/Engine/Source/Runtime/Experimental/Chaos/Private/Chaos/Collision/SolverCollisionContainer.cpp
@@ -186,18 +186,25 @@ namespace Chaos
 	{
 		CollisionSolvers.Reset(MaxCollisions);
 	}
+	
+	void FPBDCollisionSolverContainer::SetNum(const int32 MaxCollisions)
+	{
+		CollisionSolvers.SetNum(MaxCollisions, false);
+	}
 
-	void FPBDCollisionSolverContainer::AddConstraintSolver(FReal Dt, FPBDCollisionConstraint& Constraint, const int32 Particle0Level, const int32 Particle1Level, FSolverBodyContainer& SolverBodyContainer)
+	void FPBDCollisionSolverContainer::AddConstraintSolver(FReal Dt, FPBDCollisionConstraint& Constraint, const int32 Particle0Level, const int32 Particle1Level, FSolverBodyContainer& SolverBodyContainer, int32& ConstraintIndex)
 	{
 		// This container is required to allocate pointers that are valid for the whole tick,
 		// so we cannot allow the container to resize during the tick. See Reset()
-		check(CollisionSolvers.Num() < CollisionSolvers.Max());
+		check(ConstraintIndex < CollisionSolvers.Num());
 
-		const int32 SolverCollisionIndex = CollisionSolvers.AddDefaulted();
+		auto& CollisionSolver = CollisionSolvers[ConstraintIndex];
+		CollisionSolver.GetSolver().Reset();
 
-		CollisionSolvers[SolverCollisionIndex].GatherInput(Dt, Constraint, Particle0Level, Particle1Level, SolverBodyContainer);
+		CollisionSolver.GatherInput(Dt, Constraint, Particle0Level, Particle1Level, SolverBodyContainer);
 
-		bRequiresIncrementalCollisionDetection |= CollisionSolvers[SolverCollisionIndex].IsIncrementalManifold();
+		bRequiresIncrementalCollisionDetection |= CollisionSolver.IsIncrementalManifold();
+		++ConstraintIndex;
 	}
 
 	void FPBDCollisionSolverContainer::UpdatePositionShockPropagation(const FReal Dt, const int32 It, const int32 NumIts, const int32 BeginIndex, const int32 EndIndex)
diff --git a/Engine/Source/Runtime/Experimental/Chaos/Private/Chaos/Evolution/SolverBodyContainer.cpp b/Engine/Source/Runtime/Experimental/Chaos/Private/Chaos/Evolution/SolverBodyContainer.cpp
--- a/Engine/Source/Runtime/Experimental/Chaos/Private/Chaos/Evolution/SolverBodyContainer.cpp
+++ b/Engine/Source/Runtime/Experimental/Chaos/Private/Chaos/Evolution/SolverBodyContainer.cpp
@@ -4,37 +4,6 @@
 
 namespace Chaos
 {
-	void FSolverBodyAdapter::GatherInput()
-	{
-		if (Particle.IsValid())
-		{
-			FRigidTransform3 CoMTransform = FParticleUtilitiesPQ::GetCoMWorldTransform(Particle);
-			SolverBody.SetP(CoMTransform.GetLocation());
-			SolverBody.SetQ(CoMTransform.GetRotation());
-			SolverBody.SetV(Particle->V());
-			SolverBody.SetW(Particle->W());
-			SolverBody.SetCoM(Particle->CenterOfMass());
-			SolverBody.SetRoM(Particle->RotationOfMass());
-
-			if (Particle->IsDynamic())
-			{
-				FRigidTransform3 PrevCoMTransform = FParticleUtilitiesXR::GetCoMWorldTransform(Particle);
-				SolverBody.SetX(PrevCoMTransform.GetLocation());
-				SolverBody.SetR(PrevCoMTransform.GetRotation());
-
-				SolverBody.SetInvM(Particle->InvM());
-				SolverBody.SetInvILocal(Particle->InvI().GetDiagonal());
-			}
-			else
-			{
-				SolverBody.SetX(SolverBody.P());
-				SolverBody.SetR(SolverBody.Q());
-			}
-
-			SolverBody.UpdateRotationDependentState();
-		}
-	}
-
 	void FSolverBodyAdapter::ScatterOutput()
 	{
 		if (Particle.IsValid())
@@ -65,7 +34,7 @@ namespace Chaos
 		// For kinematics we cannot do this because the kinematic may be in multiple islands and 
 		// would require a different index for each island, so we use a local map instead. 
 		int32 ItemIndex = InParticle->SolverBodyIndex();
-
+	
 		if (ItemIndex == INDEX_NONE)
 		{
 			if (InParticle->IsDynamic())
@@ -89,7 +58,7 @@ namespace Chaos
 				}
 			}			
 		}
-
+	
 		check(ItemIndex != INDEX_NONE);
 		return &SolverBodies[ItemIndex].GetSolverBody();
 	}
diff --git a/Engine/Source/Runtime/Experimental/Chaos/Private/Chaos/PBDCollisionConstraints.cpp b/Engine/Source/Runtime/Experimental/Chaos/Private/Chaos/PBDCollisionConstraints.cpp
--- a/Engine/Source/Runtime/Experimental/Chaos/Private/Chaos/PBDCollisionConstraints.cpp
+++ b/Engine/Source/Runtime/Experimental/Chaos/Private/Chaos/PBDCollisionConstraints.cpp
@@ -357,7 +357,7 @@ namespace Chaos
 		if (SolverType == EConstraintSolverType::QuasiPbd)
 		{
 			FPBDCollisionSolverContainer& SolverContainer = GetConstraintSolverContainer(SolverData);
-			SolverContainer.Reset(NumIslandConstraints);
+			SolverContainer.SetNum(NumIslandConstraints);
 
 			// @todo(chaos): we shouldn't have to pass settings down to the solvers every frame. Ideally the solver containers would be created by the constraint containers and
 			// the settings would be passed on then. Currently the solver containers are created here FPBDIslandSolverData::AddConstraintDatas
@@ -383,7 +383,7 @@ namespace Chaos
 			check(Constraint.IsEnabled());
 
 			FPBDCollisionSolverContainer& SolverContainer = GetConstraintSolverContainer(SolverData);
-			SolverContainer.AddConstraintSolver(Dt, Constraint, Particle0Level, Particle1Level, SolverData.GetBodyContainer());
+			SolverContainer.AddConstraintSolver(Dt, Constraint, Particle0Level, Particle1Level, SolverData.GetBodyContainer(), SolverData.GetConstraintIndex(ContainerId));
 		}
 		else
 		{
diff --git a/Engine/Source/Runtime/Experimental/Chaos/Private/Chaos/PBDConstraintRule.cpp b/Engine/Source/Runtime/Experimental/Chaos/Private/Chaos/PBDConstraintRule.cpp
--- a/Engine/Source/Runtime/Experimental/Chaos/Private/Chaos/PBDConstraintRule.cpp
+++ b/Engine/Source/Runtime/Experimental/Chaos/Private/Chaos/PBDConstraintRule.cpp
@@ -259,7 +259,8 @@ namespace Chaos
 					int32 ConstraintSetEnd = 0;
 					
 					Constraints.SetNumIslandConstraints(ConstraintGraph->GetIslandConstraints(Island).Num(), *IslandSolver);
-					
+					IslandSolver->InitConstraintIndex(Constraints.GetContainerId());
+
 					for (int32 Level = 0; Level < MaxLevel; ++Level)
 					{
 						for (int32 Color = 0; Color < MaxColor; ++Color)
diff --git a/Engine/Source/Runtime/Experimental/Chaos/Private/Chaos/PBDRigidsEvolutionGBF.cpp b/Engine/Source/Runtime/Experimental/Chaos/Private/Chaos/PBDRigidsEvolutionGBF.cpp
--- a/Engine/Source/Runtime/Experimental/Chaos/Private/Chaos/PBDRigidsEvolutionGBF.cpp
+++ b/Engine/Source/Runtime/Experimental/Chaos/Private/Chaos/PBDRigidsEvolutionGBF.cpp
@@ -525,7 +525,39 @@ void FPBDRigidsEvolutionGBF::GatherSolverInput(FReal Dt, int32 Island)
 {
 	// We must initialize the solver body container to be large enough to hold all particles in the
 	// island so that the pointers remain valid (the array should not grow and relocate)
-	ConstraintGraph.GetSolverIsland(Island)->GetBodyContainer().Reset(ConstraintGraph.GetIslandParticles(Island).Num());
+	FPBDIslandSolver* SolverIsland = ConstraintGraph.GetSolverIsland(Island);
+	if(SolverIsland)
+	{
+		// We pre-add all the dynamic particles solver bodies involved in any constraint to speed up the gather 
+		const TArray<FGeometryParticleHandle*>& IslandParticles = SolverIsland->GetParticles();
+		SolverIsland->GetBodyContainer().Reset(IslandParticles.Num());
+
+		if(!SolverIsland->IsSleeping())
+		{
+			for(auto& Particle : IslandParticles)
+			{
+				FGenericParticleHandle GenericHandle = Particle->Handle();
+				GenericHandle->SetSolverBodyIndex(INDEX_NONE);
+				
+				if(GenericHandle->ObjectState() == EObjectStateType::Dynamic)
+				{
+					if(FPBDRigidParticleHandle* PBDRigid = Particle->Handle()->CastToRigidParticle())
+					{
+						if(ConstraintGraph.GetGraphNodes().IsValidIndex(PBDRigid->ConstraintGraphIndex()))
+						{
+							if(ConstraintGraph.GetGraphNodes()[PBDRigid->ConstraintGraphIndex()].NodeEdges.Num() > 0)
+							{
+								// Only the non sleeping dynamic rigid particles involved in any constraints will be added
+								// by default to the solver bodies
+								const int32 BodyIndex = SolverIsland->GetBodyContainer().GetBodies().Emplace(GenericHandle);
+								GenericHandle->SetSolverBodyIndex(BodyIndex);
+							}
+						}
+					}
+				}
+			}
+		}
+	}
 
 	// NOTE: SolverBodies are gathered as part of the constraint gather, in the order that they are first seen 
 	for (FPBDConstraintGraphRule* ConstraintRule : PrioritizedConstraintRules)
diff --git a/Engine/Source/Runtime/Experimental/Chaos/Public/Chaos/Collision/SolverCollisionContainer.h b/Engine/Source/Runtime/Experimental/Chaos/Public/Chaos/Collision/SolverCollisionContainer.h
--- a/Engine/Source/Runtime/Experimental/Chaos/Public/Chaos/Collision/SolverCollisionContainer.h
+++ b/Engine/Source/Runtime/Experimental/Chaos/Public/Chaos/Collision/SolverCollisionContainer.h
@@ -35,8 +35,13 @@ namespace Chaos
 
 		virtual void Reset(const int32 InMaxCollisions) override;
 
+		/** Set the number of collision solvers in the container
+		 * @param MaxCollisions Max number of collision solvers
+		 */
+		void SetNum(const int32 MaxCollisions);
+
 		// Add a solver constraint for the specified constraint and gather the required solver data
-		void AddConstraintSolver(FReal Dt, FPBDCollisionConstraint& Constraint, const int32 Particle0Level, const int32 Particle1Level, FSolverBodyContainer& SolverBodyContainer);
+		void AddConstraintSolver(FReal Dt, FPBDCollisionConstraint& Constraint, const int32 Particle0Level, const int32 Particle1Level, FSolverBodyContainer& SolverBodyContainer, int32& ConstraintIndex);
 
 		bool SolvePositionSerial(const FReal Dt, const int32 It, const int32 NumIts, const int32 BeginIndex, const int32 EndIndex);
 		bool SolvePositionParallel(const FReal Dt, const int32 It, const int32 NumIts, const int32 BeginIndex, const int32 EndIndex);
diff --git a/Engine/Source/Runtime/Experimental/Chaos/Public/Chaos/CollisionResolutionUtil.h b/Engine/Source/Runtime/Experimental/Chaos/Public/Chaos/CollisionResolutionUtil.h
--- a/Engine/Source/Runtime/Experimental/Chaos/Public/Chaos/CollisionResolutionUtil.h
+++ b/Engine/Source/Runtime/Experimental/Chaos/Public/Chaos/CollisionResolutionUtil.h
@@ -36,6 +36,23 @@ namespace Chaos
 				-V[1] * (-V[1] * M.M[0][0] + V[0] * M.M[1][0]) + V[0] * (-V[1] * M.M[1][0] + V[0] * M.M[1][1]) + Im);
 		}
 
+		/** Compute the diagonal part of the rigid objects rotational contribution
+		 * @param V Position relative to the center of mass
+		 * @param M Mass of the particle
+		 * @param Im Inertia tensor of the particle
+		 */
+		template<typename TRealType>
+		inline TVec3<TRealType> ComputeFactorDiagonal3(const TVec3<TRealType>& V, const TMatrix33<TRealType>& M, const TRealType Im)
+		{
+			// Rigid objects rotational contribution diagonal to the impulse.
+			// Vx*M*VxT+Im
+			check(Im > FLT_MIN);
+			return TVec3<TRealType>(
+				-V[2] * (-V[2] * M.M[1][1] + V[1] * M.M[2][1]) + V[1] * (-V[2] * M.M[2][1] + V[1] * M.M[2][2]) + Im,
+				 V[2] * ( V[2] * M.M[0][0] - V[0] * M.M[2][0]) - V[0] * ( V[2] * M.M[2][0] - V[0] * M.M[2][2]) + Im,
+				-V[1] * (-V[1] * M.M[0][0] + V[0] * M.M[1][0]) + V[0] * (-V[1] * M.M[1][0] + V[0] * M.M[1][1]) + Im);
+		}
+
 		FVec3
 		GetEnergyClampedImpulse(const TPBDRigidParticleHandle<FReal, 3>* PBDRigid0, const TPBDRigidParticleHandle<FReal, 3>* PBDRigid1, const FVec3& Impulse, const FVec3& VectorToPoint1, const FVec3& VectorToPoint2, const FVec3& Velocity1, const FVec3& Velocity2);
 
diff --git a/Engine/Source/Runtime/Experimental/Chaos/Public/Chaos/Evolution/SolverBodyContainer.h b/Engine/Source/Runtime/Experimental/Chaos/Public/Chaos/Evolution/SolverBodyContainer.h
--- a/Engine/Source/Runtime/Experimental/Chaos/Public/Chaos/Evolution/SolverBodyContainer.h
+++ b/Engine/Source/Runtime/Experimental/Chaos/Public/Chaos/Evolution/SolverBodyContainer.h
@@ -4,6 +4,7 @@
 #include "Chaos/Evolution/SolverBody.h"
 #include "Chaos/Framework/ContainerItemHandle.h"
 #include "Chaos/ParticleHandle.h"
+#include "Chaos/Particle/ParticleUtilities.h"
 
 #include "Containers/Array.h"
 #include "Containers/Map.h"
@@ -28,9 +29,38 @@ namespace Chaos
 		}
 		
 		FSolverBody& GetSolverBody() { return SolverBody; }
-		const FGenericParticleHandle& GetParticle() { return Particle; }
+		const FGenericParticleHandle& GetParticle() const { return Particle; }
 
-		void GatherInput();
+		FORCEINLINE void GatherInput()
+		{
+			if (Particle.IsValid())
+			{
+				FRigidTransform3 CoMTransform = FParticleUtilitiesPQ::GetCoMWorldTransform(Particle);
+				SolverBody.SetP(CoMTransform.GetLocation());
+				SolverBody.SetQ(CoMTransform.GetRotation());
+				SolverBody.SetV(Particle->V());
+				SolverBody.SetW(Particle->W());
+				SolverBody.SetCoM(Particle->CenterOfMass());
+				SolverBody.SetRoM(Particle->RotationOfMass());
+
+				if (Particle->IsDynamic())
+				{
+					FRigidTransform3 PrevCoMTransform = FParticleUtilitiesXR::GetCoMWorldTransform(Particle);
+					SolverBody.SetX(PrevCoMTransform.GetLocation());
+					SolverBody.SetR(PrevCoMTransform.GetRotation());
+
+					SolverBody.SetInvM(Particle->InvM());
+					SolverBody.SetInvILocal(Particle->InvI().GetDiagonal());
+				}
+				else
+				{
+					SolverBody.SetX(SolverBody.P());
+					SolverBody.SetR(SolverBody.Q());
+				}
+				// No need to update the UpdateRotationDependentState since this function is only
+				// valid for dynamic particle for which the SetInvILocal is already doing the job
+			}
+		}
 		void ScatterOutput();
 
 	private:
@@ -116,6 +146,12 @@ namespace Chaos
 			return SolverBodies[Index];
 		}
 
+		/** Return all the solver bodies inside the container */
+		inline TArray<FSolverBodyAdapter>& GetBodies()
+		{
+			return SolverBodies;
+		}
+
 		// Whether the specified index is valid.
 		// @param Index The index of the body to return
 		inline bool IsValid(int Index) const
diff --git a/Engine/Source/Runtime/Experimental/Chaos/Public/Chaos/Evolution/SolverDatas.h b/Engine/Source/Runtime/Experimental/Chaos/Public/Chaos/Evolution/SolverDatas.h
--- a/Engine/Source/Runtime/Experimental/Chaos/Public/Chaos/Evolution/SolverDatas.h
+++ b/Engine/Source/Runtime/Experimental/Chaos/Public/Chaos/Evolution/SolverDatas.h
@@ -52,6 +52,13 @@ public :
 	template<typename ConstraintType> const ConstraintType* GetConstraintHandle(const int32 ContainerId, const int32 ConstraintIndex) const;
 	template<typename ConstraintType> ConstraintType* GetConstraintHandle(const int32 ContainerId, const int32 ConstraintIndex);
 
+	/** Accessors of the current constraint index */
+	const int32& GetConstraintIndex(const int32 ContainerId) const { return ConstraintDatas[ContainerId].ConstraintIndex; };
+	int32& GetConstraintIndex(const int32 ContainerId) { return ConstraintDatas[ContainerId].ConstraintIndex; };
+
+	/** Init the current constraint index */
+	void InitConstraintIndex(const int32 ContainerId) { ConstraintDatas[ContainerId].ConstraintIndex = 0; };
+
 	/** Add a constraint container to the solver island given a container id
 	* @param ContainerId Constraints container id from which the solver constraint datas is being built
 	*/
@@ -68,6 +75,9 @@ protected:
 
 		/** Constraint handles that will be used for legacy solver */
 		TArray<FConstraintHandle*> ConstraintHandles;
+
+		/** Current constraint index in the constraint container */
+		int32 ConstraintIndex;
 	};
 
 	/** Solver body container of that datas */
