Chaos - Fix average contact velocity solve NetPushOutNormal
- Also make it use cached angular axes and eliminate some redundant work

diff --git a/Engine/Source/Runtime/Experimental/Chaos/Private/Chaos/Collision/PBDCollisionSolver.cpp b/Engine/Source/Runtime/Experimental/Chaos/Private/Chaos/Collision/PBDCollisionSolver.cpp
--- a/Engine/Source/Runtime/Experimental/Chaos/Private/Chaos/Collision/PBDCollisionSolver.cpp
+++ b/Engine/Source/Runtime/Experimental/Chaos/Private/Chaos/Collision/PBDCollisionSolver.cpp
@@ -49,6 +49,7 @@ namespace Chaos
 		bool bChaos_PBDCollisionSolver_Velocity_DynamicFrictionEnabled = true;
 		bool bChaos_PBDCollisionSolver_Velocity_NegativeImpulseEnabled = true;
 		bool bChaos_PBDCollisionSolver_Velocity_ImpulseClampEnabled = true;
+		bool bChaos_PBDCollisionSolver_Velocity_AveragePointEnabled = true;
 
 		FAutoConsoleVariableRef CVarChaos_PBDCollisionSolver_Velocity_SolveEnabled(TEXT("p.Chaos.PBDCollisionSolver.Velocity.SolveEnabled"), bChaos_PBDCollisionSolver_Velocity_SolveEnabled, TEXT(""));
 		FAutoConsoleVariableRef CVarChaos_PBDCollisionSolver_Velocity_UseShockPropagation(TEXT("p.Chaos.PBDCollisionSolver.Velocity.ShockPropagationIterations"), Chaos_PBDCollisionSolver_Velocity_ShockPropagationIterations, TEXT(""));
@@ -56,6 +57,7 @@ namespace Chaos
 		FAutoConsoleVariableRef CVarChaos_PBDCollisionSolver_Velocity_DynamicFrictionEnabled(TEXT("p.Chaos.PBDCollisionSolver.Velocity.DynamicFrictionEnabled"), bChaos_PBDCollisionSolver_Velocity_DynamicFrictionEnabled, TEXT(""));
 		FAutoConsoleVariableRef CVarChaos_PBDCollisionSolver_Velocity_NegativeImpulseEnabled(TEXT("p.Chaos.PBDCollisionSolver.Velocity.NegativeImpulseEnabled"), bChaos_PBDCollisionSolver_Velocity_NegativeImpulseEnabled, TEXT(""));
 		FAutoConsoleVariableRef CVarChaos_PBDCollisionSolver_Velocity_ImpulseClampEnabled(TEXT("p.Chaos.PBDCollisionSolver.Velocity.ImpulseClampEnabled"), bChaos_PBDCollisionSolver_Velocity_ImpulseClampEnabled, TEXT(""));
+		FAutoConsoleVariableRef CVarChaos_PBDCollisionSolver_Velocity_AveragePointEnabled(TEXT("p.Chaos.PBDCollisionSolver.Velocity.AveragePointEnabled"), bChaos_PBDCollisionSolver_Velocity_AveragePointEnabled, TEXT(""));
 	}
 	using namespace CVars;
 
@@ -318,15 +320,18 @@ namespace Chaos
 
 		// Generate a new contact point at the average of all the active contacts
 		int32 NumActiveManifoldPoints = 0;
-		FSolverVec3 RelativeContactPosition0 = FVec3(0);
-		FSolverVec3 RelativeContactPosition1 = FVec3(0);
+		FSolverVec3 RelativeContactPosition0 = FSolverVec3(0);
+		FSolverVec3 RelativeContactPosition1 = FSolverVec3(0);
 		FSolverVec3 WorldContactNormal = FSolverVec3(0);
 		FSolverReal NetPushOutNormal = FSolverReal(0);
-		FSolverReal WorldContactVelocityTargetNormal = FReal(0);
+		FSolverReal WorldContactVelocityTargetNormal = FSolverReal(0);
 		for (int32 PointIndex = 0; PointIndex < State.NumManifoldPoints; ++PointIndex)
 		{
 			FPBDCollisionSolverManifoldPoint& SolverManifoldPoint = State.ManifoldPoints[PointIndex];
-			if (!FMath::IsNearlyZero(SolverManifoldPoint.NetPushOutNormal))
+
+			// An "active" point here is defined as one that started out penetrating or applied a net impulse.
+			const bool bUsePoint = (SolverManifoldPoint.WorldContactDeltaNormal < 0) || (SolverManifoldPoint.NetPushOutNormal > FSolverReal(SMALL_NUMBER));
+			if (bUsePoint)
 			{
 				RelativeContactPosition0 += SolverManifoldPoint.RelativeContactPosition0;
 				RelativeContactPosition1 += SolverManifoldPoint.RelativeContactPosition1;
@@ -345,7 +350,6 @@ namespace Chaos
 		// than we would have if we just solved the corners of the box in sequence.
 		if (NumActiveManifoldPoints > 1)
 		{
-			const FSolverReal DynamicFriction = FSolverReal(0);
 			const FSolverReal Scale = FSolverReal(1) / FSolverReal(NumActiveManifoldPoints);
 
 			FPBDCollisionSolverManifoldPoint AverageManifoldPoint;
@@ -353,20 +357,42 @@ namespace Chaos
 			AverageManifoldPoint.RelativeContactPosition1 = RelativeContactPosition1 * Scale;
 			AverageManifoldPoint.WorldContactNormal = WorldContactNormal;
 			AverageManifoldPoint.WorldContactVelocityTargetNormal = WorldContactVelocityTargetNormal * Scale;
-			AverageManifoldPoint.NetPushOutNormal = NetPushOutNormal * Scale;
-
-			const FSolverMatrix33 ContactMassInv =
-				(Body0.IsDynamic() ? Collisions::ComputeFactorMatrix3(AverageManifoldPoint.RelativeContactPosition0, Body0.InvI(), Body0.InvM()) : FSolverMatrix33(0)) +
-				(Body1.IsDynamic() ? Collisions::ComputeFactorMatrix3(AverageManifoldPoint.RelativeContactPosition1, Body1.InvI(), Body1.InvM()) : FSolverMatrix33(0));
-			const FSolverReal ContactMassInvNormal = FSolverVec3::DotProduct(AverageManifoldPoint.WorldContactNormal, Utilities::Multiply(ContactMassInv, AverageManifoldPoint.WorldContactNormal));
-			const FSolverReal ContactMassNormal = (ContactMassInvNormal > FSolverReal(SMALL_NUMBER)) ? FSolverReal(1) / ContactMassInvNormal : FSolverReal(0);
-			AverageManifoldPoint.ContactMassNormal = ContactMassNormal;
+			
+			// Total pushout (not average) which is correct if the average point is also the centroid but otherwise probably an overestimate.
+			// This is used to limit the possibly-attractive impulse that corrects implicit velocity errors from the PBD solve, but it
+			// which might cause stickiness if overestimated. Although it should be recitifed when we solve the normal contact point velocities
+			// so I don't think it matters...we could do better though if it is a problem
+			AverageManifoldPoint.NetPushOutNormal = NetPushOutNormal;
+
+			// Calculate the contact mass (and derived properties) for this point
+			FSolverVec3 WorldContactNormalAngular0 = FSolverVec3(0);
+			FSolverVec3 WorldContactNormalAngular1 = FSolverVec3(0);
+			FSolverReal ContactMassInvNormal = FSolverReal(0);
+			if (Body0.IsDynamic())
+			{
+				const FSolverVec3 R0xN = FSolverVec3::CrossProduct(RelativeContactPosition0, WorldContactNormal);
+				WorldContactNormalAngular0 = Body0.InvI() * R0xN;
+				ContactMassInvNormal += FSolverVec3::DotProduct(R0xN, WorldContactNormalAngular0) + Body0.InvM();
+			}
+			if (Body1.IsDynamic())
+			{
+				const FSolverVec3 R1xN = FSolverVec3::CrossProduct(RelativeContactPosition1, WorldContactNormal);
+				WorldContactNormalAngular1 = Body1.InvI() * R1xN;
+				ContactMassInvNormal += FSolverVec3::DotProduct(R1xN, WorldContactNormalAngular1) + Body1.InvM();
+			}
+			AverageManifoldPoint.ContactMassNormal = (ContactMassInvNormal > FSolverReal(SMALL_NUMBER)) ? FSolverReal(1) / ContactMassInvNormal : FSolverReal(0);
+			AverageManifoldPoint.WorldContactNormalAngular0 = WorldContactNormalAngular0;
+			AverageManifoldPoint.WorldContactNormalAngular1 = WorldContactNormalAngular1;
 
-			// @todo(chaos): All unused for the average contact solve - optimize
+			// @todo(chaos): we don't use these - maybe do the calculation without an actual manifold point object...
 			AverageManifoldPoint.ContactMassTangentU = FSolverReal(0);
 			AverageManifoldPoint.ContactMassTangentV = FSolverReal(0);
 			AverageManifoldPoint.WorldContactTangentU = FSolverVec3(0);
 			AverageManifoldPoint.WorldContactTangentV = FSolverVec3(0);
+			AverageManifoldPoint.WorldContactTangentUAngular0 = FSolverVec3(0);
+			AverageManifoldPoint.WorldContactTangentVAngular0 = FSolverVec3(0);
+			AverageManifoldPoint.WorldContactTangentUAngular1 = FSolverVec3(0);
+			AverageManifoldPoint.WorldContactTangentVAngular1 = FSolverVec3(0);
 			AverageManifoldPoint.WorldContactDeltaNormal = FSolverReal(0);
 			AverageManifoldPoint.WorldContactDeltaTangent0 = FSolverReal(0);
 			AverageManifoldPoint.WorldContactDeltaTangent1 = FSolverReal(0);
@@ -378,17 +404,13 @@ namespace Chaos
 			AverageManifoldPoint.FrictionNetPushOutNormal = FSolverReal(0);
 			AverageManifoldPoint.bInsideStaticFrictionCone = false;
 
-			// @todo(chaos): don't need tangential velocity delta here since we are setting DynamicFriction to 0
-			FSolverReal ContactVelocityDeltaNormal, ContactVelocityDeltaTangentU, ContactVelocityDeltaTangentV;
-			AverageManifoldPoint.CalculateContactVelocityError(Body0, Body1, DynamicFriction, Dt, ContactVelocityDeltaNormal, ContactVelocityDeltaTangentU, ContactVelocityDeltaTangentV);
+			FSolverReal ContactVelocityDeltaNormal;
+			AverageManifoldPoint.CalculateContactVelocityErrorNormal(Body0, Body1, ContactVelocityDeltaNormal);
 
-			ApplyVelocityCorrection(
+			ApplyVelocityCorrectionNormal(
 				State.Stiffness,
 				Dt,
-				DynamicFriction,
 				ContactVelocityDeltaNormal,
-				ContactVelocityDeltaTangentU,
-				ContactVelocityDeltaTangentV,
 				AverageManifoldPoint,
 				Body0,
 				Body1);
@@ -400,7 +422,8 @@ namespace Chaos
 		// Apply restitution at the average contact point
 		// This means we don't need to run as many iterations to get stable bouncing
 		// It also helps with zero restitution to counter any velocioty added by the PBD solve
-		if (State.NumManifoldPoints > 1)
+		const bool bSolveAverageContact = (State.NumManifoldPoints > 1) && bChaos_PBDCollisionSolver_Velocity_AveragePointEnabled;
+		if (bSolveAverageContact)
 		{
 			SolveVelocityAverage(Dt);
 		}
diff --git a/Engine/Source/Runtime/Experimental/Chaos/Private/Chaos/Collision/PBDCollisionSolver.h b/Engine/Source/Runtime/Experimental/Chaos/Private/Chaos/Collision/PBDCollisionSolver.h
--- a/Engine/Source/Runtime/Experimental/Chaos/Private/Chaos/Collision/PBDCollisionSolver.h
+++ b/Engine/Source/Runtime/Experimental/Chaos/Private/Chaos/Collision/PBDCollisionSolver.h
@@ -87,6 +87,7 @@ namespace Chaos
 		/**
 		 * @brief Calculate the velocity error at the current transforms
 		*/
+		void CalculateContactVelocityErrorNormal(const FConstraintSolverBody& Body0, const FConstraintSolverBody& Body1, FSolverReal& OutContactVelocityDeltaNormal) const;
 		void CalculateContactVelocityError(const FConstraintSolverBody& Body0, const FConstraintSolverBody& Body1, const FSolverReal DynamicFriction, const FSolverReal Dt, FSolverReal& OutContactVelocityDeltaNormal, FSolverReal& OutContactVelocityDeltaTangent0, FSolverReal& OutContactVelocityDeltaTangent1) const;
 
 		// @todo(chaos): make private
@@ -474,7 +475,7 @@ namespace Chaos
 		}
 		if (Body1.IsDynamic())
 		{
-			const FSolverVec3 DX1 = -(Body1.InvM() * PushOut);
+			const FSolverVec3 DX1 = -Body1.InvM() * PushOut;
 			const FSolverVec3 DR1 = ManifoldPoint.WorldContactNormalAngular1 * -PushOutNormal + ManifoldPoint.WorldContactTangentUAngular1 * -PushOutTangentU + ManifoldPoint.WorldContactTangentVAngular1 * -PushOutTangentV;
 			Body1.ApplyPositionDelta(DX1);
 			Body1.ApplyRotationDelta(DR1);
@@ -532,7 +533,7 @@ namespace Chaos
 		
 		// Clamp the total impulse to be positive along the normal. We can apply a net negative impulse, 
 		// but only to correct the velocity that was added by pushout.
-		FSolverReal NetImpulseNormal = ManifoldPoint.NetImpulseNormal + ImpulseNormal;
+		const FSolverReal NetImpulseNormal = ManifoldPoint.NetImpulseNormal + ImpulseNormal;
 		const FSolverReal PushOutImpulseNormal = FMath::Max(0.0f, ManifoldPoint.NetPushOutNormal / Dt);
 		if (NetImpulseNormal < -PushOutImpulseNormal)
 		{
@@ -560,13 +561,52 @@ namespace Chaos
 		// Calculate the velocity deltas from the impulse
 		if (Body0.IsDynamic())
 		{
-			const FSolverVec3 AngularImpulse = FSolverVec3::CrossProduct(ManifoldPoint.RelativeContactPosition0, Impulse);
-			Body0.ApplyVelocityDelta(Body0.InvM() * Impulse, Body0.InvI() * AngularImpulse);
+			const FSolverVec3 DV0 = Body0.InvM() * Impulse;
+			const FSolverVec3 DW0 = ManifoldPoint.WorldContactNormalAngular0 * ImpulseNormal + ManifoldPoint.WorldContactTangentUAngular0 * ImpulseTangentU + ManifoldPoint.WorldContactTangentVAngular0 * ImpulseTangentV;
+			Body0.ApplyVelocityDelta(DV0, DW0);
 		}
 		if (Body1.IsDynamic())
 		{
-			const FSolverVec3 AngularImpulse = FSolverVec3::CrossProduct(ManifoldPoint.RelativeContactPosition1, Impulse);
-			Body1.ApplyVelocityDelta(Body1.InvM() * -Impulse, Body1.InvI() * -AngularImpulse);
+			const FSolverVec3 DV1 = -Body1.InvM() * Impulse;
+			const FSolverVec3 DW1 = ManifoldPoint.WorldContactNormalAngular1 * -ImpulseNormal + ManifoldPoint.WorldContactTangentUAngular1 * -ImpulseTangentU + ManifoldPoint.WorldContactTangentVAngular1 * -ImpulseTangentV;
+			Body1.ApplyVelocityDelta(DV1, DW1);
+		}
+	}
+
+	FORCEINLINE_DEBUGGABLE void ApplyVelocityCorrectionNormal(
+		const FSolverReal Stiffness,
+		const FSolverReal Dt,
+		const FSolverReal ContactVelocityDeltaNormal,
+		FPBDCollisionSolverManifoldPoint& ManifoldPoint,
+		FConstraintSolverBody& Body0,
+		FConstraintSolverBody& Body1)
+	{
+		FSolverReal ImpulseNormal = -(Stiffness * ManifoldPoint.ContactMassNormal) * ContactVelocityDeltaNormal;
+
+		// See comments in ApplyVelocityCorrection
+		const FSolverReal NetImpulseNormal = ManifoldPoint.NetImpulseNormal + ImpulseNormal;
+		const FSolverReal PushOutImpulseNormal = FMath::Max(0.0f, ManifoldPoint.NetPushOutNormal / Dt);
+		if (NetImpulseNormal < -PushOutImpulseNormal)
+		{
+			ImpulseNormal = ImpulseNormal - (NetImpulseNormal + PushOutImpulseNormal);
+		}
+
+		ManifoldPoint.NetImpulseNormal += ImpulseNormal;
+
+		FSolverVec3 Impulse = ImpulseNormal * ManifoldPoint.WorldContactNormal;
+
+		// Calculate the velocity deltas from the impulse
+		if (Body0.IsDynamic())
+		{
+			const FSolverVec3 DV0 = Body0.InvM() * Impulse;
+			const FSolverVec3 DW0 = ManifoldPoint.WorldContactNormalAngular0 * ImpulseNormal;
+			Body0.ApplyVelocityDelta(DV0, DW0);
+		}
+		if (Body1.IsDynamic())
+		{
+			const FSolverVec3 DV1 = -Body1.InvM() * Impulse;
+			const FSolverVec3 DW1 = ManifoldPoint.WorldContactNormalAngular1 * -ImpulseNormal;
+			Body1.ApplyVelocityDelta(DV1, DW1);
 		}
 	}
 
@@ -627,9 +667,6 @@ namespace Chaos
 		const FSolverVec3 ContactVelocity = ContactVelocity0 - ContactVelocity1;
 		const FSolverReal ContactVelocityNormal = FSolverVec3::DotProduct(ContactVelocity, WorldContactNormal);
 
-		// Target normal velocity, including restitution
-		const FSolverReal ContactVelocityTargetNormal = WorldContactVelocityTargetNormal;
-
 		// Add up the errors in the velocity (current velocity - desired velocity)
 		OutContactVelocityDeltaNormal = (ContactVelocityNormal - WorldContactVelocityTargetNormal);
 		OutContactVelocityDeltaTangent0 = FSolverReal(0);
@@ -660,6 +697,16 @@ namespace Chaos
 		}
 	}
 
+	FORCEINLINE_DEBUGGABLE void FPBDCollisionSolverManifoldPoint::CalculateContactVelocityErrorNormal(const FConstraintSolverBody& Body0, const FConstraintSolverBody& Body1, FSolverReal& OutContactVelocityDeltaNormal) const
+	{
+		const FSolverVec3 ContactVelocity0 = Body0.V() + FSolverVec3::CrossProduct(Body0.W(), RelativeContactPosition0);
+		const FSolverVec3 ContactVelocity1 = Body1.V() + FSolverVec3::CrossProduct(Body1.W(), RelativeContactPosition1);
+		const FSolverVec3 ContactVelocity = ContactVelocity0 - ContactVelocity1;
+		const FSolverReal ContactVelocityNormal = FSolverVec3::DotProduct(ContactVelocity, WorldContactNormal);
+
+		// Add up the errors in the velocity (current velocity - desired velocity)
+		OutContactVelocityDeltaNormal = (ContactVelocityNormal - WorldContactVelocityTargetNormal);
+	}
 
 	//////////////////////////////////////////////////////////////////////////////////////////////////
 	//////////////////////////////////////////////////////////////////////////////////////////////////
