CHaos - BackPort size specific data for geometry collection
CL 15260032 (Size specific damage )
CL 16429701 (step 1 - size specific parameters restructuring )
CL 16439958 ( review fixes from 16429701 )
CL 16456707 (transfer defaults to size specific data)
CL 16476610 ( Force Size Specific data to always have a single element )
CL 16555840 ( Allow empty collision arrays )

diff --git a/Engine/Source/Runtime/Core/Public/UObject/PhysicsObjectVersion.h b/Engine/Source/Runtime/Core/Public/UObject/PhysicsObjectVersion.h
--- a/Engine/Source/Runtime/Core/Public/UObject/PhysicsObjectVersion.h
+++ b/Engine/Source/Runtime/Core/Public/UObject/PhysicsObjectVersion.h
@@ -51,6 +51,9 @@ struct CORE_API FPhysicsObjectVersion
 		// Convex structure data has a list of unique edges (half of the half edges)
 		ChaosConvexHasUniqueEdgeSet,
 
+		// Chaos FGeometryCollectionObject user defined collision shapes support
+		GeometryCollectionUserDefinedCollisionShapes,
+
 		// -----<new versions can be added above this line>-------------------------------------------------
 		VersionPlusOne,
 		LatestVersion = VersionPlusOne - 1
diff --git a/Engine/Source/Runtime/Experimental/GeometryCollectionEngine/Private/GeometryCollection/GeometryCollectionObject.cpp b/Engine/Source/Runtime/Experimental/GeometryCollectionEngine/Private/GeometryCollection/GeometryCollectionObject.cpp
--- a/Engine/Source/Runtime/Experimental/GeometryCollectionEngine/Private/GeometryCollection/GeometryCollectionObject.cpp
+++ b/Engine/Source/Runtime/Experimental/GeometryCollectionEngine/Private/GeometryCollection/GeometryCollectionObject.cpp
@@ -141,15 +141,22 @@ FGeometryCollectionSizeSpecificData::FGeometryCollectionSizeSpecificData()
 bool FGeometryCollectionSizeSpecificData::Serialize(FArchive& Ar)
 {
 	Ar.UsingCustomVersion(FUE5MainStreamObjectVersion::GUID);
+	Ar.UsingCustomVersion(FPhysicsObjectVersion::GUID);
 	return false;	//We only have this function to mark custom GUID. Still want serialize tagged properties
 }
 
 #if WITH_EDITORONLY_DATA
 void FGeometryCollectionSizeSpecificData::PostSerialize(const FArchive& Ar)
 {
-	const int32 CustomVer = Ar.CustomVer(FUE5MainStreamObjectVersion::GUID);
+	const int32 PhysicsObjectVersion = Ar.CustomVer(FPhysicsObjectVersion::GUID);
+	const int32 StreamObjectVersion = Ar.CustomVer(FUE5MainStreamObjectVersion::GUID);
 	// make sure to load back the deprecated values in the new structure if necessary
-	if (Ar.IsLoading() && CustomVer < FUE5MainStreamObjectVersion::GeometryCollectionUserDefinedCollisionShapes)
+	// IMPORTANT : this was merge backed in UE4 and PhysicsObjectVersion had to be used,
+	//		that's why we need to test both version to make sure backward asset compatibility is maintained
+	if (Ar.IsLoading() && (
+		StreamObjectVersion < FUE5MainStreamObjectVersion::GeometryCollectionUserDefinedCollisionShapes &&
+		PhysicsObjectVersion < FPhysicsObjectVersion::GeometryCollectionUserDefinedCollisionShapes
+		)) 
 	{
 		if (CollisionShapes.Num())
 		{
