Add framework for optimizing joint solver
- Add PBDJointCachedSolverGaussSeidel class
- Add CVar to toggle between PBDJointSolverGaussSeidel and PBDJointCachedSolverGaussSeidel
- Cache data for point, spherical, cylinder, planar constraints

diff --git a/Engine/Source/Runtime/Experimental/Chaos/Private/Chaos/Joint/JointConstraintsCVars.cpp b/Engine/Source/Runtime/Experimental/Chaos/Private/Chaos/Joint/JointConstraintsCVars.cpp
--- a/Engine/Source/Runtime/Experimental/Chaos/Private/Chaos/Joint/JointConstraintsCVars.cpp
+++ b/Engine/Source/Runtime/Experimental/Chaos/Private/Chaos/Joint/JointConstraintsCVars.cpp
@@ -27,4 +27,7 @@ float Chaos_Joint_LinearVelocityThresholdToApplyRestitution = 1e-2f;
 FAutoConsoleVariableRef CVarChaosJointLinearVelocityThresholdToApplyRestitution(TEXT("p.Chaos.Joint.LinearVelocityThresholdToApplyRestitution"), Chaos_Joint_LinearVelocityThresholdToApplyRestitution, TEXT("Apply restitution only if initial velocity is higher than this threshold (used in Quasipbd)"));
 
 float Chaos_Joint_AngularVelocityThresholdToApplyRestitution = 1e-2f;
-FAutoConsoleVariableRef CVarChaosJointAngularVelocityThresholdToApplyRestitution(TEXT("p.Chaos.Joint.AngularVelocityThresholdToApplyRestitution"), Chaos_Joint_AngularVelocityThresholdToApplyRestitution, TEXT("Apply restitution only if initial velocity is higher than this threshold (used in Quasipbd)"));
\ No newline at end of file
+FAutoConsoleVariableRef CVarChaosJointAngularVelocityThresholdToApplyRestitution(TEXT("p.Chaos.Joint.AngularVelocityThresholdToApplyRestitution"), Chaos_Joint_AngularVelocityThresholdToApplyRestitution, TEXT("Apply restitution only if initial velocity is higher than this threshold (used in Quasipbd)"));
+
+bool bChaos_Joint_UseCachedSolver = false;
+FAutoConsoleVariableRef CVarChaosJointUseCachedSolver(TEXT("p.Chaos.Joint.UseCachedSolver"), bChaos_Joint_UseCachedSolver, TEXT("Use cached version of joint solver. (default is false"));
diff --git a/Engine/Source/Runtime/Experimental/Chaos/Private/Chaos/Joint/JointConstraintsCVars.h b/Engine/Source/Runtime/Experimental/Chaos/Private/Chaos/Joint/JointConstraintsCVars.h
--- a/Engine/Source/Runtime/Experimental/Chaos/Private/Chaos/Joint/JointConstraintsCVars.h
+++ b/Engine/Source/Runtime/Experimental/Chaos/Private/Chaos/Joint/JointConstraintsCVars.h
@@ -26,3 +26,5 @@ extern float Chaos_Joint_LinearVelocityThresholdToApplyRestitution;
 
 extern float Chaos_Joint_AngularVelocityThresholdToApplyRestitution;
 
+extern bool bChaos_Joint_UseCachedSolver;
+
diff --git a/Engine/Source/Runtime/Experimental/Chaos/Private/Chaos/Joint/PBDJointCachedSolverGaussSeidel.cpp b/Engine/Source/Runtime/Experimental/Chaos/Private/Chaos/Joint/PBDJointCachedSolverGaussSeidel.cpp
new file mode 100644
--- /dev/null
+++ b/Engine/Source/Runtime/Experimental/Chaos/Private/Chaos/Joint/PBDJointCachedSolverGaussSeidel.cpp
@@ -0,0 +1,3193 @@
+// Copyright Epic Games, Inc. All Rights Reserved.
+#include "Chaos/Joint/PBDJointCachedSolverGaussSeidel.h"
+#include "Chaos/Joint/ChaosJointLog.h"
+#include "Chaos/Joint/JointConstraintsCVars.h"
+#include "Chaos/Joint/JointSolverConstraints.h"
+#include "Chaos/ParticleHandle.h"
+#include "Chaos/PBDJointConstraintUtilities.h"
+#include "Chaos/Utilities.h"
+#include "ChaosStats.h"
+#if 0
+#include "PBDJointSolverGaussSeidel.ispc.generated.h"
+#endif
+
+//PRAGMA_DISABLE_OPTIMIZATION
+
+namespace Chaos
+{
+
+
+    //
+    //
+    //////////////////////////////////////////////////////////////////////////
+    //
+    //
+
+    FPBDJointCachedSolver::FPBDJointCachedSolver()
+    {
+        if (bChaos_Joint_ISPC_Enabled)
+        {
+#if 0
+            check(sizeof(FPBDJointCachedSolver) == ispc::SizeofFPBDJointSolver());
+#endif
+        }
+    }
+
+    void FPBDJointCachedSolver::InitDerivedState()
+    {
+        InitConnectorXs[0] = X(0) + R(0) * LocalConnectorXs[0].GetTranslation();
+        InitConnectorXs[1] = X(1) + R(1) * LocalConnectorXs[1].GetTranslation();
+        InitConnectorRs[0] = R(0) * LocalConnectorXs[0].GetRotation();
+        InitConnectorRs[1] = R(1) * LocalConnectorXs[1].GetRotation();
+        InitConnectorRs[1].EnforceShortestArcWith(InitConnectorRs[0]);
+
+        ConnectorXs[0] = InitConnectorXs[0];
+        ConnectorRs[0] = InitConnectorRs[0];
+
+        ConnectorXs[1] = InitConnectorXs[1];
+        ConnectorRs[1] = InitConnectorRs[1];
+    }
+
+
+    void FPBDJointCachedSolver::UpdateDerivedState()
+    {
+        // Kinematic bodies will not be moved, so we don't update derived state during iterations
+        if (InvM(0) > SMALL_NUMBER)
+        {
+            const FVec3 BodyP0 = P(0);
+            const FRotation3 BodyQ0 = Q(0);
+            ConnectorXs[0] = BodyP0 + BodyQ0 * LocalConnectorXs[0].GetTranslation();
+            ConnectorRs[0] = BodyQ0 * LocalConnectorXs[0].GetRotation();
+        }
+        if (InvM(1) > SMALL_NUMBER)
+        {
+            const FVec3 BodyP1 = P(1);
+            const FRotation3 BodyQ1 = Q(1);
+            ConnectorXs[1] = BodyP1 + BodyQ1 * LocalConnectorXs[1].GetTranslation();
+            ConnectorRs[1] = BodyQ1 * LocalConnectorXs[1].GetRotation();
+        }
+        ConnectorRs[1].EnforceShortestArcWith(ConnectorRs[0]);
+    }
+
+
+    void FPBDJointCachedSolver::UpdateDerivedState(const int32 BodyIndex)
+    {
+        const FVec3 BodyP = P(BodyIndex);
+        const FRotation3 BodyQ = Q(BodyIndex);
+        ConnectorXs[BodyIndex] = BodyP + BodyQ * LocalConnectorXs[BodyIndex].GetTranslation();
+        ConnectorRs[BodyIndex] = BodyQ * LocalConnectorXs[BodyIndex].GetRotation();
+        ConnectorRs[1].EnforceShortestArcWith(ConnectorRs[0]);
+    }
+
+    bool FPBDJointCachedSolver::UpdateIsActive(
+        const FReal Dt,
+        const FPBDJointSolverSettings& SolverSettings,
+        const FPBDJointSettings& JointSettings)
+    {
+        // NumActiveConstraints is initialized to -1, so there's no danger of getting invalid LastPs/Qs
+        // We also check SolverStiffness mainly for testing when solver stiffness is 0 (so we don't exit immediately)
+        if ((NumActiveConstraints >= 0) && (SolverStiffness > 0.0f))
+        {
+            bool bIsSolved =
+                FVec3::IsNearlyEqual(P(0), LastPs[0], PositionTolerance)
+                && FVec3::IsNearlyEqual(P(1), LastPs[1], PositionTolerance)
+                && FRotation3::IsNearlyEqual(Q(0), LastQs[0], 0.5f * AngleTolerance)
+                && FRotation3::IsNearlyEqual(Q(1), LastQs[1], 0.5f * AngleTolerance);
+            bIsActive = !bIsSolved;
+        }
+
+        LastPs[0] = P(0);
+        LastPs[1] = P(1);
+        LastQs[0] = Q(0);
+        LastQs[1] = Q(1);
+
+        return bIsActive;
+    }
+
+
+    void FPBDJointCachedSolver::Init(
+        const FReal Dt,
+        const FSolverBodyPtrPair& SolverBodyPair,
+        const FPBDJointSolverSettings& SolverSettings,
+        const FPBDJointSettings& JointSettings,
+        const FRigidTransform3& XL0,
+        const FRigidTransform3& XL1)
+    {
+        SolverBodies[0] = *SolverBodyPair[0];
+        SolverBodies[1] = *SolverBodyPair[1];
+
+        LocalConnectorXs[0] = XL0;
+        LocalConnectorXs[1] = XL1;
+
+        // \todo(chaos): joint should support parent/child in either order
+        SolverBodies[0].SetInvMScale(JointSettings.ParentInvMassScale);
+        SolverBodies[1].SetInvMScale(FReal(1));
+
+        FVec3 ConditionedInvILs[2];
+        FPBDJointUtilities::ConditionInverseMassAndInertia(Body0().InvM(), Body1().InvM(), Body0().InvILocal(), Body1().InvILocal(), SolverSettings.MinParentMassRatio, SolverSettings.MaxInertiaRatio, ConditionedInvMs[0], ConditionedInvMs[1], ConditionedInvILs[0], ConditionedInvILs[1]);
+        ConditionedInvIs[0] = Utilities::ComputeWorldSpaceInertia(Q(0), ConditionedInvILs[0]);
+        ConditionedInvIs[1] = Utilities::ComputeWorldSpaceInertia(Q(1), ConditionedInvILs[1]);
+
+        NetLinearImpulse = FVec3(0);
+        NetAngularImpulse = FVec3(0);
+
+        LinearSoftLambda = 0;
+        TwistSoftLambda = 0;
+        SwingSoftLambda = 0;
+        LinearDriveLambdas = FVec3(0);
+        RotationDriveLambdas = FVec3(0);
+
+        LinearConstraintPadding = FVec3(-1);
+        AngularConstraintPadding = FVec3(-1);
+
+        // Tolerances are positional errors below visible detection. But in PBD the errors
+        // we leave behind get converted to velocity, so we need to ensure that the resultant
+        // movement from that erroneous velocity is less than the desired position tolerance.
+        // Assume that the tolerances were defined for a 60Hz simulation, then it must be that
+        // the position error is less than the position change from constant external forces
+        // (e.g., gravity). So, we are saying that the tolerance was chosen because the position
+        // error is less that F.dt^2. We need to scale the tolerance to work at our current dt.
+        const FReal ToleranceScale = FMath::Min(1.f, 60.f * 60.f * Dt * Dt);
+        PositionTolerance = ToleranceScale * SolverSettings.PositionTolerance;
+        AngleTolerance = ToleranceScale * SolverSettings.AngleTolerance;
+
+        NumActiveConstraints = -1;
+        bIsActive = true;
+
+        SolverStiffness = 1.0f;
+
+        LinearHardLambda = FVec3(0);
+        AngularHardLambda = FVec3(0);
+
+        InitDerivedState();
+
+        if (JointSettings.LinearRestitution != 0.0f)
+        {
+            CalculateConstraintAxisLinearVelocities(JointSettings, InitConstraintAxisLinearVelocities);
+        }
+        if (JointSettings.TwistRestitution != 0.0f || JointSettings.SwingRestitution != 0.0f)
+        {
+            CalculateConstraintAxisAngularVelocities(SolverSettings, JointSettings, InitConstraintAxisAngularVelocities);
+        }
+
+        InitPositionConstraints(SolverSettings, JointSettings);
+        LastPs[0] = FVec3(0.f);
+        LastPs[1] = FVec3(0.f);
+        LastQs[0] = FRotation3::FromIdentity();
+        LastQs[1] = FRotation3::FromIdentity();
+    }
+
+    void FPBDJointCachedSolver::InitPositionConstraints(
+        const FPBDJointSolverSettings& SolverSettings,
+        const FPBDJointSettings& JointSettings)
+    {
+        bool bHasPositionConstraints =
+            (JointSettings.LinearMotionTypes[0] != EJointMotionType::Free)
+            || (JointSettings.LinearMotionTypes[1] != EJointMotionType::Free)
+            || (JointSettings.LinearMotionTypes[2] != EJointMotionType::Free);
+        if (!bHasPositionConstraints)
+        {
+            return;
+        }
+
+        if (FPBDJointUtilities::GetSoftLinearLimitEnabled(SolverSettings, JointSettings))
+        {
+            LinearStiffness = FPBDJointUtilities::GetSoftLinearStiffness(SolverSettings, JointSettings);
+            LinearDamping = FPBDJointUtilities::GetSoftLinearDamping(SolverSettings, JointSettings);
+            bLinearAccelerationMode = FPBDJointUtilities::GetLinearSoftAccelerationMode(SolverSettings, JointSettings);
+        }
+        else
+        {
+            LinearStiffness = FPBDJointUtilities::GetLinearStiffness(SolverSettings, JointSettings);
+        }
+
+        const TVec3<EJointMotionType>& LinearMotion = JointSettings.LinearMotionTypes;
+        const TVec3<bool> bLinearLocked =
+        {
+            (LinearMotion[0] == EJointMotionType::Locked),
+            (LinearMotion[1] == EJointMotionType::Locked),
+            (LinearMotion[2] == EJointMotionType::Locked),
+        };
+        const TVec3<bool> bLinearLimited =
+        {
+            (LinearMotion[0] == EJointMotionType::Limited),
+            (LinearMotion[1] == EJointMotionType::Limited),
+            (LinearMotion[2] == EJointMotionType::Limited),
+        };
+
+        if (bLinearLocked[0] && bLinearLocked[1] && bLinearLocked[2])
+        {
+            // Hard point constraint (most common case)
+            if (!Body0().IsDynamic())
+            {
+                InitPointPositionConstraintKD(0, 1);
+            }
+            else if (!Body1().IsDynamic())
+            {
+                InitPointPositionConstraintKD(1, 0);
+            }
+            else
+            {
+                InitPointPositionConstraintDD();
+            }
+        }
+        else if (bLinearLimited[0] && bLinearLimited[1] && bLinearLimited[2])
+        {
+            // Spherical constraint
+            InitSphericalPositionConstraint(JointSettings);
+        }
+        else if (bLinearLocked[1] && bLinearLocked[2] && !bLinearLocked[0])
+        {
+            // Line constraint along X axis
+            InitCylindricalPositionConstraint(0, LinearMotion[0], EJointMotionType::Locked, JointSettings);
+        }
+        else if (bLinearLocked[0] && bLinearLocked[2] && !bLinearLocked[1])
+        {
+            // Line constraint along Y axis
+            InitCylindricalPositionConstraint(1, LinearMotion[1], EJointMotionType::Locked, JointSettings);
+        }
+        else if (bLinearLocked[0] && bLinearLocked[1] && !bLinearLocked[2])
+        {
+            // Line constraint along Z axis
+            InitCylindricalPositionConstraint(2, LinearMotion[2], EJointMotionType::Locked, JointSettings);
+        }
+        else if (bLinearLimited[1] && bLinearLimited[2] && !bLinearLimited[0])
+        {
+            // Cylindrical constraint along X axis
+            InitCylindricalPositionConstraint(0, LinearMotion[0], EJointMotionType::Limited, JointSettings);
+        }
+        else if (bLinearLimited[0] && bLinearLimited[2] && !bLinearLimited[1])
+        {
+            // Cylindrical constraint along Y axis
+            InitCylindricalPositionConstraint(1, LinearMotion[1], EJointMotionType::Limited, JointSettings);
+        }
+        else if (bLinearLimited[0] && bLinearLimited[1] && !bLinearLimited[2])
+        {
+            // Cylindrical constraint along Z axis
+            InitCylindricalPositionConstraint(2, LinearMotion[2], EJointMotionType::Limited, JointSettings);
+        }
+        else if (bLinearLocked[0] || bLinearLimited[0])
+        {
+            // Planar constraint along X axis
+            InitPlanarPositionConstraint(0, LinearMotion[0], JointSettings);
+        }
+        else if (bLinearLocked[1] || bLinearLimited[1])
+        {
+            // Planar constraint along Y axis
+            InitPlanarPositionConstraint(1, LinearMotion[1], JointSettings);
+        }
+        else if (bLinearLocked[2] || bLinearLimited[2])
+        {
+            // Planar constraint along Z axis
+            InitPlanarPositionConstraint(2, LinearMotion[2], JointSettings);
+        }
+    }
+
+    void FPBDJointCachedSolver::InitPointPositionConstraintKD(
+        const int32 KIndex,
+        const int32 DIndex)
+    {
+        check(!Body(KIndex).IsDynamic());
+        check(Body(DIndex).IsDynamic());
+        const FVec3 ConstraintArm = ConnectorXs[DIndex] - P(DIndex);
+        const FMatrix33 M = Utilities::ComputeJointFactorMatrix(ConstraintArm, InvI(DIndex), InvM(DIndex));
+        ConstraintArms[PointPositionConstraintIndex][DIndex] = ConstraintArm;
+        ConstraintArms[PointPositionConstraintIndex][KIndex] = ConnectorXs[KIndex] - P(KIndex);
+        PointPositionConstraintIM = M.Inverse();
+    }
+
+    void FPBDJointCachedSolver::InitPointPositionConstraintDD()
+    {
+        check(Body(0).IsDynamic());
+        check(Body(1).IsDynamic());
+        const FVec3 ConstraintArm0 = ConnectorXs[0] - P(0);
+        const FVec3 ConstraintArm1 = ConnectorXs[1] - P(1);
+        const FMatrix33 M0 = Utilities::ComputeJointFactorMatrix(ConstraintArm0, InvI(0), InvM(0));
+        const FMatrix33 M1 = Utilities::ComputeJointFactorMatrix(ConstraintArm1, InvI(1), InvM(1));
+        ConstraintArms[PointPositionConstraintIndex][0] = ConstraintArm0;
+        ConstraintArms[PointPositionConstraintIndex][1] = ConstraintArm1;
+        PointPositionConstraintIM = (M0 + M1).Inverse();
+    }
+
+    void FPBDJointCachedSolver::InitSphericalPositionConstraint(
+        const FPBDJointSettings& JointSettings)
+    {
+        FVec3 Axis;
+        FReal Delta;
+        FPBDJointUtilities::GetSphericalAxisDelta(ConnectorXs[0], ConnectorXs[1], Axis, Delta);
+
+        const FReal LimitPadding = GetLinearConstraintPadding(0);
+        const FReal Limit = FMath::Max(JointSettings.LinearLimit - LimitPadding, (FReal)0.);
+
+        // This uses ConnectorXs[1] rather than ConnectorXs[0] for stability reasons. It is more stable this way when ConnectorXs[0] and ConnectorXs[1] are far apart 
+        const FVec3 ConstraintArm0 = ConnectorXs[1] - P(0);
+        const FVec3 ConstraintArm1 = ConnectorXs[1] - P(1);
+
+        InitPositionConstraint(Axis, ConstraintArm0, ConstraintArm1, SphericalPositionConstraintIndex);
+    }
+
+    void FPBDJointCachedSolver::InitCylindricalPositionConstraint(
+        const int32 AxisIndex,
+        const EJointMotionType AxialMotion,
+        const EJointMotionType RadialMotion,
+        const FPBDJointSettings& JointSettings)
+    {
+        check(AxialMotion != RadialMotion);
+
+        FVec3 Axis, RadialAxis;
+        FReal AxialDelta, RadialDelta;
+        FPBDJointUtilities::GetCylindricalAxesDeltas(ConnectorRs[0], ConnectorXs[0], ConnectorXs[1], AxisIndex, Axis, AxialDelta, RadialAxis, RadialDelta);
+
+        if (AxialDelta < 0.0f)
+        {
+            Axis = -Axis;
+        }
+
+        // This uses ConnectorXs[1] rather than ConnectorXs[0] for stability reasons. It is more stable this way when ConnectorXs[0] and ConnectorXs[1] are far apart 
+        const FVec3 ConstraintArm0 = ConnectorXs[1] - P(0);
+        const FVec3 ConstraintArm1 = ConnectorXs[1] - P(1);
+        InitPositionConstraint(Axis, ConstraintArm0, ConstraintArm1, (int32)EJointCylindricalPositionConstraintType::Axial);   
+        InitPositionConstraint(RadialAxis, ConstraintArm0, ConstraintArm1, (int32)EJointCylindricalPositionConstraintType::Radial);
+    }
+
+    void FPBDJointCachedSolver::InitPlanarPositionConstraint(
+        const int32 AxisIndex,
+        const EJointMotionType AxialMotion,
+        const FPBDJointSettings& JointSettings)
+    {
+        FVec3 Axis;
+        FReal Delta;
+        FPBDJointUtilities::GetPlanarAxisDelta(ConnectorRs[0], ConnectorXs[0], ConnectorXs[1], AxisIndex, Axis, Delta);
+
+        if (Delta < 0.0f)
+        {
+            Axis = -Axis;
+        }
+        // This uses ConnectorXs[1] rather than ConnectorXs[0] for stability reasons. It is more stable this way when ConnectorXs[0] and ConnectorXs[1] are far apart 
+        const FVec3 ConstraintArm0 = ConnectorXs[1] - P(0);
+        const FVec3 ConstraintArm1 = ConnectorXs[1] - P(1);
+        InitPositionConstraint(Axis, ConstraintArm0, ConstraintArm1, PlanarPositionConstraintIndex);
+    }
+
+    void FPBDJointCachedSolver::InitPositionConstraint(
+        const FVec3& Axis,
+        const FVec3& ConstraintArm0,
+        const FVec3& ConstraintArm1,
+        const int32 ConstraintIndex)
+    {
+        PositionConstraintAxis[ConstraintIndex] = Axis;
+        ConstraintArms[ConstraintIndex][0] = ConstraintArm0;
+        ConstraintArms[ConstraintIndex][1] = ConstraintArm1;
+
+        UpdatePositionConstraintIM(Axis, ConstraintArm0, ConstraintArm1, ConstraintIndex);
+    }
+
+    void FPBDJointCachedSolver::UpdatePositionConstraintIM(
+        const FVec3& Axis,
+        const FVec3& ConstraintArm0,
+        const FVec3& ConstraintArm1,
+        const int32 ConstraintIndex)
+    {
+        const FVec3 AngularAxis0 = FVec3::CrossProduct(ConstraintArm0, Axis);
+        const FVec3 AngularAxis1 = FVec3::CrossProduct(ConstraintArm1, Axis);
+        const FVec3 IA0 = Utilities::Multiply(InvI(0), AngularAxis0);
+        const FVec3 IA1 = Utilities::Multiply(InvI(1), AngularAxis1);
+        const FReal II0 = FVec3::DotProduct(AngularAxis0, IA0);
+        const FReal II1 = FVec3::DotProduct(AngularAxis1, IA1);
+
+        PositionConstraintDRAxis[ConstraintIndex][0] = IA0;
+        PositionConstraintDRAxis[ConstraintIndex][1] = -IA1;
+        PositionConstraintIM[ConstraintIndex] = InvM(0) + II0 + InvM(1) + II1;
+    }
+
+    void FPBDJointCachedSolver::Deinit()
+    {
+        SolverBodies[0].Reset();
+        SolverBodies[1].Reset();
+    }
+
+
+    void FPBDJointCachedSolver::Update(
+        const FReal Dt,
+        const FPBDJointSolverSettings& SolverSettings,
+        const FPBDJointSettings& JointSettings)
+    {
+        UpdateDerivedState();
+
+        UpdateIsActive(Dt, SolverSettings, JointSettings);
+    }
+
+    void FPBDJointCachedSolver::EnableProjection()
+    {
+        Body0().SetInvMScale(0);
+    }
+
+
+    void FPBDJointCachedSolver::ApplyConstraints(
+        const FReal Dt,
+        const FReal InSolverStiffness,
+        const FPBDJointSolverSettings& SolverSettings,
+        const FPBDJointSettings& JointSettings)
+    {
+        NumActiveConstraints = 0;
+        SolverStiffness = InSolverStiffness;
+
+        ApplyPositionConstraints(Dt, SolverSettings, JointSettings);
+        ApplyRotationConstraints(Dt, SolverSettings, JointSettings);
+    
+        ApplyPositionDrives(Dt, SolverSettings, JointSettings);
+        ApplyRotationDrives(Dt, SolverSettings, JointSettings);
+
+        UpdateIsActive(Dt, SolverSettings, JointSettings);
+    }
+
+
+    void FPBDJointCachedSolver::ApplyVelocityConstraints(
+        const FReal Dt,
+        const FReal InSolverStiffness,
+        const FPBDJointSolverSettings& SolverSettings,
+        const FPBDJointSettings& JointSettings)
+    {
+        SolverStiffness = InSolverStiffness;
+
+        // This is used for the QuasiPbd solver. If the Pbd step applied impulses to
+        // correct position errors, it will have introduced a velocity equal to the 
+        // corection divided by the timestep. We ensure that the velocity constraints
+        // (including restitution) are also enforced. This also prevents any position
+        // errors from the previous frame getting converted into energy.
+
+        ApplyLinearVelocityConstraints(Dt, SolverSettings, JointSettings);
+        ApplyAngularVelocityConstraints(Dt, SolverSettings, JointSettings);
+
+        // @todo(chaos): We can also apply velocity drives here rather than in the Pbd pass
+    }
+
+
+    void FPBDJointCachedSolver::ApplyProjections(
+        const FReal Dt,
+        const FReal InSolverStiffness,
+        const FPBDJointSolverSettings& SolverSettings,
+        const FPBDJointSettings& JointSettings)
+    {
+        // @todo(chaos): We need to handle parent/child being the other way round
+        if (!Body1().IsDynamic())
+        {
+            // If child is kinematic, return. 
+            return;
+        }
+
+        SolverStiffness = InSolverStiffness;
+
+        const FReal LinearProjection = FPBDJointUtilities::GetLinearProjection(SolverSettings, JointSettings);
+        const FReal AngularProjection = FPBDJointUtilities::GetAngularProjection(SolverSettings, JointSettings);
+
+        FVec3 DP1 = FVec3(0);
+        FVec3 DR1 = FVec3(0);
+
+        // Position Projection
+        const bool bLinearSoft = FPBDJointUtilities::GetSoftLinearLimitEnabled(SolverSettings, JointSettings);
+        const bool bLinearProjectionEnabled = (bLinearSoft && JointSettings.bSoftProjectionEnabled) || (!bLinearSoft && JointSettings.bProjectionEnabled);
+        const TVec3<EJointMotionType>& LinearMotion = JointSettings.LinearMotionTypes;
+        const bool bLinearLocked =
+            (LinearMotion[0] == EJointMotionType::Locked)
+            && (LinearMotion[1] == EJointMotionType::Locked)
+            && (LinearMotion[2] == EJointMotionType::Locked);
+        const bool bLinearLimited =
+            (LinearMotion[0] == EJointMotionType::Limited)
+            && (LinearMotion[1] == EJointMotionType::Limited)
+            && (LinearMotion[2] == EJointMotionType::Limited);
+        if (bLinearProjectionEnabled && (LinearProjection > 0))
+        {
+            if (bLinearLocked)
+            {
+                ApplyPointProjection(Dt, SolverSettings, JointSettings, LinearProjection, DP1, DR1);
+            }
+            else if (bLinearLimited)
+            {
+                ApplySphereProjection(Dt, SolverSettings, JointSettings, LinearProjection, DP1, DR1);
+            }
+            // @todo(ccaulfield): support mixed linear projection
+        }
+
+        // Twist projection
+        const bool bTwistSoft = FPBDJointUtilities::GetSoftTwistLimitEnabled(SolverSettings, JointSettings);
+        const bool bTwistProjectionEnabled = SolverSettings.bEnableTwistLimits && ((bTwistSoft && JointSettings.bSoftProjectionEnabled) || (!bTwistSoft && JointSettings.bProjectionEnabled));
+        if (bTwistProjectionEnabled && (AngularProjection > 0.0f))
+        {
+            const EJointMotionType TwistMotion = JointSettings.AngularMotionTypes[(int32)EJointAngularConstraintIndex::Twist];
+            if (TwistMotion != EJointMotionType::Free)
+            {
+                ApplyTwistProjection(Dt, SolverSettings, JointSettings, AngularProjection, bLinearLocked, DP1, DR1);
+            }
+        }
+
+        // Swing projection
+        const bool bSwingSoft = FPBDJointUtilities::GetSoftSwingLimitEnabled(SolverSettings, JointSettings);
+        const bool bSwingProjectionEnabled = SolverSettings.bEnableSwingLimits && ((bSwingSoft && JointSettings.bSoftProjectionEnabled) || (!bSwingSoft && JointSettings.bProjectionEnabled));
+        if (bSwingProjectionEnabled && (AngularProjection > 0.0f))
+        {
+            const EJointMotionType Swing1Motion = JointSettings.AngularMotionTypes[(int32)EJointAngularConstraintIndex::Swing1];
+            const EJointMotionType Swing2Motion = JointSettings.AngularMotionTypes[(int32)EJointAngularConstraintIndex::Swing2];
+            if ((Swing1Motion == EJointMotionType::Limited) && (Swing2Motion == EJointMotionType::Limited))
+            {
+                ApplyConeProjection(Dt, SolverSettings, JointSettings, AngularProjection, bLinearLocked, DP1, DR1);
+            }
+            else if ((Swing1Motion == EJointMotionType::Limited) && (Swing2Motion == EJointMotionType::Locked))
+            {
+                ApplySingleLockedSwingProjection(Dt, SolverSettings, JointSettings, EJointAngularConstraintIndex::Swing2, AngularProjection, bLinearLocked, DP1, DR1);
+                ApplySwingProjection(Dt, SolverSettings, JointSettings, EJointAngularConstraintIndex::Swing1, AngularProjection, bLinearLocked, DP1, DR1);
+            }
+            else if ((Swing1Motion == EJointMotionType::Limited) && (Swing2Motion == EJointMotionType::Free))
+            {
+                ApplyDualConeSwingProjection(Dt, SolverSettings, JointSettings, EJointAngularConstraintIndex::Swing1, AngularProjection, bLinearLocked, DP1, DR1);
+            }
+            else if ((Swing1Motion == EJointMotionType::Locked) && (Swing2Motion == EJointMotionType::Limited))
+            {
+                ApplySingleLockedSwingProjection(Dt, SolverSettings, JointSettings, EJointAngularConstraintIndex::Swing1, AngularProjection, bLinearLocked, DP1, DR1);
+                ApplySwingProjection(Dt, SolverSettings, JointSettings, EJointAngularConstraintIndex::Swing2, AngularProjection, bLinearLocked, DP1, DR1);
+            }
+            else if ((Swing1Motion == EJointMotionType::Locked) && (Swing2Motion == EJointMotionType::Locked))
+            {
+                ApplyDoubleLockedSwingProjection(Dt, SolverSettings, JointSettings, AngularProjection, bLinearLocked, DP1, DR1);
+            }
+            else if ((Swing1Motion == EJointMotionType::Locked) && (Swing2Motion == EJointMotionType::Free))
+            {
+                ApplySingleLockedSwingProjection(Dt, SolverSettings, JointSettings, EJointAngularConstraintIndex::Swing1, AngularProjection, bLinearLocked, DP1, DR1);
+            }
+            else if ((Swing1Motion == EJointMotionType::Free) && (Swing2Motion == EJointMotionType::Limited))
+            {
+                ApplyDualConeSwingProjection(Dt, SolverSettings, JointSettings, EJointAngularConstraintIndex::Swing2, AngularProjection, bLinearLocked, DP1, DR1);
+            }
+            else if ((Swing1Motion == EJointMotionType::Free) && (Swing2Motion == EJointMotionType::Locked))
+            {
+                ApplySingleLockedSwingProjection(Dt, SolverSettings, JointSettings, EJointAngularConstraintIndex::Swing2, AngularProjection, bLinearLocked, DP1, DR1);
+            }
+        }
+
+        // Final position fixup
+        if (bLinearLocked)
+        {
+            if (bLinearProjectionEnabled && (LinearProjection > 0))
+            {
+                ApplyTranslateProjection(Dt, SolverSettings, JointSettings, LinearProjection, DP1, DR1);
+            }
+
+            // Add velocity correction from the net projection motion
+            if (Chaos_Joint_VelProjectionAlpha > 0.0f)
+            {
+                ApplyVelocityProjection(Dt, SolverSettings, JointSettings, Chaos_Joint_VelProjectionAlpha, DP1, DR1);
+            }
+        }
+    
+        UpdateIsActive(Dt, SolverSettings, JointSettings);
+    }
+
+
+    void FPBDJointCachedSolver::ApplyRotationConstraints(
+        const FReal Dt,
+        const FPBDJointSolverSettings& SolverSettings,
+        const FPBDJointSettings& JointSettings)
+    {
+        bool bHasRotationConstraints =
+            (JointSettings.AngularMotionTypes[0] != EJointMotionType::Free)
+            || (JointSettings.AngularMotionTypes[1] != EJointMotionType::Free)
+            || (JointSettings.AngularMotionTypes[2] != EJointMotionType::Free);
+        if (!bHasRotationConstraints)
+        {
+            return;
+        }
+
+        // Locked axes always use hard constraints. Limited axes use hard or soft depending on settings
+        EJointMotionType TwistMotion = JointSettings.AngularMotionTypes[(int32)EJointAngularConstraintIndex::Twist];
+        EJointMotionType Swing1Motion = JointSettings.AngularMotionTypes[(int32)EJointAngularConstraintIndex::Swing1];
+        EJointMotionType Swing2Motion = JointSettings.AngularMotionTypes[(int32)EJointAngularConstraintIndex::Swing2];
+        bool bTwistSoft = FPBDJointUtilities::GetSoftTwistLimitEnabled(SolverSettings, JointSettings);
+        bool bSwingSoft = FPBDJointUtilities::GetSoftSwingLimitEnabled(SolverSettings, JointSettings);
+
+        // If the twist axes are opposing, we cannot decompose the orientation into swing and twist angles, so just give up
+        const FVec3 Twist0 = ConnectorRs[0] * FJointConstants::TwistAxis();
+        const FVec3 Twist1 = ConnectorRs[1] * FJointConstants::TwistAxis();
+        const FReal Twist01Dot = FVec3::DotProduct(Twist0, Twist1);
+        const bool bDegenerate = (Twist01Dot < Chaos_Joint_DegenerateRotationLimit);
+        if (bDegenerate)
+        {
+            UE_LOG(LogChaosJoint, VeryVerbose, TEXT(" Degenerate rotation at Swing %f deg"), FMath::RadiansToDegrees(FMath::Acos(Twist01Dot)));
+        }
+
+        // Apply twist constraint
+        // NOTE: Cannot calculate twist angle at 180degree swing
+        if (SolverSettings.bEnableTwistLimits && !bDegenerate)
+        {
+            if (TwistMotion == EJointMotionType::Limited)
+            {
+                ApplyTwistConstraint(Dt, SolverSettings, JointSettings, bTwistSoft);
+            }
+            else if (TwistMotion == EJointMotionType::Locked)
+            {
+                // Covered below
+            }
+            else if (TwistMotion == EJointMotionType::Free)
+            {
+            }
+        }
+
+        // Apply swing constraints
+        // NOTE: Cannot separate swing angles at 180degree swing (but we can still apply locks)
+        if (SolverSettings.bEnableSwingLimits)
+        {
+            if ((Swing1Motion == EJointMotionType::Limited) && (Swing2Motion == EJointMotionType::Limited))
+            {
+                ApplyConeConstraint(Dt, SolverSettings, JointSettings, bSwingSoft);
+            }
+            else if ((Swing1Motion == EJointMotionType::Limited) && (Swing2Motion == EJointMotionType::Locked))
+            {
+                ApplySingleLockedSwingConstraint(Dt, SolverSettings, JointSettings, EJointAngularConstraintIndex::Swing2, false);
+                if (!bDegenerate)
+                {
+                    ApplySwingConstraint(Dt, SolverSettings, JointSettings, EJointAngularConstraintIndex::Swing1, bSwingSoft);
+                }
+            }
+            else if ((Swing1Motion == EJointMotionType::Limited) && (Swing2Motion == EJointMotionType::Free))
+            {
+                if (!bDegenerate)
+                {
+                    ApplyDualConeSwingConstraint(Dt, SolverSettings, JointSettings, EJointAngularConstraintIndex::Swing1, bSwingSoft);
+                }
+            }
+            else if ((Swing1Motion == EJointMotionType::Locked) && (Swing2Motion == EJointMotionType::Limited))
+            {
+                ApplySingleLockedSwingConstraint(Dt, SolverSettings, JointSettings, EJointAngularConstraintIndex::Swing1, false);
+                if (!bDegenerate)
+                {
+                    ApplySwingConstraint(Dt, SolverSettings, JointSettings, EJointAngularConstraintIndex::Swing2, bSwingSoft);
+                }
+            }
+            else if ((Swing1Motion == EJointMotionType::Locked) && (Swing2Motion == EJointMotionType::Locked))
+            {
+                // Covered below
+            }
+            else if ((Swing1Motion == EJointMotionType::Locked) && (Swing2Motion == EJointMotionType::Free))
+            {
+                ApplySingleLockedSwingConstraint(Dt, SolverSettings, JointSettings, EJointAngularConstraintIndex::Swing1, false);
+            }
+            else if ((Swing1Motion == EJointMotionType::Free) && (Swing2Motion == EJointMotionType::Limited))
+            {
+                if (!bDegenerate)
+                {
+                    ApplyDualConeSwingConstraint(Dt, SolverSettings, JointSettings, EJointAngularConstraintIndex::Swing2, bSwingSoft);
+                }
+            }
+            else if ((Swing1Motion == EJointMotionType::Free) && (Swing2Motion == EJointMotionType::Locked))
+            {
+                ApplySingleLockedSwingConstraint(Dt, SolverSettings, JointSettings, EJointAngularConstraintIndex::Swing2, false);
+            }
+            else if ((Swing1Motion == EJointMotionType::Free) && (Swing2Motion == EJointMotionType::Free))
+            {
+            }
+        }
+
+        // Note: single-swing locks are already handled above so we only need to do something here if both are locked
+        bool bLockedTwist = SolverSettings.bEnableTwistLimits 
+            && (JointSettings.AngularMotionTypes[(int32)EJointAngularConstraintIndex::Twist] == EJointMotionType::Locked);
+        bool bLockedSwing = SolverSettings.bEnableSwingLimits 
+            && (JointSettings.AngularMotionTypes[(int32)EJointAngularConstraintIndex::Swing1] == EJointMotionType::Locked) 
+            && (JointSettings.AngularMotionTypes[(int32)EJointAngularConstraintIndex::Swing2] == EJointMotionType::Locked);
+        if (bLockedTwist || bLockedSwing)
+        {
+            ApplyLockedRotationConstraints(Dt, SolverSettings, JointSettings, bLockedTwist, bLockedSwing);
+        }
+    }
+
+
+    void FPBDJointCachedSolver::ApplyRotationDrives(
+        const FReal Dt,
+        const FPBDJointSolverSettings& SolverSettings,
+        const FPBDJointSettings& JointSettings)
+    {
+        bool bHasRotationDrives =
+            JointSettings.bAngularTwistPositionDriveEnabled
+            || JointSettings.bAngularTwistVelocityDriveEnabled
+            || JointSettings.bAngularSwingPositionDriveEnabled
+            || JointSettings.bAngularSwingVelocityDriveEnabled
+            || JointSettings.bAngularSLerpPositionDriveEnabled
+            || JointSettings.bAngularSLerpVelocityDriveEnabled;
+        if (!bHasRotationDrives)
+        {
+            return;
+        }
+
+        EJointMotionType TwistMotion = JointSettings.AngularMotionTypes[(int32)EJointAngularConstraintIndex::Twist];
+        EJointMotionType Swing1Motion = JointSettings.AngularMotionTypes[(int32)EJointAngularConstraintIndex::Swing1];
+        EJointMotionType Swing2Motion = JointSettings.AngularMotionTypes[(int32)EJointAngularConstraintIndex::Swing2];
+
+        if (SolverSettings.bEnableDrives)
+        {
+            bool bTwistLocked = TwistMotion == EJointMotionType::Locked;
+            bool bSwing1Locked = Swing1Motion == EJointMotionType::Locked;
+            bool bSwing2Locked = Swing2Motion == EJointMotionType::Locked;
+
+            // No SLerp drive if we have a locked rotation (it will be grayed out in the editor in this case, but could still have been set before the rotation was locked)
+            // @todo(ccaulfield): setting should be cleaned up before being passed to the solver
+            if ((JointSettings.bAngularSLerpPositionDriveEnabled || JointSettings.bAngularSLerpVelocityDriveEnabled) && !bTwistLocked && !bSwing1Locked && !bSwing2Locked)
+            {
+                ApplySLerpDrive(Dt, SolverSettings, JointSettings);
+            }
+            else
+            {
+                const bool bTwistDriveEnabled = ((JointSettings.bAngularTwistPositionDriveEnabled || JointSettings.bAngularTwistVelocityDriveEnabled) && !bTwistLocked);
+                const bool bSwingDriveEnabled = (JointSettings.bAngularSwingPositionDriveEnabled || JointSettings.bAngularSwingVelocityDriveEnabled);
+                const bool bSwing1DriveEnabled = bSwingDriveEnabled && !bSwing1Locked;
+                const bool bSwing2DriveEnabled = bSwingDriveEnabled && !bSwing2Locked;
+                if (bTwistDriveEnabled || bSwing1DriveEnabled || bSwing2DriveEnabled)
+                {
+                    ApplySwingTwistDrives(Dt, SolverSettings, JointSettings, bTwistDriveEnabled, bSwing1DriveEnabled, bSwing2DriveEnabled);
+                }
+            }
+        }
+    }
+
+
+    void FPBDJointCachedSolver::ApplyPositionConstraints(
+        const FReal Dt,
+        const FPBDJointSolverSettings& SolverSettings,
+        const FPBDJointSettings& JointSettings)
+    {
+        bool bHasPositionConstraints =
+            (JointSettings.LinearMotionTypes[0] != EJointMotionType::Free)
+            || (JointSettings.LinearMotionTypes[1] != EJointMotionType::Free)
+            || (JointSettings.LinearMotionTypes[2] != EJointMotionType::Free);
+        if (!bHasPositionConstraints)
+        {
+            return;
+        }
+
+        const TVec3<EJointMotionType>& LinearMotion = JointSettings.LinearMotionTypes;
+        const TVec3<bool> bLinearLocked =
+        {
+            (LinearMotion[0] == EJointMotionType::Locked),
+            (LinearMotion[1] == EJointMotionType::Locked),
+            (LinearMotion[2] == EJointMotionType::Locked),
+        };
+        const TVec3<bool> bLinearLimited =
+        {
+            (LinearMotion[0] == EJointMotionType::Limited),
+            (LinearMotion[1] == EJointMotionType::Limited),
+            (LinearMotion[2] == EJointMotionType::Limited),
+        };
+
+        if (bLinearLocked[0] && bLinearLocked[1] && bLinearLocked[2])
+        {
+            // Hard point constraint (most common case)
+            if (!Body0().IsDynamic())
+            {
+                ApplyPointPositionConstraintKD(0, 1, Dt, SolverSettings, JointSettings);
+            }
+            else if (!Body1().IsDynamic())
+            {
+                ApplyPointPositionConstraintKD(1, 0, Dt, SolverSettings, JointSettings);
+            }
+            else
+            {
+                ApplyPointPositionConstraintDD(Dt, SolverSettings, JointSettings);
+            }
+        }
+        else if (bLinearLimited[0] && bLinearLimited[1] && bLinearLimited[2])
+        {
+            // Spherical constraint
+            ApplySphericalPositionConstraint(Dt, SolverSettings, JointSettings);
+        }
+        else if (bLinearLocked[1] && bLinearLocked[2] && !bLinearLocked[0])
+        {
+            // Line constraint along X axis
+            ApplyCylindricalPositionConstraint(Dt, 0, LinearMotion[0], EJointMotionType::Locked, SolverSettings, JointSettings);
+        }
+        else if (bLinearLocked[0] && bLinearLocked[2] && !bLinearLocked[1])
+        {
+            // Line constraint along Y axis
+            ApplyCylindricalPositionConstraint(Dt, 1, LinearMotion[1], EJointMotionType::Locked, SolverSettings, JointSettings);
+        }
+        else if (bLinearLocked[0] && bLinearLocked[1] && !bLinearLocked[2])
+        {
+            // Line constraint along Z axis
+            ApplyCylindricalPositionConstraint(Dt, 2, LinearMotion[2], EJointMotionType::Locked, SolverSettings, JointSettings);
+        }
+        else if (bLinearLimited[1] && bLinearLimited[2] && !bLinearLimited[0])
+        {
+            // Cylindrical constraint along X axis
+            ApplyCylindricalPositionConstraint(Dt, 0, LinearMotion[0], EJointMotionType::Limited, SolverSettings, JointSettings);
+        }
+        else if (bLinearLimited[0] && bLinearLimited[2] && !bLinearLimited[1])
+        {
+            // Cylindrical constraint along Y axis
+            ApplyCylindricalPositionConstraint(Dt, 1, LinearMotion[1], EJointMotionType::Limited, SolverSettings, JointSettings);
+        }
+        else if (bLinearLimited[0] && bLinearLimited[1] && !bLinearLimited[2])
+        {
+            // Cylindrical constraint along Z axis
+            ApplyCylindricalPositionConstraint(Dt, 2, LinearMotion[2], EJointMotionType::Limited, SolverSettings, JointSettings);
+        }
+        else if (bLinearLocked[0] || bLinearLimited[0])
+        {
+            // Planar constraint along X axis
+            ApplyPlanarPositionConstraint(Dt, 0, LinearMotion[0], SolverSettings, JointSettings);
+        }
+        else if (bLinearLocked[1] || bLinearLimited[1])
+        {
+            // Planar constraint along Y axis
+            ApplyPlanarPositionConstraint(Dt, 1, LinearMotion[1], SolverSettings, JointSettings);
+        }
+        else if (bLinearLocked[2] || bLinearLimited[2])
+        {
+            // Planar constraint along Z axis
+            ApplyPlanarPositionConstraint(Dt, 2, LinearMotion[2], SolverSettings, JointSettings);
+        }
+    }
+
+
+    void FPBDJointCachedSolver::ApplyPositionDrives(
+        const FReal Dt,
+        const FPBDJointSolverSettings& SolverSettings,
+        const FPBDJointSettings& JointSettings)
+    {
+        if (SolverSettings.bEnableDrives)
+        {
+            TVec3<bool> bDriven =
+            {
+                (JointSettings.bLinearPositionDriveEnabled[0] || JointSettings.bLinearVelocityDriveEnabled[0]) && (JointSettings.LinearMotionTypes[0] != EJointMotionType::Locked),
+                (JointSettings.bLinearPositionDriveEnabled[1] || JointSettings.bLinearVelocityDriveEnabled[1]) && (JointSettings.LinearMotionTypes[1] != EJointMotionType::Locked),
+                (JointSettings.bLinearPositionDriveEnabled[2] || JointSettings.bLinearVelocityDriveEnabled[2]) && (JointSettings.LinearMotionTypes[2] != EJointMotionType::Locked),
+            };
+
+            // Rectangular position drives
+            if (bDriven[0] || bDriven[1] || bDriven[2])
+            {
+                const FMatrix33 R0M = ConnectorRs[0].ToMatrix();
+                const FVec3 XTarget = ConnectorXs[0] + ConnectorRs[0] * JointSettings.LinearDrivePositionTarget;
+                const FVec3 VTarget = ConnectorRs[0] * JointSettings.LinearDriveVelocityTarget;
+                const FVec3 CX = ConnectorXs[1] - XTarget;
+
+                for (int32 AxisIndex = 0; AxisIndex < 3; ++AxisIndex)
+                {
+                    if (bDriven[AxisIndex])
+                    {
+                        const FVec3 Axis = R0M.GetAxis(AxisIndex);
+                        const FReal DeltaPos = FVec3::DotProduct(CX, Axis);
+                        const FReal DeltaVel = FVec3::DotProduct(VTarget, Axis);
+
+                        ApplyPositionDrive(Dt, AxisIndex, SolverSettings, JointSettings, Axis, DeltaPos, DeltaVel);
+                    }
+                }
+            }
+        }
+    }
+
+
+    //
+    //
+    //////////////////////////////////////////////////////////////////////////
+    //
+    //
+
+
+    void FPBDJointCachedSolver::ApplyLinearVelocityConstraints(
+        const FReal Dt,
+        const FPBDJointSolverSettings& SolverSettings,
+        const FPBDJointSettings& JointSettings)
+    {
+        bool bHasPositionConstraints =
+            (JointSettings.LinearMotionTypes[0] != EJointMotionType::Free)
+            || (JointSettings.LinearMotionTypes[1] != EJointMotionType::Free)
+            || (JointSettings.LinearMotionTypes[2] != EJointMotionType::Free);
+        if (!bHasPositionConstraints)
+        {
+            return;
+        }
+
+        const TVector<EJointMotionType, 3>& LinearMotion = JointSettings.LinearMotionTypes;
+        const TVector<bool, 3> bLinearLocked =
+        {
+            (LinearMotion[0] == EJointMotionType::Locked),
+            (LinearMotion[1] == EJointMotionType::Locked),
+            (LinearMotion[2] == EJointMotionType::Locked),
+        };
+        const TVector<bool, 3> bLinearLimited =
+        {
+            (LinearMotion[0] == EJointMotionType::Limited),
+            (LinearMotion[1] == EJointMotionType::Limited),
+            (LinearMotion[2] == EJointMotionType::Limited),
+        };
+
+        if (bLinearLocked[0] && bLinearLocked[1] && bLinearLocked[2])
+        {
+            // Hard point constraint (most common case)
+            ApplyPointVelocityConstraint(Dt, SolverSettings, JointSettings);
+        }
+        else if (bLinearLimited[0] && bLinearLimited[1] && bLinearLimited[2])
+        {
+            ApplySphericalVelocityConstraint(Dt, SolverSettings, JointSettings);
+        }
+        else if (bLinearLocked[1] && bLinearLocked[2] && !bLinearLocked[0])
+        {
+            // Line constraint along X axis
+            ApplyCylindricalVelocityConstraint(Dt, 0, LinearMotion[0], EJointMotionType::Locked, SolverSettings, JointSettings);
+        }
+        else if (bLinearLocked[0] && bLinearLocked[2] && !bLinearLocked[1])
+        {
+            // Line constraint along Y axis
+            ApplyCylindricalVelocityConstraint(Dt, 1, LinearMotion[1], EJointMotionType::Locked, SolverSettings, JointSettings);
+        }
+        else if (bLinearLocked[0] && bLinearLocked[1] && !bLinearLocked[2])
+        {
+            // Line constraint along Z axis
+            ApplyCylindricalVelocityConstraint(Dt, 2, LinearMotion[2], EJointMotionType::Locked, SolverSettings, JointSettings);
+        }
+        else if (bLinearLimited[1] && bLinearLimited[2] && !bLinearLimited[0])
+        {
+            // Cylindrical constraint along X axis
+            ApplyCylindricalVelocityConstraint(Dt, 0, LinearMotion[0], EJointMotionType::Limited, SolverSettings, JointSettings);
+        }
+        else if (bLinearLimited[0] && bLinearLimited[2] && !bLinearLimited[1])
+        {
+            // Cylindrical constraint along Y axis
+            ApplyCylindricalVelocityConstraint(Dt, 1, LinearMotion[1], EJointMotionType::Limited, SolverSettings, JointSettings);
+        }
+        else if (bLinearLimited[0] && bLinearLimited[1] && !bLinearLimited[2])
+        {
+            // Cylindrical constraint along Z axis
+            ApplyCylindricalVelocityConstraint(Dt, 2, LinearMotion[2], EJointMotionType::Limited, SolverSettings, JointSettings);
+        }
+        else if (bLinearLocked[0] || bLinearLimited[0])
+        {
+            // Planar constraint along X axis
+            ApplyPlanarVelocityConstraint(Dt, 0, LinearMotion[0], SolverSettings, JointSettings);
+        }
+        else if (bLinearLocked[1] || bLinearLimited[1])
+        {
+            // Planar constraint along Y axis
+            ApplyPlanarVelocityConstraint(Dt, 1, LinearMotion[1], SolverSettings, JointSettings);
+        }
+        else if (bLinearLocked[2] || bLinearLimited[2])
+        {
+            // Planar constraint along Z axis
+            ApplyPlanarVelocityConstraint(Dt, 2, LinearMotion[2], SolverSettings, JointSettings);
+        }
+    }
+
+    void FPBDJointCachedSolver::ApplyAngularVelocityConstraints(
+        const FReal Dt,
+        const FPBDJointSolverSettings& SolverSettings,
+        const FPBDJointSettings& JointSettings)
+    {
+        bool bHasRotationConstraints =
+            (JointSettings.AngularMotionTypes[0] != EJointMotionType::Free)
+            || (JointSettings.AngularMotionTypes[1] != EJointMotionType::Free)
+            || (JointSettings.AngularMotionTypes[2] != EJointMotionType::Free);
+        if (!bHasRotationConstraints)
+        {
+            return;
+        }
+
+        // Locked axes always use hard constraints. Limited axes use hard or soft depending on settings
+        EJointMotionType TwistMotion = JointSettings.AngularMotionTypes[(int32)EJointAngularConstraintIndex::Twist];
+        EJointMotionType Swing1Motion = JointSettings.AngularMotionTypes[(int32)EJointAngularConstraintIndex::Swing1];
+        EJointMotionType Swing2Motion = JointSettings.AngularMotionTypes[(int32)EJointAngularConstraintIndex::Swing2];
+        bool bTwistSoft = FPBDJointUtilities::GetSoftTwistLimitEnabled(SolverSettings, JointSettings);
+        bool bSwingSoft = FPBDJointUtilities::GetSoftSwingLimitEnabled(SolverSettings, JointSettings);
+
+        // If the twist axes are opposing, we cannot decompose the orientation into swing and twist angles, so just give up
+        const FVec3 Twist0 = ConnectorRs[0] * FJointConstants::TwistAxis();
+        const FVec3 Twist1 = ConnectorRs[1] * FJointConstants::TwistAxis();
+        const FReal Twist01Dot = FVec3::DotProduct(Twist0, Twist1);
+        const bool bDegenerate = (Twist01Dot < Chaos_Joint_DegenerateRotationLimit);
+        if (bDegenerate)
+        {
+            UE_LOG(LogChaosJoint, VeryVerbose, TEXT(" Degenerate rotation at Swing %f deg"), FMath::RadiansToDegrees(FMath::Acos(Twist01Dot)));
+        }
+
+        // Apply twist constraint
+        // NOTE: Cannot calculate twist angle at 180degree swing
+        if (SolverSettings.bEnableTwistLimits && !bDegenerate)
+        {
+            if (TwistMotion == EJointMotionType::Limited)
+            {
+                ApplyTwistVelocityConstraint(Dt, SolverSettings, JointSettings, bTwistSoft);
+            }
+            else if (TwistMotion == EJointMotionType::Locked)
+            {
+                // Covered below
+            }
+            else if (TwistMotion == EJointMotionType::Free)
+            {
+            }
+        }
+
+        // Apply swing constraints
+        // NOTE: Cannot separate swing angles at 180degree swing (but we can still apply locks)
+        if (SolverSettings.bEnableSwingLimits)
+        {
+            if ((Swing1Motion == EJointMotionType::Limited) && (Swing2Motion == EJointMotionType::Limited))
+            {
+                ApplyConeVelocityConstraint(Dt, SolverSettings, JointSettings, bSwingSoft);
+            }
+            else if ((Swing1Motion == EJointMotionType::Limited) && (Swing2Motion == EJointMotionType::Locked))
+            {
+                ApplySingleLockedSwingVelocityConstraint(Dt, SolverSettings, JointSettings, EJointAngularConstraintIndex::Swing2, false);
+                if (!bDegenerate)
+                {
+                    ApplySwingVelocityConstraint(Dt, SolverSettings, JointSettings, EJointAngularConstraintIndex::Swing1, bSwingSoft);
+                }
+            }
+            else if ((Swing1Motion == EJointMotionType::Limited) && (Swing2Motion == EJointMotionType::Free))
+            {
+                if (!bDegenerate)
+                {
+                    ApplyDualConeSwingVelocityConstraint(Dt, SolverSettings, JointSettings, EJointAngularConstraintIndex::Swing1, bSwingSoft);
+                }
+            }
+            else if ((Swing1Motion == EJointMotionType::Locked) && (Swing2Motion == EJointMotionType::Limited))
+            {
+                ApplySingleLockedSwingVelocityConstraint(Dt, SolverSettings, JointSettings, EJointAngularConstraintIndex::Swing1, false);
+                if (!bDegenerate)
+                {
+                    ApplySwingVelocityConstraint(Dt, SolverSettings, JointSettings, EJointAngularConstraintIndex::Swing2, bSwingSoft);
+                }
+            }
+            else if ((Swing1Motion == EJointMotionType::Locked) && (Swing2Motion == EJointMotionType::Locked))
+            {
+                // Covered below
+            }
+            else if ((Swing1Motion == EJointMotionType::Locked) && (Swing2Motion == EJointMotionType::Free))
+            {
+                ApplySingleLockedSwingVelocityConstraint(Dt, SolverSettings, JointSettings, EJointAngularConstraintIndex::Swing1, false);
+            }
+            else if ((Swing1Motion == EJointMotionType::Free) && (Swing2Motion == EJointMotionType::Limited))
+            {
+                if (!bDegenerate)
+                {
+                    ApplyDualConeSwingVelocityConstraint(Dt, SolverSettings, JointSettings, EJointAngularConstraintIndex::Swing2, bSwingSoft);
+                }
+            }
+            else if ((Swing1Motion == EJointMotionType::Free) && (Swing2Motion == EJointMotionType::Locked))
+            {
+                ApplySingleLockedSwingVelocityConstraint(Dt, SolverSettings, JointSettings, EJointAngularConstraintIndex::Swing2, false);
+            }
+            else if ((Swing1Motion == EJointMotionType::Free) && (Swing2Motion == EJointMotionType::Free))
+            {
+            }
+        }
+
+        // Note: single-swing locks are already handled above so we only need to do something here if both are locked
+        bool bLockedTwist = SolverSettings.bEnableTwistLimits 
+            && (JointSettings.AngularMotionTypes[(int32)EJointAngularConstraintIndex::Twist] == EJointMotionType::Locked);
+        bool bLockedSwing = SolverSettings.bEnableSwingLimits 
+            && (JointSettings.AngularMotionTypes[(int32)EJointAngularConstraintIndex::Swing1] == EJointMotionType::Locked) 
+            && (JointSettings.AngularMotionTypes[(int32)EJointAngularConstraintIndex::Swing2] == EJointMotionType::Locked);
+        if (bLockedTwist || bLockedSwing)
+        {
+            ApplyLockedRotationVelocityConstraints(Dt, SolverSettings, JointSettings, bLockedTwist, bLockedSwing);
+        }
+    }
+
+    void FPBDJointCachedSolver::ApplyTwistVelocityConstraint(
+        const FReal Dt,
+        const FPBDJointSolverSettings& SolverSettings,
+        const FPBDJointSettings& JointSettings,
+        const bool bUseSoftLimit)
+    {
+        if(!NetAngularImpulse.IsNearlyZero() && !bUseSoftLimit && AngularHardLambda[(int32)EJointAngularConstraintIndex::Twist] > SMALL_NUMBER)
+        {
+            FVec3 TwistAxis;
+            FReal TwistAngle;
+            FPBDJointUtilities::GetTwistAxisAngle(ConnectorRs[0], ConnectorRs[1], TwistAxis, TwistAngle);
+
+            if (TwistAngle < 0.0f)
+            {
+                TwistAxis = -TwistAxis;
+            }
+
+            const FReal TwistStiffness = SolverStiffness * FPBDJointUtilities::GetTwistStiffness(SolverSettings, JointSettings);
+
+            FReal TargetVel = 0.0f;
+            if (JointSettings.TwistRestitution != 0.0f)
+            {
+                const FReal InitVel = InitConstraintAxisAngularVelocities[(int32)EJointAngularConstraintIndex::Twist];
+                TargetVel = InitVel > Chaos_Joint_AngularVelocityThresholdToApplyRestitution ? -JointSettings.TwistRestitution * InitVel : 0.0f;
+            }
+            ApplyAngularVelocityConstraint(TwistStiffness, TwistAxis, TargetVel);
+        }
+    }
+
+    void FPBDJointCachedSolver::ApplyConeVelocityConstraint(
+        const FReal Dt,
+        const FPBDJointSolverSettings& SolverSettings,
+        const FPBDJointSettings& JointSettings,
+        const bool bUseSoftLimit)
+    {
+        if(!NetAngularImpulse.IsNearlyZero() && !bUseSoftLimit && AngularHardLambda[(int32)EJointAngularConstraintIndex::Swing2] > SMALL_NUMBER)
+        {
+            FVec3 SwingAxisLocal;
+            FReal DSwingAngle = 0.0f;
+            const FReal LimitPadding = GetAngularConstraintPadding(EJointAngularConstraintIndex::Swing1);
+            const FReal Swing1Limit = FMath::Max(JointSettings.AngularLimits[(int32)EJointAngularConstraintIndex::Swing1] - LimitPadding, 0.0f);
+            const FReal Swing2Limit = FMath::Max(JointSettings.AngularLimits[(int32)EJointAngularConstraintIndex::Swing2] - LimitPadding, 0.0f);
+            FPBDJointUtilities::GetEllipticalConeAxisErrorLocal(ConnectorRs[0], ConnectorRs[1], Swing2Limit, Swing1Limit, SwingAxisLocal, DSwingAngle);
+            // SwingAxisLocal has the size of Sine(SwingAngle), thus we need to normalize it
+            SwingAxisLocal.SafeNormalize();
+            const FVec3 SwingAxis = ConnectorRs[0] * SwingAxisLocal;
+            const FReal SwingStiffness = SolverStiffness * FPBDJointUtilities::GetSwingStiffness(SolverSettings, JointSettings);
+            FReal TargetVel = 0.0f;
+            if (JointSettings.SwingRestitution != 0.0f)
+            {
+                const FReal InitVel = InitConstraintAxisAngularVelocities[(int32)EJointAngularConstraintIndex::Swing1];
+                TargetVel = InitVel > Chaos_Joint_AngularVelocityThresholdToApplyRestitution ? -JointSettings.SwingRestitution * InitVel : 0.0f;
+            }
+            ApplyAngularVelocityConstraint(SwingStiffness, SwingAxis, TargetVel);
+        }
+    }
+
+    void FPBDJointCachedSolver::ApplySingleLockedSwingVelocityConstraint(
+        const FReal Dt,
+        const FPBDJointSolverSettings& SolverSettings,
+        const FPBDJointSettings& JointSettings,
+        const EJointAngularConstraintIndex SwingConstraintIndex,
+        const bool bUseSoftLimit)
+    {
+        if(!NetAngularImpulse.IsNearlyZero() && !bUseSoftLimit && AngularHardLambda[(int32)SwingConstraintIndex] > SMALL_NUMBER)
+        {
+            FVec3 SwingAxis;
+            FReal SwingAngle;
+            FPBDJointUtilities::GetLockedSwingAxisAngle(ConnectorRs[0], ConnectorRs[1], SwingConstraintIndex, SwingAxis, SwingAngle);
+            // Swing axis has the size of Sine(SwingAngle), thus we need to normalize it
+            SwingAxis.SafeNormalize();
+
+            const FReal SwingStiffness = SolverStiffness * FPBDJointUtilities::GetSwingStiffness(SolverSettings, JointSettings);
+            ApplyAngularVelocityConstraint(SwingStiffness, SwingAxis);
+        }
+    }
+
+    void FPBDJointCachedSolver::ApplyDualConeSwingVelocityConstraint(
+        const FReal Dt,
+        const FPBDJointSolverSettings& SolverSettings,
+        const FPBDJointSettings& JointSettings,
+        const EJointAngularConstraintIndex SwingConstraintIndex,
+        const bool bUseSoftLimit)
+    {
+        if(!NetAngularImpulse.IsNearlyZero() && !bUseSoftLimit && AngularHardLambda[(int32)SwingConstraintIndex] > SMALL_NUMBER)
+        {
+            FVec3 SwingAxis;
+            FReal SwingAngle;
+            FPBDJointUtilities::GetDualConeSwingAxisAngle(ConnectorRs[0], ConnectorRs[1], SwingConstraintIndex, SwingAxis, SwingAngle);
+
+            const FReal SwingStiffness = SolverStiffness * FPBDJointUtilities::GetSwingStiffness(SolverSettings, JointSettings);
+            FReal TargetVel = 0.0f;
+            if (JointSettings.SwingRestitution != 0.0f)
+            {
+                const FReal InitVel = InitConstraintAxisAngularVelocities[(int32)SwingConstraintIndex];
+                TargetVel = InitVel > Chaos_Joint_AngularVelocityThresholdToApplyRestitution ? -JointSettings.SwingRestitution * InitVel : 0.0f;
+            }
+            ApplyAngularVelocityConstraint(SwingStiffness, SwingAxis, TargetVel);
+        }
+    }
+
+    void FPBDJointCachedSolver::ApplySwingVelocityConstraint(
+        const FReal Dt,
+        const FPBDJointSolverSettings& SolverSettings,
+        const FPBDJointSettings& JointSettings,
+        const EJointAngularConstraintIndex SwingConstraintIndex,
+        const bool bUseSoftLimit)
+    {
+        if(!NetAngularImpulse.IsNearlyZero() && !bUseSoftLimit && AngularHardLambda[(int32)SwingConstraintIndex] > SMALL_NUMBER)
+        {
+            FVec3 SwingAxis;
+            FReal SwingAngle;
+            FPBDJointUtilities::GetSwingAxisAngle(ConnectorRs[0], ConnectorRs[1], SolverSettings.SwingTwistAngleTolerance, SwingConstraintIndex, SwingAxis, SwingAngle);
+
+            const FReal SwingStiffness = SolverStiffness * FPBDJointUtilities::GetSwingStiffness(SolverSettings, JointSettings);
+            FReal TargetVel = 0.0f;
+            if (JointSettings.SwingRestitution != 0.0f)
+            {
+                const FReal InitVel = InitConstraintAxisAngularVelocities[(int32)SwingConstraintIndex];
+                TargetVel = InitVel > Chaos_Joint_AngularVelocityThresholdToApplyRestitution ? -JointSettings.SwingRestitution * InitVel : 0.0f;
+            }
+            ApplyAngularVelocityConstraint(SwingStiffness, SwingAxis, TargetVel);
+        }
+    }
+
+    void FPBDJointCachedSolver::ApplyAngularVelocityConstraint(
+        const FReal Stiffness,
+        const FVec3& Axis,
+        const FReal TargetVel)
+    {
+        const FVec3 CW = W(1) - W(0);
+        const FVec3 IA0 = Utilities::Multiply(InvI(0), Axis);
+        const FVec3 IA1 = Utilities::Multiply(InvI(1), Axis);
+        const FReal II0 = FVec3::DotProduct(Axis, IA0);
+        const FReal II1 = FVec3::DotProduct(Axis, IA1);
+
+        const FVec3 AngularImpulse = Stiffness * (FVec3::DotProduct(CW, Axis) - TargetVel) / (II0 + II1) * Axis;
+        const FVec3 DW0 = Utilities::Multiply(InvI(0), AngularImpulse);
+        const FVec3 DW1 = Utilities::Multiply(InvI(1), -AngularImpulse);
+
+        ApplyAngularVelocityDelta(DW0, DW1);
+    }
+
+    void FPBDJointCachedSolver::ApplyLockedRotationVelocityConstraints(
+        const FReal Dt,
+        const FPBDJointSolverSettings& SolverSettings,
+        const FPBDJointSettings& JointSettings,
+        const bool bApplyTwist,
+        const bool bApplySwing)
+    {
+        if(!NetAngularImpulse.IsNearlyZero())
+        {
+            FVec3 Axis0, Axis1, Axis2;
+            FPBDJointUtilities::GetLockedRotationAxes(ConnectorRs[0], ConnectorRs[1], Axis0, Axis1, Axis2);
+            const FRotation3 R01 = ConnectorRs[0].Inverse() * ConnectorRs[1];
+
+            if (bApplyTwist && AngularHardLambda[(int32)EJointAngularConstraintIndex::Twist] > SMALL_NUMBER)
+            {
+                FReal TwistStiffness = SolverStiffness * FPBDJointUtilities::GetTwistStiffness(SolverSettings, JointSettings);
+                Axis0.SafeNormalize();
+                ApplyAngularVelocityConstraint(TwistStiffness, Axis0);
+            }
+
+            if (bApplySwing && AngularHardLambda[(int32)EJointAngularConstraintIndex::Swing2] + AngularHardLambda[(int32)EJointAngularConstraintIndex::Swing1] > SMALL_NUMBER)
+            {
+                FReal SwingStiffness = SolverStiffness * FPBDJointUtilities::GetSwingStiffness(SolverSettings, JointSettings);
+                Axis1.SafeNormalize();
+                ApplyAngularVelocityConstraint(SwingStiffness, Axis1);
+                Axis2.SafeNormalize();
+                ApplyAngularVelocityConstraint(SwingStiffness, Axis2);
+            }
+        }
+    }
+
+
+    //
+    //
+    //////////////////////////////////////////////////////////////////////////
+    //
+    //
+
+
+    void FPBDJointCachedSolver::ApplyPositionDelta(
+        const int32 BodyIndex,
+        const FVec3& DP)
+    {
+        //UE_LOG(LogChaosJoint, VeryVerbose, TEXT("      Apply DP%d %f %f %f"), BodyIndex, DP.X, DP.Y, DP.Z);
+
+        Body(BodyIndex).ApplyPositionDelta(DP);
+
+        ConnectorXs[BodyIndex] += DP;
+    }
+
+
+    void FPBDJointCachedSolver::ApplyPositionDelta(
+        const FVec3& DP0,
+        const FVec3& DP1)
+    {
+        //UE_LOG(LogChaosJoint, VeryVerbose, TEXT("      Apply DP%d %f %f %f"), 0, DP0.X, DP0.Y, DP0.Z);
+        //UE_LOG(LogChaosJoint, VeryVerbose, TEXT("      Apply DP%d %f %f %f"), 1, DP1.X, DP1.Y, DP1.Z);
+
+        Body(0).ApplyPositionDelta(DP0);
+        Body(1).ApplyPositionDelta(DP1);
+
+        ConnectorXs[0] += DP0;
+        ConnectorXs[1] += DP1;
+    }
+
+
+    void FPBDJointCachedSolver::ApplyRotationDelta(
+        const int32 BodyIndex,
+        const FVec3& DR)
+    {
+        //UE_LOG(LogChaosJoint, VeryVerbose, TEXT("      Apply DR%d %f %f %f"), BodyIndex, DR.X, DR.Y, DR.Z);
+
+        Body(BodyIndex).ApplyRotationDelta(DR);
+
+        UpdateDerivedState(BodyIndex);
+    }
+
+
+    void FPBDJointCachedSolver::ApplyRotationDelta(
+        const FVec3& DR0,
+        const FVec3& DR1)
+    {
+        //UE_LOG(LogChaosJoint, VeryVerbose, TEXT("      Apply DR%d %f %f %f"), 0, DR0.X, DR0.Y, DR0.Z);
+        //UE_LOG(LogChaosJoint, VeryVerbose, TEXT("      Apply DR%d %f %f %f"), 1, DR1.X, DR1.Y, DR1.Z);
+
+        if (bRealTypeCompatibleWithISPC && bChaos_Joint_ISPC_Enabled)
+        {
+#if 0
+            ispc::ApplyRotationDelta2((ispc::FPBDJointCachedSolver*)this, (ispc::FVector&)DR0, (ispc::FVector&)DR1);
+#endif
+        }
+        else
+        {
+            if (Body0().IsDynamic())
+            {
+                Body(0).ApplyRotationDelta(DR0);
+            }
+            if (Body1().IsDynamic())
+            {
+                Body(1).ApplyRotationDelta(DR1);
+            }
+
+            UpdateDerivedState();
+        }
+    }
+
+    void FPBDJointCachedSolver::ApplyDelta(
+        const int32 BodyIndex,
+        const FVec3& DP,
+        const FVec3& DR)
+    {
+        //UE_LOG(LogChaosJoint, VeryVerbose, TEXT("      Apply DP%d %f %f %f"), BodyIndex, DP.X, DP.Y, DP.Z);
+        //UE_LOG(LogChaosJoint, VeryVerbose, TEXT("      Apply DR%d %f %f %f"), BodyIndex, DR.X, DR.Y, DR.Z);
+
+        Body(BodyIndex).ApplyTransformDelta(DP, DR);
+
+        UpdateDerivedState(BodyIndex);
+    }
+
+
+    void FPBDJointCachedSolver::ApplyVelocityDelta(
+        const int32 BodyIndex,
+        const FVec3& DV,
+        const FVec3& DW)
+    {
+        Body(BodyIndex).ApplyVelocityDelta(DV, DW);
+    }
+
+
+    void FPBDJointCachedSolver::ApplyVelocityDelta(
+        const FVec3& DV0,
+        const FVec3& DW0,
+        const FVec3& DV1,
+        const FVec3& DW1)
+    {
+        Body(0).ApplyVelocityDelta(DV0, DW0);
+        Body(1).ApplyVelocityDelta(DV1, DW1);
+    }
+
+    void FPBDJointCachedSolver::ApplyAngularVelocityDelta(
+        const FVec3& DW0,
+        const FVec3& DW1)
+    {
+        Body(0).ApplyAngularVelocityDelta(DW0);
+        Body(1).ApplyAngularVelocityDelta(DW1);
+    }
+
+    void FPBDJointCachedSolver::ApplyPositionConstraint(
+        const FReal Stiffness,
+        const FVec3& Axis,
+        const FReal Delta,
+        const FReal IM,
+        const FVec3& Arm0,
+        const FVec3& Arm1,
+        const FVec3& DRAxis0,
+        const FVec3& DRAxis1,
+        const int32 LinearHardLambdaIndex)
+    {
+        const FReal DLambda = Stiffness * Delta / IM;
+        const FVec3 DX = Axis * DLambda;
+        const FVec3 DP0 = InvM(0) * DX;
+        const FVec3 DP1 = -InvM(1) * DX;
+        const FVec3 DR0 = DRAxis0 * DLambda;
+        const FVec3 DR1 = DRAxis1 * DLambda;
+
+        ApplyPositionDelta(DP0, DP1);
+        ApplyRotationDelta(DR0, DR1);
+
+        NetLinearImpulse += DX;
+
+        if(LinearHardLambdaIndex >= 0)
+        {
+            LinearHardLambda[LinearHardLambdaIndex] += DLambda;
+        }
+
+        ++NumActiveConstraints;
+    }
+
+
+    void FPBDJointCachedSolver::ApplyPositionConstraintSoft(
+        const FReal Dt,
+        const FReal JointStiffness,
+        const FReal JointDamping,
+        const bool bAccelerationMode,
+        const FVec3& Axis,
+        const FReal Delta,
+        const FReal TargetVel,
+        FReal& Lambda)
+    {
+        if (bRealTypeCompatibleWithISPC && bChaos_Joint_ISPC_Enabled)
+        {
+#if 0
+            float ReturnedLambda = Lambda;
+            ispc::ApplyPositionConstraintSoft((ispc::FPBDJointCachedSolver*)this, Dt, JointStiffness, JointDamping, bAccelerationMode, (ispc::FVector&)Axis, Delta, TargetVel, ReturnedLambda);
+            Lambda = ReturnedLambda;
+#endif
+        }
+        else
+        {
+            // Joint-space inverse mass
+            const FVec3 AngularAxis0 = FVec3::CrossProduct(ConnectorXs[0] - P(0), Axis);
+            const FVec3 AngularAxis1 = FVec3::CrossProduct(ConnectorXs[1] - P(1), Axis);
+            const FVec3 IA0 = Utilities::Multiply(InvI(0), AngularAxis0);
+            const FVec3 IA1 = Utilities::Multiply(InvI(1), AngularAxis1);
+            const FReal II0 = FVec3::DotProduct(AngularAxis0, IA0);
+            const FReal II1 = FVec3::DotProduct(AngularAxis1, IA1);
+            const FReal II = (InvM(0) + II0 + InvM(1) + II1);
+    
+            FReal VelDt = 0;
+            if (JointDamping > KINDA_SMALL_NUMBER)
+            {
+                const FVec3 V0Dt = FVec3::CalculateVelocity(InitConnectorXs[0], ConnectorXs[0], 1.0f);
+                const FVec3 V1Dt = FVec3::CalculateVelocity(InitConnectorXs[1], ConnectorXs[1], 1.0f);
+                VelDt = TargetVel * Dt + FVec3::DotProduct(V0Dt - V1Dt, Axis);
+            }
+    
+            const FReal SpringMassScale = (bAccelerationMode) ? 1.0f / (InvM(0) + InvM(1)) : 1.0f;
+            const FReal S = SpringMassScale * JointStiffness * Dt * Dt;
+            const FReal D = SpringMassScale * JointDamping * Dt;
+            const FReal Multiplier = (FReal)1 / ((S + D) * II + (FReal)1);
+            const FReal DLambda = SolverStiffness * Multiplier * (S * Delta - D * VelDt - Lambda);
+    
+            const FVec3 DP0 = (InvM(0) * DLambda) * Axis;
+            const FVec3 DP1 = (-InvM(1) * DLambda) * Axis;
+            const FVec3 DR0 = DLambda * Utilities::Multiply(InvI(0), AngularAxis0);
+            const FVec3 DR1 = -DLambda * Utilities::Multiply(InvI(1), AngularAxis1);
+    
+            ApplyPositionDelta(DP0, DP1);
+            ApplyRotationDelta(DR0, DR1);
+
+            Lambda += DLambda;
+            NetLinearImpulse += DLambda * Axis;
+        }
+
+        ++NumActiveConstraints;
+    }
+    
+
+    void FPBDJointCachedSolver::ApplyRotationConstraintKD(
+        const int32 KIndex,
+        const int32 DIndex,
+        const FReal JointStiffness,
+        const FVec3& Axis,
+        const FReal Angle,
+        const int32 AngularHardLambdaIndex)
+    {
+        const FReal Stiffness = SolverStiffness * JointStiffness;
+
+        const FVec3 IA1 = Utilities::Multiply(InvI(DIndex), Axis);
+        const FReal II1 = FVec3::DotProduct(Axis, IA1);
+        const FReal DR = Stiffness * (Angle / II1);
+        const FVec3 DR1 = IA1 * -DR;
+        ApplyRotationDelta(DIndex, DR1);
+
+        NetAngularImpulse += (KIndex == 0 )? DR * Axis : -DR * Axis;
+        if (AngularHardLambdaIndex >= 0)
+        {
+            AngularHardLambda[AngularHardLambdaIndex] += DR;
+        }
+    }
+
+
+    void FPBDJointCachedSolver::ApplyRotationConstraintDD(
+        const FReal JointStiffness,
+        const FVec3& Axis,
+        const FReal Angle,
+        const int32 AngularHardLambdaIndex)
+    {
+        const FReal Stiffness = SolverStiffness * JointStiffness;
+
+        // Joint-space inverse mass
+        const FVec3 IA0 = Utilities::Multiply(InvI(0), Axis);
+        const FVec3 IA1 = Utilities::Multiply(InvI(1), Axis);
+        const FReal II0 = FVec3::DotProduct(Axis, IA0);
+        const FReal II1 = FVec3::DotProduct(Axis, IA1);
+
+        const FReal DR = Stiffness * Angle / (II0 + II1);
+        const FVec3 DR0 = IA0 * DR;
+        const FVec3 DR1 = IA1 * -DR;
+
+        ApplyRotationDelta(DR0, DR1);
+
+        NetAngularImpulse += Axis * DR;
+        if (AngularHardLambdaIndex >= 0)
+        {
+            AngularHardLambda[AngularHardLambdaIndex] += DR;
+        }
+    }
+
+
+    void FPBDJointCachedSolver::ApplyRotationConstraint(
+        const FReal JointStiffness,
+        const FVec3& Axis,
+        const FReal Angle,
+        const int32 AngularHardLambdaIndex)
+    {
+        if (!Body0().IsDynamic())
+        {
+            ApplyRotationConstraintKD(0, 1, JointStiffness, Axis, Angle, AngularHardLambdaIndex);
+        }
+        else if (!Body1().IsDynamic())
+        {
+            ApplyRotationConstraintKD(1, 0, JointStiffness, Axis, -Angle, AngularHardLambdaIndex);
+        }
+        else
+        {
+            ApplyRotationConstraintDD(JointStiffness, Axis, Angle, AngularHardLambdaIndex);
+        }
+
+        ++NumActiveConstraints;
+    }
+
+
+    // See "XPBD: Position-Based Simulation of Compliant Constrained Dynamics"
+    void FPBDJointCachedSolver::ApplyRotationConstraintSoftKD(
+        const int32 KIndex,
+        const int32 DIndex,
+        const FReal Dt,
+        const FReal JointStiffness,
+        const FReal JointDamping,
+        const bool bAccelerationMode,
+        const FVec3& Axis,
+        const FReal Angle,
+        const FReal AngVelTarget,
+        FReal& Lambda)
+    {
+        check(!Body(KIndex).IsDynamic());
+        check(Body(DIndex).IsDynamic());
+
+        if (bRealTypeCompatibleWithISPC && bChaos_Joint_ISPC_Enabled)
+        {
+#if 0
+            float ReturnedLambda = Lambda;
+            ispc::ApplyRotationConstraintSoftKD((ispc::FPBDJointCachedSolver*)this, KIndex, DIndex, Dt, JointStiffness, JointDamping, bAccelerationMode, (ispc::FVector&) Axis, Angle, AngVelTarget, ReturnedLambda);
+            Lambda = ReturnedLambda;
+#endif
+        }
+        else
+        {
+            // World-space inverse mass
+            const FVec3 IA1 = Utilities::Multiply(InvI(DIndex), Axis);
+
+            // Joint-space inverse mass
+            FReal II1 = FVec3::DotProduct(Axis, IA1);
+            const FReal II = II1;
+
+            // Damping angular velocity
+            FReal AngVelDt = 0;
+            if (JointDamping > KINDA_SMALL_NUMBER)
+            {
+                const FVec3 W0Dt = FRotation3::CalculateAngularVelocity(InitConnectorRs[KIndex], ConnectorRs[KIndex], 1.0f);
+                const FVec3 W1Dt = FRotation3::CalculateAngularVelocity(InitConnectorRs[DIndex], ConnectorRs[DIndex], 1.0f);
+                AngVelDt = AngVelTarget * Dt + FVec3::DotProduct(Axis, W0Dt - W1Dt);
+            }
+
+            const FReal SpringMassScale = (bAccelerationMode) ? 1.0f / II : 1.0f;
+            const FReal S = SpringMassScale * JointStiffness * Dt * Dt;
+            const FReal D = SpringMassScale * JointDamping * Dt;
+            const FReal Multiplier = (FReal)1 / ((S + D) * II + (FReal)1);
+            const FReal DLambda = SolverStiffness * Multiplier * (S * Angle - D * AngVelDt - Lambda);
+
+            //const FVec3 DR1 = IA1 * -DLambda;
+            const FVec3 DR1 = Axis * -(DLambda * II1);
+
+            ApplyRotationDelta(DIndex, DR1);
+    
+            Lambda += DLambda;
+            NetAngularImpulse += (KIndex == 0 ? 1 : -1) * DLambda * Axis;
+        }
+    }
+
+    // See "XPBD: Position-Based Simulation of Compliant Constrained Dynamics"
+    void FPBDJointCachedSolver::ApplyRotationConstraintSoftDD(
+        const FReal Dt,
+        const FReal JointStiffness,
+        const FReal JointDamping,
+        const bool bAccelerationMode,
+        const FVec3& Axis,
+        const FReal Angle,
+        const FReal AngVelTarget,
+        FReal& Lambda)
+    {
+        check(Body(0).IsDynamic());
+        check(Body(1).IsDynamic());
+
+        if (bRealTypeCompatibleWithISPC && bChaos_Joint_ISPC_Enabled)
+        {
+#if 0
+            float ReturnedLambda = Lambda;
+            ispc::ApplyRotationConstraintSoftDD((ispc::FPBDJointCachedSolver*)this, Dt, JointStiffness, JointDamping, bAccelerationMode, (ispc::FVector&) Axis, Angle, AngVelTarget, ReturnedLambda);
+            Lambda = ReturnedLambda;
+#endif
+        }
+        else
+        {
+            // World-space inverse mass
+            const FVec3 IA0 = Utilities::Multiply(InvI(0), Axis);
+            const FVec3 IA1 = Utilities::Multiply(InvI(1), Axis);
+
+            // Joint-space inverse mass
+            FReal II0 = FVec3::DotProduct(Axis, IA0);
+            FReal II1 = FVec3::DotProduct(Axis, IA1);
+            const FReal II = (II0 + II1);
+
+            // Damping angular velocity
+            FReal AngVelDt = 0;
+            if (JointDamping > KINDA_SMALL_NUMBER)
+            {
+                const FVec3 W0Dt = FRotation3::CalculateAngularVelocity(InitConnectorRs[0], ConnectorRs[0], 1.0f);
+                const FVec3 W1Dt = FRotation3::CalculateAngularVelocity(InitConnectorRs[1], ConnectorRs[1], 1.0f);
+                AngVelDt = AngVelTarget * Dt + FVec3::DotProduct(Axis, W0Dt - W1Dt);
+            }
+
+            const FReal SpringMassScale = (bAccelerationMode) ? 1.0f / II : 1.0f;
+            const FReal S = SpringMassScale * JointStiffness * Dt * Dt;
+            const FReal D = SpringMassScale * JointDamping * Dt;
+            const FReal Multiplier = (FReal)1 / ((S + D) * II + (FReal)1);
+            const FReal DLambda = SolverStiffness * Multiplier * (S * Angle - D * AngVelDt - Lambda);
+
+            //const FVec3 DR0 = IA0 * DLambda;
+            //const FVec3 DR1 = IA1 * -DLambda;
+            const FVec3 DR0 = Axis * (DLambda * II0);
+            const FVec3 DR1 = Axis * -(DLambda * II1);
+
+            ApplyRotationDelta(DR0, DR1);
+
+            Lambda += DLambda;
+            NetAngularImpulse += DLambda * Axis;
+        }
+    }
+
+    void FPBDJointCachedSolver::ApplyRotationConstraintSoft(
+        const FReal Dt,
+        const FReal JointStiffness,
+        const FReal JointDamping,
+        const bool bAccelerationMode,
+        const FVec3& Axis,
+        const FReal Angle,
+        const FReal AngVelTarget,
+        FReal& Lambda)
+    {
+        if (!Body0().IsDynamic())
+        {
+            ApplyRotationConstraintSoftKD(0, 1, Dt, JointStiffness, JointDamping, bAccelerationMode, Axis, Angle, AngVelTarget, Lambda);
+        }
+        else if (!Body1().IsDynamic())
+        {
+            ApplyRotationConstraintSoftKD(1, 0, Dt, JointStiffness, JointDamping, bAccelerationMode, Axis, -Angle, -AngVelTarget, Lambda);
+        }
+        else
+        {
+            ApplyRotationConstraintSoftDD(Dt, JointStiffness, JointDamping, bAccelerationMode, Axis, Angle, AngVelTarget, Lambda);
+        }
+
+        ++NumActiveConstraints;
+    }
+
+    // Used for non-zero restitution. We pad constraints by an amount such that the velocity
+    // calculated after solving constraint positions will as required for the restitution.
+    void FPBDJointCachedSolver::CalculateLinearConstraintPadding(
+        const FReal Dt,
+        const FPBDJointSolverSettings& SolverSettings,
+        const FPBDJointSettings& JointSettings,
+        const FReal Restitution,
+        const int32 AxisIndex,
+        const FVec3 Axis,
+        FReal& InOutPos)
+    {
+        // NOTE: We only calculate the padding after the constraint is first violated, and after
+        // that the padding is fixed for the rest of the iterations in the current step.
+        if ((Restitution > 0.0f) && (InOutPos > 0.0f) && !HasLinearConstraintPadding(AxisIndex))
+        {
+            SetLinearConstraintPadding(AxisIndex, 0.0f);
+
+            // Calculate the velocity we want to match
+            const FVec3 V0Dt = FVec3::CalculateVelocity(InitConnectorXs[0], ConnectorXs[0], 1.0f);
+            const FVec3 V1Dt = FVec3::CalculateVelocity(InitConnectorXs[1], ConnectorXs[1], 1.0f);
+            const FReal AxisVDt = FVec3::DotProduct(V1Dt - V0Dt, Axis);
+
+            // Calculate the padding to apply to the constraint that will result in the
+            // desired outward velocity (assuming the constraint is fully resolved)
+            const FReal Padding = (1.0f + Restitution) * AxisVDt - InOutPos;
+            if (Padding > 0.0f)
+            {
+                SetLinearConstraintPadding(AxisIndex, Padding);
+                InOutPos += Padding;
+            }
+        }
+    }
+
+    // Used for non-zero restitution. We pad constraints by an amount such that the velocity
+    // calculated after solving constraint positions will as required for the restitution.
+    void FPBDJointCachedSolver::CalculateAngularConstraintPadding(
+        const FReal Dt,
+        const FPBDJointSolverSettings& SolverSettings,
+        const FPBDJointSettings& JointSettings,
+        const FReal Restitution,
+        const EJointAngularConstraintIndex ConstraintIndex,
+        const FVec3 Axis,
+        FReal& InOutAngle)
+    {
+        // NOTE: We only calculate the padding after the constraint is first violated, and after
+        // that the padding is fixed for the rest of the iterations in the current step.
+        if ((Restitution > 0.0f) && (InOutAngle > 0.0f) && !HasAngularConstraintPadding(ConstraintIndex))
+        {
+            SetAngularConstraintPadding(ConstraintIndex, 0.0f);
+
+            // Calculate the velocity we want to match
+            const FVec3 W0Dt = FRotation3::CalculateAngularVelocity(InitConnectorRs[0], ConnectorRs[0], 1.0f);
+            const FVec3 W1Dt = FRotation3::CalculateAngularVelocity(InitConnectorRs[1], ConnectorRs[1], 1.0f);
+            const FReal AxisWDt = FVec3::DotProduct(W1Dt - W0Dt, Axis);
+
+            // Calculate the padding to apply to the constraint that will result in the
+            // desired outward velocity (assuming the constraint is fully resolved)
+            const FReal Padding = (1.0f + Restitution) * AxisWDt - InOutAngle;
+            if (Padding > 0.0f)
+            {
+                SetAngularConstraintPadding(ConstraintIndex, Padding);
+                InOutAngle += Padding;
+            }
+        }
+    }
+
+    //
+    //
+    //////////////////////////////////////////////////////////////////////////
+    //
+    //
+
+    void FPBDJointCachedSolver::ApplyLockedRotationConstraints(
+        const FReal Dt,
+        const FPBDJointSolverSettings& SolverSettings,
+        const FPBDJointSettings& JointSettings,
+        const bool bApplyTwist,
+        const bool bApplySwing)
+    {
+        FVec3 Axis0, Axis1, Axis2;
+        FPBDJointUtilities::GetLockedRotationAxes(ConnectorRs[0], ConnectorRs[1], Axis0, Axis1, Axis2);
+
+        const FRotation3 R01 = ConnectorRs[0].Inverse() * ConnectorRs[1];
+
+        if (bApplyTwist)
+        {
+            FReal TwistStiffness = FPBDJointUtilities::GetTwistStiffness(SolverSettings, JointSettings);
+            ApplyRotationConstraint(TwistStiffness, Axis0, R01.X, (int32)EJointAngularConstraintIndex::Twist);
+        }
+
+        if (bApplySwing)
+        {
+            FReal SwingStiffness = FPBDJointUtilities::GetSwingStiffness(SolverSettings, JointSettings);
+            ApplyRotationConstraint(SwingStiffness, Axis1, R01.Y, (int32)EJointAngularConstraintIndex::Swing2);
+            ApplyRotationConstraint(SwingStiffness, Axis2, R01.Z, (int32)EJointAngularConstraintIndex::Swing1);
+        }
+    }
+
+    void FPBDJointCachedSolver::ApplyTwistConstraint(
+        const FReal Dt,
+        const FPBDJointSolverSettings& SolverSettings,
+        const FPBDJointSettings& JointSettings,
+        const bool bUseSoftLimit)
+    {
+        FVec3 TwistAxis;
+        FReal TwistAngle;
+        FPBDJointUtilities::GetTwistAxisAngle(ConnectorRs[0], ConnectorRs[1], TwistAxis, TwistAngle);
+
+        // Calculate the twist correction to apply to each body
+        const FReal LimitPadding = GetAngularConstraintPadding(EJointAngularConstraintIndex::Twist);
+        FReal DTwistAngle = 0;
+        FReal TwistAngleMax = FMath::Max(JointSettings.AngularLimits[(int32)EJointAngularConstraintIndex::Twist] - LimitPadding, (FReal)0.);
+        if (TwistAngle > TwistAngleMax)
+        {
+            DTwistAngle = TwistAngle - TwistAngleMax;
+        }
+        else if (TwistAngle < -TwistAngleMax)
+        {
+            // Keep Twist error positive
+            DTwistAngle = -TwistAngle - TwistAngleMax;
+            TwistAxis = -TwistAxis;
+        }
+
+        UE_LOG(LogChaosJoint, VeryVerbose, TEXT("    Twist Angle %f [Limit %f]"), FMath::RadiansToDegrees(TwistAngle), FMath::RadiansToDegrees(TwistAngleMax));
+
+        // Apply twist correction
+        if (DTwistAngle > AngleTolerance)
+        {
+            if (bUseSoftLimit)
+            {
+                const FReal TwistStiffness = FPBDJointUtilities::GetSoftTwistStiffness(SolverSettings, JointSettings);
+                const FReal TwistDamping = FPBDJointUtilities::GetSoftTwistDamping(SolverSettings, JointSettings);
+                const bool bAccelerationMode = FPBDJointUtilities::GetAngularSoftAccelerationMode(SolverSettings, JointSettings);
+                ApplyRotationConstraintSoft(Dt, TwistStiffness, TwistDamping, bAccelerationMode, TwistAxis, DTwistAngle, 0.0f, TwistSoftLambda);
+            }
+            else
+            {
+                if (JointSettings.TwistRestitution > 0.0f)
+                {
+                    CalculateAngularConstraintPadding(Dt, SolverSettings, JointSettings, JointSettings.TwistRestitution, EJointAngularConstraintIndex::Twist, TwistAxis, DTwistAngle);
+                }
+
+                FReal TwistStiffness = FPBDJointUtilities::GetTwistStiffness(SolverSettings, JointSettings);
+                ApplyRotationConstraint(TwistStiffness, TwistAxis, DTwistAngle, (int32)EJointAngularConstraintIndex::Twist);
+            }
+        }
+    }
+
+    void FPBDJointCachedSolver::ApplyConeConstraint(
+        const FReal Dt,
+        const FPBDJointSolverSettings& SolverSettings,
+        const FPBDJointSettings& JointSettings,
+        const bool bUseSoftLimit)
+    {
+        FVec3 SwingAxisLocal;
+        FReal DSwingAngle = 0.0f;
+
+        const FReal LimitPadding = GetAngularConstraintPadding(EJointAngularConstraintIndex::Swing1);
+        const FReal Swing1Limit = FMath::Max(JointSettings.AngularLimits[(int32)EJointAngularConstraintIndex::Swing1] - LimitPadding, (FReal)0.);
+        const FReal Swing2Limit = FMath::Max(JointSettings.AngularLimits[(int32)EJointAngularConstraintIndex::Swing2] - LimitPadding, (FReal)0.);
+        FPBDJointUtilities::GetEllipticalConeAxisErrorLocal(ConnectorRs[0], ConnectorRs[1], Swing2Limit, Swing1Limit, SwingAxisLocal, DSwingAngle);
+
+        UE_LOG(LogChaosJoint, VeryVerbose, TEXT("    Cone Error %f [Limits %f %f]"), FMath::RadiansToDegrees(DSwingAngle), FMath::RadiansToDegrees(Swing2Limit), FMath::RadiansToDegrees(Swing1Limit));
+
+        const FVec3 SwingAxis = ConnectorRs[0] * SwingAxisLocal;
+
+        // Apply swing correction to each body
+        if (DSwingAngle > AngleTolerance)
+        {
+            if (bUseSoftLimit)
+            {
+                const FReal SoftSwingStiffness = FPBDJointUtilities::GetSoftSwingStiffness(SolverSettings, JointSettings);
+                const FReal SoftSwingDamping = FPBDJointUtilities::GetSoftSwingDamping(SolverSettings, JointSettings);
+                const bool bAccelerationMode = FPBDJointUtilities::GetAngularSoftAccelerationMode(SolverSettings, JointSettings);
+                ApplyRotationConstraintSoft(Dt, SoftSwingStiffness, SoftSwingDamping, bAccelerationMode, SwingAxis, DSwingAngle, 0.0f, SwingSoftLambda);
+            }
+            else
+            {
+                if (JointSettings.SwingRestitution > 0.0f)
+                {
+                    CalculateAngularConstraintPadding(Dt, SolverSettings, JointSettings, JointSettings.SwingRestitution, EJointAngularConstraintIndex::Swing1, SwingAxis, DSwingAngle);
+                }
+
+                FReal SwingStiffness = FPBDJointUtilities::GetSwingStiffness(SolverSettings, JointSettings);
+                // For cone constraints, the lambda are all accumulated in Swing2
+                ApplyRotationConstraint(SwingStiffness, SwingAxis, DSwingAngle, (int32)EJointAngularConstraintIndex::Swing2);
+            }
+        }
+    }
+
+
+    void FPBDJointCachedSolver::ApplySingleLockedSwingConstraint(
+        const FReal Dt,
+        const FPBDJointSolverSettings& SolverSettings,
+        const FPBDJointSettings& JointSettings,
+        const EJointAngularConstraintIndex SwingConstraintIndex,
+        const bool bUseSoftLimit)
+    {
+        // NOTE: SwingAxis is not normalized in this mode. It has length Sin(SwingAngle).
+        // Likewise, the SwingAngle is actually Sin(SwingAngle)
+        FVec3 SwingAxis;
+        FReal SwingAngle;
+        FPBDJointUtilities::GetLockedSwingAxisAngle(ConnectorRs[0], ConnectorRs[1], SwingConstraintIndex, SwingAxis, SwingAngle);
+
+        UE_LOG(LogChaosJoint, VeryVerbose, TEXT("    LockedSwing%d Angle %f [Tolerance %f]"), (SwingConstraintIndex == EJointAngularConstraintIndex::Swing1) ? 1 : 2, FMath::RadiansToDegrees(SwingAngle), FMath::RadiansToDegrees(AngleTolerance));
+
+        // Apply swing correction
+        FReal DSwingAngle = SwingAngle;
+        if (FMath::Abs(DSwingAngle) > AngleTolerance)
+        {
+            if (bUseSoftLimit)
+            {
+                const FReal SoftSwingStiffness = FPBDJointUtilities::GetSoftSwingStiffness(SolverSettings, JointSettings);
+                const FReal SoftSwingDamping = FPBDJointUtilities::GetSoftSwingDamping(SolverSettings, JointSettings);
+                const bool bAccelerationMode = FPBDJointUtilities::GetAngularSoftAccelerationMode(SolverSettings, JointSettings);
+                ApplyRotationConstraintSoft(Dt, SoftSwingStiffness, SoftSwingDamping, bAccelerationMode, SwingAxis, DSwingAngle, 0.0f, SwingSoftLambda);
+            }
+            else
+            {
+                const FReal SwingStiffness = FPBDJointUtilities::GetSwingStiffness(SolverSettings, JointSettings);
+                ApplyRotationConstraint(SwingStiffness, SwingAxis, DSwingAngle, (int32)SwingConstraintIndex);
+            }
+        }
+    }
+
+
+    void FPBDJointCachedSolver::ApplyDualConeSwingConstraint(
+        const FReal Dt,
+        const FPBDJointSolverSettings& SolverSettings,
+        const FPBDJointSettings& JointSettings,
+        const EJointAngularConstraintIndex SwingConstraintIndex,
+        const bool bUseSoftLimit)
+    {
+        FVec3 SwingAxis;
+        FReal SwingAngle;
+        FPBDJointUtilities::GetDualConeSwingAxisAngle(ConnectorRs[0], ConnectorRs[1], SwingConstraintIndex, SwingAxis, SwingAngle);
+
+        // Calculate swing error we need to correct
+        const FReal LimitPadding = GetAngularConstraintPadding(SwingConstraintIndex);
+        FReal DSwingAngle = 0;
+        const FReal SwingAngleMax = FMath::Max(JointSettings.AngularLimits[(int32)SwingConstraintIndex] - LimitPadding, (FReal)0.);
+        if (SwingAngle > SwingAngleMax)
+        {
+            DSwingAngle = SwingAngle - SwingAngleMax;
+        }
+        else if (SwingAngle < -SwingAngleMax)
+        {
+            // Keep the error positive
+            DSwingAngle = -SwingAngle - SwingAngleMax;
+            SwingAxis = -SwingAxis;
+        }
+
+        UE_LOG(LogChaosJoint, VeryVerbose, TEXT("    DualConeSwing%d Angle %f [Limit %f]"), (SwingConstraintIndex == EJointAngularConstraintIndex::Swing1) ? 1 : 2, FMath::RadiansToDegrees(SwingAngle), FMath::RadiansToDegrees(SwingAngleMax));
+
+        // Apply swing correction
+        if (DSwingAngle > SolverSettings.AngleTolerance)
+        {
+            if (bUseSoftLimit)
+            {
+                const FReal SoftSwingStiffness = FPBDJointUtilities::GetSoftSwingStiffness(SolverSettings, JointSettings);
+                const FReal SoftSwingDamping = FPBDJointUtilities::GetSoftSwingDamping(SolverSettings, JointSettings);
+                const bool bAccelerationMode = FPBDJointUtilities::GetAngularSoftAccelerationMode(SolverSettings, JointSettings);
+                ApplyRotationConstraintSoft(Dt, SoftSwingStiffness, SoftSwingDamping, bAccelerationMode, SwingAxis, DSwingAngle, 0.0f, SwingSoftLambda);
+            }
+            else
+            {
+                if (JointSettings.SwingRestitution > 0.0f)
+                {
+                    CalculateAngularConstraintPadding(Dt, SolverSettings, JointSettings, JointSettings.SwingRestitution, SwingConstraintIndex, SwingAxis, DSwingAngle);
+                }
+
+                const FReal SwingStiffness = FPBDJointUtilities::GetSwingStiffness(SolverSettings, JointSettings);
+                ApplyRotationConstraint(SwingStiffness, SwingAxis, DSwingAngle, (int32)SwingConstraintIndex);
+            }
+        }
+    }
+
+
+    void FPBDJointCachedSolver::ApplySwingConstraint(
+        const FReal Dt,
+        const FPBDJointSolverSettings& SolverSettings,
+        const FPBDJointSettings& JointSettings,
+        const EJointAngularConstraintIndex SwingConstraintIndex,
+        const bool bUseSoftLimit)
+    {
+        FVec3 SwingAxis;
+        FReal SwingAngle;
+        FPBDJointUtilities::GetSwingAxisAngle(ConnectorRs[0], ConnectorRs[1], SolverSettings.SwingTwistAngleTolerance, SwingConstraintIndex, SwingAxis, SwingAngle);
+
+        // Calculate swing error we need to correct
+        const FReal LimitPadding = GetAngularConstraintPadding(SwingConstraintIndex);
+        FReal DSwingAngle = 0;
+        const FReal SwingAngleMax = FMath::Max(JointSettings.AngularLimits[(int32)SwingConstraintIndex] - LimitPadding, (FReal)0.);
+        if (SwingAngle > SwingAngleMax)
+        {
+            DSwingAngle = SwingAngle - SwingAngleMax;
+        }
+        else if (SwingAngle < -SwingAngleMax)
+        {
+            // Keep swing error positive
+            DSwingAngle = -SwingAngle - SwingAngleMax;
+            SwingAxis = -SwingAxis;
+        }
+
+        UE_LOG(LogChaosJoint, VeryVerbose, TEXT("    Swing%d Angle %f [Limit %f]"), (SwingConstraintIndex == EJointAngularConstraintIndex::Swing1) ? 1 : 2, FMath::RadiansToDegrees(SwingAngle), FMath::RadiansToDegrees(SwingAngleMax));
+
+        // Apply swing correction
+        if (DSwingAngle > AngleTolerance)
+        {
+            if (bUseSoftLimit)
+            {
+                const FReal SoftSwingStiffness = FPBDJointUtilities::GetSoftSwingStiffness(SolverSettings, JointSettings);
+                const FReal SoftSwingDamping = FPBDJointUtilities::GetSoftSwingDamping(SolverSettings, JointSettings);
+                const bool bAccelerationMode = FPBDJointUtilities::GetAngularSoftAccelerationMode(SolverSettings, JointSettings);
+                ApplyRotationConstraintSoft(Dt, SoftSwingStiffness, SoftSwingDamping, bAccelerationMode, SwingAxis, DSwingAngle, 0.0f, SwingSoftLambda);
+            }
+            else
+            {
+                if (JointSettings.SwingRestitution > 0.0f)
+                {
+                    CalculateAngularConstraintPadding(Dt, SolverSettings, JointSettings, JointSettings.SwingRestitution, SwingConstraintIndex, SwingAxis, DSwingAngle);
+                }
+
+                const FReal SwingStiffness = FPBDJointUtilities::GetSwingStiffness(SolverSettings, JointSettings);
+                ApplyRotationConstraint(SwingStiffness, SwingAxis, DSwingAngle, (int32)SwingConstraintIndex);
+            }
+        }
+    }
+
+
+    void FPBDJointCachedSolver::ApplySwingTwistDrives(
+        const FReal Dt,
+        const FPBDJointSolverSettings& SolverSettings,
+        const FPBDJointSettings& JointSettings,
+        const bool bTwistDriveEnabled,
+        const bool bSwing1DriveEnabled,
+        const bool bSwing2DriveEnabled)
+    {
+        FRotation3 R1Target = ConnectorRs[0] * JointSettings.AngularDrivePositionTarget;
+        R1Target.EnforceShortestArcWith(ConnectorRs[1]);
+        FRotation3 R1Error = R1Target.Inverse() * ConnectorRs[1];
+        FVec3 R1TwistAxisError = R1Error * FJointConstants::TwistAxis();
+
+        // Angle approximation Angle ~= Sin(Angle) for small angles, underestimates for large angles
+        const FReal DTwistAngle = 2.0f * R1Error.X;
+        const FReal DSwing1Angle = R1TwistAxisError.Y;
+        const FReal DSwing2Angle = -R1TwistAxisError.Z;
+
+        const FReal AngularTwistDriveStiffness = FPBDJointUtilities::GetAngularTwistDriveStiffness(SolverSettings, JointSettings);
+        const FReal AngularTwistDriveDamping = FPBDJointUtilities::GetAngularTwistDriveDamping(SolverSettings, JointSettings);
+        const FReal AngularSwingDriveStiffness = FPBDJointUtilities::GetAngularSwingDriveStiffness(SolverSettings, JointSettings);
+        const FReal AngularSwingDriveDamping = FPBDJointUtilities::GetAngularSwingDriveDamping(SolverSettings, JointSettings);
+        const bool bAccelerationMode = FPBDJointUtilities::GetDriveAccelerationMode(SolverSettings, JointSettings);
+
+        const bool bUseTwistDrive = bTwistDriveEnabled && (((FMath::Abs(DTwistAngle) > AngleTolerance) && (AngularTwistDriveStiffness > 0.0f)) || (AngularTwistDriveDamping > 0.0f));
+        if (bUseTwistDrive)
+        {
+            FReal ReturnedLambda = RotationDriveLambdas[(int32)EJointAngularConstraintIndex::Twist];
+            const FVec3 TwistAxis = ConnectorRs[1] * FJointConstants::TwistAxis();
+            ApplyRotationConstraintSoft(Dt, AngularTwistDriveStiffness, AngularTwistDriveDamping, bAccelerationMode, TwistAxis, DTwistAngle, JointSettings.AngularDriveVelocityTarget[(int32)EJointAngularConstraintIndex::Twist], ReturnedLambda);
+            RotationDriveLambdas[(int32)EJointAngularConstraintIndex::Twist] = ReturnedLambda;
+        }
+
+        const bool bUseSwing1Drive = bSwing1DriveEnabled && (((FMath::Abs(DSwing1Angle) > AngleTolerance) && (AngularSwingDriveStiffness > 0.0f)) || (AngularSwingDriveDamping > 0.0f));
+        if (bUseSwing1Drive)
+        {
+            FReal ReturnedLambda = RotationDriveLambdas[(int32)EJointAngularConstraintIndex::Swing1];
+            const FVec3 Swing1Axis = ConnectorRs[1] * FJointConstants::Swing1Axis();
+            ApplyRotationConstraintSoft(Dt, AngularSwingDriveStiffness, AngularSwingDriveDamping, bAccelerationMode, Swing1Axis, DSwing1Angle, JointSettings.AngularDriveVelocityTarget[(int32)EJointAngularConstraintIndex::Swing1], ReturnedLambda);
+            RotationDriveLambdas[(int32)EJointAngularConstraintIndex::Swing1] = ReturnedLambda;
+        }
+
+        const bool bUseSwing2Drive = bSwing2DriveEnabled && (((FMath::Abs(DSwing2Angle) > AngleTolerance) && (AngularSwingDriveStiffness > 0.0f)) || (AngularSwingDriveDamping > 0.0f));
+        if (bUseSwing2Drive)
+        {
+            FReal ReturnedLambda = RotationDriveLambdas[(int32)EJointAngularConstraintIndex::Swing2];
+            const FVec3 Swing2Axis = ConnectorRs[1] * FJointConstants::Swing2Axis();
+            ApplyRotationConstraintSoft(Dt, AngularSwingDriveStiffness, AngularSwingDriveDamping, bAccelerationMode, Swing2Axis, DSwing2Angle, JointSettings.AngularDriveVelocityTarget[(int32)EJointAngularConstraintIndex::Swing2], ReturnedLambda);
+            RotationDriveLambdas[(int32)EJointAngularConstraintIndex::Swing2] = ReturnedLambda;
+        }
+    }
+
+
+    void FPBDJointCachedSolver::ApplySLerpDrive(
+        const FReal Dt,
+        const FPBDJointSolverSettings& SolverSettings,
+        const FPBDJointSettings& JointSettings)
+    {
+        const FReal AngularDriveStiffness = FPBDJointUtilities::GetAngularSLerpDriveStiffness(SolverSettings, JointSettings);
+        const FReal AngularDriveDamping = FPBDJointUtilities::GetAngularSLerpDriveDamping(SolverSettings, JointSettings);
+        const bool bAccelerationMode = FPBDJointUtilities::GetDriveAccelerationMode(SolverSettings, JointSettings);
+
+        // If damping is enabled, we need to apply the drive about all 3 axes, but without damping we can just drive along the axis of error
+        if (AngularDriveDamping > 0.0f)
+        {
+            // NOTE: Slerp target velocity only works properly if we have a stiffness of zero.
+            FVec3 Axes[3] = { FVec3(1, 0, 0), FVec3(0, 1, 0), FVec3(0, 0, 1) };
+            if (AngularDriveStiffness > 0.0f)
+            {
+                FPBDJointUtilities::GetLockedRotationAxes(ConnectorRs[0], ConnectorRs[1], Axes[0], Axes[1], Axes[2]);
+                Utilities::NormalizeSafe(Axes[0], KINDA_SMALL_NUMBER);
+                Utilities::NormalizeSafe(Axes[1], KINDA_SMALL_NUMBER);
+                Utilities::NormalizeSafe(Axes[2], KINDA_SMALL_NUMBER);
+            }
+
+            const FRotation3 R01 = ConnectorRs[0].Inverse() * ConnectorRs[1];
+            FRotation3 TargetAngPos = JointSettings.AngularDrivePositionTarget;
+            TargetAngPos.EnforceShortestArcWith(R01);
+            const FRotation3 R1Error = TargetAngPos.Inverse() * R01;
+            FReal AxisAngles[3] = 
+            { 
+                2.0f * FMath::Asin(R1Error.X), 
+                2.0f * FMath::Asin(R1Error.Y), 
+                2.0f * FMath::Asin(R1Error.Z) 
+            };
+
+            const FVec3 TargetAngVel = ConnectorRs[0] * JointSettings.AngularDriveVelocityTarget;
+
+            for (int32 AxisIndex = 0; AxisIndex < 3; ++AxisIndex)
+            {
+                FReal ReturnedLambda = RotationDriveLambdas[AxisIndex];
+                const FReal AxisAngVel = FVec3::DotProduct(TargetAngVel, Axes[AxisIndex]);
+                ApplyRotationConstraintSoft(Dt, AngularDriveStiffness, AngularDriveDamping, bAccelerationMode, Axes[AxisIndex], AxisAngles[AxisIndex], AxisAngVel, ReturnedLambda);
+                RotationDriveLambdas[AxisIndex] = ReturnedLambda;
+            }
+        }
+        else
+        {
+            const FRotation3 TargetR1 = ConnectorRs[0] * JointSettings.AngularDrivePositionTarget;
+            const FRotation3 DR = TargetR1 * ConnectorRs[1].Inverse();
+
+            FVec3 SLerpAxis;
+            FReal SLerpAngle;
+            if (DR.ToAxisAndAngleSafe(SLerpAxis, SLerpAngle, FVec3(1, 0, 0)))
+            {
+                if (SLerpAngle > (FReal)PI)
+                {
+                    SLerpAngle = SLerpAngle - (FReal)2 * PI;
+                }
+
+                UE_LOG(LogChaosJoint, VeryVerbose, TEXT("      SLerpDrive Pos: %f Axis: %f %f %f"), -SLerpAngle, SLerpAxis.X, SLerpAxis.Y, SLerpAxis.Z);
+
+                if (FMath::Abs(SLerpAngle) > AngleTolerance)
+                {
+                    FReal ReturnedLambda = RotationDriveLambdas[(int32)EJointAngularConstraintIndex::Swing1];
+                    FReal AngVelTarget = (JointSettings.AngularDriveDamping > FReal(0)) ? FVec3::DotProduct(SLerpAxis, ConnectorRs[0] * JointSettings.AngularDriveVelocityTarget) : 0.0f;
+                    ApplyRotationConstraintSoft(Dt, AngularDriveStiffness, AngularDriveDamping, bAccelerationMode, SLerpAxis, -SLerpAngle, AngVelTarget, ReturnedLambda);
+                    RotationDriveLambdas[(int32)EJointAngularConstraintIndex::Swing1] = ReturnedLambda;
+                }
+            }
+        }
+    }
+
+
+    // Kinematic-Dynamic bodies
+    void FPBDJointCachedSolver::ApplyPointPositionConstraintKD(
+        const int32 KIndex,
+        const int32 DIndex,
+        const FReal Dt,
+        const FPBDJointSolverSettings& SolverSettings,
+        const FPBDJointSettings& JointSettings)
+    {
+        check(!Body(KIndex).IsDynamic());
+        check(Body(DIndex).IsDynamic());
+
+        FReal Stiffness = SolverStiffness * FPBDJointUtilities::GetLinearStiffness(SolverSettings, JointSettings);
+        const FVec3 CX = ConnectorXs[DIndex] - ConnectorXs[KIndex];
+
+        UE_LOG(LogChaosJoint, VeryVerbose, TEXT("    PointKD Delta %f [Limit %f]"), CX.Size(), PositionTolerance);
+
+        if (CX.SizeSquared() > PositionTolerance * PositionTolerance)
+        {
+            if (bRealTypeCompatibleWithISPC && bChaos_Joint_ISPC_Enabled)
+            {
+                check(false);
+#if 0
+                ispc::ApplyPointPositionConstraintKD((ispc::FPBDJointCachedSolver*)this, KIndex, DIndex, (ispc::FVector&)CX, Stiffness);
+#endif
+            }
+            else
+            {
+                // Calculate constraint correction
+                const FVec3 DX = Stiffness * Utilities::Multiply(PointPositionConstraintIM, CX);
+
+                // Apply constraint correction
+                const FVec3 ArmD = ConstraintArms[PointPositionConstraintIndex][DIndex];
+                const FVec3 DP1 = -InvM(DIndex) * DX;
+                const FVec3 DR1 = Utilities::Multiply(InvI(DIndex), FVec3::CrossProduct(ArmD, -DX));
+
+                ApplyDelta(DIndex, DP1, DR1);
+
+                NetLinearImpulse += (KIndex == 0) ? DX : -DX;
+            }
+
+            ++NumActiveConstraints;
+        }
+    }
+
+
+    // Dynamic-Dynamic bodies
+    void FPBDJointCachedSolver::ApplyPointPositionConstraintDD(
+        const FReal Dt,
+        const FPBDJointSolverSettings& SolverSettings,
+        const FPBDJointSettings& JointSettings)
+    {
+        check(Body(0).IsDynamic());
+        check(Body(1).IsDynamic());
+
+        FReal Stiffness = SolverStiffness * FPBDJointUtilities::GetLinearStiffness(SolverSettings, JointSettings);
+        const FVec3 CX = ConnectorXs[1] - ConnectorXs[0];
+
+        UE_LOG(LogChaosJoint, VeryVerbose, TEXT("    PointDD Delta %f [Limit %f]"), CX.Size(), PositionTolerance);
+
+        if (CX.SizeSquared() > PositionTolerance * PositionTolerance)
+        {
+            if (bRealTypeCompatibleWithISPC && bChaos_Joint_ISPC_Enabled)
+            {
+                check(false);
+#if 0
+                ispc::ApplyPointPositionConstraintDD((ispc::FPBDJointCachedSolver*)this, (ispc::FVector&)CX, Stiffness);
+#endif
+            }
+            else
+            {
+                // Calculate constraint correction
+                const FVec3 DX = Stiffness * Utilities::Multiply(PointPositionConstraintIM, CX);
+
+                // Apply constraint correction
+                const FVec3 Arm0 = ConstraintArms[PointPositionConstraintIndex][0];
+                const FVec3 Arm1 = ConstraintArms[PointPositionConstraintIndex][1];
+                const FVec3 DP0 = InvM(0) * DX;
+                const FVec3 DP1 = -InvM(1) * DX;
+                const FVec3 DR0 = Utilities::Multiply(InvI(0), FVec3::CrossProduct(Arm0, DX));
+                const FVec3 DR1 = Utilities::Multiply(InvI(1), FVec3::CrossProduct(Arm1, -DX));
+
+                ApplyPositionDelta(DP0, DP1);
+                ApplyRotationDelta(DR0, DR1);
+
+                NetLinearImpulse += DX;
+            }
+            
+            ++NumActiveConstraints;
+        }
+    }
+
+
+    void FPBDJointCachedSolver::ApplySphericalPositionConstraint(
+        const FReal Dt,
+        const FPBDJointSolverSettings& SolverSettings,
+        const FPBDJointSettings& JointSettings)
+    {
+        const FVec3 Axis = PositionConstraintAxis[SphericalPositionConstraintIndex];
+        const FReal Delta = FVec3::DotProduct(ConnectorXs[1] - ConnectorXs[0], Axis);
+
+
+        const FReal LimitPadding = GetLinearConstraintPadding(0);
+        const FReal Limit = FMath::Max(JointSettings.LinearLimit - LimitPadding, (FReal)0.);
+
+        FReal Error = Delta - Limit;
+        if (Error > PositionTolerance)
+        {
+            if (!FPBDJointUtilities::GetSoftLinearLimitEnabled(SolverSettings, JointSettings))
+            {
+                if (JointSettings.LinearRestitution > 0.0f)
+                {
+                    CalculateLinearConstraintPadding(Dt, SolverSettings, JointSettings, JointSettings.LinearRestitution, 0, Axis, Error);
+                }
+
+                const FReal& IM = PositionConstraintIM[SphericalPositionConstraintIndex];
+                const FVec3& Arm0 = ConstraintArms[SphericalPositionConstraintIndex][0];
+                const FVec3& Arm1 = ConstraintArms[SphericalPositionConstraintIndex][1];
+                const FVec3& DRAxis0 = PositionConstraintDRAxis[SphericalPositionConstraintIndex][0];
+                const FVec3& DRAxis1 = PositionConstraintDRAxis[SphericalPositionConstraintIndex][1];
+                ApplyPositionConstraint(LinearStiffness, Axis, Error, IM, Arm0, Arm1, DRAxis0, DRAxis1);
+            }
+            else
+            {
+                const FReal JointStiffness = FPBDJointUtilities::GetSoftLinearStiffness(SolverSettings, JointSettings);
+                const FReal JointDamping = FPBDJointUtilities::GetSoftLinearDamping(SolverSettings, JointSettings);
+                const bool bAccelerationMode = FPBDJointUtilities::GetLinearSoftAccelerationMode(SolverSettings, JointSettings);
+                ApplyPositionConstraintSoft(Dt, JointStiffness, JointDamping, bAccelerationMode, Axis, Error, 0.0f, LinearSoftLambda);
+            }
+        }
+    }
+
+
+    void FPBDJointCachedSolver::ApplyCylindricalPositionConstraint(
+        const FReal Dt,
+        const int32 AxisIndex,
+        const EJointMotionType AxialMotion,
+        const EJointMotionType RadialMotion,
+        const FPBDJointSolverSettings& SolverSettings,
+        const FPBDJointSettings& JointSettings)
+    {
+        check(AxialMotion != RadialMotion);
+
+        FVec3 Axis = PositionConstraintAxis[(int32)EJointCylindricalPositionConstraintType::Axial];
+        FReal AxialDelta = FVec3::DotProduct(ConnectorXs[1] - ConnectorXs[0], Axis);
+        if (AxialDelta < 0.f)
+        {
+            AxialDelta = -AxialDelta;
+            Axis = -Axis;
+        }
+        const FReal AxialLimitPadding = GetLinearConstraintPadding(0);
+        const FReal AxialLimit = (AxialMotion == EJointMotionType::Locked) ? 0.0f : FMath::Max(JointSettings.LinearLimit - AxialLimitPadding, (FReal)0.);
+        FReal AxialError = AxialDelta - AxialLimit;
+
+        if (AxialError > PositionTolerance)
+        {
+            if ((AxialMotion == EJointMotionType::Limited) && FPBDJointUtilities::GetSoftLinearLimitEnabled(SolverSettings, JointSettings))
+            {
+                // Soft Axial constraint
+                const FReal JointStiffness = FPBDJointUtilities::GetSoftLinearStiffness(SolverSettings, JointSettings);
+                const FReal JointDamping = FPBDJointUtilities::GetSoftLinearDamping(SolverSettings, JointSettings);
+                const bool bAccelerationMode = FPBDJointUtilities::GetLinearSoftAccelerationMode(SolverSettings, JointSettings);
+                ApplyPositionConstraintSoft(Dt, JointStiffness, JointDamping, bAccelerationMode, Axis, AxialError, 0.0f, LinearSoftLambda);
+            }
+            else if (AxialMotion != EJointMotionType::Free)
+            {
+                if (JointSettings.LinearRestitution > 0.0f)
+                {
+                    CalculateLinearConstraintPadding(Dt, SolverSettings, JointSettings, JointSettings.LinearRestitution, 0, Axis, AxialError);
+                }
+                const FReal& IM = PositionConstraintIM[(int32)EJointCylindricalPositionConstraintType::Axial];
+                const FVec3& Arm0 = ConstraintArms[(int32)EJointCylindricalPositionConstraintType::Axial][0];
+                const FVec3& Arm1 = ConstraintArms[(int32)EJointCylindricalPositionConstraintType::Axial][1];
+                const FVec3& DRAxis0 = PositionConstraintDRAxis[(int32)EJointCylindricalPositionConstraintType::Axial][0];
+                const FVec3& DRAxis1 = PositionConstraintDRAxis[(int32)EJointCylindricalPositionConstraintType::Axial][1];
+                ApplyPositionConstraint(LinearStiffness, Axis, AxialError, IM, Arm0, Arm1, DRAxis0, DRAxis1, (int32)EJointCylindricalPositionConstraintType::Axial);
+            }
+        }
+
+        const FVec3& RadialAxis = PositionConstraintAxis[(int32)EJointCylindricalPositionConstraintType::Radial];
+        const FReal RadialDelta = FVec3::DotProduct(ConnectorXs[1] - ConnectorXs[0], RadialAxis);
+        const FReal RadialLimitPadding = GetLinearConstraintPadding(1);
+        const FReal RadialLimit = (RadialMotion == EJointMotionType::Locked) ? 0.0f : FMath::Max(JointSettings.LinearLimit - AxialLimitPadding, (FReal)0.);
+        FReal RadialError = RadialDelta - RadialLimit;
+        if (RadialError > PositionTolerance)
+        {
+            if ((RadialMotion == EJointMotionType::Limited) && FPBDJointUtilities::GetSoftLinearLimitEnabled(SolverSettings, JointSettings))
+            {
+                // Soft Radial constraint
+                const FReal JointStiffness = FPBDJointUtilities::GetSoftLinearStiffness(SolverSettings, JointSettings);
+                const FReal JointDamping = FPBDJointUtilities::GetSoftLinearDamping(SolverSettings, JointSettings);
+                const bool bAccelerationMode = FPBDJointUtilities::GetLinearSoftAccelerationMode(SolverSettings, JointSettings);
+                ApplyPositionConstraintSoft(Dt, JointStiffness, JointDamping, bAccelerationMode, RadialAxis, RadialError, 0.0f, LinearSoftLambda);
+            }
+            else if (RadialMotion != EJointMotionType::Free)
+            {
+                // Hard Radial constraint
+                if (JointSettings.LinearRestitution > 0.0f)
+                {
+                    CalculateLinearConstraintPadding(Dt, SolverSettings, JointSettings, JointSettings.LinearRestitution, 1, RadialAxis, RadialError);
+                }
+
+                const FReal& IM = PositionConstraintIM[(int32)EJointCylindricalPositionConstraintType::Radial];
+                const FVec3& Arm0 = ConstraintArms[(int32)EJointCylindricalPositionConstraintType::Radial][0];
+                const FVec3& Arm1 = ConstraintArms[(int32)EJointCylindricalPositionConstraintType::Radial][1];
+                const FVec3& DRAxis0 = PositionConstraintDRAxis[(int32)EJointCylindricalPositionConstraintType::Radial][0];
+                const FVec3& DRAxis1 = PositionConstraintDRAxis[(int32)EJointCylindricalPositionConstraintType::Radial][1];
+                ApplyPositionConstraint(LinearStiffness, RadialAxis, RadialError, IM, Arm0, Arm1, DRAxis0, DRAxis1, (int32)EJointCylindricalPositionConstraintType::Radial);
+            }
+        }
+    }
+
+    void FPBDJointCachedSolver::ApplyPlanarPositionConstraint(
+        const FReal Dt,
+        const int32 AxisIndex,
+        const EJointMotionType AxialMotion,
+        const FPBDJointSolverSettings& SolverSettings,
+        const FPBDJointSettings& JointSettings)
+    {
+        FVec3 Axis = PositionConstraintAxis[PlanarPositionConstraintIndex];
+        FReal Delta = FVec3::DotProduct(ConnectorXs[1] - ConnectorXs[0], Axis);
+        if (Delta < 0.0f)
+        {
+            Delta = -Delta;
+            Axis = -Axis;
+        }
+
+        const FReal LimitPadding = GetLinearConstraintPadding(0);
+        const FReal Limit = (AxialMotion == EJointMotionType::Locked) ? 0 : FMath::Max(JointSettings.LinearLimit - LimitPadding, (FReal)0.);
+        FReal Error = Delta - Limit;
+        if (Error > PositionTolerance)
+        {
+            if ((AxialMotion == EJointMotionType::Limited) && FPBDJointUtilities::GetSoftLinearLimitEnabled(SolverSettings, JointSettings))
+            {
+                const FReal JointStiffness = FPBDJointUtilities::GetSoftLinearStiffness(SolverSettings, JointSettings);
+                const FReal JointDamping = FPBDJointUtilities::GetSoftLinearDamping(SolverSettings, JointSettings);
+                const bool bAccelerationMode = FPBDJointUtilities::GetLinearSoftAccelerationMode(SolverSettings, JointSettings);
+                ApplyPositionConstraintSoft(Dt, JointStiffness, JointDamping, bAccelerationMode, Axis, Error, 0.0f, LinearSoftLambda);
+            }
+            else
+            {
+                if (JointSettings.LinearRestitution > 0.0f)
+                {
+                    CalculateLinearConstraintPadding(Dt, SolverSettings, JointSettings, JointSettings.LinearRestitution, 0, Axis, Error);
+                }
+
+                const FReal& IM = PositionConstraintIM[PlanarPositionConstraintIndex];
+                const FVec3& Arm0 = ConstraintArms[PlanarPositionConstraintIndex][0];
+                const FVec3& Arm1 = ConstraintArms[PlanarPositionConstraintIndex][1];
+                const FVec3& DRAxis0 = PositionConstraintDRAxis[PlanarPositionConstraintIndex][0];
+                const FVec3& DRAxis1 = PositionConstraintDRAxis[PlanarPositionConstraintIndex][1];
+                ApplyPositionConstraint(LinearStiffness, Axis, Error, IM, Arm0, Arm1, DRAxis0, DRAxis1, PlanarPositionConstraintIndex);
+            }
+        }
+    }
+
+
+    void FPBDJointCachedSolver::ApplyPositionDrive(
+        const FReal Dt,
+        const int32 AxisIndex,
+        const FPBDJointSolverSettings& SolverSettings,
+        const FPBDJointSettings& JointSettings,
+        const FVec3& Axis,
+        const FReal DeltaPos,
+        const FReal DeltaVel)
+    {
+        const FReal JointStiffness = FPBDJointUtilities::GetLinearDriveStiffness(SolverSettings, JointSettings, AxisIndex);
+        const FReal JointDamping = FPBDJointUtilities::GetLinearDriveDamping(SolverSettings, JointSettings, AxisIndex);
+        const bool bAccelerationMode = FPBDJointUtilities::GetDriveAccelerationMode(SolverSettings, JointSettings);
+
+        if ((FMath::Abs(DeltaPos) > PositionTolerance) || (JointDamping > 0.0f))
+        {
+            FReal ReturnedLambda = LinearDriveLambdas[AxisIndex];
+            ApplyPositionConstraintSoft(Dt, JointStiffness, JointDamping, bAccelerationMode, Axis, DeltaPos, DeltaVel, ReturnedLambda);
+            LinearDriveLambdas[AxisIndex] = ReturnedLambda;
+        }
+    }
+
+
+    void FPBDJointCachedSolver::ApplyPointProjection(
+        const FReal Dt,
+        const FPBDJointSolverSettings& SolverSettings,
+        const FPBDJointSettings& JointSettings,
+        const FReal Alpha,
+        FVec3& NetDP1,
+        FVec3& NetDR1)
+    {
+        const FReal ProjectionPositionTolerance = 0.0f;//PositionTolerance;
+
+        const FVec3 CX = ConnectorXs[1] - ConnectorXs[0];
+        if (CX.Size() > ProjectionPositionTolerance)
+        {
+            FMatrix33 J = Utilities::ComputeJointFactorMatrix(ConnectorXs[1] - P(1), InvI(1), InvM(1));
+            const FMatrix33 IJ = J.Inverse();
+            const FVec3 DX = Utilities::Multiply(IJ, CX);
+
+            const FVec3 DP1 = -Alpha * InvM(1) * DX;
+            const FVec3 DR1 = -Alpha * Utilities::Multiply(InvI(1), FVec3::CrossProduct(ConnectorXs[1] - P(1), DX));
+            ApplyDelta(1, DP1, DR1);
+
+            NetDP1 += DP1;
+            NetDR1 += DR1;
+        }
+    }
+
+    void FPBDJointCachedSolver::ApplySphereProjection(
+        const FReal Dt,
+        const FPBDJointSolverSettings& SolverSettings,
+        const FPBDJointSettings& JointSettings,
+        const FReal Alpha,
+        FVec3& NetDP1,
+        FVec3& NetDR1)
+    {
+        const FReal ProjectionPositionTolerance = 0.0f;//PositionTolerance;
+
+        FVec3 Axis;
+        FReal Delta;
+        FPBDJointUtilities::GetSphericalAxisDelta(ConnectorXs[0], ConnectorXs[1], Axis, Delta);
+        const FReal Error = FMath::Max((FReal)0, Delta - JointSettings.LinearLimit);
+        if (FMath::Abs(Error) > ProjectionPositionTolerance)
+        {
+            const FVec3 AngularAxis1 = FVec3::CrossProduct(ConnectorXs[1] - P(1), Axis);
+            const FVec3 IA1 = Utilities::Multiply(InvI(1), AngularAxis1);
+            const FReal II1 = FVec3::DotProduct(AngularAxis1, IA1);
+            const FReal IM = InvM(1) + II1;
+            const FVec3 DX = Axis * Error / IM;
+
+            const FVec3 DP1 = -Alpha * InvM(1) * DX;
+            const FVec3 DR1 = -Alpha * Utilities::Multiply(InvI(1), FVec3::CrossProduct(ConnectorXs[1] - P(1), DX));
+            ApplyDelta(1, DP1, DR1);
+
+            NetDP1 += DP1;
+            NetDR1 += DR1;
+        }
+    }
+
+    void FPBDJointCachedSolver::ApplyTranslateProjection(
+        const FReal Dt,
+        const FPBDJointSolverSettings& SolverSettings,
+        const FPBDJointSettings& JointSettings,
+        const FReal Alpha,
+        FVec3& NetDP1,
+        FVec3& NetDR1)
+    {
+        const FReal ProjectionPositionTolerance = 0.0f;//PositionTolerance;
+
+        FVec3 Axis;
+        FReal Delta;
+        FPBDJointUtilities::GetSphericalAxisDelta(ConnectorXs[0], ConnectorXs[1], Axis, Delta);
+        const FReal Error = FMath::Max((FReal)0, Delta - JointSettings.LinearLimit);
+        if (Error > ProjectionPositionTolerance)
+        {
+            const FVec3 DP1 = -Alpha * Error * Axis;
+            ApplyPositionDelta(1, DP1);
+            
+            NetDP1 += DP1;
+        }
+    }
+
+    void FPBDJointCachedSolver::ApplyConeProjection(
+        const FReal Dt,
+        const FPBDJointSolverSettings& SolverSettings,
+        const FPBDJointSettings& JointSettings,
+        const FReal Alpha,
+        const bool bPositionLocked,
+        FVec3& NetDP1,
+        FVec3& NetDR1)
+    {
+        const FReal ProjectionAngleTolerance = 0.0f;//AngleTolerance;
+
+        FVec3 SwingAxisLocal;
+        FReal DSwingAngle = 0.0f;
+        const FReal Swing1Limit = JointSettings.AngularLimits[(int32)EJointAngularConstraintIndex::Swing1];
+        const FReal Swing2Limit = JointSettings.AngularLimits[(int32)EJointAngularConstraintIndex::Swing2];
+        FPBDJointUtilities::GetEllipticalConeAxisErrorLocal(ConnectorRs[0], ConnectorRs[1], Swing2Limit, Swing1Limit, SwingAxisLocal, DSwingAngle);
+        FVec3 SwingAxis = ConnectorRs[0] * SwingAxisLocal;
+        if (DSwingAngle > ProjectionAngleTolerance)
+        {
+            const FVec3 DR1 = -Alpha * DSwingAngle * SwingAxis;
+            FVec3 DP1 = FVec3(0);
+            if (bPositionLocked)
+            {
+                DP1 = -Alpha * FVec3::CrossProduct(DR1, ConnectorXs[1] - P(1));
+            }
+            ApplyDelta(1, DP1, DR1);
+
+            NetDP1 += DP1;
+            NetDR1 += DR1;
+        }
+    }
+
+    void FPBDJointCachedSolver::ApplySwingProjection(
+        const FReal Dt,
+        const FPBDJointSolverSettings& SolverSettings,
+        const FPBDJointSettings& JointSettings,
+        const EJointAngularConstraintIndex SwingConstraintIndex,
+        const FReal Alpha,
+        const bool bPositionLocked,
+        FVec3& NetDP1,
+        FVec3& NetDR1)
+    {
+        const FReal ProjectionAngleTolerance = 0.0f;//AngleTolerance;
+
+        FVec3 SwingAxis;
+        FReal SwingAngle;
+        FPBDJointUtilities::GetSwingAxisAngle(ConnectorRs[0], ConnectorRs[1], SolverSettings.SwingTwistAngleTolerance, SwingConstraintIndex, SwingAxis, SwingAngle);
+
+        // Calculate swing error we need to correct
+        FReal DSwingAngle = 0;
+        const FReal SwingAngleMax = JointSettings.AngularLimits[(int32)SwingConstraintIndex];
+        if (SwingAngle > SwingAngleMax)
+        {
+            DSwingAngle = SwingAngle - SwingAngleMax;
+        }
+        else if (SwingAngle < -SwingAngleMax)
+        {
+            DSwingAngle = SwingAngle + SwingAngleMax;
+        }
+
+        UE_LOG(LogChaosJoint, VeryVerbose, TEXT("    Swing%d Angle %f [Limit %f]"), (SwingConstraintIndex == EJointAngularConstraintIndex::Swing1) ? 1 : 2, FMath::RadiansToDegrees(SwingAngle), FMath::RadiansToDegrees(SwingAngleMax));
+
+        // Apply swing correction
+        if (FMath::Abs(DSwingAngle) > ProjectionAngleTolerance)
+        {
+            const FVec3 DR1 = -Alpha * DSwingAngle * SwingAxis;
+            FVec3 DP1 = FVec3(0);
+            if (bPositionLocked)
+            {
+                DP1 = -Alpha * FVec3::CrossProduct(DR1, ConnectorXs[1] - P(1));
+            }
+            ApplyDelta(1, DP1, DR1);
+
+            NetDP1 += DP1;
+            NetDR1 += DR1;
+        }
+    }
+
+    void FPBDJointCachedSolver::ApplySingleLockedSwingProjection(
+        const FReal Dt,
+        const FPBDJointSolverSettings& SolverSettings,
+        const FPBDJointSettings& JointSettings,
+        const EJointAngularConstraintIndex SwingConstraintIndex,
+        const FReal Alpha,
+        const bool bPositionLocked,
+        FVec3& NetDP1,
+        FVec3& NetDR1)
+    {
+        ApplySwingProjection(Dt, SolverSettings, JointSettings, SwingConstraintIndex, Alpha, bPositionLocked, NetDP1, NetDR1);
+    }
+
+    void FPBDJointCachedSolver::ApplyDoubleLockedSwingProjection(
+        const FReal Dt,
+        const FPBDJointSolverSettings& SolverSettings,
+        const FPBDJointSettings& JointSettings,
+        const FReal Alpha,
+        const bool bPositionLocked,
+        FVec3& NetDP1,
+        FVec3& NetDR1)
+    {
+        const FReal ProjectionAngleTolerance = 0.0f;//AngleTolerance;
+
+        FVec3 SwingAxisLocal;
+        FReal DSwingAngle = 0.0f;
+        FPBDJointUtilities::GetCircularConeAxisErrorLocal(ConnectorRs[0], ConnectorRs[1], 0.0f, SwingAxisLocal, DSwingAngle);
+        FVec3 SwingAxis = ConnectorRs[0] * SwingAxisLocal;
+        if (DSwingAngle > ProjectionAngleTolerance)
+        {
+            const FVec3 DR1 = -Alpha * DSwingAngle * SwingAxis;
+            FVec3 DP1 = FVec3(0);
+            if (bPositionLocked)
+            {
+                DP1 = -Alpha * FVec3::CrossProduct(DR1, ConnectorXs[1] - P(1));
+            }
+            ApplyDelta(1, DP1, DR1);
+
+            NetDP1 += DP1;
+            NetDR1 += DR1;
+        }
+    }
+
+    void FPBDJointCachedSolver::ApplyDualConeSwingProjection(
+        const FReal Dt,
+        const FPBDJointSolverSettings& SolverSettings,
+        const FPBDJointSettings& JointSettings,
+        const EJointAngularConstraintIndex SwingConstraintIndex,
+        const FReal Alpha,
+        const bool bPositionLocked,
+        FVec3& NetDP1,
+        FVec3& NetDR1)
+    {
+        const FReal ProjectionAngleTolerance = 0.0f;//AngleTolerance;
+
+        FVec3 SwingAxis;
+        FReal SwingAngle;
+        FPBDJointUtilities::GetDualConeSwingAxisAngle(ConnectorRs[0], ConnectorRs[1], SwingConstraintIndex, SwingAxis, SwingAngle);
+
+        // Calculate swing error we need to correct
+        FReal DSwingAngle = 0;
+        const FReal SwingAngleMax = JointSettings.AngularLimits[(int32)SwingConstraintIndex];
+        if (SwingAngle > SwingAngleMax)
+        {
+            DSwingAngle = SwingAngle - SwingAngleMax;
+        }
+        else if (SwingAngle < -SwingAngleMax)
+        {
+            DSwingAngle = SwingAngle + SwingAngleMax;
+        }
+
+        // Apply swing correction
+        if (FMath::Abs(DSwingAngle) > ProjectionAngleTolerance)
+        {
+            const FVec3 DR1 = -Alpha * DSwingAngle * SwingAxis;
+            FVec3 DP1 = FVec3(0);
+            if (bPositionLocked)
+            {
+                DP1 = -Alpha * FVec3::CrossProduct(DR1, ConnectorXs[1] - P(1));
+            }
+            ApplyDelta(1, DP1, DR1);
+
+            NetDP1 += DP1;
+            NetDR1 += DR1;
+        }
+    }
+
+    void FPBDJointCachedSolver::ApplyTwistProjection(
+        const FReal Dt,
+        const FPBDJointSolverSettings& SolverSettings,
+        const FPBDJointSettings& JointSettings,
+        const FReal Alpha,
+        const bool bPositionLocked,
+        FVec3& NetDP1,
+        FVec3& NetDR1)
+    {
+        const FReal ProjectionAngleTolerance = 0.0f;//AngleTolerance;
+
+        FVec3 TwistAxis;
+        FReal TwistAngle;
+        FPBDJointUtilities::GetTwistAxisAngle(ConnectorRs[0], ConnectorRs[1], TwistAxis, TwistAngle);
+        FReal DTwistAngle = 0;
+        const FReal TwistLimit = JointSettings.AngularLimits[(int32)EJointAngularConstraintIndex::Twist];
+        if (TwistAngle > TwistLimit)
+        {
+            DTwistAngle = TwistAngle - TwistLimit;
+        }
+        else if (TwistAngle < -TwistLimit)
+        {
+            DTwistAngle = TwistAngle + TwistLimit;
+        }
+
+        if (FMath::Abs(DTwistAngle) > ProjectionAngleTolerance)
+        {
+            const FVec3 DR1 = -Alpha * DTwistAngle * TwistAxis;
+            FVec3 DP1 = FVec3(0);
+            if (bPositionLocked)
+            {
+                DP1 = -Alpha * FVec3::CrossProduct(DR1, ConnectorXs[1] - P(1));
+            }
+            ApplyDelta(1, DP1, DR1);
+
+            NetDP1 += DP1;
+            NetDR1 += DR1;
+        }
+    }
+
+    void FPBDJointCachedSolver::ApplyVelocityProjection(
+        const FReal Dt,
+        const FPBDJointSolverSettings& SolverSettings,
+        const FPBDJointSettings& JointSettings,
+        const FReal Alpha,
+        const FVec3& DP1,
+        const FVec3& DR1)
+    {
+        const FVec3 DV1 = Alpha * DP1 / Dt;
+        const FVec3 DW1 = Alpha * DR1 / Dt;
+        ApplyVelocityDelta(1, DV1, DW1);
+    }
+
+    void FPBDJointCachedSolver::ApplyLinearVelocityConstraint(
+        const FReal Stiffness,
+        const FVec3& Axis,
+        const FReal IM,
+        const FVec3& Arm0,
+        const FVec3& Arm1,
+        const FVec3& DRAxis0,
+        const FVec3& DRAxis1,
+        const FReal TargetVel)
+    {
+        const FVec3 CV0 = V(0) + FVec3::CrossProduct(W(0), Arm0);
+        const FVec3 CV1 = V(1) + FVec3::CrossProduct(W(1), Arm1);
+        const FVec3 CV = CV1 - CV0;
+        const FReal ImpulseSize = Stiffness * (FVec3::DotProduct(CV, Axis) - TargetVel) / IM;
+        const FVec3 Impulse = ImpulseSize * Axis;
+
+        const FVec3 DV0 = InvM(0) * Impulse;
+        const FVec3 DV1 = -InvM(1) * Impulse;
+        const FVec3 DW0 = DRAxis0 * ImpulseSize;
+        const FVec3 DW1 = DRAxis1 * ImpulseSize;
+        ApplyVelocityDelta(DV0, DW0, DV1, DW1);
+    }
+
+    void FPBDJointCachedSolver::ApplyPointVelocityConstraint(
+        const FReal Stiffness,
+        const FMatrix33& ConstraintIM,
+        const FVec3& ConstraintArm0,
+        const FVec3& ConstraintArm1)
+    {
+        const FVec3 CV0 = V(0) + FVec3::CrossProduct(W(0), ConstraintArm0);
+        const FVec3 CV1 = V(1) + FVec3::CrossProduct(W(1), ConstraintArm1);
+        const FVec3 CV = CV1 - CV0;
+        const FVec3 Impulse = Stiffness * Utilities::Multiply(ConstraintIM, CV);
+
+        const FVec3 DV0 = InvM(0) * Impulse;
+        const FVec3 DV1 = -InvM(1) * Impulse;
+        const FVec3 DW0 = Utilities::Multiply(InvI(0), FVec3::CrossProduct(ConstraintArm0, Impulse));
+        const FVec3 DW1 = Utilities::Multiply(InvI(1), FVec3::CrossProduct(ConstraintArm1, -Impulse));
+
+        ApplyVelocityDelta(DV0, DW0, DV1, DW1);
+    }
+
+    void FPBDJointCachedSolver::ApplyPointVelocityConstraint(
+        const FReal Dt,
+        const FPBDJointSolverSettings& SolverSettings,
+        const FPBDJointSettings& JointSettings)
+    {
+        if (!NetLinearImpulse.IsNearlyZero() && !FPBDJointUtilities::GetSoftLinearLimitEnabled(SolverSettings, JointSettings))
+        {
+            const FVec3& Arm0 = ConstraintArms[PointPositionConstraintIndex][0];
+            const FVec3& Arm1 = ConstraintArms[PointPositionConstraintIndex][1];
+			ApplyPointVelocityConstraint(LinearStiffness, PointPositionConstraintIM, Arm0, Arm1);
+        }
+    }
+
+    void FPBDJointCachedSolver::ApplySphericalVelocityConstraint(
+        const FReal Dt,
+        const FPBDJointSolverSettings& SolverSettings,
+        const FPBDJointSettings& JointSettings)
+    {
+        if(!NetLinearImpulse.IsNearlyZero() && !FPBDJointUtilities::GetSoftLinearLimitEnabled(SolverSettings, JointSettings))
+        {
+            const FVec3 Axis = PositionConstraintAxis[SphericalPositionConstraintIndex];
+            FReal TargetVel = 0.0f;
+            if (JointSettings.LinearRestitution != 0.0f)
+            {
+                const FReal InitVel = InitConstraintAxisLinearVelocities[0];
+                TargetVel = InitVel > Chaos_Joint_LinearVelocityThresholdToApplyRestitution ? -JointSettings.LinearRestitution * InitVel : 0.0f; 
+            }
+            const FReal& IM = PositionConstraintIM[SphericalPositionConstraintIndex];
+            const FVec3& Arm0 = ConstraintArms[SphericalPositionConstraintIndex][0];
+            const FVec3& Arm1 = ConstraintArms[SphericalPositionConstraintIndex][1];
+            const FVec3& DRAxis0 = PositionConstraintDRAxis[SphericalPositionConstraintIndex][0];
+            const FVec3& DRAxis1 = PositionConstraintDRAxis[SphericalPositionConstraintIndex][1];
+            ApplyLinearVelocityConstraint(LinearStiffness, Axis, IM, Arm0, Arm1, DRAxis0, DRAxis1, TargetVel);
+        }
+    }
+
+    void FPBDJointCachedSolver::ApplyCylindricalVelocityConstraint(
+        const FReal Dt,
+        const int32 AxisIndex,
+        const EJointMotionType AxialMotion,
+        const EJointMotionType RadialMotion,
+        const FPBDJointSolverSettings& SolverSettings,
+        const FPBDJointSettings& JointSettings)
+    {
+        if(!NetLinearImpulse.IsNearlyZero() && !FPBDJointUtilities::GetSoftLinearLimitEnabled(SolverSettings, JointSettings))
+        {
+            if (LinearHardLambda[(int32)EJointCylindricalPositionConstraintType::Axial] > SMALL_NUMBER)
+            {
+                const FVec3 Axis = PositionConstraintAxis[(int32)EJointCylindricalPositionConstraintType::Axial];
+                FReal TargetVel = 0.0f;
+                if (AxialMotion == EJointMotionType::Limited && JointSettings.LinearRestitution != 0.0f)
+                {
+                    const FReal InitVel = InitConstraintAxisLinearVelocities[(int32)EJointCylindricalPositionConstraintType::Axial];
+                    TargetVel = InitVel > Chaos_Joint_LinearVelocityThresholdToApplyRestitution ? -JointSettings.LinearRestitution * InitVel : 0.0f; 
+                }
+                const FReal& IM = PositionConstraintIM[(int32)EJointCylindricalPositionConstraintType::Axial];
+                const FVec3& Arm0 = ConstraintArms[(int32)EJointCylindricalPositionConstraintType::Axial][0];
+                const FVec3& Arm1 = ConstraintArms[(int32)EJointCylindricalPositionConstraintType::Axial][1];
+                const FVec3& DRAxis0 = PositionConstraintDRAxis[(int32)EJointCylindricalPositionConstraintType::Axial][0];
+                const FVec3& DRAxis1 = PositionConstraintDRAxis[(int32)EJointCylindricalPositionConstraintType::Axial][1];
+                ApplyLinearVelocityConstraint(LinearStiffness, Axis, IM, Arm0, Arm1, DRAxis0, DRAxis1, TargetVel);
+            }
+
+            if (LinearHardLambda[(int32)EJointCylindricalPositionConstraintType::Radial] > SMALL_NUMBER)
+            {
+                const FVec3 RadialAxis = PositionConstraintAxis[(int32)EJointCylindricalPositionConstraintType::Radial];
+                FReal TargetVel = 0.0f;
+                if (RadialMotion == EJointMotionType::Limited && JointSettings.LinearRestitution != 0.0f)
+                {
+                    const FReal InitVel = InitConstraintAxisLinearVelocities[(int32)EJointCylindricalPositionConstraintType::Radial];
+                    TargetVel = InitVel > Chaos_Joint_LinearVelocityThresholdToApplyRestitution ? -JointSettings.LinearRestitution * InitVel : 0.0f; 
+                }
+                const FReal& IM = PositionConstraintIM[(int32)EJointCylindricalPositionConstraintType::Radial];
+                const FVec3& Arm0 = ConstraintArms[(int32)EJointCylindricalPositionConstraintType::Radial][0];
+                const FVec3& Arm1 = ConstraintArms[(int32)EJointCylindricalPositionConstraintType::Radial][1];
+                const FVec3& DRAxis0 = PositionConstraintDRAxis[(int32)EJointCylindricalPositionConstraintType::Radial][0];
+                const FVec3& DRAxis1 = PositionConstraintDRAxis[(int32)EJointCylindricalPositionConstraintType::Radial][1];
+                ApplyLinearVelocityConstraint(LinearStiffness, RadialAxis, IM, Arm0, Arm1, DRAxis0, DRAxis1, TargetVel);
+            }
+        }
+    }
+
+    void FPBDJointCachedSolver::ApplyPlanarVelocityConstraint(
+        const FReal Dt,
+        const int32 AxisIndex,
+        const EJointMotionType AxialMotion,
+        const FPBDJointSolverSettings& SolverSettings,
+        const FPBDJointSettings& JointSettings)
+    {
+        if(!NetLinearImpulse.IsNearlyZero() && !FPBDJointUtilities::GetSoftLinearLimitEnabled(SolverSettings, JointSettings))
+        {
+			FVec3 Axis = PositionConstraintAxis[PlanarPositionConstraintIndex];
+            FReal TargetVel = 0.0f;
+            if (AxialMotion == EJointMotionType::Limited && JointSettings.LinearRestitution != 0.0f)
+            {
+                const FReal InitVel = InitConstraintAxisLinearVelocities[0];
+                TargetVel = InitVel > Chaos_Joint_LinearVelocityThresholdToApplyRestitution ? -JointSettings.LinearRestitution * InitVel : 0.0f; 
+            }
+            const FReal& IM = PositionConstraintIM[PlanarPositionConstraintIndex];
+            const FVec3& Arm0 = ConstraintArms[PlanarPositionConstraintIndex][0];
+            const FVec3& Arm1 = ConstraintArms[PlanarPositionConstraintIndex][1];
+            const FVec3& DRAxis0 = PositionConstraintDRAxis[PlanarPositionConstraintIndex][0];
+            const FVec3& DRAxis1 = PositionConstraintDRAxis[PlanarPositionConstraintIndex][1];
+            ApplyLinearVelocityConstraint(LinearStiffness, Axis, IM, Arm0, Arm1, DRAxis0, DRAxis1, TargetVel);
+        }
+    }
+
+    void FPBDJointCachedSolver::CalculateConstraintAxisLinearVelocities(
+        const FPBDJointSettings& JointSettings,
+        FVec3& ConstraintAxisLinearVelocities) const
+    {
+        bool bHasPositionConstraints =
+            (JointSettings.LinearMotionTypes[0] != EJointMotionType::Free)
+            || (JointSettings.LinearMotionTypes[1] != EJointMotionType::Free)
+            || (JointSettings.LinearMotionTypes[2] != EJointMotionType::Free);
+        if (!bHasPositionConstraints)
+        {
+            return;
+        }
+
+        const TVector<EJointMotionType, 3>& LinearMotion = JointSettings.LinearMotionTypes;
+        const TVector<bool, 3> bLinearLocked =
+        {
+            (LinearMotion[0] == EJointMotionType::Locked),
+            (LinearMotion[1] == EJointMotionType::Locked),
+            (LinearMotion[2] == EJointMotionType::Locked),
+        };
+        const TVector<bool, 3> bLinearLimited =
+        {
+            (LinearMotion[0] == EJointMotionType::Limited),
+            (LinearMotion[1] == EJointMotionType::Limited),
+            (LinearMotion[2] == EJointMotionType::Limited),
+        };
+
+        if (bLinearLocked[0] && bLinearLocked[1] && bLinearLocked[2])
+        {
+            
+        }
+        else if (bLinearLimited[0] && bLinearLimited[1] && bLinearLimited[2])
+        {
+            CalculateSphericalConstraintAxisLinearVelocities(JointSettings, ConstraintAxisLinearVelocities);
+        }
+        else if (bLinearLocked[1] && bLinearLocked[2] && !bLinearLocked[0])
+        {
+            // Line constraint along X axis
+            CalculateCylindricalConstraintAxisLinearVelocities(0, LinearMotion[0], EJointMotionType::Locked, JointSettings, ConstraintAxisLinearVelocities);
+        }
+        else if (bLinearLocked[0] && bLinearLocked[2] && !bLinearLocked[1])
+        {
+            // Line constraint along Y axis
+            CalculateCylindricalConstraintAxisLinearVelocities(1, LinearMotion[1], EJointMotionType::Locked, JointSettings, ConstraintAxisLinearVelocities);
+        }
+        else if (bLinearLocked[0] && bLinearLocked[1] && !bLinearLocked[2])
+        {
+            // Line constraint along Z axis
+            CalculateCylindricalConstraintAxisLinearVelocities(2, LinearMotion[2], EJointMotionType::Locked, JointSettings, ConstraintAxisLinearVelocities);
+        }
+        else if (bLinearLimited[1] && bLinearLimited[2] && !bLinearLimited[0])
+        {
+            // Cylindrical constraint along X axis
+            CalculateCylindricalConstraintAxisLinearVelocities(0, LinearMotion[0], EJointMotionType::Limited, JointSettings, ConstraintAxisLinearVelocities);
+        }
+        else if (bLinearLimited[0] && bLinearLimited[2] && !bLinearLimited[1])
+        {
+            // Cylindrical constraint along Y axis
+            CalculateCylindricalConstraintAxisLinearVelocities(1, LinearMotion[1], EJointMotionType::Limited, JointSettings, ConstraintAxisLinearVelocities);
+        }
+        else if (bLinearLimited[0] && bLinearLimited[1] && !bLinearLimited[2])
+        {
+            // Cylindrical constraint along Z axis
+            CalculateCylindricalConstraintAxisLinearVelocities(2, LinearMotion[2], EJointMotionType::Limited, JointSettings, ConstraintAxisLinearVelocities);
+        }
+        else if (bLinearLimited[0])
+        {
+            // Planar constraint along X axis
+            CalculatePlanarConstraintAxisLinearVelocities(0, LinearMotion[0], JointSettings, ConstraintAxisLinearVelocities);
+        }
+        else if (bLinearLimited[1])
+        {
+            // Planar constraint along Y axis
+            CalculatePlanarConstraintAxisLinearVelocities(1, LinearMotion[1], JointSettings, ConstraintAxisLinearVelocities);
+        }
+        else if (bLinearLimited[2])
+        {
+            // Planar constraint along Z axis
+            CalculatePlanarConstraintAxisLinearVelocities(2, LinearMotion[2], JointSettings, ConstraintAxisLinearVelocities);
+        }
+    }
+
+    void FPBDJointCachedSolver::CalculateSphericalConstraintAxisLinearVelocities(
+            const FPBDJointSettings& JointSettings,
+            FVec3& ConstraintAxisLinearVelocities) const
+    {
+        FVec3 Axis;
+        FReal Delta;
+        FPBDJointUtilities::GetSphericalAxisDelta(ConnectorXs[0], ConnectorXs[1], Axis, Delta);
+        const FVec3 DConnectorX = ConnectorXs[1] - ConnectorXs[0];
+        const FVec3 Connector0Correction = Delta > JointSettings.LinearLimit ? Axis * JointSettings.LinearLimit : DConnectorX;
+
+        const FVec3 Arm0 = ConnectorXs[0] + Connector0Correction - P(0);
+        const FVec3 Arm1 = ConnectorXs[1] - P(1);
+        const FVec3 CV0 = V(0) + FVec3::CrossProduct(W(0), Arm0);
+        const FVec3 CV1 = V(1) + FVec3::CrossProduct(W(1), Arm1);
+        const FVec3 CV = CV1 - CV0;
+
+        ConstraintAxisLinearVelocities[0] = FVec3::DotProduct(CV, Axis);
+    }
+
+    void FPBDJointCachedSolver::CalculateCylindricalConstraintAxisLinearVelocities(
+            const int32 AxisIndex,
+            const EJointMotionType AxialMotion,
+            const EJointMotionType RadialMotion,
+            const FPBDJointSettings& JointSettings,
+            FVec3& ConstraintAxisLinearVelocities) const
+    {
+        FVec3 Axis, RadialAxis;
+        FReal AxialDelta, RadialDelta;
+        FPBDJointUtilities::GetCylindricalAxesDeltas(ConnectorRs[0], ConnectorXs[0], ConnectorXs[1], AxisIndex, Axis, AxialDelta, RadialAxis, RadialDelta);
+        const FVec3 DConnectorX = ConnectorXs[1] - ConnectorXs[0];
+
+        if (AxialMotion == EJointMotionType::Limited)
+        {
+            if (AxialDelta < 0.0f)
+            {
+                AxialDelta = -AxialDelta;
+                Axis = -Axis;
+            }
+            const FReal AxialLimit = (AxialMotion == EJointMotionType::Locked) ? 0.0f :JointSettings.LinearLimit;
+            const FReal AxialError = AxialDelta - AxialLimit;
+            
+            const FVec3 Connector0Correction = DConnectorX - AxialError * Axis;
+
+            const FVec3 Arm0 = ConnectorXs[0] + Connector0Correction - P(0);
+            const FVec3 Arm1 = ConnectorXs[1] - P(1);
+            const FVec3 CV0 = V(0) + FVec3::CrossProduct(W(0), Arm0);
+            const FVec3 CV1 = V(1) + FVec3::CrossProduct(W(1), Arm1);
+            const FVec3 CV = CV1 - CV0;
+
+            ConstraintAxisLinearVelocities[(int32)EJointCylindricalPositionConstraintType::Axial] = FVec3::DotProduct(CV, Axis);
+        }
+
+        if (RadialMotion == EJointMotionType::Limited)
+        {
+            const FReal RadialLimit = (RadialMotion == EJointMotionType::Locked) ? 0.0f : JointSettings.LinearLimit;
+            const FReal RadialError = FMath::Max(RadialDelta - RadialLimit, 0.0f);
+            const FVec3 Connector0Correction = DConnectorX - RadialError * RadialAxis;
+            
+            const FVec3 Arm0 = ConnectorXs[0] + Connector0Correction - P(0);
+            const FVec3 Arm1 = ConnectorXs[1] - P(1);
+            const FVec3 CV0 = V(0) + FVec3::CrossProduct(W(0), Arm0);
+            const FVec3 CV1 = V(1) + FVec3::CrossProduct(W(1), Arm1);
+            const FVec3 CV = CV1 - CV0;
+
+            ConstraintAxisLinearVelocities[(int32)EJointCylindricalPositionConstraintType::Radial] = FVec3::DotProduct(CV, RadialAxis);
+        }
+    }
+
+    void FPBDJointCachedSolver::CalculatePlanarConstraintAxisLinearVelocities(
+            const int32 AxisIndex,
+            const EJointMotionType AxialMotion,
+            const FPBDJointSettings& JointSettings,
+            FVec3& ConstraintAxisLinearVelocities) const
+    {
+        FVec3 Axis;
+        FReal Delta;
+        FPBDJointUtilities::GetPlanarAxisDelta(ConnectorRs[0], ConnectorXs[0], ConnectorXs[1], AxisIndex, Axis, Delta);
+
+        if (Delta < 0.0f)
+        {
+            Delta = -Delta;
+            Axis = -Axis;
+        }
+
+        const FReal Limit = (AxialMotion == EJointMotionType::Locked) ? 0 : JointSettings.LinearLimit;
+        const FReal Error = Delta - Limit;
+
+        const FVec3 DConnectorX = ConnectorXs[1] - ConnectorXs[0];
+        const FVec3 Connector0Correction = DConnectorX - Axis * Error; 
+
+        const FVec3 Arm0 = ConnectorXs[0] + Connector0Correction - P(0);
+        const FVec3 Arm1 = ConnectorXs[1] - P(1);
+        const FVec3 CV0 = V(0) + FVec3::CrossProduct(W(0), Arm0);
+        const FVec3 CV1 = V(1) + FVec3::CrossProduct(W(1), Arm1);
+        const FVec3 CV = CV1 - CV0;
+
+        ConstraintAxisLinearVelocities[0] = FVec3::DotProduct(CV, Axis);
+    }
+
+    void FPBDJointCachedSolver::CalculateConstraintAxisAngularVelocities(
+        const FPBDJointSolverSettings& SolverSettings,
+        const FPBDJointSettings& JointSettings,
+        FVec3& ConstraintAxisAngularVelocities
+        ) const
+    {
+        bool bHasRotationConstraints =
+            (JointSettings.AngularMotionTypes[0] != EJointMotionType::Free)
+            || (JointSettings.AngularMotionTypes[1] != EJointMotionType::Free)
+            || (JointSettings.AngularMotionTypes[2] != EJointMotionType::Free);
+        if (!bHasRotationConstraints)
+        {
+            return;
+        }
+
+        // Locked axes always use hard constraints. Limited axes use hard or soft depending on settings
+        EJointMotionType TwistMotion = JointSettings.AngularMotionTypes[(int32)EJointAngularConstraintIndex::Twist];
+        EJointMotionType Swing1Motion = JointSettings.AngularMotionTypes[(int32)EJointAngularConstraintIndex::Swing1];
+        EJointMotionType Swing2Motion = JointSettings.AngularMotionTypes[(int32)EJointAngularConstraintIndex::Swing2];
+
+        // If the twist axes are opposing, we cannot decompose the orientation into swing and twist angles, so just give up
+        const FVec3 Twist0 = ConnectorRs[0] * FJointConstants::TwistAxis();
+        const FVec3 Twist1 = ConnectorRs[1] * FJointConstants::TwistAxis();
+        const FReal Twist01Dot = FVec3::DotProduct(Twist0, Twist1);
+        const bool bDegenerate = (Twist01Dot < Chaos_Joint_DegenerateRotationLimit);
+        if (bDegenerate)
+        {
+            UE_LOG(LogChaosJoint, VeryVerbose, TEXT(" Degenerate rotation at Swing %f deg"), FMath::RadiansToDegrees(FMath::Acos(Twist01Dot)));
+        }
+
+        // Apply twist constraint
+        // NOTE: Cannot calculate twist angle at 180degree swing
+        if (SolverSettings.bEnableTwistLimits && !bDegenerate)
+        {
+            if (TwistMotion == EJointMotionType::Limited)
+            {
+                CalculateTwistConstraintAxisAngularVelocities(JointSettings, ConstraintAxisAngularVelocities);
+            }
+        }
+
+        // Apply swing constraints
+        // NOTE: Cannot separate swing angles at 180degree swing (but we can still apply locks)
+        if (SolverSettings.bEnableSwingLimits)
+        {
+            if ((Swing1Motion == EJointMotionType::Limited) && (Swing2Motion == EJointMotionType::Limited))
+            {
+                CalculateConeConstraintAxisAngularVelocities(JointSettings, ConstraintAxisAngularVelocities);
+            }
+            else if ((Swing1Motion == EJointMotionType::Limited) && (Swing2Motion == EJointMotionType::Locked))
+            {
+                if (!bDegenerate)
+                {
+                    CalculateSwingConstraintAxisAngularVelocities(SolverSettings, JointSettings, EJointAngularConstraintIndex::Swing1, ConstraintAxisAngularVelocities);
+                }
+            }
+            else if ((Swing1Motion == EJointMotionType::Limited) && (Swing2Motion == EJointMotionType::Free))
+            {
+                if (!bDegenerate)
+                {
+                    CalculateDualConeSwingConstraintAxisAngularVelocities(JointSettings, EJointAngularConstraintIndex::Swing1, ConstraintAxisAngularVelocities);
+                }
+            }
+            else if ((Swing1Motion == EJointMotionType::Locked) && (Swing2Motion == EJointMotionType::Limited))
+            {
+                if (!bDegenerate)
+                {
+                    CalculateSwingConstraintAxisAngularVelocities(SolverSettings, JointSettings, EJointAngularConstraintIndex::Swing2, ConstraintAxisAngularVelocities);
+                }
+            }
+            else if ((Swing1Motion == EJointMotionType::Locked) && (Swing2Motion == EJointMotionType::Locked))
+            {
+            }
+            else if ((Swing1Motion == EJointMotionType::Locked) && (Swing2Motion == EJointMotionType::Free))
+            {
+            }
+            else if ((Swing1Motion == EJointMotionType::Free) && (Swing2Motion == EJointMotionType::Limited))
+            {
+                if (!bDegenerate)
+                {
+                    CalculateDualConeSwingConstraintAxisAngularVelocities(JointSettings, EJointAngularConstraintIndex::Swing2, ConstraintAxisAngularVelocities);
+                }
+            }
+            else if ((Swing1Motion == EJointMotionType::Free) && (Swing2Motion == EJointMotionType::Locked))
+            {
+            }
+            else if ((Swing1Motion == EJointMotionType::Free) && (Swing2Motion == EJointMotionType::Free))
+            {
+            }
+        }
+    }
+
+    void FPBDJointCachedSolver::CalculateTwistConstraintAxisAngularVelocities(
+        const FPBDJointSettings& JointSettings,
+        FVec3& ConstraintAxisAngularVelocities) const
+    {
+        FVec3 TwistAxis;
+        FReal TwistAngle;
+        FPBDJointUtilities::GetTwistAxisAngle(ConnectorRs[0], ConnectorRs[1], TwistAxis, TwistAngle);
+
+        if (TwistAngle < 0.0f)
+        {
+            TwistAxis = -TwistAxis;
+        }
+        ConstraintAxisAngularVelocities[(int32)EJointAngularConstraintIndex::Twist] = FVec3::DotProduct(W(1) - W(0), TwistAxis);
+    }
+
+    void FPBDJointCachedSolver::CalculateConeConstraintAxisAngularVelocities(
+            const FPBDJointSettings& JointSettings,
+            FVec3& ConstraintAxisAngularVelocities) const
+    {
+        FVec3 SwingAxisLocal;
+        FReal DSwingAngle = 0.0f;
+        const FReal Swing1Limit = JointSettings.AngularLimits[(int32)EJointAngularConstraintIndex::Swing1];
+        const FReal Swing2Limit = JointSettings.AngularLimits[(int32)EJointAngularConstraintIndex::Swing2];
+        FPBDJointUtilities::GetEllipticalConeAxisErrorLocal(ConnectorRs[0], ConnectorRs[1], Swing2Limit, Swing1Limit, SwingAxisLocal, DSwingAngle);
+        SwingAxisLocal.SafeNormalize();
+        const FVec3 SwingAxis = ConnectorRs[0] * SwingAxisLocal;
+        ConstraintAxisAngularVelocities[(int32)EJointAngularConstraintIndex::Swing1] = FVec3::DotProduct(W(1) - W(0), SwingAxis);   
+    }
+
+    void FPBDJointCachedSolver::CalculateDualConeSwingConstraintAxisAngularVelocities(
+            const FPBDJointSettings& JointSettings,
+            const EJointAngularConstraintIndex SwingConstraintIndex,
+            FVec3& ConstraintAxisAngularVelocities) const
+    {
+        FVec3 SwingAxis;
+        FReal SwingAngle;
+        FPBDJointUtilities::GetDualConeSwingAxisAngle(ConnectorRs[0], ConnectorRs[1], SwingConstraintIndex, SwingAxis, SwingAngle);
+        ConstraintAxisAngularVelocities[(int32)SwingConstraintIndex] = FVec3::DotProduct(W(1) - W(0), SwingAxis);
+    }
+
+    void FPBDJointCachedSolver::CalculateSwingConstraintAxisAngularVelocities(
+            const FPBDJointSolverSettings& SolverSettings,
+            const FPBDJointSettings& JointSettings,
+            const EJointAngularConstraintIndex SwingConstraintIndex,
+            FVec3& ConstraintAxisAngularVelocities) const
+    {
+        FVec3 SwingAxis;
+        FReal SwingAngle;
+        FPBDJointUtilities::GetSwingAxisAngle(ConnectorRs[0], ConnectorRs[1], SolverSettings.SwingTwistAngleTolerance, SwingConstraintIndex, SwingAxis, SwingAngle);
+        ConstraintAxisAngularVelocities[(int32)SwingConstraintIndex] = FVec3::DotProduct(W(1) - W(0), SwingAxis);
+    }
+}
diff --git a/Engine/Source/Runtime/Experimental/Chaos/Private/Chaos/PBDJointConstraints.cpp b/Engine/Source/Runtime/Experimental/Chaos/Private/Chaos/PBDJointConstraints.cpp
--- a/Engine/Source/Runtime/Experimental/Chaos/Private/Chaos/PBDJointConstraints.cpp
+++ b/Engine/Source/Runtime/Experimental/Chaos/Private/Chaos/PBDJointConstraints.cpp
@@ -719,14 +719,28 @@ namespace Chaos
 			bJointsDirty = false;
 		}
 
-		ConstraintSolvers.SetNum(NumConstraints());
+		if (bChaos_Joint_UseCachedSolver)
+		{
+			CachedConstraintSolvers.SetNum(NumConstraints());
+		}
+		else
+		{
+			ConstraintSolvers.SetNum(NumConstraints());		
+		}
 	}
 
 	void FPBDJointConstraints::UnprepareTick()
 	{
 		SCOPE_CYCLE_COUNTER(STAT_Joints_UnprepareTick);
 
-		ConstraintSolvers.Empty();
+		if (bChaos_Joint_UseCachedSolver)
+		{
+			CachedConstraintSolvers.Empty();
+		}
+		else
+		{
+			ConstraintSolvers.Empty();
+		}
 	}
 	
 	void FPBDJointConstraints::CalculateConstraintSpace(int32 ConstraintIndex, FVec3& OutX0, FMatrix33& OutR0, FVec3& OutX1, FMatrix33& OutR1) const
@@ -801,7 +815,6 @@ namespace Chaos
 
 		const FPBDJointSettings& JointSettings = ConstraintSettings[ConstraintIndex];
 		const FTransformPair& JointFrames = JointSettings.ConnectorTransforms;
-		FPBDJointSolver& Solver = ConstraintSolvers[ConstraintIndex];
 
 		int32 Index0, Index1;
 		GetConstrainedParticleIndices(ConstraintIndex, Index0, Index1);
@@ -813,13 +826,28 @@ namespace Chaos
 		FSolverBody* Body0 = SolverData.GetBodyContainer().FindOrAdd(Particle0);
 		FSolverBody* Body1 = SolverData.GetBodyContainer().FindOrAdd(Particle1);
 
-		Solver.Init(
-			Dt,
-			{ Body0, Body1 },
-			Settings,
-			JointSettings,
-			FParticleUtilities::ParticleLocalToCoMLocal(Particle0, JointFrames[Index0]),
-			FParticleUtilities::ParticleLocalToCoMLocal(Particle1, JointFrames[Index1]));
+		if (bChaos_Joint_UseCachedSolver)
+		{
+			FPBDJointCachedSolver& Solver = CachedConstraintSolvers[ConstraintIndex];
+			Solver.Init(
+				Dt,
+				{ Body0, Body1 },
+				Settings,
+				JointSettings,
+				FParticleUtilities::ParticleLocalToCoMLocal(Particle0, JointFrames[Index0]),
+				FParticleUtilities::ParticleLocalToCoMLocal(Particle1, JointFrames[Index1]));
+		}
+		else
+		{
+			FPBDJointSolver& Solver = ConstraintSolvers[ConstraintIndex];
+			Solver.Init(
+				Dt,
+				{ Body0, Body1 },
+				Settings,
+				JointSettings,
+				FParticleUtilities::ParticleLocalToCoMLocal(Particle0, JointFrames[Index0]),
+				FParticleUtilities::ParticleLocalToCoMLocal(Particle1, JointFrames[Index1]));
+		}
 
 		// Plasticity should not be turned on in the middle of simulation.
 		const bool bUseLinearPlasticity = JointSettings.LinearPlasticityLimit != FLT_MAX;
@@ -851,36 +879,68 @@ namespace Chaos
 			if (!ConstraintStates[ConstraintIndex].bDisabled)
 			{
 				FPBDJointState& JointState = ConstraintStates[ConstraintIndex];
-				FPBDJointSolver& Solver = ConstraintSolvers[ConstraintIndex];
 
 				int32 Index0, Index1;
 				GetConstrainedParticleIndices(ConstraintIndex, Index0, Index1);
 
-				// NOTE: LinearImpulse/AngularImpulse in the solver are not really impulses - they are mass-weighted position/rotation delta, or (impulse x dt).
-				if (Dt > SMALL_NUMBER)
+				if (bChaos_Joint_UseCachedSolver)
 				{
-					if (Index0 == 0)
+					FPBDJointCachedSolver& Solver = CachedConstraintSolvers[ConstraintIndex];
+					// NOTE: LinearImpulse/AngularImpulse in the solver are not really impulses - they are mass-weighted position/rotation delta, or (impulse x dt).
+					if (Dt > SMALL_NUMBER)
 					{
-						JointState.LinearImpulse = Solver.GetNetLinearImpulse() / Dt;
-						JointState.AngularImpulse = Solver.GetNetAngularImpulse() / Dt;
+						if (Index0 == 0)
+						{
+							JointState.LinearImpulse = Solver.GetNetLinearImpulse() / Dt;
+							JointState.AngularImpulse = Solver.GetNetAngularImpulse() / Dt;
+						}
+						else
+						{
+							// Particles were flipped in the solver...
+							JointState.LinearImpulse = -Solver.GetNetLinearImpulse() / Dt;
+							JointState.AngularImpulse = -Solver.GetNetAngularImpulse() / Dt;
+						}
 					}
 					else
 					{
-						// Particles were flipped in the solver...
-						JointState.LinearImpulse = -Solver.GetNetLinearImpulse() / Dt;
-						JointState.AngularImpulse = -Solver.GetNetAngularImpulse() / Dt;
+						JointState.LinearImpulse = FVec3(0);
+						JointState.AngularImpulse = FVec3(0);
 					}
+
+					ApplyPlasticityLimits(ConstraintIndex);
+
+					// Remove our solver body reference (they are not valid between frames)
+					CachedConstraintSolvers[ConstraintIndex].Deinit();
 				}
 				else
 				{
-					JointState.LinearImpulse = FVec3(0);
-					JointState.AngularImpulse = FVec3(0);
-				}
+					FPBDJointSolver& Solver = ConstraintSolvers[ConstraintIndex];
+					// NOTE: LinearImpulse/AngularImpulse in the solver are not really impulses - they are mass-weighted position/rotation delta, or (impulse x dt).
+					if (Dt > SMALL_NUMBER)
+					{
+						if (Index0 == 0)
+						{
+							JointState.LinearImpulse = Solver.GetNetLinearImpulse() / Dt;
+							JointState.AngularImpulse = Solver.GetNetAngularImpulse() / Dt;
+						}
+						else
+						{
+							// Particles were flipped in the solver...
+							JointState.LinearImpulse = -Solver.GetNetLinearImpulse() / Dt;
+							JointState.AngularImpulse = -Solver.GetNetAngularImpulse() / Dt;
+						}
+					}
+					else
+					{
+						JointState.LinearImpulse = FVec3(0);
+						JointState.AngularImpulse = FVec3(0);
+					}
 
-				ApplyPlasticityLimits(ConstraintIndex);
+					ApplyPlasticityLimits(ConstraintIndex);
 
-				// Remove our solver body reference (they are not valid between frames)
-				ConstraintSolvers[ConstraintIndex].Deinit();
+					// Remove our solver body reference (they are not valid between frames)
+					ConstraintSolvers[ConstraintIndex].Deinit();
+				}
 			}
 		}
 	}
@@ -909,6 +969,7 @@ namespace Chaos
 		int32 NumActive = 0;
 		if (Settings.ApplyPairIterations > 0)
 		{
+			// UpdateConstraintProjection(It, NumIts, SolverData);
 			for (int32 ConstraintIndex : SolverData.GetConstraintIndices(ContainerId))
 			{
 				NumActive += ApplyPhase2Single(Dt, ConstraintIndex, Settings.ApplyPairIterations, It, NumIts);
@@ -917,6 +978,26 @@ namespace Chaos
 		return (NumActive > 0);
 	}
 
+	void FPBDJointConstraints::UpdateConstraintProjection(const int32 It, const int32 NumIts, const FPBDIslandSolverData& SolverData)
+	{
+		const bool bUseProjection = It == NumIts - 1;
+		if (bUseProjection)
+		{
+			for (int32 ConstraintIndex : SolverData.GetConstraintIndices(ContainerId))
+			{
+				const FPBDJointSettings& JointSettings = ConstraintSettings[ConstraintIndex];
+				if (bChaos_Joint_UseCachedSolver)
+				{
+					if (JointSettings.bProjectionEnabled)
+					{
+						FPBDJointCachedSolver& Solver = CachedConstraintSolvers[ConstraintIndex];
+						Solver.EnableProjection();
+					}
+				}
+			}
+		}
+	}
+
 	//////////////////////////////////////////////////////////////////////////
 	//
 	// End Graph API Solver.
@@ -964,13 +1045,27 @@ namespace Chaos
 		}
 
 		// check valid particle and solver state
-		const FPBDJointSolver& Solver = ConstraintSolvers[ConstraintIndex];
-		if ((Particle0->Sleeping() && Particle1->Sleeping())
-			|| (Particle0->IsKinematic() && Particle1->Sleeping())
-			|| (Particle0->Sleeping() && Particle1->IsKinematic())
-			|| (FMath::IsNearlyZero(Solver.InvM(0)) && FMath::IsNearlyZero(Solver.InvM(1))))
+		if (bChaos_Joint_UseCachedSolver)
 		{
-			return false;
+			const FPBDJointCachedSolver& Solver = CachedConstraintSolvers[ConstraintIndex];
+			if ((Particle0->Sleeping() && Particle1->Sleeping())
+				|| (Particle0->IsKinematic() && Particle1->Sleeping())
+				|| (Particle0->Sleeping() && Particle1->IsKinematic())
+				|| (FMath::IsNearlyZero(Solver.InvM(0)) && FMath::IsNearlyZero(Solver.InvM(1))))
+			{
+				return false;
+			}
+		}
+		else
+		{
+			const FPBDJointSolver& Solver = ConstraintSolvers[ConstraintIndex];
+			if ((Particle0->Sleeping() && Particle1->Sleeping())
+				|| (Particle0->IsKinematic() && Particle1->Sleeping())
+				|| (Particle0->Sleeping() && Particle1->IsKinematic())
+				|| (FMath::IsNearlyZero(Solver.InvM(0)) && FMath::IsNearlyZero(Solver.InvM(1))))
+			{
+				return false;
+			}
 		}
 		return true;
 	}
@@ -988,59 +1083,118 @@ namespace Chaos
 		UE_LOG(LogChaosJoint, VeryVerbose, TEXT("Solve Joint Constraint %d %s %s (dt = %f; it = %d / %d)"), ConstraintIndex, *Constraint[0]->ToString(), *Constraint[1]->ToString(), Dt, It, NumIts);
 
 		const FPBDJointSettings& JointSettings = ConstraintSettings[ConstraintIndex];
-		FPBDJointSolver& Solver = ConstraintSolvers[ConstraintIndex];
-
-		// @todo(chaos): store this on the Solver object and don't access the particles here
-		int32 Index0, Index1;
-		GetConstrainedParticleIndices(ConstraintIndex, Index0, Index1);
-		FGenericParticleHandle Particle0 = FGenericParticleHandle(ConstraintParticles[ConstraintIndex][Index0]);
-		FGenericParticleHandle Particle1 = FGenericParticleHandle(ConstraintParticles[ConstraintIndex][Index1]);
-		if ((Particle0->Sleeping() && Particle1->Sleeping())
-			|| (Particle0->IsKinematic() && Particle1->Sleeping()) 
-			|| (Particle0->Sleeping() && Particle1->IsKinematic())
-			|| (FMath::IsNearlyZero(Solver.InvM(0)) && FMath::IsNearlyZero(Solver.InvM(1))))
+		if (bChaos_Joint_UseCachedSolver)
 		{
-			return false;
-		}
+			FPBDJointCachedSolver& Solver = CachedConstraintSolvers[ConstraintIndex];
 
-		// If we were solved last iteration and nothing has changed since, we are done
-		const bool bWasActive = Solver.GetIsActive();
-		Solver.Update(Dt, Settings, JointSettings);
+			// @todo(chaos): store this on the Solver object and don't access the particles here
+			int32 Index0, Index1;
+			GetConstrainedParticleIndices(ConstraintIndex, Index0, Index1);
+			FGenericParticleHandle Particle0 = FGenericParticleHandle(ConstraintParticles[ConstraintIndex][Index0]);
+			FGenericParticleHandle Particle1 = FGenericParticleHandle(ConstraintParticles[ConstraintIndex][Index1]);
+			if ((Particle0->Sleeping() && Particle1->Sleeping())
+				|| (Particle0->IsKinematic() && Particle1->Sleeping()) 
+				|| (Particle0->Sleeping() && Particle1->IsKinematic())
+				|| (FMath::IsNearlyZero(Solver.InvM(0)) && FMath::IsNearlyZero(Solver.InvM(1))))
+			{
+				return false;
+			}
 
-		if (!bWasActive && !Solver.GetIsActive() && bChaos_Joint_EarlyOut_Enabled)
-		{
-			return false;
-		}
+			// If we were solved last iteration and nothing has changed since, we are done
+			const bool bWasActive = Solver.GetIsActive();
+			Solver.Update(Dt, Settings, JointSettings);
 
-		const FReal IterationStiffness = CalculateIterationStiffness(It, NumIts);
-		for (int32 PairIt = 0; PairIt < NumPairIts; ++PairIt)
-		{
-			UE_LOG(LogChaosJoint, VeryVerbose, TEXT("  Pair Iteration %d / %d"), PairIt, NumPairIts);
+			if (!bWasActive && !Solver.GetIsActive() && bChaos_Joint_EarlyOut_Enabled)
+			{
+				return false;
+			}
+
+			const FReal IterationStiffness = CalculateIterationStiffness(It, NumIts);
+			for (int32 PairIt = 0; PairIt < NumPairIts; ++PairIt)
+			{
+				UE_LOG(LogChaosJoint, VeryVerbose, TEXT("  Pair Iteration %d / %d"), PairIt, NumPairIts);
 
-			// This is the same position solver for all SolverTypes (which makes it wrong for the GbfPbd solver which should be a pbd-esque velocity solve)
-			Solver.ApplyConstraints(Dt, IterationStiffness, Settings, JointSettings);
+				// This is the same position solver for all SolverTypes (which makes it wrong for the GbfPbd solver which should be a pbd-esque velocity solve)
+				Solver.ApplyConstraints(Dt, IterationStiffness, Settings, JointSettings);
+
+				if (!Solver.GetIsActive() && bChaos_Joint_EarlyOut_Enabled)
+				{
+					break;
+				}
+			}
 
-			if (!Solver.GetIsActive() && bChaos_Joint_EarlyOut_Enabled)
+			// If the solver we are running uses a velocity solver in the Apply step (for other constraint types), 
+			// we need to update the implicit velocity because we did a position solve above. E.g., GbfPbd collisions.
+			if (SolverType == EConstraintSolverType::GbfPbd)
 			{
-				break;
+				Solver.Body(0).SetImplicitVelocity(Dt);
+				Solver.Body(1).SetImplicitVelocity(Dt);
 			}
-		}
 
-		// If the solver we are running uses a velocity solver in the Apply step (for other constraint types), 
-		// we need to update the implicit velocity because we did a position solve above. E.g., GbfPbd collisions.
-		if (SolverType == EConstraintSolverType::GbfPbd)
-		{
-			Solver.Body(0).SetImplicitVelocity(Dt);
-			Solver.Body(1).SetImplicitVelocity(Dt);
-		}
+			// @todo(ccaulfield): The break limit should really be applied to the impulse in the solver to prevent 1-frame impulses larger than the threshold
+			if ((JointSettings.LinearBreakForce!=FLT_MAX) || (JointSettings.AngularBreakTorque!=FLT_MAX))
+			{
+				ApplyBreakThreshold(Dt, ConstraintIndex, Solver.GetNetLinearImpulse(), Solver.GetNetAngularImpulse());
+			}
 
-		// @todo(ccaulfield): The break limit should really be applied to the impulse in the solver to prevent 1-frame impulses larger than the threshold
-		if ((JointSettings.LinearBreakForce!=FLT_MAX) || (JointSettings.AngularBreakTorque!=FLT_MAX))
-		{
-			ApplyBreakThreshold(Dt, ConstraintIndex, Solver.GetNetLinearImpulse(), Solver.GetNetAngularImpulse());
+			return Solver.GetIsActive() || !bChaos_Joint_EarlyOut_Enabled;
 		}
+		else
+		{
+			FPBDJointSolver& Solver = ConstraintSolvers[ConstraintIndex];
+
+			// @todo(chaos): store this on the Solver object and don't access the particles here
+			int32 Index0, Index1;
+			GetConstrainedParticleIndices(ConstraintIndex, Index0, Index1);
+			FGenericParticleHandle Particle0 = FGenericParticleHandle(ConstraintParticles[ConstraintIndex][Index0]);
+			FGenericParticleHandle Particle1 = FGenericParticleHandle(ConstraintParticles[ConstraintIndex][Index1]);
+			if ((Particle0->Sleeping() && Particle1->Sleeping())
+				|| (Particle0->IsKinematic() && Particle1->Sleeping()) 
+				|| (Particle0->Sleeping() && Particle1->IsKinematic())
+				|| (FMath::IsNearlyZero(Solver.InvM(0)) && FMath::IsNearlyZero(Solver.InvM(1))))
+			{
+				return false;
+			}
+
+			// If we were solved last iteration and nothing has changed since, we are done
+			const bool bWasActive = Solver.GetIsActive();
+			Solver.Update(Dt, Settings, JointSettings);
+
+			if (!bWasActive && !Solver.GetIsActive() && bChaos_Joint_EarlyOut_Enabled)
+			{
+				return false;
+			}
+
+			const FReal IterationStiffness = CalculateIterationStiffness(It, NumIts);
+			for (int32 PairIt = 0; PairIt < NumPairIts; ++PairIt)
+			{
+				UE_LOG(LogChaosJoint, VeryVerbose, TEXT("  Pair Iteration %d / %d"), PairIt, NumPairIts);
 
-		return Solver.GetIsActive() || !bChaos_Joint_EarlyOut_Enabled;
+				// This is the same position solver for all SolverTypes (which makes it wrong for the GbfPbd solver which should be a pbd-esque velocity solve)
+				Solver.ApplyConstraints(Dt, IterationStiffness, Settings, JointSettings);
+
+				if (!Solver.GetIsActive() && bChaos_Joint_EarlyOut_Enabled)
+				{
+					break;
+				}
+			}
+
+			// If the solver we are running uses a velocity solver in the Apply step (for other constraint types), 
+			// we need to update the implicit velocity because we did a position solve above. E.g., GbfPbd collisions.
+			if (SolverType == EConstraintSolverType::GbfPbd)
+			{
+				Solver.Body(0).SetImplicitVelocity(Dt);
+				Solver.Body(1).SetImplicitVelocity(Dt);
+			}
+
+			// @todo(ccaulfield): The break limit should really be applied to the impulse in the solver to prevent 1-frame impulses larger than the threshold
+			if ((JointSettings.LinearBreakForce!=FLT_MAX) || (JointSettings.AngularBreakTorque!=FLT_MAX))
+			{
+				ApplyBreakThreshold(Dt, ConstraintIndex, Solver.GetNetLinearImpulse(), Solver.GetNetAngularImpulse());
+			}
+
+			return Solver.GetIsActive() || !bChaos_Joint_EarlyOut_Enabled;
+		}
 	}
 
 	bool FPBDJointConstraints::ApplyPhase2Single(const FReal Dt, const int32 ConstraintIndex, const int32 NumPairIts, const int32 It, const int32 NumIts)
@@ -1054,56 +1208,112 @@ namespace Chaos
 		UE_LOG(LogChaosJoint, VeryVerbose, TEXT("Project Joint Constraint %d %s %s (dt = %f; it = %d / %d)"), ConstraintIndex, *Constraint[0]->ToString(), *Constraint[1]->ToString(), Dt, It, NumIts);
 
 		const FPBDJointSettings& JointSettings = ConstraintSettings[ConstraintIndex];
-		FPBDJointSolver& Solver = ConstraintSolvers[ConstraintIndex];
-
-		// @todo(chaos): store this on the Solver object and don't access the particles here
-		int32 Index0, Index1;
-		GetConstrainedParticleIndices(ConstraintIndex, Index0, Index1);
-		FGenericParticleHandle Particle0 = FGenericParticleHandle(ConstraintParticles[ConstraintIndex][Index0]);
-		FGenericParticleHandle Particle1 = FGenericParticleHandle(ConstraintParticles[ConstraintIndex][Index1]);
-		if ((Particle0->Sleeping() && Particle1->Sleeping())
-			|| (Particle0->IsKinematic() && Particle1->Sleeping())
-			|| (Particle0->Sleeping() && Particle1->IsKinematic()) 
-			|| (FMath::IsNearlyZero(Solver.InvM(0)) && FMath::IsNearlyZero(Solver.InvM(1))))
+		if (bChaos_Joint_UseCachedSolver)
 		{
-			return false;
-		}
+			FPBDJointCachedSolver& Solver = CachedConstraintSolvers[ConstraintIndex];
 
-		// If we were solved last iteration and nothing has changed since, we are done
-		const bool bWasActive = Solver.GetIsActive();
-		Solver.Update(Dt, Settings, JointSettings);
-		
-		// For quasipbd, we may still need to stabilize velocities and solve for restitution.
-		if (!bWasActive && !Solver.GetIsActive() && SolverType != EConstraintSolverType::QuasiPbd && bChaos_Joint_EarlyOut_Enabled)
-		{
-			return false;
-		}
+			// @todo(chaos): store this on the Solver object and don't access the particles here
+			int32 Index0, Index1;
+			GetConstrainedParticleIndices(ConstraintIndex, Index0, Index1);
+			FGenericParticleHandle Particle0 = FGenericParticleHandle(ConstraintParticles[ConstraintIndex][Index0]);
+			FGenericParticleHandle Particle1 = FGenericParticleHandle(ConstraintParticles[ConstraintIndex][Index1]);
+			if ((Particle0->Sleeping() && Particle1->Sleeping())
+				|| (Particle0->IsKinematic() && Particle1->Sleeping())
+				|| (Particle0->Sleeping() && Particle1->IsKinematic()) 
+				|| (FMath::IsNearlyZero(Solver.InvM(0)) && FMath::IsNearlyZero(Solver.InvM(1))))
+			{
+				return false;
+			}
 
-		const FReal IterationStiffness = CalculateIterationStiffness(It, NumIts);
-		for (int32 PairIt = 0; PairIt < NumPairIts; ++PairIt)
-		{
-			switch(SolverType)
+			// If we were solved last iteration and nothing has changed since, we are done
+			const bool bWasActive = Solver.GetIsActive();
+			Solver.Update(Dt, Settings, JointSettings);
+			
+			// For quasipbd, we may still need to stabilize velocities and solve for restitution.
+			if (!bWasActive && !Solver.GetIsActive() && SolverType != EConstraintSolverType::QuasiPbd && bChaos_Joint_EarlyOut_Enabled)
 			{
-			case EConstraintSolverType::None:
-				break;
-			case EConstraintSolverType::GbfPbd:
-			case EConstraintSolverType::StandardPbd:
-				Solver.ApplyProjections(Dt, IterationStiffness, Settings, JointSettings);
-				break;
-			case EConstraintSolverType::QuasiPbd:
-				Solver.ApplyVelocityConstraints(Dt, IterationStiffness, Settings, JointSettings);
-				break;
+				return false;
 			}
 
-			if (!Solver.GetIsActive() && bChaos_Joint_EarlyOut_Enabled)
+			const FReal IterationStiffness = CalculateIterationStiffness(It, NumIts);
+			for (int32 PairIt = 0; PairIt < NumPairIts; ++PairIt)
 			{
-				break;
+				switch(SolverType)
+				{
+				case EConstraintSolverType::None:
+					break;
+				case EConstraintSolverType::GbfPbd:
+				case EConstraintSolverType::StandardPbd:
+					Solver.ApplyProjections(Dt, IterationStiffness, Settings, JointSettings);
+					break;
+				case EConstraintSolverType::QuasiPbd:
+					Solver.ApplyVelocityConstraints(Dt, IterationStiffness, Settings, JointSettings);
+					break;
+				}
+
+				if (!Solver.GetIsActive() && bChaos_Joint_EarlyOut_Enabled)
+				{
+					break;
+				}
 			}
+
+			// @todo(ccaulfield): should probably add to net impulses in push out too...(for breaking etc)
+
+			return Solver.GetIsActive() || !bChaos_Joint_EarlyOut_Enabled;
 		}
+		else
+		{
+			FPBDJointSolver& Solver = ConstraintSolvers[ConstraintIndex];
 
-		// @todo(ccaulfield): should probably add to net impulses in push out too...(for breaking etc)
+			// @todo(chaos): store this on the Solver object and don't access the particles here
+			int32 Index0, Index1;
+			GetConstrainedParticleIndices(ConstraintIndex, Index0, Index1);
+			FGenericParticleHandle Particle0 = FGenericParticleHandle(ConstraintParticles[ConstraintIndex][Index0]);
+			FGenericParticleHandle Particle1 = FGenericParticleHandle(ConstraintParticles[ConstraintIndex][Index1]);
+			if ((Particle0->Sleeping() && Particle1->Sleeping())
+				|| (Particle0->IsKinematic() && Particle1->Sleeping())
+				|| (Particle0->Sleeping() && Particle1->IsKinematic()) 
+				|| (FMath::IsNearlyZero(Solver.InvM(0)) && FMath::IsNearlyZero(Solver.InvM(1))))
+			{
+				return false;
+			}
+
+			// If we were solved last iteration and nothing has changed since, we are done
+			const bool bWasActive = Solver.GetIsActive();
+			Solver.Update(Dt, Settings, JointSettings);
+			
+			// For quasipbd, we may still need to stabilize velocities and solve for restitution.
+			if (!bWasActive && !Solver.GetIsActive() && SolverType != EConstraintSolverType::QuasiPbd && bChaos_Joint_EarlyOut_Enabled)
+			{
+				return false;
+			}
+
+			const FReal IterationStiffness = CalculateIterationStiffness(It, NumIts);
+			for (int32 PairIt = 0; PairIt < NumPairIts; ++PairIt)
+			{
+				switch(SolverType)
+				{
+				case EConstraintSolverType::None:
+					break;
+				case EConstraintSolverType::GbfPbd:
+				case EConstraintSolverType::StandardPbd:
+					Solver.ApplyProjections(Dt, IterationStiffness, Settings, JointSettings);
+					break;
+				case EConstraintSolverType::QuasiPbd:
+					Solver.ApplyVelocityConstraints(Dt, IterationStiffness, Settings, JointSettings);
+					break;
+				}
 
-		return Solver.GetIsActive() || !bChaos_Joint_EarlyOut_Enabled;
+				if (!Solver.GetIsActive() && bChaos_Joint_EarlyOut_Enabled)
+				{
+					break;
+				}
+			}
+
+			// @todo(ccaulfield): should probably add to net impulses in push out too...(for breaking etc)
+
+			return Solver.GetIsActive() || !bChaos_Joint_EarlyOut_Enabled;
+		}
 	}
 
 	void FPBDJointConstraints::ApplyBreakThreshold(const FReal Dt, int32 ConstraintIndex, const FVec3& LinearImpulse, const FVec3& AngularImpulse)
@@ -1167,8 +1377,8 @@ namespace Chaos
 			}
 		}
 
-		FConstraintSolverBody& Body0 = ConstraintSolvers[ConstraintIndex].Body(0);
-		FConstraintSolverBody& Body1 = ConstraintSolvers[ConstraintIndex].Body(1);
+		FConstraintSolverBody& Body0 = bChaos_Joint_UseCachedSolver ? CachedConstraintSolvers[ConstraintIndex].Body(0) : ConstraintSolvers[ConstraintIndex].Body(0);
+		FConstraintSolverBody& Body1 = bChaos_Joint_UseCachedSolver ? CachedConstraintSolvers[ConstraintIndex].Body(1) : ConstraintSolvers[ConstraintIndex].Body(1);
 
 		const FTransformPair& ConstraintFramesLocal = JointSettings.ConnectorTransforms;
 		FTransformPair ConstraintFramesGlobal(ConstraintFramesLocal[Index0] * FRigidTransform3(Body0.ActorP(), Body0.ActorQ()), ConstraintFramesLocal[Index1] * FRigidTransform3(Body1.ActorP(), Body1.ActorQ()));
diff --git a/Engine/Source/Runtime/Experimental/Chaos/Public/Chaos/Joint/PBDJointCachedSolverGaussSeidel.h b/Engine/Source/Runtime/Experimental/Chaos/Public/Chaos/Joint/PBDJointCachedSolverGaussSeidel.h
new file mode 100644
--- /dev/null
+++ b/Engine/Source/Runtime/Experimental/Chaos/Public/Chaos/Joint/PBDJointCachedSolverGaussSeidel.h
@@ -0,0 +1,815 @@
+// Copyright Epic Games, Inc. All Rights Reserved.
+#pragma once
+
+#include "CoreMinimal.h"
+
+#include "Chaos/DenseMatrix.h"
+#include "Chaos/Evolution/SolverBody.h"
+#include "Chaos/ParticleHandleFwd.h"
+#include "Chaos/PBDJointConstraintTypes.h"
+#include "Chaos/PBDJointConstraintUtilities.h"
+#include "Chaos/Utilities.h"
+
+namespace Chaos
+{
+	class FPBDJointCachedSolver;
+
+	/**
+	 * Calculate new positions and rotations for a pair of bodies connected by a joint.
+	 *
+	 * This solver treats of the 6 possible constraints (up to 3 linear and 3 angular)
+	 * individually and resolves them in sequence.
+	 *
+	 * \see FJointSolverCholesky
+	 */
+	class FPBDJointCachedSolver
+	{
+	public:
+		static const int32 MaxConstrainedBodies = 2;
+
+		// Access to the SolverBodies that the constraint is operating on.
+		// \note SolverBodies are transient and exist only as long as we are in the Island's constraint solver loop
+		inline FConstraintSolverBody& Body(int32 BodyIndex)
+		{
+			check((BodyIndex >= 0) && (BodyIndex < 2));
+			check(SolverBodies[BodyIndex].IsValid());
+
+			return SolverBodies[BodyIndex];
+		}
+
+		inline const FConstraintSolverBody& Body(int32 BodyIndex) const
+		{
+			check((BodyIndex >= 0) && (BodyIndex < 2));
+			check(SolverBodies[BodyIndex].IsValid());
+
+			return SolverBodies[BodyIndex];
+		}
+
+		inline FConstraintSolverBody& Body0()
+		{
+			return SolverBodies[0];
+		}
+
+		inline const FConstraintSolverBody& Body0() const
+		{
+			return SolverBodies[0];
+		}
+
+		inline FConstraintSolverBody& Body1()
+		{
+			return SolverBodies[1];
+		}
+
+		inline const FConstraintSolverBody& Body1() const
+		{
+			return SolverBodies[1];
+		}
+
+		inline const FVec3& X(int BodyIndex) const
+		{
+			return Body(BodyIndex).X();
+		}
+
+		inline const FRotation3& R(int BodyIndex) const
+		{
+			return Body(BodyIndex).R();
+		}
+
+		inline const FVec3 P(int BodyIndex) const
+		{
+			// NOTE: Joints always use the latest post-correction position and rotation. This makes the joint error calculations non-linear and more robust against explosion
+			// but adds a cost because we need to apply the latest correction each time we request the latest transform
+			return Body(BodyIndex).CorrectedP();
+		}
+
+		inline const FRotation3 Q(int BodyIndex) const
+		{
+			// NOTE: Joints always use the latest post-correction position and rotation. This makes the joint error calculations non-linear and more robust against explosion
+			// but adds a cost because we need to apply the latest correction each time we request the latest transform
+			return Body(BodyIndex).CorrectedQ();
+		}
+
+		inline const FVec3& V(int BodyIndex) const
+		{
+			return Body(BodyIndex).V();
+		}
+
+		inline const FVec3& W(int BodyIndex) const
+		{
+			return Body(BodyIndex).W();
+		}
+
+		inline FReal InvM(int32 BodyIndex) const
+		{
+			return ConditionedInvMs[BodyIndex];
+		}
+
+		inline FMatrix33 InvI(int32 BodyIndex) const
+		{
+			return ConditionedInvIs[BodyIndex];
+		}
+
+		inline const FVec3& GetNetLinearImpulse() const
+		{
+			return NetLinearImpulse;
+		}
+
+		inline const FVec3& GetNetAngularImpulse() const
+		{
+			return NetAngularImpulse;
+		}
+
+		inline int32 GetNumActiveConstraints() const
+		{
+			// We use -1 as unitialized, but that should not be exposed outside the solver
+			return FMath::Max(NumActiveConstraints, 0);
+		}
+
+		inline bool GetIsActive() const
+		{
+			return bIsActive;
+		}
+
+
+		FPBDJointCachedSolver();
+
+		// Called once per frame to initialze the joint solver from the joint settings
+		void Init(
+			const FReal Dt,
+			const FSolverBodyPtrPair& SolverBodyPair,
+			const FPBDJointSolverSettings& SolverSettings,
+			const FPBDJointSettings& JointSettings,
+			const FRigidTransform3& XL0,
+			const FRigidTransform3& XL1);
+
+		void InitPositionConstraints(
+			const FPBDJointSolverSettings& SolverSettings,
+			const FPBDJointSettings& JointSettings);
+
+		void InitPointPositionConstraintKD(
+			const int32 KIndex,
+			const int32 DIndex);
+
+		void InitPointPositionConstraintDD();
+
+		void InitSphericalPositionConstraint(
+			const FPBDJointSettings& JointSettings);
+
+		void InitCylindricalPositionConstraint(
+		    const int32 AxisIndex,
+		    const EJointMotionType AxialMotion,
+		    const EJointMotionType RadialMotion,
+		    const FPBDJointSettings& JointSettings);
+
+		void InitPlanarPositionConstraint(
+			const int32 AxisIndex,
+			const EJointMotionType AxialMotion,
+			const FPBDJointSettings& JointSettings);
+
+		void InitPositionConstraint(
+			const FVec3& Axis,
+			const FVec3& ConstraintArm0,
+			const FVec3& ConstraintArm1,
+			const int32 ConstraintIndex);
+
+		void UpdatePositionConstraintIM(
+			const FVec3& Axis,
+			const FVec3& ConstraintArm0,
+			const FVec3& ConstraintArm1,
+			const int32 ConstraintIndex);
+
+		void Deinit();
+
+		// @todo(chaos): this doesn't do much now we have SolverBodies - remove it
+		void Update(
+			const FReal Dt,
+			const FPBDJointSolverSettings& SolverSettings,
+			const FPBDJointSettings& JointSettings);
+
+		// Run the position solve for the constraints
+		void ApplyConstraints(
+			const FReal Dt,
+			const FReal InSolverStiffness,
+			const FPBDJointSolverSettings& SolverSettings,
+			const FPBDJointSettings& JointSettings);
+
+		// Run the velocity solve for the constraints
+		void ApplyVelocityConstraints(
+			const FReal Dt,
+			const FReal InSolverStiffness,
+			const FPBDJointSolverSettings& SolverSettings,
+			const FPBDJointSettings& JointSettings);
+
+		// Apply projection (a position solve where the parent has infinite mass) for the constraints
+		// @todo(chaos): this can be build into the Apply phase when the whole solver is PBD
+		void ApplyProjections(
+			const FReal Dt,
+			const FReal InSolverStiffness,
+			const FPBDJointSolverSettings& SolverSettings,
+			const FPBDJointSettings& JointSettings);
+
+		void EnableProjection();
+
+	private:
+
+		// Initialize state dependent on particle positions
+		void InitDerivedState();
+
+		// Update state dependent on particle positions (after moving one or both of them)
+		void UpdateDerivedState(const int32 BodyIndex);
+		void UpdateDerivedState();
+
+		// Check to see if this constraint still needs further solved
+		// @todo(chaos): the term "active" is used inconsistently with the meaning elsewhere. Active should
+		// mean "contributed impusles". Should use "solved" rather than "active"
+		bool UpdateIsActive(
+			const FReal Dt,
+			const FPBDJointSolverSettings& SolverSettings,
+			const FPBDJointSettings& JointSettings);
+
+		void ApplyPositionConstraints(
+			const FReal Dt,
+			const FPBDJointSolverSettings& SolverSettings,
+			const FPBDJointSettings& JointSettings);
+
+		void ApplyRotationConstraints(
+			const FReal Dt,
+			const FPBDJointSolverSettings& SolverSettings,
+			const FPBDJointSettings& JointSettings);
+
+		void ApplyPositionDrives(
+			const FReal Dt,
+			const FPBDJointSolverSettings& SolverSettings,
+			const FPBDJointSettings& JointSettings);
+
+		void ApplyRotationDrives(
+			const FReal Dt,
+			const FPBDJointSolverSettings& SolverSettings,
+			const FPBDJointSettings& JointSettings);
+
+		void ApplyLinearVelocityConstraints(
+			const FReal Dt,
+			const FPBDJointSolverSettings& SolverSettings,
+			const FPBDJointSettings& JointSettings);
+
+		void ApplyAngularVelocityConstraints(
+			const FReal Dt,
+			const FPBDJointSolverSettings& SolverSettings,
+			const FPBDJointSettings& JointSettings);
+
+		void ApplyPositionDelta(
+			const int32 BodyIndex,
+			const FVec3& DP);
+
+		void ApplyPositionDelta(
+			const FVec3& DP0,
+			const FVec3& DP1);
+
+		void ApplyRotationDelta(
+			const int32 BodyIndex,
+			const FVec3& DR);
+
+		void ApplyRotationDelta(
+			const FVec3& DR0,
+			const FVec3& DR1);
+
+		void ApplyDelta(
+			const int32 BodyIndex,
+			const FVec3& DP,
+			const FVec3& DR1);
+
+		void ApplyVelocityDelta(
+			const int32 BodyIndex,
+			const FVec3& DV,
+			const FVec3& DW);
+
+		void ApplyVelocityDelta(
+			const FVec3& DV0,
+			const FVec3& DW0,
+			const FVec3& DV1,
+			const FVec3& DW1);
+
+		void ApplyAngularVelocityDelta(
+			const FVec3& DW0,
+			const FVec3& DW1);
+
+		void ApplyPositionConstraint(
+			const FReal Stiffness,
+			const FVec3& Axis,
+			const FReal Delta,
+			const FReal IM,
+			const FVec3& Arm0,
+			const FVec3& Arm1,
+			const FVec3& DRAxis0,
+			const FVec3& DRAxis1,
+			const int32 LinearHardLambdaIndex = -1);
+
+		void ApplyPositionConstraintSoft(
+			const FReal Dt,
+			const FReal JointStiffness,
+			const FReal JointDamping,
+			const bool bAccelerationMode,
+			const FVec3& Axis,
+			const FReal Delta,
+			const FReal TargetVel,
+			FReal& Lambda);
+
+		void ApplyRotationConstraint(
+			const FReal JointStiffness,
+			const FVec3& Axis,
+			const FReal Angle,
+			const int32 AngularHardLambdaIndex = -1);
+
+		void ApplyRotationConstraintKD(
+			const int32 KIndex,
+			const int32 DIndex,
+			const FReal JointStiffness,
+			const FVec3& Axis,
+			const FReal Angle,
+			const int32 AngularHardLambdaIndex = -1);
+
+		void ApplyRotationConstraintDD(
+			const FReal JointStiffness,
+			const FVec3& Axis,
+			const FReal Angle,
+			const int32 AngularHardLambdaIndex = -1);
+
+		void ApplyRotationConstraintSoft(
+			const FReal Dt,
+			const FReal JointStiffness,
+			const FReal JointDamping,
+			const bool bAccelerationMode,
+			const FVec3& Axis,
+			const FReal Angle,
+			const FReal AngVelTarget,
+			FReal& Lambda);
+
+		void ApplyRotationConstraintSoftKD(
+			const int32 KIndex,
+			const int32 DIndex,
+			const FReal Dt,
+			const FReal JointStiffness,
+			const FReal JointDamping,
+			const bool bAccelerationMode,
+			const FVec3& Axis,
+			const FReal Angle,
+			const FReal AngVelTarget,
+			FReal& Lambda);
+
+		void ApplyRotationConstraintSoftDD(
+			const FReal Dt,
+			const FReal JointStiffness,
+			const FReal JointDamping,
+			const bool bAccelerationMode,
+			const FVec3& Axis,
+			const FReal Angle,
+			const FReal AngVelTarget,
+			FReal& Lambda);
+
+		void ApplyLockedRotationConstraints(
+			const FReal Dt,
+			const FPBDJointSolverSettings& SolverSettings,
+			const FPBDJointSettings& JointSettings,
+			const bool bApplyTwist,
+			const bool bApplySwing);
+
+		void ApplyTwistConstraint(
+			const FReal Dt,
+			const FPBDJointSolverSettings& SolverSettings,
+			const FPBDJointSettings& JointSettings,
+			const bool bUseSoftLimit);
+
+		void ApplyConeConstraint(
+			const FReal Dt,
+			const FPBDJointSolverSettings& SolverSettings,
+			const FPBDJointSettings& JointSettings,
+			const bool bUseSoftLimit);
+
+		// One Swing axis is free, and the other locked. This applies the lock: Body1 Twist axis is confined to a plane.
+		void ApplySingleLockedSwingConstraint(
+			const FReal Dt,
+			const FPBDJointSolverSettings& SolverSettings,
+			const FPBDJointSettings& JointSettings,
+			const EJointAngularConstraintIndex SwingConstraintIndex,
+			const bool bUseSoftLimit);
+
+		// One Swing axis is free, and the other limited. This applies the limit: Body1 Twist axis is confined to space between two cones.
+		void ApplyDualConeSwingConstraint(
+			const FReal Dt,
+			const FPBDJointSolverSettings& SolverSettings,
+			const FPBDJointSettings& JointSettings,
+			const EJointAngularConstraintIndex SwingConstraintIndex,
+			const bool bUseSoftLimit);
+
+		// One swing axis is locked, the other limited or locked. This applies the Limited axis (ApplyDualConeSwingConstraint is used for the locked axis).
+		void ApplySwingConstraint(
+			const FReal Dt,
+			const FPBDJointSolverSettings& SolverSettings,
+			const FPBDJointSettings& JointSettings,
+			const EJointAngularConstraintIndex SwingConstraintIndex,
+			const bool bUseSoftLimit);
+
+		void ApplySwingTwistDrives(
+			const FReal Dt,
+			const FPBDJointSolverSettings& SolverSettings,
+			const FPBDJointSettings& JointSettings,
+			const bool bTwistDriveEnabled,
+			const bool bSwing1DriveEnabled,
+			const bool bSwing2DriveEnabled);
+
+		void ApplySLerpDrive(
+			const FReal Dt,
+			const FPBDJointSolverSettings& SolverSettings,
+			const FPBDJointSettings& JointSettings);
+
+		void ApplyPointPositionConstraintKD(
+			const int32 KIndex,
+			const int32 DIndex,
+			const FReal Dt,
+			const FPBDJointSolverSettings& SolverSettings,
+			const FPBDJointSettings& JointSettings);
+
+		void ApplyPointPositionConstraintDD(
+			const FReal Dt,
+			const FPBDJointSolverSettings& SolverSettings,
+			const FPBDJointSettings& JointSettings);
+
+		void ApplySphericalPositionConstraint(
+			const FReal Dt,
+			const FPBDJointSolverSettings& SolverSettings,
+			const FPBDJointSettings& JointSettings);
+
+		void ApplyCylindricalPositionConstraint(
+			const FReal Dt,
+			const int32 AxisIndex,
+			const EJointMotionType AxialMotion,
+			const EJointMotionType RadialMotion,
+			const FPBDJointSolverSettings& SolverSettings,
+			const FPBDJointSettings& JointSettings);
+
+		void ApplyPlanarPositionConstraint(
+			const FReal Dt,
+			const int32 AxisIndex,
+			const EJointMotionType AxialMotion,
+			const FPBDJointSolverSettings& SolverSettings,
+			const FPBDJointSettings& JointSettings);
+
+		void ApplyPositionDrive(
+			const FReal Dt,
+			const int32 AxisIndex,
+			const FPBDJointSolverSettings& SolverSettings,
+			const FPBDJointSettings& JointSettings,
+			const FVec3& Axis,
+			const FReal DeltaPos,
+			const FReal DeltaVel);
+
+
+		void ApplyPointProjection(
+			const FReal Dt,
+			const FPBDJointSolverSettings& SolverSettings,
+			const FPBDJointSettings& JointSettings,
+			const FReal Alpha,
+			FVec3& NetDP1,
+			FVec3& NetDR1);
+
+		void ApplySphereProjection(
+			const FReal Dt,
+			const FPBDJointSolverSettings& SolverSettings,
+			const FPBDJointSettings& JointSettings,
+			const FReal Alpha,
+			FVec3& NetDP1,
+			FVec3& NetDR1);
+
+		void ApplyTranslateProjection(
+			const FReal Dt,
+			const FPBDJointSolverSettings& SolverSettings,
+			const FPBDJointSettings& JointSettings,
+			const FReal Alpha,
+			FVec3& NetDP1,
+			FVec3& NetDR1);
+
+		void ApplyConeProjection(
+			const FReal Dt,
+			const FPBDJointSolverSettings& SolverSettings,
+			const FPBDJointSettings& JointSettings,
+			const FReal Alpha,
+			const bool bPositionLocked,
+			FVec3& NetDP1,
+			FVec3& NetDR1);
+
+		void ApplySwingProjection(
+			const FReal Dt,
+			const FPBDJointSolverSettings& SolverSettings,
+			const FPBDJointSettings& JointSettings,
+			const EJointAngularConstraintIndex SwingConstraintIndex,
+			const FReal Alpha,
+			const bool bPositionLocked,
+			FVec3& NetDP1,
+			FVec3& NetDR1);
+
+		void ApplySingleLockedSwingProjection(
+			const FReal Dt,
+			const FPBDJointSolverSettings& SolverSettings,
+			const FPBDJointSettings& JointSettings,
+			const EJointAngularConstraintIndex SwingConstraintIndex,
+			const FReal Alpha,
+			const bool bPositionLocked,
+			FVec3& NetDP1,
+			FVec3& NetDR1);
+
+		void ApplyDoubleLockedSwingProjection(
+			const FReal Dt,
+			const FPBDJointSolverSettings& SolverSettings,
+			const FPBDJointSettings& JointSettings,
+			const FReal Alpha,
+			const bool bPositionLocked,
+			FVec3& NetDP1,
+			FVec3& NetDR1);
+
+		void ApplyDualConeSwingProjection(
+			const FReal Dt,
+			const FPBDJointSolverSettings& SolverSettings,
+			const FPBDJointSettings& JointSettings,
+			const EJointAngularConstraintIndex SwingConstraintIndex,
+			const FReal Alpha,
+			const bool bPositionLocked,
+			FVec3& NetDP1,
+			FVec3& NetDR1);
+
+		void ApplyTwistProjection(
+			const FReal Dt,
+			const FPBDJointSolverSettings& SolverSettings,
+			const FPBDJointSettings& JointSettings,
+			const FReal Alpha,
+			const bool bPositionLocked,
+			FVec3& NetDP1,
+			FVec3& NetDR1);
+
+		void ApplyVelocityProjection(
+			const FReal Dt,
+			const FPBDJointSolverSettings& SolverSettings,
+			const FPBDJointSettings& JointSettings,
+			const FReal Alpha,
+			const FVec3& DP1,
+			const FVec3& DR1);
+
+		void ApplyLinearVelocityConstraint(
+			const FReal Stiffness,
+			const FVec3& Axis,
+			const FReal IM,
+			const FVec3& Arm0,
+			const FVec3& Arm1,
+			const FVec3& DRAxis0,
+			const FVec3& DRAxis1,
+			const FReal TargetVel = 0);
+
+		void ApplyPointVelocityConstraint(
+			const FReal Dt,
+			const FPBDJointSolverSettings& SolverSettings,
+			const FPBDJointSettings& JointSettings);
+
+		void ApplyPointVelocityConstraint(
+		    const FReal Stiffness,
+		    const FMatrix33& ConstraintIM,
+		    const FVec3& ConstraintArm0,
+		    const FVec3& ConstraintArm1);
+
+		void ApplySphericalVelocityConstraint(
+			const FReal Dt,
+			const FPBDJointSolverSettings& SolverSettings,
+			const FPBDJointSettings& JointSettings);
+
+		void ApplyCylindricalVelocityConstraint(
+			const FReal Dt,
+			const int32 AxisIndex,
+			const EJointMotionType AxialMotion,
+			const EJointMotionType RadialMotion,
+			const FPBDJointSolverSettings& SolverSettings,
+			const FPBDJointSettings& JointSettings);
+
+		void ApplyPlanarVelocityConstraint(
+			const FReal Dt,
+			const int32 AxisIndex,
+			const EJointMotionType AxialMotion,
+			const FPBDJointSolverSettings& SolverSettings,
+			const FPBDJointSettings& JointSettings);
+
+		void ApplyTwistVelocityConstraint(
+			const FReal Dt,
+			const FPBDJointSolverSettings& SolverSettings,
+			const FPBDJointSettings& JointSettings,
+			const bool bUseSoftLimit);
+
+		void ApplyConeVelocityConstraint(
+			const FReal Dt,
+			const FPBDJointSolverSettings& SolverSettings,
+			const FPBDJointSettings& JointSettings,
+			const bool bUseSoftLimit);
+
+		// One Swing axis is free, and the other locked. This applies the lock: Body1 Twist axis is confined to a plane.
+		void ApplySingleLockedSwingVelocityConstraint(
+			const FReal Dt,
+			const FPBDJointSolverSettings& SolverSettings,
+			const FPBDJointSettings& JointSettings,
+			const EJointAngularConstraintIndex SwingConstraintIndex,
+			const bool bUseSoftLimit);
+
+		// One Swing axis is free, and the other limited. This applies the limit: Body1 Twist axis is confined to space between two cones.
+		void ApplyDualConeSwingVelocityConstraint(
+			const FReal Dt,
+			const FPBDJointSolverSettings& SolverSettings,
+			const FPBDJointSettings& JointSettings,
+			const EJointAngularConstraintIndex SwingConstraintIndex,
+			const bool bUseSoftLimit);
+
+		// One swing axis is locked, the other limited or locked. This applies the Limited axis (ApplyDualConeSwingConstraint is used for the locked axis).
+		void ApplySwingVelocityConstraint(
+			const FReal Dt,
+			const FPBDJointSolverSettings& SolverSettings,
+			const FPBDJointSettings& JointSettings,
+			const EJointAngularConstraintIndex SwingConstraintIndex,
+			const bool bUseSoftLimit);
+
+		void ApplyAngularVelocityConstraint(
+			const FReal Stiffness,
+			const FVec3& Axis,
+			const FReal TargetVel = 0.0f);
+
+		void ApplyLockedRotationVelocityConstraints(
+			const FReal Dt,
+			const FPBDJointSolverSettings& SolverSettings,
+			const FPBDJointSettings& JointSettings,
+			const bool bApplyTwist,
+			const bool bApplySwing);
+
+		void CalculateLinearConstraintPadding(
+			const FReal Dt,
+			const FPBDJointSolverSettings& SolverSettings,
+			const FPBDJointSettings& JointSettings,
+			const FReal Restitution,
+			const int32 AxisIndex,
+			const FVec3 Axis,
+			FReal& InOutPos);
+
+		void CalculateAngularConstraintPadding(
+			const FReal Dt,
+			const FPBDJointSolverSettings& SolverSettings,
+			const FPBDJointSettings& JointSettings,
+			const FReal Restitution,
+			const EJointAngularConstraintIndex ConstraintIndex,
+			const FVec3 Axis,
+			FReal& InOutAngle);
+
+		inline bool HasLinearConstraintPadding(const int32 AxisIndex) const
+		{
+			return LinearConstraintPadding[AxisIndex] >= 0.0f;
+		}
+
+		inline FReal GetLinearConstraintPadding(const int32 AxisIndex) const
+		{
+			return FMath::Max(LinearConstraintPadding[AxisIndex], 0.0f);
+		}
+
+		inline void SetLinearConstraintPadding(const int32 AxisIndex, FReal Padding)
+		{
+			LinearConstraintPadding[AxisIndex] = Padding;
+		}
+
+
+		inline bool HasAngularConstraintPadding(const EJointAngularConstraintIndex ConstraintIndex) const
+		{
+			return AngularConstraintPadding[(int32)ConstraintIndex] >= 0.0f;
+		}
+
+		inline FReal GetAngularConstraintPadding(const EJointAngularConstraintIndex ConstraintIndex) const
+		{
+			return FMath::Max(AngularConstraintPadding[(int32)ConstraintIndex], 0.0f);
+		}
+
+		inline void SetAngularConstraintPadding(const EJointAngularConstraintIndex ConstraintIndex, FReal Padding)
+		{
+			AngularConstraintPadding[(int32)ConstraintIndex] = Padding;
+		}
+
+		// Calculate linear velocities along constraint axes based on the constraint type in JointSettings. CV is the relative linear velocity of joint connectors.
+		void CalculateConstraintAxisLinearVelocities(
+			const FPBDJointSettings& JointSettings,
+			FVec3& ConstraintAxisLinearVelocities) const;
+
+		void CalculateSphericalConstraintAxisLinearVelocities(
+			const FPBDJointSettings& JointSettings,
+			FVec3& ConstraintAxisLinearVelocities
+			) const;
+
+		void CalculateCylindricalConstraintAxisLinearVelocities(
+			const int32 AxisIndex,
+			const EJointMotionType AxialMotion,
+			const EJointMotionType RadialMotion,
+			const FPBDJointSettings& JointSettings,
+			FVec3& ConstraintAxisLinearVelocities) const;
+
+		void CalculatePlanarConstraintAxisLinearVelocities(
+			const int32 AxisIndex,
+			const EJointMotionType AxialMotion,
+			const FPBDJointSettings& JointSettings,
+			FVec3& ConstraintAxisLinearVelocities) const;
+
+		// Calculate angular velocities along constraint axes based on the constraint type in JointSettings. CW is the relative angular velocity of joint connectors.
+		void CalculateConstraintAxisAngularVelocities(
+			const FPBDJointSolverSettings& SolverSettings,
+			const FPBDJointSettings& JointSettings,
+			FVec3& ConstraintAxisAngularVelocities
+			) const;
+
+		void CalculateTwistConstraintAxisAngularVelocities(
+			const FPBDJointSettings& JointSettings,
+			FVec3& ConstraintAxisAngularVelocities) const;
+
+		void CalculateConeConstraintAxisAngularVelocities(
+			const FPBDJointSettings& JointSettings,
+			FVec3& ConstraintAxisAngularVelocities) const;
+
+		void CalculateDualConeSwingConstraintAxisAngularVelocities(
+			const FPBDJointSettings& JointSettings,
+			const EJointAngularConstraintIndex SwingConstraintIndex,
+			FVec3& ConstraintAxisAngularVelocities) const;
+
+		void CalculateSwingConstraintAxisAngularVelocities(
+			const FPBDJointSolverSettings& SolverSettings,
+			const FPBDJointSettings& JointSettings,
+			const EJointAngularConstraintIndex SwingConstraintIndex,
+			FVec3& ConstraintAxisAngularVelocities) const;
+
+		// The cached body state on which the joint operates
+		FConstraintSolverBody SolverBodies[MaxConstrainedBodies];
+
+		// Local-space constraint settings
+		FRigidTransform3 LocalConnectorXs[MaxConstrainedBodies];	// Local(CoM)-space joint connector transforms
+
+		// World-space constraint settings
+		FVec3 ConnectorXs[MaxConstrainedBodies];			// World-space joint connector positions
+		FRotation3 ConnectorRs[MaxConstrainedBodies];		// World-space joint connector rotations
+
+		// XPBD Initial iteration world-space body state
+		FVec3 InitConnectorXs[MaxConstrainedBodies];		// World-space joint connector positions
+		FRotation3 InitConnectorRs[MaxConstrainedBodies];	// World-space joint connector rotations
+
+		// Conditioned InvM and InvI
+		FReal ConditionedInvMs[MaxConstrainedBodies];
+		FMatrix33 ConditionedInvIs[MaxConstrainedBodies];
+
+		// Accumulated Impulse and AngularImpulse (Impulse * Dt since they are mass multiplied position corrections)
+		FVec3 NetLinearImpulse;
+		FVec3 NetAngularImpulse;
+
+		// Lagrange multipliers of the position constraints.
+		// Currently these are only used in ApplyCylindricalVelocityConstraints
+		FVec3 LinearHardLambda;
+		// Lagrange multipliers of the rotation constraints
+		// Currently these are used in ApplyAngularVelocityConstraints
+		FVec3 AngularHardLambda;
+
+		// XPBD Accumulators (net impulse for each soft constraint/drive)
+		FReal LinearSoftLambda;
+		FReal TwistSoftLambda;
+		FReal SwingSoftLambda;
+		FVec3 LinearDriveLambdas;
+		FVec3 RotationDriveLambdas;
+
+		// Solver stiffness - increased over iterations for stability
+		// \todo(chaos): remove Stiffness from SolverSettings (because it is not a solver constant)
+		FReal SolverStiffness;
+
+		// Constraint padding which can act something like a velocity constraint (for reslockfreetitution)
+		FVec3 LinearConstraintPadding;
+		FVec3 AngularConstraintPadding;
+
+		// Tolerances below which we stop solving
+		FReal PositionTolerance;					// Distance error below which we consider a constraint or drive solved
+		FReal AngleTolerance;						// Angle error below which we consider a constraint or drive solved
+
+		// Tracking whether the solver is resolved
+		FVec3 LastPs[MaxConstrainedBodies];			// Positions at the beginning of the iteration
+		FRotation3 LastQs[MaxConstrainedBodies];	// Rotations at the beginning of the iteration
+		FVec3 InitConstraintAxisLinearVelocities; // Linear velocities along the constraint axes at the begining of the frame, used by restitution
+		FVec3 InitConstraintAxisAngularVelocities; // Angular velocities along the constraint axes at the begining of the frame, used by restitution
+		int32 NumActiveConstraints;					// The number of active constraints and drives in the last iteration (-1 initial value)
+		bool bIsActive;								// Whether any constraints actually moved any bodies in last iteration
+
+		// Caching data
+		FReal LinearStiffness, LinearDamping;
+		FReal AngularStiffness, AngularDamping;
+		FVec3 ConstraintArms[3][2];
+		FMatrix33 PointPositionConstraintIM;
+		FReal PositionConstraintIM[2];
+		FVec3 PositionConstraintAxis[2];
+		FVec3 PositionConstraintDRAxis[2][2]; // Cache for InvI * cross(Arm, Axis)
+		bool bLinearAccelerationMode;
+
+		// dummy indices
+		static constexpr int32 PointPositionConstraintIndex = 0;
+		static constexpr int32 SphericalPositionConstraintIndex = 0;
+		static constexpr int32 PlanarPositionConstraintIndex = 0;
+	};
+
+}
\ No newline at end of file
diff --git a/Engine/Source/Runtime/Experimental/Chaos/Public/Chaos/PBDJointConstraints.h b/Engine/Source/Runtime/Experimental/Chaos/Public/Chaos/PBDJointConstraints.h
--- a/Engine/Source/Runtime/Experimental/Chaos/Public/Chaos/PBDJointConstraints.h
+++ b/Engine/Source/Runtime/Experimental/Chaos/Public/Chaos/PBDJointConstraints.h
@@ -13,6 +13,7 @@
 #include "Chaos/Evolution/SolverDatas.h"
 #include "Chaos/Joint/JointSolverConstraints.h"
 #include "Chaos/Joint/PBDJointSolverGaussSeidel.h"
+#include "Chaos/Joint/PBDJointCachedSolverGaussSeidel.h"
 #include "Chaos/ParticleHandleFwd.h"
 #include "Chaos/PBDConstraintContainer.h"
 #include "Chaos/PBDJointConstraintTypes.h"
@@ -302,6 +303,8 @@ namespace Chaos
 		void ApplyBreakThreshold(const FReal Dt, int32 ConstraintIndex, const FVec3& LinearImpulse, const FVec3& AngularImpulse);
 		void ApplyPlasticityLimits(const int32 ConstraintIndex);
 
+		void UpdateConstraintProjection(const int32 It, const int32 NumIts, const FPBDIslandSolverData& SolverData);
+
 		FPBDJointSolverSettings Settings;
 
 		TArray<FPBDJointSettings> ConstraintSettings;
@@ -316,6 +319,7 @@ namespace Chaos
 
 		// @todo(ccaulfield): optimize storage for joint solver
 		TArray<FPBDJointSolver> ConstraintSolvers;
+		TArray<FPBDJointCachedSolver> CachedConstraintSolvers;
 
 		EConstraintSolverType SolverType;
 	};
