Chaos - disable manifold GJK point injection until we can make it better
- disable manifold restore for spheres and capsules

diff --git a/Engine/Source/Runtime/Experimental/Chaos/Private/Chaos/Collision/PBDCollisionConstraint.cpp b/Engine/Source/Runtime/Experimental/Chaos/Private/Chaos/Collision/PBDCollisionConstraint.cpp
--- a/Engine/Source/Runtime/Experimental/Chaos/Private/Chaos/Collision/PBDCollisionConstraint.cpp
+++ b/Engine/Source/Runtime/Experimental/Chaos/Private/Chaos/Collision/PBDCollisionConstraint.cpp
@@ -18,6 +18,8 @@
 
 namespace Chaos
 {
+	extern bool bChaos_Collision_EnableManifoldInject;
+
 	FRealSingle Chaos_Manifold_MatchPositionTolerance = 0.3f;		// Fraction of object size position tolerance
 	FRealSingle Chaos_Manifold_MatchNormalTolerance = 0.02f;		// Dot product tolerance
 	FAutoConsoleVariableRef CVarChaos_Manifold_MatchPositionTolerance(TEXT("p.Chaos.Collision.Manifold.MatchPositionTolerance"), Chaos_Manifold_MatchPositionTolerance, TEXT("A tolerance as a fraction of object size used to determine if two contact points are the same"));
@@ -599,13 +601,14 @@ namespace Chaos
 					ManifoldPoint.ContactPoint.ShapeContactPoints[1] = ShapeContactPoint1;
 					ManifoldPoint.ContactPoint.Phi = ContactPhi;
 				}
-				else if (ManifoldPointToRemove == INDEX_NONE)
+				else if ((ManifoldPointToRemove == INDEX_NONE) && bChaos_Collision_EnableManifoldInject)
 				{
+					// We can reject up to 1 point (if we have GJK point injection enabled)
 					ManifoldPointToRemove = ManifoldPointIndex;
 				}
 				else
 				{
-					// We want to remove a second point, but we will never reuse it now so throw away the whole manifold
+					// We want to remove a(nother) point, but we will never reuse the manifold now so throw it away
 					ResetActiveManifoldContacts();
 					return false;
 				}
diff --git a/Engine/Source/Runtime/Experimental/Chaos/Private/Chaos/Collision/ParticlePairMidPhase.cpp b/Engine/Source/Runtime/Experimental/Chaos/Private/Chaos/Collision/ParticlePairMidPhase.cpp
--- a/Engine/Source/Runtime/Experimental/Chaos/Private/Chaos/Collision/ParticlePairMidPhase.cpp
+++ b/Engine/Source/Runtime/Experimental/Chaos/Private/Chaos/Collision/ParticlePairMidPhase.cpp
@@ -20,18 +20,18 @@ namespace Chaos
 		bool bChaos_Collision_MidPhase_EnableBoundsChecks = true;
 		FAutoConsoleVariableRef CVarChaos_Collision_EnableBoundsChecks(TEXT("p.Chaos.Collision.EnableBoundsChecks"), bChaos_Collision_MidPhase_EnableBoundsChecks, TEXT(""));
 
-		bool bChaos_Collision_EnableManifoldRestore = true;
 		Chaos::FRealSingle Chaos_Collision_RestoreTolerance_NoContact_Position = 0.005f;	// About 0.5cm for a meter cube
 		Chaos::FRealSingle Chaos_Collision_RestoreTolerance_NoContact_Rotation = 0.1f;		// About 10deg
 		Chaos::FRealSingle Chaos_Collision_RestoreTolerance_Contact_Position = 0.02f;		// About 2cm for a meter cube
 		Chaos::FRealSingle Chaos_Collision_RestoreTolerance_Contact_Rotation = 0.1f;		// About 10deg
-		FAutoConsoleVariableRef CVarChaos_Collision_EnableManifoldRestore(TEXT("p.Chaos.Collision.EnableManifoldRestore"), bChaos_Collision_EnableManifoldRestore, TEXT(""));
 		FAutoConsoleVariableRef CVarChaos_Collision_RestoreTolerance_NoContact_Position(TEXT("p.Chaos.Collision.RestoreTolerance.NoContact.Position"), Chaos_Collision_RestoreTolerance_NoContact_Position, TEXT("Fraction of Size. Particle pairs that move less than this may have their contacts reinstated"));
 		FAutoConsoleVariableRef CVarChaos_Collision_RestoreTolerance_NoContact_Rotation(TEXT("p.Chaos.Collision.RestoreTolerance.NoContact.Rotation"), Chaos_Collision_RestoreTolerance_NoContact_Rotation, TEXT("Quaternion Dot Product Limit. Particle pairs that move less than this may have their contacts reinstated"));
 		FAutoConsoleVariableRef CVarChaos_Collision_RestoreTolerance_Contact_Position(TEXT("p.Chaos.Collision.RestoreTolerance.WithContact.Position"), Chaos_Collision_RestoreTolerance_Contact_Position, TEXT("Fraction of Size. Particle pairs that move less than this may have their contacts reinstated"));
 		FAutoConsoleVariableRef CVarChaos_Collision_RestoreTolerance_Contact_Rotation(TEXT("p.Chaos.Collision.RestoreTolerance.WithContact.Rotation"), Chaos_Collision_RestoreTolerance_Contact_Rotation, TEXT("Quaternion Dot Product Limit. Particle pairs that move less than this may have their contacts reinstated"));
 
+		bool bChaos_Collision_EnableManifoldRestore = true;
 		bool bChaos_Collision_EnableManifoldUpdate = true;
+		FAutoConsoleVariableRef CVarChaos_Collision_EnableManifoldRestore(TEXT("p.Chaos.Collision.EnableManifoldRestore"), bChaos_Collision_EnableManifoldRestore, TEXT(""));
 		FAutoConsoleVariableRef CVarChaos_Collision_EnableManifoldUpdate(TEXT("p.Chaos.Collision.EnableManifoldUpdate"), bChaos_Collision_EnableManifoldUpdate, TEXT(""));
 
 		Chaos::FRealSingle Chaos_Collision_CullDistanceScaleInverseSize = 0.01f;	// 100cm
@@ -128,18 +128,21 @@ namespace Chaos
 		const EImplicitObjectType ImplicitType1 = (Implicit1 != nullptr) ? GetInnerType(Implicit1->GetCollisionType()) : ImplicitObjectType::Unknown;
 		const bool bIsSphere0 = (ImplicitType0 == ImplicitObjectType::Sphere);
 		const bool bIsSphere1 = (ImplicitType1 == ImplicitObjectType::Sphere);
+		const bool bIsCapsule0 = (ImplicitType0 == ImplicitObjectType::Capsule);
+		const bool bIsCapsule1 = (ImplicitType1 == ImplicitObjectType::Capsule);
 
 		const bool bAllowBoundsChecked = bChaos_Collision_MidPhase_EnableBoundsChecks && bHasBounds0 && bHasBounds1;
-		Flags.bEnableAABBCheck = bAllowBoundsChecked && (!bIsSphere0 || !bIsSphere1);
-		Flags.bEnableOBBCheck0 = bAllowBoundsChecked && !bIsSphere0;
-		Flags.bEnableOBBCheck1 = bAllowBoundsChecked && !bIsSphere1;
+		Flags.bEnableAABBCheck = bAllowBoundsChecked && !(bIsSphere0 && bIsSphere1);	// No AABB test if both are spheres
+		Flags.bEnableOBBCheck0 = bAllowBoundsChecked && !bIsSphere0;					// No OBB test for spheres
+		Flags.bEnableOBBCheck1 = bAllowBoundsChecked && !bIsSphere1;					// No OBB test for spheres
 		
 		if (bAllowBoundsChecked && bIsSphere0 && bIsSphere1)
 		{
-			SphereBoundsCheckSize = Shape0->GetLeafGeometry()->GetMargin() + Shape1->GetLeafGeometry()->GetMargin();
+			SphereBoundsCheckSize = Implicit0->GetMargin() + Implicit1->GetMargin();	// Sphere-Sphere bounds test
 		}
 
-		Flags.bEnableManifoldCheck = bChaos_Collision_EnableManifoldUpdate && !bIsSphere0 && !bIsSphere1;
+		// Do not try to reuse manifold points for capsules or spheres (against anything)
+		Flags.bEnableManifoldCheck = bChaos_Collision_EnableManifoldUpdate && !bIsSphere0 && !bIsSphere1 && !bIsCapsule0 && !bIsCapsule1;
 	}
 
 	FSingleShapePairCollisionDetector::~FSingleShapePairCollisionDetector()
@@ -796,6 +799,14 @@ namespace Chaos
 					TryAddShapePair(Shape0, Shape1);
 				}
 			}
+
+			// If we have only 1 shape pair, we may disable Manifold Restore functionality if it is spheres or capsules
+			// @todo(chaos): make manifold restoration work better with rolling shapes (spheres and capsules)
+			const bool bDisableManifoldRestore = (ShapePairDetectors.Num() == 1) && !ShapePairDetectors[0].EnableManifoldRestore();
+			if (bDisableManifoldRestore)
+			{
+				Flags.bRestorable = false;
+			}
 		}
 	}
 
diff --git a/Engine/Source/Runtime/Experimental/Chaos/Private/Chaos/CollisionOneShotManifolds.cpp b/Engine/Source/Runtime/Experimental/Chaos/Private/Chaos/CollisionOneShotManifolds.cpp
--- a/Engine/Source/Runtime/Experimental/Chaos/Private/Chaos/CollisionOneShotManifolds.cpp
+++ b/Engine/Source/Runtime/Experimental/Chaos/Private/Chaos/CollisionOneShotManifolds.cpp
@@ -40,7 +40,7 @@ namespace Chaos
 	FAutoConsoleVariableRef CVarForceOneShotManifoldEdgeEdgeCaseZeroCullDistance(TEXT("p.Chaos.Collision.Manifold.ForceOneShotManifoldEdgeEdgeCaseZeroCullDistance"), ForceOneShotManifoldEdgeEdgeCaseZeroCullDistance,
 	TEXT("If enabled, if one shot manifold hits edge/edge case, we will force a cull distance of zero. That means edge/edge contacts will be thrown out if separated at all. Only applies to Convex/Convex oneshot impl."));
 
-	bool bChaos_Collision_EnableManifoldInject = true;
+	bool bChaos_Collision_EnableManifoldInject = false;
 	FAutoConsoleVariableRef CVarChaos_Collision_EnableManifoldInject(TEXT("p.Chaos.Collision.EnableManifolInject"), bChaos_Collision_EnableManifoldInject, TEXT(""));
 
 	
diff --git a/Engine/Source/Runtime/Experimental/Chaos/Public/Chaos/Collision/ParticlePairMidPhase.h b/Engine/Source/Runtime/Experimental/Chaos/Public/Chaos/Collision/ParticlePairMidPhase.h
--- a/Engine/Source/Runtime/Experimental/Chaos/Public/Chaos/Collision/ParticlePairMidPhase.h
+++ b/Engine/Source/Runtime/Experimental/Chaos/Public/Chaos/Collision/ParticlePairMidPhase.h
@@ -86,6 +86,15 @@ namespace Chaos
 		*/
 		void SetCollision(const FPBDCollisionConstraint& Constraint);
 
+		/**
+		 * @brief Whether we want to use manifold restore for this shape type
+		 * @see FParticlePairMidPhase::ShouldRestoreConstraints
+		*/
+		bool EnableManifoldRestore() const
+		{
+			return Flags.bEnableManifoldCheck;
+		}
+
 	private:
 		int32 GenerateCollisionImpl(const FReal CullDistance, const bool bUseCCD, const FReal Dt);
 
