Remove particle collision constraints after changing NonFrequentData

diff --git a/Engine/Source/Runtime/Experimental/Chaos/Private/Chaos/PBDRigidsEvolutionGBF.cpp b/Engine/Source/Runtime/Experimental/Chaos/Private/Chaos/PBDRigidsEvolutionGBF.cpp
--- a/Engine/Source/Runtime/Experimental/Chaos/Private/Chaos/PBDRigidsEvolutionGBF.cpp
+++ b/Engine/Source/Runtime/Experimental/Chaos/Private/Chaos/PBDRigidsEvolutionGBF.cpp
@@ -871,5 +871,13 @@ void FPBDRigidsEvolutionGBF::TransferJointConstraintCollisions()
 #endif
 }
 
+void FPBDRigidsEvolutionGBF::DestroyParticleCollisionsInAllocator(FGeometryParticleHandle* Particle)
+{
+	if (Particle != nullptr)
+	{
+		CollisionConstraints.GetConstraintAllocator().RemoveParticle(Particle);
+	}
+}
+
 }
 
diff --git a/Engine/Source/Runtime/Experimental/Chaos/Private/PhysicsProxy/SingleParticlePhysicsProxy.cpp b/Engine/Source/Runtime/Experimental/Chaos/Private/PhysicsProxy/SingleParticlePhysicsProxy.cpp
--- a/Engine/Source/Runtime/Experimental/Chaos/Private/PhysicsProxy/SingleParticlePhysicsProxy.cpp
+++ b/Engine/Source/Runtime/Experimental/Chaos/Private/PhysicsProxy/SingleParticlePhysicsProxy.cpp
@@ -70,7 +70,9 @@ void PushToPhysicsStateImp(const Chaos::FDirtyPropertiesManager& Manager, Chaos:
 			Handle->SetNonFrequentData(*NewNonFrequentData);
 
 			// geometry may have changed, we need to invalidate the particle so it can be removed from caching structures in the evolution
+			// @todo(chaos): Remove collision constraints only. Invalidate particle may remove joint constraints in constraint graph. If joint constraint is persistent, this may cause issues.
 			Evolution.InvalidateParticle(Handle);
+			Evolution.DestroyParticleCollisionsInAllocator(Handle);
 		}
 
 		auto NewVelocities = bHasKinematicData ? ParticleData.FindVelocities(Manager, DataIdx) : nullptr;
diff --git a/Engine/Source/Runtime/Experimental/Chaos/Public/Chaos/PBDRigidsEvolutionGBF.h b/Engine/Source/Runtime/Experimental/Chaos/Public/Chaos/PBDRigidsEvolutionGBF.h
--- a/Engine/Source/Runtime/Experimental/Chaos/Public/Chaos/PBDRigidsEvolutionGBF.h
+++ b/Engine/Source/Runtime/Experimental/Chaos/Public/Chaos/PBDRigidsEvolutionGBF.h
@@ -175,6 +175,8 @@ namespace Chaos
 
 		void DestroyConstraint(FConstraintHandle* Constraint);
 
+		void DestroyParticleCollisionsInAllocator(FGeometryParticleHandle* Particle);
+
 		CHAOS_API inline void EndFrame(FReal Dt)
 		{
 			Particles.GetNonDisabledDynamicView().ParallelFor([&](auto& Particle, int32 Index) {
