Chaos - LWC - change Convex implicits storage to use float instead of doubles
- also changed the convex generation to use float as it was before LWC
- updated all unit test accordingly

diff --git a/Engine/Plugins/Experimental/ChaosCloth/Source/Chaos/Private/ChaosCloth/ChaosClothingSimulation.cpp b/Engine/Plugins/Experimental/ChaosCloth/Source/Chaos/Private/ChaosCloth/ChaosClothingSimulation.cpp
--- a/Engine/Plugins/Experimental/ChaosCloth/Source/Chaos/Private/ChaosCloth/ChaosClothingSimulation.cpp
+++ b/Engine/Plugins/Experimental/ChaosCloth/Source/Chaos/Private/ChaosCloth/ChaosClothingSimulation.cpp
@@ -1182,29 +1182,29 @@ static void DrawTaperedCylinder(FPrimitiveDrawInterface* PDI, const FTaperedCyli
 
 static void DrawConvex(FPrimitiveDrawInterface* PDI, const FConvex& Convex, const FQuat& Rotation, const FVector& Position, const FLinearColor& Color)
 {
-	const TArray<TPlaneConcrete<FReal, 3>>& Planes = Convex.GetFaces();
+	const TArray<FConvex::FPlaneType>& Planes = Convex.GetFaces();
 	for (int32 PlaneIndex1 = 0; PlaneIndex1 < Planes.Num(); ++PlaneIndex1)
 	{
-		const TPlaneConcrete<FReal, 3>& Plane1 = Planes[PlaneIndex1];
+		const FConvex::FPlaneType& Plane1 = Planes[PlaneIndex1];
 
 		for (int32 PlaneIndex2 = PlaneIndex1 + 1; PlaneIndex2 < Planes.Num(); ++PlaneIndex2)
 		{
-			const TPlaneConcrete<FReal, 3>& Plane2 = Planes[PlaneIndex2];
+			const FConvex::FPlaneType& Plane2 = Planes[PlaneIndex2];
 
 			// Find the two surface points that belong to both Plane1 and Plane2
 			uint32 ParticleIndex1 = INDEX_NONE;
 
-			const TArray<FVec3>& Vertices = Convex.GetVertices();
+			const TArray<FConvex::FVec3Type>& Vertices = Convex.GetVertices();
 			for (int32 ParticleIndex = 0; ParticleIndex < Vertices.Num(); ++ParticleIndex)
 			{
-				const FVec3& X = Vertices[ParticleIndex];
+				const FConvex::FVec3Type& X = Vertices[ParticleIndex];
 
 				if (FMath::Square(Plane1.SignedDistance(X)) < KINDA_SMALL_NUMBER && 
 					FMath::Square(Plane2.SignedDistance(X)) < KINDA_SMALL_NUMBER)
 				{
 					if (ParticleIndex1 != INDEX_NONE)
 					{
-						const FVec3& X1 = Vertices[ParticleIndex1];
+						const FVector X1(Vertices[ParticleIndex1]);
 						const FVector Position1 = Position + Rotation.RotateVector(X1);
 						const FVector Position2 = Position + Rotation.RotateVector(X);
 						DrawLine(PDI, Position1, Position2, Color);
diff --git a/Engine/Plugins/Experimental/ChaosCloth/Source/Chaos/Private/ChaosCloth/ChaosClothingSimulationCollider.cpp b/Engine/Plugins/Experimental/ChaosCloth/Source/Chaos/Private/ChaosCloth/ChaosClothingSimulationCollider.cpp
--- a/Engine/Plugins/Experimental/ChaosCloth/Source/Chaos/Private/ChaosCloth/ChaosClothingSimulationCollider.cpp
+++ b/Engine/Plugins/Experimental/ChaosCloth/Source/Chaos/Private/ChaosCloth/ChaosClothingSimulationCollider.cpp
@@ -205,7 +205,7 @@ void FClothingSimulationCollider::FLODData::Add(
 			BoneIndices[Index] = GetMappedBoneIndex(UsedBoneIndices, Convex.BoneIndex);
 			UE_LOG(LogChaosCloth, VeryVerbose, TEXT("Found collision convex on bone index %d."), BoneIndices[Index]);
 
-			TArray<TPlaneConcrete<FReal, 3>> Planes;
+			TArray<FConvex::FPlaneType> Planes;
 			TArray<TArray<int32>> FaceIndices;
 
 			const int32 NumSurfacePoints = Convex.SurfacePoints.Num();
@@ -238,9 +238,9 @@ void FClothingSimulationCollider::FLODData::Add(
 						UE_LOG(LogChaosCloth, Warning, TEXT("Invalid convex collision: no face indices, this asset might be using a deprecated Apex collision shape."));
 						break;
 					}
-					const FVec3 Normal(static_cast<FVector>(NormalizedPlane));
-					const FVec3 Base = Normal * NormalizedPlane.W * InScale;
-					Planes.Add(TPlaneConcrete<FReal, 3>(Base, Normal));
+					const FConvex::FVec3Type Normal(static_cast<FVector>(NormalizedPlane));
+					const FConvex::FVec3Type Base = Normal * NormalizedPlane.W * InScale;
+					Planes.Add(FConvex::FPlaneType(Base, Normal));
 					FaceIndices.Add(Face.Indices);
 				}
 			}
@@ -250,7 +250,7 @@ void FClothingSimulationCollider::FLODData::Add(
 				check(Planes.Num() == FaceIndices.Num());
 
 				// Retrieve the convex vertices and face indices
-				TArray<FVec3> Vertices;
+				TArray<FConvex::FVec3Type> Vertices;
 				Vertices.SetNum(NumSurfacePoints);
 				for (int32 ParticleIndex = 0; ParticleIndex < NumSurfacePoints; ++ParticleIndex)
 				{
@@ -523,14 +523,14 @@ void FClothingSimulationCollider::ExtractPhysicsAssetCollision(FClothCollisionDa
 				const FConvex& ChaosConvex = ChaosConvexMesh.GetObjectChecked<FConvex>();
 
 				// Copy faces' planes and indices
-				const TArray<TPlaneConcrete<FReal, 3>>& Faces = ChaosConvex.GetFaces();
+				const TArray<FConvex::FPlaneType>& Faces = ChaosConvex.GetFaces();
 				Convex.Faces.SetNum(Faces.Num());
 				for (int32 FaceIndex = 0; FaceIndex < Faces.Num(); ++FaceIndex)
 				{
 					FClothCollisionPrim_ConvexFace& Face = Convex.Faces[FaceIndex];
 
 					// Copy face plane
-					const TPlaneConcrete<FReal, 3>& PlaneConcrete = Faces[FaceIndex];
+					const FConvex::FPlaneType& PlaneConcrete = Faces[FaceIndex];
 					Face.Plane = FPlane(PlaneConcrete.X(), PlaneConcrete.Normal());
 
 					// Copy face indices
diff --git a/Engine/Plugins/Experimental/Water/Source/Runtime/Private/OceanCollisionComponent.cpp b/Engine/Plugins/Experimental/Water/Source/Runtime/Private/OceanCollisionComponent.cpp
--- a/Engine/Plugins/Experimental/Water/Source/Runtime/Private/OceanCollisionComponent.cpp
+++ b/Engine/Plugins/Experimental/Water/Source/Runtime/Private/OceanCollisionComponent.cpp
@@ -94,7 +94,7 @@ void UOceanCollisionComponent::UpdateBodySetup(const TArray<FKConvexElem>& Conve
 	for (FKConvexElem& Elem : CachedBodySetup->AggGeom.ConvexElems)
 	{
 		int32 NumHullVerts = Elem.VertexData.Num();
-		TArray<Chaos::FVec3> ConvexVertices;
+		TArray<Chaos::FConvex::FVec3Type> ConvexVertices;
 		ConvexVertices.SetNum(NumHullVerts);
 		for (int32 VertIndex = 0; VertIndex < NumHullVerts; ++VertIndex)
 		{
diff --git a/Engine/Plugins/Experimental/Water/Source/Runtime/Private/WaterBodyRiverComponent.cpp b/Engine/Plugins/Experimental/Water/Source/Runtime/Private/WaterBodyRiverComponent.cpp
--- a/Engine/Plugins/Experimental/Water/Source/Runtime/Private/WaterBodyRiverComponent.cpp
+++ b/Engine/Plugins/Experimental/Water/Source/Runtime/Private/WaterBodyRiverComponent.cpp
@@ -258,7 +258,7 @@ void UWaterBodyRiverComponent::UpdateSplineMesh(USplineMeshComponent* MeshComp,
 			for (FKConvexElem& Elem : MeshComp->GetBodySetup()->AggGeom.ConvexElems)
 			{
 				const int32 NumHullVerts = Elem.VertexData.Num();
-				TArray<Chaos::FVec3> ConvexVertices;
+				TArray<Chaos::FConvex::FVec3Type> ConvexVertices;
 				ConvexVertices.SetNum(NumHullVerts);
 				for (int32 VertIndex = 0; VertIndex < NumHullVerts; ++VertIndex)
 				{
diff --git a/Engine/Source/Editor/UnrealEd/Private/Fbx/FbxMainExport.cpp b/Engine/Source/Editor/UnrealEd/Private/Fbx/FbxMainExport.cpp
--- a/Engine/Source/Editor/UnrealEd/Private/Fbx/FbxMainExport.cpp
+++ b/Engine/Source/Editor/UnrealEd/Private/Fbx/FbxMainExport.cpp
@@ -3890,7 +3890,7 @@ private:
 		{
 			return;
 		}
-		const TArray<Chaos::TPlaneConcrete<Chaos::FReal, 3>>& Faces = ConvexMesh->GetFaces();
+		const TArray<Chaos::FConvex::FPlaneType>& Faces = ConvexMesh->GetFaces();
 		for (int32 PolyIndex = 0; PolyIndex < Faces.Num(); ++PolyIndex)
 		{
 			FVector Normal = Faces[PolyIndex].Normal().GetSafeNormal();
diff --git a/Engine/Source/Programs/HeadlessChaos/Private/HeadlessChaosTestConvex.cpp b/Engine/Source/Programs/HeadlessChaos/Private/HeadlessChaosTestConvex.cpp
--- a/Engine/Source/Programs/HeadlessChaos/Private/HeadlessChaosTestConvex.cpp
+++ b/Engine/Source/Programs/HeadlessChaos/Private/HeadlessChaosTestConvex.cpp
@@ -19,21 +19,14 @@ namespace ChaosTest
 	// Check that convex creation with face merging is working correctly.
 	// The initial creation generates a set of triangles, and the merge step should
 	// leave the hull with only one face per normal.
-	void TestConvexBuilderConvexBoxFaceMerge(const FVec3* Vertices, const int32 NumVertices)
+	void TestConvexBuilderConvexBoxFaceMerge(const TArray<FConvex::FVec3Type>& Vertices)
 	{
-		TArray<Chaos::FVec3> Particles;
-		Particles.SetNum(NumVertices);
-		for (int32 ParticleIndex = 0; ParticleIndex < NumVertices; ++ParticleIndex)
-		{
-			Particles[ParticleIndex] = Vertices[ParticleIndex];
-		}
-
-		TArray<TPlaneConcrete<FReal, 3>> Planes;
+		TArray<FConvex::FPlaneType> Planes;
 		TArray<TArray<int32>> FaceVertices;
-		TArray<Chaos::FVec3> SurfaceParticles;
-		TAABB<FReal, 3> LocalBounds;
+		TArray<Chaos::FConvex::FVec3Type> SurfaceParticles;
+		FConvex::FAABB3Type LocalBounds;
 
-		FConvexBuilder::Build(Particles, Planes, FaceVertices, SurfaceParticles, LocalBounds);
+		FConvexBuilder::Build(Vertices, Planes, FaceVertices, SurfaceParticles, LocalBounds);
 		FConvexBuilder::MergeFaces(Planes, FaceVertices, SurfaceParticles, 1.0f);
 
 		// Check that we have the right number of faces and particles
@@ -71,19 +64,19 @@ namespace ChaosTest
 	// Check that face merging works for a convex box
 	GTEST_TEST(ConvexStructureTests, TestConvexBoxFaceMerge)
 	{
-		const FVec3 Vertices[] =
+		const TArray<FConvex::FVec3Type> Vertices =
 		{
-			FVec3(-50,		-50,	-50),
-			FVec3(-50,		-50,	50),
-			FVec3(-50,		50,		-50),
-			FVec3(-50,		50,		50),
-			FVec3(50,		-50,	-50),
-			FVec3(50,		-50,	50),
-			FVec3(50,		50,		-50),
-			FVec3(50,		50,		50),
+			{-50,	-50,	-50},
+			{-50,	-50,	50},
+			{-50,	50,		-50},
+			{-50,	50,		50},
+			{50,	-50,	-50},
+			{50,	-50,	50},
+			{50,	50,		-50},
+			{50,	50,		50},
 		};
 
-		TestConvexBuilderConvexBoxFaceMerge(Vertices, UE_ARRAY_COUNT(Vertices));
+		TestConvexBuilderConvexBoxFaceMerge(Vertices);
 	}
 
 	// Check that the convex structure data is consistent (works for TBox and TConvex)
@@ -99,7 +92,7 @@ namespace ChaosTest
 			// All vertices should be on the plane
 			for (int32 PlaneVertexIndex = 0; PlaneVertexIndex < Convex.NumPlaneVertices(PlaneIndex); ++PlaneVertexIndex)
 			{
-				const TPlaneConcrete<FReal, 3> Plane = Convex.GetPlane(PlaneIndex);
+				const auto Plane = Convex.GetPlane(PlaneIndex);
 				const int32 VertexIndex = Convex.GetPlaneVertex(PlaneIndex, PlaneVertexIndex);
 				const FVec3 Vertex = Convex.GetVertex(VertexIndex);
 				const FReal VertexDistance = FVec3::DotProduct(Plane.Normal(), Vertex - Plane.X());
@@ -118,7 +111,7 @@ namespace ChaosTest
 
 			for (int32 PlaneIndex : PlaneIndices)
 			{
-				const TPlaneConcrete<FReal, 3> Plane = Convex.GetPlane(PlaneIndex);
+				const auto Plane = Convex.GetPlane(PlaneIndex);
 				const FVec3 Vertex = Convex.GetVertex(VertexIndex);
 				const FReal VertexDistance = FVec3::DotProduct(Plane.Normal(), Vertex - Plane.X());
 				EXPECT_NEAR(VertexDistance, 0.0f, Tolerance);
@@ -127,65 +120,56 @@ namespace ChaosTest
 	}
 
 	// Check that the convex structure data is consistent
-	void TestConvexStructureData(const FVec3* Vertices, const int32 NumVertices)
+	void TestConvexStructureData(const TArray<FConvex::FVec3Type>& Vertices)
 	{
-		TArray<Chaos::FVec3> Particles;
-		Particles.SetNum(NumVertices);
-		for (int32 ParticleIndex = 0; ParticleIndex < NumVertices; ++ParticleIndex)
-		{
-			Particles[ParticleIndex] = Vertices[ParticleIndex];
-		}
-
-		FConvex Convex(Particles, 0.0f);
-
+		FConvex Convex(Vertices, 0.0f);
 		TestConvexStructureDataImpl(Convex);
 	}
 
 	// Check that the convex structure data is consistent for a simple convex box
 	GTEST_TEST(ConvexStructureTests, TestConvexStructureData)
 	{
-		const FVec3 Vertices[] =
+		const TArray<FConvex::FVec3Type> Vertices =
 		{
-			FVec3(-50,		-50,	-50),
-			FVec3(-50,		-50,	50),
-			FVec3(-50,		50,		-50),
-			FVec3(-50,		50,		50),
-			FVec3(50,		-50,	-50),
-			FVec3(50,		-50,	50),
-			FVec3(50,		50,		-50),
-			FVec3(50,		50,		50),
+			{-50,		-50,	-50},
+			{-50,		-50,	50},
+			{-50,		50,		-50},
+			{-50,		50,		50},
+			{50,		-50,	-50},
+			{50,		-50,	50},
+			{50,		50,		-50},
+			{50,		50,		50},
 		};
 
-		TestConvexStructureData(Vertices, UE_ARRAY_COUNT(Vertices));
+		TestConvexStructureData(Vertices);
 	}
 
 	// Check that the convex structure data is consistent for a complex convex shape
 	GTEST_TEST(ConvexStructureTests, TestConvexStructureData2)
 	{
-		const FVec3 Vertices[] =
-
+		const TArray<FConvex::FVec3Type> Vertices =
 		{
-			FVec3(0, 0, 12.0f),
-			FVec3(-0.707f, -0.707f, 10.0f),
-			FVec3(0, -1, 10.0f),
-			FVec3(0.707f, -0.707f, 10.0f),
-			FVec3(1, 0, 10.0f),
-			FVec3(0.707f, 0.707f, 10.0f),
-			FVec3(0.0f, 1.0f, 10.0f),
-			FVec3(-0.707f, 0.707f, 10.0f),
-			FVec3(-1.0f, 0.0f, 10.0f),
-			FVec3(-0.707f, -0.707f, 0.0f),
-			FVec3(0, -1, 0.0f),
-			FVec3(0.707f, -0.707f, 0.0f),
-			FVec3(1, 0, 0.0f),
-			FVec3(0.707f, 0.707f, 0.0f),
-			FVec3(0.0f, 1.0f, 0.0f),
-			FVec3(-0.707f, 0.707f, 0.0f),
-			FVec3(-1.0f, 0.0f, 0.0f),
-			FVec3(0, 0, -2.0f),
+			{0, 0, 12.0f},
+			{-0.707f, -0.707f, 10.0f},
+			{0, -1, 10.0f},
+			{0.707f, -0.707f, 10.0f},
+			{1, 0, 10.0f},
+			{0.707f, 0.707f, 10.0f},
+			{0.0f, 1.0f, 10.0f},
+			{-0.707f, 0.707f, 10.0f},
+			{-1.0f, 0.0f, 10.0f},
+			{-0.707f, -0.707f, 0.0f},
+			{0, -1, 0.0f},
+			{0.707f, -0.707f, 0.0f},
+			{1, 0, 0.0f},
+			{0.707f, 0.707f, 0.0f},
+			{0.0f, 1.0f, 0.0f},
+			{-0.707f, 0.707f, 0.0f},
+			{-1.0f, 0.0f, 0.0f},
+			{0, 0, -2.0f},
 		};
 
-		TestConvexStructureData(Vertices, UE_ARRAY_COUNT(Vertices));
+		TestConvexStructureData(Vertices);
 	}
 
 	// Check that the convex structure data is consistent for a standard box
@@ -236,16 +220,15 @@ namespace ChaosTest
 		const FReal Radius = 1000.0f;
 
 		const int32 NumVertices = TestGeometry2::RawVertexArray.Num() / 3;
-		TArray<FVec3> Particles;
+		TArray<FConvex::FVec3Type> Particles;
 		Particles.SetNum(NumVertices);
 		for (int32 ParticleIndex = 0; ParticleIndex < NumVertices; ++ParticleIndex)
 		{
-			const FVec3 VertexPos = FVec3(
+			Particles[ParticleIndex] = FConvex::FVec3Type(
 				TestGeometry2::RawVertexArray[3 * ParticleIndex + 0],
 				TestGeometry2::RawVertexArray[3 * ParticleIndex + 1],
 				TestGeometry2::RawVertexArray[3 * ParticleIndex + 2]
 			);
-			Particles[ParticleIndex] = VertexPos;
 		}
 
 		FConvex Convex(Particles, 0.0f);
@@ -266,14 +249,13 @@ namespace ChaosTest
 		const int32 NumVertices = 50000;
 
 		// Make a convex with points on a sphere.
-		TArray<FVec3> Vertices;
+		TArray<FConvex::FVec3Type> Vertices;
 		Vertices.SetNum(NumVertices);
 		for (int32 VertexIndex = 0; VertexIndex < NumVertices; ++VertexIndex)
 		{
-			const FReal Theta = FMath::RandRange(-PI, PI);
-			const FReal Phi = FMath::RandRange(-0.5f * PI, 0.5f * PI);
-			const FVec3 VertexPos = Radius * FVec3(FMath::Cos(Theta), FMath::Sin(Theta), FMath::Sin(Phi));
-			Vertices[VertexIndex] = VertexPos;
+			const FConvex::FRealType Theta = FMath::RandRange(-PI, PI);
+			const FConvex::FRealType Phi = FMath::RandRange(-0.5f * PI, 0.5f * PI);
+			Vertices[VertexIndex] = Radius * FConvex::FVec3Type(FMath::Cos(Theta), FMath::Sin(Theta), FMath::Sin(Phi));
 		}
 		FConvex Convex(Vertices, 0.0f);
 
@@ -291,18 +273,18 @@ namespace ChaosTest
 		// Create a long mesh with a extremely small end (YZ plane) 
 		// so that it generate extremely sized triangle that will produce extremely small (unormalized) normals
 		const float SmallNumber = 0.001f;
-		const FVec3 Range{ 100.0f, SmallNumber, SmallNumber };
+		const FConvex::FVec3Type Range{ 100.0f, SmallNumber, SmallNumber };
 
-		const FVec3 Vertices[] =
+		const TArray<FConvex::FVec3Type> Vertices =
 		{
-			FVec3(0, 0, 0),
-			FVec3(Range.X, 0, 0),
-			FVec3(Range.X, Range.Y, 0),
-			FVec3(Range.X, Range.Y, Range.Z),
-			FVec3(Range.X+ SmallNumber, Range.Y*0.5f, Range.Z * 0.5f),
+			{0, 0, 0},
+			{Range.X, 0, 0},
+			{Range.X, Range.Y, 0},
+			{Range.X, Range.Y, Range.Z},
+			{Range.X + SmallNumber, Range.Y * 0.5f, Range.Z * 0.5f},
 		};
 
-		TestConvexStructureData(Vertices, UE_ARRAY_COUNT(Vertices));
+		TestConvexStructureData(Vertices);
 	}
 
 	GTEST_TEST(ConvexStructureTests, TestConvexFailsSafelyOnPlanarObject)
@@ -313,7 +295,7 @@ namespace ChaosTest
 		// a check to fire inside the convex builder as we classified the object incorrectly and didn't
 		// safely handle a failure due to a planar object. This test verifies that the builder can
 		// safely fail to build a convex from a plane.
-		const TArray<FVec3> Vertices =
+		const TArray<FConvex::FVec3Type> Vertices =
 		{
 			{-15.1425571, 16.9698563, 0.502334476},
 			{-15.1425571, 16.9698563, 0.502334476},
@@ -365,10 +347,10 @@ namespace ChaosTest
 			{15.1425571, -16.9698563, -0.502334476}
 		};
 
-		TArray <TPlaneConcrete<FReal, 3>> Planes;
+		TArray<FConvex::FPlaneType> Planes;
 		TArray<TArray<int32>> FaceIndices;
-		TArray<FVec3> FinalVertices;
-		TAABB<FReal, 3> LocalBounds;
+		TArray<FConvex::FVec3Type> FinalVertices;
+		FConvex::FAABB3Type LocalBounds;
 
 		{
 			// Temporarily set LogChaos to error, we're expecting this to fire warnings and don't want that to fail a CIS run.
@@ -382,7 +364,7 @@ namespace ChaosTest
 
 	GTEST_TEST(ConvexStructureTests, TestConvexHalfEdgeStructureData_Box)
 	{
-		const TArray<FVec3> InputVertices =
+		const TArray<FConvex::FVec3Type> InputVertices =
 		{
 			FVec3(-50,		-50,	-50),
 			FVec3(-50,		-50,	50),
@@ -394,10 +376,10 @@ namespace ChaosTest
 			FVec3(50,		50,		50),
 		};
 
-		TArray<TPlaneConcrete<FReal, 3>> Planes;
+		TArray<FConvex::FPlaneType> Planes;
 		TArray<TArray<int32>> FaceVertices;
-		TArray<Chaos::FVec3> Vertices;
-		TAABB<FReal, 3> LocalBounds;
+		TArray<FConvex::FVec3Type> Vertices;
+		FConvex::FAABB3Type LocalBounds;
 		FConvexBuilder::Build(InputVertices, Planes, FaceVertices, Vertices, LocalBounds);
 		FConvexBuilder::MergeFaces(Planes, FaceVertices, Vertices, 1.0f);
 
diff --git a/Engine/Source/Programs/HeadlessChaos/Private/HeadlessChaosTestEPA.cpp b/Engine/Source/Programs/HeadlessChaos/Private/HeadlessChaosTestEPA.cpp
--- a/Engine/Source/Programs/HeadlessChaos/Private/HeadlessChaosTestEPA.cpp
+++ b/Engine/Source/Programs/HeadlessChaos/Private/HeadlessChaosTestEPA.cpp
@@ -670,7 +670,7 @@ namespace ChaosTest
 					{{768.000000, -448.000000, 3.81469727e-06}, {6.62273836e-09, 6.62273836e-09, -1.00000000}}
 				});
 
-			TArray<FVec3> SurfaceParticles(
+			TArray<FConvex::FVec3Type> SurfaceParticles(
 				{
 					{0.000000000, -1024.00000, 2.84217094e-14},
 					{768.000000, -1024.00000, 2.84217094e-14},
@@ -766,7 +766,7 @@ namespace ChaosTest
 					{{0.000000000,0.000000000,0.000000000},{-0.000000000,0.000000000,1.00000000}}
 					});
 
-				TArray<FVec3> SurfaceParticles(
+				TArray<FConvex::FVec3Type> SurfaceParticles(
 					{
 						{0.000000000,-512.000000,-32.0000000},
 					{512.000000,0.000000000,-32.0000000},
@@ -821,7 +821,7 @@ namespace ChaosTest
 		}
 
 		// Defining boat geom data outside of single test scope as it's used in multiple.
-		const TArray<Chaos::FVec3> BoatSurfaceVertices(
+		const TArray<FConvex::FVec3Type> BoatSurfaceVertices(
 			{
 				{-118.965088, -100.379936, 105.818298},
 				{-128.562881, 80.0933762, 107.703270},
@@ -1031,7 +1031,7 @@ namespace ChaosTest
 					{{0.000000000,0.000000000,0.000000000},{-0.000000000,0.000000000,1.00000000}}
 				});
 
-			TArray<FVec3> GroundSurfaceParticles(
+			TArray<FConvex::FVec3Type> GroundSurfaceParticles(
 				{
 					{0.000000000,-512.000000,-32.0000000},
 					{512.000000,0.000000000,-32.0000000},
@@ -1197,7 +1197,7 @@ namespace ChaosTest
 		GTEST_TEST(EPATests, EPARealFailures_ConvexTrimeshRotationalDifferencesBreakNormal)
 		{
 			using namespace Chaos;
-			TArray<FVec3> ConvexBoxSurfaceParticles(
+			TArray<FConvex::FVec3Type> ConvexBoxSurfaceParticles(
 				{
 				{50.0999985, -50.1124992, -50.1250000},
 				{-50.0999985, 50.1250000, -50.1250000},
diff --git a/Engine/Source/Programs/HeadlessChaos/Private/HeadlessChaosTestGJK.cpp b/Engine/Source/Programs/HeadlessChaos/Private/HeadlessChaosTestGJK.cpp
--- a/Engine/Source/Programs/HeadlessChaos/Private/HeadlessChaosTestGJK.cpp
+++ b/Engine/Source/Programs/HeadlessChaos/Private/HeadlessChaosTestGJK.cpp
@@ -549,7 +549,7 @@ namespace ChaosTest
 
 		{
 			//Tetrahedron
-			TArray<Chaos::FVec3> HullParticles;
+			TArray<FConvex::FVec3Type> HullParticles;
 			HullParticles.SetNum(4);
 			HullParticles[0] = { -1,-1,-1 };
 			HullParticles[1] = { 1,-1,-1 };
@@ -584,7 +584,7 @@ namespace ChaosTest
 
 		{
 			//Triangle
-			TArray<Chaos::FVec3> TriangleParticles;
+			TArray<FConvex::FVec3Type> TriangleParticles;
 			TriangleParticles.SetNum(3);
 			TriangleParticles[0] = { -1,-1,-1 };
 			TriangleParticles[1] = { 1,-1,-1 };
@@ -919,7 +919,7 @@ namespace ChaosTest
 	void GJKSphereConvexSweep()
 	{
 		//Tetrahedron
-		TArray<Chaos::FVec3> HullParticles;
+		TArray<FConvex::FVec3Type> HullParticles;
 		HullParticles.SetNum(4);
 		HullParticles[0] = { 3,0,4 };
 		HullParticles[1] = { 3,1,0 };
@@ -1349,7 +1349,7 @@ namespace ChaosTest
 
 		{
 			//based on real sweep from game
-			TArray<Chaos::FVec3> ConvexParticles;
+			TArray<FConvex::FVec3Type> ConvexParticles;
 			ConvexParticles.SetNum(10);
 
 			ConvexParticles[0] = { 51870.2305, 54369.6719, 19200.0000 };
@@ -1387,7 +1387,7 @@ namespace ChaosTest
 	GTEST_TEST(GJKTests, DISABLED_TestGJKCapsuleConvexInitialOverlapSweep_Fixed)
 	{
 		{
-			TArray<Chaos::FVec3> ConvexParticles;
+			TArray<FConvex::FVec3Type> ConvexParticles;
 			ConvexParticles.SetNum(8);
 
 			ConvexParticles[0] ={-256.000031,12.0000601,384.000061};
@@ -1424,7 +1424,7 @@ namespace ChaosTest
 	void GJKCapsuleConvexInitialOverlapSweep()
 	{
 		{
-			TArray<FVec3> ConvexParticles;
+			TArray<FConvex::FVec3Type> ConvexParticles;
 			ConvexParticles.SetNum(16);
 
 			ConvexParticles[0] ={-127.216454,203.240234,124.726524};
@@ -1493,7 +1493,7 @@ namespace ChaosTest
 
 		{
 			//capsule vs triangle as we make the sweep longer the world space point of impact should stay the same
-			TArray<FVec3> ConvexParticles;
+			TArray<FConvex::FVec3Type> ConvexParticles;
 			ConvexParticles.SetNum(3);
 
 			ConvexParticles[0] ={7400.00000, 12600.0000, 206.248123};
@@ -1704,7 +1704,7 @@ namespace ChaosTest
 	void GJKConvexConvexEPABoundaryCondition()
 	{
 		// These verts are those from a rectangular box with bevelled edges
-		TArray<FVec3> CoreShapeVerts = 
+		TArray<FConvex::FVec3Type> CoreShapeVerts =
 		{
 			{3.54999995f, -1.04999995f, 0.750000000f},
 			{3.75000000f, 1.04999995f, 0.549999952f},
@@ -1786,7 +1786,7 @@ namespace ChaosTest
 
 	void NegativeScaleConvexTest()
 	{
-		TArray<FVec3> ConvexVerts =
+		TArray<FConvex::FVec3Type> ConvexVerts =
 		{
 			{512.000061, -1279.99988, -383.999939},
 			{511.999969, 6.81566016e-05, 2.23802308e-05},
@@ -1797,7 +1797,7 @@ namespace ChaosTest
 			{7.96019594e-05, -1280.00000, -383.999969},
 			{512.000061, -1023.99994, -383.999939}
 		};
-		TArray<FVec3> ConvexVertices(MoveTemp(ConvexVerts));
+		TArray<FConvex::FVec3Type> ConvexVertices(MoveTemp(ConvexVerts));
 		TUniquePtr<FImplicitConvex3> CoreConvex = MakeUnique<FImplicitConvex3>(ConvexVertices, 0.0f);
 		const TImplicitObjectScaled<FImplicitConvex3> ScaledConvex(MakeSerializable(CoreConvex), FVec3(-1,1,1), 38.4000015);
 		const TSphere<FReal, 3> Sphere(FVec3(0,0,0), 32);
@@ -1825,7 +1825,7 @@ namespace ChaosTest
 		//	{-2.36513770e-05, -1.52587909e-05, -2.84217094e-14},
 		//	{1.80563184e-05, -256.000031, -2.84217094e-14},
 		//};
-		TArray<FVec3> ConvexVerts =
+		TArray<FConvex::FVec3Type> ConvexVerts =
 		{
 			// subset of verts from above test.
 			FVec3(-512, -1280, -384),
@@ -1834,7 +1834,7 @@ namespace ChaosTest
 			FVec3(0, 0, 0),
 			FVec3(0, -256, 0),
 		};
-		TArray<FVec3> ConvexVertices(MoveTemp(ConvexVerts));
+		TArray<FConvex::FVec3Type> ConvexVertices(MoveTemp(ConvexVerts));
 		FImplicitConvex3 CoreConvex = FImplicitConvex3(ConvexVertices, 38.4000015);
 		const TSphere<FReal, 3> Sphere(FVec3(0, 0, 0), 32);
 		const FRigidTransform3 StartTM(FVec3(-172.000000, -48.0000000, 52.0000000), FRotation3::FromIdentity());
diff --git a/Engine/Source/Programs/HeadlessChaos/Private/HeadlessChaosTestImplicits.cpp b/Engine/Source/Programs/HeadlessChaos/Private/HeadlessChaosTestImplicits.cpp
--- a/Engine/Source/Programs/HeadlessChaos/Private/HeadlessChaosTestImplicits.cpp
+++ b/Engine/Source/Programs/HeadlessChaos/Private/HeadlessChaosTestImplicits.cpp
@@ -1554,14 +1554,14 @@ namespace ChaosTest {
 				EXPECT_NE(Tri.Z, 8);
 			}
 
-			TArray<FVec3> Vertices;
+			TArray<FConvex::FVec3Type> Vertices;
 			Vertices.SetNum((int32)Particles.Size());
 			for (int32 VertexIndex = 0; VertexIndex < (int32)Particles.Size(); ++VertexIndex)
 			{
 				Vertices[VertexIndex] = Particles.X(VertexIndex);
 			}
 			FConvex Convex(Vertices, 0.0f);
-			const TArray<FVec3>& CulledParticles = Convex.GetVertices();
+			const TArray<FConvex::FVec3Type>& CulledParticles = Convex.GetVertices();
 			EXPECT_EQ(CulledParticles.Num(), 8);
 
 			for (int32 Idx = 0; Idx < CulledParticles.Num(); ++Idx)
@@ -1600,11 +1600,11 @@ namespace ChaosTest {
 	{
 		{
 			//degenerates
-			TArray<Chaos::FVec3> Particles;
+			TArray<FConvex::FVec3Type> Particles;
 			Particles.SetNum(3);
-			Particles[0] = Chaos::FVec3(-1, -1, -1);
-			Particles[1] = Chaos::FVec3(1, -1, -1);
-			Particles[2] = Chaos::FVec3(1, 1, -1);
+			Particles[0] = { -1, -1, -1};
+			Particles[1] = { 1, -1, -1 };
+			Particles[2] = { 1, 1, -1 };
 			TArray<TVector<int32, 3>>Indices;
 			Chaos::FConvexBuilder::BuildConvexHull(Particles, Indices);
 			EXPECT_EQ(Indices.Num(), 0);
@@ -1613,17 +1613,17 @@ namespace ChaosTest {
 			EXPECT_EQ(Indices.Num(), 0);
 		}
 		{
-			TArray <Chaos::FVec3> Particles;
+			TArray <FConvex::FVec3Type> Particles;
 			Particles.SetNum(9);
-			Particles[0] = Chaos::FVec3(-1, -1, -1);
-			Particles[1] = Chaos::FVec3(-1, -1, 1);
-			Particles[2] = Chaos::FVec3(-1, 1, -1);
-			Particles[3] = Chaos::FVec3(-1, 1, 1);
-			Particles[4] = Chaos::FVec3(1, -1, -1);
-			Particles[5] = Chaos::FVec3(1, -1, 1);
-			Particles[6] = Chaos::FVec3(1, 1, -1);
-			Particles[7] = Chaos::FVec3(1, 1, 1);
-			Particles[8] = Chaos::FVec3(0, 0, 0);
+			Particles[0] = { -1, -1, -1 };
+			Particles[1] = { -1, -1, 1 };
+			Particles[2] = { -1, 1, -1 };
+			Particles[3] = { -1, 1, 1 };
+			Particles[4] = { 1, -1, -1 };
+			Particles[5] = { 1, -1, 1 };
+			Particles[6] = { 1, 1, -1 };
+			Particles[7] = { 1, 1, 1 };
+			Particles[8] = { 0, 0, 0 };
 			TArray<TVector<int32, 3>>Indices;
 			Chaos::FConvexBuilder::BuildConvexHull(Particles, Indices);
 			EXPECT_EQ(Indices.Num(), 12);
@@ -1635,26 +1635,26 @@ namespace ChaosTest {
 			}
 		}
 		{
-			TArray<FVec3> Particles;
+			TArray<FConvex::FVec3Type> Particles;
 			Particles.SetNum(5);
-			Particles[0] = Chaos::FVec3(-1, -1, -1);
-			Particles[1] = Chaos::FVec3(1, -1, -1);
-			Particles[2] = Chaos::FVec3(1, 1, -1);
-			Particles[3] = Chaos::FVec3(0, 0, 0.5);
-			Particles[4] = (Particles[3]- Particles[1]) * 0.5 + Particles[1]+ Chaos::TVec3<float>(0, 0, 0.1);
+			Particles[0] = { -1, -1, -1 };
+			Particles[1] = { 1, -1, -1 };
+			Particles[2] = { 1, 1, -1 };
+			Particles[3] = { 0, 0, 0.5 };
+			Particles[4] = (Particles[3] - Particles[1]) * 0.5 + Particles[1] + FConvex::FVec3Type{ 0, 0, 0.1 };
 			TArray<TVector<int32, 3>> Indices;
 			Chaos::FConvexBuilder::BuildConvexHull(Particles, Indices);
 			EXPECT_EQ(Indices.Num(), 6);
 		}
 		{
-			TArray<FVec3> Particles;
+			TArray<FConvex::FVec3Type> Particles;
 			Particles.SetNum(6);
-			Particles[0] = Chaos::FVec3(-1, -1, -1);
-			Particles[1] = Chaos::FVec3(1, -1, -1);
-			Particles[2] = Chaos::FVec3(1, 1, -1);
-			Particles[3] = Chaos::FVec3(0, 0, 0.5);
-			Particles[4] = (Particles[3]- Particles[1]) * 0.5 + Particles[1] + Chaos::TVec3<float>(0, 0, 0.1);
-			Particles[5] = Particles[4]+ Chaos::TVec3<float>(-0.1, 0, 0);
+			Particles[0] = { -1, -1, -1 };
+			Particles[1] = { 1, -1, -1 };
+			Particles[2] = { 1, 1, -1 };
+			Particles[3] = { 0, 0, 0.5 };
+			Particles[4] = (Particles[3] - Particles[1]) * 0.5 + Particles[1] + FConvex::FVec3Type{ 0, 0, 0.1 };
+			Particles[5] = Particles[4] + FConvex::FVec3Type{ -0.1, 0, 0 };
 			TArray<TVector<int32, 3>> Indices;
 			Chaos::FConvexBuilder::BuildConvexHull(Particles, Indices);
 			EXPECT_EQ(Indices.Num(), 8);
@@ -1664,17 +1664,17 @@ namespace ChaosTest {
 			// a large enough epsilon for building horizons in hull generation
 			// (tested to fail with 1e-1) we will generate a non-convex hull
 			// Using a scaled epsilon resolves this case
-			TArray<FVec3> Particles;
+			TArray<FConvex::FVec3Type> Particles;
 			Particles.SetNum(9);
-			Particles[0] = Chaos::FVec3(-1, -1, -1);
-			Particles[1] = Chaos::FVec3(-1, -1, 1);
-			Particles[2] = Chaos::FVec3(-1, 1, -1);
-			Particles[3] = Chaos::FVec3(-1, 1, 1);
-			Particles[4] = Chaos::FVec3(1, -1, -1);
-			Particles[5] = Chaos::FVec3(1, -1, 1);
-			Particles[6] = Chaos::FVec3(1, 1, -1);
-			Particles[7] = Chaos::FVec3(1, 1, 1);
-			Particles[8] = Chaos::FVec3(0.966962576, -0.0577232838, 0.959515572);
+			Particles[0] = { -1, -1, -1 };
+			Particles[1] = { -1, -1, 1 };
+			Particles[2] = { -1, 1, -1 };
+			Particles[3] = { -1, 1, 1 };
+			Particles[4] = { 1, -1, -1 };
+			Particles[5] = { 1, -1, 1 };
+			Particles[6] = { 1, 1, -1 };
+			Particles[7] = { 1, 1, 1 };
+			Particles[8] = { 0.966962576, -0.0577232838, 0.959515572 };
 			
 			TArray<TVec3<int32>> Indices;
 			Chaos::FConvexBuilder::Params BuildParams;
@@ -1700,17 +1700,17 @@ namespace ChaosTest {
 			// only the original box - ignoring all interior and coplanar points.
 			// Note: If hull generation is changed to support non-triangular faces the conditions here
 			// will need to change as a correct hull in that method will produce only 6 faces not 12
-			TArray<FVec3> Particles;
+			TArray<FConvex::FVec3Type> Particles;
 			int32 NumParticles = 3600;
 			Particles.SetNum(NumParticles);
-			Particles[0] = Chaos::FVec3(-1, -1, -1);
-			Particles[1] = Chaos::FVec3(-1, -1, 1);
-			Particles[2] = Chaos::FVec3(-1, 1, -1);
-			Particles[3] = Chaos::FVec3(-1, 1, 1);
-			Particles[4] = Chaos::FVec3(1, -1, -1);
-			Particles[5] = Chaos::FVec3(1, -1, 1);
-			Particles[6] = Chaos::FVec3(1, 1, -1);
-			Particles[7] = Chaos::FVec3(1, 1, 1);
+			Particles[0] = { -1, -1, -1 };
+			Particles[1] = { -1, -1, 1 };
+			Particles[2] = { -1, 1, -1 };
+			Particles[3] = { -1, 1, 1 };
+			Particles[4] = { 1, -1, -1 };
+			Particles[5] = { 1, -1, 1 };
+			Particles[6] = { 1, 1, -1 };
+			Particles[7] = { 1, 1, 1 };
 			FRandomStream Stream(42);
 			for(int i = 8; i < NumParticles; ++i)
 			{
@@ -1739,26 +1739,26 @@ namespace ChaosTest {
 
 	void Simplify()
 	{
-		TArray<FVec3> Particles;
+		TArray<FConvex::FVec3Type> Particles;
 		Particles.SetNum(18);
-		Particles[0] = Chaos::FVec3(0, 0, 12.0f);
-		Particles[1] = Chaos::FVec3(-0.707f, -0.707f, 10.0f);
-		Particles[2] = Chaos::FVec3(0, -1, 10.0f);
-		Particles[3] = Chaos::FVec3(0.707f, -0.707f, 10.0f);
-		Particles[4] = Chaos::FVec3(1, 0, 10.0f);
-		Particles[5] = Chaos::FVec3(0.707f, 0.707f, 10.0f);
-		Particles[6] = Chaos::FVec3(0.0f, 1.0f, 10.0f);
-		Particles[7] = Chaos::FVec3(-0.707f, 0.707f, 10.0f);
-		Particles[8] = Chaos::FVec3(-1.0f, 0.0f, 10.0f);
-		Particles[9] = Chaos::FVec3(-0.707f, -0.707f, 0.0f);
-		Particles[10] = Chaos::FVec3(0, -1, 0.0f);
-		Particles[11] = Chaos::FVec3(0.707f, -0.707f, 0.0f);
-		Particles[12] = Chaos::FVec3(1, 0, 0.0f);
-		Particles[13] = Chaos::FVec3(0.707f, 0.707f, 0.0f);
-		Particles[14] = Chaos::FVec3(0.0f, 1.0f, 0.0f);
-		Particles[15] = Chaos::FVec3(-0.707f, 0.707f, 0.0f);
-		Particles[16] = Chaos::FVec3(-1.0f, 0.0f, 0.0f);
-		Particles[17] = Chaos::FVec3(0, 0, -2.0f);
+		Particles[0] = { 0, 0, 12.0f };
+		Particles[1] = { -0.707f, -0.707f, 10.0f };
+		Particles[2] = { 0, -1, 10.0f };
+		Particles[3] = { 0.707f, -0.707f, 10.0f };
+		Particles[4] = { 1, 0, 10.0f };
+		Particles[5] = { 0.707f, 0.707f, 10.0f };
+		Particles[6] = { 0.0f, 1.0f, 10.0f };
+		Particles[7] = { -0.707f, 0.707f, 10.0f };
+		Particles[8] = { -1.0f, 0.0f, 10.0f };
+		Particles[9] = { -0.707f, -0.707f, 0.0f };
+		Particles[10] = { 0, -1, 0.0f };
+		Particles[11] = { 0.707f, -0.707f, 0.0f };
+		Particles[12] = { 1, 0, 0.0f };
+		Particles[13] = { 0.707f, 0.707f, 0.0f };
+		Particles[14] = { 0.0f, 1.0f, 0.0f };
+		Particles[15] = { -0.707f, 0.707f, 0.0f };
+		Particles[16] = { -1.0f, 0.0f, 0.0f };
+		Particles[17] = { 0, 0, -2.0f };
 					
 		FConvex Convex(Particles, 0.0f);
 
@@ -1767,8 +1767,8 @@ namespace ChaosTest {
 		int32 OriginalNumberFaces = Convex.GetFaces().Num();
 		FAABB3 OriginalBoundingBox = Convex.BoundingBox();
 
-		const TArray<FVec3>& CulledParticles = Convex.GetVertices();
-		const TArray<TPlaneConcrete<Chaos::FReal, 3>> Planes = Convex.GetFaces();
+		const TArray<FConvex::FVec3Type>& CulledParticles = Convex.GetVertices();
+		const TArray<FConvex::FPlaneType> Planes = Convex.GetFaces();
 
 		// set target number of particles in simplified convex
 		FConvexBuilder::PerformGeometryReduction = 1;
diff --git a/Engine/Source/Programs/HeadlessChaos/Private/HeadlessChaosTestMostOpposing.cpp b/Engine/Source/Programs/HeadlessChaos/Private/HeadlessChaosTestMostOpposing.cpp
--- a/Engine/Source/Programs/HeadlessChaos/Private/HeadlessChaosTestMostOpposing.cpp
+++ b/Engine/Source/Programs/HeadlessChaos/Private/HeadlessChaosTestMostOpposing.cpp
@@ -76,15 +76,15 @@ namespace ChaosTest
 		FVec3 Normal;
 		int32 FaceIndex;
 
-		TArray<FVec3> Particles;
+		TArray<FConvex::FVec3Type> Particles;
 		Particles.SetNum(6);
-		Particles[0] = FVec3(1, 1, 1);
-		Particles[1] = FVec3(5, 1, 1);
-		Particles[2] = FVec3(1, 5, 1);
+		Particles[0] = { 1, 1, 1 };
+		Particles[1] = { 5, 1, 1 };
+		Particles[2] = { 1, 5, 1 };
 
-		Particles[3] = FVec3(1, 1, 1);
-		Particles[4] = FVec3(1, 5, 1);
-		Particles[5] = FVec3(1, 1, -5);
+		Particles[3] = { 1, 1, 1 };
+		Particles[4] = { 1, 5, 1 };
+		Particles[5] = { 1, 1, -5};
 
 		FConvex Convex(MoveTemp(Particles), 0.0f);
 
@@ -130,15 +130,15 @@ namespace ChaosTest
 		FVec3 Normal;
 		int32 FaceIndex;
 
-		TArray<FVec3> Particles;
+		TArray<FConvex::FVec3Type> Particles;
 		Particles.SetNum(6);
-		Particles[0] = FVec3(0, -1, 1);
-		Particles[1] = FVec3(1, -1, -1);
-		Particles[2] = FVec3(0, 1, 1);
+		Particles[0] = { 0, -1, 1 };
+		Particles[1] = { 1, -1, -1 };
+		Particles[2] = { 0, 1, 1 };
 
-		Particles[3] = FVec3(0, -1, 1);
-		Particles[4] = FVec3(0, 1, 1);
-		Particles[5] = FVec3(-1, -1, -1);
+		Particles[3] = { 0, -1, 1 };
+		Particles[4] = { 0, 1, 1 };
+		Particles[5] = { -1, -1, -1 };
 
 		TUniquePtr<FImplicitObject> Convex = MakeUnique<FConvex>(MoveTemp(Particles), 0.0f);
 
diff --git a/Engine/Source/Programs/HeadlessChaos/Private/HeadlessChaosTestOneShotManifolds.cpp b/Engine/Source/Programs/HeadlessChaos/Private/HeadlessChaosTestOneShotManifolds.cpp
--- a/Engine/Source/Programs/HeadlessChaos/Private/HeadlessChaosTestOneShotManifolds.cpp
+++ b/Engine/Source/Programs/HeadlessChaos/Private/HeadlessChaosTestOneShotManifolds.cpp
@@ -227,21 +227,21 @@ namespace ChaosTest
 	GTEST_TEST(OneShotManifoldTests, TestConvexMarginEdgeEdge)
 	{
 		FReal Dt = 1 / 30.0f;
-		FReal HalfSize = 100.0f;
+		FConvex::FRealType HalfSize = 100.0f;
 		FReal Margin = 0.2f * HalfSize;
 		float ExpectedPhi = 0.0f;
 		float Offset = 2.0f * HalfSize * FMath::Sqrt(2.0f);
 
-		TArray<FVec3> BoxVerts = 
+		TArray<FConvex::FVec3Type> BoxVerts =
 		{
-			FVec3(-HalfSize, -HalfSize, -HalfSize),
-			FVec3(-HalfSize,  HalfSize, -HalfSize),
-			FVec3( HalfSize,  HalfSize, -HalfSize),
-			FVec3( HalfSize, -HalfSize, -HalfSize),
-			FVec3(-HalfSize, -HalfSize,  HalfSize),
-			FVec3(-HalfSize,  HalfSize,  HalfSize),
-			FVec3( HalfSize,  HalfSize,  HalfSize),
-			FVec3( HalfSize, -HalfSize,  HalfSize),
+			{-HalfSize, -HalfSize, -HalfSize},
+			{-HalfSize,  HalfSize, -HalfSize},
+			{ HalfSize,  HalfSize, -HalfSize},
+			{ HalfSize, -HalfSize, -HalfSize},
+			{-HalfSize, -HalfSize,  HalfSize},
+			{-HalfSize,  HalfSize,  HalfSize},
+			{ HalfSize,  HalfSize,  HalfSize},
+			{ HalfSize, -HalfSize,  HalfSize},
 		};
 
 		// First box rotated 45 degrees about Z and placed at the oirigin
diff --git a/Engine/Source/Programs/HeadlessChaos/Private/HeadlessChaosTestUtility.cpp b/Engine/Source/Programs/HeadlessChaos/Private/HeadlessChaosTestUtility.cpp
--- a/Engine/Source/Programs/HeadlessChaos/Private/HeadlessChaosTestUtility.cpp
+++ b/Engine/Source/Programs/HeadlessChaos/Private/HeadlessChaosTestUtility.cpp
@@ -512,17 +512,17 @@ namespace ChaosTest {
 		Particle->R() = FRotation3::MakeFromEuler(FVec3(0, 0, 0)).GetNormalized();
 		Particle->W() = FVec3(0, 0, 0);
 
-		TArray<Chaos::FVec3> Cube;
+		TArray<Chaos::FConvex::FVec3Type> Cube;
 		Cube.SetNum(9);
-		Cube[0] = Chaos::FVec3(-1000, -1000, -20);
-		Cube[1] = Chaos::FVec3(-1000, -1000, 0);
-		Cube[2] = Chaos::FVec3(-1000, 1000, -20);
-		Cube[3] = Chaos::FVec3(-1000, 1000, 0);
-		Cube[4] = Chaos::FVec3(1000, -1000, -20);
-		Cube[5] = Chaos::FVec3(1000, -1000, 0);
-		Cube[6] = Chaos::FVec3(1000, 1000, -20);
-		Cube[7] = Chaos::FVec3(1000, 1000, 0);
-		Cube[8] = Chaos::FVec3(0, 0, 0);
+		Cube[0] = { -1000, -1000, -20 };
+		Cube[1] = { -1000, -1000, 0 };
+		Cube[2] = { -1000, 1000, -20 };
+		Cube[3] = { -1000, 1000, 0 };
+		Cube[4] = { 1000, -1000, -20 };
+		Cube[5] = { 1000, -1000, 0 };
+		Cube[6] = { 1000, 1000, -20 };
+		Cube[7] = { 1000, 1000, 0 };
+		Cube[8] = { 0, 0, 0 };
 
 		Particle->SetDynamicGeometry(MakeUnique<FConvex>(Cube, 0.0f));
 
@@ -559,17 +559,19 @@ namespace ChaosTest {
 	/**/
 	void AppendDynamicParticleConvexBox(FPBDRigidParticleHandle& InParticles, const FVec3& Scale, FReal Margin)
 	{
-		TArray<FVec3> Cube;
+		FConvex::FVec3Type ScaleFloat{ FRealSingle(Scale.X), FRealSingle(Scale.Y), FRealSingle(Scale.Z) };
+
+		TArray<FConvex::FVec3Type> Cube;
 		Cube.SetNum(9);
-		Cube[0] = Chaos::FVec3(-1, -1, -1)*Scale;
-		Cube[1] = Chaos::FVec3(-1, -1, 1)*Scale;
-		Cube[2] = Chaos::FVec3(-1, 1, -1)*Scale;
-		Cube[3] = Chaos::FVec3(-1, 1, 1)*Scale;
-		Cube[4] = Chaos::FVec3(1, -1, -1)*Scale;
-		Cube[5] = Chaos::FVec3(1, -1, 1)*Scale;
-		Cube[6] = Chaos::FVec3(1, 1, -1)*Scale;
-		Cube[7] = Chaos::FVec3(1, 1, 1)*Scale;
-		Cube[8] = Chaos::FVec3(0, 0, 0);
+		Cube[0] = FConvex::FVec3Type{ -1, -1, -1 } * ScaleFloat;
+		Cube[1] = FConvex::FVec3Type{ -1, -1, 1 } * ScaleFloat;
+		Cube[2] = FConvex::FVec3Type{ -1, 1, -1 } * ScaleFloat;
+		Cube[3] = FConvex::FVec3Type{ -1, 1, 1 } * ScaleFloat;
+		Cube[4] = FConvex::FVec3Type{ 1, -1, -1 } * ScaleFloat;
+		Cube[5] = FConvex::FVec3Type{ 1, -1, 1 } * ScaleFloat;
+		Cube[6] = FConvex::FVec3Type{ 1, 1, -1 } * ScaleFloat;
+		Cube[7] = FConvex::FVec3Type{ 1, 1, 1 } * ScaleFloat;
+		Cube[8] = FConvex::FVec3Type{ 0, 0, 0 };
 
 		InParticles.X() = FVec3(0, 0, 0);
 		InParticles.V() = FVec3(0, 0, 0);
@@ -702,46 +704,46 @@ namespace ChaosTest {
 		}
 	}
 
-	TArray<FVec3> MakeBoxVerts(const FVec3& Center, const FVec3& HalfSize)
+	TArray<FConvex::FVec3Type> MakeBoxVerts(const FConvex::FVec3Type& Center, const FConvex::FVec3Type& HalfSize)
 	{
 		return
 		{
-			Center + FVec3(-HalfSize.X, -HalfSize.Y, -HalfSize.Z),
-			Center + FVec3(-HalfSize.X,  HalfSize.Y, -HalfSize.Z),
-			Center + FVec3(HalfSize.X,  HalfSize.Y, -HalfSize.Z),
-			Center + FVec3(HalfSize.X, -HalfSize.Y, -HalfSize.Z),
-			Center + FVec3(-HalfSize.X, -HalfSize.Y,  HalfSize.Z),
-			Center + FVec3(-HalfSize.X,  HalfSize.Y,  HalfSize.Z),
-			Center + FVec3(HalfSize.X,  HalfSize.Y,  HalfSize.Z),
-			Center + FVec3(HalfSize.X, -HalfSize.Y,  HalfSize.Z),
+			Center + FConvex::FVec3Type(-HalfSize.X, -HalfSize.Y, -HalfSize.Z),
+			Center + FConvex::FVec3Type(-HalfSize.X,  HalfSize.Y, -HalfSize.Z),
+			Center + FConvex::FVec3Type(HalfSize.X,  HalfSize.Y, -HalfSize.Z),
+			Center + FConvex::FVec3Type(HalfSize.X, -HalfSize.Y, -HalfSize.Z),
+			Center + FConvex::FVec3Type(-HalfSize.X, -HalfSize.Y,  HalfSize.Z),
+			Center + FConvex::FVec3Type(-HalfSize.X,  HalfSize.Y,  HalfSize.Z),
+			Center + FConvex::FVec3Type(HalfSize.X,  HalfSize.Y,  HalfSize.Z),
+			Center + FConvex::FVec3Type(HalfSize.X, -HalfSize.Y,  HalfSize.Z),
 		};
 	}
 
-	FImplicitConvex3 CreateConvexBox(const FVec3& BoxSize, const FReal Margin)
+	FImplicitConvex3 CreateConvexBox(const FConvex::FVec3Type& BoxSize, const FReal Margin)
 	{
-		const FVec3 HalfSize = 0.5f * BoxSize;
-		return FImplicitConvex3(MakeBoxVerts(FVec3(0), HalfSize), Margin);
+		const FConvex::FVec3Type HalfSize = 0.5f * BoxSize;
+		return FImplicitConvex3(MakeBoxVerts(FConvex::FVec3Type(0), HalfSize), Margin);
 	}
 
-	FImplicitConvex3 CreateConvexBox(const FVec3& BoxMin, const FVec3& BoxMax, const FReal Margin)
+	FImplicitConvex3 CreateConvexBox(const FConvex::FVec3Type& BoxMin, const FConvex::FVec3Type& BoxMax, const FReal Margin)
 	{
-		const FVec3 Center = 0.5f * (BoxMin + BoxMax);
-		const FVec3 HalfSize = 0.5f * (BoxMax - BoxMin);
+		const FConvex::FVec3Type Center = 0.5f * (BoxMin + BoxMax);
+		const FConvex::FVec3Type HalfSize = 0.5f * (BoxMax - BoxMin);
 		return FImplicitConvex3(MakeBoxVerts(Center, HalfSize), Margin);
 	}
 
 
-	TImplicitObjectInstanced<FImplicitConvex3> CreateInstancedConvexBox(const FVec3& BoxSize, const FReal Margin)
+	TImplicitObjectInstanced<FImplicitConvex3> CreateInstancedConvexBox(const FConvex::FVec3Type& BoxSize, const FReal Margin)
 	{
-		const FVec3 HalfSize = 0.5f * BoxSize;
-		TSharedPtr<FImplicitConvex3, ESPMode::ThreadSafe> BoxConvex = MakeShared<FImplicitConvex3, ESPMode::ThreadSafe>(MakeBoxVerts(FVec3(0), HalfSize), 0.0f);
+		const FConvex::FVec3Type HalfSize = 0.5f * BoxSize;
+		TSharedPtr<FImplicitConvex3, ESPMode::ThreadSafe> BoxConvex = MakeShared<FImplicitConvex3, ESPMode::ThreadSafe>(MakeBoxVerts(FConvex::FVec3Type(0), HalfSize), 0.0f);
 		return TImplicitObjectInstanced<FImplicitConvex3>(BoxConvex, Margin);
 	}
 
-	TImplicitObjectScaled<FImplicitConvex3> CreateScaledConvexBox(const FVec3& BoxSize, const FVec3 BoxScale, const FReal Margin)
+	TImplicitObjectScaled<FImplicitConvex3> CreateScaledConvexBox(const FConvex::FVec3Type& BoxSize, const FVec3 BoxScale, const FReal Margin)
 	{
-		const FVec3 HalfSize = 0.5f * BoxSize;
-		TSharedPtr<FImplicitConvex3, ESPMode::ThreadSafe> BoxConvex = MakeShared<FImplicitConvex3, ESPMode::ThreadSafe>(MakeBoxVerts(FVec3(0), HalfSize), 0.0f);
+		const FConvex::FVec3Type HalfSize = 0.5f * BoxSize;
+		TSharedPtr<FImplicitConvex3, ESPMode::ThreadSafe> BoxConvex = MakeShared<FImplicitConvex3, ESPMode::ThreadSafe>(MakeBoxVerts(FConvex::FVec3Type(0), HalfSize), 0.0f);
 		return TImplicitObjectScaled<FImplicitConvex3>(BoxConvex, BoxScale, Margin);
 	}
 
diff --git a/Engine/Source/Programs/HeadlessChaos/Public/HeadlessChaosTestUtility.h b/Engine/Source/Programs/HeadlessChaos/Public/HeadlessChaosTestUtility.h
--- a/Engine/Source/Programs/HeadlessChaos/Public/HeadlessChaosTestUtility.h
+++ b/Engine/Source/Programs/HeadlessChaos/Public/HeadlessChaosTestUtility.h
@@ -171,8 +171,8 @@ namespace ChaosTest {
 	}
 
 
-	extern FImplicitConvex3 CreateConvexBox(const FVec3& BoxSize, const FReal Margin);
-	extern FImplicitConvex3 CreateConvexBox(const FVec3& BoxMin, const FVec3& BoxMax, const FReal Margin);
-	extern TImplicitObjectInstanced<FImplicitConvex3> CreateInstancedConvexBox(const FVec3& BoxSize, const FReal Margin);
-	extern TImplicitObjectScaled<FImplicitConvex3> CreateScaledConvexBox(const FVec3& BoxSize, const FVec3 BoxScale, const FReal Margin);
+	extern FImplicitConvex3 CreateConvexBox(const FConvex::FVec3Type& BoxSize, const FReal Margin);
+	extern FImplicitConvex3 CreateConvexBox(const FConvex::FVec3Type& BoxMin, const FConvex::FVec3Type& BoxMax, const FReal Margin);
+	extern TImplicitObjectInstanced<FImplicitConvex3> CreateInstancedConvexBox(const FConvex::FVec3Type& BoxSize, const FReal Margin);
+	extern TImplicitObjectScaled<FImplicitConvex3> CreateScaledConvexBox(const FConvex::FVec3Type& BoxSize, const FVec3 BoxScale, const FReal Margin);
 }
diff --git a/Engine/Source/Runtime/Engine/Private/PhysicsEngine/BodySetup.cpp b/Engine/Source/Runtime/Engine/Private/PhysicsEngine/BodySetup.cpp
--- a/Engine/Source/Runtime/Engine/Private/PhysicsEngine/BodySetup.cpp
+++ b/Engine/Source/Runtime/Engine/Private/PhysicsEngine/BodySetup.cpp
@@ -2027,7 +2027,7 @@ TArray<int32> FKConvexElem::GetChaosConvexIndices() const
 	const int32 NumVerts = VertexData.Num();
 	if (NumVerts > 0)
 	{
-		TArray<Chaos::FVec3> ConvexVertices;
+		TArray<Chaos::FConvex::FVec3Type> ConvexVertices;
 		ConvexVertices.SetNum(NumVerts);
 
 		for (int32 VertIndex = 0; VertIndex < NumVerts; ++VertIndex)
diff --git a/Engine/Source/Runtime/Engine/Private/PhysicsEngine/Experimental/ChaosCooking.cpp b/Engine/Source/Runtime/Engine/Private/PhysicsEngine/Experimental/ChaosCooking.cpp
--- a/Engine/Source/Runtime/Engine/Private/PhysicsEngine/Experimental/ChaosCooking.cpp
+++ b/Engine/Source/Runtime/Engine/Private/PhysicsEngine/Experimental/ChaosCooking.cpp
@@ -163,13 +163,13 @@ namespace Chaos
 					}
 
 					// Create the corner vertices for the convex
-					TArray<FVec3> ConvexVertices;
+					TArray<FConvex::FVec3Type> ConvexVertices;
 					ConvexVertices.SetNumZeroed(NumHullVerts);
 
 					for(int32 VertIndex = 0; VertIndex < NumHullVerts; ++VertIndex)
 					{
 						const FVector& HullVert = HullVerts[VertIndex];
-						ConvexVertices[VertIndex] = FVector(bMirrored ? -HullVert.X : HullVert.X, HullVert.Y, HullVert.Z);
+						ConvexVertices[VertIndex] = FConvex::FVec3Type(bMirrored ? -HullVert.X : HullVert.X, HullVert.Y, HullVert.Z);
 					}
 
 					// Margin is always zero on convex shapes - they are intended to be instanced
diff --git a/Engine/Source/Runtime/Experimental/Chaos/Private/Chaos/ChaosDebugDraw.cpp b/Engine/Source/Runtime/Experimental/Chaos/Private/Chaos/ChaosDebugDraw.cpp
--- a/Engine/Source/Runtime/Experimental/Chaos/Private/Chaos/ChaosDebugDraw.cpp
+++ b/Engine/Source/Runtime/Experimental/Chaos/Private/Chaos/ChaosDebugDraw.cpp
@@ -616,14 +616,16 @@ namespace Chaos
 				// Warning: extremely expensive!
 				if (bChaosDebugDebugDrawExactCoreShapes)
 				{
-					const FRigidTransform3 ShrunkScaledTransform = FRigidTransform3(ShapeTransform.GetTranslation(), ShapeTransform.GetRotation());
-					TArray<FVec3> ScaledVerts = Convex->GetVertices();
-					for (FVec3& Vert : ScaledVerts)
+					TArray<FConvex::FVec3Type> ScaledVerts(Convex->GetVertices());
+					FConvex::FVec3Type Scale(ShapeTransform.GetScale3D());
+					for (FConvex::FVec3Type& Vert : ScaledVerts)
 					{
-						Vert = Vert * ShapeTransform.GetScale3D();
+						Vert *= Scale;
 					}
-					FConvex ShrunkScaledConvex(ScaledVerts, 0.0f);
-					ShrunkScaledConvex.MovePlanesAndRebuild(-Margin);
+					FConvex ShrunkScaledConvex(ScaledVerts, FReal(0));
+					ShrunkScaledConvex.MovePlanesAndRebuild(FConvex::FRealType(-Margin)); // potential loss of precision but margin should remain within the float range
+
+					const FRigidTransform3 ShrunkScaledTransform = FRigidTransform3(ShapeTransform.GetTranslation(), ShapeTransform.GetRotation());
 					DrawShapesConvexImpl(Particle, ShrunkScaledTransform, &ShrunkScaledConvex, 0.0f, FColor::Red, Settings);
 				}
 				break;
@@ -1454,7 +1456,7 @@ namespace Chaos
 
 			virtual void Box(const FAABB3& InBox, const TVec3<FReal>& InLinearColor, float InThickness) override
 			{
-				FDebugDrawQueue::GetInstance().DrawDebugBox(InBox.Center(), InBox.Extents(), FQuat::Identity, FLinearColor(InLinearColor).ToFColor(false), false, -1.f, Settings.DrawPriority, InThickness);
+				FDebugDrawQueue::GetInstance().DrawDebugBox(InBox.Center(), InBox.Extents() * FReal(0.5), FQuat::Identity, FLinearColor(InLinearColor).ToFColor(false), false, -1.f, Settings.DrawPriority, InThickness);
 			}
 
 			virtual void Line(const TVec3<FReal>& InBegin, const TVec3<FReal>& InEnd, const TVec3<FReal>& InLinearColor, float InThickness) override
diff --git a/Engine/Source/Runtime/Experimental/Chaos/Private/Chaos/CollisionConvexMesh.cpp b/Engine/Source/Runtime/Experimental/Chaos/Private/Chaos/CollisionConvexMesh.cpp
--- a/Engine/Source/Runtime/Experimental/Chaos/Private/Chaos/CollisionConvexMesh.cpp
+++ b/Engine/Source/Runtime/Experimental/Chaos/Private/Chaos/CollisionConvexMesh.cpp
@@ -27,11 +27,11 @@ namespace Chaos
 
 
 
-	void FConvexBuilder::Build(const TArray<FVec3>& InVertices, TArray <TPlaneConcrete<FReal, 3>>& OutPlanes, TArray<TArray<int32>>& OutFaceIndices, TArray<FVec3>& OutVertices, FAABB3& OutLocalBounds)
+	void FConvexBuilder::Build(const TArray<FVec3Type>& InVertices, TArray <FPlaneType>& OutPlanes, TArray<TArray<int32>>& OutFaceIndices, TArray<FVec3Type>& OutVertices, FAABB3Type& OutLocalBounds)
 	{
 		OutPlanes.Reset();
 		OutVertices.Reset();
-		OutLocalBounds = FAABB3::EmptyAABB();
+		OutLocalBounds = FAABB3Type::EmptyAABB();
 
 		const int32 NumVerticesIn = InVertices.Num();
 		if (NumVerticesIn == 0)
@@ -40,11 +40,11 @@ namespace Chaos
 		}
 
 
-		const TArray<FVec3>* VerticesToUse = &InVertices;
-		TArray<FVec3> ModifiedVertices;
+		const TArray<FVec3Type>* VerticesToUse = &InVertices;
+		TArray<FVec3Type> ModifiedVertices;
 
 		// For triangles and planar shapes, create a very thin prism as a convex
-		auto Inflate = [](const TArray<FVec3>& Source, TArray<FVec3>& Destination, const FVec3& Normal, FReal Inflation)
+		auto Inflate = [](const TArray<FVec3Type>& Source, TArray<FVec3Type>& Destination, const FVec3Type& Normal, FRealType Inflation)
 		{
 			const int32 NumSource = Source.Num();
 			Destination.Reset();
@@ -57,7 +57,7 @@ namespace Chaos
 			}
 		};
 
-		FVec3 PlanarNormal(0);
+		FVec3Type PlanarNormal(0);
 		if (NumVerticesIn == 3)
 		{
 			const bool bIsValidTriangle = IsValidTriangle(InVertices[0], InVertices[1], InVertices[2], PlanarNormal);
@@ -83,7 +83,7 @@ namespace Chaos
 
 		const int32 NumVerticesToUse = VerticesToUse->Num();
 
-		OutLocalBounds = FAABB3((*VerticesToUse)[0], (*VerticesToUse)[0]);
+		OutLocalBounds = FAABB3Type((*VerticesToUse)[0], (*VerticesToUse)[0]);
 		for (int32 VertexIndex = 0; VertexIndex < NumVerticesToUse; ++VertexIndex)
 		{
 			OutLocalBounds.GrowToInclude((*VerticesToUse)[VertexIndex]);
@@ -98,8 +98,8 @@ namespace Chaos
 
 			if (bUseGeometryTConvexHull3) // Use the newer Geometry Tools code path for generating convex hulls.
 			{
-				UE::Geometry::TConvexHull3<FReal> HullCompute;
-				if (HullCompute.Solve<FVec3>(*VerticesToUse))
+				UE::Geometry::TConvexHull3<FRealType> HullCompute;
+				if (HullCompute.Solve<FVec3Type>(*VerticesToUse))
 				{
 					// Get Planes, FaceIndices, Vertices, and LocalBoundingBox
 					TMap<int32, int32> HullVertMap;
@@ -110,7 +110,7 @@ namespace Chaos
 								int32 Index = Triangle[j];
 								if (HullVertMap.Contains(Index) == false)
 								{
-									const FVec3& OrigPos = (*VerticesToUse)[Index];
+									const FVec3Type& OrigPos = (*VerticesToUse)[Index];
 									int32 NewVID = OutVertices.Num();
 									OutVertices.Add(OrigPos);
 									HullVertMap.Add(Index, NewVID);
@@ -125,13 +125,13 @@ namespace Chaos
 							// Winding is backwards from what Chaos expects
 							OutFaceIndices.Add({ Triangle[0], Triangle[2], Triangle[1] });
 
-							UE::Geometry::THalfspace3<FReal> Halfspace(OutVertices[Triangle[0]], OutVertices[Triangle[1]], OutVertices[Triangle[2]]);
-							const FVec3& Normal = Halfspace.Normal;
-							const FVec3& Point = OutVertices[Triangle[0]];
-							OutPlanes.Add(TPlaneConcrete<FReal, 3>{ Point, Normal });
+							UE::Geometry::THalfspace3<FRealType> Halfspace(OutVertices[Triangle[0]], OutVertices[Triangle[1]], OutVertices[Triangle[2]]);
+							const FVec3Type& Normal = Halfspace.Normal;
+							const FVec3Type& Point = OutVertices[Triangle[0]];
+							OutPlanes.Add(FPlaneType{ Point, Normal });
 						});
 
-					OutLocalBounds = FAABB3(OutVertices[0], OutVertices[0]);
+					OutLocalBounds = FAABB3Type(OutVertices[0], OutVertices[0]);
 					for (int32 VertexIndex = 1; VertexIndex < OutVertices.Num(); ++VertexIndex)
 					{
 						OutLocalBounds.GrowToInclude(OutVertices[VertexIndex]);
@@ -161,9 +161,9 @@ namespace Chaos
 
 				for (const TVec3<int32>& Idx : Indices)
 				{
-					FVec3 Vs[3] = { (*VerticesToUse)[Idx[0]], (*VerticesToUse)[Idx[1]], (*VerticesToUse)[Idx[2]] };
-					const FVec3 Normal = FVec3::CrossProduct(Vs[1] - Vs[0], Vs[2] - Vs[0]).GetUnsafeNormal();
-					OutPlanes.Add(TPlaneConcrete<FReal, 3>(Vs[0], Normal));
+					FVec3Type Vs[3] = { (*VerticesToUse)[Idx[0]], (*VerticesToUse)[Idx[1]], (*VerticesToUse)[Idx[2]] };
+					const FVec3Type Normal = FVec3Type::CrossProduct(Vs[1] - Vs[0], Vs[2] - Vs[0]).GetUnsafeNormal();
+					OutPlanes.Add(FPlaneType(Vs[0], Normal));
 					TArray<int32> FaceIndices;
 					FaceIndices.SetNum(3);
 					FaceIndices[0] = AddIndex(Idx[0]);
diff --git a/Engine/Source/Runtime/Experimental/Chaos/Private/Chaos/Convex.cpp b/Engine/Source/Runtime/Experimental/Chaos/Private/Chaos/Convex.cpp
--- a/Engine/Source/Runtime/Experimental/Chaos/Private/Chaos/Convex.cpp
+++ b/Engine/Source/Runtime/Experimental/Chaos/Private/Chaos/Convex.cpp
@@ -31,7 +31,7 @@ namespace Chaos
 		FReal MostOpposingDot = TNumericLimits<FReal>::Max();
 		for (int32 Idx = 0; Idx < Planes.Num(); ++Idx)
 		{
-			const TPlaneConcrete<FReal, 3>& Plane = Planes[Idx];
+			const FPlaneType& Plane = Planes[Idx];
 			const FReal Distance = Plane.SignedDistance(Position);
 			if (FMath::Abs(Distance) < SearchDist)
 			{
@@ -59,7 +59,7 @@ namespace Chaos
 		FReal MostOpposingDot = TNumericLimits<FReal>::Max();
 		for (int32 Idx = 0; Idx < Planes.Num(); ++Idx)
 		{
-			const TPlaneConcrete<FReal, 3> Plane = TPlaneConcrete<FReal, 3>::MakeScaledUnsafe(Planes[Idx], Scale);
+			const FPlaneType Plane = FPlaneType::MakeScaledUnsafe(Planes[Idx], Scale);
 			const FReal Distance = Plane.SignedDistance(Position);
 			if (FMath::Abs(Distance) < SearchDist)
 			{
@@ -90,7 +90,7 @@ namespace Chaos
 		TSet<int32> IncludedParticles;
 		for (int32 Idx = 0; Idx < Planes.Num(); ++Idx)
 		{
-			const TPlaneConcrete<FReal, 3>& Plane = Planes[Idx];
+			const FPlaneType& Plane = Planes[Idx];
 			FReal AbsOfSignedDistance = FMath::Abs(Plane.SignedDistance(Position));
 			if (AbsOfSignedDistance < SearchDist)
 			{
@@ -118,7 +118,7 @@ namespace Chaos
 		FReal MostOpposingDot = TNumericLimits<FReal>::Max();
 		for (int32 Idx = 0; Idx < Planes.Num(); ++Idx)
 		{
-			const TPlaneConcrete<FReal, 3>& Plane = Planes[Idx];
+			const FPlaneType& Plane = Planes[Idx];
 			const FReal Dot = FVec3::DotProduct(Plane.Normal(), Normal);
 			if (Dot < MostOpposingDot)
 			{
@@ -138,7 +138,7 @@ namespace Chaos
 		FReal MostOpposingDot = TNumericLimits<FReal>::Max();
 		for (int32 Idx = 0; Idx < Planes.Num(); ++Idx)
 		{
-			const TPlaneConcrete<FReal, 3>& Plane = Planes[Idx];
+			const FPlaneType& Plane = Planes[Idx];
 			const FVec3 ScaledNormal = (Plane.Normal() / Scale).GetSafeNormal();
 			const FReal Dot = FVec3::DotProduct(ScaledNormal, Normal);
 			if (Dot < MostOpposingDot)
@@ -285,10 +285,10 @@ namespace Chaos
 		StructureData.SetPlaneVertices(MoveTemp(PlaneVertexIndices), Vertices.Num());
 	}
 
-	void FConvex::MovePlanesAndRebuild(const FReal InDelta)
+	void FConvex::MovePlanesAndRebuild(const FRealType InDelta)
 	{
 		TArray<FPlane4f> NewPlanes;
-		TArray<FVec3> NewPoints;
+		TArray<FVec3Type> NewPoints;
 		NewPlanes.Reserve(Planes.Num());
 
 		// Move all the planes inwards
@@ -300,15 +300,15 @@ namespace Chaos
 
 		// Recalculate the set of points from the intersection of all combinations of 3 planes
 		// There will be NC3 of these (N! / (3! * (N-3)!)
-		const FReal PointTolerance = 1e-2f;
+		const FRealType PointTolerance = 1e-2f;
 		for (int32 PlaneIndex0 = 0; PlaneIndex0 < NewPlanes.Num(); ++PlaneIndex0)
 		{
 			for (int32 PlaneIndex1 = PlaneIndex0 + 1; PlaneIndex1 < NewPlanes.Num(); ++PlaneIndex1)
 			{
 				for (int32 PlaneIndex2 = PlaneIndex1 + 1; PlaneIndex2 < NewPlanes.Num(); ++PlaneIndex2)
 				{
-					FVec3 PlanesPos;
-					if (FMath::IntersectPlanes3<FReal>(PlanesPos, NewPlanes[PlaneIndex0], NewPlanes[PlaneIndex1], NewPlanes[PlaneIndex2]))
+					FVec3Type PlanesPos;
+					if (FMath::IntersectPlanes3<FRealType>(PlanesPos, NewPlanes[PlaneIndex0], NewPlanes[PlaneIndex1], NewPlanes[PlaneIndex2]))
 					{
 						// Reject duplicate points
 						int32 NewPointIndex = INDEX_NONE;
@@ -330,14 +330,14 @@ namespace Chaos
 		}
 
 		// Reject points outside the planes to get down to a sensible number for the build step
-		const FReal PointPlaneTolerance = PointTolerance;
+		const FRealType PointPlaneTolerance = PointTolerance;
 		int32 NumNewPoints = NewPoints.Num();
 		for (int32 PointIndex = 0; PointIndex < NumNewPoints; ++PointIndex)
 		{
 			const int32 NumNewPlanes = NewPlanes.Num();
 			for (int32 PlaneIndex = 0; PlaneIndex < NumNewPlanes; ++PlaneIndex)
 			{
-				const FReal PointPlaneDistance = NewPlanes[PlaneIndex].PlaneDot(NewPoints[PointIndex]);
+				const FRealType PointPlaneDistance = NewPlanes[PlaneIndex].PlaneDot(NewPoints[PointIndex]);
 				if (PointPlaneDistance > PointPlaneTolerance)
 				{
 					NewPoints.RemoveAtSwap(PointIndex);
diff --git a/Engine/Source/Runtime/Experimental/Chaos/Private/Chaos/Levelset.cpp b/Engine/Source/Runtime/Experimental/Chaos/Private/Chaos/Levelset.cpp
--- a/Engine/Source/Runtime/Experimental/Chaos/Private/Chaos/Levelset.cpp
+++ b/Engine/Source/Runtime/Experimental/Chaos/Private/Chaos/Levelset.cpp
@@ -1486,7 +1486,13 @@ void GetGeomSurfaceSamples(const FCapsule& InGeom, TArray<FVec3>& OutSamples)
 
 void GetGeomSurfaceSamples(const FConvex& InGeom, TArray<FVec3>& OutSamples)
 {
-	OutSamples = InGeom.GetVertices();
+	// because Convex store in single precision we need to convert to FReal precision
+	const TArray<FConvex::FVec3Type>& ConvexVertices = InGeom.GetVertices();
+	OutSamples.Reset(ConvexVertices.Num());
+	for (const FConvex::FVec3Type& Vertex : ConvexVertices)
+	{
+		OutSamples.Add(FVec3{ Vertex }); // conversion from single to double precision
+	}
 }
 
 template<typename InnerT>
diff --git a/Engine/Source/Runtime/Experimental/Chaos/Private/Chaos/MassProperties.cpp b/Engine/Source/Runtime/Experimental/Chaos/Private/Chaos/MassProperties.cpp
--- a/Engine/Source/Runtime/Experimental/Chaos/Private/Chaos/MassProperties.cpp
+++ b/Engine/Source/Runtime/Experimental/Chaos/Private/Chaos/MassProperties.cpp
@@ -87,8 +87,14 @@ namespace Chaos
 		return FinalRotation;
 	}
 
-	template<typename TSurfaces>
-	void CalculateVolumeAndCenterOfMass(const FParticles& Vertices, const TSurfaces& Surfaces, FReal& OutVolume, FVec3& OutCenterOfMass)
+	template<typename T, typename TSurfaces>
+	void CHAOS_API CalculateVolumeAndCenterOfMass(const TParticles<T, 3>& Vertices, const TSurfaces& Surfaces, T& OutVolume, TVec3<T>& OutCenterOfMass)
+	{
+		CalculateVolumeAndCenterOfMass(Vertices.AllX(), Surfaces, OutVolume, OutCenterOfMass);
+	}
+
+	template<typename T, typename TSurfaces>
+	void CalculateVolumeAndCenterOfMass(const TArray<TVec3<T>>& Vertices, const TSurfaces& Surfaces, T& OutVolume, TVec3<T>& OutCenterOfMass)
 	{
 		if (!Surfaces.Num())
 		{
@@ -96,19 +102,19 @@ namespace Chaos
 			return;
 		}
 		
-		FReal Volume = 0;
-		FVec3 VolumeTimesSum(0);
-		FVec3 Center = Vertices.X(Surfaces[0][0]);
+		T Volume = 0;
+		TVec3<T> VolumeTimesSum(0);
+		TVec3<T> Center = Vertices[Surfaces[0][0]];
 		for (const auto& Element : Surfaces)
 		{
 			// For now we only support triangular elements
 			ensure(Element.Num() == 3);
 
-			FMatrix33 DeltaMatrix;
-			FVec3 PerElementSize;
+			PMatrix<T,3,3> DeltaMatrix;
+			TVec3<T> PerElementSize;
 			for (int32 i = 0; i < Element.Num(); ++i)
 			{
-				FVec3 DeltaVector = Vertices.X(Element[i]) - Center;
+				TVec3<T> DeltaVector = Vertices[Element[i]] - Center;
 				DeltaMatrix.M[0][i] = DeltaVector[0];
 				DeltaMatrix.M[1][i] = DeltaVector[1];
 				DeltaMatrix.M[2][i] = DeltaVector[2];
@@ -116,7 +122,7 @@ namespace Chaos
 			PerElementSize[0] = DeltaMatrix.M[0][0] + DeltaMatrix.M[0][1] + DeltaMatrix.M[0][2];
 			PerElementSize[1] = DeltaMatrix.M[1][0] + DeltaMatrix.M[1][1] + DeltaMatrix.M[1][2];
 			PerElementSize[2] = DeltaMatrix.M[2][0] + DeltaMatrix.M[2][1] + DeltaMatrix.M[2][2];
-			FReal Det = DeltaMatrix.M[0][0] * (DeltaMatrix.M[1][1] * DeltaMatrix.M[2][2] - DeltaMatrix.M[1][2] * DeltaMatrix.M[2][1]) -
+			T Det = DeltaMatrix.M[0][0] * (DeltaMatrix.M[1][1] * DeltaMatrix.M[2][2] - DeltaMatrix.M[1][2] * DeltaMatrix.M[2][1]) -
 				DeltaMatrix.M[0][1] * (DeltaMatrix.M[1][0] * DeltaMatrix.M[2][2] - DeltaMatrix.M[1][2] * DeltaMatrix.M[2][0]) +
 				DeltaMatrix.M[0][2] * (DeltaMatrix.M[1][0] * DeltaMatrix.M[2][1] - DeltaMatrix.M[1][1] * DeltaMatrix.M[2][0]);
 			Volume += Det;
@@ -220,8 +226,12 @@ namespace Chaos
 	template CHAOS_API FMassProperties CalculateMassProperties(const FParticles& Vertices, const TArray<TVec3<int32>>& Surfaces, const FReal Mass);
 	template CHAOS_API FMassProperties CalculateMassProperties(const FParticles & Vertices, const TArray<TArray<int32>>& Surfaces, const FReal Mass);
 
-	template CHAOS_API void CalculateVolumeAndCenterOfMass(const FParticles& Vertices, const TArray<TVec3<int32>>& Surfaces, FReal& OutVolume, FVec3& OutCenterOfMass);
-	template CHAOS_API void CalculateVolumeAndCenterOfMass(const FParticles& Vertices, const TArray<TArray<int32>>& Surfaces, FReal& OutVolume, FVec3& OutCenterOfMass);
+	template CHAOS_API void CalculateVolumeAndCenterOfMass(const TParticles<FRealDouble, 3>& Vertices, const TArray<TVec3<int32>>& Surfaces, FRealDouble& OutVolume, TVec3<FRealDouble>& OutCenterOfMass);
+	template CHAOS_API void CalculateVolumeAndCenterOfMass(const TParticles<FRealDouble, 3>& Vertices, const TArray<TArray<int32>>& Surfaces, FRealDouble& OutVolume, TVec3<FRealDouble>& OutCenterOfMass);
+
+	template CHAOS_API void CalculateVolumeAndCenterOfMass(const TArray<TVec3<FRealSingle>>& Vertices, const TArray<TArray<int32>>& Surfaces, FRealSingle& OutVolume, TVec3<FRealSingle>& OutCenterOfMass);
+	template CHAOS_API void CalculateVolumeAndCenterOfMass(const TArray<TVec3<FRealDouble>>& Vertices, const TArray<TArray<int32>>& Surfaces, FRealDouble& OutVolume, TVec3<FRealDouble>& OutCenterOfMass);
+
 
 	template CHAOS_API void CalculateInertiaAndRotationOfMass(const FParticles& Vertices, const TArray<TVec3<int32>>& Surface, const FReal Density,
 		const FVec3& CenterOfMass, FMatrix33& OutInertiaTensor, FRotation3& OutRotationOfMass);
diff --git a/Engine/Source/Runtime/Experimental/Chaos/Private/GeometryCollection/GeometryCollectionCollisionStructureManager.cpp b/Engine/Source/Runtime/Experimental/Chaos/Private/GeometryCollection/GeometryCollectionCollisionStructureManager.cpp
--- a/Engine/Source/Runtime/Experimental/Chaos/Private/GeometryCollection/GeometryCollectionCollisionStructureManager.cpp
+++ b/Engine/Source/Runtime/Experimental/Chaos/Private/GeometryCollection/GeometryCollectionCollisionStructureManager.cpp
@@ -255,12 +255,14 @@ FCollisionStructureManager::NewImplicitConvex(
 	const FTransform& MassTransform,
 	const Chaos::FReal CollisionMarginFraction)
 {
+	using FConvexVec3 = Chaos::FConvex::FVec3Type;
+
 	if (ConvexIndices.Num())
 	{
 		TArray<TUniquePtr<Chaos::FImplicitObject>> Implicits;
 		for (auto& Index : ConvexIndices)
 		{
-			TArray<Chaos::FVec3> ConvexVertices = (*ConvexGeometry)[Index]->GetVertices();
+			TArray<FConvexVec3> ConvexVertices = (*ConvexGeometry)[Index]->GetVertices();
 			for (int32 Idx = 0; Idx < ConvexVertices.Num(); Idx++)
 			{
 				ConvexVertices[Idx] = MassTransform.InverseTransformPosition(ConvexVertices[Idx]);
diff --git a/Engine/Source/Runtime/Experimental/Chaos/Private/GeometryCollection/GeometryCollectionConvexUtility.cpp b/Engine/Source/Runtime/Experimental/Chaos/Private/GeometryCollection/GeometryCollectionConvexUtility.cpp
--- a/Engine/Source/Runtime/Experimental/Chaos/Private/GeometryCollection/GeometryCollectionConvexUtility.cpp
+++ b/Engine/Source/Runtime/Experimental/Chaos/Private/GeometryCollection/GeometryCollectionConvexUtility.cpp
@@ -70,7 +70,7 @@ namespace
 typedef Chaos::TPlaneConcrete<Chaos::FReal, 3> FChaosPlane;
 
 // filter points s.t. they are spaced at least more than SimplificationDistanceThreshold apart (after starting with the 4 'extreme' points to ensure we cover a volume)
-void FilterHullPoints(const TArray<Chaos::FVec3>& InPts, TArray<Chaos::FVec3>& OutPts, double SimplificationDistanceThreshold)
+void FilterHullPoints(const TArray<Chaos::FConvex::FVec3Type>& InPts, TArray<Chaos::FConvex::FVec3Type>& OutPts, double SimplificationDistanceThreshold)
 {
 	if (SimplificationDistanceThreshold > 0)
 	{
@@ -145,21 +145,21 @@ void FilterHullPoints(const TArray<Chaos::FVec3>& InPts, TArray<Chaos::FVec3>& O
 	}
 }
 
-void FilterHullPoints(TArray<Chaos::FVec3>& InOutPts, double SimplificationDistanceThreshold)
+void FilterHullPoints(TArray<Chaos::FConvex::FVec3Type>& InOutPts, double SimplificationDistanceThreshold)
 {
 	if (SimplificationDistanceThreshold > 0)
 	{
-		TArray<Chaos::FVec3> FilteredPts;
+		TArray<Chaos::FConvex::FVec3Type> FilteredPts;
 		FilterHullPoints(InOutPts, FilteredPts, SimplificationDistanceThreshold);
 		InOutPts = MoveTemp(FilteredPts);
 	}
 }
 
-Chaos::FConvex MakeHull(const TArray<Chaos::FVec3>& Pts, double SimplificationDistanceThreshold)
+Chaos::FConvex MakeHull(const TArray<Chaos::FConvex::FVec3Type>& Pts, double SimplificationDistanceThreshold)
 {
 	if (SimplificationDistanceThreshold > 0)
 	{
-		TArray<Chaos::FVec3> FilteredPts;
+		TArray<Chaos::FConvex::FVec3Type> FilteredPts;
 		FilterHullPoints(Pts, FilteredPts, SimplificationDistanceThreshold);
 		
 		return Chaos::FConvex(FilteredPts, KINDA_SMALL_NUMBER);
@@ -172,9 +172,9 @@ Chaos::FConvex MakeHull(const TArray<Chaos::FVec3>& Pts, double SimplificationDi
 
 /// Cut hull with plane, generating the point set of a new hull
 /// @return false if plane did not cut any points on the hull
-bool CutHull(const Chaos::FConvex& HullIn, FChaosPlane Plane, bool KeepSide, TArray<Chaos::FVec3>& HullPtsOut)
+bool CutHull(const Chaos::FConvex& HullIn, FChaosPlane Plane, bool KeepSide, TArray<Chaos::FConvex::FVec3Type>& HullPtsOut)
 {
-	const TArray<Chaos::FVec3>& Vertices = HullIn.GetVertices();
+	const TArray<Chaos::FConvex::FVec3Type>& Vertices = HullIn.GetVertices();
 	const Chaos::FConvexStructureData& HullData = HullIn.GetStructureData();
 	bool bHasOutside = false;
 	for (int VertIdx = 0; VertIdx < Vertices.Num(); VertIdx++)
@@ -223,7 +223,7 @@ bool CutHull(const Chaos::FConvex& HullIn, FChaosPlane Plane, bool KeepSide, TAr
 /// @return false if plane did not cut any points on the hull
 bool SplitHull(const Chaos::FConvex& HullIn, FChaosPlane Plane, bool KeepSide, TArray<Chaos::FVec3>& InsidePtsOut, TArray<Chaos::FVec3>& OutsidePtsOut)
 {
-	const TArray<Chaos::FVec3>& Vertices = HullIn.GetVertices();
+	const TArray<Chaos::FConvex::FVec3Type>& Vertices = HullIn.GetVertices();
 	const Chaos::FConvexStructureData& HullData = HullIn.GetStructureData();
 	bool bHasOutside = false;
 	for (int VertIdx = 0; VertIdx < Vertices.Num(); VertIdx++)
@@ -252,8 +252,8 @@ bool SplitHull(const Chaos::FConvex& HullIn, FChaosPlane Plane, bool KeepSide, T
 		for (int32 PlaneVertexIdx = 0; PlaneVertexIdx < NumPlaneVerts; PlaneVertexIdx++)
 		{
 			int32 NextVertIdx = (PlaneVertexIdx + 1) % NumPlaneVerts;
-			const Chaos::FVec3& V0 = Vertices[HullData.GetPlaneVertex(PlaneIdx, PlaneVertexIdx)];
-			const Chaos::FVec3& V1 = Vertices[HullData.GetPlaneVertex(PlaneIdx, NextVertIdx)];
+			const Chaos::FVec3 V0 = Vertices[HullData.GetPlaneVertex(PlaneIdx, PlaneVertexIdx)];
+			const Chaos::FVec3 V1 = Vertices[HullData.GetPlaneVertex(PlaneIdx, NextVertIdx)];
 			if ((Plane.SignedDistance(V0) < 0) != (Plane.SignedDistance(V1) < 0))
 			{
 				Chaos::Pair<Chaos::FVec3, bool> Res = Plane.FindClosestIntersection(V0, V1, 0);
@@ -476,7 +476,7 @@ void CreateNonoverlappingConvexHulls(
 				return false;
 			}
 			FChaosPlane CutPlane(IdealCenter, IdealNormal);
-			TArray<Chaos::FVec3> CutHullPts;
+			TArray<Chaos::FConvex::FVec3Type> CutHullPts;
 			if (CutHull(*Convexes[ConvexA], CutPlane, true, CutHullPts))
 			{
 				*Convexes[ConvexA] = MakeHull(CutHullPts, SimplificationDistanceThreshold);
@@ -584,7 +584,7 @@ void CreateNonoverlappingConvexHulls(
 			{
 				if (TransformToConvexIndices[Bone].Num() == 0)
 				{
-					TArray<Chaos::FVec3> JoinedHullPts;
+					TArray<Chaos::FConvex::FVec3Type> JoinedHullPts;
 					for (int32 Child : Children[Bone])
 					{
 						for (int32 ConvexIdx : TransformToConvexIndices[Child])
@@ -660,7 +660,7 @@ void CreateNonoverlappingConvexHulls(
 		{
 			if (Depths[Bone] == ProcessDepth && Children[Bone].Num() > 0 && !SkipBone(Bone))
 			{
-				TArray<Chaos::FVec3> JoinedHullPts;
+				TArray<Chaos::FConvex::FVec3Type> JoinedHullPts;
 				for (int32 Child : Children[Bone])
 				{
 					for (int32 ConvexIdx : TransformToConvexIndices[Child])
@@ -694,7 +694,7 @@ void CreateNonoverlappingConvexHulls(
 			// Tentatively create the clipped hulls
 			Chaos::FConvex CutHullA, CutHullB;
 			bool bCreatedA = false, bCreatedB = false;
-			TArray<Chaos::FVec3> CutHullPts;
+			TArray<Chaos::FConvex::FVec3Type> CutHullPts;
 			if (CutHull(*Convexes[ConvexA], CutPlane, true, CutHullPts))
 			{
 				if (CutHullPts.Num() < 4) // immediate reject zero-volume results
@@ -850,7 +850,7 @@ void HullsFromGeometry(
 			int32 VStart = Geometry.VertexStart[GeomIdx];
 			int32 VCount = Geometry.VertexCount[GeomIdx];
 			int32 VEnd = VStart + VCount;
-			TArray<Chaos::FVec3> HullPts;
+			TArray<Chaos::FConvex::FVec3Type> HullPts;
 			HullPts.Reserve(VCount);
 			for (int32 VIdx = VStart; VIdx < VEnd; VIdx++)
 			{
@@ -862,18 +862,18 @@ void HullsFromGeometry(
 			if (Convexes[ConvexIdx]->NumVertices() == 0 && HullPts.Num() > 0)
 			{
 				// if we've failed to make a convex hull, add a tiny bounding box just to ensure every geometry has a hull
-				Chaos::FAABB3 AABB = Convexes[ConvexIdx]->BoundingBox();
-				AABB.Thicken((Chaos::FReal).001);
-				Chaos::FVec3 Min = AABB.Min();
-				Chaos::FVec3 Max = AABB.Max();
+				Chaos::FConvex::FAABB3Type AABB = Convexes[ConvexIdx]->GetLocalBoundingBox();
+				AABB.Thicken(.001f);
+				Chaos::FConvex::FVec3Type Min = AABB.Min();
+				Chaos::FConvex::FVec3Type Max = AABB.Max();
 				HullPts.Add(Min);
 				HullPts.Add(Max);
-				HullPts.Add(Chaos::FVec3(Min.X, Min.Y, Max.Z));
-				HullPts.Add(Chaos::FVec3(Min.X, Max.Y, Max.Z));
-				HullPts.Add(Chaos::FVec3(Max.X, Min.Y, Max.Z));
-				HullPts.Add(Chaos::FVec3(Max.X, Max.Y, Min.Z));
-				HullPts.Add(Chaos::FVec3(Max.X, Min.Y, Min.Z));
-				HullPts.Add(Chaos::FVec3(Min.X, Max.Y, Min.Z));
+				HullPts.Add({Min.X, Min.Y, Max.Z});
+				HullPts.Add({Min.X, Max.Y, Max.Z});
+				HullPts.Add({Max.X, Min.Y, Max.Z});
+				HullPts.Add({Max.X, Max.Y, Min.Z});
+				HullPts.Add({Max.X, Min.Y, Min.Z});
+				HullPts.Add({Min.X, Max.Y, Min.Z});
 				// note: Do not use SimplificationDistanceThreshold for this fixed tiny hull
 				*Convexes[ConvexIdx] = Chaos::FConvex(HullPts, KINDA_SMALL_NUMBER);
 			}
@@ -892,11 +892,11 @@ void TransformHullsToLocal(
 	for (int32 Bone = 0; Bone < TransformToConvexIndices.Num(); Bone++)
 	{
 		FTransform& Transform = GlobalTransformArray[Bone];
-		TArray<Chaos::FVec3> HullPts;
+		TArray<Chaos::FConvex::FVec3Type> HullPts;
 		for (int32 ConvexIdx : TransformToConvexIndices[Bone])
 		{
 			HullPts.Reset();
-			for (const Chaos::FVec3& P : Convexes[ConvexIdx]->GetVertices())
+			for (const Chaos::FConvex::FVec3Type& P : Convexes[ConvexIdx]->GetVertices())
 			{
 				HullPts.Add(Transform.InverseTransformPosition(P));
 			}
@@ -960,7 +960,7 @@ TUniquePtr<Chaos::FConvex> FGeometryCollectionConvexUtility::FindConvexHull(cons
 	int32 VertexCount = GeometryCollection->VertexCount[GeometryIndex];
 	int32 VertexStart = GeometryCollection->VertexStart[GeometryIndex];
 
-	TArray<Chaos::FVec3> Vertices;
+	TArray<Chaos::FConvex::FVec3Type> Vertices;
 	Vertices.SetNum(VertexCount);
 	for (int32 VertexIndex = 0; VertexIndex < VertexCount; ++VertexIndex)
 	{
diff --git a/Engine/Source/Runtime/Experimental/Chaos/Public/Chaos/CollisionConvexMesh.h b/Engine/Source/Runtime/Experimental/Chaos/Public/Chaos/CollisionConvexMesh.h
--- a/Engine/Source/Runtime/Experimental/Chaos/Public/Chaos/CollisionConvexMesh.h
+++ b/Engine/Source/Runtime/Experimental/Chaos/Public/Chaos/CollisionConvexMesh.h
@@ -28,11 +28,16 @@ namespace Chaos
 	// a prism with a small thickness to emulate the desired result as a hull. Otherwise hull generation will fail on
 	// these cases. Verbose logging on LogChaos will point out when this path is taken for further scrutiny about
 	// the geometry
-	static constexpr FReal TriQuadPrismInflation() { return 0.1f; }
-	static constexpr FReal DefaultHorizonEpsilon() { return 0.1f; }
+	static constexpr FRealSingle TriQuadPrismInflation() { return 0.1f; }
+	static constexpr FRealSingle DefaultHorizonEpsilon() { return 0.1f; }
 
 	class FConvexBuilder
 	{
+		using FRealType = FRealSingle;
+		using FVec3Type = TVec3<FRealType>;
+		using FPlaneType = TPlaneConcrete<FRealType, 3>;
+		using FAABB3Type = TAABB<FRealType, 3>;
+
 	public:
 
 		class Params
@@ -42,10 +47,10 @@ namespace Chaos
 				: HorizonEpsilon(DefaultHorizonEpsilon())
 			{}
 
-			FReal HorizonEpsilon;
+			FRealType HorizonEpsilon;
 		};
 
-		static CHAOS_API FReal SuggestEpsilon(const TArray<FVec3>& InVertices)
+		static CHAOS_API FRealType SuggestEpsilon(const TArray<FVec3Type>& InVertices)
 		{
 			if (ComputeHorizonEpsilonFromMeshExtends == 0)
 			{
@@ -57,7 +62,7 @@ namespace Chaos
 			// above 1.0 such that 1.0 + FLT_EPSILON != 1.0. Floats aren't equally disanced though so big or small numbers
 			// don't work well with it. Here we take the max absolute of each axis and scale that for a wider margin and
 			// use that to scale FLT_EPSILON to get a more relevant value.
-			FVec3 MaxAxes(TNumericLimits<FReal>::Lowest());
+			FVec3Type MaxAxes(TNumericLimits<FRealType>::Lowest());
 			const int32 NumVertices = InVertices.Num();
 
 			if (NumVertices <= 1)
@@ -67,7 +72,7 @@ namespace Chaos
 
 			for (int32 Index = 0; Index < NumVertices; ++Index)
 			{
-				FVec3 PositionAbs = InVertices[Index].GetAbs();
+				FVec3Type PositionAbs = InVertices[Index].GetAbs();
 
 				MaxAxes[0] = FMath::Max(MaxAxes[0], PositionAbs[0]);
 				MaxAxes[1] = FMath::Max(MaxAxes[1], PositionAbs[1]);
@@ -77,30 +82,30 @@ namespace Chaos
 			return 3.0f * (MaxAxes[0] + MaxAxes[1] + MaxAxes[2]) * FLT_EPSILON;
 		}
 
-		static bool IsValidTriangle(const FVec3& A, const FVec3& B, const FVec3& C, FVec3& OutNormal)
+		static bool IsValidTriangle(const FVec3Type& A, const FVec3Type& B, const FVec3Type& C, FVec3Type& OutNormal)
 		{
-			const FVec3 BA = B - A;
-			const FVec3 CA = C - A;
-			const FVec3 Cross = FVec3::CrossProduct(BA, CA);
+			const FVec3Type BA = B - A;
+			const FVec3Type CA = C - A;
+			const FVec3Type Cross = FVec3Type::CrossProduct(BA, CA);
 			OutNormal = Cross.GetUnsafeNormal();
 			return Cross.Size() > 1e-4;
 		}
 
-		static bool IsValidTriangle(const FVec3& A, const FVec3& B, const FVec3& C)
+		static bool IsValidTriangle(const FVec3Type& A, const FVec3Type& B, const FVec3Type& C)
 		{
-			FVec3 Normal(0);
+			FVec3Type Normal(0);
 			return IsValidTriangle(A, B, C, Normal);
 		}
 
-		static bool IsValidQuad(const FVec3& A, const FVec3& B, const FVec3& C, const FVec3& D, FVec3& OutNormal)
+		static bool IsValidQuad(const FVec3Type& A, const FVec3Type& B, const FVec3Type& C, const FVec3Type& D, FVec3Type& OutNormal)
 		{
-			FPlane TriPlane(A, B, C);
-			const FReal DPointDistance = FMath::Abs(TriPlane.PlaneDot(D));
-			OutNormal = FVec3(TriPlane.X, TriPlane.Y, TriPlane.Z);
-			return FMath::IsNearlyEqual(DPointDistance, (FReal)0, (FReal)KINDA_SMALL_NUMBER);
+			const UE::Math::TPlane<FRealType> TriPlane(A, B, C);
+			const FRealType DPointDistance = FMath::Abs(TriPlane.PlaneDot(D));
+			OutNormal = FVec3Type(TriPlane.X, TriPlane.Y, TriPlane.Z);
+			return FMath::IsNearlyEqual(DPointDistance, FRealType(0), FRealType(KINDA_SMALL_NUMBER));
 		}
 		
-		static bool IsPlanarShape(const TArray<FVec3>& InVertices, FVec3& OutNormal)
+		static bool IsPlanarShape(const TArray<FVec3Type>& InVertices, FVec3Type& OutNormal)
 		{
 			bool bResult = false;
 			const int32 NumVertices = InVertices.Num();
@@ -112,13 +117,13 @@ namespace Chaos
 			}
 			else // > 3 points
 			{
-				FPlane TriPlane(InVertices[0], InVertices[1], InVertices[2]);
-				OutNormal = FVec3(TriPlane.X, TriPlane.Y, TriPlane.Z);
+				const UE::Math::TPlane<FRealType> TriPlane(InVertices[0], InVertices[1], InVertices[2]);
+				OutNormal = FVec3Type(TriPlane.X, TriPlane.Y, TriPlane.Z);
 
 				for (int32 Index = 3; Index < NumVertices; ++Index)
 				{
-					const FReal PointPlaneDot = FMath::Abs(TriPlane.PlaneDot(InVertices[Index]));
-					if(!FMath::IsNearlyEqual(PointPlaneDot, (FReal)0, (FReal)KINDA_SMALL_NUMBER))
+					const FRealType PointPlaneDot = FMath::Abs(TriPlane.PlaneDot(InVertices[Index]));
+					if(!FMath::IsNearlyEqual(PointPlaneDot, FRealType(0), FRealType(KINDA_SMALL_NUMBER)))
 					{
 						return false;
 					}
@@ -135,19 +140,19 @@ namespace Chaos
 		{
 			FHalfEdge* FirstEdge;
 			FHalfEdge* ConflictList; //Note that these half edges are really just free verts grouped together
-			TPlaneConcrete<FReal,3> Plane;
+			FPlaneType Plane;
 			FConvexFace* Prev;
 			FConvexFace* Next; //these have no geometric meaning, just used for book keeping
 			int32 PoolIdx;
 
 		private:
-			FConvexFace(int32 InPoolIdx, const TPlaneConcrete<FReal, 3>& FacePlane)
+			FConvexFace(int32 InPoolIdx, const FPlaneType& FacePlane)
 				: PoolIdx(InPoolIdx)
 			{
 				Reset(FacePlane);
 			}
 
-			void Reset(const TPlaneConcrete<FReal,3>& FacePlane)
+			void Reset(const FPlaneType& FacePlane)
 			{
 				ConflictList = nullptr;
 				Plane = FacePlane;
@@ -210,7 +215,7 @@ namespace Chaos
 				}
 			}
 
-			FConvexFace* AllocConvexFace(const TPlaneConcrete<FReal,3>& FacePlane)
+			FConvexFace* AllocConvexFace(const FPlaneType& FacePlane)
 			{
 				if(FacesFreeIndices.Num())
 				{
@@ -247,9 +252,9 @@ namespace Chaos
 
 	public:
 
-		static CHAOS_API void Build(const TArray<FVec3>& InVertices, TArray <TPlaneConcrete<FReal, 3>>& OutPlanes, TArray<TArray<int32>>& OutFaceIndices, TArray<FVec3>& OutVertices, FAABB3& OutLocalBounds);
+		static CHAOS_API void Build(const TArray<FVec3Type>& InVertices, TArray<FPlaneType>& OutPlanes, TArray<TArray<int32>>& OutFaceIndices, TArray<FVec3Type>& OutVertices, FAABB3Type& OutLocalBounds);
 
-		static void BuildConvexHull(const TArray<FVec3>& InVertices, TArray<TVec3<int32>>& OutIndices, const Params& InParams = Params())
+		static void BuildConvexHull(const TArray<FVec3Type>& InVertices, TArray<TVec3<int32>>& OutIndices, const Params& InParams = Params())
 		{
 			TRACE_CPUPROFILER_EVENT_SCOPE(Chaos::BuildConvexHull);
 
@@ -284,7 +289,7 @@ namespace Chaos
 #if DEBUG_HULL_GENERATION
 #if DEBUG_HULL_GENERATION_HULLSTEPS_TO_OBJ
 				UE_LOG(LogChaos, VeryVerbose, TEXT("# ======================================================"));
-				const FVec3 ConflictPos = InVertices[ConflictV->Vertex];
+				const FVec3Type ConflictPos = InVertices[ConflictV->Vertex];
 				UE_LOG(LogChaos, VeryVerbose, TEXT("# GENERATED HULL before adding Vtx %d (%f %f %f)"), ConflictV->Vertex, ConflictPos.X, ConflictPos.Y, ConflictPos.Z);
 				UE_LOG(LogChaos, VeryVerbose, TEXT("# ------------------------------------------------------"));
 				FConvexFace* Face = DummyFace->Next;
@@ -316,7 +321,7 @@ namespace Chaos
 			}
 		}
 
-		static FTriangleMesh BuildConvexHullTriMesh(const TArray<FVec3>& InVertices)
+		static FTriangleMesh BuildConvexHullTriMesh(const TArray<FVec3Type>& InVertices)
 		{
 			TArray<TVec3<int32>> Indices;
 			BuildConvexHull(InVertices, Indices);
@@ -343,7 +348,7 @@ namespace Chaos
 			return FString::Printf(TEXT("Planes %d, Vertices %d"), NumPlanes, NumVertices);
 		}
 
-		static CHAOS_API void Simplify(TArray <TPlaneConcrete<FReal, 3>>& InOutPlanes, TArray<TArray<int32>>& InOutFaces, TArray<FVec3>& InOutVertices, FAABB3& InOutLocalBounds)
+		static CHAOS_API void Simplify(TArray<FPlaneType>& InOutPlanes, TArray<TArray<int32>>& InOutFaces, TArray<FVec3Type>& InOutVertices, FAABB3Type& InOutLocalBounds)
 		{
 			struct TPair
 			{
@@ -357,7 +362,7 @@ namespace Chaos
 			int32 NumToDelete = NumberOfVerticesWeHave - NumberOfVerticesRequired;
 
 			int32 Size = InOutVertices.Num();
-			TArray<FVec3> Vertices(InOutVertices);
+			TArray<FVec3Type> Vertices(InOutVertices);
 
 			TArray<bool> IsDeleted;
 			IsDeleted.Reset();
@@ -368,7 +373,7 @@ namespace Chaos
 				for (uint32 Iteration = 0; Iteration < (uint32)NumToDelete; Iteration++)
 				{
 					TPair ClosestPair;
-					FReal ClosestDistSqr = FLT_MAX;
+					FRealType ClosestDistSqr = FLT_MAX;
 
 					for (int32 A = 0; A < (Size - 1); A++)
 					{
@@ -378,8 +383,8 @@ namespace Chaos
 							{
 								if (!IsDeleted[B])
 								{
-									FVec3 Vec = Vertices[A] - Vertices[B];
-									FReal LengthSqr = Vec.SizeSquared();
+									FVec3Type Vec = Vertices[A] - Vertices[B];
+									FRealType LengthSqr = Vec.SizeSquared();
 									if (LengthSqr < ClosestDistSqr)
 									{
 										ClosestDistSqr = LengthSqr;
@@ -400,7 +405,7 @@ namespace Chaos
 				}
 			}
 
-			TArray<FVec3> TmpVertices;
+			TArray<FVec3Type> TmpVertices;
 			for (int Idx = 0; Idx < Vertices.Num(); Idx++)
 			{
 				// Only add vertices that have not been merged away
@@ -415,31 +420,31 @@ namespace Chaos
 		}
 
 		// Convert multi-triangle faces to single n-gons
-		static void MergeFaces(TArray<TPlaneConcrete<FReal, 3>>& InOutPlanes, TArray<TArray<int32>>& InOutFaceVertexIndices, const TArray<FVec3>& Vertices, FReal DistanceThreshold)
+		static void MergeFaces(TArray<FPlaneType>& InOutPlanes, TArray<TArray<int32>>& InOutFaceVertexIndices, const TArray<FVec3Type>& Vertices, FRealType DistanceThreshold)
 		{
-			const FReal NormalThreshold = 1.e-4f;
+			const FRealType NormalThreshold = 1.e-4f;
 
 			// Find planes with equal normal within the threshold and merge them
 			for (int32 PlaneIndex0 = 0; PlaneIndex0 < InOutPlanes.Num(); ++PlaneIndex0)
 			{
-				const TPlaneConcrete<FReal, 3>& Plane0 = InOutPlanes[PlaneIndex0];
+				const FPlaneType& Plane0 = InOutPlanes[PlaneIndex0];
 				TArray<int32>& Vertices0 = InOutFaceVertexIndices[PlaneIndex0];
 
 				for (int32 PlaneIndex1 = PlaneIndex0 + 1; PlaneIndex1 < InOutPlanes.Num(); ++PlaneIndex1)
 				{
-					const TPlaneConcrete<FReal, 3>& Plane1 = InOutPlanes[PlaneIndex1];
+					const FPlaneType& Plane1 = InOutPlanes[PlaneIndex1];
 					const TArray<int32>& Vertices1 = InOutFaceVertexIndices[PlaneIndex1];
 
 					// First similarity test: normals are close - this will reject all very dissimilar faces
-					const FReal PlaneNormalDot = FVec3::DotProduct(Plane0.Normal(), Plane1.Normal());
+					const FRealType PlaneNormalDot = FVec3Type::DotProduct(Plane0.Normal(), Plane1.Normal());
 					if (PlaneNormalDot > 1.0f - NormalThreshold)
 					{
 						// Second similarity test: vertices of one plane are within threshold distance of the other. This is slower but more accurate
 						bool bWithinDistanceThreshold = true;
 						for (int32 Plane0VertexIndex : Vertices0)
 						{
-							const FVec3 Plane0Vertex = Vertices[Plane0VertexIndex];
-							const FReal Plane0VertexDistance = FMath::Abs(FVec3::DotProduct(Plane1.X() - Plane0Vertex, Plane1.Normal()));
+							const FVec3Type Plane0Vertex = Vertices[Plane0VertexIndex];
+							const FRealType Plane0VertexDistance = FMath::Abs(FVec3Type::DotProduct(Plane1.X() - Plane0Vertex, Plane1.Normal()));
 							if (Plane0VertexDistance > DistanceThreshold)
 							{
 								bWithinDistanceThreshold = false;
@@ -450,8 +455,8 @@ namespace Chaos
 						{
 							for (int32 Plane1VertexIndex : Vertices1)
 							{
-								const FVec3 Plane1Vertex = Vertices[Plane1VertexIndex];
-								const FReal Plane1VertexDistance = FMath::Abs(FVec3::DotProduct(Plane0.X() - Plane1Vertex, Plane0.Normal()));
+								const FVec3Type Plane1Vertex = Vertices[Plane1VertexIndex];
+								const FRealType Plane1VertexDistance = FMath::Abs(FVec3Type::DotProduct(Plane0.X() - Plane1Vertex, Plane0.Normal()));
 								if (Plane1VertexDistance > DistanceThreshold)
 								{
 									bWithinDistanceThreshold = false;
@@ -485,24 +490,24 @@ namespace Chaos
 		}
 
 		// Reorder the vertices to be counter-clockwise about the normal
-		static void SortFaceVerticesCCW(const TPlaneConcrete<FReal, 3>& Face, TArray<int32>& InOutFaceVertexIndices, const TArray<FVec3>& Vertices)
+		static void SortFaceVerticesCCW(const FPlaneType& Face, TArray<int32>& InOutFaceVertexIndices, const TArray<FVec3Type>& Vertices)
 		{
 			FMatrix33 FaceMatrix = (FMatrix44f)FRotationMatrix::MakeFromZ(Face.Normal());
 
-			FVec3 Centroid = FVec3(0);
+			FVec3Type Centroid(0);
 			for (int32 VertexIndex = 0; VertexIndex < InOutFaceVertexIndices.Num(); ++VertexIndex)
 			{
 				Centroid += Vertices[InOutFaceVertexIndices[VertexIndex]];
 			}
-			Centroid /= FReal(InOutFaceVertexIndices.Num());
+			Centroid /= FRealType(InOutFaceVertexIndices.Num());
 
 			// [2, -2] based on clockwise angle about the normal
 			auto VertexScore = [&Centroid, &FaceMatrix, &Vertices](int32 VertexIndex)
 			{
-				const FVec3 CentroidOffsetDir = (Vertices[VertexIndex] - Centroid).GetSafeNormal();
-				const FReal DotX = FVec3::DotProduct(CentroidOffsetDir, FaceMatrix.GetAxis(0));
-				const FReal DotY = FVec3::DotProduct(CentroidOffsetDir, FaceMatrix.GetAxis(1));
-				const FReal Score = (DotX >= 0.0f) ? 1.0f + DotY : -1.0f - DotY;
+				const FVec3Type CentroidOffsetDir = (Vertices[VertexIndex] - Centroid).GetSafeNormal();
+				const FRealType DotX = FVec3Type::DotProduct(CentroidOffsetDir, FaceMatrix.GetAxis(0));
+				const FRealType DotY = FVec3Type::DotProduct(CentroidOffsetDir, FaceMatrix.GetAxis(1));
+				const FRealType Score = (DotX >= 0.0f) ? 1.0f + DotY : -1.0f - DotY;
 				return Score;
 			};
 
@@ -515,14 +520,14 @@ namespace Chaos
 		}
 
 		// Generate the vertex indices for all planes in CCW order (used to serialize old data that did not have structure data)
-		static void BuildPlaneVertexIndices(const TArray<TPlaneConcrete<FReal, 3>>& InPlanes, const TArray<FVec3>& Vertices, TArray<TArray<int32>>& OutFaceVertexIndices, const FReal DistanceTolerance = 1.e-3f)
+		static void BuildPlaneVertexIndices(const TArray<FPlaneType>& InPlanes, const TArray<FVec3Type>& Vertices, TArray<TArray<int32>>& OutFaceVertexIndices, const FRealType DistanceTolerance = 1.e-3f)
 		{
 			OutFaceVertexIndices.Reset(InPlanes.Num());
 			for (int32 PlaneIndex = 0; PlaneIndex < InPlanes.Num(); ++PlaneIndex)
 			{
 				for (int32 VertexIndex = 0; VertexIndex < Vertices.Num(); ++VertexIndex)
 				{
-					const FReal PlaneVertexDistance = FVec3::DotProduct(InPlanes[PlaneIndex].Normal(), Vertices[VertexIndex] - InPlanes[PlaneIndex].X());
+					const FRealType PlaneVertexDistance = FVec3Type::DotProduct(InPlanes[PlaneIndex].Normal(), Vertices[VertexIndex] - InPlanes[PlaneIndex].X());
 					if (FMath::Abs(PlaneVertexDistance) < DistanceTolerance)
 					{
 						OutFaceVertexIndices[PlaneIndex].Add(VertexIndex);
@@ -542,12 +547,12 @@ namespace Chaos
 
 	private:
 
-		static FVec3 ComputeFaceNormal(const FVec3& A, const FVec3& B, const FVec3& C)
+		static FVec3Type ComputeFaceNormal(const FVec3Type& A, const FVec3Type& B, const FVec3Type& C)
 		{
-			return FVec3::CrossProduct((B - A), (C - A));
+			return FVec3Type::CrossProduct((B - A), (C - A));
 		}
 
-		static FConvexFace* CreateFace(FMemPool& Pool, const TArray<FVec3>& InVertices, FHalfEdge* RS, FHalfEdge* ST, FHalfEdge* TR)
+		static FConvexFace* CreateFace(FMemPool& Pool, const TArray<FVec3Type>& InVertices, FHalfEdge* RS, FHalfEdge* ST, FHalfEdge* TR)
 		{
 			RS->Prev = TR;
 			RS->Next = ST;
@@ -555,8 +560,8 @@ namespace Chaos
 			ST->Next = TR;
 			TR->Prev = ST;
 			TR->Next = RS;
-			FVec3 RSTNormal = ComputeFaceNormal(InVertices[RS->Vertex], InVertices[ST->Vertex], InVertices[TR->Vertex]);
-			const FReal RSTNormalSize = RSTNormal.Size();
+			FVec3Type RSTNormal = ComputeFaceNormal(InVertices[RS->Vertex], InVertices[ST->Vertex], InVertices[TR->Vertex]);
+			const FRealType RSTNormalSize = RSTNormal.Size();
 			if (RSTNormalSize > 0)
 			{
 				RSTNormal = RSTNormal * (1 / RSTNormalSize);
@@ -566,7 +571,7 @@ namespace Chaos
 				// degenerated face, use a valid neighbor face normal 
 				RSTNormal = RS->Twin->Face->Plane.Normal();
 			}
-			FConvexFace* RST = Pool.AllocConvexFace(TPlaneConcrete<FReal, 3>(InVertices[RS->Vertex], RSTNormal));
+			FConvexFace* RST = Pool.AllocConvexFace(FPlaneType(InVertices[RS->Vertex], RSTNormal));
 			RST->FirstEdge = RS;
 			RS->Face = RST;
 			ST->Face = RST;
@@ -574,16 +579,16 @@ namespace Chaos
 			return RST;
 		}
 
-		static void StealConflictList(FMemPool& Pool, const TArray<FVec3>& InVertices, FHalfEdge* OldList, FConvexFace** Faces, int32 NumFaces)
+		static void StealConflictList(FMemPool& Pool, const TArray<FVec3Type>& InVertices, FHalfEdge* OldList, FConvexFace** Faces, int32 NumFaces)
 		{
 			FHalfEdge* Cur = OldList;
 			while(Cur)
 			{
-				FReal MaxD = 1e-4f;
+				FRealType MaxD = 1e-4f;
 				int32 MaxIdx = -1;
 				for(int32 Idx = 0; Idx < NumFaces; ++Idx)
 				{
-					FReal Distance = Faces[Idx]->Plane.SignedDistance(InVertices[Cur->Vertex]);
+					FRealType Distance = Faces[Idx]->Plane.SignedDistance(InVertices[Cur->Vertex]);
 					if(Distance > MaxD)
 					{
 						MaxD = Distance;
@@ -596,7 +601,7 @@ namespace Chaos
 				{
 					//let's make sure faces created with this new conflict vertex will be valid. The plane check above is not sufficient because long thin triangles will have a plane with its point at one of these. Combined with normal and precision we can have errors
 					auto PretendNormal = [&InVertices](FHalfEdge* A, FHalfEdge* B, FHalfEdge* C) {
-						return FVec3::CrossProduct(InVertices[B->Vertex] - InVertices[A->Vertex], InVertices[C->Vertex] - InVertices[A->Vertex]).SizeSquared();
+						return FVec3Type::CrossProduct(InVertices[B->Vertex] - InVertices[A->Vertex], InVertices[C->Vertex] - InVertices[A->Vertex]).SizeSquared();
 					};
 					FHalfEdge* Edge = Faces[MaxIdx]->FirstEdge;
 					do
@@ -633,21 +638,21 @@ namespace Chaos
 			}
 		}
 
-		static FConvexFace* BuildInitialHull(FMemPool& Pool, const TArray<FVec3>& InVertices)
+		static FConvexFace* BuildInitialHull(FMemPool& Pool, const TArray<FVec3Type>& InVertices)
 		{
 			if (InVertices.Num() < 4) //not enough points
 			{
 				return nullptr;
 			}
 
-			constexpr FReal Epsilon = 1e-4f;
+			constexpr FRealType Epsilon = 1e-4f;
 
 			const int32 NumVertices = InVertices.Num();
 
 			//We store the vertex directly in the half-edge. We use its next to group free vertices by context list
 			//create a starting triangle by finding min/max on X and max on Y
-			FReal MinX = TNumericLimits<FReal>::Max();
-			FReal MaxX = TNumericLimits<FReal>::Lowest();
+			FRealType MinX = TNumericLimits<FRealType>::Max();
+			FRealType MaxX = TNumericLimits<FRealType>::Lowest();
 			FHalfEdge* A = nullptr; //min x
 			FHalfEdge* B = nullptr; //max x
 			FHalfEdge* DummyHalfEdge = Pool.AllocHalfEdge(-1);
@@ -661,7 +666,7 @@ namespace Chaos
 				Prev->Next = VHalf;
 				VHalf->Prev = Prev;
 				VHalf->Next = nullptr;
-				const FVec3& V = InVertices[i];
+				const FVec3Type& V = InVertices[i];
 
 				if(V[0] < MinX)
 				{
@@ -696,12 +701,12 @@ namespace Chaos
 			}
 
 			//find C so that we get the biggest base triangle
-			FReal MaxTriSize = Epsilon;
-			const FVec3 AToB = InVertices[B->Vertex] - InVertices[A->Vertex];
+			FRealType MaxTriSize = Epsilon;
+			const FVec3Type AToB = InVertices[B->Vertex] - InVertices[A->Vertex];
 			FHalfEdge* C = nullptr;
 			for(FHalfEdge* V = DummyHalfEdge->Next; V; V = V->Next)
 			{
-				FReal TriSize = FVec3::CrossProduct(AToB, InVertices[V->Vertex] - InVertices[A->Vertex]).SizeSquared();
+				FRealType TriSize = FVec3Type::CrossProduct(AToB, InVertices[V->Vertex] - InVertices[A->Vertex]).SizeSquared();
 				if(TriSize > MaxTriSize)
 				{
 					MaxTriSize = TriSize;
@@ -722,16 +727,16 @@ namespace Chaos
 			}
 
 			//find farthest D along normal
-			const FVec3 AToC = InVertices[C->Vertex] - InVertices[A->Vertex];
-			const FVec3 Normal = FVec3::CrossProduct(AToB, AToC).GetSafeNormal();
+			const FVec3Type AToC = InVertices[C->Vertex] - InVertices[A->Vertex];
+			const FVec3Type Normal = FVec3Type::CrossProduct(AToB, AToC).GetSafeNormal();
 
-			FReal MaxPosDistance = Epsilon;
-			FReal MaxNegDistance = Epsilon;
+			FRealType MaxPosDistance = Epsilon;
+			FRealType MaxNegDistance = Epsilon;
 			FHalfEdge* PosD = nullptr;
 			FHalfEdge* NegD = nullptr;
 			for(FHalfEdge* V = DummyHalfEdge->Next; V; V = V->Next)
 			{
-				FReal Dot = FVec3::DotProduct(InVertices[V->Vertex] - InVertices[A->Vertex], Normal);
+				FRealType Dot = FVec3Type::DotProduct(InVertices[V->Vertex] - InVertices[A->Vertex], Normal);
 				if(Dot > MaxPosDistance)
 				{
 					MaxPosDistance = Dot;
@@ -803,7 +808,7 @@ namespace Chaos
 			return Faces[0];
 		}
 
-		static FHalfEdge* FindConflictVertex(const TArray<FVec3>& InVertices, FConvexFace* FaceList)
+		static FHalfEdge* FindConflictVertex(const TArray<FVec3Type>& InVertices, FConvexFace* FaceList)
 		{
 			UE_CLOG(DEBUG_HULL_GENERATION, LogChaos, VeryVerbose, TEXT("Finding conflict vertex"));
 
@@ -811,12 +816,12 @@ namespace Chaos
 			{
 				UE_CLOG(DEBUG_HULL_GENERATION, LogChaos, VeryVerbose, TEXT("\tTesting Face (%d %d %d)"), CurFace->FirstEdge->Vertex, CurFace->FirstEdge->Next->Vertex, CurFace->FirstEdge->Prev->Vertex);
 
-				FReal MaxD = TNumericLimits<FReal>::Lowest();
+				FRealType MaxD = TNumericLimits<FRealType>::Lowest();
 				FHalfEdge* MaxV = nullptr;
 				for(FHalfEdge* CurFaceVertex = CurFace->ConflictList; CurFaceVertex; CurFaceVertex = CurFaceVertex->Next)
 				{
 					//is it faster to cache this from stealing stage?
-					FReal Dist = FVec3::DotProduct(InVertices[CurFaceVertex->Vertex], CurFace->Plane.Normal());
+					FRealType Dist = FVec3Type::DotProduct(InVertices[CurFaceVertex->Vertex], CurFace->Plane.Normal());
 					if(Dist > MaxD)
 					{
 						MaxD = Dist;
@@ -850,7 +855,7 @@ namespace Chaos
 			return nullptr;
 		}
 
-		static void BuildHorizon(const TArray<FVec3>& InVertices, FHalfEdge* ConflictV, TArray<FHalfEdge*>& HorizonEdges, TArray<FConvexFace*>& FacesToDelete, const Params& InParams)
+		static void BuildHorizon(const TArray<FVec3Type>& InVertices, FHalfEdge* ConflictV, TArray<FHalfEdge*>& HorizonEdges, TArray<FConvexFace*>& FacesToDelete, const Params& InParams)
 		{
 #if DEBUG_HULL_GENERATION
 			UE_LOG(LogChaos, VeryVerbose, TEXT("Generate horizon - START"));
@@ -858,8 +863,8 @@ namespace Chaos
 			//We must flood fill from the initial face and mark edges of faces the conflict vertex cannot see
 			//In order to return a CCW ordering we must traverse each face in CCW order from the edge we crossed over
 			//This should already be the ordering in the half edge
-			const FReal Epsilon = InParams.HorizonEpsilon;
-			const FVec3 V = InVertices[ConflictV->Vertex];
+			const FRealType Epsilon = InParams.HorizonEpsilon;
+			const FVec3Type V = InVertices[ConflictV->Vertex];
 			TSet<FConvexFace*> Processed;
 			TArray<FHalfEdge*> Queue;
 			check(ConflictV->Face);
@@ -895,7 +900,7 @@ namespace Chaos
 #endif
 					continue;
 				}
-				const FReal Distance = NextFace->Plane.SignedDistance(V);
+				const FRealType Distance = NextFace->Plane.SignedDistance(V);
 				if(Distance > Epsilon)
 				{
 #if DEBUG_HULL_GENERATION
@@ -928,7 +933,7 @@ namespace Chaos
 #if DEBUG_HULL_GENERATION
 #if DEBUG_HULL_GENERATION_BUILDHORIZON_TO_OBJ
 			UE_LOG(LogChaos, VeryVerbose, TEXT("# ======================================================"));
-			const FVec3 ConflictPos = InVertices[ConflictV->Vertex];
+			const FVec3Type ConflictPos = InVertices[ConflictV->Vertex];
 			UE_LOG(LogChaos, VeryVerbose, TEXT("# BUILD_HORIZON - Conflict Vertex = %d (%f %f %f)"), ConflictV->Vertex, ConflictPos.X, ConflictPos.Y, ConflictPos.Z);
 			UE_LOG(LogChaos, VeryVerbose, TEXT("# ------------------------------------------------------"));
 			for (TSet<FConvexFace*>::TConstIterator SetIt(Processed); SetIt; ++SetIt)
@@ -957,7 +962,7 @@ namespace Chaos
 
 		}
 
-		static void BuildFaces(FMemPool& Pool, const TArray<FVec3>& InVertices, const FHalfEdge* ConflictV, const TArray<FHalfEdge*>& HorizonEdges, const TArray<FConvexFace*> OldFaces, TArray<FConvexFace*>& NewFaces)
+		static void BuildFaces(FMemPool& Pool, const TArray<FVec3Type>& InVertices, const FHalfEdge* ConflictV, const TArray<FHalfEdge*>& HorizonEdges, const TArray<FConvexFace*> OldFaces, TArray<FConvexFace*>& NewFaces)
 		{
 			//The HorizonEdges are in CCW order. We must make new faces and edges to join from ConflictV to these edges
 			check(HorizonEdges.Num() >= 3);
@@ -1017,7 +1022,7 @@ namespace Chaos
 			StartFace->Prev = OldFace;
 		}
 
-		static void AddVertex(FMemPool& Pool, const TArray<FVec3>& InVertices, FHalfEdge* ConflictV, const Params& InParams)
+		static void AddVertex(FMemPool& Pool, const TArray<FVec3Type>& InVertices, FHalfEdge* ConflictV, const Params& InParams)
 		{
 			UE_CLOG(DEBUG_HULL_GENERATION, LogChaos, VeryVerbose, TEXT("Adding Vertex %d"), ConflictV->Vertex);
 
diff --git a/Engine/Source/Runtime/Experimental/Chaos/Public/Chaos/Convex.h b/Engine/Source/Runtime/Experimental/Chaos/Public/Chaos/Convex.h
--- a/Engine/Source/Runtime/Experimental/Chaos/Public/Chaos/Convex.h
+++ b/Engine/Source/Runtime/Experimental/Chaos/Public/Chaos/Convex.h
@@ -26,13 +26,19 @@ namespace Chaos
 	{
 	public:
 		using FImplicitObject::GetTypeName;
-		using TType = FReal;
+
+		using TType = FRealSingle;
+		using FRealType = TType;
+		using FVec3Type = TVec3<FRealType>;
+		using FPlaneType = TPlaneConcrete<FRealType, 3>;
+		using FAABB3Type = TAABB<FRealType, 3>;
+
 		static constexpr unsigned D = 3;
 
 		FConvex()
 		    : FImplicitObject(EImplicitObject::IsConvex | EImplicitObject::HasBoundingBox, ImplicitObjectType::Convex)
 			, Volume(0.f)
-			, CenterOfMass(FVec3(0.f))
+			, CenterOfMass(FVec3Type(0.f))
 		{}
 		FConvex(const FConvex&) = delete;
 		FConvex(FConvex&& Other)
@@ -52,11 +58,10 @@ namespace Chaos
 		// @todo(chaos): Add plane vertex indices in the constructor and call CreateStructureData
 		// @todo(chaos): Merge planes? Or assume the input is a good convex hull?
 		UE_DEPRECATED(4.27, "Use the constructor version with the face indices.")
-		FConvex(TArray<TPlaneConcrete<FReal, 3>>&& InPlanes, TArray<FVec3>&& InVertices)
+		FConvex(TArray<FPlaneType>&& InPlanes, TArray<FVec3Type>&& InVertices)
 		    : FImplicitObject(EImplicitObject::IsConvex | EImplicitObject::HasBoundingBox, ImplicitObjectType::Convex)
 			, Planes(MoveTemp(InPlanes))
 		    , Vertices(MoveTemp(InVertices))
-		    , LocalBoundingBox(FAABB3::EmptyAABB())
 		{
 			for (int32 ParticleIndex = 0; ParticleIndex < Vertices.Num(); ++ParticleIndex)
 			{
@@ -68,11 +73,10 @@ namespace Chaos
 			Volume = LocalBoundingBox.GetVolume();
 		}
 
-		FConvex(TArray<TPlaneConcrete<FReal, 3>>&& InPlanes, TArray<TArray<int32>>&& InFaceIndices, TArray<FVec3>&& InVertices)
+		FConvex(TArray<FPlaneType>&& InPlanes, TArray<TArray<int32>>&& InFaceIndices, TArray<FVec3Type>&& InVertices)
 		    : FImplicitObject(EImplicitObject::IsConvex | EImplicitObject::HasBoundingBox, ImplicitObjectType::Convex)
 			, Planes(MoveTemp(InPlanes))
 		    , Vertices(MoveTemp(InVertices))
-		    , LocalBoundingBox(FAABB3::EmptyAABB())
 		{
 			for (int32 ParticleIndex = 0; ParticleIndex < Vertices.Num(); ++ParticleIndex)
 			{
@@ -86,7 +90,7 @@ namespace Chaos
 			CreateStructureData(MoveTemp(InFaceIndices));
 		}
 
-		FConvex(const TArray<FVec3>& InVertices, const FReal InMargin)
+		FConvex(const TArray<FVec3Type>& InVertices, const FReal InMargin)
 		    : FImplicitObject(EImplicitObject::IsConvex | EImplicitObject::HasBoundingBox, ImplicitObjectType::Convex)
 		{
 			const int32 NumVertices = InVertices.Num();
@@ -100,8 +104,7 @@ namespace Chaos
 			CHAOS_ENSURE(Planes.Num() == FaceIndices.Num());
 
 			// @todo(chaos): this only works with triangles. Fix that an we can run MergeFaces before calling this
-			TParticles<FReal, 3> VertexParticles(CopyTemp(Vertices));
-			CalculateVolumeAndCenterOfMass(VertexParticles, FaceIndices, Volume, CenterOfMass);
+			CalculateVolumeAndCenterOfMass(Vertices, FaceIndices, Volume, CenterOfMass);
 
 			// @todo(chaos):																				 should be based on size, or passed in
 			if (!FConvexBuilder::bUseGeometryTConvexHull3)
@@ -110,7 +113,7 @@ namespace Chaos
 				// it appears that this code path can leave the convex in an
 				// undefined state. We should consider removing the merge faces when
 				// we transition to the UE::Geometry convex generation. 
-				const FReal DistanceTolerance = 1.0f;
+				const FRealType DistanceTolerance = 1.0f;
 				FConvexBuilder::MergeFaces(Planes, FaceIndices, Vertices, DistanceTolerance);
 				CHAOS_ENSURE(Planes.Num() == FaceIndices.Num());
 			}
@@ -153,7 +156,7 @@ namespace Chaos
 
 		virtual TUniquePtr<FImplicitObject> CopyWithScale(const FVec3& Scale) const override;
 
-		void MovePlanesAndRebuild(FReal InDelta);
+		void MovePlanesAndRebuild(FRealType InDelta);
 
 	private:
 		void CreateStructureData(TArray<TArray<int32>>&& FaceIndices);
@@ -175,6 +178,12 @@ namespace Chaos
 		}
 
 		virtual const FAABB3 BoundingBox() const override
+		{
+			// LWC : conversion from float to double
+			return FAABB3(LocalBoundingBox.Min(), LocalBoundingBox.Max());
+		}
+
+		FConvex::FAABB3Type GetLocalBoundingBox() const 
 		{
 			return LocalBoundingBox;
 		}
@@ -403,13 +412,14 @@ namespace Chaos
 		}
 
 		// Get the plane at the specified index (e.g., indices from FindVertexPlanes)
-		const TPlaneConcrete<FReal, 3>& GetPlane(int32 FaceIndex) const
+		const TPlaneConcrete<FReal, 3> GetPlane(int32 FaceIndex) const
 		{
-			return Planes[FaceIndex];
+			// @todo(chaos) this is needed because this API is shared with BOx implicit - we shoudl eventually only need local space planes and be able to use single precision
+			return TPlaneConcrete<FReal, 3>::MakeFrom(Planes[FaceIndex]);
 		}
 
 		// Get the vertex at the specified index (e.g., indices from GetPlaneVertexs)
-		const FVec3& GetVertex(int32 VertexIndex) const
+		const FVec3Type& GetVertex(int32 VertexIndex) const
 		{
 			return Vertices[VertexIndex];
 		}
@@ -426,7 +436,7 @@ namespace Chaos
 			// passes in a valid face index.
 			if (CHAOS_ENSURE(FaceIndex != INDEX_NONE))
 			{
-				const TPlaneConcrete<FReal, 3>& OpposingFace = GetFaces()[FaceIndex];
+				const FPlaneType& OpposingFace = GetFaces()[FaceIndex];
 				return OpposingFace.Normal();
 			}
 			return FVec3(0.f, 0.f, 1.f);
@@ -670,12 +680,12 @@ namespace Chaos
 			return FString::Printf(TEXT("Convex"));
 		}
 
-		const TArray<FVec3>& GetVertices() const
+		const TArray<FVec3Type>& GetVertices() const
 		{
 			return Vertices;
 		}
 
-		const TArray<TPlaneConcrete<FReal, 3>>& GetFaces() const
+		const TArray<FPlaneType>& GetFaces() const
 		{
 			return Planes;
 		}
@@ -705,14 +715,14 @@ namespace Chaos
 		{
 			uint32 Result = LocalBoundingBox.GetTypeHash();
 
-			for (const FVec3& Vertex: Vertices)
+			for (const FVec3Type& Vertex: Vertices)
 			{
 				Result = HashCombine(Result, ::GetTypeHash(Vertex[0]));
 				Result = HashCombine(Result, ::GetTypeHash(Vertex[1]));
 				Result = HashCombine(Result, ::GetTypeHash(Vertex[2]));
 			}
 
-			for(const TPlaneConcrete<FReal, 3>& Plane : Planes)
+			for(const FPlaneType& Plane : Planes)
 			{
 				Result = HashCombine(Result, Plane.GetTypeHash());
 			}
@@ -720,7 +730,7 @@ namespace Chaos
 			return Result;
 		}
 
-		FORCEINLINE void SerializeImp(FArchive& Ar)
+		FORCEINLINE void SerializeImp(FArchive& Ar)	
 		{
 			Ar.UsingCustomVersion(FExternalPhysicsCustomObjectVersion::GUID);
 			Ar.UsingCustomVersion(FUE5MainStreamObjectVersion::GUID);
@@ -730,7 +740,7 @@ namespace Chaos
 
 			if (Ar.CustomVer(FExternalPhysicsCustomObjectVersion::GUID) < FExternalPhysicsCustomObjectVersion::ConvexUsesTPlaneConcrete)
 			{
-				TArray<TPlane<FReal, 3>> TmpPlanes;
+				TArray<TPlane<FRealType, 3>> TmpPlanes;
 				Ar << TmpPlanes;
 
 				Planes.SetNum(TmpPlanes.Num());
@@ -755,7 +765,7 @@ namespace Chaos
 
 			if (!bConvexVerticesNewFormat)
 			{
-				TParticles<FReal, 3> TmpSurfaceParticles;
+				TParticles<FRealType, 3> TmpSurfaceParticles;
 				Ar << TmpSurfaceParticles;
 
 				const int32 NumVertices = (int32)TmpSurfaceParticles.Size();
@@ -771,13 +781,13 @@ namespace Chaos
 			}
 			
 
-			TBox<FReal,3>::SerializeAsAABB(Ar, LocalBoundingBox);
+			TBox<FRealType,3>::SerializeAsAABB(Ar, LocalBoundingBox);
 
 			if (Ar.CustomVer(FExternalPhysicsCustomObjectVersion::GUID) >= FExternalPhysicsCustomObjectVersion::AddConvexCenterOfMassAndVolume)
 			{
-				FRealSingle VolumeFloat = (FRealSingle)Volume; // LWC_TODO : potential precision loss, to be changed when we can serialize FReal as double
+				FRealSingle VolumeFloat = (FRealSingle)Volume; // LWC_TODO : potential precision lossdepending on FRealType, to be changed when we can serialize FReal as double
 				Ar << VolumeFloat;
-				Volume = (FReal)VolumeFloat;
+				Volume = (FRealType)VolumeFloat;
 
 				Ar << CenterOfMass;
 			}
@@ -786,14 +796,11 @@ namespace Chaos
 				// Rebuild convex in order to extract face indices.
 				// @todo(chaos): Make it so it can take Vertices as both input and output without breaking...
 				TArray<TArray<int32>> FaceIndices;
-				TArray<FVec3> TempVertices;
+				TArray<FVec3Type> TempVertices;
 				FConvexBuilder::Build(Vertices, Planes, FaceIndices, TempVertices, LocalBoundingBox);
 
 				// Copy vertices and move into particles.
-				// @todo(chaos): make CalculateVolumeAndCenterOfMass take array of positions rather than particles
-				TArray<FVec3> VerticesCopy = Vertices;
-				const FParticles SurfaceParticles(MoveTemp(VerticesCopy));
-				CalculateVolumeAndCenterOfMass(SurfaceParticles, FaceIndices, Volume, CenterOfMass);
+				CalculateVolumeAndCenterOfMass(Vertices, FaceIndices, Volume, CenterOfMass);
 			}
 
 			Ar.UsingCustomVersion(FReleaseObjectVersion::GUID);
@@ -857,7 +864,7 @@ namespace Chaos
 			FConvexBuilder::Simplify(Planes, FaceIndices, Vertices, LocalBoundingBox);
 
 			// @todo(chaos): DistanceTolerance should be based on size, or passed in
-			const FReal DistanceTolerance = 1.0f;
+			const FRealType DistanceTolerance = 1.0f;
 			FConvexBuilder::MergeFaces(Planes, FaceIndices, Vertices, DistanceTolerance);
 
 			CreateStructureData(MoveTemp(FaceIndices));
@@ -869,11 +876,11 @@ namespace Chaos
 		}
 
 	private:
-		TArray<TPlaneConcrete<FReal, 3>> Planes;
-		TArray<FVec3> Vertices; //copy of the vertices that are just on the convex hull boundary
-		FAABB3 LocalBoundingBox;
+		TArray<FPlaneType> Planes;
+		TArray<FVec3Type> Vertices; //copy of the vertices that are just on the convex hull boundary
+		FAABB3Type LocalBoundingBox;
 		FConvexStructureData StructureData;
-		FReal Volume;
-		FVec3 CenterOfMass;
+		FRealType Volume;
+		FVec3Type CenterOfMass;
 	};
 }
diff --git a/Engine/Source/Runtime/Experimental/Chaos/Public/Chaos/ImplicitObjectScaled.h b/Engine/Source/Runtime/Experimental/Chaos/Public/Chaos/ImplicitObjectScaled.h
--- a/Engine/Source/Runtime/Experimental/Chaos/Public/Chaos/ImplicitObjectScaled.h
+++ b/Engine/Source/Runtime/Experimental/Chaos/Public/Chaos/ImplicitObjectScaled.h
@@ -127,12 +127,12 @@ public:
 		return MObject->GetDoCollide();
 	}
 
-	virtual T PhiWithNormal(const TVector<T, d>& X, TVector<T, d>& Normal) const override
+	virtual FReal PhiWithNormal(const FVec3& X, FVec3& Normal) const override
 	{
 		return MObject->PhiWithNormal(X, Normal);
 	}
 
-	virtual bool Raycast(const TVector<T, d>& StartPoint, const TVector<T, d>& Dir, const T Length, const T Thickness, T& OutTime, TVector<T, d>& OutPosition, TVector<T, d>& OutNormal, int32& OutFaceIndex) const override
+	virtual bool Raycast(const FVec3& StartPoint, const FVec3& Dir, const FReal Length, const FReal Thickness, FReal& OutTime, FVec3& OutPosition, FVec3& OutNormal, int32& OutFaceIndex) const override
 	{
 		return MObject->Raycast(StartPoint, Dir, Length, Thickness, OutTime, OutPosition, OutNormal, OutFaceIndex);
 	}
@@ -144,40 +144,40 @@ public:
 		Ar << MObject;
 	}
 
-	virtual int32 FindMostOpposingFace(const TVector<T, d>& Position, const TVector<T, d>& UnitDir, int32 HintFaceIndex, T SearchDist) const override
+	virtual int32 FindMostOpposingFace(const FVec3& Position, const FVec3& UnitDir, int32 HintFaceIndex, FReal SearchDist) const override
 	{
 		return MObject->FindMostOpposingFace(Position, UnitDir, HintFaceIndex, SearchDist);
 	}
 
-	virtual TVector<T, 3> FindGeometryOpposingNormal(const TVector<T, d>& DenormDir, int32 HintFaceIndex, const TVector<T, d>& OriginalNormal) const override
+	virtual FVec3 FindGeometryOpposingNormal(const FVec3& DenormDir, int32 HintFaceIndex, const FVec3& OriginalNormal) const override
 	{
 		return MObject->FindGeometryOpposingNormal(DenormDir, HintFaceIndex, OriginalNormal);
 	}
 
-	virtual bool Overlap(const TVector<T, d>& Point, const T Thickness) const override
+	virtual bool Overlap(const FVec3& Point, const FReal Thickness) const override
 	{
 		return MObject->Overlap(Point, Thickness);
 	}
 
 	// The support position from the specified direction
-	FORCEINLINE TVector<T, d> Support(const TVector<T, d>& Direction, const T Thickness, int32& VertexIndex) const
+	FORCEINLINE FVec3 Support(const FVec3& Direction, const FReal Thickness, int32& VertexIndex) const
 	{
 		return MObject->Support(Direction, Thickness, VertexIndex); 
 	}
 
 	// this shouldn't be called, but is required until we remove the explicit function implementations in CollisionResolution.cpp
-	FORCEINLINE TVector<T, d> SupportScaled(const TVector<T, d>& Direction, const T Thickness, const FVec3& Scale, int32& VertexIndex) const
+	FORCEINLINE FVec3 SupportScaled(const FVec3& Direction, const T Thickness, const FVec3& Scale, int32& VertexIndex) const
 	{
 		return MObject->SupportScaled(Direction, Thickness, Scale, VertexIndex);
 	}
 
 	// The support position from the specified direction, if the shape is reduced by the margin
-	FORCEINLINE TVector<T, d> SupportCore(const TVector<T, d>& Direction, const FReal InMargin, FReal* OutSupportDelta, int32& VertexIndex) const
+	FORCEINLINE FVec3 SupportCore(const FVec3& Direction, const FReal InMargin, FReal* OutSupportDelta, int32& VertexIndex) const
 	{
 		return MObject->SupportCore(Direction, InMargin, OutSupportDelta, VertexIndex);
 	}
 
-	virtual const TAABB<T, d> BoundingBox() const override 
+	virtual const FAABB3 BoundingBox() const override 
 	{ 
 		return MObject->BoundingBox();
 	}
@@ -388,7 +388,7 @@ public:
 	using ObjectType = typename TChooseClass<bInstanced, TSerializablePtr<TConcrete>, TUniquePtr<TConcrete>>::Result;
 	using FImplicitObject::GetTypeName;
 
-	TImplicitObjectScaled(ObjectType Object, const TVector<T, d>& Scale, T InMargin = 0)
+	TImplicitObjectScaled(ObjectType Object, const FVec3& Scale, FReal InMargin = 0)
 	    : FImplicitObjectScaled(EImplicitObject::HasBoundingBox, Object->GetType())
 	    , MObject(MoveTemp(Object))
 		, MSharedPtrForRefCount(nullptr)
@@ -409,7 +409,7 @@ public:
 		SetScale(Scale);
 	}
 
-	TImplicitObjectScaled(TSharedPtr<TConcrete, ESPMode::ThreadSafe> Object, const TVector<T, d>& Scale, T InMargin = 0)
+	TImplicitObjectScaled(TSharedPtr<TConcrete, ESPMode::ThreadSafe> Object, const FVec3& Scale, FReal InMargin = 0)
 	    : FImplicitObjectScaled(EImplicitObject::HasBoundingBox, Object->GetType())
 	    , MObject(MakeSerializable<TConcrete, ESPMode::ThreadSafe>(Object))
 		, MSharedPtrForRefCount(Object)
@@ -430,7 +430,7 @@ public:
 		SetScale(Scale);
 	}
 
-	TImplicitObjectScaled(ObjectType Object, TUniquePtr<Chaos::FImplicitObject> &&ObjectOwner, const TVector<T, d>& Scale, T InMargin = 0)
+	TImplicitObjectScaled(ObjectType Object, TUniquePtr<Chaos::FImplicitObject> &&ObjectOwner, const FVec3& Scale, FReal InMargin = 0)
 	    : FImplicitObjectScaled(EImplicitObject::HasBoundingBox, Object->GetType())
 	    , MObject(Object)
 		, MSharedPtrForRefCount(nullptr)
@@ -550,34 +550,34 @@ public:
 		return (MObject->GetRadius() > 0.0f) ? Margin : 0.0f;
 	}
 	
-	virtual T PhiWithNormal(const TVector<T, d>& X, TVector<T, d>& Normal) const override
+	virtual FReal PhiWithNormal(const FVec3& X, FVec3& Normal) const override
 	{
 		return MObject->PhiWithNormalScaled(X, MScale, Normal);
 	}
 
-	virtual bool Raycast(const TVector<T, d>& StartPoint, const TVector<T, d>& Dir, const T Length, const T Thickness, T& OutTime, TVector<T, d>& OutPosition, TVector<T, d>& OutNormal, int32& OutFaceIndex) const override
+	virtual bool Raycast(const FVec3& StartPoint, const FVec3& Dir, const FReal Length, const FReal Thickness, FReal& OutTime, FVec3& OutPosition, FVec3& OutNormal, int32& OutFaceIndex) const override
 	{
 		ensure(Length > 0);
 		ensure(FMath::IsNearlyEqual(Dir.SizeSquared(), (FReal)1, (FReal)KINDA_SMALL_NUMBER));
 		ensure(Thickness == 0 || (FMath::IsNearlyEqual(MScale[0], MScale[1]) && FMath::IsNearlyEqual(MScale[0], MScale[2])));	//non uniform turns sphere into an ellipsoid so no longer a raycast and requires a more expensive sweep
 
-		const TVector<T, d> UnscaledStart = MInvScale * StartPoint;
-		const TVector<T, d> UnscaledDirDenorm = MInvScale * Dir;
-		const T LengthScale = UnscaledDirDenorm.Size();
-		if (ensure(LengthScale > TNumericLimits<T>::Min()))
+		const FVec3 UnscaledStart = MInvScale * StartPoint;
+		const FVec3 UnscaledDirDenorm = MInvScale * Dir;
+		const FReal LengthScale = UnscaledDirDenorm.Size();
+		if (ensure(LengthScale > TNumericLimits<FReal>::Min()))
 		{
-			const T LengthScaleInv = 1.f / LengthScale;
-			const T UnscaledLength = Length * LengthScale;
-			const TVector<T, d> UnscaledDir = UnscaledDirDenorm * LengthScaleInv;
+			const FReal LengthScaleInv = FReal(1) / LengthScale;
+			const FReal UnscaledLength = Length * LengthScale;
+			const FVec3 UnscaledDir = UnscaledDirDenorm * LengthScaleInv;
 			
-			TVector<T, d> UnscaledPosition;
-			TVector<T, d> UnscaledNormal;
-			T UnscaledTime;
+			FVec3 UnscaledPosition;
+			FVec3 UnscaledNormal;
+			FReal UnscaledTime;
 
 			if (MObject->Raycast(UnscaledStart, UnscaledDir, UnscaledLength, Thickness * MInvScale[0], UnscaledTime, UnscaledPosition, UnscaledNormal, OutFaceIndex))
 			{
 				//We double check that NewTime < Length because of potential precision issues. When that happens we always keep the shortest hit first
-				const T NewTime = LengthScaleInv * UnscaledTime;
+				const FReal NewTime = LengthScaleInv * UnscaledTime;
 				if (NewTime == 0) // Normal/Position output may be uninitialized with TOI 0 so we use ray information for that as the ray origin is likely inside the shape
 				{
 					OutPosition = StartPoint;
@@ -588,7 +588,7 @@ public:
 				else if (NewTime < Length) 
 				{
 					OutPosition = MScale * UnscaledPosition;
-					OutNormal = (MInvScale * UnscaledNormal).GetSafeNormal(TNumericLimits<T>::Min());
+					OutNormal = (MInvScale * UnscaledNormal).GetSafeNormal(TNumericLimits<FReal>::Min());
 					OutTime = NewTime;
 					return true;
 				}
@@ -761,55 +761,55 @@ public:
 	}
 
 
-	virtual int32 FindMostOpposingFace(const TVector<T, d>& Position, const TVector<T, d>& UnitDir, int32 HintFaceIndex, T SearchDist) const override
+	virtual int32 FindMostOpposingFace(const FVec3& Position, const FVec3& UnitDir, int32 HintFaceIndex, FReal SearchDist) const override
 	{
-		ensure(FMath::IsNearlyEqual(UnitDir.SizeSquared(), (FReal)1, (FReal)KINDA_SMALL_NUMBER));
+		ensure(FMath::IsNearlyEqual(UnitDir.SizeSquared(), FReal(1), FReal(KINDA_SMALL_NUMBER)));
 		return MObject->FindMostOpposingFaceScaled(Position, UnitDir, HintFaceIndex, SearchDist, MScale);
 	}
 
-	virtual TVector<T, 3> FindGeometryOpposingNormal(const TVector<T, d>& DenormDir, int32 HintFaceIndex, const TVector<T, d>& OriginalNormal) const override
+	virtual FVec3 FindGeometryOpposingNormal(const FVec3& DenormDir, int32 HintFaceIndex, const FVec3& OriginalNormal) const override
 	{
 		// @todo(chaos): we need a virtual FindGeometryOpposingNormal. Some types (Convex, Box) can just return the
 		// normal from the face index without needing calculating the unscaled normals.
-		ensure(FMath::IsNearlyEqual(OriginalNormal.SizeSquared(), (FReal)1, (FReal)KINDA_SMALL_NUMBER));
+		ensure(FMath::IsNearlyEqual(OriginalNormal.SizeSquared(), FReal(1), FReal(KINDA_SMALL_NUMBER)));
 
 		// Get unscaled dir and normal
-		const TVector<T, 3> LocalDenormDir = DenormDir * MScale;
-		const TVector<T, 3> LocalOriginalNormalDenorm = OriginalNormal * MScale;
-		const T NormalLengthScale = LocalOriginalNormalDenorm.Size();
-		const TVector<T, 3> LocalOriginalNormal
+		const FVec3 LocalDenormDir = DenormDir * MScale;
+		const FVec3 LocalOriginalNormalDenorm = OriginalNormal * MScale;
+		const FReal NormalLengthScale = LocalOriginalNormalDenorm.Size();
+		const FVec3 LocalOriginalNormal
 			= ensure(NormalLengthScale > SMALL_NUMBER)
 			? LocalOriginalNormalDenorm / NormalLengthScale
-			: TVector<T, d>(0, 0, 1);
+			: FVec3(0, 0, 1);
 
 		// Compute final normal
-		const TVector<T, d> LocalNormal = MObject->FindGeometryOpposingNormal(LocalDenormDir, HintFaceIndex, LocalOriginalNormal);
-		TVector<T, d> Normal = LocalNormal * MInvScale;
-		FReal NormalLength = Normal.SafeNormalize(TNumericLimits<T>::Min());
+		const FVec3 LocalNormal = MObject->FindGeometryOpposingNormal(LocalDenormDir, HintFaceIndex, LocalOriginalNormal);
+		FVec3 Normal = LocalNormal * MInvScale;
+		FReal NormalLength = Normal.SafeNormalize(TNumericLimits<FReal>::Min());
 		if (NormalLength == 0)
 		{
 			CHAOS_ENSURE(false);
-			Normal = TVector<T, 3>(0, 0, 1);
+			Normal = FVec3(0, 0, 1);
 		}
 
 		return Normal;
 	}
 
-	virtual bool Overlap(const TVector<T, d>& Point, const T Thickness) const override
+	virtual bool Overlap(const FVec3& Point, const FReal Thickness) const override
 	{
-		const TVector<T, d> UnscaledPoint = MInvScale * Point;
+		const FVec3 UnscaledPoint = MInvScale * Point;
 
 		// TODO: consider alternative that handles thickness scaling properly in 3D, only works for uniform scaling right now
-		const T UnscaleThickness = MInvScale[0] * Thickness; 
+		const FReal UnscaleThickness = MInvScale[0] * Thickness; 
 
 		return MObject->Overlap(UnscaledPoint, UnscaleThickness);
 	}
 
-	virtual Pair<TVector<T, d>, bool> FindClosestIntersectionImp(const TVector<T, d>& StartPoint, const TVector<T, d>& EndPoint, const T Thickness) const override
+	virtual Pair<FVec3, bool> FindClosestIntersectionImp(const FVec3& StartPoint, const FVec3& EndPoint, const FReal Thickness) const override
 	{
 		ensure(OuterMargin == 0);	//not supported: do we care?
-		const TVector<T,d> UnscaledStart = MInvScale * StartPoint;
-		const TVector<T, d> UnscaledEnd = MInvScale * EndPoint;
+		const FVec3 UnscaledStart = MInvScale * StartPoint;
+		const FVec3 UnscaledEnd = MInvScale * EndPoint;
 		auto ClosestIntersection = MObject->FindClosestIntersection(UnscaledStart, UnscaledEnd, Thickness);
 		if (ClosestIntersection.Second)
 		{
@@ -818,7 +818,7 @@ public:
 		return ClosestIntersection;
 	}
 
-	virtual int32 FindClosestFaceAndVertices(const FVec3& Position, TArray<FVec3>& FaceVertices, FReal SearchDist = 0.01f) const override
+	virtual int32 FindClosestFaceAndVertices(const FVec3& Position, TArray<FVec3>& FaceVertices, FReal SearchDist = FReal(0.01)) const override
 	{
 		const FVec3 UnscaledPoint = MInvScale * Position;
 		const FReal UnscaledSearchDist = SearchDist * MInvScale.Max();	//this is not quite right since it's no longer a sphere, but the whole thing is fuzzy anyway
@@ -834,7 +834,7 @@ public:
 	}
 
 
-	FORCEINLINE_DEBUGGABLE TVector<T, d> Support(const TVector<T, d>& Direction, const T Thickness, int32& VertexIndex) const
+	FORCEINLINE_DEBUGGABLE FVec3 Support(const FVec3& Direction, const FReal Thickness, int32& VertexIndex) const
 	{
 		// Support_obj(dir) = pt => for all x in obj, pt \dot dir >= x \dot dir
 		// We want Support_objScaled(dir) = Support_obj(dir') where dir' is some modification of dir so we can use the unscaled support function
@@ -842,11 +842,11 @@ public:
 		// But this is the same as pt \dot dir >= dir^T Ax = (dir^TA) x = (A^T dir)^T x
 		//So let dir' = A^T dir.
 		//Since we only support scaling on the principal axes A is a diagonal (and therefore symmetric) matrix and so a simple component wise multiplication is sufficient
-		const TVector<T, d> UnthickenedPt = MObject->Support(Direction * MScale, 0.0f, VertexIndex) * MScale;
-		return Thickness > 0 ? TVector<T, d>(UnthickenedPt + Direction.GetSafeNormal() * Thickness) : UnthickenedPt;
+		const FVec3 UnthickenedPt = MObject->Support(Direction * MScale, 0.0f, VertexIndex) * MScale;
+		return Thickness > 0 ? FVec3(UnthickenedPt + Direction.GetSafeNormal() * Thickness) : UnthickenedPt;
 	}
 
-	FORCEINLINE_DEBUGGABLE TVector<T, d> SupportCore(const TVector<T, d>& Direction, const FReal InMargin, FReal* OutSupportDelta, int32& VertexIndex) const
+	FORCEINLINE_DEBUGGABLE FVec3 SupportCore(const FVec3& Direction, const FReal InMargin, FReal* OutSupportDelta, int32& VertexIndex) const
 	{
 		return MObject->SupportCoreScaled(Direction, InMargin, MScale, OutSupportDelta, VertexIndex);
 	}
@@ -901,7 +901,7 @@ public:
 		FChaosArchiveScopedMemory ScopedMemory(Ar, GetTypeName(), false);
 		FImplicitObject::SerializeImp(Ar);
 		Ar << MObject << MScale << MInvScale;
-		TBox<T,d>::SerializeAsAABB(Ar, MLocalBoundingBox);
+		TBox<FReal,d>::SerializeAsAABB(Ar, MLocalBoundingBox);
 		ensure(OuterMargin == 0);	//not supported: do we care?
 
 		Ar.UsingCustomVersion(FExternalPhysicsCustomObjectVersion::GUID);
@@ -949,10 +949,10 @@ private:
 
 	void UpdateBounds()
 	{
-		const TAABB<T, d> UnscaledBounds = MObject->BoundingBox();
-		const TVector<T, d> Vector1 = UnscaledBounds.Min() * MScale;
-		MLocalBoundingBox = TAABB<T, d>(Vector1, Vector1);	//need to grow it out one vector at a time in case scale is negative
-		const TVector<T, d> Vector2 = UnscaledBounds.Max() *MScale;
+		const FAABB3 UnscaledBounds = MObject->BoundingBox();
+		const FVec3 Vector1 = UnscaledBounds.Min() * MScale;
+		MLocalBoundingBox = FAABB3(Vector1, Vector1);	//need to grow it out one vector at a time in case scale is negative
+		const FVec3 Vector2 = UnscaledBounds.Max() * MScale;
 		MLocalBoundingBox.GrowToInclude(Vector2);
 	}
 
diff --git a/Engine/Source/Runtime/Experimental/Chaos/Public/Chaos/MassProperties.h b/Engine/Source/Runtime/Experimental/Chaos/Public/Chaos/MassProperties.h
--- a/Engine/Source/Runtime/Experimental/Chaos/Public/Chaos/MassProperties.h
+++ b/Engine/Source/Runtime/Experimental/Chaos/Public/Chaos/MassProperties.h
@@ -35,8 +35,11 @@ namespace Chaos
 
 	FRotation3 CHAOS_API TransformToLocalSpace(FMatrix33& Inertia);
 
-	template<typename TSurfaces>
-	void CHAOS_API CalculateVolumeAndCenterOfMass(const FParticles& Vertices, const TSurfaces& Surfaces, FReal& OutVolume, FVec3& OutCenterOfMass);
+	template<typename T, typename TSurfaces>
+	void CHAOS_API CalculateVolumeAndCenterOfMass(const TParticles<T,3>& Vertices, const TSurfaces& Surfaces, T& OutVolume, TVec3<T>& OutCenterOfMass);
+
+	template<typename T, typename TSurfaces>
+	void CHAOS_API CalculateVolumeAndCenterOfMass(const TArray<TVec3<T>>& Vertices, const TSurfaces& Surfaces, T& OutVolume, TVec3<T>& OutCenterOfMass);
 
 	template<typename TSurfaces>
 	FMassProperties CHAOS_API CalculateMassProperties(const FParticles& Vertices, const TSurfaces& Surfaces, const FReal Mass);
diff --git a/Engine/Source/Runtime/Experimental/Chaos/Public/Chaos/Plane.h b/Engine/Source/Runtime/Experimental/Chaos/Public/Chaos/Plane.h
--- a/Engine/Source/Runtime/Experimental/Chaos/Public/Chaos/Plane.h
+++ b/Engine/Source/Runtime/Experimental/Chaos/Public/Chaos/Plane.h
@@ -69,10 +69,11 @@ public:
 	}
 
 	// Scale the plane and assume that none of the scale components are zero
-	static TPlaneConcrete<T> MakeScaledUnsafe(const TPlaneConcrete<T>& Plane, const TVec3<T>& Scale)
+	template <typename U>
+	static TPlaneConcrete<T> MakeScaledUnsafe(const TPlaneConcrete<U>& Plane, const TVec3<T>& Scale)
 	{
-		const TVec3<T> ScaledX = Plane.MX * Scale;
-		TVec3<T> ScaledN = Plane.MNormal / Scale;
+		const TVec3<T> ScaledX = TVec3<T>(Plane.X()) * Scale;
+		TVec3<T> ScaledN = TVec3<T>(Plane.Normal()) / Scale;
 
 		// We don't handle zero scales, but we could still end up with a small normal
 		const T ScaleN2 = ScaledN.SizeSquared();
@@ -82,12 +83,18 @@ public:
 		}
 		else
 		{
-			ScaledN = Plane.MNormal;
+			ScaledN = TVec3<T>(Plane.Normal());
 		}
 
 		return TPlaneConcrete<T>(ScaledX, ScaledN);
 	}
 
+	template <typename U>
+	static TPlaneConcrete<T> MakeFrom(const TPlaneConcrete<U>& Plane)
+	{
+		return TPlaneConcrete<T>(TVec3<T>(Plane.X()), TVec3<T>(Plane.Normal()));
+	}
+
 
 	TPlaneConcrete() = default;
 	TPlaneConcrete(const TVec3<T>& InX, const TVec3<T>& InNormal)
@@ -100,9 +107,10 @@ public:
 	/**
 	 * Phi is positive on the side of the normal, and negative otherwise.
 	 */
-	FReal SignedDistance(const FVec3& x) const
+	template <typename U>
+	U SignedDistance(const TVec3<U>& x) const
 	{
-		return FVec3::DotProduct(x - (FVec3)MX, (FVec3)MNormal);
+		return TVec3<U>::DotProduct(x - TVec3<U>(MX), TVec3<U>(MNormal));
 	}
 
 	/**
diff --git a/Engine/Source/Runtime/Experimental/ChaosCore/Public/Chaos/Matrix.h b/Engine/Source/Runtime/Experimental/ChaosCore/Public/Chaos/Matrix.h
--- a/Engine/Source/Runtime/Experimental/ChaosCore/Public/Chaos/Matrix.h
+++ b/Engine/Source/Runtime/Experimental/ChaosCore/Public/Chaos/Matrix.h
@@ -25,12 +25,12 @@ namespace Chaos
 	};
 
 	template<>
-	class PMatrix<FReal, 3, 2>
+	class PMatrix<FRealDouble, 3, 2>
 	{
 	public:
-		FReal M[6];
+		FRealDouble M[6];
 
-		PMatrix(const TVector<FReal, 3>& C1, const TVector<FReal, 3>& C2)
+		PMatrix(const TVector<FRealDouble, 3>& C1, const TVector<FRealDouble, 3>& C2)
 		{
 			M[0] = C1.X;
 			M[1] = C1.Y;
@@ -40,7 +40,7 @@ namespace Chaos
 			M[5] = C2.Z;
 		}
 
-		PMatrix(const FReal x00, const FReal x10, const FReal x20, const FReal x01, const FReal x11, const FReal x21)
+		PMatrix(const FRealDouble x00, const FRealDouble x10, const FRealDouble x20, const FRealDouble x01, const FRealDouble x11, const FRealDouble x21)
 		{
 			M[0] = x00;
 			M[1] = x10;
@@ -50,12 +50,47 @@ namespace Chaos
 			M[5] = x21;
 		}
 
-		TVector<FReal, 3> operator*(const TVector<FReal, 2>& Other)
+		TVector<FRealDouble, 3> operator*(const TVector<FRealDouble, 2>& Other)
 		{
-			return TVector<FReal, 3>(
-			    M[0] * Other.X + M[3] * Other.Y,
-			    M[1] * Other.X + M[4] * Other.Y,
-			    M[2] * Other.X + M[5] * Other.Y);
+			return TVector<FRealDouble, 3>(
+			    M[0] * Other[0] + M[3] * Other[1],
+			    M[1] * Other[0] + M[4] * Other[1],
+			    M[2] * Other[0] + M[5] * Other[1]);
+		}
+	};
+
+	template<>
+	class PMatrix<FRealSingle, 3, 2>
+	{
+	public:
+		FRealSingle M[6];
+
+		PMatrix(const TVector<FRealSingle, 3>& C1, const TVector<FRealSingle, 3>& C2)
+		{
+			M[0] = C1.X;
+			M[1] = C1.Y;
+			M[2] = C1.Z;
+			M[3] = C2.X;
+			M[4] = C2.Y;
+			M[5] = C2.Z;
+		}
+
+		PMatrix(const FRealSingle x00, const FRealSingle x10, const FRealSingle x20, const FRealSingle x01, const FRealSingle x11, const FRealSingle x21)
+		{
+			M[0] = x00;
+			M[1] = x10;
+			M[2] = x20;
+			M[3] = x01;
+			M[4] = x11;
+			M[5] = x21;
+		}
+
+		TVector<FRealSingle, 3> operator*(const TVector<FRealSingle, 2>& Other)
+		{
+			return TVector<FRealSingle, 3>(
+				M[0] * Other[0] + M[3] * Other[1],
+				M[1] * Other[0] + M[4] * Other[1],
+				M[2] * Other[0] + M[5] * Other[1]);
 		}
 	};
 
@@ -157,19 +192,19 @@ namespace Chaos
 
 	// TODO(mlentine): Do not use 4x4 matrix for 3x3 implementation
 	template<>
-	class PMatrix<FReal, 3, 3> : public UE::Math::TMatrix<FReal>
+	class PMatrix<FRealDouble, 3, 3> : public UE::Math::TMatrix<FRealDouble>
 	{
 	public:
 		PMatrix()
-		    : UE::Math::TMatrix<FReal>() {}
-		PMatrix(UE::Math::TMatrix<FReal>&& Other)
-		    : UE::Math::TMatrix<FReal>(MoveTemp(Other)) {}
+		    : UE::Math::TMatrix<FRealDouble>() {}
+		PMatrix(UE::Math::TMatrix<FRealDouble>&& Other)
+		    : UE::Math::TMatrix<FRealDouble>(MoveTemp(Other)) {}
 		PMatrix(const UE::Math::TMatrix<FRealSingle>& Other)
-		    : UE::Math::TMatrix<FReal>((UE::Math::TMatrix<FReal>)Other) {}
+		    : UE::Math::TMatrix<FRealDouble>((UE::Math::TMatrix<FRealDouble>)Other) {}
 		PMatrix(const UE::Math::TMatrix<FRealDouble>& Other)
-			: UE::Math::TMatrix<FReal>((UE::Math::TMatrix<FReal>)Other) {}
-		PMatrix(const FReal x00, const FReal x11, const FReal x22)
-		    : UE::Math::TMatrix<FReal>()
+			: UE::Math::TMatrix<FRealDouble>((UE::Math::TMatrix<FRealDouble>)Other) {}
+		PMatrix(const FRealDouble x00, const FRealDouble x11, const FRealDouble x22)
+		    : UE::Math::TMatrix<FRealDouble>()
 		{
 			M[0][0] = x00;
 			M[1][0] = 0;
@@ -189,8 +224,8 @@ namespace Chaos
 			M[1][3] = 0;
 			M[2][3] = 0;
 		}
-		PMatrix(const FReal x00, const FReal x10, const FReal x20, const FReal x11, const FReal x21, const FReal x22)
-		    : UE::Math::TMatrix<FReal>()
+		PMatrix(const FRealDouble x00, const FRealDouble x10, const FRealDouble x20, const FRealDouble x11, const FRealDouble x21, const FRealDouble x22)
+		    : UE::Math::TMatrix<FRealDouble>()
 		{
 			M[0][0] = x00;
 			M[1][0] = x10;
@@ -210,8 +245,8 @@ namespace Chaos
 			M[1][3] = 0;
 			M[2][3] = 0;
 		}
-		PMatrix(const FReal x00, const FReal x10, const FReal x20, const FReal x01, const FReal x11, const FReal x21, const FReal x02, const FReal x12, const FReal x22)
-		    : UE::Math::TMatrix<FReal>()
+		PMatrix(const FRealDouble x00, const FRealDouble x10, const FRealDouble x20, const FRealDouble x01, const FRealDouble x11, const FRealDouble x21, const FRealDouble x02, const FRealDouble x12, const FRealDouble x22)
+		    : UE::Math::TMatrix<FRealDouble>()
 		{
 			M[0][0] = x00;
 			M[1][0] = x10;
@@ -231,8 +266,8 @@ namespace Chaos
 			M[1][3] = 0;
 			M[2][3] = 0;
 		}
-		PMatrix(const FReal x)
-		    : UE::Math::TMatrix<FReal>()
+		PMatrix(const FRealDouble x)
+		    : UE::Math::TMatrix<FRealDouble>()
 		{
 			M[0][0] = x;
 			M[1][0] = x;
@@ -252,7 +287,7 @@ namespace Chaos
 			M[1][3] = 0;
 			M[2][3] = 0;
 		}
-		PMatrix(const TVector<FReal, 3>& C1, const TVector<FReal, 3>& C2, const TVector<FReal, 3>& C3)
+		PMatrix(const TVector<FRealDouble, 3>& C1, const TVector<FRealDouble, 3>& C2, const TVector<FRealDouble, 3>& C3)
 		{
 			M[0][0] = C1.X;
 			M[1][0] = C1.Y;
@@ -275,13 +310,13 @@ namespace Chaos
 #if COMPILE_WITHOUT_UNREAL_SUPPORT
 		PMatrix<FReal, 3, 3> GetTransposed()
 		{
-			return PMatrix<FReal, 3, 3>(M[0][0], M[0][1], M[0][2], M[1][0], M[1][1], M[1][2], M[2][0], M[2][1], M[2][2]);
+			return PMatrix<FRealDouble, 3, 3>(M[0][0], M[0][1], M[0][2], M[1][0], M[1][1], M[1][2], M[2][0], M[2][1], M[2][2]);
 		}
 		FReal Determinant()
 		{
 			return M[0][0] * (M[1][1] * M[2][2] - M[1][2] * M[2][1]) - M[0][1] * (M[1][0] * M[2][2] - M[1][2] * M[2][0]) + M[0][2] * (M[1][0] * M[2][1] - M[1][1] * M[2][0]);
 		}
-		PMatrix<FReal, 3, 3>& operator+=(const PMatrix<FReal, 3, 3>& Other)
+		PMatrix<FRealDouble, 3, 3>& operator+=(const PMatrix<FRealDouble, 3, 3>& Other)
 		{
 			M[0][0] += Other.M[0][0];
 			M[0][1] += Other.M[0][1];
@@ -296,16 +331,16 @@ namespace Chaos
 		}
 #endif
 		// TDOD(mlentine): This should really be a vector multiply and sum for each entry using sse
-		TVector<FReal, 3> operator*(const TVector<FReal, 3>& Other) const
+		TVector<FRealDouble, 3> operator*(const TVector<FRealDouble, 3>& Other) const
 		{
-			return TVector<FReal, 3>(
+			return TVector<FRealDouble, 3>(
 			    M[0][0] * Other[0] + M[0][1] * Other[1] + M[0][2] * Other[2],
 			    M[1][0] * Other[0] + M[1][1] * Other[1] + M[1][2] * Other[2],
 			    M[2][0] * Other[0] + M[2][1] * Other[1] + M[2][2] * Other[2]);
 		}
-		PMatrix<FReal, 3, 3> operator+(const PMatrix<FReal, 3, 3>& Other) const
+		PMatrix<FRealDouble, 3, 3> operator+(const PMatrix<FRealDouble, 3, 3>& Other) const
 		{
-			return PMatrix<FReal, 3, 3>(
+			return PMatrix<FRealDouble, 3, 3>(
 				M[0][0] + Other.M[0][0],
 				M[1][0] + Other.M[1][0],
 				M[2][0] + Other.M[2][0],
@@ -316,13 +351,13 @@ namespace Chaos
 				M[1][2] + Other.M[1][2],
 				M[2][2] + Other.M[2][2]);
 		}
-		friend PMatrix<FReal, 3, 3> operator+(const PMatrix<FReal, 3, 3>& Other)
+		friend PMatrix<FRealDouble, 3, 3> operator+(const PMatrix<FRealDouble, 3, 3>& Other)
 		{
 			return Other;
 		}
-		PMatrix<FReal, 3, 3> operator-(const PMatrix<FReal, 3, 3>& Other) const
+		PMatrix<FRealDouble, 3, 3> operator-(const PMatrix<FRealDouble, 3, 3>& Other) const
 		{
-			return PMatrix<FReal, 3, 3>(
+			return PMatrix<FRealDouble, 3, 3>(
 				M[0][0] - Other.M[0][0],
 				M[1][0] - Other.M[1][0],
 				M[2][0] - Other.M[2][0],
@@ -333,9 +368,9 @@ namespace Chaos
 				M[1][2] - Other.M[1][2],
 				M[2][2] - Other.M[2][2]);
 		}
-		friend PMatrix<FReal, 3, 3> operator-(const PMatrix<FReal, 3, 3>& Other)
+		friend PMatrix<FRealDouble, 3, 3> operator-(const PMatrix<FRealDouble, 3, 3>& Other)
 		{
-			return PMatrix<FReal, 3, 3>(
+			return PMatrix<FRealDouble, 3, 3>(
 				-Other.M[0][0],
 				-Other.M[1][0],
 				-Other.M[2][0],
@@ -346,13 +381,13 @@ namespace Chaos
 				-Other.M[1][2],
 				-Other.M[2][2]);
 		}
-		PMatrix<FReal, 3, 3> operator*(const PMatrix<FReal, 3, 3>& Other) const
+		PMatrix<FRealDouble, 3, 3> operator*(const PMatrix<FRealDouble, 3, 3>& Other) const
 		{
-			return static_cast<const UE::Math::TMatrix<FReal>*>(this)->operator*(static_cast<const UE::Math::TMatrix<FReal>&>(Other));
+			return static_cast<const UE::Math::TMatrix<FRealDouble>*>(this)->operator*(static_cast<const UE::Math::TMatrix<FRealDouble>&>(Other));
 		}
-		PMatrix<FReal, 3, 3> operator*(const FReal Other) const
+		PMatrix<FRealDouble, 3, 3> operator*(const FRealDouble Other) const
 		{
-			return PMatrix<FReal, 3, 3>(
+			return PMatrix<FRealDouble, 3, 3>(
 			    M[0][0] * Other,
 			    M[1][0] * Other,
 			    M[2][0] * Other,
@@ -363,13 +398,13 @@ namespace Chaos
 			    M[1][2] * Other,
 			    M[2][2] * Other);
 		}
-		friend PMatrix<FReal, 3, 3> operator*(const FReal OtherF, const PMatrix<FReal, 3, 3>& OtherM)
+		friend PMatrix<FRealDouble, 3, 3> operator*(const FRealDouble OtherF, const PMatrix<FRealDouble, 3, 3>& OtherM)
 		{
 			return OtherM * OtherF;
 		}
-		PMatrix<FReal, 3, 2> operator*(const PMatrix<FReal, 3, 2>& Other) const
+		PMatrix<FRealDouble, 3, 2> operator*(const PMatrix<FRealDouble, 3, 2>& Other) const
 		{
-			return PMatrix<FReal, 3, 2>(
+			return PMatrix<FRealDouble, 3, 2>(
 			    M[0][0] * Other.M[0] + M[0][1] * Other.M[1] + M[0][2] * Other.M[2],
 			    M[1][0] * Other.M[0] + M[1][1] * Other.M[1] + M[1][2] * Other.M[2],
 			    M[2][0] * Other.M[0] + M[2][1] * Other.M[1] + M[2][2] * Other.M[2],
@@ -377,9 +412,9 @@ namespace Chaos
 			    M[1][0] * Other.M[3] + M[1][1] * Other.M[4] + M[1][2] * Other.M[5],
 			    M[2][0] * Other.M[3] + M[2][1] * Other.M[4] + M[2][2] * Other.M[5]);
 		}
-		PMatrix<FReal, 3, 3> SubtractDiagonal(const FReal Scalar) const
+		PMatrix<FRealDouble, 3, 3> SubtractDiagonal(const FRealDouble Scalar) const
 		{
-			return PMatrix<FReal, 3, 3>(
+			return PMatrix<FRealDouble, 3, 3>(
 			    M[0][0] - Scalar,
 			    M[1][0],
 			    M[2][0],
@@ -390,9 +425,9 @@ namespace Chaos
 			    M[1][2],
 			    M[2][2] - Scalar);
 		}
-		PMatrix<FReal, 3, 3> SymmetricCofactorMatrix() const
+		PMatrix<FRealDouble, 3, 3> SymmetricCofactorMatrix() const
 		{
-			return PMatrix<FReal, 3, 3>(
+			return PMatrix<FRealDouble, 3, 3>(
 			    M[1][1] * M[2][2] - M[2][1] * M[2][1],
 			    M[2][1] * M[2][0] - M[1][0] * M[2][2],
 			    M[1][0] * M[2][1] - M[1][1] * M[2][0],
@@ -400,43 +435,43 @@ namespace Chaos
 			    M[1][0] * M[2][0] - M[0][0] * M[2][1],
 			    M[0][0] * M[1][1] - M[1][0] * M[1][0]);
 		}
-		TVector<FReal, 3> LargestColumnNormalized() const
+		TVector<FRealDouble, 3> LargestColumnNormalized() const
 		{
-			FReal m10 = M[1][0] * M[1][0];
-			FReal m20 = M[2][0] * M[2][0];
-			FReal m21 = M[2][1] * M[2][1];
-			FReal c0 = M[0][0] * M[0][0] + m10 + m20;
-			FReal c1 = m10 + M[1][1] * M[1][1] + m21;
-			FReal c2 = m20 + m21 + M[2][2] * M[2][2];
+			FRealDouble m10 = M[1][0] * M[1][0];
+			FRealDouble m20 = M[2][0] * M[2][0];
+			FRealDouble m21 = M[2][1] * M[2][1];
+			FRealDouble c0 = M[0][0] * M[0][0] + m10 + m20;
+			FRealDouble c1 = m10 + M[1][1] * M[1][1] + m21;
+			FRealDouble c2 = m20 + m21 + M[2][2] * M[2][2];
 			if (c0 > c1 && c0 > c2)
 			{
-				return TVector<FReal, 3>(M[0][0], M[1][0], M[2][0]) / sqrt(c0);
+				return TVector<FRealDouble, 3>(M[0][0], M[1][0], M[2][0]) / sqrt(c0);
 			}
 			if (c1 > c2)
 			{
-				return TVector<FReal, 3>(M[1][0], M[1][1], M[2][1]) / sqrt(c1);
+				return TVector<FRealDouble, 3>(M[1][0], M[1][1], M[2][1]) / sqrt(c1);
 			}
 			if (c2 > 0)
 			{
-				return TVector<FReal, 3>(M[2][0], M[2][1], M[2][2]) / sqrt(c2);
+				return TVector<FRealDouble, 3>(M[2][0], M[2][1], M[2][2]) / sqrt(c2);
 			}
-			return TVector<FReal, 3>(1, 0, 0);
+			return TVector<FRealDouble, 3>(1, 0, 0);
 		}
 
 		/**
 		 * Get the specified axis (0-indexed, X,Y,Z).
 		 * @note: we are treating matrices as column major, so axis elements are sequential in memory
 		 */
-		FORCEINLINE TVector<FReal, 3> GetAxis(int32 AxisIndex) const
+		FORCEINLINE TVector<FRealDouble, 3> GetAxis(int32 AxisIndex) const
 		{
-			return TVector<FReal, 3>(M[AxisIndex][0], M[AxisIndex][1], M[AxisIndex][2]);
+			return TVector<FRealDouble, 3>(M[AxisIndex][0], M[AxisIndex][1], M[AxisIndex][2]);
 		}
 
 		/**
 		 * Set the specified axis (0-indexed, X,Y,Z).
 		 * @note: we are treating matrices as column major, so axis elements are sequential in memory
 		 */
-		FORCEINLINE void SetAxis(int32 AxisIndex, const TVector<FReal, 3>& Axis)
+		FORCEINLINE void SetAxis(int32 AxisIndex, const TVector<FRealDouble, 3>& Axis)
 		{
 			M[AxisIndex][0] = Axis.X;
 			M[AxisIndex][1] = Axis.Y;
@@ -449,9 +484,9 @@ namespace Chaos
 		 * @note: we are treating matrices as column major, so rows are not sequential in memory
 		 * @seealso GetAxis, GetColumn
 		 */
-		FORCEINLINE TVector<FReal, 3> GetRow(int32 RowIndex) const
+		FORCEINLINE TVector<FRealDouble, 3> GetRow(int32 RowIndex) const
 		{
-			return TVector<FReal, 3>(M[0][RowIndex], M[1][RowIndex], M[2][RowIndex]);
+			return TVector<FRealDouble, 3>(M[0][RowIndex], M[1][RowIndex], M[2][RowIndex]);
 		}
 
 		/**
@@ -459,7 +494,7 @@ namespace Chaos
 		 * @note: we are treating matrices as column major, so axis elements are sequential in memory
 		 * @seealso SetAxis, SetColumn
 		 */
-		FORCEINLINE void SetRow(int32 RowIndex, const TVector<FReal, 3>& V)
+		FORCEINLINE void SetRow(int32 RowIndex, const TVector<FRealDouble, 3>& V)
 		{
 			M[0][RowIndex] = V.X;
 			M[1][RowIndex] = V.Y;
@@ -472,7 +507,7 @@ namespace Chaos
 		 * @note: we are treating matrices as column major, so columns are sequential in memory
 		 * @seealso GetAxis, GetRow
 		 */
-		FORCEINLINE TVector<FReal, 3> GetColumn(int32 ColumnIndex) const
+		FORCEINLINE TVector<FRealDouble, 3> GetColumn(int32 ColumnIndex) const
 		{
 			return GetAxis(ColumnIndex);
 		}
@@ -482,7 +517,7 @@ namespace Chaos
 		 * @note: we are treating matrices as column major, so axis elements are sequential in memory
 		 * @seealso SetAxis, SetRow
 		 */
-		FORCEINLINE void SetColumn(int32 ColumnIndex, const TVector<FReal, 3>& V)
+		FORCEINLINE void SetColumn(int32 ColumnIndex, const TVector<FRealDouble, 3>& V)
 		{
 			SetAxis(ColumnIndex, V);
 		}
@@ -490,17 +525,17 @@ namespace Chaos
 		/**
 		 * Get the diagonal elements as a vector.
 		 */
-		FORCEINLINE TVector<FReal, 3> GetDiagonal() const
+		FORCEINLINE TVector<FRealDouble, 3> GetDiagonal() const
 		{
-			return TVector<FReal, 3>(M[0][0], M[1][1], M[2][2]);
+			return TVector<FRealDouble, 3>(M[0][0], M[1][1], M[2][2]);
 		}
 
-		FORCEINLINE FReal GetAt(int32 RowIndex, int32 ColIndex) const
+		FORCEINLINE FRealDouble GetAt(int32 RowIndex, int32 ColIndex) const
 		{
 			return M[ColIndex][RowIndex];
 		}
 
-		FORCEINLINE void SetAt(int32 RowIndex, int32 ColIndex, FReal V)
+		FORCEINLINE void SetAt(int32 RowIndex, int32 ColIndex, FRealDouble V)
 		{
 			M[ColIndex][RowIndex] = V;
 		}
@@ -508,16 +543,16 @@ namespace Chaos
 		/**
 		 * Return a diagonal matrix with the specified elements
 		 */
-		static PMatrix<FReal, 3, 3> FromDiagonal(const TVector<FReal, 3>& D)
+		static PMatrix<FRealDouble, 3, 3> FromDiagonal(const TVector<FRealDouble, 3>& D)
 		{
-			return PMatrix<FReal, 3, 3>(D.X, D.Y, D.Z);
+			return PMatrix<FRealDouble, 3, 3>(D.X, D.Y, D.Z);
 		}
 
 #if COMPILE_WITHOUT_UNREAL_SUPPORT
 		// TODO(mlentine): Document which one is row and which one is column
-		PMatrix<FReal, 3, 3> operator*(const PMatrix<FReal, 3, 3>& Other)
+		PMatrix<FRealDouble, 3, 3> operator*(const PMatrix<FRealDouble, 3, 3>& Other)
 		{
-			return PMatrix<FReal, 3, 3>(
+			return PMatrix<FRealDouble, 3, 3>(
 			    M[0][0] * Other.M[0][0] + M[0][1] * Other.M[1][0] + M[0][2] * Other.M[2][0],
 			    M[1][0] * Other.M[0][0] + M[1][1] * Other.M[1][0] + M[1][2] * Other.M[2][0],
 			    M[2][0] * Other.M[0][0] + M[2][1] * Other.M[1][0] + M[2][2] * Other.M[2][0],
@@ -528,9 +563,9 @@ namespace Chaos
 			    M[1][0] * Other.M[0][2] + M[1][1] * Other.M[1][2] + M[1][2] * Other.M[2][2],
 			    M[2][0] * Other.M[0][2] + M[2][1] * Other.M[1][2] + M[2][2] * Other.M[2][2]);
 		}
-		PMatrix<FReal, 3, 3> operator*(const FReal Scalar)
+		PMatrix<FRealDouble, 3, 3> operator*(const FRealDouble Scalar)
 		{
-			return PMatrix<FReal, 3, 3>(
+			return PMatrix<FRealDouble, 3, 3>(
 			    M[0][0] * Scalar + M[0][1] * Scalar + M[0][2] * Scalar,
 			    M[1][0] * Scalar + M[1][1] * Scalar + M[1][2] * Scalar,
 			    M[2][0] * Scalar + M[2][1] * Scalar + M[2][2] * Scalar,
@@ -542,7 +577,7 @@ namespace Chaos
 			    M[2][0] * Scalar + M[2][1] * Scalar + M[2][2] * Scalar);
 		}
 #endif
-		inline bool Equals(const PMatrix<FReal, 3, 3>& Other, FReal Tolerance = KINDA_SMALL_NUMBER) const
+		inline bool Equals(const PMatrix<FRealDouble, 3, 3>& Other, FRealDouble Tolerance = KINDA_SMALL_NUMBER) const
 		{
 			return true
 				&& (FMath::Abs(Other.M[0][0] - M[0][0]) <= Tolerance)
@@ -557,11 +592,10 @@ namespace Chaos
 		}
 
 
-		static const PMatrix<FReal, 3, 3> Zero;
-		static const PMatrix<FReal, 3, 3> Identity;
+		static const PMatrix<FRealDouble, 3, 3> Zero;
+		static const PMatrix<FRealDouble, 3, 3> Identity;
 	};
 
-#if !UE_LARGE_WORLD_COORDINATES_DISABLED
 	template<>
 	class PMatrix<FRealSingle, 3, 3> : public UE::Math::TMatrix<FRealSingle>
 	{
@@ -773,6 +807,16 @@ namespace Chaos
 		{
 			return OtherM * OtherF;
 		}
+		PMatrix<FRealSingle, 3, 2> operator*(const PMatrix<FRealSingle, 3, 2>& Other) const
+		{
+			return PMatrix<FRealSingle, 3, 2>(
+				M[0][0] * Other.M[0] + M[0][1] * Other.M[1] + M[0][2] * Other.M[2],
+				M[1][0] * Other.M[0] + M[1][1] * Other.M[1] + M[1][2] * Other.M[2],
+				M[2][0] * Other.M[0] + M[2][1] * Other.M[1] + M[2][2] * Other.M[2],
+				M[0][0] * Other.M[3] + M[0][1] * Other.M[4] + M[0][2] * Other.M[5],
+				M[1][0] * Other.M[3] + M[1][1] * Other.M[4] + M[1][2] * Other.M[5],
+				M[2][0] * Other.M[3] + M[2][1] * Other.M[4] + M[2][2] * Other.M[5]);
+		}
 		PMatrix<FRealSingle, 3, 3> SubtractDiagonal(const FRealSingle Scalar) const
 		{
 			return PMatrix<FRealSingle, 3, 3>(
@@ -956,5 +1000,4 @@ namespace Chaos
 		static const PMatrix<FRealSingle, 3, 3> Zero;
 		static const PMatrix<FRealSingle, 3, 3> Identity;
 	};
-#endif
 }
\ No newline at end of file
