Reset ParticleCollisions in FPBDRigidsEvolutionGBF::DestroyParticleCollisionsInAllocator

diff --git a/Engine/Source/Runtime/Experimental/Chaos/Private/Chaos/PBDRigidsEvolutionGBF.cpp b/Engine/Source/Runtime/Experimental/Chaos/Private/Chaos/PBDRigidsEvolutionGBF.cpp
--- a/Engine/Source/Runtime/Experimental/Chaos/Private/Chaos/PBDRigidsEvolutionGBF.cpp
+++ b/Engine/Source/Runtime/Experimental/Chaos/Private/Chaos/PBDRigidsEvolutionGBF.cpp
@@ -876,6 +876,7 @@ void FPBDRigidsEvolutionGBF::DestroyParticleCollisionsInAllocator(FGeometryParti
 	if (Particle != nullptr)
 	{
 		CollisionConstraints.GetConstraintAllocator().RemoveParticle(Particle);
+		Particle->ParticleCollisions().Reset();
 	}
 }
 
