Chaos : Fix linux build issue ( related to 18576858 )

diff --git a/Engine/Source/Runtime/Experimental/Chaos/Private/Chaos/MassProperties.cpp b/Engine/Source/Runtime/Experimental/Chaos/Private/Chaos/MassProperties.cpp
--- a/Engine/Source/Runtime/Experimental/Chaos/Private/Chaos/MassProperties.cpp
+++ b/Engine/Source/Runtime/Experimental/Chaos/Private/Chaos/MassProperties.cpp
@@ -229,6 +229,9 @@ namespace Chaos
 	template CHAOS_API void CalculateVolumeAndCenterOfMass(const TParticles<FRealDouble, 3>& Vertices, const TArray<TVec3<int32>>& Surfaces, FRealDouble& OutVolume, TVec3<FRealDouble>& OutCenterOfMass);
 	template CHAOS_API void CalculateVolumeAndCenterOfMass(const TParticles<FRealDouble, 3>& Vertices, const TArray<TArray<int32>>& Surfaces, FRealDouble& OutVolume, TVec3<FRealDouble>& OutCenterOfMass);
 
+	template CHAOS_API void CalculateVolumeAndCenterOfMass(const TParticles<FRealSingle, 3>& Vertices, const TArray<TVec3<int32>>& Surfaces, FRealSingle& OutVolume, TVec3<FRealSingle>& OutCenterOfMass);
+	template CHAOS_API void CalculateVolumeAndCenterOfMass(const TParticles<FRealSingle, 3>& Vertices, const TArray<TArray<int32>>& Surfaces, FRealSingle& OutVolume, TVec3<FRealSingle>& OutCenterOfMass);
+	
 	template CHAOS_API void CalculateVolumeAndCenterOfMass(const TArray<TVec3<FRealSingle>>& Vertices, const TArray<TArray<int32>>& Surfaces, FRealSingle& OutVolume, TVec3<FRealSingle>& OutCenterOfMass);
 	template CHAOS_API void CalculateVolumeAndCenterOfMass(const TArray<TVec3<FRealDouble>>& Vertices, const TArray<TArray<int32>>& Surfaces, FRealDouble& OutVolume, TVec3<FRealDouble>& OutCenterOfMass);
 
