Chaos - fix debug darw folliwiung change to NetMode in editor

diff --git a/Engine/Source/Runtime/Experimental/ChaosSolverEngine/Private/Chaos/ChaosDebugDrawComponent.cpp b/Engine/Source/Runtime/Experimental/ChaosSolverEngine/Private/Chaos/ChaosDebugDrawComponent.cpp
--- a/Engine/Source/Runtime/Experimental/ChaosSolverEngine/Private/Chaos/ChaosDebugDrawComponent.cpp
+++ b/Engine/Source/Runtime/Experimental/ChaosSolverEngine/Private/Chaos/ChaosDebugDrawComponent.cpp
@@ -303,7 +303,7 @@ void UChaosDebugDrawComponent::HandlePostWorldInitialization(UWorld* World, cons
 void UChaosDebugDrawComponent::CreateDebugDrawActor(UWorld* World)
 {
 #if CHAOS_DEBUG_DRAW
-	if (!World->IsNetMode(NM_Client))
+	if ((World->GetNetMode() != ENetMode::NM_Client) && (World->GetNetMode() != ENetMode::NM_Standalone))
 		return;
 
 	static FName NAME_ChaosDebugDrawActor = TEXT("ChaosDebugDrawActor");
