Chaos - friction improvement
- dynamic friction handled in position solve now (except for quadratic shapes)
- use velocity to estimate previous friction anchors when previous contact not available
- reproject vertex contacts onto the contact plane when reusing (friction stability)

diff --git a/Engine/Source/Programs/HeadlessChaos/Private/HeadlessChaosTestCollisions.cpp b/Engine/Source/Programs/HeadlessChaos/Private/HeadlessChaosTestCollisions.cpp
--- a/Engine/Source/Programs/HeadlessChaos/Private/HeadlessChaosTestCollisions.cpp
+++ b/Engine/Source/Programs/HeadlessChaos/Private/HeadlessChaosTestCollisions.cpp
@@ -330,16 +330,9 @@ namespace ChaosTest {
 		Collisions.ComputeConstraints(Dt);
 		EXPECT_EQ(Collisions.NumConstraints(), 1);
 
-		FPBDCollisionConstraint& Constraint = Collisions.GetConstraint(0);
-		if (auto PBDRigid = Constraint.Particle[0]->CastToRigidParticle())
-		{
-			PBDRigid->CollisionParticles()->UpdateAccelerationStructures();
-		}
-
 		Collisions.GatherInput(Dt);
 
-		Collisions.UpdateLevelsetConstraint(Constraint);
-
+		FPBDCollisionConstraint& Constraint = Collisions.GetConstraint(0);
 		EXPECT_EQ(Constraint.Particle[0], Box);
 		EXPECT_EQ(Constraint.Particle[1], Floor);
 		EXPECT_TRUE(Constraint.GetNormal().operator==(FVec3(0, 0, 1)));
diff --git a/Engine/Source/Programs/HeadlessChaos/Private/HeadlessChaosTestConstraintSolverCore.cpp b/Engine/Source/Programs/HeadlessChaos/Private/HeadlessChaosTestConstraintSolverCore.cpp
--- a/Engine/Source/Programs/HeadlessChaos/Private/HeadlessChaosTestConstraintSolverCore.cpp
+++ b/Engine/Source/Programs/HeadlessChaos/Private/HeadlessChaosTestConstraintSolverCore.cpp
@@ -79,7 +79,7 @@ namespace ChaosTest
 
 
 		Collision.SetSolverBodies(&Bodies[0], &Bodies[1]);
-		Collision.SetFriction(0, 0);
+		Collision.SetFriction(0, 0, 0);
 		Collision.SetNumManifoldPoints(1);
 		Collision.InitContact(0, FVec3(0,0,Radius), FVec3(0,0,-Radius), FVec3(0,0,1));
 		Collision.InitMaterial(0, 0, false, 0);
diff --git a/Engine/Source/Programs/HeadlessChaos/Private/HeadlessChaosTestContactModification.cpp b/Engine/Source/Programs/HeadlessChaos/Private/HeadlessChaosTestContactModification.cpp
--- a/Engine/Source/Programs/HeadlessChaos/Private/HeadlessChaosTestContactModification.cpp
+++ b/Engine/Source/Programs/HeadlessChaos/Private/HeadlessChaosTestContactModification.cpp
@@ -188,7 +188,7 @@ namespace ChaosTest
 					int32 NumContacts = PairModifier.GetNumContacts();
 					for (int32 PointIdx = 0; PointIdx < NumContacts; ++PointIdx)
 					{
-						// Tell solver contacts are SeparationPadding units closer than they are, so object will float 5 over floor.
+						// Tell solver contacts are SeparationPadding units closer than they are, so object will float SeparationPadding units over floor.
 						FReal Separation = PairModifier.GetSeparation(PointIdx);
 						PairModifier.ModifySeparation(Separation - SeparationPadding, PointIdx);
 					}
@@ -596,7 +596,8 @@ namespace ChaosTest
 		SlidingCubeParticle.SetGeometry(SlidingCubeGeom);
 		Solver->RegisterObject(SlidingCubeProxy);
 		SlidingCubeParticle.SetGravityEnabled(true);
-		SlidingCubeParticle.SetX(FVec3(200, -200, 100));
+		SlidingCubeParticle.SetX(FVec3(200, -200, 50));
+		SlidingCubeParticle.SetR(FQuat::MakeFromEuler(FVec3(0, 20, 0)));
 		SetCubeInertiaTensor(SlidingCubeParticle, /*Dimension=*/200, /*Mass=*/1);
 		ChaosTest::SetParticleSimDataToCollide({ SlidingCubeProxy->GetParticle_LowLevel() });
 
@@ -606,7 +607,8 @@ namespace ChaosTest
 		ModifiedCubeParticle.SetGeometry(ModifiedCubeGeom);
 		Solver->RegisterObject(ModifiedCubeProxy);
 		ModifiedCubeParticle.SetGravityEnabled(true);
-		ModifiedCubeParticle.SetX(FVec3(200, 200, 100));
+		ModifiedCubeParticle.SetX(FVec3(200, 200, 50));
+		ModifiedCubeParticle.SetR(FQuat::MakeFromEuler(FVec3(0, 20, 0)));
 		SetCubeInertiaTensor(ModifiedCubeParticle, /*Dimension=*/200, /*Mass=*/1);
 		ChaosTest::SetParticleSimDataToCollide({ ModifiedCubeProxy->GetParticle_LowLevel() });
 
@@ -646,7 +648,7 @@ namespace ChaosTest
 
 
 		const float Dt = 0.1f;
-		const int32 Steps = 30;
+		const int32 Steps = 50;
 		for (int Step = 0; Step < Steps; ++Step)
 		{
 			Solver->AdvanceAndDispatch_External(Dt);
diff --git a/Engine/Source/Runtime/Experimental/Chaos/Private/Chaos/ChaosDebugDraw.cpp b/Engine/Source/Runtime/Experimental/Chaos/Private/Chaos/ChaosDebugDraw.cpp
--- a/Engine/Source/Runtime/Experimental/Chaos/Private/Chaos/ChaosDebugDraw.cpp
+++ b/Engine/Source/Runtime/Experimental/Chaos/Private/Chaos/ChaosDebugDraw.cpp
@@ -874,7 +874,7 @@ namespace Chaos
 						FColor ImpulseColor = FColor(0, 0, 200);
 						FColor PushOutColor = FColor(200, 200, 0);
 						FColor PushOutImpusleColor = FColor(0, 200, 200);
-						if (ManifoldPoint.bInsideStaticFrictionCone)
+						if (ManifoldPoint.Flags.bInsideStaticFrictionCone)
 						{
 							DiscColor = FColor(150, 200, 0);
 						}
@@ -920,7 +920,7 @@ namespace Chaos
 							const FReal BoxScale = Settings.DrawScale * Settings.ContactWidth;
 							FDebugDrawQueue::GetInstance().DrawDebugBox(WorldPlaneLocation, FVec3(BoxScale, BoxScale, FReal(0.01)), FRotation3(FRotationMatrix::MakeFromZ(WorldPlaneNormal)), FColor::Blue, false, KINDA_SMALL_NUMBER, Settings.DrawPriority, 0.5f * Settings.LineThickness);
 						}
-						else if (ManifoldPoint.bWasRestored)
+						else if (ManifoldPoint.Flags.bWasRestored)
 						{
 							const FReal BoxScale = Settings.DrawScale * Settings.ContactWidth;
 							FDebugDrawQueue::GetInstance().DrawDebugBox(WorldPlaneLocation, FVec3(BoxScale, BoxScale, FReal(0.01)), FRotation3(FRotationMatrix::MakeFromZ(WorldPlaneNormal)), FColor::Purple, false, KINDA_SMALL_NUMBER, Settings.DrawPriority, 0.5f * Settings.LineThickness);
diff --git a/Engine/Source/Runtime/Experimental/Chaos/Private/Chaos/Collision/PBDCollisionConstraint.cpp b/Engine/Source/Runtime/Experimental/Chaos/Private/Chaos/Collision/PBDCollisionConstraint.cpp
--- a/Engine/Source/Runtime/Experimental/Chaos/Private/Chaos/Collision/PBDCollisionConstraint.cpp
+++ b/Engine/Source/Runtime/Experimental/Chaos/Private/Chaos/Collision/PBDCollisionConstraint.cpp
@@ -121,7 +121,7 @@ namespace Chaos
 	FPBDCollisionConstraint FPBDCollisionConstraint::MakeTriangle(const FImplicitObject* Implicit0)
 	{
 		FPBDCollisionConstraint Constraint;
-		Constraint.InitMargins(Implicit0->GetCollisionType(), ImplicitObjectType::Triangle, Implicit0->GetMargin(), FReal(0));
+		Constraint.InitMarginsAndTolerances(Implicit0->GetCollisionType(), ImplicitObjectType::Triangle, Implicit0->GetMargin(), FReal(0));
 		return Constraint;
 	}
 
@@ -228,10 +228,10 @@ namespace Chaos
 		const FReal Margin1 = GetImplicit1()->GetMargin();
 		const EImplicitObjectType ImplicitType0 = GetInnerType(GetImplicit0()->GetCollisionType());
 		const EImplicitObjectType ImplicitType1 = GetInnerType(GetImplicit1()->GetCollisionType());
-		InitMargins(ImplicitType0, ImplicitType1, Margin0, Margin1);
+		InitMarginsAndTolerances(ImplicitType0, ImplicitType1, Margin0, Margin1);
 	}
 
-	void FPBDCollisionConstraint::InitMargins(const EImplicitObjectType ImplicitType0, const EImplicitObjectType ImplicitType1, const FReal Margin0, const FReal Margin1)
+	void FPBDCollisionConstraint::InitMarginsAndTolerances(const EImplicitObjectType ImplicitType0, const EImplicitObjectType ImplicitType1, const FReal Margin0, const FReal Margin1)
 	{
 		// Set up the margins and tolerances to be used during the narrow phase.
 		// One shape in a collision will always have a margin. Only triangles have zero margin and we don't 
@@ -268,6 +268,9 @@ namespace Chaos
 			CollisionMargins[1] = Margin1;
 			CollisionTolerance = QuadraticToleranceScale * Margin1;
 		}
+
+		Flags.bIsQuadratic0 = bIsQuadratic0;
+		Flags.bIsQuadratic1 = bIsQuadratic1;
 	}
 
 	void FPBDCollisionConstraint::SetIsSleeping(const bool bInIsSleeping)
@@ -386,7 +389,7 @@ namespace Chaos
 
 			UpdateManifoldPointFromContact(ManifoldPointIndex);
 
-			ManifoldPoint.bInsideStaticFrictionCone = Flags.bUseManifold;
+			ManifoldPoint.Flags.bInsideStaticFrictionCone = Flags.bUseManifold;
 
 			// Copy currently active point
 			if (ManifoldPoint.ContactPoint.Phi < Manifold.Phi)
@@ -495,8 +498,9 @@ namespace Chaos
 	void FPBDCollisionConstraint::UpdateManifoldPointFromContact(const int32 ManifoldPointIndex)
 	{
 		FManifoldPoint& ManifoldPoint = ManifoldPoints[ManifoldPointIndex];
-		ManifoldPoint.WorldContactPoints[0] = ShapeWorldTransform0.TransformPositionNoScale(ManifoldPoint.ContactPoint.ShapeContactPoints[0]);
-		ManifoldPoint.WorldContactPoints[1] = ShapeWorldTransform1.TransformPositionNoScale(ManifoldPoint.ContactPoint.ShapeContactPoints[1]);
+
+		ManifoldPoint.WorldContactPoints[0] = ShapeWorldTransform0.TransformPositionNoScale(ManifoldPoint.ShapeAnchorPoints[0]);
+		ManifoldPoint.WorldContactPoints[1] = ShapeWorldTransform1.TransformPositionNoScale(ManifoldPoint.ShapeAnchorPoints[1]);
 	}
 
 	void FPBDCollisionConstraint::SetActiveContactPoint(const FContactPoint& ContactPoint)
@@ -528,12 +532,19 @@ namespace Chaos
 		Flags.bWasManifoldRestored = false;
 	}
 
-	void FPBDCollisionConstraint::RestoreManifold()
+	void FPBDCollisionConstraint::RestoreManifold(const bool bReproject)
 	{
-		// We want to restore the manifold as-is and will skip the narrow phase, which means we leave the manifold in place, 
-		// but we still have some cleanup to do to account for slight movement of the bodies. E.g., we need to update the 
-		// world-space state for the contact modifiers
-		UpdateManifoldContacts();
+		// We want to restore the manifold as-is and skip the narrow phase, but...
+		if (bReproject)
+		{
+			// The bodies may have moved a moderate amount so we need to update the shape-space (and subsequently the world-space) contact points
+			ReprojectManifoldContacts();
+		}
+		else
+		{
+			// The bodies have barely moved so we keep the same shape-space contacts. We still need to update world-space contacts though.
+			UpdateManifoldContacts();
+		}
 
 		Flags.bWasManifoldRestored = true;
 	}
@@ -551,6 +562,54 @@ namespace Chaos
 		LastShapeWorldTransform1 = InShapeWorldTransform1;
 	}
 
+	void FPBDCollisionConstraint::ReprojectManifoldPoint(const int32 ManifoldPointIndex)
+	{
+		FManifoldPoint& ManifoldPoint = ManifoldPoints[ManifoldPointIndex];
+
+		// Calculate the world-space contact location and separation at the current shape transforms
+		// @todo(chaos): this should use the normal owner. Currently we assume body 1 is the owner
+		const FRigidTransform3 Shape0ToShape1Transform = ShapeWorldTransform0.GetRelativeTransformNoScale(ShapeWorldTransform1);
+		const FVec3 Contact0In1 = Shape0ToShape1Transform.TransformPositionNoScale(ManifoldPoint.InitialShapeContactPoints[0]);
+		const FVec3& Contact1In1 = ManifoldPoint.InitialShapeContactPoints[1];
+		const FVec3 ContactNormalIn1 = ShapeWorldTransform1.InverseTransformVectorNoScale(ManifoldPoint.ContactPoint.Normal);
+
+		const FVec3 ContactDeltaIn1 = Contact0In1 - Contact1In1;
+		const FReal ContactPhi = FVec3::DotProduct(ContactDeltaIn1, ContactNormalIn1);
+
+		// Recalculate the contact points at the new location
+		const FVec3 ShapeContactPoint1 = Contact0In1 - ContactPhi * ContactNormalIn1;
+		ManifoldPoint.ContactPoint.ShapeContactPoints[1] = ShapeContactPoint1;
+		ManifoldPoint.ContactPoint.Phi = ContactPhi;
+
+		// Restore friction anchors if we have them for this point
+		TryRestoreFrictionData(ManifoldPointIndex);
+
+		// Update world-space contact locations
+		UpdateManifoldPointFromContact(ManifoldPointIndex);
+		ManifoldPoint.ContactPoint.Location = FReal(0.5) * (ManifoldPoint.WorldContactPoints[0] + ManifoldPoint.WorldContactPoints[1]);
+	}
+
+	void FPBDCollisionConstraint::ReprojectManifoldContacts()
+	{
+		Manifold.Reset();
+
+		for (int32 ManifoldPointIndex = 0; ManifoldPointIndex < ManifoldPoints.Num(); ManifoldPointIndex++)
+		{
+			FManifoldPoint& ManifoldPoint = ManifoldPoints[ManifoldPointIndex];
+
+			ReprojectManifoldPoint(ManifoldPointIndex);
+
+			ManifoldPoint.Flags.bInsideStaticFrictionCone = Flags.bUseManifold;
+			ManifoldPoint.Flags.bWasRestored = true;
+
+			// Copy currently active point
+			if (ManifoldPoint.ContactPoint.Phi < Manifold.Phi)
+			{
+				SetActiveContactPoint(ManifoldPoint.ContactPoint);
+			}
+		}
+	}
+
 	bool FPBDCollisionConstraint::UpdateAndTryRestoreManifold()
 	{
 		const FCollisionTolerances& Tolerances = Chaos_Manifold_Tolerances;
@@ -632,7 +691,7 @@ namespace Chaos
 				UpdateManifoldPointFromContact(ManifoldPointIndex);
 				ManifoldPoint.ContactPoint.Location = FReal(0.5) * (ManifoldPoint.WorldContactPoints[0] + ManifoldPoint.WorldContactPoints[1]);
 
-				ManifoldPoint.bWasRestored = true;
+				ManifoldPoint.Flags.bWasRestored = true;
 
 				if (ManifoldPoint.ContactPoint.Phi < Manifold.Phi)
 				{
@@ -713,7 +772,7 @@ namespace Chaos
 					ManifoldPoint.ContactPoint = NewContactPoint;
 					ManifoldPoint.InitialShapeContactPoints[0] = NewContactPoint.ShapeContactPoints[0];
 					ManifoldPoint.InitialShapeContactPoints[1] = NewContactPoint.ShapeContactPoints[1];
-					ManifoldPoint.bWasRestored = false;
+					ManifoldPoint.Flags.bWasRestored = false;
 					TryRestoreFrictionData(ManifoldPointIndex);
 					UpdateManifoldPointFromContact(ManifoldPointIndex);
 					if (NewContactPoint.Phi < GetPhi())
@@ -861,7 +920,7 @@ namespace Chaos
 				ManifoldPoints[ManifoldPointIndex].ContactPoint = NewContactPoint;
 				ManifoldPoints[ManifoldPointIndex].InitialShapeContactPoints[0] = NewContactPoint.ShapeContactPoints[0];
 				ManifoldPoints[ManifoldPointIndex].InitialShapeContactPoints[1] = NewContactPoint.ShapeContactPoints[1];
-				ManifoldPoints[ManifoldPointIndex].bWasRestored = false;
+				ManifoldPoints[ManifoldPointIndex].Flags.bWasRestored = false;
 				UpdateManifoldPointFromContact(ManifoldPointIndex);
 				if (NewContactPoint.Phi < Manifold.Phi)
 				{
@@ -873,19 +932,66 @@ namespace Chaos
 		return true;
 	}
 
-	const FManifoldPointSavedData* FPBDCollisionConstraint::FindManifoldPointSavedData(const FManifoldPoint& ManifoldPoint) const
+	FReal FPBDCollisionConstraint::CalculateSavedManifoldPointScore(const FSavedManifoldPoint& SavedManifoldPoint, const FManifoldPoint& ManifoldPoint, const FReal DistanceToleranceSq) const
+	{
+		// If we have a vertex-plane (or vertex-vertex) contact, we want to know if we have the same vertex(es).
+		// If we have and edge-edge contact, we want to know if we have the same edges.
+		// But we don't know what type of contact we have, so for now...
+		// If the contact point is in the same spot on one of the bodies, assume it is the same contact
+		// @todo(chaos) - collision detection should provide the contact point types (vertex/edge/plane)
+		FReal DP0Sq = TNumericLimits<FReal>::Max();
+		FReal DP1Sq = TNumericLimits<FReal>::Max();
+		const FVec3 DP0 = ManifoldPoint.ContactPoint.ShapeContactPoints[0] - SavedManifoldPoint.ShapeContactPoints[0];
+		const FVec3 DP1 = ManifoldPoint.ContactPoint.ShapeContactPoints[1] - SavedManifoldPoint.ShapeContactPoints[1];
+
+		// When only one shape is quadratic, we only look at the quadratic contact point so we don't identify
+		// a sphere spinning on the spot as a stationary contact
+		// @todo(chaos): handle quadratic shapes better with static friction
+		if (IsQuadratic0() && !IsQuadratic1())
+		{
+			DP0Sq = DP0.SizeSquared();
+		}
+		else if (IsQuadratic1() && !IsQuadratic0())
+		{
+			DP1Sq = DP1.SizeSquared();
+		}
+		else
+		{
+			DP0Sq = DP0.SizeSquared();
+			DP1Sq = DP1.SizeSquared();
+		}
+
+		const FReal MinDPSq = FMath::Min(DP0Sq, DP1Sq);
+		if (MinDPSq < DistanceToleranceSq)
+		{
+			return MinDPSq;
+		}
+
+		return TNumericLimits<FReal>::Max();
+	}
+
+	const FSavedManifoldPoint* FPBDCollisionConstraint::FindSavedManifoldPoint(const FManifoldPoint& ManifoldPoint) const
 	{
 		if (bChaos_Manifold_EnableFrictionRestore)
 		{
 			const FReal DistanceToleranceSq = FMath::Square(Chaos_Manifold_FrictionPositionTolerance);
+			FReal BestScore = DistanceToleranceSq;
+			int32 MatchIndex = INDEX_NONE;
 			for (int32 SavedManifoldPointIndex = 0; SavedManifoldPointIndex < SavedManifoldPoints.Num(); ++SavedManifoldPointIndex)
 			{
-				const FManifoldPointSavedData& SavedManifoldPoint = SavedManifoldPoints[SavedManifoldPointIndex];
-				if (SavedManifoldPoint.IsMatch(ManifoldPoint, DistanceToleranceSq))
+				const FSavedManifoldPoint& SavedManifoldPoint = SavedManifoldPoints[SavedManifoldPointIndex];
+				const FReal Score = CalculateSavedManifoldPointScore(SavedManifoldPoint, ManifoldPoint, DistanceToleranceSq);
+				if (Score < BestScore)
 				{
-					return &SavedManifoldPoint;
+					BestScore = Score;
+					MatchIndex = SavedManifoldPointIndex;
 				}
 			}
+
+			if (MatchIndex != INDEX_NONE)
+			{
+				return &SavedManifoldPoints[MatchIndex];
+			}
 		}
 		return nullptr;
 	}
@@ -894,17 +1000,24 @@ namespace Chaos
 	{
 		FManifoldPoint& ManifoldPoint = ManifoldPoints[ManifoldPointIndex];
 
-		// Assume we have no matching point from the previous tick, but that we can retain friction from now on
+		// Assume we have no matching point from the previous tick, but that we can hold friction from now on
 		// Not supported for non-manifolds yet (hopefully we don't need to)
-		ManifoldPoint.bInsideStaticFrictionCone = Flags.bUseManifold;
-		ManifoldPoint.StaticFrictionMax = FReal(0);
+		ManifoldPoint.Flags.bInsideStaticFrictionCone = Flags.bUseManifold;
 
 		// Find the previous manifold point that matches if there is one
-		const FManifoldPointSavedData* PrevManifoldPoint = FindManifoldPointSavedData(ManifoldPoint);
-		if (PrevManifoldPoint != nullptr)
+		const FSavedManifoldPoint* SavedManifoldPoint = FindSavedManifoldPoint(ManifoldPoint);
+		if (SavedManifoldPoint != nullptr)
 		{
 			// We have data from the previous tick and static friction was enabled - restore the data
-			PrevManifoldPoint->Restore(ManifoldPoint);
+			ManifoldPoint.ShapeAnchorPoints[0] = SavedManifoldPoint->ShapeContactPoints[0];
+			ManifoldPoint.ShapeAnchorPoints[1] = SavedManifoldPoint->ShapeContactPoints[1];
+			ManifoldPoint.Flags.bWasFrictionRestored = true;
+		}
+		else
+		{
+			ManifoldPoint.ShapeAnchorPoints[0] = ManifoldPoint.ContactPoint.ShapeContactPoints[0];
+			ManifoldPoint.ShapeAnchorPoints[1] = ManifoldPoint.ContactPoint.ShapeContactPoints[1];
+			ManifoldPoint.Flags.bWasFrictionRestored = false;
 		}
 	}
 
diff --git a/Engine/Source/Runtime/Experimental/Chaos/Private/Chaos/Collision/PBDCollisionSolver.cpp b/Engine/Source/Runtime/Experimental/Chaos/Private/Chaos/Collision/PBDCollisionSolver.cpp
--- a/Engine/Source/Runtime/Experimental/Chaos/Private/Chaos/Collision/PBDCollisionSolver.cpp
+++ b/Engine/Source/Runtime/Experimental/Chaos/Private/Chaos/Collision/PBDCollisionSolver.cpp
@@ -27,9 +27,7 @@ namespace Chaos
 		int32 Chaos_PBDCollisionSolver_Position_ShockPropagationIterations = 3;
 		float Chaos_PBDCollisionSolver_Position_MinInvMassScale = 0.3f;
 		int32 Chaos_PBDCollisionSolver_Position_ZeroFrictionIterations = 4;
-		bool bChaos_PBDCollisionSolver_Position_NegativePushOutEnabled = true;
 		float Chaos_PBDCollisionSolver_Position_StaticFrictionStiffness = 0.5f;
-		float Chaos_PBDCollisionSolver_Position_StaticFrictionLerpRate = 0.1f;
 		float Chaos_PBDCollisionSolver_Position_PositionSolverTolerance = 0.001f;		// cms
 		float Chaos_PBDCollisionSolver_Position_RotationSolverTolerance = 0.001f;		// rads
 
@@ -38,7 +36,6 @@ namespace Chaos
 		FAutoConsoleVariableRef CVarChaos_PBDCollisionSolver_Position_MinInvMassScale(TEXT("p.Chaos.PBDCollisionSolver.Position.MinInvMassScale"), Chaos_PBDCollisionSolver_Position_MinInvMassScale, TEXT(""));
 		FAutoConsoleVariableRef CVarChaos_PBDCollisionSolver_Position_ZeroFrictionIterations(TEXT("p.Chaos.PBDCollisionSolver.Position.ZeroFrictionIterations"), Chaos_PBDCollisionSolver_Position_ZeroFrictionIterations, TEXT(""));
 		FAutoConsoleVariableRef CVarChaos_PBDCollisionSolver_Position_StaticFrictionStiffness(TEXT("p.Chaos.PBDCollisionSolver.Position.StaticFriction.Stiffness"), Chaos_PBDCollisionSolver_Position_StaticFrictionStiffness, TEXT(""));
-		FAutoConsoleVariableRef CVarChaos_PBDCollisionSolver_Position_StaticFrictionLerpRate(TEXT("p.Chaos.PBDCollisionSolver.Position.StaticFriction.LerpRate"), Chaos_PBDCollisionSolver_Position_StaticFrictionLerpRate, TEXT(""));
 		FAutoConsoleVariableRef CVarChaos_PBDCollisionSolver_Position_PositionSolverTolerance(TEXT("p.Chaos.PBDCollisionSolver.Position.PositionTolerance"), Chaos_PBDCollisionSolver_Position_PositionSolverTolerance, TEXT(""));
 		FAutoConsoleVariableRef CVarChaos_PBDCollisionSolver_Position_RotationSolverTolerance(TEXT("p.Chaos.PBDCollisionSolver.Position.RotationTolerance"), Chaos_PBDCollisionSolver_Position_RotationSolverTolerance, TEXT(""));
 
@@ -46,28 +43,26 @@ namespace Chaos
 		int32 Chaos_PBDCollisionSolver_Velocity_ShockPropagationIterations = 1;
 		// If this is the same as Chaos_PBDCollisionSolver_Position_MinInvMassScale and all velocity iterations have shockpropagation, we avoid recalculating constraiunt-space mass
 		float Chaos_PBDCollisionSolver_Velocity_MinInvMassScale = Chaos_PBDCollisionSolver_Position_MinInvMassScale;
-		bool bChaos_PBDCollisionSolver_Velocity_DynamicFrictionEnabled = true;
-		bool bChaos_PBDCollisionSolver_Velocity_NegativeImpulseEnabled = true;
-		bool bChaos_PBDCollisionSolver_Velocity_ImpulseClampEnabled = true;
+		bool bChaos_PBDCollisionSolver_Velocity_FrictionEnabled = true;
 		bool bChaos_PBDCollisionSolver_Velocity_AveragePointEnabled = true;
 
 		FAutoConsoleVariableRef CVarChaos_PBDCollisionSolver_Velocity_SolveEnabled(TEXT("p.Chaos.PBDCollisionSolver.Velocity.SolveEnabled"), bChaos_PBDCollisionSolver_Velocity_SolveEnabled, TEXT(""));
 		FAutoConsoleVariableRef CVarChaos_PBDCollisionSolver_Velocity_UseShockPropagation(TEXT("p.Chaos.PBDCollisionSolver.Velocity.ShockPropagationIterations"), Chaos_PBDCollisionSolver_Velocity_ShockPropagationIterations, TEXT(""));
 		FAutoConsoleVariableRef CVarChaos_PBDCollisionSolver_Velocity_MinInvMassScale(TEXT("p.Chaos.PBDCollisionSolver.Velocity.MinInvMassScale"), Chaos_PBDCollisionSolver_Velocity_MinInvMassScale, TEXT(""));
-		FAutoConsoleVariableRef CVarChaos_PBDCollisionSolver_Velocity_DynamicFrictionEnabled(TEXT("p.Chaos.PBDCollisionSolver.Velocity.DynamicFrictionEnabled"), bChaos_PBDCollisionSolver_Velocity_DynamicFrictionEnabled, TEXT(""));
-		FAutoConsoleVariableRef CVarChaos_PBDCollisionSolver_Velocity_NegativeImpulseEnabled(TEXT("p.Chaos.PBDCollisionSolver.Velocity.NegativeImpulseEnabled"), bChaos_PBDCollisionSolver_Velocity_NegativeImpulseEnabled, TEXT(""));
-		FAutoConsoleVariableRef CVarChaos_PBDCollisionSolver_Velocity_ImpulseClampEnabled(TEXT("p.Chaos.PBDCollisionSolver.Velocity.ImpulseClampEnabled"), bChaos_PBDCollisionSolver_Velocity_ImpulseClampEnabled, TEXT(""));
+		FAutoConsoleVariableRef CVarChaos_PBDCollisionSolver_Velocity_FrictionEnabled(TEXT("p.Chaos.PBDCollisionSolver.Velocity.FrictionEnabled"), bChaos_PBDCollisionSolver_Velocity_FrictionEnabled, TEXT(""));
 		FAutoConsoleVariableRef CVarChaos_PBDCollisionSolver_Velocity_AveragePointEnabled(TEXT("p.Chaos.PBDCollisionSolver.Velocity.AveragePointEnabled"), bChaos_PBDCollisionSolver_Velocity_AveragePointEnabled, TEXT(""));
 	}
 	using namespace CVars;
 
 
 	void FPBDCollisionSolverManifoldPoint::InitContact(
+		const FSolverReal Dt,
 		const FConstraintSolverBody& Body0,
 		const FConstraintSolverBody& Body1,
 		const FVec3& InWorldAnchorPoint0,
 		const FVec3& InWorldAnchorPoint1,
-		const FSolverVec3& InWorldContactNormal)
+		const FSolverVec3& InWorldContactNormal,
+		const bool bAddVelocityToFriction)
 	{
 		NetPushOutNormal = FSolverReal(0);
 		NetPushOutTangentU = FSolverReal(0);
@@ -76,7 +71,7 @@ namespace Chaos
 		NetImpulseTangentU = FSolverReal(0);
 		NetImpulseTangentV = FSolverReal(0);
 
-		UpdateContact(Body0, Body1, InWorldAnchorPoint0, InWorldAnchorPoint1, InWorldContactNormal);
+		UpdateContact(Dt, Body0, Body1, InWorldAnchorPoint0, InWorldAnchorPoint1, InWorldContactNormal, bAddVelocityToFriction);
 
 		UpdateMass(Body0, Body1);
 	}
@@ -85,12 +80,9 @@ namespace Chaos
 		const FConstraintSolverBody& Body0,
 		const FConstraintSolverBody& Body1,
 		const FSolverReal InRestitution,
-		const FSolverReal InRestitutionVelocityThreshold,
-		const bool bInEnableStaticFriction,
-		const FSolverReal InFrictionNetPushOutNormal)
+		const FSolverReal InRestitutionVelocityThreshold)
 	{
-		FrictionNetPushOutNormal = InFrictionNetPushOutNormal;
-		bInsideStaticFrictionCone = bInEnableStaticFriction;
+		StaticFrictionRatio = FSolverReal(0);
 		WorldContactVelocityTargetNormal = FSolverReal(0);
 
 		if (InRestitution > FSolverReal(0))
@@ -106,11 +98,13 @@ namespace Chaos
 	}
 
 	inline void FPBDCollisionSolverManifoldPoint::UpdateContact(
+		const FSolverReal Dt,
 		const FConstraintSolverBody& Body0, 
 		const FConstraintSolverBody& Body1,
 		const FVec3& InWorldAnchorPoint0,
 		const FVec3& InWorldAnchorPoint1, 
-		const FSolverVec3& InWorldContactNormal)
+		const FSolverVec3& InWorldContactNormal,
+		const bool bAddVelocityToFriction)
 	{
 		// The world-space point where we apply impulses/corrections (same world-space point for momentum conservation)
 		const FVec3 WorldContactPosition = FReal(0.5) * (InWorldAnchorPoint0 + InWorldAnchorPoint1);
@@ -133,8 +127,19 @@ namespace Chaos
 		// The contact point error we are trying to correct in this solver
 		const FSolverVec3 WorldContactDelta = InWorldAnchorPoint0 - InWorldAnchorPoint1;
 		WorldContactDeltaNormal = FSolverVec3::DotProduct(WorldContactDelta, WorldContactNormal);
-		WorldContactDeltaTangent0 = FSolverVec3::DotProduct(WorldContactDelta, WorldContactTangentU);
-		WorldContactDeltaTangent1 = FSolverVec3::DotProduct(WorldContactDelta, WorldContactTangentV);
+		WorldContactDeltaTangentU = FSolverVec3::DotProduct(WorldContactDelta, WorldContactTangentU);
+		WorldContactDeltaTangentV = FSolverVec3::DotProduct(WorldContactDelta, WorldContactTangentV);
+
+		// For rolling friction to work (e.g., on spheres) we need to include the angular velocity in the friction calculation
+		if (bAddVelocityToFriction)
+		{
+			// Contact velocity from rotation for quadratic surfaces. Corrects for rolling friction.
+			FVec3 ContactVel = (Body0.V() + FVec3::CrossProduct(Body0.W(), RelativeContactPosition0)) - (Body1.V() + FVec3::CrossProduct(Body1.W(), RelativeContactPosition1));
+			const FSolverReal ContactVelocityU = FSolverVec3::DotProduct(ContactVel, WorldContactTangentU);
+			const FSolverReal ContactVelocityV = FSolverVec3::DotProduct(ContactVel, WorldContactTangentV);
+			WorldContactDeltaTangentU += ContactVelocityU * Dt;
+			WorldContactDeltaTangentV += ContactVelocityV * Dt;
+		}
 	}
 
 	inline void FPBDCollisionSolverManifoldPoint::UpdateMass(const FConstraintSolverBody& Body0, const FConstraintSolverBody& Body1)
@@ -267,52 +272,34 @@ namespace Chaos
 
 	void FPBDCollisionSolver::InitContact(
 		const int32 ManifoldPoiontIndex,
+		const FReal Dt,
 		const FVec3& InWorldAnchorPoint0,
 		const FVec3& InWorldAnchorPoint1,
-		const FVec3& InWorldContactNormal)
+		const FVec3& InWorldContactNormal,
+		const bool bAddVelocityToFriction)
 	{
 		State.ManifoldPoints[ManifoldPoiontIndex].InitContact(
+			FSolverReal(Dt),
 			State.SolverBodies[0], 
 			State.SolverBodies[1], 
 			InWorldAnchorPoint0, 
 			InWorldAnchorPoint1, 
-			FSolverVec3(InWorldContactNormal));
+			FSolverVec3(InWorldContactNormal),
+			bAddVelocityToFriction);
 	}
 
 	void FPBDCollisionSolver::InitMaterial(
 		const int32 ManifoldPoiontIndex,
 		const FReal InRestitution,
-		const FReal InRestitutionVelocityThreshold,
-		const bool bInEnableStaticFriction,
-		const FReal InFrictionNetPushOutNormal)
+		const FReal InRestitutionVelocityThreshold)
 	{
 		State.ManifoldPoints[ManifoldPoiontIndex].InitMaterial(
 			State.SolverBodies[0], 
 			State.SolverBodies[1], 
 			FSolverReal(InRestitution), 
-			FSolverReal(InRestitutionVelocityThreshold),
-			bInEnableStaticFriction, 
-			FSolverReal(InFrictionNetPushOutNormal));
-
-		// Track if any points have restitution enabled. See SolveVelocity
-		State.bHaveRestitution = State.bHaveRestitution || (State.ManifoldPoints[ManifoldPoiontIndex].WorldContactVelocityTargetNormal > 0);
-	}
-
-	void FPBDCollisionSolver::UpdateContact(
-		const int32 ManifoldPoiontIndex,
-		const FVec3& InWorldAnchorPoint0,
-		const FVec3& InWorldAnchorPoint1,
-		const FVec3& InWorldContactNormal)
-	{
-		State.ManifoldPoints[ManifoldPoiontIndex].UpdateContact(
-			State.SolverBodies[0], 
-			State.SolverBodies[1], 
-			InWorldAnchorPoint0, 
-			InWorldAnchorPoint1, 
-			FSolverVec3(InWorldContactNormal));
+			FSolverReal(InRestitutionVelocityThreshold));
 	}
 
-
 	void FPBDCollisionSolver::SolveVelocityAverage(const FSolverReal Dt)
 	{
 		FConstraintSolverBody& Body0 = SolverBody0();
@@ -328,10 +315,7 @@ namespace Chaos
 		for (int32 PointIndex = 0; PointIndex < State.NumManifoldPoints; ++PointIndex)
 		{
 			FPBDCollisionSolverManifoldPoint& SolverManifoldPoint = State.ManifoldPoints[PointIndex];
-
-			// An "active" point here is defined as one that started out penetrating or applied a net impulse.
-			const bool bUsePoint = (SolverManifoldPoint.WorldContactDeltaNormal < 0) || (SolverManifoldPoint.NetPushOutNormal > FSolverReal(SMALL_NUMBER));
-			if (bUsePoint)
+			if (SolverManifoldPoint.ShouldSolveVelocity())
 			{
 				RelativeContactPosition0 += SolverManifoldPoint.RelativeContactPosition0;
 				RelativeContactPosition1 += SolverManifoldPoint.RelativeContactPosition1;
@@ -346,17 +330,17 @@ namespace Chaos
 		// We only do this if we have multiple active contacts
 		// NOTE: this average contact isn't really correct, especially when all contacts do not equally
 		// contribute to the pushout (which is normal even for a simple box on a plane). E.g., the WorldContactVelocityTargetNormal
-		// and NetPushOutNormal will not be right. The goal here though is just to get close to the correct result in frewer iterations
+		// and NetPushOutNormal will not be right. The goal here though is just to get close to the correct result in fewer iterations
 		// than we would have if we just solved the corners of the box in sequence.
 		if (NumActiveManifoldPoints > 1)
 		{
-			const FSolverReal Scale = FSolverReal(1) / FSolverReal(NumActiveManifoldPoints);
+			const FSolverReal InvCount = FSolverReal(1) / FSolverReal(NumActiveManifoldPoints);
 
 			FPBDCollisionSolverManifoldPoint AverageManifoldPoint;
-			AverageManifoldPoint.RelativeContactPosition0 = RelativeContactPosition0 * Scale;
-			AverageManifoldPoint.RelativeContactPosition1 = RelativeContactPosition1 * Scale;
+			AverageManifoldPoint.RelativeContactPosition0 = RelativeContactPosition0 * InvCount;
+			AverageManifoldPoint.RelativeContactPosition1 = RelativeContactPosition1 * InvCount;
 			AverageManifoldPoint.WorldContactNormal = WorldContactNormal;
-			AverageManifoldPoint.WorldContactVelocityTargetNormal = WorldContactVelocityTargetNormal * Scale;
+			AverageManifoldPoint.WorldContactVelocityTargetNormal = WorldContactVelocityTargetNormal * InvCount;
 			
 			// Total pushout (not average) which is correct if the average point is also the centroid but otherwise probably an overestimate.
 			// This is used to limit the possibly-attractive impulse that corrects implicit velocity errors from the PBD solve, but it
@@ -394,26 +378,37 @@ namespace Chaos
 			AverageManifoldPoint.WorldContactTangentUAngular1 = FSolverVec3(0);
 			AverageManifoldPoint.WorldContactTangentVAngular1 = FSolverVec3(0);
 			AverageManifoldPoint.WorldContactDeltaNormal = FSolverReal(0);
-			AverageManifoldPoint.WorldContactDeltaTangent0 = FSolverReal(0);
-			AverageManifoldPoint.WorldContactDeltaTangent1 = FSolverReal(0);
+			AverageManifoldPoint.WorldContactDeltaTangentU = FSolverReal(0);
+			AverageManifoldPoint.WorldContactDeltaTangentV = FSolverReal(0);
 			AverageManifoldPoint.NetPushOutTangentU = FSolverReal(0);
 			AverageManifoldPoint.NetPushOutTangentV = FSolverReal(0);
 			AverageManifoldPoint.NetImpulseNormal = FSolverReal(0);
 			AverageManifoldPoint.NetImpulseTangentU = FSolverReal(0);
 			AverageManifoldPoint.NetImpulseTangentV = FSolverReal(0);
-			AverageManifoldPoint.FrictionNetPushOutNormal = FSolverReal(0);
-			AverageManifoldPoint.bInsideStaticFrictionCone = false;
+			AverageManifoldPoint.StaticFrictionRatio = FSolverReal(0);
 
 			FSolverReal ContactVelocityDeltaNormal;
 			AverageManifoldPoint.CalculateContactVelocityErrorNormal(Body0, Body1, ContactVelocityDeltaNormal);
 
+			const FSolverReal MinImpulseNormal = FMath::Min(FSolverReal(0), -AverageManifoldPoint.NetPushOutNormal / Dt);
+
 			ApplyVelocityCorrectionNormal(
 				State.Stiffness,
-				Dt,
 				ContactVelocityDeltaNormal,
+				MinImpulseNormal,
 				AverageManifoldPoint,
 				Body0,
 				Body1);
+
+			// Now distribute the net impulse among the active points so we don't over-correct pushout from initial overlaps
+			for (int32 PointIndex = 0; PointIndex < State.NumManifoldPoints; ++PointIndex)
+			{
+				FPBDCollisionSolverManifoldPoint& SolverManifoldPoint = State.ManifoldPoints[PointIndex];
+				if (SolverManifoldPoint.ShouldSolveVelocity())
+				{
+					SolverManifoldPoint.NetImpulseNormal += AverageManifoldPoint.NetImpulseNormal * InvCount;
+				}
+			}
 		}
 	}
 
@@ -433,26 +428,45 @@ namespace Chaos
 
 		// NOTE: this dynamic friction implementation is iteration-count sensitive
 		// @todo(chaos): fix iteration count dependence of dynamic friction
-		const FSolverReal DynamicFriction = (bApplyDynamicFriction && bChaos_PBDCollisionSolver_Velocity_DynamicFrictionEnabled) ? State.DynamicFriction : FSolverReal(0);
+		const FSolverReal DynamicFriction = (bApplyDynamicFriction && (Dt > 0) && bChaos_PBDCollisionSolver_Velocity_FrictionEnabled) ? State.VelocityFriction : FSolverReal(0);
 
 		for (int32 PointIndex = 0; PointIndex < State.NumManifoldPoints; ++PointIndex)
 		{
 			FPBDCollisionSolverManifoldPoint& SolverManifoldPoint = State.ManifoldPoints[PointIndex];
-			if (!FMath::IsNearlyZero(SolverManifoldPoint.NetPushOutNormal))
+			if (SolverManifoldPoint.ShouldSolveVelocity())
 			{
-				FSolverReal ContactVelocityDeltaNormal, ContactVelocityDeltaTangentU, ContactVelocityDeltaTangentV;
-				SolverManifoldPoint.CalculateContactVelocityError(Body0, Body1, DynamicFriction, Dt, ContactVelocityDeltaNormal, ContactVelocityDeltaTangentU, ContactVelocityDeltaTangentV);
-
-				ApplyVelocityCorrection(
-					State.Stiffness,
-					Dt,
-					DynamicFriction,
-					ContactVelocityDeltaNormal,
-					ContactVelocityDeltaTangentU,
-					ContactVelocityDeltaTangentV,
-					SolverManifoldPoint,
-					Body0,
-					Body1);
+				const FSolverReal MinImpulseNormal = FMath::Min(FSolverReal(0), -SolverManifoldPoint.NetPushOutNormal / Dt);
+
+				if (DynamicFriction > 0)
+				{
+					FSolverReal ContactVelocityDeltaNormal, ContactVelocityDeltaTangentU, ContactVelocityDeltaTangentV;
+					SolverManifoldPoint.CalculateContactVelocityError(Body0, Body1, DynamicFriction, Dt, ContactVelocityDeltaNormal, ContactVelocityDeltaTangentU, ContactVelocityDeltaTangentV);
+
+					ApplyVelocityCorrection(
+						State.Stiffness,
+						Dt,
+						DynamicFriction,
+						ContactVelocityDeltaNormal,
+						ContactVelocityDeltaTangentU,
+						ContactVelocityDeltaTangentV,
+						MinImpulseNormal,
+						SolverManifoldPoint,
+						Body0,
+						Body1);
+				}
+				else
+				{
+					FSolverReal ContactVelocityDeltaNormal;
+					SolverManifoldPoint.CalculateContactVelocityErrorNormal(Body0, Body1, ContactVelocityDeltaNormal);
+
+					ApplyVelocityCorrectionNormal(
+						State.Stiffness,
+						ContactVelocityDeltaNormal,
+						MinImpulseNormal,
+						SolverManifoldPoint,
+						Body0,
+						Body1);
+				}
 			}
 		}
 
diff --git a/Engine/Source/Runtime/Experimental/Chaos/Private/Chaos/Collision/PBDCollisionSolver.h b/Engine/Source/Runtime/Experimental/Chaos/Private/Chaos/Collision/PBDCollisionSolver.h
--- a/Engine/Source/Runtime/Experimental/Chaos/Private/Chaos/Collision/PBDCollisionSolver.h
+++ b/Engine/Source/Runtime/Experimental/Chaos/Private/Chaos/Collision/PBDCollisionSolver.h
@@ -8,6 +8,8 @@
 #include "Chaos/Defines.h"
 #include "Chaos/Evolution/SolverBody.h"
 
+CHAOS_API DECLARE_LOG_CATEGORY_EXTERN(LogChaosCollision, Log, All);
+
 // Set to 0 to use a linearized error calculation, and set to 1 to use a non-linear error calculation in collision detection. 
 // In principle nonlinear is more accurate when large rotation corrections occur, but this is not too important for collisions because 
 // when the bodies settle the corrections are small. The linearized version is significantly faster than the non-linear version because 
@@ -28,10 +30,6 @@ namespace Chaos
 	namespace CVars
 	{
 		extern float Chaos_PBDCollisionSolver_Position_StaticFrictionStiffness;
-		extern float Chaos_PBDCollisionSolver_Position_StaticFrictionLerpRate;
-		
-		extern bool bChaos_PBDCollisionSolver_Velocity_NegativeImpulseEnabled;
-		extern bool bChaos_PBDCollisionSolver_Velocity_ImpulseClampEnabled;
 	}
 
 
@@ -45,11 +43,13 @@ namespace Chaos
 		 * @brief Initialize the geometric data for the contact
 		*/
 		void InitContact(
+			const FSolverReal Dt,
 			const FConstraintSolverBody& Body0,
 			const FConstraintSolverBody& Body1,
 			const FVec3& InWorldAnchorPoint0,
 			const FVec3& InWorldAnchorPoint1,
-			const FSolverVec3& InWorldContactNormal);
+			const FSolverVec3& InWorldContactNormal,
+			const bool bAddVelocityToFriction);
 
 		/**
 		 * @brief Initialize the material related properties of the contact
@@ -58,19 +58,19 @@ namespace Chaos
 			const FConstraintSolverBody& Body0,
 			const FConstraintSolverBody& Body1,
 			const FSolverReal InRestitution,
-			const FSolverReal InRestitutionVelocityThreshold,
-			const bool bInEnableStaticFriction,
-			const FSolverReal InFrictionNetPushOutNormal);
+			const FSolverReal InRestitutionVelocityThreshold);
 
 		/**
 		 * @brief Update the world-space relative contact points based on current body transforms and body-space contact positions
 		*/
 		void UpdateContact(
+			const FSolverReal Dt,
 			const FConstraintSolverBody& Body0,
 			const FConstraintSolverBody& Body1,
 			const FVec3& InWorldAnchorPoint0,
 			const FVec3& InWorldAnchorPoint1,
-			const FSolverVec3& InWorldContactNormal);
+			const FSolverVec3& InWorldContactNormal,
+			const bool bAddVelocityToFriction);
 
 		/**
 		 * @brief Update the cached mass properties based on the current body transforms
@@ -100,6 +100,8 @@ namespace Chaos
 		*/
 		FSolverVec3 CalculateContactVelocity(const FConstraintSolverBody& Body0, const FConstraintSolverBody& Body1) const;
 
+		bool ShouldSolveVelocity() const;
+
 		// World-space body-relative contact points
 		FSolverVec3 RelativeContactPosition0;
 		FSolverVec3 RelativeContactPosition1;
@@ -124,8 +126,8 @@ namespace Chaos
 
 		// Initial world-space contact separation that we are trying to correct
 		FSolverReal WorldContactDeltaNormal;
-		FSolverReal WorldContactDeltaTangent0;
-		FSolverReal WorldContactDeltaTangent1;
+		FSolverReal WorldContactDeltaTangentU;
+		FSolverReal WorldContactDeltaTangentV;
 
 		// Desired final normal velocity, taking Restitution into account
 		FSolverReal WorldContactVelocityTargetNormal;
@@ -138,11 +140,10 @@ namespace Chaos
 		FSolverReal NetImpulseTangentU;
 		FSolverReal NetImpulseTangentV;
 
-		// A smoothed NetPushOutNormal along the normal, used for clipping to the static friction cone
-		FSolverReal FrictionNetPushOutNormal;
-
-		// Whether we are still in the static friction cone
-		bool bInsideStaticFrictionCone;
+		// A measure of how much we exceeded the static friction threshold.
+		// Equal to (NormalPushOut / TangentialPushOut) before clamping to the friction cone.
+		// Used to move the static friction anchors to the edge of the cone in Scatter.
+		FSolverReal StaticFrictionRatio;
 	};
 
 	/**
@@ -167,11 +168,13 @@ namespace Chaos
 
 		FSolverReal StaticFriction() const { return State.StaticFriction; }
 		FSolverReal DynamicFriction() const { return State.DynamicFriction; }
+		FSolverReal VelocityFriction() const { return State.VelocityFriction; }
 
-		void SetFriction(const FReal InStaticFriction, const FReal InDynamicFriction)
+		void SetFriction(const FReal InStaticFriction, const FReal InDynamicFriction, const FReal InVelocityFriction)
 		{
 			State.StaticFriction = FSolverReal(InStaticFriction);
 			State.DynamicFriction = FSolverReal(InDynamicFriction);
+			State.VelocityFriction = FSolverReal(InVelocityFriction);
 		}
 
 		void SetStiffness(const FReal InStiffness)
@@ -210,22 +213,16 @@ namespace Chaos
 
 		void InitContact(
 			const int32 ManifoldPoiontIndex,
+			const FReal Dt,
 			const FVec3& InWorldAnchorPoint0,
 			const FVec3& InWorldAnchorPoint1,
-			const FVec3& InWorldContactNormal);
+			const FVec3& InWorldContactNormal,
+			const bool bAddVelocityToFriction);
 
 		void InitMaterial(
 			const int32 ManifoldPoiontIndex,
 			const FReal InRestitution,
-			const FReal InRestitutionVelocityThreshold,
-			const bool bInEnableStaticFriction,
-			const FReal InFrictionNetPushOutNormal);
-
-		void UpdateContact(
-			const int32 ManifoldPoiontIndex,
-			const FVec3& InWorldAnchorPoint0,
-			const FVec3& InWorldAnchorPoint1,
-			const FVec3& InWorldContactNormal);
+			const FReal InRestitutionVelocityThreshold);
 
 		/**
 		 * @brief Get the first (decaorated) solver body
@@ -289,9 +286,9 @@ namespace Chaos
 			FState()
 				: StaticFriction(0)
 				, DynamicFriction(0)
+				, VelocityFriction(0)
 				, Stiffness(1)
 				, NumManifoldPoints(0)
-				, bHaveRestitution(false)
 				, SolverBodies()
 				, ManifoldPoints()
 			{
@@ -302,16 +299,25 @@ namespace Chaos
 			{
 				StaticFriction = 0;
 				DynamicFriction = 0;
+				VelocityFriction = 0;
 				Stiffness = 1;
-				bHaveRestitution = false;
 				NumManifoldPoints = 0;
 			}
 
+			// Static Friction in the position-solve phase
 			FSolverReal StaticFriction;
+
+			// Dynamic Friction in the position-solve phase
 			FSolverReal DynamicFriction;
+			
+			// Dynamic Friction in the velocity-solve phase
+			FSolverReal VelocityFriction;
+
+			// Solver stiffness (scales all pushout and impulses)
 			FSolverReal Stiffness;
+
+			// Bodies and contacts
 			int32 NumManifoldPoints;
-			bool bHaveRestitution;
 			FConstraintSolverBody SolverBodies[MaxConstrainedBodies];
 			FPBDCollisionSolverManifoldPoint ManifoldPoints[MaxPointsPerConstraint];
 		};
@@ -369,16 +375,10 @@ namespace Chaos
 		FSolverReal& InOutPushOutTangentV,
 		FSolverReal& InOutNetPushOutTangentU,
 		FSolverReal& InOutNetPushOutTangentV,
-		FSolverReal& InOutFrictionNetPushOutNormal,
-		bool& bInOutInsideStaticFrictionCone)
+		FSolverReal& OutStaticFrictionRatio)
 	{
-		bInOutInsideStaticFrictionCone = true;
-		// Static friction limit: immediately increase maximum lateral correction, but smoothly decay maximum static friction limit. 
-		// This is so that small variations in position (jitter) and therefore NetPushOutNormal don't cause static friction to slip
-		// @todo(chaos): StaticFriction smoothing is iteration count depenendent - try to make it not so
-		const FSolverReal StaticFrictionLerpRate = CVars::Chaos_PBDCollisionSolver_Position_StaticFrictionLerpRate;
-		const FSolverReal StaticFrictionTarget = FMath::Max(NetPushOutNormal, FSolverReal(0));
-		InOutFrictionNetPushOutNormal = FMath::Lerp(FMath::Max(InOutFrictionNetPushOutNormal, StaticFrictionTarget), StaticFrictionTarget, StaticFrictionLerpRate);
+		// Assume we stay in the friction cone...
+		OutStaticFrictionRatio = FSolverReal(1);
 
 		if (NetPushOutNormal < FSolverReal(SMALL_NUMBER))
 		{
@@ -386,29 +386,26 @@ namespace Chaos
 			// so we must subtract it again to calculate the actual pushout we want to undo (i.e., the net pushout that has been applied to the body so far)
 			InOutPushOutTangentU = -(InOutNetPushOutTangentU - InOutPushOutTangentU);
 			InOutPushOutTangentV = -(InOutNetPushOutTangentV - InOutPushOutTangentV);
-			InOutNetPushOutTangentU = FReal(0);
-			InOutNetPushOutTangentV = FReal(0);
-			bInOutInsideStaticFrictionCone = false;
+			InOutNetPushOutTangentU = FSolverReal(0);
+			InOutNetPushOutTangentV = FSolverReal(0);
+			OutStaticFrictionRatio = FSolverReal(0);
 		}
-
-		// If we exceed the friction cone, clip to cone and stop adding frictional corrections
-		if (bInOutInsideStaticFrictionCone)
+		else
 		{
-			const FSolverReal MaxPushOutTangent = StaticFriction * InOutFrictionNetPushOutNormal;
-			const FSolverReal MaxPushOutTangentSq = FMath::Square(MaxPushOutTangent);
+			// If we exceed the static friction cone, clip to the dynamic friction cone
+			const FSolverReal MaxStaticPushOutTangentSq = FMath::Square(StaticFriction * NetPushOutNormal);
 			const FSolverReal NetPushOutTangentSq = FMath::Square(InOutNetPushOutTangentU) + FMath::Square(InOutNetPushOutTangentV);
-			if (NetPushOutTangentSq > MaxPushOutTangentSq)
+			if (NetPushOutTangentSq > MaxStaticPushOutTangentSq)
 			{
-				// @todo(chaos): clamp to dynamic friction rather than static
-				const FSolverReal ClippedMaxPushOutTangent = MaxPushOutTangent;
-				const FSolverReal InvNetPushOutTangent = FMath::InvSqrt(NetPushOutTangentSq);
-				const FSolverReal NetPushOutTangentU = ClippedMaxPushOutTangent * InOutNetPushOutTangentU * InvNetPushOutTangent;
-				const FSolverReal NetPushOutTangentV = ClippedMaxPushOutTangent * InOutNetPushOutTangentV * InvNetPushOutTangent;
+				const FSolverReal MaxDynamicPushOutTangent = DynamicFriction * NetPushOutNormal;
+				const FSolverReal FrictionMultiplier = MaxDynamicPushOutTangent * FMath::InvSqrt(NetPushOutTangentSq);
+				const FSolverReal NetPushOutTangentU = FrictionMultiplier * InOutNetPushOutTangentU;
+				const FSolverReal NetPushOutTangentV = FrictionMultiplier * InOutNetPushOutTangentV;
 				InOutPushOutTangentU = NetPushOutTangentU - (InOutNetPushOutTangentU - InOutPushOutTangentU);
 				InOutPushOutTangentV = NetPushOutTangentV - (InOutNetPushOutTangentV - InOutPushOutTangentV);
 				InOutNetPushOutTangentU = NetPushOutTangentU;
 				InOutNetPushOutTangentV = NetPushOutTangentV;
-				bInOutInsideStaticFrictionCone = false;
+				OutStaticFrictionRatio = FrictionMultiplier;
 			}
 		}
 	}
@@ -460,12 +457,10 @@ namespace Chaos
 			PushOutTangentV,									// Out
 			ManifoldPoint.NetPushOutTangentU,					// Out
 			ManifoldPoint.NetPushOutTangentV,					// Out
-			ManifoldPoint.FrictionNetPushOutNormal,				// Out
-			ManifoldPoint.bInsideStaticFrictionCone);			// Out
-
-		const FSolverVec3 PushOut = PushOutNormal * ManifoldPoint.WorldContactNormal + PushOutTangentU * ManifoldPoint.WorldContactTangentU + PushOutTangentV * ManifoldPoint.WorldContactTangentV;
+			ManifoldPoint.StaticFrictionRatio);					// Out
 
 		// Update the particle state based on the pushout
+		const FSolverVec3 PushOut = PushOutNormal * ManifoldPoint.WorldContactNormal + PushOutTangentU * ManifoldPoint.WorldContactTangentU + PushOutTangentV * ManifoldPoint.WorldContactTangentV;
 		if (Body0.IsDynamic())
 		{
 			const FSolverVec3 DX0 = Body0.InvM() * PushOut;
@@ -525,40 +520,45 @@ namespace Chaos
 		const FSolverReal ContactVelocityDeltaNormal,
 		const FSolverReal ContactVelocityDeltaTangent0,
 		const FSolverReal ContactVelocityDeltaTangent1,
+		const FSolverReal MinImpulseNormal,
 		FPBDCollisionSolverManifoldPoint& ManifoldPoint,
 		FConstraintSolverBody& Body0,
 		FConstraintSolverBody& Body1)
 	{
-		FSolverReal ImpulseNormal = -(Stiffness * ManifoldPoint.ContactMassNormal) * ContactVelocityDeltaNormal;
-		
+		FSolverReal ImpulseNormal = -Stiffness * ManifoldPoint.ContactMassNormal * ContactVelocityDeltaNormal;
+		FSolverReal ImpulseTangentU = -Stiffness * ManifoldPoint.ContactMassTangentU * ContactVelocityDeltaTangent0;
+		FSolverReal ImpulseTangentV = -Stiffness * ManifoldPoint.ContactMassTangentV * ContactVelocityDeltaTangent1;
+
 		// Clamp the total impulse to be positive along the normal. We can apply a net negative impulse, 
-		// but only to correct the velocity that was added by pushout.
+		// but only to correct the velocity that was added by pushout (in which case MinImpulseNormal will be negative).
 		const FSolverReal NetImpulseNormal = ManifoldPoint.NetImpulseNormal + ImpulseNormal;
-		const FSolverReal PushOutImpulseNormal = FMath::Max(0.0f, ManifoldPoint.NetPushOutNormal / Dt);
-		if (NetImpulseNormal < -PushOutImpulseNormal)
+		if ((ManifoldPoint.NetImpulseNormal + ImpulseNormal) < MinImpulseNormal)
 		{
 			// We are trying to apply a net negative impulse larger than one to counteract the effective pushout impulse
 			// so clamp the net impulse to be equal to minus the pushout impulse along the normal.
-			// NOTE: NetImpulseNormal is negative here
-			ImpulseNormal = ImpulseNormal - (NetImpulseNormal + PushOutImpulseNormal);
+			ImpulseNormal = MinImpulseNormal - ManifoldPoint.NetImpulseNormal;
 		}
 
-		FSolverReal ImpulseTangentU = FSolverReal(0);
-		FSolverReal ImpulseTangentV = FSolverReal(0);
-		const bool bApplyFriction = (DynamicFriction > 0) && !FMath::IsNearlyZero(ImpulseNormal) && (Dt > 0);
-		if (bApplyFriction)
+		// Clamp the tangential impulses to the friction cone
+		if ((DynamicFriction > 0) && (Dt > 0))
 		{
-			ImpulseTangentU = -(Stiffness * ManifoldPoint.ContactMassTangentU) * ContactVelocityDeltaTangent0;
-			ImpulseTangentV = -(Stiffness * ManifoldPoint.ContactMassTangentV) * ContactVelocityDeltaTangent1;
+			const FSolverReal MaxImpulseTangent = FMath::Max(FSolverReal(0), DynamicFriction * (ManifoldPoint.NetImpulseNormal + ImpulseNormal + ManifoldPoint.NetPushOutNormal / Dt));
+			const FSolverReal MaxImpulseTangentSq = FMath::Square(MaxImpulseTangent);
+			const FSolverReal ImpulseTangentSq = FMath::Square(ImpulseTangentU) + FMath::Square(ImpulseTangentV);
+			if (ImpulseTangentSq > MaxImpulseTangentSq)
+			{
+				const FSolverReal ImpulseTangentScale = MaxImpulseTangent * FMath::InvSqrt(ImpulseTangentSq);
+				ImpulseTangentU *= ImpulseTangentScale;
+				ImpulseTangentV *= ImpulseTangentScale;
+			}
 		}
 
 		ManifoldPoint.NetImpulseNormal += ImpulseNormal;
 		ManifoldPoint.NetImpulseTangentU += ImpulseTangentU;
 		ManifoldPoint.NetImpulseTangentV += ImpulseTangentV;
 
-		FSolverVec3 Impulse = ImpulseNormal * ManifoldPoint.WorldContactNormal + ImpulseTangentU * ManifoldPoint.WorldContactTangentU + ImpulseTangentV * ManifoldPoint.WorldContactTangentV;
-
-		// Calculate the velocity deltas from the impulse
+		// Apply the velocity deltas from the impulse
+		const FSolverVec3 Impulse = ImpulseNormal * ManifoldPoint.WorldContactNormal + ImpulseTangentU * ManifoldPoint.WorldContactTangentU + ImpulseTangentV * ManifoldPoint.WorldContactTangentV;
 		if (Body0.IsDynamic())
 		{
 			const FSolverVec3 DV0 = Body0.InvM() * Impulse;
@@ -575,8 +575,8 @@ namespace Chaos
 
 	FORCEINLINE_DEBUGGABLE void ApplyVelocityCorrectionNormal(
 		const FSolverReal Stiffness,
-		const FSolverReal Dt,
 		const FSolverReal ContactVelocityDeltaNormal,
+		const FSolverReal MinImpulseNormal,
 		FPBDCollisionSolverManifoldPoint& ManifoldPoint,
 		FConstraintSolverBody& Body0,
 		FConstraintSolverBody& Body1)
@@ -585,17 +585,15 @@ namespace Chaos
 
 		// See comments in ApplyVelocityCorrection
 		const FSolverReal NetImpulseNormal = ManifoldPoint.NetImpulseNormal + ImpulseNormal;
-		const FSolverReal PushOutImpulseNormal = FMath::Max(0.0f, ManifoldPoint.NetPushOutNormal / Dt);
-		if (NetImpulseNormal < -PushOutImpulseNormal)
+		if (NetImpulseNormal < MinImpulseNormal)
 		{
-			ImpulseNormal = ImpulseNormal - (NetImpulseNormal + PushOutImpulseNormal);
+			ImpulseNormal = ImpulseNormal - NetImpulseNormal + MinImpulseNormal;
 		}
 
 		ManifoldPoint.NetImpulseNormal += ImpulseNormal;
 
-		FSolverVec3 Impulse = ImpulseNormal * ManifoldPoint.WorldContactNormal;
-
 		// Calculate the velocity deltas from the impulse
+		FSolverVec3 Impulse = ImpulseNormal * ManifoldPoint.WorldContactNormal;
 		if (Body0.IsDynamic())
 		{
 			const FSolverVec3 DV0 = Body0.InvM() * Impulse;
@@ -619,8 +617,8 @@ namespace Chaos
 	FORCEINLINE_DEBUGGABLE void FPBDCollisionSolverManifoldPoint::CalculateContactPositionErrorNormal(const FConstraintSolverBody& Body0, const FConstraintSolverBody& Body1, const FSolverReal MaxPushOut, FSolverReal& OutContactDeltaNormal) const
 	{
 		// Linear version: calculate the contact delta assuming linear motion after applying a positional impulse at the contact point. There will be an error that depends on the size of the rotation.
-		const FSolverVec3 ContactDelta0 = FSolverVec3(Body0.DP()) + FSolverVec3::CrossProduct(Body0.DQ(), RelativeContactPosition0);
-		const FSolverVec3 ContactDelta1 = FSolverVec3(Body1.DP()) + FSolverVec3::CrossProduct(Body1.DQ(), RelativeContactPosition1);
+		const FSolverVec3 ContactDelta0 = Body0.DP() + FSolverVec3::CrossProduct(Body0.DQ(), RelativeContactPosition0);
+		const FSolverVec3 ContactDelta1 = Body1.DP() + FSolverVec3::CrossProduct(Body1.DQ(), RelativeContactPosition1);
 		const FSolverVec3 ContactDelta = ContactDelta0 - ContactDelta1;
 		OutContactDeltaNormal = WorldContactDeltaNormal + FSolverVec3::DotProduct(ContactDelta, WorldContactNormal);
 
@@ -633,23 +631,26 @@ namespace Chaos
 		}
 	}
 
-	FORCEINLINE_DEBUGGABLE void FPBDCollisionSolverManifoldPoint::CalculateContactPositionError(const FConstraintSolverBody& Body0, const FConstraintSolverBody& Body1, const FSolverReal MaxPushOut, FSolverReal& OutContactDeltaNormal, FSolverReal& OutContactDeltaTangent0, FSolverReal& OutContactDeltaTangent1) const
+	FORCEINLINE_DEBUGGABLE void FPBDCollisionSolverManifoldPoint::CalculateContactPositionError(const FConstraintSolverBody& Body0, const FConstraintSolverBody& Body1, const FSolverReal MaxPushOut, FSolverReal& OutContactDeltaNormal, FSolverReal& OutContactDeltaTangentU, FSolverReal& OutContactDeltaTangentV) const
 	{
 #if CHAOS_NONLINEAR_COLLISIONS_ENABLED
-#error "Non-linear collision solver no longer functional - fix it"
+#error "Non-linear collision solver probably no longer functional - fix it!"
 		// Non-linear version: calculate the contact delta after we have converted the current positional impulses into position and rotation corrections.
 		// We could precalculate and store the LocalContactPositions if we really want to use this nonlinear version
-		const FSolverVec3 LocalContactPosition0 = Body0.Q().Inverse() * WorldRelativeContactPosition0;
-		const FSolverVec3 LocalContactPosition1 = Body1.Q().Inverse() * WorldRelativeContactPosition1;
+		const FSolverVec3 LocalContactPosition0 = Body0.Q().Inverse() * RelativeContactPosition0;
+		const FSolverVec3 LocalContactPosition1 = Body1.Q().Inverse() * RelativeContactPosition1;
 		const FSolverVec3 ContactDelta = (Body0.CorrectedP() + Body0.CorrectedQ() * LocalContactPosition0) - (Body1.CorrectedP() + Body1.CorrectedQ() * LocalContactPosition1);
+		OutContactDeltaNormal = FSolverVec3::DotProduct(ContactDelta, WorldContactNormal);
+		OutContactDeltaTangentU = FSolverVec3::DotProduct(ContactDelta, WorldContactTangentU);
+		OutContactDeltaTangentV = FSolverVec3::DotProduct(ContactDelta, WorldContactTangentV);
 #else
 		// Linear version: calculate the contact delta assuming linear motion after applying a positional impulse at the contact point. There will be an error that depends on the size of the rotation.
-		const FSolverVec3 ContactDelta0 = FSolverVec3(Body0.DP()) + FSolverVec3::CrossProduct(Body0.DQ(), RelativeContactPosition0);
-		const FSolverVec3 ContactDelta1 = FSolverVec3(Body1.DP()) + FSolverVec3::CrossProduct(Body1.DQ(), RelativeContactPosition1);
+		const FSolverVec3 ContactDelta0 = Body0.DP() + FSolverVec3::CrossProduct(Body0.DQ(), RelativeContactPosition0);
+		const FSolverVec3 ContactDelta1 = Body1.DP() + FSolverVec3::CrossProduct(Body1.DQ(), RelativeContactPosition1);
 		const FSolverVec3 ContactDelta = ContactDelta0 - ContactDelta1;
 		OutContactDeltaNormal = WorldContactDeltaNormal + FSolverVec3::DotProduct(ContactDelta, WorldContactNormal);
-		OutContactDeltaTangent0 = WorldContactDeltaTangent0 + FSolverVec3::DotProduct(ContactDelta, WorldContactTangentU);
-		OutContactDeltaTangent1 = WorldContactDeltaTangent1 + FSolverVec3::DotProduct(ContactDelta, WorldContactTangentV);
+		OutContactDeltaTangentU = WorldContactDeltaTangentU + FSolverVec3::DotProduct(ContactDelta, WorldContactTangentU);
+		OutContactDeltaTangentV = WorldContactDeltaTangentV + FSolverVec3::DotProduct(ContactDelta, WorldContactTangentV);
 #endif
 
 		// NOTE: OutContactDeltaNormal is negative for penetration
@@ -666,35 +667,12 @@ namespace Chaos
 		const FSolverVec3 ContactVelocity1 = Body1.V() + FSolverVec3::CrossProduct(Body1.W(), RelativeContactPosition1);
 		const FSolverVec3 ContactVelocity = ContactVelocity0 - ContactVelocity1;
 		const FSolverReal ContactVelocityNormal = FSolverVec3::DotProduct(ContactVelocity, WorldContactNormal);
+		const FSolverReal ContactVelocityTangent0 = FSolverVec3::DotProduct(ContactVelocity, WorldContactTangentU);
+		const FSolverReal ContactVelocityTangent1 = FSolverVec3::DotProduct(ContactVelocity, WorldContactTangentV);
 
-		// Add up the errors in the velocity (current velocity - desired velocity)
 		OutContactVelocityDeltaNormal = (ContactVelocityNormal - WorldContactVelocityTargetNormal);
-		OutContactVelocityDeltaTangent0 = FSolverReal(0);
-		OutContactVelocityDeltaTangent1 = FSolverReal(0);
-
-		const bool bApplyFriction = (DynamicFriction > 0) && (Dt > 0);
-		if (bApplyFriction)
-		{
-			const FSolverReal ContactVelocityTangent0 = FSolverVec3::DotProduct(ContactVelocity, WorldContactTangentU);
-			const FSolverReal ContactVelocityTangent1 = FSolverVec3::DotProduct(ContactVelocity, WorldContactTangentV);
-			OutContactVelocityDeltaTangent0 = ContactVelocityTangent0;
-			OutContactVelocityDeltaTangent1 = ContactVelocityTangent1;
-
-			const FSolverReal ContactVelocityTangentialSq = FMath::Square(ContactVelocityTangent0) + FMath::Square(ContactVelocityTangent1);
-			if (ContactVelocityTangentialSq > FSolverReal(SMALL_NUMBER))
-			{
-				// PushOut = ContactMass * DP, where DP is the contact positional correction
-				// Friction force is proportional to the normal force, so friction velocity correction
-				// is proprtional to normal velocity correction, or DVn = DPn/dt = PushOut.N / (ContactMass * dt);
-				const FSolverReal MaxContactVelocityTangential = DynamicFriction * NetPushOutNormal / (ContactMassNormal * Dt);
-				if (ContactVelocityTangentialSq > FMath::Square(MaxContactVelocityTangential))
-				{
-					const FSolverReal ContactVelocityTangentialScale = MaxContactVelocityTangential * FMath::InvSqrt(ContactVelocityTangentialSq);
-					OutContactVelocityDeltaTangent0 *= ContactVelocityTangentialScale;
-					OutContactVelocityDeltaTangent1 *= ContactVelocityTangentialScale;
-				}
-			}
-		}
+		OutContactVelocityDeltaTangent0 = ContactVelocityTangent0;
+		OutContactVelocityDeltaTangent1 = ContactVelocityTangent1;
 	}
 
 	FORCEINLINE_DEBUGGABLE void FPBDCollisionSolverManifoldPoint::CalculateContactVelocityErrorNormal(const FConstraintSolverBody& Body0, const FConstraintSolverBody& Body1, FSolverReal& OutContactVelocityDeltaNormal) const
@@ -708,6 +686,12 @@ namespace Chaos
 		OutContactVelocityDeltaNormal = (ContactVelocityNormal - WorldContactVelocityTargetNormal);
 	}
 
+	FORCEINLINE_DEBUGGABLE bool FPBDCollisionSolverManifoldPoint::ShouldSolveVelocity() const
+	{
+		// We ensure positive separating velocity for close contacts even if they didn't receive a pushout
+		return (NetImpulseNormal > FSolverReal(0)) || (WorldContactDeltaNormal < FSolverReal(0));
+	}
+
 	//////////////////////////////////////////////////////////////////////////////////////////////////
 	//////////////////////////////////////////////////////////////////////////////////////////////////
 	//////////////////////////////////////////////////////////////////////////////////////////////////
@@ -728,34 +712,19 @@ namespace Chaos
 			FSolverReal ContactDeltaNormal, ContactDeltaTangent0, ContactDeltaTangent1;
 			SolverManifoldPoint.CalculateContactPositionError(Body0.SolverBody(), Body1.SolverBody(), MaxPushOut, ContactDeltaNormal, ContactDeltaTangent0, ContactDeltaTangent1);
 
-			const bool bProcessManifoldPoint = (ContactDeltaNormal < FSolverReal(0)) || !FMath::IsNearlyZero(SolverManifoldPoint.NetPushOutNormal);
+			const bool bProcessManifoldPoint = (ContactDeltaNormal < FSolverReal(0)) || (SolverManifoldPoint.NetPushOutNormal > FSolverReal(SMALL_NUMBER));
 			if (bProcessManifoldPoint)
 			{
-				// If we want to add static friction...
-				// (note: we run a few iterations without friction by temporarily setting StaticFriction to 0)
-				const bool bApplyFriction = (State.StaticFriction > 0) && SolverManifoldPoint.bInsideStaticFrictionCone;
-				if (bApplyFriction)
-				{
-					ApplyPositionCorrectionWithFriction(
-						State.Stiffness,
-						State.StaticFriction,
-						State.DynamicFriction,
-						ContactDeltaNormal,
-						ContactDeltaTangent0,
-						ContactDeltaTangent1,
-						SolverManifoldPoint,
-						Body0,
-						Body1);
-				}
-				else
-				{
-					ApplyPositionCorrectionNoFriction(
-						State.Stiffness, 
-						ContactDeltaNormal, 
-						SolverManifoldPoint, 
-						Body0, 
-						Body1);
-				}
+				ApplyPositionCorrectionWithFriction(
+					State.Stiffness,
+					State.StaticFriction,
+					State.DynamicFriction,
+					ContactDeltaNormal,
+					ContactDeltaTangent0,
+					ContactDeltaTangent1,
+					SolverManifoldPoint,
+					Body0,
+					Body1);
 			}
 		}
 
@@ -776,7 +745,7 @@ namespace Chaos
 			FSolverReal ContactDeltaNormal;
 			SolverManifoldPoint.CalculateContactPositionErrorNormal(Body0.SolverBody(), Body1.SolverBody(), MaxPushOut, ContactDeltaNormal);
 
-			const bool bProcessManifoldPoint = (ContactDeltaNormal < FSolverReal(0)) || !FMath::IsNearlyZero(SolverManifoldPoint.NetPushOutNormal);
+			const bool bProcessManifoldPoint = (ContactDeltaNormal < FSolverReal(0)) || (SolverManifoldPoint.NetPushOutNormal > FSolverReal(SMALL_NUMBER));
 			if (bProcessManifoldPoint)
 			{
 				ApplyPositionCorrectionNoFriction(
@@ -793,4 +762,3 @@ namespace Chaos
 
 }
 
-CHAOS_API DECLARE_LOG_CATEGORY_EXTERN(LogChaosCollision, Log, All);
diff --git a/Engine/Source/Runtime/Experimental/Chaos/Private/Chaos/Collision/ParticlePairMidPhase.cpp b/Engine/Source/Runtime/Experimental/Chaos/Private/Chaos/Collision/ParticlePairMidPhase.cpp
--- a/Engine/Source/Runtime/Experimental/Chaos/Private/Chaos/Collision/ParticlePairMidPhase.cpp
+++ b/Engine/Source/Runtime/Experimental/Chaos/Private/Chaos/Collision/ParticlePairMidPhase.cpp
@@ -24,12 +24,14 @@ namespace Chaos
 		Chaos::FRealSingle Chaos_Collision_RestoreTolerance_NoContact_Rotation = 0.1f;		// About 10deg
 		Chaos::FRealSingle Chaos_Collision_RestoreTolerance_Contact_Position = 0.02f;		// About 2cm for a meter cube
 		Chaos::FRealSingle Chaos_Collision_RestoreTolerance_Contact_Rotation = 0.1f;		// About 10deg
+		Chaos::FRealSingle Chaos_Collision_RestoreTolerance_ReprojectFactor = 0.0f;
 		FAutoConsoleVariableRef CVarChaos_Collision_RestoreTolerance_NoContact_Position(TEXT("p.Chaos.Collision.RestoreTolerance.NoContact.Position"), Chaos_Collision_RestoreTolerance_NoContact_Position, TEXT("Fraction of Size. Particle pairs that move less than this may have their contacts reinstated"));
 		FAutoConsoleVariableRef CVarChaos_Collision_RestoreTolerance_NoContact_Rotation(TEXT("p.Chaos.Collision.RestoreTolerance.NoContact.Rotation"), Chaos_Collision_RestoreTolerance_NoContact_Rotation, TEXT("Quaternion Dot Product Limit. Particle pairs that move less than this may have their contacts reinstated"));
 		FAutoConsoleVariableRef CVarChaos_Collision_RestoreTolerance_Contact_Position(TEXT("p.Chaos.Collision.RestoreTolerance.WithContact.Position"), Chaos_Collision_RestoreTolerance_Contact_Position, TEXT("Fraction of Size. Particle pairs that move less than this may have their contacts reinstated"));
 		FAutoConsoleVariableRef CVarChaos_Collision_RestoreTolerance_Contact_Rotation(TEXT("p.Chaos.Collision.RestoreTolerance.WithContact.Rotation"), Chaos_Collision_RestoreTolerance_Contact_Rotation, TEXT("Quaternion Dot Product Limit. Particle pairs that move less than this may have their contacts reinstated"));
+		FAutoConsoleVariableRef CVarChaos_Collision_RestoreTolerance_ReprojectFactor(TEXT("p.Chaos.Collision.RestoreTolerance.ReprojectFactor"), Chaos_Collision_RestoreTolerance_ReprojectFactor, TEXT("Quaternion Dot Product Limit. Particle pairs that move less than this may have their contacts reinstated"));
 
-		bool bChaos_Collision_EnableManifoldRestore = true;
+		bool bChaos_Collision_EnableManifoldRestore = false;	// This causes subtle problems with friction unless points are reprojected, but then it's basically the same as EnableManifoldUpdate
 		bool bChaos_Collision_EnableManifoldUpdate = true;
 		FAutoConsoleVariableRef CVarChaos_Collision_EnableManifoldRestore(TEXT("p.Chaos.Collision.EnableManifoldRestore"), bChaos_Collision_EnableManifoldRestore, TEXT(""));
 		FAutoConsoleVariableRef CVarChaos_Collision_EnableManifoldUpdate(TEXT("p.Chaos.Collision.EnableManifoldUpdate"), bChaos_Collision_EnableManifoldUpdate, TEXT(""));
@@ -347,7 +349,7 @@ namespace Chaos
 		return 0;
 	}
 
-	int32 FSingleShapePairCollisionDetector::RestoreCollision(const FReal CullDistance)
+	int32 FSingleShapePairCollisionDetector::RestoreCollision(const FReal CullDistance, const bool bReproject)
 	{
 		// Only restore constraints if active last tick. Any older than that and the shapes were separated for a bit
 		const int32 CurrentEpoch = MidPhase.GetCollisionAllocator().GetCurrentEpoch();
@@ -358,7 +360,8 @@ namespace Chaos
 			const FRigidTransform3& ShapeWorldTransform1 = Shape1->GetLeafWorldTransform();
 			Constraint->SetShapeWorldTransforms(ShapeWorldTransform0, ShapeWorldTransform1);
 
-			Constraint->RestoreManifold();
+			Constraint->RestoreManifold(bReproject);
+
 			if (Constraint->GetPhi() <= CullDistance)
 			{
 				if (MidPhase.GetCollisionAllocator().ActivateConstraint(Constraint.Get()))
@@ -601,7 +604,7 @@ namespace Chaos
 		return Constraints.Add(Key.GetKey(), MoveTemp(Constraint)).Get();
 	}
 
-	int32 FMultiShapePairCollisionDetector::RestoreCollisions(const FReal CullDistance)
+	int32 FMultiShapePairCollisionDetector::RestoreCollisions(const FReal CullDistance, const bool bReproject)
 	{
 		if (Constraints.Num() == 0)
 		{
@@ -622,7 +625,8 @@ namespace Chaos
 				const FRigidTransform3 ShapeWorldTransform1 = Constraint->ImplicitTransform[1] * ParticleTransform1;
 				Constraint->SetShapeWorldTransforms(ShapeWorldTransform0, ShapeWorldTransform1);
 
-				Constraint->RestoreManifold();
+				Constraint->RestoreManifold(bReproject);
+
 				if (Constraint->GetPhi() < CullDistance)
 				{
 					if (MidPhase.GetCollisionAllocator().ActivateConstraint(Constraint.Get()))
@@ -903,7 +907,7 @@ namespace Chaos
 		CullDistanceScale = FMath::Max3(CullDistanceScale0, CullDistanceScale1, MinCullDistanceScale);
 	}
 
-	bool FParticlePairMidPhase::ShouldRestoreConstraints(const FReal Dt)
+	ECollisionRestoreType FParticlePairMidPhase::ShouldRestoreConstraints(const FReal Dt)
 	{
 		const FVec3& ParticleP0 = FConstGenericParticleHandle(Particle0)->P();
 		const FRotation3& ParticleQ0 = FConstGenericParticleHandle(Particle0)->Q();
@@ -911,28 +915,53 @@ namespace Chaos
 		const FRotation3& ParticleQ1 = FConstGenericParticleHandle(Particle1)->Q();
 
 		// We can only restore collisions if they were created or updated last tick
+		ECollisionRestoreType RestoreType = ECollisionRestoreType::None;
 		if (IsUsedSince(CollisionAllocator->GetCurrentEpoch() - 1))
 		{
-			const FReal PositionThreshold = (NumActiveConstraints == 0) ? RestoreThresholdZeroContacts.PositionThreshold : RestoreThreshold.PositionThreshold;
-			const FReal RotationThreshold = (NumActiveConstraints == 0) ? RestoreThresholdZeroContacts.RotationThreshold : RestoreThreshold.RotationThreshold;
+			const FReal PositionReprojectThreshold = (NumActiveConstraints == 0) ? RestoreThresholdZeroContacts.PositionThreshold : RestoreThreshold.PositionThreshold;
+			const FReal PositionReuseThreshold = (NumActiveConstraints == 0) ? PositionReprojectThreshold : Chaos_Collision_RestoreTolerance_ReprojectFactor * PositionReprojectThreshold;
 
-			// If either particle has moved or rotated in world space we cannot reuse the constraint
-			if ((ParticleP0 - RestoreParticleP0).IsNearlyZero(PositionThreshold) && (ParticleP1 - RestoreParticleP1).IsNearlyZero(PositionThreshold))
+			const FReal MaxDP0 = FVec3(ParticleP0 - RestoreParticleP0).GetAbsMax();
+			const FReal MaxDP1 = FVec3(ParticleP1 - RestoreParticleP1).GetAbsMax();
+			const FReal MaxDP = FMath::Max(MaxDP0, MaxDP1);
+
+			// Check position delta and calculate the best-case reuse option (ignoring rotation for now)
+			if (MaxDP < PositionReuseThreshold)
 			{
-				if (FRotation3::IsNearlyEqual(ParticleQ0, RestoreParticleQ0, RotationThreshold) && FRotation3::IsNearlyEqual(ParticleQ1, RestoreParticleQ1, RotationThreshold))
+				// We moved a tiny amount - reuse the manifold exactly as-is
+				RestoreType = ECollisionRestoreType::Reuse;
+			}
+			else if (MaxDP < PositionReprojectThreshold)
+			{
+				// We moved a bit but within threshold - reuse the manifold but reproject the vertex contacts onto the planes
+				RestoreType = ECollisionRestoreType::Reproject;
+			}
+
+			// If we passed the position test, check the rotation threshold and maybe veto the restore decision
+			if (RestoreType != ECollisionRestoreType::None)
+			{
+				const FReal RotationThreshold = (NumActiveConstraints == 0) ? RestoreThresholdZeroContacts.RotationThreshold : RestoreThreshold.RotationThreshold;
+				const FReal MaxDQ0 = FRotation3(ParticleQ0 - RestoreParticleQ0).GetAbsMax();
+				const FReal MaxDQ1 = FRotation3(ParticleQ1 - RestoreParticleQ1).GetAbsMax();
+				const FReal MaxDQ = FMath::Max(MaxDQ0, MaxDQ1);
+				if (MaxDQ > RotationThreshold)
 				{
-					// We passed the gauntlet - reuse the constraint
-					return true;
+					RestoreType = ECollisionRestoreType::None;
 				}
 			}
 		}
 
-		// We have moved and should rebuild the manifold. Update the current manifold transforms for future restore checks
-		RestoreParticleP0 = ParticleP0;
-		RestoreParticleP1 = ParticleP1;
-		RestoreParticleQ0 = ParticleQ0;
-		RestoreParticleQ1 = ParticleQ1;
-		return false;
+		if (RestoreType != ECollisionRestoreType::None)
+		{
+			// We have moved and will rebuild the manifold so
+			// update the current manifold transforms for future restore checks
+			RestoreParticleP0 = ParticleP0;
+			RestoreParticleP1 = ParticleP1;
+			RestoreParticleQ0 = ParticleQ0;
+			RestoreParticleQ1 = ParticleQ1;
+		}
+
+		return RestoreType;
 	}
 
 	bool FParticlePairMidPhase::TryRestoreConstraints(const FReal Dt, const FReal CullDistance)
@@ -941,16 +970,19 @@ namespace Chaos
 		PHYSICS_CSV_SCOPED_EXPENSIVE(PhysicsVerbose, DetectCollisions_RestoreCollision);
 
 		// If the particles haven't moved relative to each other, we can just reuse the constraint as-is
-		if (ShouldRestoreConstraints(Dt))
+		const ECollisionRestoreType CollisionRestoreType = ShouldRestoreConstraints(Dt);
+		if (CollisionRestoreType != ECollisionRestoreType::None)
 		{
+			const bool bReproject = (CollisionRestoreType == ECollisionRestoreType::Reproject);
+
 			int32 NumRestored = 0;
 			for (FSingleShapePairCollisionDetector& ShapePair : ShapePairDetectors)
 			{
-				NumRestored += ShapePair.RestoreCollision(CullDistance);
+				NumRestored += ShapePair.RestoreCollision(CullDistance, bReproject);
 			}
 			for (FMultiShapePairCollisionDetector& MultiShapePair : MultiShapePairDetectors)
 			{
-				NumRestored += MultiShapePair.RestoreCollisions(CullDistance);
+				NumRestored += MultiShapePair.RestoreCollisions(CullDistance, bReproject);
 			}
 
 			NumActiveConstraints = NumRestored;
diff --git a/Engine/Source/Runtime/Experimental/Chaos/Private/Chaos/Collision/SolverCollisionContainer.cpp b/Engine/Source/Runtime/Experimental/Chaos/Private/Chaos/Collision/SolverCollisionContainer.cpp
--- a/Engine/Source/Runtime/Experimental/Chaos/Private/Chaos/Collision/SolverCollisionContainer.cpp
+++ b/Engine/Source/Runtime/Experimental/Chaos/Private/Chaos/Collision/SolverCollisionContainer.cpp
@@ -76,9 +76,20 @@ namespace Chaos
 			Body0->SetLevel(Particle0Level);
 			Body1->SetLevel(Particle1Level);
 
-			Solver.SetFriction(
-				FMath::Max(Constraint->Manifold.AngularFriction, Constraint->Manifold.Friction),
-				Constraint->Manifold.Friction);
+			// Friction values. Static and Dynamic friction are applied in the position solve for most shapes.
+			// For quadratic shapes, we run dynamic friction in the velocity solve for better rolling behaviour.
+			// @todo(chaos): fix static/dynamic friction for quadratic shapes
+			const FReal StaticFriction = FMath::Max(Constraint->Manifold.AngularFriction, Constraint->Manifold.Friction);
+			const FReal DynamicFriction = Constraint->Manifold.Friction;
+			const bool bIsQuadratic = Constraint->HasQuadraticShape();
+			if (!bIsQuadratic)
+			{
+				Solver.SetFriction(StaticFriction, DynamicFriction, FReal(0));
+			}
+			else
+			{
+				Solver.SetFriction(StaticFriction, FReal(0), DynamicFriction);
+			}
 
 			Solver.SetStiffness(Constraint->GetStiffness());
 
@@ -104,20 +115,23 @@ namespace Chaos
 				TArrayView<FManifoldPoint> ManifoldPoints = Constraint->GetManifoldPoints();
 				const FManifoldPoint& ManifoldPoint = ManifoldPoints[ManifoldPointIndex];
 
+				// If we didn't have friction data restored, we estimate the static friction target by adding velocity to the contact error				
+				const bool bAddVelocityToFriction = !ManifoldPoint.Flags.bWasFrictionRestored;
+
 				// Initialize the structural data in the contact (relative contact points, contact mass etc)
 				Solver.InitContact(
 					ManifoldPointIndex,
+					Dt,
 					ManifoldPoint.WorldContactPoints[0],
 					ManifoldPoint.WorldContactPoints[1],
-					ManifoldPoint.ContactPoint.Normal);
+					ManifoldPoint.ContactPoint.Normal,
+					bAddVelocityToFriction);
 
 				// Initialize the material properties (restitution and friction related)
 				Solver.InitMaterial(
 					ManifoldPointIndex,
 					Restitution,
-					RestitutionVelocityThreshold,
-					(Solver.StaticFriction() > 0),
-					ManifoldPoint.StaticFrictionMax);
+					RestitutionVelocityThreshold);
 			}
 		}
 
@@ -147,8 +161,7 @@ namespace Chaos
 				Constraint->SetSolverResults(PointIndex, 
 					NetPushOut, 
 					NetImpulse, 
-					SolverManifoldPoint.bInsideStaticFrictionCone,
-					SolverManifoldPoint.FrictionNetPushOutNormal,
+					SolverManifoldPoint.StaticFrictionRatio,
 					Dt);
 			}
 
diff --git a/Engine/Source/Runtime/Experimental/Chaos/Private/Chaos/ContactModification.cpp b/Engine/Source/Runtime/Experimental/Chaos/Private/Chaos/ContactModification.cpp
--- a/Engine/Source/Runtime/Experimental/Chaos/Private/Chaos/ContactModification.cpp
+++ b/Engine/Source/Runtime/Experimental/Chaos/Private/Chaos/ContactModification.cpp
@@ -128,6 +128,9 @@ namespace Chaos
 
 		ManifoldPoint.ContactPoint.ShapeContactPoints[0] = Constraint->GetShapeWorldTransform0().InverseTransformPositionNoScale(Location0);
 		ManifoldPoint.ContactPoint.ShapeContactPoints[1] = Constraint->GetShapeWorldTransform1().InverseTransformPositionNoScale(Location1);
+		// @todo(chaos): Overwriting ShapeAnchorPoints disables static friction for this tick - we might want to do something better here
+		ManifoldPoint.ShapeAnchorPoints[0] = ManifoldPoint.ContactPoint.ShapeContactPoints[0];
+		ManifoldPoint.ShapeAnchorPoints[1] = ManifoldPoint.ContactPoint.ShapeContactPoints[1];
 		ManifoldPoint.WorldContactPoints[0] = Location0;
 		ManifoldPoint.WorldContactPoints[1] = Location1;
 		ManifoldPoint.ContactPoint.Location = 0.5 * (Location0 + Location1);
diff --git a/Engine/Source/Runtime/Experimental/Chaos/Private/PBDRigidsSolver.cpp b/Engine/Source/Runtime/Experimental/Chaos/Private/PBDRigidsSolver.cpp
--- a/Engine/Source/Runtime/Experimental/Chaos/Private/PBDRigidsSolver.cpp
+++ b/Engine/Source/Runtime/Experimental/Chaos/Private/PBDRigidsSolver.cpp
@@ -1446,7 +1446,7 @@ CSV_CUSTOM_STAT(PhysicsCounters, Name, Value, ECsvCustomStatOp::Set);
 				for (const FManifoldPoint& ManifoldPoint : Collision->GetContact().GetManifoldPoints())
 				{
 					++NumManifoldPoints;
-					if (ManifoldPoint.bWasRestored || Collision->GetContact().WasManifoldRestored())
+					if (ManifoldPoint.Flags.bWasRestored || Collision->GetContact().WasManifoldRestored())
 					{
 						++NumRestoredManifoldPoints;
 					}
diff --git a/Engine/Source/Runtime/Experimental/Chaos/Public/Chaos/Collision/PBDCollisionConstraint.h b/Engine/Source/Runtime/Experimental/Chaos/Public/Chaos/Collision/PBDCollisionConstraint.h
--- a/Engine/Source/Runtime/Experimental/Chaos/Public/Chaos/Collision/PBDCollisionConstraint.h
+++ b/Engine/Source/Runtime/Experimental/Chaos/Public/Chaos/Collision/PBDCollisionConstraint.h
@@ -43,88 +43,52 @@ namespace Chaos
 	public:
 		FManifoldPoint() 
 			: ContactPoint()
+			, ShapeAnchorPoints{ FVec3(0), FVec3(0) }
 			, InitialShapeContactPoints{ FVec3(0), FVec3(0) }
 			, WorldContactPoints{ FVec3(0), FVec3(0) }
 			, NetPushOut(0)
 			, NetImpulse(0)
-			, StaticFrictionMax(0)
-			, bInsideStaticFrictionCone(false)
-			, bWasRestored(false)
+			, Flags()
 		{}
 
 		FManifoldPoint(const FContactPoint& InContactPoint) 
 			: ContactPoint(InContactPoint)
+			, ShapeAnchorPoints{ FVec3(0), FVec3(0) }
 			, InitialShapeContactPoints{ FVec3(0), FVec3(0) }
 			, WorldContactPoints{ FVec3(0), FVec3(0) }
 			, NetPushOut(0)
 			, NetImpulse(0)
-			, StaticFrictionMax(0)
-			, bInsideStaticFrictionCone(false)
-			, bWasRestored(false)
+			, Flags()
 		{}
 
 		FContactPoint ContactPoint;			// Contact point results of low-level collision detection
+		FVec3 ShapeAnchorPoints[2];
 		FVec3 InitialShapeContactPoints[2];	// ShapeContactPoints when the constraint was first initialized. Used to track reusablility
 		FVec3 WorldContactPoints[2];		// World-space contact points on the two bodies
 		FVec3 NetPushOut;					// Total pushout applied at this contact point
 		FVec3 NetImpulse;					// Total impulse applied by this contact point
-		FReal StaticFrictionMax;			// A proxy for the normal impulse used to limit static friction correction. Used for smoothing static friction limits
-		bool bInsideStaticFrictionCone;		// Whether we are inside the static friction cone (used in PushOut)
-		bool bWasRestored;
+
+		union FFlags
+		{
+			FFlags() : Bits(0) {}
+			struct
+			{
+				uint32 bInsideStaticFrictionCone : 1;		// Whether we are inside the static friction cone (used in PushOut)
+				uint32 bWasRestored : 1;
+				uint32 bWasFrictionRestored : 1;
+			};
+			uint32 Bits;
+		} Flags;
 	};
 
 	/**
 	 * @brief The friction data for a manifold point
 	 * This is the information that needs to be stored between ticks to implement static friction.
 	*/
-	class CHAOS_API FManifoldPointSavedData
+	class CHAOS_API FSavedManifoldPoint
 	{
 	public:
-		FManifoldPointSavedData()
-		{
-		}
-
-		FManifoldPointSavedData(const FManifoldPoint& ManifoldPoint)
-		{
-			Save(ManifoldPoint);
-		}
-
-		/**
-		 * @brief Copy the ManifoldPoint data needed for static friction next tick
-		*/
-		inline void Save(const FManifoldPoint& ManifoldPoint)
-		{
-			ShapeContactPoints[0] = ManifoldPoint.ContactPoint.ShapeContactPoints[0];
-			ShapeContactPoints[1] = ManifoldPoint.ContactPoint.ShapeContactPoints[1];
-			StaticFrictionMax = ManifoldPoint.StaticFrictionMax;
-		}
-
-		/**
-		 * @brief Replace the ManifoldPoint's data with the saved data from the previous tick
-		*/
-		inline void Restore(FManifoldPoint& ManifoldPoint) const
-		{
-			ManifoldPoint.ContactPoint.ShapeContactPoints[0] = ShapeContactPoints[0];
-			ManifoldPoint.ContactPoint.ShapeContactPoints[1] = ShapeContactPoints[1];
-			ManifoldPoint.StaticFrictionMax = StaticFrictionMax;
-		}
-
-		/**
-		 * @brief Is the data for the specified manifold point?
-		*/
-		inline bool IsMatch(const FManifoldPoint& ManifoldPoint, const FReal DistanceToleranceSq) const
-		{
-			// If the contact point is in the same spot on one of the bodies, assume it is the same contact
-			// @todo(chaos): more robust same-point test. E.g., this won't work for very small or very large objects,
-			// so at least make the tolerance size-dependent
-			const FVec3 DP0 = ManifoldPoint.ContactPoint.ShapeContactPoints[0] - ShapeContactPoints[0];
-			const FVec3 DP1 = ManifoldPoint.ContactPoint.ShapeContactPoints[1] - ShapeContactPoints[1];
-			return ((DP0.SizeSquared() < DistanceToleranceSq) || (DP1.SizeSquared() < DistanceToleranceSq));
-		}
-
 		FVec3 ShapeContactPoints[2];
-		FVec3 NetPushOut;
-		FReal StaticFrictionMax;
 	};
 
 	/*
@@ -340,6 +304,12 @@ namespace Chaos
 		const FReal GetCollisionMargin1() const { return CollisionMargins[1]; }
 		const FReal GetCollisionTolerance() const { return CollisionTolerance; }
 
+		const bool IsQuadratic0() const { return Flags.bIsQuadratic0; }
+		const bool IsQuadratic1() const { return Flags.bIsQuadratic1; }
+		const bool HasQuadraticShape() const { return (Flags.bIsQuadratic0 || Flags.bIsQuadratic1); }
+		const FReal GetCollisionRadius0() const { return (Flags.bIsQuadratic0) ? CollisionMargins[0] : FReal(0); }
+		const FReal GetCollisionRadius1() const { return (Flags.bIsQuadratic1) ? CollisionMargins[1] : FReal(0); }
+
 		// @todo(chaos): half of this API is wrong for the new multi-point manifold constraints. Remove it
 
 		const FCollisionContact& GetManifold() const { return Manifold; }
@@ -452,7 +422,7 @@ namespace Chaos
 		 * @brief Restore the contact manifold (assumes relative motion of the two bodies is small)
 		 * @see IsWithinManifoldRestorationThreshold
 		*/
-		void RestoreManifold();
+		void RestoreManifold(const bool bReproject);
 
 		/**
 		 * Determine the constraint direction based on Normal and Phi.
@@ -460,9 +430,6 @@ namespace Chaos
 		 */
 		ECollisionConstraintDirection GetConstraintDirection(const FReal Dt) const;
 
-
-		const FManifoldPointSavedData* FindManifoldPointSavedData(const FManifoldPoint& ManifoldPoint) const;
-
 		/**
 		 * @brief Called before SetSolverResults() to reset accumulators
 		*/
@@ -473,24 +440,52 @@ namespace Chaos
 		}
 
 		/**
-		 * @brief Store the data (from the solver) that is retained between ticks for the specified manifold point
+		 * @brief Store the data from the solver that is retained between ticks for the specified manifold point or used by dependent systems (plasticity, breaking, etc.)
 		*/
-		void SetSolverResults(const int32 ManifoldPointIndex, const FVec3& NetPushOut, const FVec3& NetImpulse, const bool bInsideStaticFrictionCone, const FReal StaticFrictionMax, const FReal Dt)
+		void SetSolverResults(
+			const int32 ManifoldPointIndex, 
+			const FVec3& NetPushOut, 
+			const FVec3& NetImpulse, 
+			const FReal StaticFrictionRatio,
+			const FReal Dt)
 		{
 			FManifoldPoint& ManifoldPoint = ManifoldPoints[ManifoldPointIndex];
 
 			ManifoldPoint.NetPushOut = NetPushOut;
 			ManifoldPoint.NetImpulse = NetImpulse;
-			ManifoldPoint.bInsideStaticFrictionCone = bInsideStaticFrictionCone;
-			ManifoldPoint.StaticFrictionMax = StaticFrictionMax;
+			ManifoldPoint.Flags.bInsideStaticFrictionCone = FMath::IsNearlyEqual(StaticFrictionRatio, FReal(1));
 
 			AccumulatedImpulse += NetImpulse + (NetPushOut / Dt);
 
-			// If are still satisfying the static friction condition, save off the state required for next tick
-			if (bInsideStaticFrictionCone && !NetPushOut.IsNearlyZero() && !SavedManifoldPoints.IsFull())
+			// Save contact data for friction
+			// NOTE: we do this even for points that did not apply PushOut or Impulse so that
+			// we get previous contact data for initial contacts (sometimes). Otherwise we
+			// end up having to estimate the previous contact from velocities
+			if (!SavedManifoldPoints.IsFull())
 			{
 				const int32 SavedIndex = SavedManifoldPoints.Add();
-				SavedManifoldPoints[SavedIndex].Save(ManifoldPoint);
+				FSavedManifoldPoint& SavedManifoldPoint = SavedManifoldPoints[SavedIndex];
+
+				if (FMath::IsNearlyEqual(StaticFrictionRatio, FReal(1)))
+				{
+					// Static friction held - we keep the same contacts points as-is for use next frame
+					SavedManifoldPoint.ShapeContactPoints[0] = ManifoldPoint.ShapeAnchorPoints[0];
+					SavedManifoldPoint.ShapeContactPoints[1] = ManifoldPoint.ShapeAnchorPoints[1];
+				}
+				else if (FMath::IsNearlyEqual(StaticFrictionRatio, FReal(0)))
+				{
+					// No friction - discard the friction anchors
+					SavedManifoldPoint.ShapeContactPoints[0] = ManifoldPoint.ContactPoint.ShapeContactPoints[0];
+					SavedManifoldPoint.ShapeContactPoints[1] = ManifoldPoint.ContactPoint.ShapeContactPoints[1];
+				}
+				else
+				{
+					// We exceeded the friction cone. Slide the friction anchor toward the last-detected contact position
+					// so that it sits at the edge of the friction cone.
+					const FReal Alpha = FMath::Clamp(StaticFrictionRatio, FReal(0), FReal(1));
+					SavedManifoldPoint.ShapeContactPoints[0] = FVec3::Lerp(ManifoldPoint.ContactPoint.ShapeContactPoints[0], ManifoldPoint.ShapeAnchorPoints[0], Alpha);
+					SavedManifoldPoint.ShapeContactPoints[1] = FVec3::Lerp(ManifoldPoint.ContactPoint.ShapeContactPoints[1], ManifoldPoint.ShapeAnchorPoints[1], Alpha);
+				}
 			}
 		}
 
@@ -536,7 +531,13 @@ namespace Chaos
 		*/
 		void TryRestoreFrictionData(const int32 ManifoldPointIndex);
 
-		void InitMargins(const EImplicitObjectType ImplicitType0, const EImplicitObjectType ImplicitType1, const FReal Margin0, const FReal Margin1);
+		void ReprojectManifoldContacts();
+		void ReprojectManifoldPoint(const int32 ManifoldPointIndex);
+
+		void InitMarginsAndTolerances(const EImplicitObjectType ImplicitType0, const EImplicitObjectType ImplicitType1, const FReal Margin0, const FReal Margin1);
+
+		const FSavedManifoldPoint* FindSavedManifoldPoint(const FManifoldPoint& ManifoldPoint) const;
+		FReal CalculateSavedManifoldPointScore(const FSavedManifoldPoint& SavedManifoldPoint, const FManifoldPoint& ManifoldPoint, const FReal DistanceToleranceSq) const;
 
 	public:
 		//@todo(chaos): make this stuff private
@@ -556,17 +557,19 @@ namespace Chaos
 		TCArray<FManifoldPoint, MaxManifoldPoints> ManifoldPoints;
 
 		// The manifold points from the previous tick when we don't reuse the manifold. Used by static friction.
-		TCArray<FManifoldPointSavedData, MaxManifoldPoints> SavedManifoldPoints;
+		TCArray<FSavedManifoldPoint, MaxManifoldPoints> SavedManifoldPoints;
 		FReal CullDistance;
 
 		// The margins to use during collision detection. We don't always use the margins on the shapes directly.
-		// E.g., we use the smallest non-zero margin for 2 convex shapes. See InitMargins
+		// E.g., we use the smallest non-zero margin for 2 convex shapes. See InitMarginsAndTolerances
 		FReal CollisionMargins[2];
 
 		// The collision tolerance is used to determine whether a new contact matches an old on. It is derived from the
 		// margins of the two shapes, as well as their types
 		FReal CollisionTolerance;
 
+		FReal FrictionPositionTolerance;
+
 		union FFlags
 		{
 			FFlags() : Bits(0) {}
@@ -575,6 +578,8 @@ namespace Chaos
 				bool bUseManifold;
 				bool bUseIncrementalManifold;
 				bool bWasManifoldRestored;
+				bool bIsQuadratic0;
+				bool bIsQuadratic1;
 			};
 			uint32 Bits;
 		} Flags;
diff --git a/Engine/Source/Runtime/Experimental/Chaos/Public/Chaos/Collision/ParticlePairMidPhase.h b/Engine/Source/Runtime/Experimental/Chaos/Public/Chaos/Collision/ParticlePairMidPhase.h
--- a/Engine/Source/Runtime/Experimental/Chaos/Public/Chaos/Collision/ParticlePairMidPhase.h
+++ b/Engine/Source/Runtime/Experimental/Chaos/Public/Chaos/Collision/ParticlePairMidPhase.h
@@ -21,6 +21,21 @@ namespace Chaos
 	class FPerShapeData;
 	class FSingleShapePairCollisionDetector;
 
+	/**
+	 * @brief The result of the test to see whether we can reuse a manifold (when the bodies have not moved much)
+	*/
+	enum class ECollisionRestoreType
+	{
+		// Do not retore the manifold - bodies have moved beyond the threshold
+		None,
+
+		// Restore the manifold exactly as-is - the bodies have barely moved
+		Reuse,
+
+		// Restore the manifold, but the bodies have moved a bit so reproject the contacts
+		Reproject,
+	};
+
 	/**
 	 * @brief Handles collision detection for a pair of simple shapes (i.e., not compound shapes)
 	 * 
@@ -71,7 +86,7 @@ namespace Chaos
 		 * @brief Reactivate the collision exactly as it was last frame
 		 * @return The number of collisions constraints that were restored
 		*/
-		int32 RestoreCollision(const FReal CullDistance);
+		int32 RestoreCollision(const FReal CullDistance, const bool bReproject);
 
 		/**
 		 * @brief Reactivate the constraint (essentially the same as Restore but slightly optimized)
@@ -199,7 +214,7 @@ namespace Chaos
 		 * @brief Reactivate the collision exactly as it was last frame
 		 * @return The number of collisions constraints that were restored
 		*/
-		int32 RestoreCollisions(const FReal CullDistance);
+		int32 RestoreCollisions(const FReal CullDistance, const bool bReproject);
 
 		/**
 		 * @brief Reactivate the constraint (essentially the same as Restore but slightly optimized)
@@ -417,7 +432,7 @@ namespace Chaos
 		 * This will be true if neither particle has moved much. 
 		 * This is non-const because it updates some position tracking data.
 		*/
-		bool ShouldRestoreConstraints(const FReal Dt);
+		ECollisionRestoreType ShouldRestoreConstraints(const FReal Dt);
 
 		/**
 		 * @brief If the particles have not moved muc, reactivate all the colisions and skip the narrow phase.
diff --git a/Engine/Source/Runtime/Experimental/Chaos/Public/Chaos/Rotation.h b/Engine/Source/Runtime/Experimental/Chaos/Public/Chaos/Rotation.h
--- a/Engine/Source/Runtime/Experimental/Chaos/Public/Chaos/Rotation.h
+++ b/Engine/Source/Runtime/Experimental/Chaos/Public/Chaos/Rotation.h
@@ -114,6 +114,14 @@ namespace Chaos
 			OutSwing = *this * OutTwist.Inverse();
 		}
 
+		/**
+		 * @brief Return the large absolute element value
+		*/
+		FORCEINLINE FReal GetAbsMax() const
+		{
+			return FMath::Max(FMath::Max(FMath::Abs(X), FMath::Abs(Y)), FMath::Max(FMath::Abs(Z), FMath::Abs(W)));
+		}
+
 		/**
 		 * @brief Return the dot product of two quaternions
 		*/
