IslandGroups chaos optim

diff --git a/Engine/Source/Programs/HeadlessChaos/Private/HeadlessChaosTestGraph.cpp b/Engine/Source/Programs/HeadlessChaos/Private/HeadlessChaosTestGraph.cpp
--- a/Engine/Source/Programs/HeadlessChaos/Private/HeadlessChaosTestGraph.cpp
+++ b/Engine/Source/Programs/HeadlessChaos/Private/HeadlessChaosTestGraph.cpp
@@ -60,6 +60,7 @@ namespace ChaosTest {
 		TMockGraphConstraints()
 			: FPBDIndexedConstraintContainer(FConstraintContainerHandle::StaticType())
 		{
+			SetContainerId(T_TYPEID);
 		}
 
 		int32 NumConstraints() const { return Constraints.Num(); }
@@ -213,7 +214,7 @@ namespace ChaosTest {
 		}
 		
 		// Generate the constraint/particle islands
-		Graph.UpdateIslands(SOAs.GetNonDisabledDynamicView(), SOAs);
+		Graph.UpdateIslands(SOAs.GetNonDisabledDynamicView(), SOAs,2);
 		
 		// Islands should be end up with the following particles (note: particle 17 is infinite mass and can appear in multiple islands)
 		TArray<TSet<TGeometryParticleHandle<FReal, 3>*>> ExpectedIslandParticles = {
@@ -519,7 +520,7 @@ namespace ChaosTest {
 			}
 
 			// Generate the constraint/particle islands
-			Graph.UpdateIslands(SOAs.GetNonDisabledDynamicView(), SOAs);
+			Graph.UpdateIslands(SOAs.GetNonDisabledDynamicView(), SOAs, 1);
 
 			// Check the number of islands
 			EXPECT_EQ(IterData.ExpectedIslandParticleIndices.Num(), Graph.NumIslands());
@@ -611,7 +612,7 @@ namespace ChaosTest {
 			Graph.AddConstraint(0,Constraints.Handles[ConstraintIndex],TVec2<TGeometryParticleHandle<FReal,3>*>(Particles[Indices[0]],Particles[Indices[1]]));
 		}
 
-		Graph.UpdateIslands(SOAs.GetNonDisabledDynamicView(),SOAs);
+		Graph.UpdateIslands(SOAs.GetNonDisabledDynamicView(),SOAs, 1);
 		for (int32 IslandIndex = 0; IslandIndex < Graph.NumIslands(); ++IslandIndex)
 		{
 			const bool bSleeped = Graph.SleepInactive(IslandIndex,PhysicsMaterials,PhysicalMaterials);
@@ -1052,7 +1053,7 @@ namespace ChaosTest {
 			const TVec2<int32>& Indices = Constraints.ConstraintParticleIndices(ConstraintIndex);
 			Graph.AddConstraint(ContainerId, Constraints.Handles[ConstraintIndex], TVec2<TGeometryParticleHandle<FReal, 3>*>(AllParticles[Indices[0]], AllParticles[Indices[1]]) );
 		}
-		Graph.UpdateIslands(SOAs.GetNonDisabledDynamicView(), SOAs);
+		Graph.UpdateIslands(SOAs.GetNonDisabledDynamicView(), SOAs, 1);
 
 		// It's a connected grid, so only one island
 		EXPECT_EQ(Graph.NumIslands(), 1);
diff --git a/Engine/Source/Runtime/Experimental/Chaos/Private/Chaos/Island/IslandGroup.cpp b/Engine/Source/Runtime/Experimental/Chaos/Private/Chaos/Island/IslandGroup.cpp
new file mode 100644
--- /dev/null
+++ b/Engine/Source/Runtime/Experimental/Chaos/Private/Chaos/Island/IslandGroup.cpp
@@ -0,0 +1,38 @@
+// Copyright Epic Games, Inc. All Rights Reserved.
+#include "Chaos/Island/IslandGroup.h"
+#include "Chaos/Island/SolverIsland.h"
+
+namespace Chaos
+{
+	
+FPBDIslandGroup::FPBDIslandGroup(const int32 GroupIndexIn ) : FPBDIslandSolverData(GroupIndexIn),
+	 IslandSolvers(), ParticlesCount(0), ConstraintsCount(0)
+{}
+
+void FPBDIslandGroup::ReserveIslands(const int32 NumSolvers)
+{
+	IslandSolvers.Reserve(NumSolvers);
+}
+
+void FPBDIslandGroup::AddIsland(FPBDIslandSolver* IslandSolver)
+{
+	if (IslandSolver)
+	{
+		IslandSolvers.Add(IslandSolver);
+	}
+}
+	
+void FPBDIslandGroup::ClearIslands()
+{
+	IslandSolvers.Reset();
+}
+
+void FPBDIslandGroup::InitGroup()
+{
+	IslandSolvers.SetNum(0,false);
+	ParticlesCount = 0;
+	ConstraintsCount = 0;
+}
+	
+
+}
\ No newline at end of file
diff --git a/Engine/Source/Runtime/Experimental/Chaos/Private/Chaos/Island/IslandManager.cpp b/Engine/Source/Runtime/Experimental/Chaos/Private/Chaos/Island/IslandManager.cpp
--- a/Engine/Source/Runtime/Experimental/Chaos/Private/Chaos/Island/IslandManager.cpp
+++ b/Engine/Source/Runtime/Experimental/Chaos/Private/Chaos/Island/IslandManager.cpp
@@ -9,10 +9,16 @@
 #include "Chaos/PBDRigidsSOAs.h"
 #include "ChaosLog.h"
 
+#include "Framework/Threading.h"
+
 /** Cvar to enable/disable the island sleeping */
 bool ChaosSolverSleepEnabled = true;
 FAutoConsoleVariableRef CVarChaosSolverSleepEnabled(TEXT("p.Chaos.Solver.SleepEnabled"), ChaosSolverSleepEnabled, TEXT(""));
 
+/** Cvar to control the number of island groups used in the solver. The total number will be NumThreads * IslandGroupsMultiplier */
+int32 GChaosSolverIslandGroupsMultiplier = 1;
+FAutoConsoleVariableRef CVarSolverIslandGroupsMultiplier(TEXT("p.Chaos.Solver.IslandGroupsMultiplier"), GChaosSolverIslandGroupsMultiplier, TEXT("Total number of island groups in the solver will be NumThreads * IslandGroupsMultiplier.[def:1]"));
+
 /** Cvar to override the sleep counter threshold if necessary */
 int32 ChaosSolverCollisionDefaultSleepCounterThresholdCVar = 20;
 FAutoConsoleVariableRef CVarChaosSolverCollisionDefaultSleepCounterThreshold(TEXT("p.ChaosSolverCollisionDefaultSleepCounterThreshold"), ChaosSolverCollisionDefaultSleepCounterThresholdCVar, TEXT("Default counter threshold for sleeping.[def:20]"));
@@ -260,11 +266,14 @@ inline bool ComputeSleepingThresholds(FPBDIslandSolver* IslandSolver,
 }
 
 FPBDIslandManager::FPBDIslandManager() : IslandSolvers(), IslandGraph(MakeUnique<GraphType>()), MaxParticleIndex(INDEX_NONE)
-{}
+{
+	InitializeGroups();
+}
 
 FPBDIslandManager::FPBDIslandManager(const TParticleView<FPBDRigidParticles>& PBDRigids) : IslandSolvers(), IslandGraph(MakeUnique<GraphType>())
 {
 	InitializeGraph(PBDRigids);
+	InitializeGroups();
 }
 
 FPBDIslandManager::~FPBDIslandManager()
@@ -283,6 +292,17 @@ FPBDIslandManager::~FPBDIslandManager()
 	}
 }
 
+void FPBDIslandManager::InitializeGroups()
+{
+	const int32 NumThreads = FTaskGraphInterface::Get().GetNumWorkerThreads();
+	IslandGroups.SetNum(NumThreads * GChaosSolverIslandGroupsMultiplier, false);
+	
+	for(int32 GroupIndex = 0, NumGroups = IslandGroups.Num(); GroupIndex < NumGroups; ++GroupIndex)
+	{
+		IslandGroups[GroupIndex] = MakeUnique<FPBDIslandGroup>(GroupIndex);
+	}
+}
+
 void FPBDIslandManager::InitializeGraph(const TParticleView<FPBDRigidParticles>& PBDRigids)
 {
 	MaxParticleIndex = 0;
@@ -511,11 +531,17 @@ void FPBDIslandManager::ResetIslands(const TParticleView<FPBDRigidParticles>& PB
 {
 	// @todo 
 }
+
+inline bool SolverIslandSortPredicate(const TUniquePtr<FPBDIslandSolver>& SolverIslandL, const TUniquePtr<FPBDIslandSolver>& SolverIslandR)
+{
+	return SolverIslandL->NumConstraints() < SolverIslandR->NumConstraints();
+}
 	
-void FPBDIslandManager::SyncIslands(FPBDRigidsSOAs& Particles)
+void FPBDIslandManager::SyncIslands(FPBDRigidsSOAs& Particles, const int32 NumContainers)
 {
 	IslandSolvers.Reserve(IslandGraph->NumIslands());
 	IslandIndexing.SetNum(IslandGraph->NumIslands(),false);
+	SortedIslands.SetNum(IslandGraph->NumIslands(),false);
 	int32 LocalIsland = 0;
 
 	// Sync of the solver islands first and reserve the required space
@@ -531,14 +557,16 @@ void FPBDIslandManager::SyncIslands(FPBDRigidsSOAs& Particles)
 			}
 			FPBDIslandSolver*& IslandSolver = IslandGraph->GraphIslands[IslandIndex].IslandItem;
 			IslandSolver = IslandSolvers[IslandIndex].Get();
-		
+			IslandSolver->ResizeConstraintsCounts(NumContainers);
+			
 			// We then transfer the persistent flag and the graph dense index to the solver island
 			IslandSolver->SetIsPersistent(IslandGraph->GraphIslands[IslandIndex].bIsPersistent);
 			IslandSolver->SetIsSleeping(IslandGraph->GraphIslands[IslandIndex].bIsSleeping);
 			IslandSolver->GetIslandIndex() = LocalIsland;
 
 			// We update the IslandIndexing to retrieve the graph sparse and persistent index from the dense one.
-			IslandIndexing[LocalIsland++] = IslandIndex;
+			SortedIslands[LocalIsland] = IslandIndex;
+			IslandIndexing[LocalIsland++] = SortedIslands[LocalIsland];
 
 			// We finally update the solver islands based on the new particles and constraints if
 			// the island is not persistent or not sleeping. 
@@ -565,6 +593,7 @@ void FPBDIslandManager::SyncIslands(FPBDRigidsSOAs& Particles)
 	}
 	PopulateIslands(IslandGraph.Get());
 
+	
 	// Update of the sync and sleep state for each islands
 	for(auto& IslandSolver : IslandSolvers)
 	{
@@ -576,15 +605,65 @@ void FPBDIslandManager::SyncIslands(FPBDRigidsSOAs& Particles)
 	}
 	
 	IslandIndexing.SetNum(LocalIsland,false);
+	SortedIslands.SetNum(LocalIsland,false);
+
+	// Build all the island groups
+	BuildGroups(NumContainers);
+}
+
+void FPBDIslandManager::BuildGroups(const int32 NumContainers)
+{
+	// No Need to sort for now since it is adding a cost for not improving much the end result
+	TSparseArray<TUniquePtr<FPBDIslandSolver>>& LocalIslands = IslandSolvers;
+	SortedIslands.Sort([&LocalIslands](const int32 IslandIndexA, const int32 IslandIndexB) -> bool
+		{ return LocalIslands[IslandIndexA]->NumConstraints() > LocalIslands[IslandIndexB]->NumConstraints();});
+	
+	const int32 NumGroups = IslandGroups.Num();
+	const int32 GroupSize = IslandGraph->NumEdges() / NumGroups + 1;
+
+	for(TUniquePtr<FPBDIslandGroup>& IslandGroup : IslandGroups)
+	{
+		IslandGroup->InitGroup();
+		IslandGroup->ResizeConstraintsCounts(NumContainers);
+	}
+
+	int32 GroupIndex = 0;
+	int32 GroupOffset = 0;
+	for(int32& SortedIndex : SortedIslands)
+	{
+		if( FPBDIslandSolver* IslandSolver = IslandSolvers[SortedIndex].Get())
+		{
+			check(GroupIndex < NumGroups);
+			IslandGroups[GroupIndex]->AddIsland(IslandSolver);
+			IslandGroups[GroupIndex]->NumParticles() += IslandSolver->NumParticles();
+			IslandGroups[GroupIndex]->NumConstraints() += IslandSolver->NumConstraints();
+
+			check(IslandSolver->NumContainerIds() == IslandGroups[GroupIndex]->NumContainerIds());
+			for(int32 ContainerIndex = 0; ContainerIndex < IslandSolver->NumContainerIds(); ++ContainerIndex)
+			{
+				IslandGroups[GroupIndex]->ConstraintCount(ContainerIndex) +=
+					IslandSolver->ConstraintCount(ContainerIndex);
+			}
+			
+			IslandSolver->SetGroupIndex(GroupIndex);
+			GroupOffset += IslandSolver->NumConstraints();
+
+			if(GroupOffset > GroupSize)
+			{
+				GroupIndex++;
+				GroupOffset = 0;
+			}
+		}
+	}
 }
 	
-void FPBDIslandManager::UpdateIslands(const TParticleView<FPBDRigidParticles>& PBDRigids, FPBDRigidsSOAs& Particles)
+void FPBDIslandManager::UpdateIslands(const TParticleView<FPBDRigidParticles>& PBDRigids, FPBDRigidsSOAs& Particles, const int32 NumContainers)
 {
 	// Merge the graph islands if required
 	IslandGraph->UpdateGraph();
 	
 	// Sync the graph islands with the solver islands objects
-	SyncIslands(Particles);
+	SyncIslands(Particles, NumContainers);
 }
 
 bool FPBDIslandManager::SleepInactive(const int32 IslandIndex,
diff --git a/Engine/Source/Runtime/Experimental/Chaos/Private/Chaos/Island/SolverIsland.cpp b/Engine/Source/Runtime/Experimental/Chaos/Private/Chaos/Island/SolverIsland.cpp
--- a/Engine/Source/Runtime/Experimental/Chaos/Private/Chaos/Island/SolverIsland.cpp
+++ b/Engine/Source/Runtime/Experimental/Chaos/Private/Chaos/Island/SolverIsland.cpp
@@ -61,6 +61,9 @@ void FPBDIslandSolver::AddConstraint(FConstraintHandle* ConstraintHandle)
 	if (ConstraintHandle)
 	{
 		IslandConstraints.Add(ConstraintHandle);
+
+		const int32 ContainerId = ConstraintHandle->GetContainerId();
+		ConstraintCounts[ConstraintHandle->GetContainerId()]++;
 	}
 }
 
diff --git a/Engine/Source/Runtime/Experimental/Chaos/Private/Chaos/PBDConstraintRule.cpp b/Engine/Source/Runtime/Experimental/Chaos/Private/Chaos/PBDConstraintRule.cpp
--- a/Engine/Source/Runtime/Experimental/Chaos/Private/Chaos/PBDConstraintRule.cpp
+++ b/Engine/Source/Runtime/Experimental/Chaos/Private/Chaos/PBDConstraintRule.cpp
@@ -146,30 +146,37 @@ namespace Chaos
 	}
 
 	template<class ConstraintType>
-	void TPBDConstraintIslandRule<ConstraintType>::GatherSolverInput(const FReal Dt, int32 Island)
+	void TPBDConstraintIslandRule<ConstraintType>::GatherSolverInput(const FReal Dt, int32 GroupIndex)
 	{
-		if(FPBDIslandSolver* IslandSolver = ConstraintGraph->GetSolverIsland(Island))
+		if(FPBDIslandGroup* IslandGroup = ConstraintGraph->GetIslandGroup(GroupIndex))
 		{
-			const TArray<FConstraintHandleHolder>& IslandConstraints = ConstraintGraph->GetIslandConstraints(Island);
-
 			// This will reset the number of constraints inside the solver datas. For now we keep this function since according
 			// to the constraints we can use the handles, the indices or the container. when only the container will be used
 			// we can replace this call by :
 			// IslandSolver->template GetConstraintContainer<typename ConstraintType::FSolverConstraintContainerType>(ContainerId)->Reset();
-			Constraints.SetNumIslandConstraints(IslandConstraints.Num(), *IslandSolver);
-
-			for (FConstraintHandle* ConstraintHandle : IslandConstraints)
+			Constraints.SetNumIslandConstraints(IslandGroup->ConstraintCount(Constraints.GetContainerId()), *IslandGroup);
+			IslandGroup->InitConstraintIndex(Constraints.GetContainerId());
+			
+			for (auto& IslandSolver : IslandGroup->GetIslands() )
 			{
-				if (ConstraintHandle->GetContainerId() == GetContainerId())
+				if(!IslandSolver->IsSleeping() && !IslandSolver->IsUsingCache())
 				{
-					FConstraintContainerHandle* Constraint = ConstraintHandle->As<FConstraintContainerHandle>();
-
-					// Note we are building the SolverBodies as we go, in the order that we visit them. Each constraint
-					// references two bodies, so we won't strictly be accessing only in cache order, but it's about as good as it can be.
-					if (Constraint->IsEnabled())
+					const TArray<FConstraintHandleHolder>& IslandConstraints = ConstraintGraph->GetIslandConstraints(IslandSolver->GetIslandIndex());
+					
+					for (FConstraintHandle* ConstraintHandle : IslandConstraints)
 					{
-						// @todo(chaos): we should provide Particle Levels in the island rule as well (see TPBDConstraintColorRule)
-						Constraint->GatherInput(Dt, INDEX_NONE, INDEX_NONE, *IslandSolver);
+						if (ConstraintHandle->GetContainerId() == GetContainerId())
+						{
+							FConstraintContainerHandle* Constraint = ConstraintHandle->As<FConstraintContainerHandle>();
+
+							// Note we are building the SolverBodies as we go, in the order that we visit them. Each constraint
+							// references two bodies, so we won't strictly be accessing only in cache order, but it's about as good as it can be.
+							if (Constraint->IsEnabled())
+							{
+								// @todo(chaos): we should provide Particle Levels in the island rule as well (see TPBDConstraintColorRule)
+								Constraint->GatherInput(Dt, INDEX_NONE, INDEX_NONE, *IslandGroup);
+							}
+						}
 					}
 				}
 			}
@@ -177,26 +184,26 @@ namespace Chaos
 	}
 
 	template<class ConstraintType>
-	void TPBDConstraintIslandRule<ConstraintType>::ScatterSolverOutput(const FReal Dt, int32 Island)
+	void TPBDConstraintIslandRule<ConstraintType>::ScatterSolverOutput(const FReal Dt, int32 GroupIndex)
 	{
-		if(FPBDIslandSolver* IslandSolver = ConstraintGraph->GetSolverIsland(Island))
+		if(FPBDIslandGroup* IslandGroup = ConstraintGraph->GetIslandGroup(GroupIndex))
 		{
-			Constraints.ScatterOutput(Dt, *IslandSolver);
+			Constraints.ScatterOutput(Dt, *IslandGroup);
 		}
 	}
 	
 	template<class ConstraintType>
-	bool TPBDConstraintIslandRule<ConstraintType>::ApplyConstraints(const FReal Dt, int32 Island, const int32 It, const int32 NumIts)
+	bool TPBDConstraintIslandRule<ConstraintType>::ApplyConstraints(const FReal Dt, int32 GroupIndex, const int32 It, const int32 NumIts)
 	{
-		FPBDIslandSolver* IslandSolver = ConstraintGraph->GetSolverIsland(Island);
-		return IslandSolver ? Constraints.ApplyPhase1Serial(Dt, It, NumIts, *IslandSolver) : false;
+		FPBDIslandGroup* IslandGroup = ConstraintGraph->GetIslandGroup(GroupIndex);
+		return IslandGroup ? Constraints.ApplyPhase1Serial(Dt, It, NumIts, *IslandGroup) : false;
 	}
 
 	template<class ConstraintType>
-	bool TPBDConstraintIslandRule<ConstraintType>::ApplyPushOut(const FReal Dt, int32 Island, const int32 It, const int32 NumIts)
+	bool TPBDConstraintIslandRule<ConstraintType>::ApplyPushOut(const FReal Dt, int32 GroupIndex, const int32 It, const int32 NumIts)
 	{
-		FPBDIslandSolver* IslandSolver = ConstraintGraph->GetSolverIsland(Island);
-		return IslandSolver  ? Constraints.ApplyPhase2Serial(Dt, It, NumIts, *IslandSolver) : false;
+		FPBDIslandGroup* IslandGroup = ConstraintGraph->GetIslandGroup(GroupIndex);
+		return IslandGroup  ? Constraints.ApplyPhase2Serial(Dt, It, NumIts, *IslandGroup) : false;
 	}
 
 	template<class ConstraintType>
@@ -206,7 +213,7 @@ namespace Chaos
 	}
 
 	template<class ConstraintType>
-	void TPBDConstraintIslandRule<ConstraintType>::UpdateAccelerationStructures(const FReal Dt, const int32 Island)
+	void TPBDConstraintIslandRule<ConstraintType>::UpdateAccelerationStructures(const FReal Dt, const int32 GroupIndex)
 	{
 	}
 	
@@ -229,7 +236,7 @@ namespace Chaos
 	}
 
 	template<class ConstraintType>
-	void TPBDConstraintColorRule<ConstraintType>::GatherSolverInput(const FReal Dt, int32 Island)
+	void TPBDConstraintColorRule<ConstraintType>::GatherSolverInput(const FReal Dt, int32 GroupIndex)
 	{
 		FPBDConstraintGraph* LocalGraph = ConstraintGraph;
 		auto GetParticleLevel = [LocalGraph](TGeometryParticleHandle<FReal, 3>* ConstrainedParticle) -> int32
@@ -246,64 +253,79 @@ namespace Chaos
 		};
 		if (IsSortingEnabled())
 		{
-			if (FPBDIslandSolver* IslandSolver = ConstraintGraph->GetSolverIsland(Island))
+			if (FPBDIslandGroup* IslandGroup = ConstraintGraph->GetIslandGroup(GroupIndex))
 			{
-				const int32 GraphIndex = ConstraintGraph->GetGraphIndex(Island);
-				if(ConstraintGraph->GetIslandGraph()->GraphIslands.IsValidIndex(GraphIndex))
+				bool bResizeDone = false;
+				TArray<TPair<int32, int32>>& IslandConstraintSets = ConstraintSets[GroupIndex];
+				IslandConstraintSets.Reset();
+				int32 ConstraintSetBegin = 0, ConstraintSetEnd = 0;
+				
+				for (auto& IslandSolver : IslandGroup->GetIslands() )
 				{
-					const int32 MaxColor = IsSortingUsingColors() ? FMath::Max(1,ConstraintGraph->GetIslandGraph()->GraphIslands[GraphIndex].MaxColors+1) : 1;
-					const int32 MaxLevel = IsSortingUsingLevels() ? FMath::Max(1,ConstraintGraph->GetIslandGraph()->GraphIslands[GraphIndex].MaxLevels+1) : 1;
-					
-					TArray<TPair<int32, int32>>& IslandConstraintSets = ConstraintSets[Island];
-					IslandConstraintSets.Reset(MaxLevel * MaxColor);	// Pessimistic array size - we could store the actual required size in coloring algorithm
-					int32 ConstraintSetEnd = 0;
-					
-					Constraints.SetNumIslandConstraints(ConstraintGraph->GetIslandConstraints(Island).Num(), *IslandSolver);
-					IslandSolver->InitConstraintIndex(Constraints.GetContainerId());
-
-					for (int32 Level = 0; Level < MaxLevel; ++Level)
+					ConstraintSetBegin = ConstraintSetEnd;
+					if(!IslandSolver->IsSleeping() && !IslandSolver->IsUsingCache())
 					{
-						for (int32 Color = 0; Color < MaxColor; ++Color)
+						if(!bResizeDone)
 						{
-							const int32 OffsetIndex = IslandOffsets[Island] + Level * MaxColor + Color;
-							const int32 OffsetBegin = (OffsetIndex == 0) ? 0 : ConstraintOffsets[OffsetIndex-1];
-							const int32 OffsetEnd = ConstraintOffsets[OffsetIndex];
-							
-							 if (OffsetEnd != OffsetBegin)
-							 {
-							 	// Calculate the range of indices for this color as a set of independent contacts
-							 	TPair<int32, int32> ColorConstrainSet(ConstraintSetEnd, ConstraintSetEnd);
-							 	for(int32 ConstraintIndex = OffsetBegin; ConstraintIndex < OffsetEnd; ++ConstraintIndex)
-							 	{
-							 		FConstraintContainerHandle* Constraint = SortedConstraints[ConstraintIndex];
-							 		if (Constraint->IsEnabled())
-							 		{
-							 			// Levels that should be assigned to the bodies for shock propagation
-							 			// @todo(chaos): optimize the lookup
-							 			const TVector<TGeometryParticleHandle<FReal, 3>*, 2> ConstrainedParticles = Constraint->GetConstrainedParticles();
-							 			const int32 Particle0Level = GetParticleLevel(ConstrainedParticles[0]);
-							 			const int32 Particle1Level = GetParticleLevel(ConstrainedParticles[1]);
-							
-							 			// Note we are building the SolverBodies as we go, in the order that we visit them. Each constraint
-							 			// references two bodies, so we won't strictly be accessing only in cache order, but it's about as good as it can be.
-							 			Constraint->GatherInput(Dt, Particle0Level, Particle1Level, *IslandSolver);
+							IslandGroup->InitConstraintIndex(Constraints.GetContainerId());
+							Constraints.SetNumIslandConstraints(IslandGroup->ConstraintCount(Constraints.GetContainerId()), *IslandGroup);
+							bResizeDone = true;
+						}
+						const int32 Island = IslandSolver->GetIslandIndex();
+				
+						const int32 GraphIndex = ConstraintGraph->GetGraphIndex(Island);
+						if(ConstraintGraph->GetIslandGraph()->GraphIslands.IsValidIndex(GraphIndex))
+						{
+							const int32 MaxColor = IsSortingUsingColors() ? FMath::Max(1,ConstraintGraph->GetIslandGraph()->GraphIslands[GraphIndex].MaxColors+1) : 1;
+							const int32 MaxLevel = IsSortingUsingLevels() ? FMath::Max(1,ConstraintGraph->GetIslandGraph()->GraphIslands[GraphIndex].MaxLevels+1) : 1;
 							
-							 			// Update the current constraint set of this color
-							 			ColorConstrainSet.Value = ++ConstraintSetEnd;
-							 		}
-							 	}
-							 	// Remember the set of constraints of this color
-							 	if (IsSortingUsingColors()) 
-							 	{
-							 		IslandConstraintSets.Add(ColorConstrainSet);
-							 	}
-							 }
+							IslandConstraintSets.Reserve(IslandConstraintSets.Num()+MaxLevel * MaxColor);	// Pessimistic array size - we could store the actual required size in coloring algorithm
+
+							for (int32 Level = 0; Level < MaxLevel; ++Level)
+							{
+								for (int32 Color = 0; Color < MaxColor; ++Color)
+								{
+									const int32 OffsetIndex = IslandOffsets[Island] + Level * MaxColor + Color;
+									const int32 OffsetBegin = (OffsetIndex == 0) ? 0 : ConstraintOffsets[OffsetIndex-1];
+									const int32 OffsetEnd = ConstraintOffsets[OffsetIndex];
+									
+									if (OffsetEnd != OffsetBegin)
+									{
+										// Calculate the range of indices for this color as a set of independent contacts
+										TPair<int32, int32> ColorConstrainSet(ConstraintSetEnd, ConstraintSetEnd);
+										for(int32 ConstraintIndex = OffsetBegin; ConstraintIndex < OffsetEnd; ++ConstraintIndex)
+										{
+											FConstraintContainerHandle* Constraint = SortedConstraints[ConstraintIndex];
+											if (Constraint->IsEnabled())
+											{
+												// Levels that should be assigned to the bodies for shock propagation
+												// @todo(chaos): optimize the lookup
+												const TVector<TGeometryParticleHandle<FReal, 3>*, 2> ConstrainedParticles = Constraint->GetConstrainedParticles();
+												const int32 Particle0Level = GetParticleLevel(ConstrainedParticles[0]);
+												const int32 Particle1Level = GetParticleLevel(ConstrainedParticles[1]);
+									
+												// Note we are building the SolverBodies as we go, in the order that we visit them. Each constraint
+												// references two bodies, so we won't strictly be accessing only in cache order, but it's about as good as it can be.
+												Constraint->GatherInput(Dt, Particle0Level, Particle1Level, *IslandGroup);
+									
+												// Update the current constraint set of this color
+												ColorConstrainSet.Value = ++ConstraintSetEnd;
+											}
+										}
+										// Remember the set of constraints of this color
+										if (IsSortingUsingColors()) 
+										{
+											IslandConstraintSets.Add(ColorConstrainSet);
+										}
+									}
+								}
+							}
 						}
 					}
 					// If we aren't coloring, we have a single group of all constraints (they have been created in level order above)
 					if (!IsSortingUsingColors())
 					{
-						const TPair<int32, int32> ConstrainSet(0, ConstraintSetEnd);
+						const TPair<int32, int32> ConstrainSet(ConstraintSetBegin, ConstraintSetEnd);
 						IslandConstraintSets.Add(ConstrainSet);
 					}
 				}
@@ -311,51 +333,51 @@ namespace Chaos
 		}
 		else
 		{
-			Base::GatherSolverInput(Dt, Island);
+			Base::GatherSolverInput(Dt, GroupIndex);
 		}
 	}
 
 	template<class ConstraintType>
-	void TPBDConstraintColorRule<ConstraintType>::ScatterSolverOutput(const FReal Dt, int32 Island)
+	void TPBDConstraintColorRule<ConstraintType>::ScatterSolverOutput(const FReal Dt, int32 GroupIndex)
 	{
 		if (IsSortingEnabled())
 		{
-			if (FPBDIslandSolver* IslandSolver = ConstraintGraph->GetSolverIsland(Island))
+			if (FPBDIslandGroup* IslandGroup = ConstraintGraph->GetIslandGroup(GroupIndex))
 			{
-				const TArray<TPair<int32, int32>>& IslandConstraintSets = ConstraintSets[Island];
+				const TArray<TPair<int32, int32>>& IslandConstraintSets = ConstraintSets[GroupIndex];
 				for (const TPair<int32, int32>& ConstraintSet : IslandConstraintSets)
 				{
-					Constraints.ScatterOutput(Dt, ConstraintSet.Key, ConstraintSet.Value, *IslandSolver);
+					Constraints.ScatterOutput(Dt, ConstraintSet.Key, ConstraintSet.Value, *IslandGroup);
 				}
 			}
 		}
 		else
 		{
-			return Base::ScatterSolverOutput(Dt, Island);
+			return Base::ScatterSolverOutput(Dt, GroupIndex);
 		}
 	}
 
 	template<class ConstraintType>
-	bool TPBDConstraintColorRule<ConstraintType>::ApplyConstraints(const FReal Dt, int32 Island, const int32 It, const int32 NumIts)
+	bool TPBDConstraintColorRule<ConstraintType>::ApplyConstraints(const FReal Dt, int32 GroupIndex, const int32 It, const int32 NumIts)
 	{
 		if (IsSortingEnabled())
 		{
 			bool bNeedsAnotherIteration = false;
-			if (FPBDIslandSolver* IslandSolver = ConstraintGraph->GetSolverIsland(Island))
+			if (FPBDIslandGroup* IslandGroup = ConstraintGraph->GetIslandGroup(GroupIndex))
 			{
-				const TArray<TPair<int32, int32>>& IslandConstraintSets = ConstraintSets[Island];
+				const TArray<TPair<int32, int32>>& IslandConstraintSets = ConstraintSets[GroupIndex];
 				if (!IsSortingUsingColors())
 				{
 					for (const TPair<int32, int32>& ConstraintSet : IslandConstraintSets)
 					{
-						bNeedsAnotherIteration |= Constraints.ApplyPhase1Serial(Dt, It, NumIts, ConstraintSet.Key, ConstraintSet.Value, *IslandSolver);
+						bNeedsAnotherIteration |= Constraints.ApplyPhase1Serial(Dt, It, NumIts, ConstraintSet.Key, ConstraintSet.Value, *IslandGroup);
 					}
 				}
 				else
 				{
 					for (const TPair<int32, int32>& ConstraintSet : IslandConstraintSets)
 					{
-						bNeedsAnotherIteration |= Constraints.ApplyPhase1Parallel(Dt, It, NumIts, ConstraintSet.Key, ConstraintSet.Value, *IslandSolver);
+						bNeedsAnotherIteration |= Constraints.ApplyPhase1Parallel(Dt, It, NumIts, ConstraintSet.Key, ConstraintSet.Value, *IslandGroup);
 					}
 				}
 			}
@@ -363,31 +385,31 @@ namespace Chaos
 		}
 		else
 		{
-			return Base::ApplyConstraints(Dt, Island, It, NumIts);
+			return Base::ApplyConstraints(Dt, GroupIndex, It, NumIts);
 		}
 	}
 
 	template<class ConstraintType>
-	bool TPBDConstraintColorRule<ConstraintType>::ApplyPushOut(const FReal Dt, int32 Island, const int32 It, const int32 NumIts)
+	bool TPBDConstraintColorRule<ConstraintType>::ApplyPushOut(const FReal Dt, int32 GroupIndex, const int32 It, const int32 NumIts)
 	{
 		if (IsSortingEnabled())
 		{
 			bool bNeedsAnotherIteration = false;
-			if (FPBDIslandSolver* IslandSolver = ConstraintGraph->GetSolverIsland(Island))
+			if (FPBDIslandGroup* IslandGroup = ConstraintGraph->GetIslandGroup(GroupIndex))
 			{
-				const TArray<TPair<int32, int32>>& IslandConstraintSets = ConstraintSets[Island];
+				const TArray<TPair<int32, int32>>& IslandConstraintSets = ConstraintSets[GroupIndex];
 				if (!IsSortingUsingColors())
 				{
 					for (const TPair<int32, int32>& ConstraintSet : IslandConstraintSets)
 					{
-						bNeedsAnotherIteration |= Constraints.ApplyPhase2Serial(Dt, It, NumIts, ConstraintSet.Key, ConstraintSet.Value, *IslandSolver);
+						bNeedsAnotherIteration |= Constraints.ApplyPhase2Serial(Dt, It, NumIts, ConstraintSet.Key, ConstraintSet.Value, *IslandGroup);
 					}
 				}
 				else
 				{
 					for (const TPair<int32, int32>& ConstraintSet : IslandConstraintSets)
 					{
-						bNeedsAnotherIteration |= Constraints.ApplyPhase2Parallel(Dt, It, NumIts, ConstraintSet.Key, ConstraintSet.Value, *IslandSolver);
+						bNeedsAnotherIteration |= Constraints.ApplyPhase2Parallel(Dt, It, NumIts, ConstraintSet.Key, ConstraintSet.Value, *IslandGroup);
 					}
 				}
 			}
@@ -395,7 +417,7 @@ namespace Chaos
 		}
 		else
 		{
-			return Base::ApplyPushOut(Dt, Island, It, NumIts);
+			return Base::ApplyPushOut(Dt, GroupIndex, It, NumIts);
 		}
 	}
 
@@ -404,7 +426,7 @@ namespace Chaos
 	{
 		if (IsSortingEnabled())
 		{
-			ConstraintSets.SetNum(ConstraintGraph->NumIslands());
+			ConstraintSets.SetNum(ConstraintGraph->NumGroups());
 			ConstraintGraph->template AddConstraintDatas<ConstraintType>(Constraints.GetContainerId());
 		}
 		else
@@ -414,7 +436,7 @@ namespace Chaos
 	}
 
 	template<class ConstraintType>
-	void TPBDConstraintColorRule<ConstraintType>::UpdateAccelerationStructures(const FReal Dt, const int32 Island)
+	void TPBDConstraintColorRule<ConstraintType>::UpdateAccelerationStructures(const FReal Dt, const int32 GroupIndex)
 	{}
 	
 	template<class ConstraintType>
diff --git a/Engine/Source/Runtime/Experimental/Chaos/Private/Chaos/PBDRigidsEvolutionGBF.cpp b/Engine/Source/Runtime/Experimental/Chaos/Private/Chaos/PBDRigidsEvolutionGBF.cpp
--- a/Engine/Source/Runtime/Experimental/Chaos/Private/Chaos/PBDRigidsEvolutionGBF.cpp
+++ b/Engine/Source/Runtime/Experimental/Chaos/Private/Chaos/PBDRigidsEvolutionGBF.cpp
@@ -188,6 +188,61 @@ void FPBDRigidsEvolutionGBF::AdvanceOneTimeStep(const FReal Dt,const FSubStepInf
 	UnprepareTick();
 }
 
+void FPBDRigidsEvolutionGBF::ReloadParticlesCache(const int32 Island)
+{
+	if(!GetConstraintGraph().GetSolverIsland(Island)->IsSleeping())
+	{
+		FEvolutionResimCache* ResimCache = GetCurrentStepResimCache();
+		if (ResimCache && ResimCache->IsResimming() && GetConstraintGraph().IslandNeedsResim(Island) == false)
+		{
+			for (auto Particle : GetConstraintGraph().GetIslandParticles(Island))
+			{
+				if (auto Rigid = Particle->CastToRigidParticle())
+				{
+					ResimCache->ReloadParticlePostSolve(*Rigid);
+				}
+			}
+			GetConstraintGraph().GetSolverIsland(Island)->SetIsUsingCache(true);
+		}
+		GetConstraintGraph().GetSolverIsland(Island)->SetIsUsingCache(false);
+	}
+}
+
+void FPBDRigidsEvolutionGBF::BuildDisabledParticles(const int32 Island, TArray<TArray<FPBDRigidParticleHandle*>>& DisabledParticles, TArray<bool>& SleepedIslands)
+{
+	for (auto Particle : GetConstraintGraph().GetIslandParticles(Island))
+	{
+		// If a dynamic particle is moving slowly enough for long enough, disable it.
+		// @todo(mlentine): Find a good way of not doing this when we aren't using this functionality
+
+		// increment the disable count for the particle
+		auto PBDRigid = Particle->CastToRigidParticle();
+		if (PBDRigid && PBDRigid->ObjectState() == EObjectStateType::Dynamic)
+		{
+			if (PBDRigid->AuxilaryValue(PhysicsMaterials) && PBDRigid->V().SizeSquared() < PBDRigid->AuxilaryValue(PhysicsMaterials)->DisabledLinearThreshold &&
+				PBDRigid->W().SizeSquared() < PBDRigid->AuxilaryValue(PhysicsMaterials)->DisabledAngularThreshold)
+			{
+				++PBDRigid->AuxilaryValue(ParticleDisableCount);
+			}
+
+			// check if we're over the disable count threshold
+			if (PBDRigid->AuxilaryValue(ParticleDisableCount) > DisableThreshold)
+			{
+				PBDRigid->AuxilaryValue(ParticleDisableCount) = 0;
+				DisabledParticles[Island].Add(PBDRigid);
+			}
+
+			if (!(ensure(!FMath::IsNaN(PBDRigid->P()[0])) && ensure(!FMath::IsNaN(PBDRigid->P()[1])) && ensure(!FMath::IsNaN(PBDRigid->P()[2]))))
+			{
+				DisabledParticles[Island].Add(PBDRigid);
+			}
+		}
+	}
+
+	// Turn off if not moving
+	SleepedIslands[Island] = GetConstraintGraph().SleepInactive(Island, PhysicsMaterials, SolverPhysicsMaterials);
+}
+
 int32 DrawAwake = 0;
 FAutoConsoleVariableRef CVarDrawAwake(TEXT("p.chaos.DebugDrawAwake"),DrawAwake,TEXT("Draw particles that are awake"));
 
@@ -310,130 +365,88 @@ void FPBDRigidsEvolutionGBF::AdvanceOneTimeStepImpl(const FReal Dt, const FSubSt
 	if(Dt > 0)
 	{
 		SCOPE_CYCLE_COUNTER(STAT_Evolution_ParallelSolve);
-		PhysicsParallelFor(GetConstraintGraph().NumIslands(), [&](int32 Island) {
+		PhysicsParallelFor(GetConstraintGraph().NumGroups(), [&](int32 GroupIndex) {
+
+			if(GetConstraintGraph().GetIslandGroups().IsValidIndex(GroupIndex) &&
+			  !GetConstraintGraph().GetIslandGroup(GroupIndex)->IsValid())
+			  	
+				return;
+			
+			// Reload cache if necessary
+			{
+				CSV_SCOPED_TIMING_STAT(PhysicsVerbose, StepSolver_ReloadCache);
+				for(auto& IslandSolver : GetConstraintGraph().GetGroupIslands(GroupIndex))
+				{
+					ReloadParticlesCache(IslandSolver->GetIslandIndex());
+				}
+			}
 
-			bool bHasCachedData = false;
-			FEvolutionResimCache* ResimCache = GetCurrentStepResimCache();
-			if (ResimCache && ResimCache->IsResimming() && GetConstraintGraph().IslandNeedsResim(Island) == false)
+			// Collect all the data that the constraint solvers operate on
 			{
-				bHasCachedData = true;
+				SCOPE_CYCLE_COUNTER(STAT_Evolution_Gather);
+				CSV_SCOPED_TIMING_STAT(PhysicsVerbose, StepSolver_Gather);
+				GatherSolverInput(Dt, GroupIndex);
 			}
 
-			if(GetConstraintGraph().GetSolverIsland(Island)->IsSleeping())
+			// Run the first phase of the constraint solvers
+			// For GBF this is the hybrid velocity solving step (which also moves the bodies to make the implicit velocity be what it should be)
+			// For PBD/QPBD this is the position solve step
 			{
-				return;
+				SCOPE_CYCLE_COUNTER(STAT_Evolution_ApplyConstraintsPhase1);
+				CSV_SCOPED_TIMING_STAT(Chaos, ApplyConstraints);
+				CSV_SCOPED_TIMING_STAT(PhysicsVerbose, StepSolver_Apply);
+				ApplyConstraintsPhase1(Dt, GroupIndex);
 			}
-			
-			const TArray<FGeometryParticleHandle*>& IslandParticles = GetConstraintGraph().GetIslandParticles(Island);
 
-			if (bHasCachedData)
+			for(auto& IslandSolver : GetConstraintGraph().GetGroupIslands(GroupIndex))
 			{
-				for (auto Particle : IslandParticles)
+				if(!IslandSolver->IsSleeping() && !IslandSolver->IsUsingCache())
 				{
-					if (auto Rigid = Particle->CastToRigidParticle())
+					if (PostApplyCallback != nullptr)
 					{
-						ResimCache->ReloadParticlePostSolve(*Rigid);
+						PostApplyCallback(IslandSolver->GetIslandIndex());
 					}
 				}
 			}
-			else
-			{
-				// Update constraint graphs, coloring etc as required by the different constraint types in this island
-				{
-					SCOPE_CYCLE_COUNTER(STAT_Evolution_UpdateAccelerationStructures);
-					CSV_SCOPED_TIMING_STAT(PhysicsVerbose, StepSolver_UpdateAccelerationStructures);
-					UpdateAccelerationStructures(Dt, Island);
-				}
-
-				// Collect all the data that the constraint solvers operate on
-				{
-					SCOPE_CYCLE_COUNTER(STAT_Evolution_Gather);
-					CSV_SCOPED_TIMING_STAT(PhysicsVerbose, StepSolver_Gather);
-					GatherSolverInput(Dt, Island);
-				}
-
-				// Run the first phase of the constraint solvers
-				// For GBF this is the hybrid velocity solving step (which also moves the bodies to make the implicit velocity be what it should be)
-				// For PBD/QPBD this is the position solve step
-				{
-					SCOPE_CYCLE_COUNTER(STAT_Evolution_ApplyConstraintsPhase1);
-					CSV_SCOPED_TIMING_STAT(Chaos, ApplyConstraints);
-					CSV_SCOPED_TIMING_STAT(PhysicsVerbose, StepSolver_Apply);
-					ApplyConstraintsPhase1(Dt, Island);
-				}
-
-				if (PostApplyCallback != nullptr)
-				{
-					PostApplyCallback(Island);
-				}
 
-				// Update implicit velocities from results of constraint solver phase 1
-				{
-					SCOPE_CYCLE_COUNTER(STAT_Evolution_UpdateVelocites);
-					CSV_SCOPED_TIMING_STAT(PhysicsVerbose, StepSolver_UpdateVelocities);
-					SetImplicitVelocities(Dt, Island);
-				}
-
-				// Run the second phase of the constraint solvers
-				// For GBF this is the pushout step
-				// For PBD this does nothing
-				// For QPBD this is the velocity solve step
-				{
-					SCOPE_CYCLE_COUNTER(STAT_Evolution_ApplyConstraintsPhase2);
-					CSV_SCOPED_TIMING_STAT(PhysicsVerbose, StepSolver_ApplyPushOut);
-					ApplyConstraintsPhase2(Dt, Island);
-				}
-
-				// Update the particles with the results of the constraint solvers, and also update constraint data
-				// that is accessed externally (net impusles, break info, etc)
-				{
-					SCOPE_CYCLE_COUNTER(STAT_Evolution_Scatter);
-					CSV_SCOPED_TIMING_STAT(PhysicsVerbose, StepSolver_Scatter);
-					ScatterSolverOutput(Dt, Island);
-				}
+			// Update implicit velocities from results of constraint solver phase 1
+			{
+				SCOPE_CYCLE_COUNTER(STAT_Evolution_UpdateVelocites);
+				CSV_SCOPED_TIMING_STAT(PhysicsVerbose, StepSolver_UpdateVelocities);
+				SetImplicitVelocities(Dt, GroupIndex);
+			}
 
-				if (PostApplyPushOutCallback != nullptr)
-				{
-					PostApplyPushOutCallback(Island);
-				}
+			// Run the second phase of the constraint solvers
+			// For GBF this is the pushout step
+			// For PBD this does nothing
+			// For QPBD this is the velocity solve step
+			{
+				SCOPE_CYCLE_COUNTER(STAT_Evolution_ApplyConstraintsPhase2);
+				CSV_SCOPED_TIMING_STAT(PhysicsVerbose, StepSolver_ApplyPushOut);
+				ApplyConstraintsPhase2(Dt, GroupIndex);
 			}
 
-			for (auto Particle : IslandParticles)
+			// Update the particles with the results of the constraint solvers, and also update constraint data
+			// that is accessed externally (net impusles, break info, etc)
 			{
-				// If a dynamic particle is moving slowly enough for long enough, disable it.
-				// @todo(mlentine): Find a good way of not doing this when we aren't using this functionality
+				SCOPE_CYCLE_COUNTER(STAT_Evolution_Scatter);
+				CSV_SCOPED_TIMING_STAT(PhysicsVerbose, StepSolver_Scatter);
+				ScatterSolverOutput(Dt, GroupIndex);
+			}
 
-				// increment the disable count for the particle
-				auto PBDRigid = Particle->CastToRigidParticle();
-				if (PBDRigid && PBDRigid->ObjectState() == EObjectStateType::Dynamic)
+			for(auto& IslandSolver : GetConstraintGraph().GetGroupIslands(GroupIndex))
+			{
+				if(!IslandSolver->IsSleeping() && !IslandSolver->IsUsingCache())
 				{
-					if (PBDRigid->AuxilaryValue(PhysicsMaterials) && PBDRigid->V().SizeSquared() < PBDRigid->AuxilaryValue(PhysicsMaterials)->DisabledLinearThreshold &&
-						PBDRigid->W().SizeSquared() < PBDRigid->AuxilaryValue(PhysicsMaterials)->DisabledAngularThreshold)
-					{
-						++PBDRigid->AuxilaryValue(ParticleDisableCount);
-					}
-
-					// check if we're over the disable count threshold
-					if (PBDRigid->AuxilaryValue(ParticleDisableCount) > DisableThreshold)
-					{
-						PBDRigid->AuxilaryValue(ParticleDisableCount) = 0;
-						//Particles.Disabled(Index) = true;
-						DisabledParticles[Island].Add(PBDRigid);
-						//Particles.V(Index) = TVector<T, d>(0);
-						//Particles.W(Index) = TVector<T, d>(0);
-					}
-
-					if (!(ensure(!FMath::IsNaN(PBDRigid->P()[0])) && ensure(!FMath::IsNaN(PBDRigid->P()[1])) && ensure(!FMath::IsNaN(PBDRigid->P()[2]))))
+					const int32 Island = IslandSolver->GetIslandIndex();
+					if (PostApplyPushOutCallback != nullptr)
 					{
-						//Particles.Disabled(Index) = true;
-						DisabledParticles[Island].Add(PBDRigid);
+						PostApplyPushOutCallback(Island);
 					}
+					BuildDisabledParticles(Island, DisabledParticles, SleepedIslands);
 				}
 			}
-
-			// Turn off if not moving
-			SleepedIslands[Island] = GetConstraintGraph().SleepInactive(Island, PhysicsMaterials, SolverPhysicsMaterials);
-		});
+		}, true);
 	}
 
 	{
@@ -521,36 +534,39 @@ void FPBDRigidsEvolutionGBF::AdvanceOneTimeStepImpl(const FReal Dt, const FSubSt
 #endif
 }
 
-void FPBDRigidsEvolutionGBF::GatherSolverInput(FReal Dt, int32 Island)
+void FPBDRigidsEvolutionGBF::GatherSolverInput(FReal Dt, int32 GroupIndex)
 {
 	// We must initialize the solver body container to be large enough to hold all particles in the
 	// island so that the pointers remain valid (the array should not grow and relocate)
-	FPBDIslandSolver* SolverIsland = ConstraintGraph.GetSolverIsland(Island);
-	if(SolverIsland)
+	if(FPBDIslandGroup* IslandGroup = ConstraintGraph.GetIslandGroups()[GroupIndex].Get())
 	{
-		// We pre-add all the dynamic particles solver bodies involved in any constraint to speed up the gather 
-		const TArray<FGeometryParticleHandle*>& IslandParticles = SolverIsland->GetParticles();
-		SolverIsland->GetBodyContainer().Reset(IslandParticles.Num());
+		IslandGroup->GetBodyContainer().Reset(IslandGroup->GetNumParticles());
 
-		if(!SolverIsland->IsSleeping())
+		for(auto& SolverIsland : IslandGroup->GetIslands())
 		{
-			for(auto& Particle : IslandParticles)
+			// We pre-add all the dynamic particles solver bodies involved in any constraint to speed up the gather 
+			const TArray<FGeometryParticleHandle*>& IslandParticles = SolverIsland->GetParticles();
+
+			if(!SolverIsland->IsSleeping() && !SolverIsland->IsUsingCache())
 			{
-				FGenericParticleHandle GenericHandle = Particle->Handle();
-				GenericHandle->SetSolverBodyIndex(INDEX_NONE);
-				
-				if(GenericHandle->ObjectState() == EObjectStateType::Dynamic)
+				for(auto& Particle : IslandParticles)
 				{
-					if(FPBDRigidParticleHandle* PBDRigid = Particle->Handle()->CastToRigidParticle())
+					FGenericParticleHandle GenericHandle = Particle->Handle();
+					GenericHandle->SetSolverBodyIndex(INDEX_NONE);
+					
+					if(GenericHandle->ObjectState() == EObjectStateType::Dynamic)
 					{
-						if(ConstraintGraph.GetGraphNodes().IsValidIndex(PBDRigid->ConstraintGraphIndex()))
+						if(FPBDRigidParticleHandle* PBDRigid = Particle->Handle()->CastToRigidParticle())
 						{
-							if(ConstraintGraph.GetGraphNodes()[PBDRigid->ConstraintGraphIndex()].NodeEdges.Num() > 0)
+							if(ConstraintGraph.GetGraphNodes().IsValidIndex(PBDRigid->ConstraintGraphIndex()))
 							{
-								// Only the non sleeping dynamic rigid particles involved in any constraints will be added
-								// by default to the solver bodies
-								const int32 BodyIndex = SolverIsland->GetBodyContainer().GetBodies().Emplace(GenericHandle);
-								GenericHandle->SetSolverBodyIndex(BodyIndex);
+								if(ConstraintGraph.GetGraphNodes()[PBDRigid->ConstraintGraphIndex()].NodeEdges.Num() > 0)
+								{
+									// Only the non sleeping dynamic rigid particles involved in any constraints will be added
+									// by default to the solver bodies
+									const int32 BodyIndex = IslandGroup->GetBodyContainer().GetBodies().Emplace(GenericHandle);
+									GenericHandle->SetSolverBodyIndex(BodyIndex);
+								}
 							}
 						}
 					}
@@ -562,23 +578,26 @@ void FPBDRigidsEvolutionGBF::GatherSolverInput(FReal Dt, int32 Island)
 	// NOTE: SolverBodies are gathered as part of the constraint gather, in the order that they are first seen 
 	for (FPBDConstraintGraphRule* ConstraintRule : PrioritizedConstraintRules)
 	{
-		ConstraintRule->GatherSolverInput(Dt, Island);
+		ConstraintRule->GatherSolverInput(Dt, GroupIndex);
 	}
 }
 
-void FPBDRigidsEvolutionGBF::ScatterSolverOutput(FReal Dt, int32 Island)
+void FPBDRigidsEvolutionGBF::ScatterSolverOutput(FReal Dt, int32 GroupIndex)
 {
 	// Scatter solver results for constraints (impulses, break events, etc)
 	for (FPBDConstraintGraphRule* ConstraintRule : PrioritizedConstraintRules)
 	{
-		ConstraintRule->ScatterSolverOutput(Dt, Island);
+		ConstraintRule->ScatterSolverOutput(Dt, GroupIndex);
 	}
 
 	// Scatter body results back to particles (position, rotation, etc)
-	ConstraintGraph.GetSolverIsland(Island)->GetBodyContainer().ScatterOutput();
+	if(FPBDIslandGroup* IslandGroup = ConstraintGraph.GetIslandGroups()[GroupIndex].Get())
+	{
+		IslandGroup->GetBodyContainer().ScatterOutput();
+	}
 }
 
-void FPBDRigidsEvolutionGBF::ApplyConstraintsPhase1(const FReal Dt, int32 Island)
+void FPBDRigidsEvolutionGBF::ApplyConstraintsPhase1(const FReal Dt, int32 GroupIndex)
 {
 	int32 LocalNumIterations = ChaosNumContactIterationsOverride >= 0 ? ChaosNumContactIterationsOverride : NumIterations;
 	// @todo(ccaulfield): track whether we are sufficiently solved and can early-out
@@ -587,7 +606,7 @@ void FPBDRigidsEvolutionGBF::ApplyConstraintsPhase1(const FReal Dt, int32 Island
 		bool bNeedsAnotherIteration = false;
 		for (FPBDConstraintGraphRule* ConstraintRule : PrioritizedConstraintRules)
 		{
-			bNeedsAnotherIteration |= ConstraintRule->ApplyConstraints(Dt, Island, i, LocalNumIterations);
+			bNeedsAnotherIteration |= ConstraintRule->ApplyConstraints(Dt, GroupIndex, i, LocalNumIterations);
 		}
 
 		if (ChaosRigidsEvolutionApplyAllowEarlyOutCVar && !bNeedsAnotherIteration)
@@ -597,12 +616,15 @@ void FPBDRigidsEvolutionGBF::ApplyConstraintsPhase1(const FReal Dt, int32 Island
 	}
 }
 
-void FPBDRigidsEvolutionGBF::SetImplicitVelocities(const FReal Dt, int32 Island)
+void FPBDRigidsEvolutionGBF::SetImplicitVelocities(const FReal Dt, int32 GroupIndex)
 {
-	ConstraintGraph.GetSolverIsland(Island)->GetBodyContainer().SetImplicitVelocities(Dt);
+	if(FPBDIslandGroup* IslandGroup = GetConstraintGraph().GetIslandGroups()[GroupIndex].Get())
+	{
+		IslandGroup->GetBodyContainer().SetImplicitVelocities(Dt);
+	}
 }
 
-void FPBDRigidsEvolutionGBF::ApplyConstraintsPhase2(const FReal Dt, int32 Island)
+void FPBDRigidsEvolutionGBF::ApplyConstraintsPhase2(const FReal Dt, int32 GroupIndex)
 {
 	int32 LocalNumPushOutIterations = ChaosNumPushOutIterationsOverride >= 0 ? ChaosNumPushOutIterationsOverride : NumPushOutIterations;
 	bool bNeedsAnotherIteration = true;
@@ -611,7 +633,7 @@ void FPBDRigidsEvolutionGBF::ApplyConstraintsPhase2(const FReal Dt, int32 Island
 		bNeedsAnotherIteration = false;
 		for (FPBDConstraintGraphRule* ConstraintRule : PrioritizedConstraintRules)
 		{
-			bNeedsAnotherIteration |= ConstraintRule->ApplyPushOut(Dt, Island, It, LocalNumPushOutIterations);
+			bNeedsAnotherIteration |= ConstraintRule->ApplyPushOut(Dt, GroupIndex, It, LocalNumPushOutIterations);
 		}
 
 		if (ChaosRigidsEvolutionApplyPushoutAllowEarlyOutCVar && !bNeedsAnotherIteration)
diff --git a/Engine/Source/Runtime/Experimental/Chaos/Public/Chaos/Evolution/SolverDatas.h b/Engine/Source/Runtime/Experimental/Chaos/Public/Chaos/Evolution/SolverDatas.h
--- a/Engine/Source/Runtime/Experimental/Chaos/Public/Chaos/Evolution/SolverDatas.h
+++ b/Engine/Source/Runtime/Experimental/Chaos/Public/Chaos/Evolution/SolverDatas.h
@@ -34,7 +34,17 @@ public :
 	/**
 	 * @brief The number of constraint containers registered
 	*/
-	int32 NumConstraintContainerIds() const { return ConstraintDatas.Num(); }
+	int32 NumContainerIds() const { return ConstraintCounts.Num(); }
+
+	/**
+	* @brief Resize the constraint count array with a fixed number of containers
+	* @param NumContainers Number of containers extracted from the constraint rules
+	*/
+	void ResizeConstraintsCounts(const int32 NumContainers);
+
+	/** Number of constraints given a container id */
+	const int32& ConstraintCount(const int32 ContainerId) const;
+	int32& ConstraintCount(const int32 ContainerId);
 
 	/** Accessors of the constraint container */
 	template<typename ContainerType> const ContainerType& GetConstraintContainer(const int32 ContainerId) const;
@@ -89,6 +99,9 @@ protected:
 	/** List of constraint datas (collision, joints...) that will be used to solve constraints */
 	TSparseArray<FConstraintDatas> ConstraintDatas;
 
+	/** Number of constraints per container id*/
+	TArray<int32> ConstraintCounts;
+
 	/** Island index in case these datas belong to an island */
 	int32 IslandIndex;
 };
@@ -137,6 +150,22 @@ FORCEINLINE ConstraintType* FPBDIslandSolverData::GetConstraintHandle(const int3
 	return static_cast<ConstraintType*>(ConstraintDatas[ContainerId].ConstraintHandles[ConstraintIndex]);
 };
 	
+FORCEINLINE const int32& FPBDIslandSolverData::ConstraintCount(const int32 ContainerId) const
+{
+	return ConstraintCounts[ContainerId];
+};
+	
+FORCEINLINE int32& FPBDIslandSolverData::ConstraintCount(const int32 ContainerId)
+{
+	return ConstraintCounts[ContainerId];
+};
+
+FORCEINLINE void FPBDIslandSolverData::ResizeConstraintsCounts(const int32 NumContainers)
+{
+	ConstraintCounts.SetNum(NumContainers, false);
+	FMemory::Memzero(ConstraintCounts.GetData(), ConstraintCounts.Num() * sizeof(int32));
+};
+	
 template<typename ConstraintType>
 inline void FPBDIslandSolverData::AddConstraintDatas(const int32 ContainerId)
 {
diff --git a/Engine/Source/Runtime/Experimental/Chaos/Public/Chaos/Island/IslandGroup.h b/Engine/Source/Runtime/Experimental/Chaos/Public/Chaos/Island/IslandGroup.h
new file mode 100644
--- /dev/null
+++ b/Engine/Source/Runtime/Experimental/Chaos/Public/Chaos/Island/IslandGroup.h
@@ -0,0 +1,91 @@
+﻿// Copyright Epic Games, Inc. All Rights Reserved.
+#pragma once
+
+#include "Containers/Array.h"
+#include "Chaos/Evolution/SolverDatas.h"
+
+namespace Chaos
+{
+	
+/** Forward Declaration */
+class FPBDIslandSolver;
+
+/**
+* Group of islands that will be used by evolution to run the solvers steps in parallel
+*/
+class CHAOS_API FPBDIslandGroup : public FPBDIslandSolverData
+{
+public:
+
+	/**
+	* Init the island group 
+	*/
+	FPBDIslandGroup(const int32 GroupIndexIn);
+
+	/**
+	* Init group members to their default values
+	*/
+	void InitGroup();
+	
+	/**
+	* Reset the islands list and reserve a number of islands in memory 
+	* @param NumIslands number of islands to be reserved
+	*/
+	void ReserveIslands(const int32 NumIslands);
+
+	/**
+	* Add island to the group
+	* @param IslandSolver Island Solver top be added
+	*/
+	void AddIsland(FPBDIslandSolver* IslandSolver);
+
+	/**
+	* Remove all islands from the group
+	*/
+	void ClearIslands();
+
+	/**
+	* Check if the group is valid and contains islands
+	*/
+	FORCEINLINE bool IsValid() const { return (IslandSolvers.Num() > 0); }
+
+	/**
+	* Get the number of islands within the group
+	*/
+	FORCEINLINE int32 NumIslands() const { return IslandSolvers.Num(); }
+
+	/**
+	* Return the islands within the group
+	*/
+	FORCEINLINE const TArray<FPBDIslandSolver*>& GetIslands() const { return IslandSolvers; }
+
+	/**
+	* Return the islands within the group
+	*/
+	FORCEINLINE TArray<FPBDIslandSolver*>& GetIslands() { return IslandSolvers; }
+
+	/**
+	* Accessors for the number of particles
+	*/
+	FORCEINLINE const int32& GetNumParticles() const  { return ParticlesCount; }
+	FORCEINLINE int32& NumParticles() {return ParticlesCount; }
+
+	/**
+	* Accessors for the number of constraints
+	*/
+	FORCEINLINE const int32& GetNumConstraints() const  { return ConstraintsCount; }
+	FORCEINLINE int32& NumConstraints() {return ConstraintsCount; }
+	
+private:
+
+	/** List of all the island constraints handles */
+	TArray<FPBDIslandSolver*> IslandSolvers;
+
+	/** Number of particles in this group */
+	int32 ParticlesCount = 0;
+
+	/** Number of constraints in this group */
+	int32 ConstraintsCount = 0;
+};
+	
+}
\ No newline at end of file
diff --git a/Engine/Source/Runtime/Experimental/Chaos/Public/Chaos/Island/IslandManager.h b/Engine/Source/Runtime/Experimental/Chaos/Public/Chaos/Island/IslandManager.h
--- a/Engine/Source/Runtime/Experimental/Chaos/Public/Chaos/Island/IslandManager.h
+++ b/Engine/Source/Runtime/Experimental/Chaos/Public/Chaos/Island/IslandManager.h
@@ -10,6 +10,7 @@
 #include "Chaos/GeometryParticlesfwd.h"
 #include "Chaos/Island/IslandGraph.h"
 #include "Chaos/Island/SolverIsland.h"
+#include "Chaos/Island/IslandGroup.h"
 
 namespace Chaos
 {
@@ -134,7 +135,7 @@ public:
 	  * @param PBDRigids Rigid particles that will be used to fill the solver islands
 	  * @param Particles SOA used in the evolution
 	  */
-	void UpdateIslands(const TParticleView<FPBDRigidParticles>& PBDRigids, FPBDRigidsSOAs& Particles);
+	void UpdateIslands(const TParticleView<FPBDRigidParticles>& PBDRigids, FPBDRigidsSOAs& Particles, const int32 NumContainers);
 
 	/**
 	  * @brief Put all particles and constraints in an island to sleep
@@ -188,6 +189,15 @@ public:
 	{
 		return IslandIndexing.Num();
 	}
+
+	/**
+	* The number of groups in the graph.
+	* @return the number of groups
+	*/
+	FORCEINLINE int32 NumGroups() const
+	{
+		return IslandGroups.Num();
+	}
 	
 	/**
 	 * The number of islands in the graph.
@@ -229,6 +239,10 @@ public:
 	template<typename ConstraintType>
 	void AddConstraintDatas(const int32 ContainerId)
 	{
+		for(auto& IslandGroup : IslandGroups)
+		{
+			IslandGroup->AddConstraintDatas<ConstraintType>(ContainerId);
+		}
 		for(auto& IslandSolver : IslandSolvers)
 		{
 			IslandSolver->AddConstraintDatas<ConstraintType>(ContainerId);
@@ -247,10 +261,22 @@ public:
 	// 	}
 	// }
 
+	/** Accessors of the graph group */
+	const FPBDIslandGroup* GetIslandGroup(const int32 GroupIndex) const {return IslandGroups.IsValidIndex(GroupIndex) ? IslandGroups[GroupIndex].Get() : nullptr; }
+	FPBDIslandGroup* GetIslandGroup(const int32 GroupIndex) {return IslandGroups.IsValidIndex(GroupIndex) ? IslandGroups[GroupIndex].Get()  : nullptr; }
+
 	/** Accessors of the graph island */
 	const FPBDIslandSolver* GetSolverIsland(const int32 IslandIndex) const {return IslandIndexing.IsValidIndex(IslandIndex) ? IslandSolvers[IslandIndexing[IslandIndex]].Get() : nullptr; }
 	FPBDIslandSolver* GetSolverIsland(const int32 IslandIndex) {return IslandIndexing.IsValidIndex(IslandIndex) ? IslandSolvers[IslandIndexing[IslandIndex]].Get() : nullptr; }
 
+	/** Accessors of the group islands */
+	const TArray<FPBDIslandSolver*>& GetGroupIslands(const int32 GroupIndex) const {return IslandGroups[GroupIndex]->GetIslands(); }
+	TArray<FPBDIslandSolver*>& GetGroupIslands(const int32 GroupIndex) {return IslandGroups[GroupIndex]->GetIslands(); }
+
+	/** Accessors of all the graph groups */
+	const TArray<TUniquePtr<FPBDIslandGroup>>& GetIslandGroups() const {return IslandGroups; }
+	TArray<TUniquePtr<FPBDIslandGroup>>& GetIslandGroups() {return IslandGroups; }
+
 	/** Get the graph nodes */
 	const TSparseArray<GraphType::FGraphNode>& GetGraphNodes() const { return IslandGraph->GraphNodes; }
 	
@@ -265,10 +291,20 @@ public:
 	
 protected:
 
+	/**
+	* Initialize the groups according to the number of threads
+	*/
+	void InitializeGroups();
+
+	/**
+	 * Build all the island groups
+	 */
+	void BuildGroups(const int32 NumContainers);
+
 	/**
 	* Sync the IslandSolvers with the IslandGraph
 	*/
-	void SyncIslands(FPBDRigidsSOAs& Particles);
+	void SyncIslands(FPBDRigidsSOAs& Particles, const int32 NumContainers);
 
 	/** List of solver islands in the manager */
 	TSparseArray<TUniquePtr<FPBDIslandSolver>> IslandSolvers;
@@ -281,5 +317,11 @@ protected:
 
 	/** Maximum particle index that we have in the graph*/
 	int32									MaxParticleIndex = INDEX_NONE;
+	
+	/** List of island groups in the manager */
+	TArray<TUniquePtr<FPBDIslandGroup>> IslandGroups;
+
+	/** Sorted list of islands */
+	TArray<int32>							SortedIslands;
 };
 }
\ No newline at end of file
diff --git a/Engine/Source/Runtime/Experimental/Chaos/Public/Chaos/Island/SolverIsland.h b/Engine/Source/Runtime/Experimental/Chaos/Public/Chaos/Island/SolverIsland.h
--- a/Engine/Source/Runtime/Experimental/Chaos/Public/Chaos/Island/SolverIsland.h
+++ b/Engine/Source/Runtime/Experimental/Chaos/Public/Chaos/Island/SolverIsland.h
@@ -87,6 +87,16 @@ class CHAOS_API FPBDIslandSolver : public FPBDIslandSolverData
 	*/
 	void SortConstraints();
 
+	/**
+	 * Set the island group index
+	 */
+	FORCEINLINE void SetGroupIndex(const int32 GroupIndex) { IslandGroup = GroupIndex; }
+	
+	/**
+	* Get the island group index
+	*/
+	FORCEINLINE const int32& GetGroupIndex() const { return IslandGroup; }
+
 	/**
 	* Return the list of particles within the solver island
 	*/
@@ -119,6 +129,8 @@ class CHAOS_API FPBDIslandSolver : public FPBDIslandSolverData
 	FORCEINLINE int32 GetSleepCounter() const { return SleepCounter; }
 	FORCEINLINE void SetSleepCounter(const int32 SleepCounterIn) { SleepCounter = SleepCounterIn; }
 	FORCEINLINE bool SleepingChanged() const { return bSleepingChanged; }
+	FORCEINLINE void SetIsUsingCache(const bool bIsUsingCacheIn ) { bIsUsingCache = bIsUsingCacheIn; }
+	FORCEINLINE bool IsUsingCache() const { return bIsUsingCache; }
 	
 	// template<typename ConstraintType>
 	// void GatherSolverInput(const FReal Dt, const int32 IslandIndex, const int32 ContainerId);
@@ -148,6 +160,13 @@ class CHAOS_API FPBDIslandSolver : public FPBDIslandSolverData
 
 	/** List of all the island constraints handles */
 	TArray<FConstraintHandleHolder> IslandConstraints;
+
+	/** Island Group in which the solver belongs */
+	int32 IslandGroup = 0;
+	
+	/** Check if the island is using the cache or not */
+	bool bIsUsingCache = false;
+	
 };
 	
 // template<typename ConstraintType>
diff --git a/Engine/Source/Runtime/Experimental/Chaos/Public/Chaos/PBDConstraintRule.h b/Engine/Source/Runtime/Experimental/Chaos/Public/Chaos/PBDConstraintRule.h
--- a/Engine/Source/Runtime/Experimental/Chaos/Public/Chaos/PBDConstraintRule.h
+++ b/Engine/Source/Runtime/Experimental/Chaos/Public/Chaos/PBDConstraintRule.h
@@ -115,19 +115,19 @@ namespace Chaos
 		virtual void BindToGraph(FPBDConstraintGraph& InContactGraph, uint32 InContainerId) {}
 
 		// Collect all the data required to solve the constraints in the specified island. This also fills the SolverBodies
-		virtual void GatherSolverInput(const FReal Dt, int32 Island) {}
+		virtual void GatherSolverInput(const FReal Dt, int32 GroupIndex) {}
 
 		// Scatter the results of the islands constraint solver(s) out to the appropriate places (e.g., impulses, break flags, etc)
-		virtual void ScatterSolverOutput(const FReal Dt, int32 Island) {}
+		virtual void ScatterSolverOutput(const FReal Dt, int32 GroupIndex) {}
 
 		/** Called once per tick to allow constraint containers to create/alter their constraints based on particle position */
 		virtual void UpdatePositionBasedState(const FReal Dt) {}
 
 		/** Apply all corrections for constraints in the specified island. Return true if more iterations are needed. */
-		virtual bool ApplyConstraints(const FReal Dt, int32 Island, const int32 It, const int32 NumIts) { return false; }
+		virtual bool ApplyConstraints(const FReal Dt, int32 GroupIndex, const int32 It, const int32 NumIts) { return false; }
 
 		/** Apply push out for constraints in the specified island. Return true if more iterations are needed. */
-		virtual bool ApplyPushOut(const FReal Dt, int32 Island, const int32 It, const int32 NumIts) { return false; }
+		virtual bool ApplyPushOut(const FReal Dt, int32 GroupIndex, const int32 It, const int32 NumIts) { return false; }
 
 		/** Add all constraints to the connectivity graph */
 		virtual void AddToGraph() {}
@@ -136,7 +136,7 @@ namespace Chaos
 		virtual void InitializeAccelerationStructures() {}
 
 		/** Set up the perf-acceleration structures for the specified island. May be called in parallel for islands */
-		virtual void UpdateAccelerationStructures(const FReal Dt, const int32 Island) {}
+		virtual void UpdateAccelerationStructures(const FReal Dt, const int32 GroupIndex) {}
 
 		/** Sort constraints if necessary  */
 		virtual void SortConstraints() {}
@@ -237,17 +237,17 @@ namespace Chaos
 
 		virtual ~TPBDConstraintIslandRule();
 
-		virtual void GatherSolverInput(const FReal Dt, int32 Island) override;
+		virtual void GatherSolverInput(const FReal Dt, int32 GroupIndex) override;
 
-		virtual void ScatterSolverOutput(const FReal Dt, int32 Island) override;
+		virtual void ScatterSolverOutput(const FReal Dt, int32 GroupIndex) override;
 
-		virtual bool ApplyConstraints(const FReal Dt, int32 Island, const int32 It, const int32 NumIts) override;
+		virtual bool ApplyConstraints(const FReal Dt, int32 GroupIndex, const int32 It, const int32 NumIts) override;
 
-		virtual bool ApplyPushOut(const FReal Dt, int32 Island, const int32 It, const int32 NumIts) override;
+		virtual bool ApplyPushOut(const FReal Dt, int32 GroupIndex, const int32 It, const int32 NumIts) override;
 
 		virtual void InitializeAccelerationStructures() override;
 		
-		virtual void UpdateAccelerationStructures(const FReal Dt, const int32 Island) override;
+		virtual void UpdateAccelerationStructures(const FReal Dt, const int32 GroupIndex) override;
 
 	protected:
 		using Base::Constraints;
@@ -278,11 +278,11 @@ namespace Chaos
 			Constraints.UpdatePositionBasedState(Dt);
 		}
 
-		virtual void GatherSolverInput(const FReal Dt, int32 Island) override;
+		virtual void GatherSolverInput(const FReal Dt, int32 GroupIndex) override;
 
-		virtual void ScatterSolverOutput(const FReal Dt, int32 Island) override;
+		virtual void ScatterSolverOutput(const FReal Dt, int32 GroupIndex) override;
 
-		virtual bool ApplyConstraints(const FReal Dt, int32 Island, const int32 It, const int32 NumIts) override;
+		virtual bool ApplyConstraints(const FReal Dt, int32 GroupIndex, const int32 It, const int32 NumIts) override;
 
 		virtual void RemoveConstraints(const TSet<TGeometryParticleHandle<FReal, 3>*>& InConstraints)
 		{
@@ -294,11 +294,11 @@ namespace Chaos
 			Constraints.Reset();
 		}
 
-		virtual bool ApplyPushOut(const FReal Dt, int32 Island, const int32 It, const int32 NumIts) override;
+		virtual bool ApplyPushOut(const FReal Dt, int32 GroupIndex, const int32 It, const int32 NumIts) override;
 
 		virtual void InitializeAccelerationStructures() override;
 
-		virtual void UpdateAccelerationStructures(const FReal Dt, const int32 Island) override;
+		virtual void UpdateAccelerationStructures(const FReal Dt, const int32 GroupIndex) override;
 
 		/** Sort constraints according to island/level/color*/
 		virtual void SortConstraints() override;
diff --git a/Engine/Source/Runtime/Experimental/Chaos/Public/Chaos/PBDRigidsEvolution.h b/Engine/Source/Runtime/Experimental/Chaos/Public/Chaos/PBDRigidsEvolution.h
--- a/Engine/Source/Runtime/Experimental/Chaos/Public/Chaos/PBDRigidsEvolution.h
+++ b/Engine/Source/Runtime/Experimental/Chaos/Public/Chaos/PBDRigidsEvolution.h
@@ -885,7 +885,7 @@ protected:
 
 	void CreateIslands()
 	{
-		ConstraintGraph.UpdateIslands(Particles.GetNonDisabledDynamicView(), Particles);
+		ConstraintGraph.UpdateIslands(Particles.GetNonDisabledDynamicView(), Particles, ConstraintRules.Num());
 
 		for (FPBDConstraintGraphRule* ConstraintRule : ConstraintRules)
 		{
diff --git a/Engine/Source/Runtime/Experimental/Chaos/Public/Chaos/PBDRigidsEvolutionGBF.h b/Engine/Source/Runtime/Experimental/Chaos/Public/Chaos/PBDRigidsEvolutionGBF.h
--- a/Engine/Source/Runtime/Experimental/Chaos/Public/Chaos/PBDRigidsEvolutionGBF.h
+++ b/Engine/Source/Runtime/Experimental/Chaos/Public/Chaos/PBDRigidsEvolutionGBF.h
@@ -173,6 +173,20 @@ namespace Chaos
 		FORCEINLINE FPBDSuspensionConstraints& GetSuspensionConstraints() { return SuspensionConstraints; }
 		FORCEINLINE const FPBDSuspensionConstraints& GetSuspensionConstraints() const { return SuspensionConstraints; }
 
+		/**
+		 * Reload the particles cache within an island
+		 * @param Island Index of the island in which the cache will be used
+		 */
+		void ReloadParticlesCache(const int32 Island);
+
+		/**
+		 * Build the list of disables particles and update the sleeping flag on the island
+		 * @param Island Index of the island in which the cache will be used
+		 * @param DisabledParticles List of islands disabled particles
+		 * @param SleepedIslands List of islands sleeping state 
+		 */
+		void BuildDisabledParticles(const int32 Island, TArray<TArray<FPBDRigidParticleHandle*>>& DisabledParticles, TArray<bool>& SleepedIslands);
+
 		void DestroyConstraint(FConstraintHandle* Constraint);
 
 		void DestroyParticleCollisionsInAllocator(FGeometryParticleHandle* Particle);
@@ -314,15 +328,15 @@ namespace Chaos
 		// First phase of constraint solver
 		// For GBF this is the velocity solve phase
 		// For PBD/QuasiPBD this is the position solve phase
-		void ApplyConstraintsPhase1(const FReal Dt, int32 Island);
+		void ApplyConstraintsPhase1(const FReal Dt, int32 GroupIndex);
 
 		// Calculate the implicit velocites based on the change in position from ApplyConstraintsPhase1
-		void SetImplicitVelocities(const FReal Dt, int32 Island);
+		void SetImplicitVelocities(const FReal Dt, int32 GroupIndex);
 		
 		// Second phase of constraint solver (after implicit velocity calculation following results of phase 1)
 		// For GBF this is the pushout phase
 		// For QuasiPBD this is the velocity solve phase
-		void ApplyConstraintsPhase2(const FReal Dt, int32 Island);
+		void ApplyConstraintsPhase2(const FReal Dt, int32 GroupIndex);
 
 
 		CHAOS_API void Serialize(FChaosArchive& Ar);
@@ -339,8 +353,8 @@ namespace Chaos
 
 		CHAOS_API void AdvanceOneTimeStepImpl(const FReal dt, const FSubStepInfo& SubStepInfo);
 
-		void GatherSolverInput(FReal Dt, int32 Island);
-		void ScatterSolverOutput(FReal Dt, int32 Island);
+		void GatherSolverInput(FReal Dt, int32 GroupIndex);
+		void ScatterSolverOutput(FReal Dt, int32 GroupIndex);
 
 		FEvolutionResimCache* GetCurrentStepResimCache()
 		{
