Chaos - Fix potential crash if a contraint with two non-dynamic CCD is created

diff --git a/Engine/Source/Runtime/Experimental/Chaos/Private/Chaos/CCDUtilities.cpp b/Engine/Source/Runtime/Experimental/Chaos/Private/Chaos/CCDUtilities.cpp
--- a/Engine/Source/Runtime/Experimental/Chaos/Private/Chaos/CCDUtilities.cpp
+++ b/Engine/Source/Runtime/Experimental/Chaos/Private/Chaos/CCDUtilities.cpp
@@ -140,19 +140,23 @@ namespace Chaos
 				bNeedCCDSolve = true;
 			}
 
-			CCDConstraints.Add(FCCDConstraint(Constraint, CCDParticlePair, Displacements));
-			for (int32 i = 0; i < 2; i++)
+			// make sure we ignore pairs that don't include any dynamics
+			if (CCDParticlePair[0] != nullptr || CCDParticlePair[1] != nullptr)
 			{
-				if (CCDParticlePair[i] != nullptr)
+				CCDConstraints.Add(FCCDConstraint(Constraint, CCDParticlePair, Displacements));
+				for (int32 i = 0; i < 2; i++)
 				{
-					CCDParticlePair[i]->AddConstraint(&CCDConstraints.Last());
+					if (CCDParticlePair[i] != nullptr)
+					{
+						CCDParticlePair[i]->AddConstraint(&CCDConstraints.Last());
+					}
 				}
-			}
 
-			if (IsDynamic[0] && IsDynamic[1])
-			{
-				CCDParticlePair[0]->AddOverlappingDynamicParticle(CCDParticlePair[1]);
-				CCDParticlePair[1]->AddOverlappingDynamicParticle(CCDParticlePair[0]);
+				if (IsDynamic[0] && IsDynamic[1])
+				{
+					CCDParticlePair[0]->AddOverlappingDynamicParticle(CCDParticlePair[1]);
+					CCDParticlePair[1]->AddOverlappingDynamicParticle(CCDParticlePair[0]);
+				}
 			}
 		}
 		return bNeedCCDSolve;
@@ -217,6 +221,8 @@ namespace Chaos
 			}
 			if (Island == INDEX_NONE)
 			{
+				// non-dynamic pairs are already ignored in Init() so if Particle 0 is null the second one should not be 
+				ensure(CCDConstraint.Particle[1] != nullptr); 
 				Island = CCDConstraint.Particle[1]->Island;
 			}
 			CCDConstraint.Island = Island;
