Fix Android build break (A[X++] = B[X] is an undefined operation because the evaluation order of the two X references is not specified)

diff --git a/Engine/Source/Runtime/Experimental/Chaos/Private/Chaos/Island/IslandManager.cpp b/Engine/Source/Runtime/Experimental/Chaos/Private/Chaos/Island/IslandManager.cpp
--- a/Engine/Source/Runtime/Experimental/Chaos/Private/Chaos/Island/IslandManager.cpp
+++ b/Engine/Source/Runtime/Experimental/Chaos/Private/Chaos/Island/IslandManager.cpp
@@ -565,8 +565,8 @@ void FPBDIslandManager::SyncIslands(FPBDRigidsSOAs& Particles, const int32 NumCo
 			IslandSolver->GetIslandIndex() = LocalIsland;
 
 			// We update the IslandIndexing to retrieve the graph sparse and persistent index from the dense one.
-			SortedIslands[LocalIsland] = IslandIndex;
-			IslandIndexing[LocalIsland++] = SortedIslands[LocalIsland];
+			IslandIndexing[LocalIsland] = SortedIslands[LocalIsland] = IslandIndex;
+			LocalIsland++;
 
 			// We finally update the solver islands based on the new particles and constraints if
 			// the island is not persistent or not sleeping. 
