Workaround for PVS being an idiot.

diff --git a/Engine/Source/Runtime/Experimental/Chaos/Private/Chaos/Island/IslandManager.cpp b/Engine/Source/Runtime/Experimental/Chaos/Private/Chaos/Island/IslandManager.cpp
--- a/Engine/Source/Runtime/Experimental/Chaos/Private/Chaos/Island/IslandManager.cpp
+++ b/Engine/Source/Runtime/Experimental/Chaos/Private/Chaos/Island/IslandManager.cpp
@@ -565,7 +565,8 @@ void FPBDIslandManager::SyncIslands(FPBDRigidsSOAs& Particles, const int32 NumCo
 			IslandSolver->GetIslandIndex() = LocalIsland;
 
 			// We update the IslandIndexing to retrieve the graph sparse and persistent index from the dense one.
-			IslandIndexing[LocalIsland] = SortedIslands[LocalIsland] = IslandIndex;
+			IslandIndexing[LocalIsland] = IslandIndex;
+			SortedIslands[LocalIsland] = IslandIndex;
 			LocalIsland++;
 
 			// We finally update the solver islands based on the new particles and constraints if
