Split solver task into multiple parts so we can run callbacks on frozen game thread

diff --git a/Engine/Source/Runtime/Experimental/Chaos/Private/Chaos/Framework/PhysicsSolverBase.cpp b/Engine/Source/Runtime/Experimental/Chaos/Private/Chaos/Framework/PhysicsSolverBase.cpp
--- a/Engine/Source/Runtime/Experimental/Chaos/Private/Chaos/Framework/PhysicsSolverBase.cpp
+++ b/Engine/Source/Runtime/Experimental/Chaos/Private/Chaos/Framework/PhysicsSolverBase.cpp
@@ -67,15 +67,9 @@ namespace Chaos
 		ENamedThreads::HighTaskPriority // if we don't have hi pri threads, then use normal priority threads at high task priority instead
 	);
 
-	FPhysicsSolverAdvanceTask::FPhysicsSolverAdvanceTask(FPhysicsSolverBase& InSolver, FPushPhysicsData& InPushData)
-		: Solver(InSolver)
-		, PushData(&InPushData)	//store as ptr so that we can clear it after freed (but still want to force user to give us a valid push data)
+	ENamedThreads::Type FPhysicsSolverProcessPushDataTask::GetDesiredThread()
 	{
-	}
-
-	TStatId FPhysicsSolverAdvanceTask::GetStatId() const
-	{
-		RETURN_QUICK_DECLARE_CYCLE_STAT(FPhysicsSolverAdvanceTask,STATGROUP_TaskGraphTasks);
+		return CPrio_FPhysicsTickTask.Get();
 	}
 
 	ENamedThreads::Type FPhysicsSolverAdvanceTask::GetDesiredThread()
@@ -83,21 +77,47 @@ namespace Chaos
 		return CPrio_FPhysicsTickTask.Get();
 	}
 
-	ESubsequentsMode::Type FPhysicsSolverAdvanceTask::GetSubsequentsMode()
+	void FPhysicsSolverProcessPushDataTask::ProcessPushData()
 	{
-		// The completion task relies on the collection of tick tasks in flight
-		return ESubsequentsMode::TrackSubsequents;
+		using namespace Chaos;
+
+		LLM_SCOPE(ELLMTag::ChaosUpdate);
+		SCOPE_CYCLE_COUNTER(STAT_ChaosTick);
+		CSV_SCOPED_TIMING_STAT_EXCLUSIVE(Physics);
+		CSV_SCOPED_TIMING_STAT(PhysicsVerbose, StepSolver);
+
+#if PHYSICS_THREAD_CONTEXT
+		FPhysicsThreadContextScope Scope(/*IsPhysicsThreadContext=*/true);
+#endif
+
+		Solver.SetExternalTimestampConsumed_Internal(PushData->ExternalTimestamp);
+		Solver.ProcessPushedData_Internal(*PushData);
+		
+		Solver.PrepareAdvanceBy(PushData->ExternalDt);
+
 	}
 
-	void FPhysicsSolverAdvanceTask::DoTask(ENamedThreads::Type CurrentThread,const FGraphEventRef& MyCompletionGraphEvent)
+	void FPhysicsSolverFrozenGTPreSimCallbacks::GTPreSimCallbacks()
 	{
-		AdvanceSolver();
+		LLM_SCOPE(ELLMTag::ChaosUpdate);
+		SCOPE_CYCLE_COUNTER(STAT_ChaosTick);
+		CSV_SCOPED_TIMING_STAT_EXCLUSIVE(Physics);
+		CSV_SCOPED_TIMING_STAT(PhysicsVerbose, StepSolver);
+
+		//We are on GT, but we know PhysicsThread is waiting so we're actually going to operate on PT data
+#if PHYSICS_THREAD_CONTEXT
+		FPhysicsThreadContextScope Scope(/*IsPhysicsThreadContext=*/true);
+		//TODO: add something to make sure we don't access GT solver data by mistake
+#endif
+
+		Solver.SetOnFrozenGameThread(true);
+		Solver.ApplyCallbacks_Internal();
+		Solver.SetOnFrozenGameThread(false);
+		
 	}
 
 	void FPhysicsSolverAdvanceTask::AdvanceSolver()
 	{
-		using namespace Chaos;
-
 		LLM_SCOPE(ELLMTag::ChaosUpdate);
 		SCOPE_CYCLE_COUNTER(STAT_ChaosTick);
 		CSV_SCOPED_TIMING_STAT_EXCLUSIVE(Physics);
@@ -107,14 +127,12 @@ namespace Chaos
 		FPhysicsThreadContextScope Scope(/*IsPhysicsThreadContext=*/true);
 #endif
 
-		Solver.SetExternalTimestampConsumed_Internal(PushData->ExternalTimestamp);
-		Solver.ProcessPushedData_Internal(*PushData);
-		
 		// StepFraction: how much of the remaining time this step represents, used to interpolate kinematic targets
 		// E.g., for 4 steps this will be: 1/4, 1/3, 1/2, 1
 		const FReal PseudoFraction = (FReal)1 / (FReal)(PushData->IntervalNumSteps - PushData->IntervalStep);
-		
-		Solver.AdvanceSolverBy(PushData->ExternalDt, FSubStepInfo{PseudoFraction, PushData->IntervalStep, PushData->IntervalNumSteps, PushData->bSolverSubstepped});
+
+		Solver.AdvanceSolverBy(FSubStepInfo{ PseudoFraction, PushData->IntervalStep, PushData->IntervalNumSteps, PushData->bSolverSubstepped });
+
 		Solver.GetMarshallingManager().FreeDataToHistory_Internal(PushData);	//cannot use push data after this point
 		PushData = nullptr;
 
@@ -374,7 +392,7 @@ namespace Chaos
 			if(ThreadingMode == EThreadingModeTemp::SingleThread)
 			{
 				ensure(!PendingTasks || PendingTasks->IsComplete());	//if mode changed we should have already blocked
-				FPhysicsSolverAdvanceTask ImmediateTask(*this, *PushData);
+				FAllSolverTasks ImmediateTask(*this, PushData);
 #if !UE_BUILD_SHIPPING
 				if(bStealAdvanceTasksForTesting)
 				{
@@ -402,7 +420,17 @@ namespace Chaos
 					Prereqs.Add(PendingTasks);
 				}
 
-				PendingTasks = TGraphTask<FPhysicsSolverAdvanceTask>::CreateTask(&Prereqs).ConstructAndDispatchWhenReady(*this, *PushData);
+				PendingTasks = TGraphTask<FPhysicsSolverProcessPushDataTask>::CreateTask(&Prereqs).ConstructAndDispatchWhenReady(*this, PushData);
+				Prereqs.Add(PendingTasks);
+
+				if(bSolverHasFrozenGameThreadCallbacks)
+				{
+					PendingTasks = TGraphTask<FPhysicsSolverFrozenGTPreSimCallbacks>::CreateTask(&Prereqs).ConstructAndDispatchWhenReady(*this);
+					Prereqs.Add(PendingTasks);
+				}
+				
+				PendingTasks = TGraphTask<FPhysicsSolverAdvanceTask>::CreateTask(&Prereqs).ConstructAndDispatchWhenReady(*this, PushData);
+
 				if(IsUsingAsyncResults() == false)
 				{
 					BlockingTasks = PendingTasks;	//block right away
diff --git a/Engine/Source/Runtime/Experimental/Chaos/Private/PBDRigidsSolver.cpp b/Engine/Source/Runtime/Experimental/Chaos/Private/PBDRigidsSolver.cpp
--- a/Engine/Source/Runtime/Experimental/Chaos/Private/PBDRigidsSolver.cpp
+++ b/Engine/Source/Runtime/Experimental/Chaos/Private/PBDRigidsSolver.cpp
@@ -264,12 +264,8 @@ namespace Chaos
 		{
 			LLM_SCOPE(ELLMTag::ChaosUpdate);
 			UE_LOG(LogPBDRigidsSolver, Verbose, TEXT("AdvanceOneTimeStepTask::DoWork()"));
-			MSolver->StartingSceneSimulation();
 
-			if(MDeltaTime > 0)	//if delta time is 0 we are flushing data, user callbacks should not be triggered because there is no sim
-			{
-				MSolver->ApplyCallbacks_Internal(MSolver->GetSolverTime(), MDeltaTime);	//question: is SolverTime the right thing to pass in here?
-			}
+			MSolver->ApplyCallbacks_Internal();
 			MSolver->GetEvolution()->GetRigidClustering().ResetAllClusterBreakings();
 			
 			{
@@ -391,7 +387,7 @@ namespace Chaos
 
 			MSolver->FinalizeCallbackData_Internal();
 
-			MSolver->GetSolverTime() += MDeltaTime;
+			MSolver->SetSolverTime(MSolver->GetSolverTime() + MDeltaTime );
 			MSolver->GetCurrentFrame()++;
 			MSolver->PostTickDebugDraw(MDeltaTime);
 
@@ -427,8 +423,6 @@ namespace Chaos
 	FPBDRigidsSolver::FPBDRigidsSolver(const EMultiBufferMode BufferingModeIn, UObject* InOwner, FReal InAsyncDt)
 		: Super(BufferingModeIn, BufferingModeIn == EMultiBufferMode::Single ? EThreadingModeTemp::SingleThread : EThreadingModeTemp::TaskGraph, InOwner, InAsyncDt)
 		, CurrentFrame(0)
-		, MTime(0.0)
-		, MLastDt(0.0)
 		, bHasFloor(true)
 		, bIsFloorAnalytic(false)
 		, FloorHeight(0.f)
@@ -791,9 +785,8 @@ namespace Chaos
 			GetEvolution()->GetCollisionConstraints();
 	}
 
-	void FPBDRigidsSolver::AdvanceSolverBy(const FReal DeltaTime, const FSubStepInfo& SubStepInfo)
+	void FPBDRigidsSolver::PrepareAdvanceBy(const FReal DeltaTime)
 	{
-		const FReal StartSimTime = GetSolverTime();
 		MEvolution->GetCollisionDetector().GetNarrowPhase().GetContext().bDeferUpdate = (ChaosSolverCollisionDeferNarrowPhase != 0);
 		MEvolution->GetCollisionDetector().GetNarrowPhase().GetContext().bAllowManifolds = (ChaosSolverCollisionUseManifolds != 0);
 		
@@ -853,13 +846,21 @@ namespace Chaos
 		UE_LOG(LogPBDRigidsSolver, Verbose, TEXT("PBDRigidsSolver::Tick(%3.5f)"), DeltaTime);
 		MLastDt = DeltaTime;
 		EventPreSolve.Broadcast(DeltaTime);
-		AdvanceOneTimeStepTask(this, DeltaTime, SubStepInfo).DoWork();
 
-		if(DeltaTime > 0)
+		StartingSceneSimulation();
+	}
+
+	void FPBDRigidsSolver::AdvanceSolverBy(const FSubStepInfo& SubStepInfo)
+	{
+		const FReal StartSimTime = GetSolverTime();
+
+		AdvanceOneTimeStepTask(this, MLastDt, SubStepInfo).DoWork();
+
+		if (MLastDt > 0)
 		{
 			//pass information back to external thread
 			//we skip dt=0 case because sync data should be identical if dt = 0
-			MarshallingManager.FinalizePullData_Internal(MEvolution->LatestExternalTimestampConsumed_Internal, StartSimTime, DeltaTime);
+			MarshallingManager.FinalizePullData_Internal(MEvolution->LatestExternalTimestampConsumed_Internal, StartSimTime, MLastDt);
 		}
 
 		if(SubStepInfo.Step == SubStepInfo.NumSteps - 1)
@@ -1236,7 +1237,8 @@ namespace Chaos
 						}
 
 						MRewindCallback->PreResimStep_Internal(Step, bFirst);
-						FPhysicsSolverAdvanceTask ImmediateTask(*this, *PushData);
+						FAllSolverTasks ImmediateTask(*this, PushData);
+						ensure(bSolverHasFrozenGameThreadCallbacks == false);	//We don't support this for resim as it's very expensive and difficult to schedule
 						ImmediateTask.AdvanceSolver();
 						MRewindCallback->PostResimStep_Internal(Step);
 
diff --git a/Engine/Source/Runtime/Experimental/Chaos/Public/Chaos/Framework/PhysicsSolverBase.h b/Engine/Source/Runtime/Experimental/Chaos/Public/Chaos/Framework/PhysicsSolverBase.h
--- a/Engine/Source/Runtime/Experimental/Chaos/Public/Chaos/Framework/PhysicsSolverBase.h
+++ b/Engine/Source/Runtime/Experimental/Chaos/Public/Chaos/Framework/PhysicsSolverBase.h
@@ -70,19 +70,64 @@ namespace Chaos
 	};
 
 	/**
-	 * Task responsible for processing the command buffer of a single solver and advancing it by
-	 * a specified delta before completing.
+	 * Task responsible for processing the command buffer of a single solver and preparing data before solver task and callbacks are run
+	 */
+	class CHAOS_API FPhysicsSolverProcessPushDataTask
+	{
+	public:
+
+		FPhysicsSolverProcessPushDataTask(FPhysicsSolverBase& InSolver, FPushPhysicsData* InPushData)
+			: Solver(InSolver)
+			, PushData(InPushData){}
+
+		TStatId GetStatId() const { RETURN_QUICK_DECLARE_CYCLE_STAT(FPhysicsSolverProcessPushDataTask, STATGROUP_TaskGraphTasks); }
+		static ENamedThreads::Type GetDesiredThread();
+		static ESubsequentsMode::Type GetSubsequentsMode() { return ESubsequentsMode::TrackSubsequents; }
+		void DoTask(ENamedThreads::Type CurrentThread, const FGraphEventRef& MyCompletionGraphEvent) { ProcessPushData(); }
+		void ProcessPushData();
+
+	private:
+
+		FPhysicsSolverBase& Solver;
+		FPushPhysicsData* PushData;
+	};
+
+	/**
+	 * Task responsible for triggering any game thread callbacks before solver can advance (not scheduled if none are registered)
+	 */
+	class CHAOS_API FPhysicsSolverFrozenGTPreSimCallbacks
+	{
+	public:
+
+		FPhysicsSolverFrozenGTPreSimCallbacks(FPhysicsSolverBase& InSolver) : Solver(InSolver) {}
+
+		TStatId GetStatId() const { RETURN_QUICK_DECLARE_CYCLE_STAT(FPhysicsSolverFrozenGTPreSimCallbacks, STATGROUP_TaskGraphTasks); }
+		static ENamedThreads::Type GetDesiredThread() { return ENamedThreads::GameThread; }
+		static ESubsequentsMode::Type GetSubsequentsMode() { return ESubsequentsMode::TrackSubsequents; }
+		void DoTask(ENamedThreads::Type CurrentThread, const FGraphEventRef& MyCompletionGraphEvent) { GTPreSimCallbacks(); }
+		void GTPreSimCallbacks();
+
+	private:
+
+		FPhysicsSolverBase& Solver;
+	};
+
+	/**
+	 * Task responsible for advancing the solver once data has been prepared and GT callbacks have fired
 	 */
 	class CHAOS_API FPhysicsSolverAdvanceTask
 	{
 	public:
 
-		FPhysicsSolverAdvanceTask(FPhysicsSolverBase& InSolver, FPushPhysicsData& PushData);
+		FPhysicsSolverAdvanceTask(FPhysicsSolverBase& InSolver, FPushPhysicsData* InPushData)
+			: Solver(InSolver)
+			, PushData(InPushData)
+		{}
 
-		TStatId GetStatId() const;
+		TStatId GetStatId() const { RETURN_QUICK_DECLARE_CYCLE_STAT(FPhysicsSolverAdvanceTask, STATGROUP_TaskGraphTasks); }
 		static ENamedThreads::Type GetDesiredThread();
-		static ESubsequentsMode::Type GetSubsequentsMode();
-		void DoTask(ENamedThreads::Type CurrentThread, const FGraphEventRef& MyCompletionGraphEvent);
+		static ESubsequentsMode::Type GetSubsequentsMode() { return ESubsequentsMode::TrackSubsequents; }
+		void DoTask(ENamedThreads::Type CurrentThread, const FGraphEventRef& MyCompletionGraphEvent) { AdvanceSolver(); }
 		void AdvanceSolver();
 
 	private:
@@ -91,6 +136,26 @@ namespace Chaos
 		FPushPhysicsData* PushData;
 	};
 
+	struct FAllSolverTasks
+	{
+		FAllSolverTasks(FPhysicsSolverBase& Solver, FPushPhysicsData* PushData)
+			: ProcessPushData(Solver, PushData)
+			, GTPreSimCallbacks(Solver)
+			, AdvanceTask(Solver, PushData)
+		{
+		}
+
+		FPhysicsSolverProcessPushDataTask ProcessPushData;
+		FPhysicsSolverFrozenGTPreSimCallbacks GTPreSimCallbacks;
+		FPhysicsSolverAdvanceTask AdvanceTask;
+
+		void AdvanceSolver()
+		{
+			ProcessPushData.ProcessPushData();
+			GTPreSimCallbacks.GTPreSimCallbacks();
+			AdvanceTask.AdvanceSolver();
+		}
+	};
 
 	class FPersistentPhysicsTask;
 
@@ -157,6 +222,15 @@ namespace Chaos
 		inline TSimCallbackObjectType* CreateAndRegisterSimCallbackObject_External(bool bContactModification = false, bool bRegisterRewindCallback = false)
 		{
 			auto NewCallbackObject = new TSimCallbackObjectType();
+			//If at least one callback is on frozen GT we mark it as such
+			//Note this is never cleaned up (to avoid race conditions of register/unregister).
+			//Expectation is that there's typically just one of these on frozen gt that triggers all the fixed tick on gt (maybe two if another system wants direct control)
+			//This means we don't expect to optimize for the case where a gt callback goes away half way through sim
+			if (NewCallbackObject->RunOnFrozenGameThread())
+			{
+				bSolverHasFrozenGameThreadCallbacks = true;
+			}
+
 			RegisterSimCallbackObject_External(NewCallbackObject, bContactModification);
 			if (bRegisterRewindCallback)
 			{
@@ -301,13 +375,28 @@ namespace Chaos
 		}
 #endif
 
-		void ApplyCallbacks_Internal(const FReal SimTime, const FReal Dt)
+		bool OnFrozenGameThread() const { return bOnFrozenGameThread; }
+		void SetOnFrozenGameThread(bool InOnFrozenGameThread)
+		{
+			bOnFrozenGameThread = InOnFrozenGameThread;
+		}
+
+		void ApplyCallbacks_Internal()
 		{
 			QUICK_SCOPE_CYCLE_COUNTER(ApplySimCallbacks);
-			for (ISimCallbackObject* Callback : SimCallbackObjects)
+
+			//if delta time is 0 we are flushing data, user callbacks should not be triggered because there is no sim
+			if(MLastDt > 0)
 			{
-				Callback->SetSimAndDeltaTime_Internal(SimTime, Dt);
-				Callback->PreSimulate_Internal();
+				const FReal SimTime = GetSolverTime();
+				for (ISimCallbackObject* Callback : SimCallbackObjects)
+				{
+					if (Callback->RunOnFrozenGameThread() == bOnFrozenGameThread)
+					{
+						Callback->SetSimAndDeltaTime_Internal(SimTime, MLastDt);
+						Callback->PreSimulate_Internal();
+					}
+				}
 			}
 		}
 
@@ -402,6 +491,14 @@ namespace Chaos
 			}
 		}
 
+		/**/
+		void SetSolverTime(const FReal InTime) { MTime = InTime; }
+		const FReal GetSolverTime() const { return MTime; }
+
+
+		/**/
+		FReal GetLastDt() const { return MLastDt; }
+
 	protected:
 		/** Mode that the results buffers should be set to (single, double, triple) */
 		EMultiBufferMode BufferMode;
@@ -422,7 +519,8 @@ namespace Chaos
 		FPhysicsSolverBase& operator =(const FPhysicsSolverBase& InCopy) = delete;
 		FPhysicsSolverBase& operator =(FPhysicsSolverBase&& InSteal) = delete;
 
-		virtual void AdvanceSolverBy(const FReal Dt, const FSubStepInfo& SubStepInfo = FSubStepInfo()) = 0;
+		virtual void PrepareAdvanceBy(const FReal Dt) = 0;
+		virtual void AdvanceSolverBy(const FSubStepInfo& SubStepInfo = FSubStepInfo()) = 0;
 		virtual void PushPhysicsState(const FReal Dt, const int32 NumSteps, const int32 NumExternalSteps) = 0;
 		virtual void ProcessPushedData_Internal(FPushPhysicsData& PushDataArray) = 0;
 		virtual void SetExternalTimestampConsumed_Internal(const int32 Timestamp) = 0;
@@ -447,6 +545,11 @@ namespace Chaos
 
 	FGraphEventRef PendingTasks;
 
+	bool bSolverHasFrozenGameThreadCallbacks = false;
+	bool bOnFrozenGameThread = false;
+	FReal MLastDt = FReal(0);
+	FReal MTime = FReal(0);
+
 	private:
 
 		//This is private because the user should never create their own callback object
@@ -480,6 +583,7 @@ namespace Chaos
 		TUniquePtr<FPhysSceneLock> ExternalDataLock_External;
 
 		friend FChaosSolversModule;
+		friend FPhysicsSolverProcessPushDataTask;
 		friend FPhysicsSolverAdvanceTask;
 
 		template<ELockType>
@@ -533,7 +637,7 @@ namespace Chaos
 	private:
 		// instead of running advance task in single threaded, put in array for manual execution control for unit tests.
 		bool bStealAdvanceTasksForTesting;
-		TArray<FPhysicsSolverAdvanceTask> StolenSolverAdvanceTasks;
+		TArray<FAllSolverTasks> StolenSolverAdvanceTasks;
 	public:
 		void SetStealAdvanceTasks_ForTesting(bool bInStealAdvanceTasksForTesting);
 		void PopAndExecuteStolenAdvanceTask_ForTesting();
diff --git a/Engine/Source/Runtime/Experimental/Chaos/Public/Chaos/SimCallbackObject.h b/Engine/Source/Runtime/Experimental/Chaos/Public/Chaos/SimCallbackObject.h
--- a/Engine/Source/Runtime/Experimental/Chaos/Public/Chaos/SimCallbackObject.h
+++ b/Engine/Source/Runtime/Experimental/Chaos/Public/Chaos/SimCallbackObject.h
@@ -80,10 +80,12 @@ public:
 	virtual int32 TriggerRewindIfNeeded_Internal(int32 LastCompletedStep) { ensure(false); return INDEX_NONE; }
 	virtual void ApplyCorrections_Internal(int32 PhysicsStep, FSimCallbackInput* Input) { ensure(false); }
 	virtual void FirstPreResimStep_Internal(int32 PhysicsStep) { }
+
+	bool RunOnFrozenGameThread() const { return bRunOnFrozenGameThread; }
 	
 protected:
 
-	ISimCallbackObject()
+	ISimCallbackObject(bool InRunOnFrozenGameThread = false)
 	: bPendingDelete(false)
 	, bPendingDelete_External(false)
 	, bContactModification(false)
@@ -91,6 +93,7 @@ protected:
 	, Solver(nullptr)
 	, CurrentOutput_Internal(nullptr)
 	, CurrentInput_Internal(nullptr)
+	, bRunOnFrozenGameThread(InRunOnFrozenGameThread)
 	{
 	}
 
@@ -170,6 +173,7 @@ private:
 
 	FReal SimTime_Internal;
 	FReal DeltaTime_Internal;
+	bool bRunOnFrozenGameThread;
 };
 
 /** Simple callback command object. Commands are typically passed in as lambdas and there's no need for data management. Should not be used directly, see FPhysicsSolverBase::EnqueueCommand */
@@ -217,8 +221,9 @@ class TSimCallbackObject : public ISimCallbackObject
 {
 public:
 
-	TSimCallbackObject()
-	: CurrentOutput_External(nullptr)
+	TSimCallbackObject(bool InRunOnFrozenGameThread = false)
+	: ISimCallbackObject(InRunOnFrozenGameThread)
+	, CurrentOutput_External(nullptr)
 	{
 	}
 
diff --git a/Engine/Source/Runtime/Experimental/Chaos/Public/PBDRigidsSolver.h b/Engine/Source/Runtime/Experimental/Chaos/Public/PBDRigidsSolver.h
--- a/Engine/Source/Runtime/Experimental/Chaos/Public/PBDRigidsSolver.h
+++ b/Engine/Source/Runtime/Experimental/Chaos/Public/PBDRigidsSolver.h
@@ -174,13 +174,6 @@ namespace Chaos
 		void SetCurrentFrame(const int32 CurrentFrameIn) { CurrentFrame = CurrentFrameIn; }
 		int32& GetCurrentFrame() { return CurrentFrame; }
 
-		/**/
-		FReal& GetSolverTime() { return MTime; }
-		const FReal GetSolverTime() const { return MTime; }
-
-		/**/
-		FReal GetLastDt() const { return MLastDt; }
-
 		/**/
 		void SetIterations(const int32 InNumIterations) { GetEvolution()->SetNumIterations(InNumIterations); }
 		void SetPushOutIterations(const int32 InNumIterations) {  GetEvolution()->SetNumPushOutIterations(InNumIterations); }
@@ -321,7 +314,8 @@ namespace Chaos
 		void BufferPhysicsResults();
 	
 		/**/
-		virtual void AdvanceSolverBy(const FReal DeltaTime, const FSubStepInfo& SubStepInfo = FSubStepInfo()) override;
+		virtual void PrepareAdvanceBy(const FReal DeltaTime) override;
+		virtual void AdvanceSolverBy(const FSubStepInfo& SubStepInfo) override;
 		virtual void PushPhysicsState(const FReal ExternalDt, const int32 NumSteps, const int32 NumExternalSteps) override;
 		virtual void SetExternalTimestampConsumed_Internal(const int32 Timestamp) override;
 
@@ -331,8 +325,6 @@ namespace Chaos
 		// Solver Data
 		//
 		int32 CurrentFrame;
-		FReal MTime;
-		FReal MLastDt;
 		bool bHasFloor;
 		bool bIsFloorAnalytic;
 		FReal FloorHeight;
