Fix IsOnFrozenGameThread so that other threads get false and continue to use the interpolated data

diff --git a/Engine/Source/Runtime/Experimental/Chaos/Public/Chaos/Framework/PhysicsSolverBase.h b/Engine/Source/Runtime/Experimental/Chaos/Public/Chaos/Framework/PhysicsSolverBase.h
--- a/Engine/Source/Runtime/Experimental/Chaos/Public/Chaos/Framework/PhysicsSolverBase.h
+++ b/Engine/Source/Runtime/Experimental/Chaos/Public/Chaos/Framework/PhysicsSolverBase.h
@@ -375,7 +375,10 @@ namespace Chaos
 		}
 #endif
 
-		bool OnFrozenGameThread() const { return bOnFrozenGameThread; }
+		//Tells us if we're on the frozen game thread. This is needed for knowing which data to read/write to
+		//The IsInGameThread check is so that other threads (e.g audio thread) which might be running queries in parallel will continue to use the correct interpolated GT data
+		bool IsOnFrozenGameThread() const { return bOnFrozenGameThread && (IsInGameThread() || IsInParallelGameThread()); }
+
 		void SetOnFrozenGameThread(bool InOnFrozenGameThread)
 		{
 			bOnFrozenGameThread = InOnFrozenGameThread;
