Fix inconsistent hit events

diff --git a/Engine/Source/Runtime/Engine/Private/PhysicsEngine/Experimental/PhysScene_Chaos.cpp b/Engine/Source/Runtime/Engine/Private/PhysicsEngine/Experimental/PhysScene_Chaos.cpp
--- a/Engine/Source/Runtime/Engine/Private/PhysicsEngine/Experimental/PhysScene_Chaos.cpp
+++ b/Engine/Source/Runtime/Engine/Private/PhysicsEngine/Experimental/PhysScene_Chaos.cpp
@@ -33,7 +33,6 @@
 #include "PhysicsProxy/SkeletalMeshPhysicsProxy.h"
 #include "PhysicsProxy/StaticMeshPhysicsProxy.h"
 #include "PhysicsProxy/JointConstraintProxy.h"
-#include "Chaos/Real.h"
 #include "Chaos/UniformGrid.h"
 #include "Chaos/BoundingVolume.h"
 #include "Chaos/Framework/DebugSubstep.h"
@@ -382,6 +381,9 @@ FPhysScene_Chaos::FPhysScene_Chaos(AActor* InSolverActor
 	)
 	, PhysicsReplication(nullptr)
 	, SolverActor(InSolverActor)
+#if WITH_CHAOS
+	, LastEventDispatchTime(Chaos::FReal(-1))
+#endif
 #if WITH_EDITOR
 	, SingleStepCounter(0)
 #endif
@@ -680,15 +682,18 @@ FCollisionNotifyInfo& FPhysScene_Chaos::GetPendingCollisionForContactPair(const
 
 void FPhysScene_Chaos::HandleCollisionEvents(const Chaos::FCollisionEventData& Event)
 {
-
+#if WITH_CHAOS
 	ContactPairToPendingNotifyMap.Reset();
 
 	TMap<IPhysicsProxyBase*, TArray<int32>> const& PhysicsProxyToCollisionIndicesMap = Event.PhysicsProxyToCollisionIndices.PhysicsProxyToIndicesMap;
 	Chaos::FCollisionDataArray const& CollisionData = Event.CollisionData.AllCollisionsArray;
 	const Chaos::FReal MinDeltaVelocityThreshold = UPhysicsSettings::Get()->MinDeltaVelocityForHitEvents;
 	int32 NumCollisions = CollisionData.Num();
-	if (NumCollisions > 0)
+
+	if (NumCollisions > 0 && LastEventDispatchTime < Event.CollisionData.TimeCreated)
 	{
+		LastEventDispatchTime = Event.CollisionData.TimeCreated;
+
 		// look through all the components that someone is interested in, and see if they had a collision
 		// note that we only need to care about the interaction from the POV of the registered component,
 		// since if anyone wants notifications for the other component it hit, it's also registered and we'll get to that elsewhere in the list
@@ -757,6 +762,7 @@ void FPhysScene_Chaos::HandleCollisionEvents(const Chaos::FCollisionEventData& E
 
 	// Tell the world and actors about the collisions
 	DispatchPendingCollisionNotifies();
+#endif
 }
 
 void FPhysScene_Chaos::DispatchPendingCollisionNotifies()
diff --git a/Engine/Source/Runtime/Engine/Public/Physics/Experimental/PhysScene_Chaos.h b/Engine/Source/Runtime/Engine/Public/Physics/Experimental/PhysScene_Chaos.h
--- a/Engine/Source/Runtime/Engine/Public/Physics/Experimental/PhysScene_Chaos.h
+++ b/Engine/Source/Runtime/Engine/Public/Physics/Experimental/PhysScene_Chaos.h
@@ -12,6 +12,7 @@
 #include "Physics/PhysicsInterfaceUtils.h"
 #include "Chaos/ChaosScene.h"
 #include "Chaos/ContactModification.h"
+#include "Chaos/Real.h"
 
 #ifndef CHAOS_WITH_PAUSABLE_SOLVER
 #define CHAOS_WITH_PAUSABLE_SOLVER 1
@@ -345,6 +346,10 @@ private:
 	/** The SolverActor that spawned and owns this scene */
 	TWeakObjectPtr<AActor> SolverActor;
 
+#if WITH_CHAOS
+	Chaos::FReal LastEventDispatchTime;
+#endif 
+
 #if WITH_EDITOR
 	// Counter used to check a match with the single step status.
 	int32 SingleStepCounter;
diff --git a/Engine/Source/Runtime/Experimental/Chaos/Private/EventDefaults.cpp b/Engine/Source/Runtime/Experimental/Chaos/Private/EventDefaults.cpp
--- a/Engine/Source/Runtime/Experimental/Chaos/Private/EventDefaults.cpp
+++ b/Engine/Source/Runtime/Experimental/Chaos/Private/EventDefaults.cpp
@@ -69,7 +69,7 @@ namespace Chaos
 	void FEventDefaults::RegisterCollisionEvent(FEventManager& EventManager)
 	{
 		EventManager.template RegisterEvent<FCollisionEventData>(EEventType::Collision, []
-		(const Chaos::FPBDRigidsSolver* Solver, FCollisionEventData& CollisionEventData)
+		(const Chaos::FPBDRigidsSolver* Solver, FCollisionEventData& CollisionEventData, bool bResetData)
 		{
 			check(Solver);
 
@@ -84,15 +84,13 @@ namespace Chaos
 			FCollisionDataArray& AllCollisionsDataArray = CollisionEventData.CollisionData.AllCollisionsArray;
 			TMap<IPhysicsProxyBase*, TArray<int32>>& AllCollisionsIndicesByPhysicsProxy = CollisionEventData.PhysicsProxyToCollisionIndices.PhysicsProxyToIndicesMap;
 
-			
-			if (CollisionEventData.CollisionData.TimeCreated != Solver->MTime)
+			if (bResetData)
 			{
 				AllCollisionsDataArray.Reset();
 				AllCollisionsIndicesByPhysicsProxy.Reset();
-
-				CollisionEventData.CollisionData.TimeCreated = Solver->MTime;
-				CollisionEventData.PhysicsProxyToCollisionIndices.TimeCreated = Solver->MTime;
 			}
+			CollisionEventData.CollisionData.TimeCreated = Solver->MTime;
+			CollisionEventData.PhysicsProxyToCollisionIndices.TimeCreated = Solver->MTime;
 			
 			const auto* Evolution = Solver->GetEvolution();
 
@@ -147,13 +145,16 @@ namespace Chaos
 							continue;
 						}
 
-						if (!Constraint.AccumulatedImpulse.IsZero() && Body0)
+						FKinematicGeometryParticleHandle* Primary = Body0 ? Body0 : Body1;
+						FKinematicGeometryParticleHandle* Secondary = Body0 ? Body1 : Body0;
+
+						if (!Constraint.AccumulatedImpulse.IsZero() && Primary)
 						{
 							if (ensure(!Constraint.GetLocation().ContainsNaN() &&
 								!Constraint.GetNormal().ContainsNaN()) &&
-								!Body0->V().ContainsNaN() &&
-								!Body0->W().ContainsNaN() &&
-								(Body1 == nullptr || ((!Body1->V().ContainsNaN()) && !Body1->W().ContainsNaN())))
+								!Primary->V().ContainsNaN() &&
+								!Primary->W().ContainsNaN() &&
+								(Secondary == nullptr || ((!Secondary->V().ContainsNaN()) && !Secondary->W().ContainsNaN())))
 							{
 								ValidCollisionHandles[NumValidCollisions] = ContactHandle;
 								NumValidCollisions++;
@@ -267,7 +268,7 @@ namespace Chaos
 	void FEventDefaults::RegisterBreakingEvent(FEventManager& EventManager)
 	{
 		EventManager.template RegisterEvent<FBreakingEventData>(EEventType::Breaking, []
-		(const Chaos::FPBDRigidsSolver* Solver, FBreakingEventData& BreakingEventData)
+		(const Chaos::FPBDRigidsSolver* Solver, FBreakingEventData& BreakingEventData, bool bResetData)
 		{
 			check(Solver);
 
@@ -282,12 +283,12 @@ namespace Chaos
 			FBreakingDataArray& AllBreakingDataArray = BreakingEventData.BreakingData.AllBreakingsArray;
 			TMap<IPhysicsProxyBase*, TArray<int32>>& AllBreakingIndicesByPhysicsProxy = BreakingEventData.PhysicsProxyToBreakingIndices.PhysicsProxyToIndicesMap;
 
-			if (BreakingEventData.BreakingData.TimeCreated != Solver->MTime)
+			if (bResetData)
 			{
 				AllBreakingDataArray.Reset();
 				AllBreakingIndicesByPhysicsProxy.Reset();
-				BreakingEventData.BreakingData.TimeCreated = Solver->MTime;
 			}
+			BreakingEventData.BreakingData.TimeCreated = Solver->MTime;
 
 			const auto* Evolution = Solver->GetEvolution();
 			const FPBDRigidParticles& Particles = Evolution->GetParticles().GetDynamicParticles();
@@ -340,7 +341,7 @@ namespace Chaos
 	void FEventDefaults::RegisterTrailingEvent(FEventManager& EventManager)
 	{
 		EventManager.template RegisterEvent<FTrailingEventData>(EEventType::Trailing, []
-		(const Chaos::FPBDRigidsSolver* Solver, FTrailingEventData& TrailingEventData)
+		(const Chaos::FPBDRigidsSolver* Solver, FTrailingEventData& TrailingEventData, bool bResetData)
 		{
 			check(Solver);
 
@@ -359,14 +360,13 @@ namespace Chaos
 			FTrailingDataArray& AllTrailingsDataArray = TrailingEventData.TrailingData.AllTrailingsArray;
 			TMap<IPhysicsProxyBase*, TArray<int32>>& AllTrailingIndicesByPhysicsProxy = TrailingEventData.PhysicsProxyToTrailingIndices.PhysicsProxyToIndicesMap;
 
-			if (TrailingEventData.TrailingData.TimeCreated != Solver->MTime)
+			if (bResetData)
 			{
 				AllTrailingsDataArray.Reset();
 				AllTrailingIndicesByPhysicsProxy.Reset();
-
-				TrailingEventData.TrailingData.TimeCreated = Solver->MTime;
-				TrailingEventData.PhysicsProxyToTrailingIndices.TimeCreated = Solver->MTime;
 			}
+			TrailingEventData.TrailingData.TimeCreated = Solver->MTime;
+			TrailingEventData.PhysicsProxyToTrailingIndices.TimeCreated = Solver->MTime;
 
 			const TArray<TPBDRigidParticleHandle<Chaos::FReal, 3>*>& ActiveParticlesArray = Evolution->GetParticles().GetActiveParticlesArray();
 
@@ -436,7 +436,7 @@ namespace Chaos
 	void FEventDefaults::RegisterSleepingEvent(FEventManager& EventManager)
 	{
 		EventManager.template RegisterEvent<FSleepingEventData>(EEventType::Sleeping, []
-		(const Chaos::FPBDRigidsSolver* Solver, FSleepingEventData& SleepingEventData)
+		(const Chaos::FPBDRigidsSolver* Solver, FSleepingEventData& SleepingEventData, bool bResetData)
 		{
 			check(Solver);
 
@@ -478,7 +478,7 @@ namespace Chaos
 	void FEventDefaults::RegisterRemovalEvent(FEventManager& EventManager)
 	{
 		EventManager.template RegisterEvent<FRemovalEventData>(EEventType::Removal, []
-		(const Chaos::FPBDRigidsSolver* Solver, FRemovalEventData& RemovalEventData)
+		(const Chaos::FPBDRigidsSolver* Solver, FRemovalEventData& RemovalEventData, bool bResetData)
 			{
 				check(Solver);
 				EnsureIsInPhysicsThreadContext();
@@ -486,12 +486,12 @@ namespace Chaos
 				FRemovalDataArray& AllRemovalDataArray = RemovalEventData.RemovalData.AllRemovalArray;
 				TMap<IPhysicsProxyBase*, TArray<int32>>& AllRemovalIndicesByPhysicsProxy = RemovalEventData.PhysicsProxyToRemovalIndices.PhysicsProxyToIndicesMap;
 
-				if (RemovalEventData.RemovalData.TimeCreated != Solver->MTime)
+				if (bResetData)
 				{
 					AllRemovalDataArray.Reset();
 					AllRemovalIndicesByPhysicsProxy.Reset();
-					RemovalEventData.RemovalData.TimeCreated = Solver->MTime;
 				}
+				RemovalEventData.RemovalData.TimeCreated = Solver->MTime;
 
 				const TArray<FRemovalData>& AllRemovalsArray = Solver->GetEvolution()->GetAllRemovals();
 				
@@ -512,4 +512,4 @@ namespace Chaos
 			
 			});
 	}
-}
\ No newline at end of file
+}
diff --git a/Engine/Source/Runtime/Experimental/Chaos/Private/EventManager.cpp b/Engine/Source/Runtime/Experimental/Chaos/Private/EventManager.cpp
--- a/Engine/Source/Runtime/Experimental/Chaos/Private/EventManager.cpp
+++ b/Engine/Source/Runtime/Experimental/Chaos/Private/EventManager.cpp
@@ -37,14 +37,14 @@ namespace Chaos
 		ContainerLock.WriteUnlock();
 	}
 
-	void FEventManager::FillProducerData(const Chaos::FPBDRigidsSolver* Solver)
+	void FEventManager::FillProducerData(const Chaos::FPBDRigidsSolver* Solver, bool bResetData)
 	{
 		ContainerLock.ReadLock();
 		for (FEventContainerBasePtr EventContainer : EventContainers)
 		{
 			if (EventContainer)
 			{
-				EventContainer->InjectProducerData(Solver);
+				EventContainer->InjectProducerData(Solver, bResetData);
 			}
 		}
 		ContainerLock.ReadUnlock();
diff --git a/Engine/Source/Runtime/Experimental/Chaos/Private/PBDRigidsSolver.cpp b/Engine/Source/Runtime/Experimental/Chaos/Private/PBDRigidsSolver.cpp
--- a/Engine/Source/Runtime/Experimental/Chaos/Private/PBDRigidsSolver.cpp
+++ b/Engine/Source/Runtime/Experimental/Chaos/Private/PBDRigidsSolver.cpp
@@ -371,15 +371,19 @@ namespace Chaos
 				SCOPE_CYCLE_COUNTER(STAT_EventDataGathering);
 				{
 					SCOPE_CYCLE_COUNTER(STAT_FillProducerData);
-					MSolver->GetEventManager()->FillProducerData(MSolver);
+					bool ResetData = (MSubStepInfo.Step == 0);
+					MSolver->GetEventManager()->FillProducerData(MSolver, ResetData);
 					MSolver->GetEvolution()->ResetAllRemovals();
 				}
+
+				// flip on last sub-step of frame
+				if (MSubStepInfo.Step == MSubStepInfo.NumSteps - 1)
 				{
 					SCOPE_CYCLE_COUNTER(STAT_FlipBuffersIfRequired);
 					MSolver->GetEventManager()->FlipBuffersIfRequired();
 				}
 			}
-
+			
 			{
 				SCOPE_CYCLE_COUNTER(STAT_EndFrame);
 				MSolver->GetEvolution()->EndFrame(MDeltaTime);
diff --git a/Engine/Source/Runtime/Experimental/Chaos/Public/EventManager.h b/Engine/Source/Runtime/Experimental/Chaos/Public/EventManager.h
--- a/Engine/Source/Runtime/Experimental/Chaos/Public/EventManager.h
+++ b/Engine/Source/Runtime/Experimental/Chaos/Public/EventManager.h
@@ -105,7 +105,7 @@ namespace Chaos
 		/*
 		 * Inject data from the physics solver into the producer side of the buffer
 		 */
-		virtual void InjectProducerData(const FPBDRigidsSolver* Solver) = 0;
+		virtual void InjectProducerData(const FPBDRigidsSolver* Solver, bool bResetData) = 0;
 
 		/**
 		 * Flips the buffer if the buffer type is double or triple
@@ -128,7 +128,7 @@ namespace Chaos
 		/**
 		 * Regular constructor
 		 */
-		TEventContainer(const Chaos::EMultiBufferMode& BufferMode, TFunction<void(const FPBDRigidsSolver* Solver, PayloadType& EventDataInOut)> InFunction)
+		TEventContainer(const Chaos::EMultiBufferMode& BufferMode, TFunction<void(const FPBDRigidsSolver* Solver, PayloadType& EventDataInOut, bool bResetData)> InFunction)
 			: InjectedFunction(InFunction)
 			, EventBuffer(Chaos::FMultiBufferFactory<PayloadType>::CreateBuffer(BufferMode))
 		{
@@ -182,9 +182,9 @@ namespace Chaos
 		/*
 		 * Inject data from the physics solver into the producer side of the buffer
 		 */
-		virtual void InjectProducerData(const FPBDRigidsSolver* Solver)
+		virtual void InjectProducerData(const FPBDRigidsSolver* Solver, bool bResetData)
 		{
-			InjectedFunction(Solver, *EventBuffer->AccessProducerBuffer());
+			InjectedFunction(Solver, *EventBuffer->AccessProducerBuffer(), bResetData);
 		}
 
 		
@@ -228,7 +228,7 @@ namespace Chaos
 		/**
 		 * The function that handles filling the event data buffer
 		 */
-		TFunction<void(const FPBDRigidsSolver* Solver, PayloadType& EventData)> InjectedFunction;
+		TFunction<void(const FPBDRigidsSolver* Solver, PayloadType& EventData, bool bResetData)> InjectedFunction;
 
 		/**
 		 * The data buffer that is filled by the producer and read by the consumer
@@ -279,7 +279,7 @@ namespace Chaos
 		 * Register a new event into the system, providing the function that will fill the producer side of the event buffer
 		 */
 		template<typename PayloadType>
-		void RegisterEvent(const EEventType& EventType, TFunction<void(const Chaos::FPBDRigidsSolver* Solver, PayloadType& EventData)> InFunction)
+		void RegisterEvent(const EEventType& EventType, TFunction<void(const Chaos::FPBDRigidsSolver* Solver, PayloadType& EventData, bool bResetData)> InFunction)
 		{
 			ContainerLock.WriteLock();
 			InternalRegisterInjector(FEventID(EventType), new TEventContainer<PayloadType>(BufferMode, InFunction));
@@ -338,7 +338,7 @@ namespace Chaos
 		/**
 		 * Called by the solver to invoke the functions that fill the producer side of all the event data buffers
 		 */
-		void FillProducerData(const Chaos::FPBDRigidsSolver* Solver);
+		void FillProducerData(const Chaos::FPBDRigidsSolver* Solver, bool bResetData = true);
 
 		/**
 		 * Flips the event data buffer if it is of double or triple buffer type
@@ -380,7 +380,6 @@ namespace Chaos
 		/** Defer handler registration if we are currently dispatching events. */
 		bool bCurrentlyDispatchingEvents;
 		TArray<TPair<FEventID, IEventHandler*>> DeferredHandlers;
-
 	};
 
 }
\ No newline at end of file
