Chaos - remove world-space data from FContactPoint and FManifoldPoint
- Collision detection routines now only provide shape-space contact data

diff --git a/Engine/Source/Programs/HeadlessChaos/Private/GeometryCollection/GeometryCollectionTestCollisionResolution.cpp b/Engine/Source/Programs/HeadlessChaos/Private/GeometryCollection/GeometryCollectionTestCollisionResolution.cpp
--- a/Engine/Source/Programs/HeadlessChaos/Private/GeometryCollection/GeometryCollectionTestCollisionResolution.cpp
+++ b/Engine/Source/Programs/HeadlessChaos/Private/GeometryCollection/GeometryCollectionTestCollisionResolution.cpp
@@ -95,8 +95,7 @@ namespace GeometryCollectionTest
 			EXPECT_TRUE(Collection->DynamicCollection->Simplicials[0]->Size() != 0);
 
 			EXPECT_LT(FMath::Abs(Collection->RestCollection->Transform[0].GetTranslation().Z - (Radius + 10.f)), KINDA_SMALL_NUMBER);
-			// ball settles within 10% of radius (The ball will sink deeper than expected due to contact position averaging within CullDistance)
-			EXPECT_LT(FMath::Abs(Collection->DynamicCollection->Transform[0].GetTranslation().Z - Radius), Radius * 0.1f); 
+			EXPECT_NEAR(Collection->DynamicCollection->Transform[0].GetTranslation().Z, Radius, KINDA_SMALL_NUMBER); 
 		}
 	}
 
diff --git a/Engine/Source/Programs/HeadlessChaos/Private/HeadlessChaosTestCollisions.cpp b/Engine/Source/Programs/HeadlessChaos/Private/HeadlessChaosTestCollisions.cpp
--- a/Engine/Source/Programs/HeadlessChaos/Private/HeadlessChaosTestCollisions.cpp
+++ b/Engine/Source/Programs/HeadlessChaos/Private/HeadlessChaosTestCollisions.cpp
@@ -64,15 +64,15 @@ namespace ChaosTest {
 		EXPECT_EQ(Collisions.NumConstraints(), 1);
 
 		FPBDCollisionConstraint& Constraint = Collisions.GetConstraint(0);
-		if (auto PBDRigid = Constraint.Particle[0]->CastToRigidParticle())
+		if (auto PBDRigid = Constraint.GetParticle0()->CastToRigidParticle())
 		{
 			//Question: non dynamics don't have collision particles, seems wrong if the levelset is dynamic and the static is something like a box
 			PBDRigid->CollisionParticles()->UpdateAccelerationStructures();
 		}
 		Collisions.UpdateLevelsetConstraint(Constraint);
 
-		EXPECT_EQ(Constraint.Particle[0], Box2);
-		EXPECT_EQ(Constraint.Particle[1], Box1);
+		EXPECT_EQ(Constraint.GetParticle0(), Box2);
+		EXPECT_EQ(Constraint.GetParticle1(), Box1);
 		EXPECT_TRUE(Constraint.GetNormal().operator==(FVec3(0, 0, 1)));
 
 		// The contact point is the average of the surface contact points, so it should be half of Phi inside particle0
@@ -118,12 +118,12 @@ namespace ChaosTest {
 		FPBDCollisionConstraint& Constraint = Collisions.GetConstraint(0);
 		Collisions.UpdateLevelsetConstraint(Constraint);
 		
-		EXPECT_EQ(Constraint.Particle[0], Box2);
-		EXPECT_EQ(Constraint.Particle[1], Box1);
+		EXPECT_EQ(Constraint.GetParticle0(), Box2);
+		EXPECT_EQ(Constraint.GetParticle1(), Box1);
 		EXPECT_TRUE(Constraint.GetNormal().operator==(FVec3(0, 0, 1)));
 
 		// The contact point is the average of the surface contact points, so it should be half of Phi inside particle0
-		Chaos::FReal Distance = ChaosTest::SignedDistance(*Constraint.Particle[0], Constraint.GetLocation());
+		Chaos::FReal Distance = ChaosTest::SignedDistance(*Constraint.GetParticle0(), Constraint.GetLocation());
 		EXPECT_NEAR(Distance, 0.5f * Constraint.GetPhi(), KINDA_SMALL_NUMBER);
 	}
 	
@@ -162,10 +162,10 @@ namespace ChaosTest {
 		EXPECT_EQ(Collisions.NumConstraints(), 1);
 
 		FPBDCollisionConstraint& Constraint = Collisions.GetConstraint(0);
-		EXPECT_TRUE(Constraint.Particle[0] != nullptr);
-		EXPECT_TRUE(Constraint.Particle[1] != nullptr);
+		EXPECT_TRUE(Constraint.GetParticle0() != nullptr);
+		EXPECT_TRUE(Constraint.GetParticle1() != nullptr);
 
-		if (auto PBDRigid = Constraint.Particle[0]->CastToRigidParticle())
+		if (auto PBDRigid = Constraint.GetParticle0()->CastToRigidParticle())
 		{
 			PBDRigid->CollisionParticles()->UpdateAccelerationStructures();
 		}
@@ -174,8 +174,8 @@ namespace ChaosTest {
 
 		Collisions.Update(Constraint);
 
-		EXPECT_EQ(Constraint.Particle[0], Box);
-		EXPECT_EQ(Constraint.Particle[1], Floor);
+		EXPECT_EQ(Constraint.GetParticle0(), Box);
+		EXPECT_EQ(Constraint.GetParticle1(), Floor);
 		EXPECT_TRUE(Constraint.GetNormal().operator==(FVec3(0, 0, 1)));
 		EXPECT_TRUE(FMath::Abs(Constraint.GetPhi() - FReal(-0.5) ) < SMALL_THRESHOLD );
 
@@ -257,8 +257,8 @@ namespace ChaosTest {
 
 		Collisions.Update(*Constraint);
 
-		EXPECT_EQ(Constraint->Particle[0], Box);
-		EXPECT_EQ(Constraint->Particle[1], Floor);
+		EXPECT_EQ(Constraint->GetParticle0(), Box);
+		EXPECT_EQ(Constraint->GetParticle1(), Floor);
 		EXPECT_TRUE(Constraint->GetNormal().operator==(FVec3(0, 0, 1)));
 		EXPECT_TRUE(FMath::Abs(Constraint->GetPhi() - FReal(-1.0)) < SMALL_THRESHOLD);
 
@@ -333,8 +333,8 @@ namespace ChaosTest {
 		Collisions.GatherInput(Dt);
 
 		FPBDCollisionConstraint& Constraint = Collisions.GetConstraint(0);
-		EXPECT_EQ(Constraint.Particle[0], Box);
-		EXPECT_EQ(Constraint.Particle[1], Floor);
+		EXPECT_EQ(Constraint.GetParticle0(), Box);
+		EXPECT_EQ(Constraint.GetParticle1(), Floor);
 		EXPECT_TRUE(Constraint.GetNormal().operator==(FVec3(0, 0, 1)));
 		EXPECT_TRUE(FMath::Abs(Constraint.GetPhi() - FReal(-0.5)) < SMALL_THRESHOLD);
 
@@ -403,7 +403,7 @@ namespace ChaosTest {
 		EXPECT_EQ(Collisions.NumConstraints(), 1);
 
 		FPBDCollisionConstraint& Constraint = Collisions.GetConstraint(0);
-		if (auto PBDRigid = Constraint.Particle[0]->CastToRigidParticle())
+		if (auto PBDRigid = Constraint.GetParticle0()->CastToRigidParticle())
 		{
 			PBDRigid->CollisionParticles()->UpdateAccelerationStructures();
 		}
@@ -411,8 +411,8 @@ namespace ChaosTest {
 		Collisions.GatherInput(Dt);
 
 		Collisions.UpdateLevelsetConstraint(Constraint);
-		EXPECT_EQ(Constraint.Particle[0], Box);
-		EXPECT_EQ(Constraint.Particle[1], Floor);
+		EXPECT_EQ(Constraint.GetParticle0(), Box);
+		EXPECT_EQ(Constraint.GetParticle1(), Floor);
 		EXPECT_TRUE(Constraint.GetNormal().operator==(FVec3(0, 0, 1)));
 		EXPECT_TRUE(FMath::Abs(Constraint.GetPhi() - FReal(-0.5)) < SMALL_THRESHOLD);
 
@@ -491,15 +491,15 @@ namespace ChaosTest {
 		}
 
 		FPBDCollisionConstraint& Constraint = Collisions.GetConstraint(0);
-		if (FPBDRigidParticleHandle* PBDRigid = Constraint.Particle[0]->CastToRigidParticle())
+		if (FPBDRigidParticleHandle* PBDRigid = Constraint.GetParticle0()->CastToRigidParticle())
 		{
 			PBDRigid->CollisionParticles()->UpdateAccelerationStructures();
 		}
 
 		Collisions.GatherInput(Dt);
 
-		EXPECT_EQ(Constraint.Particle[0], DynamicCube);
-		EXPECT_EQ(Constraint.Particle[1], StaticCube);
+		EXPECT_EQ(Constraint.GetParticle0(), DynamicCube);
+		EXPECT_EQ(Constraint.GetParticle1(), StaticCube);
 		EXPECT_NEAR(Constraint.GetNormal().Z, FReal(1), KINDA_SMALL_NUMBER);
 		{
 			INVARIANT_XR_START(DynamicCube);
@@ -560,13 +560,13 @@ namespace ChaosTest {
 
 		FPBDCollisionConstraint& Constraint = Collisions.GetConstraint(0);
 
-		if (auto PBDRigid = Constraint.Particle[0]->CastToRigidParticle())
+		if (auto PBDRigid = Constraint.GetParticle0()->CastToRigidParticle())
 		{
 			PBDRigid->CollisionParticles()->UpdateAccelerationStructures();
 		}
 
-		EXPECT_EQ(Constraint.Particle[0], Box2);
-		EXPECT_EQ(Constraint.Particle[1], StaticBox);
+		EXPECT_EQ(Constraint.GetParticle0(), Box2);
+		EXPECT_EQ(Constraint.GetParticle1(), StaticBox);
 		EXPECT_TRUE(Constraint.GetNormal().Equals(FVector(0.0, 0.0, 1.0f)));
 		EXPECT_TRUE(FMath::Abs(Constraint.GetPhi() - FReal(-0.4)) < SMALL_THRESHOLD);
 
diff --git a/Engine/Source/Programs/HeadlessChaos/Private/HeadlessChaosTestDetectCollision.cpp b/Engine/Source/Programs/HeadlessChaos/Private/HeadlessChaosTestDetectCollision.cpp
--- a/Engine/Source/Programs/HeadlessChaos/Private/HeadlessChaosTestDetectCollision.cpp
+++ b/Engine/Source/Programs/HeadlessChaos/Private/HeadlessChaosTestDetectCollision.cpp
@@ -34,9 +34,6 @@ namespace ChaosTest
 		EXPECT_NEAR(ContactPoint.ShapeContactNormal.X, ShapeContactNormal.X, NormalTolerance);
 		EXPECT_NEAR(ContactPoint.ShapeContactNormal.Y, ShapeContactNormal.Y, NormalTolerance);
 		EXPECT_NEAR(ContactPoint.ShapeContactNormal.Z, ShapeContactNormal.Z, NormalTolerance);
-		EXPECT_NEAR(ContactPoint.Normal.X, Normal.X, NormalTolerance);
-		EXPECT_NEAR(ContactPoint.Normal.Y, Normal.Y, NormalTolerance);
-		EXPECT_NEAR(ContactPoint.Normal.Z, Normal.Z, NormalTolerance);
 		EXPECT_NEAR(ContactPoint.Phi, Phi, DistanceTolerance);
 	}
 
diff --git a/Engine/Source/Programs/HeadlessChaos/Private/HeadlessChaosTestOneShotManifolds.cpp b/Engine/Source/Programs/HeadlessChaos/Private/HeadlessChaosTestOneShotManifolds.cpp
--- a/Engine/Source/Programs/HeadlessChaos/Private/HeadlessChaosTestOneShotManifolds.cpp
+++ b/Engine/Source/Programs/HeadlessChaos/Private/HeadlessChaosTestOneShotManifolds.cpp
@@ -60,9 +60,6 @@ namespace ChaosTest
 			for (int ConstraintIndex = 0; ConstraintIndex < ContactCount; ConstraintIndex++)
 			{
 				EXPECT_NEAR(-200.0f, Constraint.GetManifoldPoints()[ConstraintIndex].ContactPoint.Phi, 0.01);
-				EXPECT_NEAR(FMath::Abs(Constraint.GetManifoldPoints()[ConstraintIndex].ContactPoint.Location.X), 100.0f, 0.01);
-				EXPECT_NEAR(FMath::Abs(Constraint.GetManifoldPoints()[ConstraintIndex].ContactPoint.Location.Y), 100.0f, 0.01);
-				EXPECT_NEAR(FMath::Abs(Constraint.GetManifoldPoints()[ConstraintIndex].ContactPoint.Location.Z), 100.0f, 0.01);
 			}
 		}
 
@@ -80,9 +77,12 @@ namespace ChaosTest
 			for (int ConstraintIndex = 0; ConstraintIndex < ContactCount; ConstraintIndex++)
 			{
 				EXPECT_NEAR(10.0f, Constraint.GetManifoldPoints()[ConstraintIndex].ContactPoint.Phi, 0.01);
-				EXPECT_NEAR(FMath::Abs(Constraint.GetManifoldPoints()[ConstraintIndex].ContactPoint.Location.X),100.0f ,0.01);
-				EXPECT_NEAR(FMath::Abs(Constraint.GetManifoldPoints()[ConstraintIndex].ContactPoint.Location.Y), 100.0f, 0.01);
-				EXPECT_NEAR(Constraint.GetManifoldPoints()[ConstraintIndex].ContactPoint.Location.Z, 100.0f, 0.01);
+				EXPECT_NEAR(FMath::Abs(Constraint.GetManifoldPoints()[ConstraintIndex].ContactPoint.ShapeContactPoints[0].X), 100.0f, 0.01);
+				EXPECT_NEAR(FMath::Abs(Constraint.GetManifoldPoints()[ConstraintIndex].ContactPoint.ShapeContactPoints[0].Y), 100.0f, 0.01);
+				EXPECT_NEAR(Constraint.GetManifoldPoints()[ConstraintIndex].ContactPoint.ShapeContactPoints[0].Z, -100.0f, 0.01);
+				EXPECT_NEAR(FMath::Abs(Constraint.GetManifoldPoints()[ConstraintIndex].ContactPoint.ShapeContactPoints[1].X), 100.0f, 0.01);
+				EXPECT_NEAR(FMath::Abs(Constraint.GetManifoldPoints()[ConstraintIndex].ContactPoint.ShapeContactPoints[1].Y), 100.0f, 0.01);
+				EXPECT_NEAR(Constraint.GetManifoldPoints()[ConstraintIndex].ContactPoint.ShapeContactPoints[1].Z, 100.0f, 0.01);
 			}
 		}
 
@@ -100,15 +100,6 @@ namespace ChaosTest
 			for (int ConstraintIndex = 0; ConstraintIndex < ContactCount; ConstraintIndex++)
 			{
 				EXPECT_NEAR(10.0f, Constraint.GetManifoldPoints()[ConstraintIndex].ContactPoint.Phi, 15.0f); // The cube is at an angle now
-				if (bChaos_Collision_Manifold_FixNormalsInWorldSpace)
-				{
-					EXPECT_NEAR(Constraint.GetManifoldPoints()[ConstraintIndex].ContactPoint.Location.Z, 100.0f, 15.0f);
-				}
-				else
-				{
-					EXPECT_NEAR(Constraint.GetManifoldPoints()[ConstraintIndex].ContactPoint.Location.Z, 110.0f, 15.0f);
-				}
-				
 			}
 		}
 
@@ -130,9 +121,12 @@ namespace ChaosTest
 			for (int ConstraintIndex = 0; ConstraintIndex < ContactCount; ConstraintIndex++)
 			{
 				EXPECT_NEAR(10.0f, Constraint.GetManifoldPoints()[ConstraintIndex].ContactPoint.Phi, 0.01);
-				EXPECT_NEAR(FMath::Abs(Constraint.GetManifoldPoints()[ConstraintIndex].ContactPoint.Location.X), 100.0f, 0.01);
-				EXPECT_NEAR(FMath::Abs(Constraint.GetManifoldPoints()[ConstraintIndex].ContactPoint.Location.Y), 100.0f, 0.01);
-				EXPECT_NEAR(Constraint.GetManifoldPoints()[ConstraintIndex].ContactPoint.Location.Z, 100.0f, 0.01);
+				EXPECT_NEAR(FMath::Abs(Constraint.GetManifoldPoints()[ConstraintIndex].ContactPoint.ShapeContactPoints[0].X - OffsetBox1.X), 100.0f, 0.01);
+				EXPECT_NEAR(FMath::Abs(Constraint.GetManifoldPoints()[ConstraintIndex].ContactPoint.ShapeContactPoints[0].Y - OffsetBox1.Y), 100.0f, 0.01);
+				EXPECT_NEAR(Constraint.GetManifoldPoints()[ConstraintIndex].ContactPoint.ShapeContactPoints[0].Z - OffsetBox1.Z, -100.0f, 0.01);
+				EXPECT_NEAR(FMath::Abs(Constraint.GetManifoldPoints()[ConstraintIndex].ContactPoint.ShapeContactPoints[1].X - OffsetBox2.X), 100.0f, 0.01);
+				EXPECT_NEAR(FMath::Abs(Constraint.GetManifoldPoints()[ConstraintIndex].ContactPoint.ShapeContactPoints[1].Y - OffsetBox2.Y), 100.0f, 0.01);
+				EXPECT_NEAR(Constraint.GetManifoldPoints()[ConstraintIndex].ContactPoint.ShapeContactPoints[1].Z - OffsetBox2.Z, 100.0f, 0.01);
 			}
 		}
 
@@ -165,10 +159,12 @@ namespace ChaosTest
 			for (int ConstraintIndex = 0; ConstraintIndex < ContactCount; ConstraintIndex++)
 			{
 				EXPECT_NEAR(10.0f, Constraint.GetManifoldPoints()[ConstraintIndex].ContactPoint.Phi, 0.01);
-				FVec3 Location = Box2Transform.InverseTransformPosition( Constraint.GetManifoldPoints()[ConstraintIndex].ContactPoint.Location) - OffsetBox2;
-				EXPECT_NEAR(FMath::Abs(Location.X), 100.0f, 0.01);
-				EXPECT_NEAR(FMath::Abs(Location.Y), 100.0f, 0.01);
-				EXPECT_NEAR(Location.Z, 100.0f, 0.01);
+				EXPECT_NEAR(FMath::Abs(Constraint.GetManifoldPoints()[ConstraintIndex].ContactPoint.ShapeContactPoints[0].X - OffsetBox1.X), 100.0f, 0.01);
+				EXPECT_NEAR(FMath::Abs(Constraint.GetManifoldPoints()[ConstraintIndex].ContactPoint.ShapeContactPoints[0].Y - OffsetBox1.Y), 100.0f, 0.01);
+				EXPECT_NEAR(Constraint.GetManifoldPoints()[ConstraintIndex].ContactPoint.ShapeContactPoints[0].Z - OffsetBox1.Z, -100.0f, 0.01);
+				EXPECT_NEAR(FMath::Abs(Constraint.GetManifoldPoints()[ConstraintIndex].ContactPoint.ShapeContactPoints[1].X - OffsetBox2.X), 100.0f, 0.01);
+				EXPECT_NEAR(FMath::Abs(Constraint.GetManifoldPoints()[ConstraintIndex].ContactPoint.ShapeContactPoints[1].Y - OffsetBox2.Y), 100.0f, 0.01);
+				EXPECT_NEAR(Constraint.GetManifoldPoints()[ConstraintIndex].ContactPoint.ShapeContactPoints[1].Z - OffsetBox2.Z, 100.0f, 0.01);
 			}
 		}
 
@@ -197,8 +193,10 @@ namespace ChaosTest
 			for (int ConstraintIndex = 0; ConstraintIndex < ContactCount; ConstraintIndex++)
 			{
 				EXPECT_NEAR(10.0f, Constraint.GetManifoldPoints()[ConstraintIndex].ContactPoint.Phi, 0.01);
-				FVec3 Location = Constraint.GetManifoldPoints()[ConstraintIndex].ContactPoint.Location;
-				EXPECT_NEAR(Location.Z, 100.0f, 0.01);
+				FVec3 Location1 = Box1Transform.TransformPosition(Constraint.GetManifoldPoints()[ConstraintIndex].ContactPoint.ShapeContactPoints[0]);
+				FVec3 Location2 = Box2Transform.TransformPosition(Constraint.GetManifoldPoints()[ConstraintIndex].ContactPoint.ShapeContactPoints[1]);
+				EXPECT_NEAR(Location1.Z, 110.0f, 0.01);
+				EXPECT_NEAR(Location2.Z, 100.0f, 0.01);
 			}
 		}
 
@@ -216,7 +214,7 @@ namespace ChaosTest
 			for (int ConstraintIndex = 0; ConstraintIndex < ContactCount; ConstraintIndex++)
 			{
 				EXPECT_NEAR(10.0f, Constraint.GetManifoldPoints()[ConstraintIndex].ContactPoint.Phi, 0.01);
-				EXPECT_NEAR(Constraint.GetManifoldPoints()[ConstraintIndex].ContactPoint.Location.Z, 100.0f, 0.01);
+				EXPECT_NEAR(Box2Transform.TransformPosition(Constraint.GetManifoldPoints()[ConstraintIndex].ContactPoint.ShapeContactPoints[1]).Z, 100.0f, 0.01);
 			}
 		}
 
@@ -258,7 +256,7 @@ namespace ChaosTest
 		if (ContactCount > 0)
 		{
 			EXPECT_NEAR(Constraint.GetManifoldPoints()[0].ContactPoint.Phi, 0.0f, 0.01f);
-			EXPECT_NEAR(Constraint.GetManifoldPoints()[0].ContactPoint.Location.X, 0.5f * Offset, 0.01f);
+			EXPECT_NEAR(Convex1Transform.TransformPosition(Constraint.GetManifoldPoints()[0].ContactPoint.ShapeContactPoints[0]).X, 0.5f * Offset, 0.01f);
 		}
 	}
 
diff --git a/Engine/Source/Programs/HeadlessChaos/Private/HeadlessChaosTestShapePadding.cpp b/Engine/Source/Programs/HeadlessChaos/Private/HeadlessChaosTestShapePadding.cpp
--- a/Engine/Source/Programs/HeadlessChaos/Private/HeadlessChaosTestShapePadding.cpp
+++ b/Engine/Source/Programs/HeadlessChaos/Private/HeadlessChaosTestShapePadding.cpp
@@ -104,10 +104,10 @@ namespace ChaosTest {
 		Constraint->ResetPhi(Constraint->GetCullDistance());
 		Collisions::UpdateConstraintFromGeometry<ECollisionUpdateType::Deepest>(*Constraint, FParticleUtilitiesPQ::GetActorWorldTransform(Box0), FParticleUtilitiesPQ::GetActorWorldTransform(Box1), 1 / 30.0f);
 
-		EXPECT_NEAR(Constraint->Manifold.Phi, ExpectedPhi, Tolerance);
-		EXPECT_NEAR(Constraint->Manifold.Normal.X, ExpectedNormal.X, Tolerance);
-		EXPECT_NEAR(Constraint->Manifold.Normal.Y, ExpectedNormal.Y, Tolerance);
-		EXPECT_NEAR(Constraint->Manifold.Normal.Z, ExpectedNormal.Z, Tolerance);
+		EXPECT_NEAR(Constraint->GetPhi(), ExpectedPhi, Tolerance);
+		EXPECT_NEAR(Constraint->GetNormal().X, ExpectedNormal.X, Tolerance);
+		EXPECT_NEAR(Constraint->GetNormal().Y, ExpectedNormal.Y, Tolerance);
+		EXPECT_NEAR(Constraint->GetNormal().Z, ExpectedNormal.Z, Tolerance);
 	}
 
 	TEST(CollisionTests, TestBoxBoxCollisionMargin)
@@ -213,10 +213,10 @@ namespace ChaosTest {
 		// Detect collisions
 		Collisions::UpdateConstraintFromGeometry<ECollisionUpdateType::Deepest>(*Constraint, FParticleUtilitiesPQ::GetActorWorldTransform(Box0), FParticleUtilitiesPQ::GetActorWorldTransform(Box1), 1 / 30.0f);
 
-		EXPECT_NEAR(Constraint->Manifold.Phi, ExpectedPhi, Tolerance);
-		EXPECT_NEAR(Constraint->Manifold.Normal.X, ExpectedNormal.X, Tolerance);
-		EXPECT_NEAR(Constraint->Manifold.Normal.Y, ExpectedNormal.Y, Tolerance);
-		EXPECT_NEAR(Constraint->Manifold.Normal.Z, ExpectedNormal.Z, Tolerance);
+		EXPECT_NEAR(Constraint->GetPhi(), ExpectedPhi, Tolerance);
+		EXPECT_NEAR(Constraint->GetNormal().X, ExpectedNormal.X, Tolerance);
+		EXPECT_NEAR(Constraint->GetNormal().Y, ExpectedNormal.Y, Tolerance);
+		EXPECT_NEAR(Constraint->GetNormal().Z, ExpectedNormal.Z, Tolerance);
 	}
 
 
diff --git a/Engine/Source/Programs/HeadlessChaos/Public/HeadlessChaosCollisionConstraints.h b/Engine/Source/Programs/HeadlessChaos/Public/HeadlessChaosCollisionConstraints.h
--- a/Engine/Source/Programs/HeadlessChaos/Public/HeadlessChaosCollisionConstraints.h
+++ b/Engine/Source/Programs/HeadlessChaos/Public/HeadlessChaosCollisionConstraints.h
@@ -84,8 +84,8 @@ public:
 			Constraint.ResetPhi(TNumericLimits<FReal>::Max());
 			Collisions::UpdateConstraintFromGeometry<ECollisionUpdateType::Deepest>(
 				Constraint, 
-				FParticleUtilities::GetActorWorldTransform(FGenericParticleHandle(Constraint.Particle[0])), 
-				FParticleUtilities::GetActorWorldTransform(FGenericParticleHandle(Constraint.Particle[1])), 
+				FParticleUtilities::GetActorWorldTransform(FGenericParticleHandle(Constraint.GetParticle0())), 
+				FParticleUtilities::GetActorWorldTransform(FGenericParticleHandle(Constraint.GetParticle1())), 
 				Dt);
 		}
 	}
@@ -96,8 +96,8 @@ public:
 		// Dt is not important for the tests that use this function
 		const FReal Dt = FReal(1) / FReal(30);
 
-		FRigidTransform3 WorldTransform0 = Constraint.ImplicitTransform[0] * Collisions::GetTransform(Constraint.Particle[0]);
-		FRigidTransform3 WorldTransform1 = Constraint.ImplicitTransform[1] * Collisions::GetTransform(Constraint.Particle[1]);
+		FRigidTransform3 WorldTransform0 = Constraint.GetShapeRelativeTransform0() * Collisions::GetTransform(Constraint.GetParticle0());
+		FRigidTransform3 WorldTransform1 = Constraint.GetShapeRelativeTransform1() * Collisions::GetTransform(Constraint.GetParticle1());
 
 		Constraint.ResetManifold();
 		Collisions::UpdateLevelsetLevelsetConstraint<ECollisionUpdateType::Deepest>(WorldTransform0, WorldTransform1, FReal(1 / 30.0f), Constraint);
diff --git a/Engine/Source/Runtime/Engine/Private/PhysicsEngine/ImmediatePhysics/ImmediatePhysicsChaos/ImmediatePhysicsSimulation_Chaos.cpp b/Engine/Source/Runtime/Engine/Private/PhysicsEngine/ImmediatePhysics/ImmediatePhysicsChaos/ImmediatePhysicsSimulation_Chaos.cpp
--- a/Engine/Source/Runtime/Engine/Private/PhysicsEngine/ImmediatePhysics/ImmediatePhysicsChaos/ImmediatePhysicsSimulation_Chaos.cpp
+++ b/Engine/Source/Runtime/Engine/Private/PhysicsEngine/ImmediatePhysics/ImmediatePhysicsChaos/ImmediatePhysicsSimulation_Chaos.cpp
@@ -946,7 +946,7 @@ namespace ImmediatePhysics_Chaos
 	{
 #if CHAOS_DEBUG_DRAW
 		using namespace Chaos;
-		if (FDebugDrawQueue::IsDebugDrawingEnabled() && bChaosImmediate_DebugDrawJoints)
+		if (FDebugDrawQueue::IsDebugDrawingEnabled())
 		{
 			if (bChaosImmediate_DebugDrawCollisions)
 			{
diff --git a/Engine/Source/Runtime/Experimental/Chaos/Private/Chaos/CCDUtilities.cpp b/Engine/Source/Runtime/Experimental/Chaos/Private/Chaos/CCDUtilities.cpp
--- a/Engine/Source/Runtime/Experimental/Chaos/Private/Chaos/CCDUtilities.cpp
+++ b/Engine/Source/Runtime/Experimental/Chaos/Private/Chaos/CCDUtilities.cpp
@@ -42,13 +42,13 @@ namespace Chaos
 	{
 		for (int32 i = 0; i < 2; i++)
 		{
-			const TPBDRigidParticleHandle<FReal, 3>* Rigid = Constraint->Particle[i]->CastToRigidParticle();
+			const TPBDRigidParticleHandle<FReal, 3>* Rigid = Constraint->GetParticle(i)->CastToRigidParticle();
 			if (Rigid && Rigid->ObjectState() == EObjectStateType::Kinematic)
 			{
 				// The same computation is carried out in UseCCDImpl function when constructing constraints. But we don't have access to FCCDConstraint at that point. This part could potentially be optimized away. 
 				const FVec3 D = Displacements[i];
 				const FReal DSizeSquared = D.SizeSquared();
-				const FReal CCDThreshold = GetParticleCCDThreshold(Constraint->Manifold.Implicit[i]);
+				const FReal CCDThreshold = GetParticleCCDThreshold(Constraint->GetImplicit(i));
 				if (DSizeSquared > CCDThreshold * CCDThreshold)
 				{
 					return i;
@@ -103,7 +103,7 @@ namespace Chaos
 			FVec3 Displacements[2] = {FVec3(0.f), FVec3(0.f)};
 			for (int i = 0; i < 2; i++)
 			{
-				TPBDRigidParticleHandle<FReal, 3>* RigidParticle = Constraint->Particle[i]->CastToRigidParticle();
+				TPBDRigidParticleHandle<FReal, 3>* RigidParticle = Constraint->GetParticle(i)->CastToRigidParticle();
 				FCCDParticle* CCDParticle = nullptr;
 				const bool IsParticleDynamic = RigidParticle && RigidParticle->ObjectState() == EObjectStateType::Dynamic;
 				if (IsParticleDynamic)
@@ -132,8 +132,8 @@ namespace Chaos
 			}
 
 			// Compute the relative displacement. If relative displacement is smaller than 0.5 * (Extents0.Min() + Extents1.Min()), it is impossible for an particle to tunnel through another particle, even though the absolute velocities of the particles might be large.
-			const FReal CCDThreshold0 = GetParticleCCDThreshold(Constraint->Manifold.Implicit[0]);
-			const FReal CCDThreshold1 = GetParticleCCDThreshold(Constraint->Manifold.Implicit[1]);
+			const FReal CCDThreshold0 = GetParticleCCDThreshold(Constraint->GetImplicit0());
+			const FReal CCDThreshold1 = GetParticleCCDThreshold(Constraint->GetImplicit1());
 			const FReal CCDConstraintThreshold = CCDThreshold0 + CCDThreshold1;
 			if ((Displacements[1] - Displacements[0]).SizeSquared() > CCDConstraintThreshold * CCDConstraintThreshold)
 			{
@@ -324,8 +324,8 @@ namespace Chaos
 					{
 						if (CCDConstraint->FastMovingKinematicIndex != INDEX_NONE)
 						{
-							const FConstGenericParticleHandle Particle1 = FGenericParticleHandle(CCDConstraint->SweptConstraint->Particle[CCDConstraint->FastMovingKinematicIndex]);
-							const FVec3 Normal = CCDConstraint->SweptConstraint->GetNormal();
+							const FConstGenericParticleHandle Particle1 = FGenericParticleHandle(CCDConstraint->SweptConstraint->GetParticle(CCDConstraint->FastMovingKinematicIndex));
+							const FVec3 Normal = CCDConstraint->SweptConstraint->CalculateWorldContactNormal();
 							const FVec3 Offset = FVec3::DotProduct(Particle1->V() * ((1.f - IslandTOI) * Dt), Normal) * Normal;
 							ClipParticleP(CCDConstraintParticles[0], Offset);
 						}
@@ -401,7 +401,7 @@ namespace Chaos
 								}
 								else
 								{
-									FGenericParticleHandle AffectedParticle = FGenericParticleHandle(AttachedCCDConstraint->SweptConstraint->Particle[j]);
+									FGenericParticleHandle AffectedParticle = FGenericParticleHandle(AttachedCCDConstraint->SweptConstraint->GetParticle(j));
 									const bool IsKinematic = AffectedParticle->ObjectState() == EObjectStateType::Kinematic;
 									if (IsKinematic)
 									{
@@ -515,13 +515,14 @@ namespace Chaos
 	void FCCDManager::ApplyImpulse(FCCDConstraint *CCDConstraint)
 	{
 		FPBDCollisionConstraint *Constraint = CCDConstraint->SweptConstraint;
-		TPBDRigidParticleHandle<FReal, 3> *Rigid0 = Constraint->Particle[0]->CastToRigidParticle();
-		TPBDRigidParticleHandle<FReal, 3> *Rigid1 = Constraint->Particle[1]->CastToRigidParticle();
+		TPBDRigidParticleHandle<FReal, 3> *Rigid0 = Constraint->GetParticle0()->CastToRigidParticle();
+		TPBDRigidParticleHandle<FReal, 3> *Rigid1 = Constraint->GetParticle1()->CastToRigidParticle();
 		check(Rigid0 != nullptr || Rigid1 != nullptr);
-		const FReal Restitution = Constraint->Manifold.Restitution;
+		const FReal Restitution = Constraint->GetRestitution();
+		const FRigidTransform3& ShapeWorldTransform1 = Constraint->GetShapeWorldTransform1();
 		for(const FManifoldPoint &ManifoldPoint : Constraint->GetManifoldPoints())
 		{
-			const FVec3 Normal = ManifoldPoint.ContactPoint.Normal;
+			const FVec3 Normal = ShapeWorldTransform1.TransformVectorNoScale(ManifoldPoint.ContactPoint.ShapeContactNormal);
 			const FVec3 V0 = Rigid0 != nullptr ? Rigid0->V() : FVec3(0.f);
 			const FVec3 V1 = Rigid1 != nullptr ? Rigid1->V() : FVec3(0.f);
 			const FReal NormalV = FVec3::DotProduct(V0 - V1, Normal);
@@ -553,7 +554,7 @@ namespace Chaos
 			FRigidTransform3 RigidTransforms[2];
 			for (int32 i = 0; i < 2; i++)
 			{
-				FGenericParticleHandle Particle = FGenericParticleHandle(SweptConstraint->Particle[i]);
+				FGenericParticleHandle Particle = FGenericParticleHandle(SweptConstraint->GetParticle(i));
 				const bool IsStatic = Particle->ObjectState() == EObjectStateType::Static;
 				if (IsStatic)
 				{
diff --git a/Engine/Source/Runtime/Experimental/Chaos/Private/Chaos/ChaosDebugDraw.cpp b/Engine/Source/Runtime/Experimental/Chaos/Private/Chaos/ChaosDebugDraw.cpp
--- a/Engine/Source/Runtime/Experimental/Chaos/Private/Chaos/ChaosDebugDraw.cpp
+++ b/Engine/Source/Runtime/Experimental/Chaos/Private/Chaos/ChaosDebugDraw.cpp
@@ -831,14 +831,14 @@ namespace Chaos
 			{
 				if (Contact.GetUseManifold())
 				{
-					const FConstGenericParticleHandle Particle0 = Contact.Particle[0];
-					const FConstGenericParticleHandle Particle1 = Contact.Particle[1];
+					const FConstGenericParticleHandle Particle0 = Contact.GetParticle0();
+					const FConstGenericParticleHandle Particle1 = Contact.GetParticle1();
 					const FRigidTransform3 WorldActorTransform0 = FParticleUtilities::GetActorWorldTransform(Particle0);
 					const FRigidTransform3 WorldActorTransform1 = FParticleUtilities::GetActorWorldTransform(Particle1);
 
 					// Are we within the region of interest?
-					const FReal Particle0Size = Contact.Particle[0]->HasBounds() ? 0.5f * Contact.Particle[0]->LocalBounds().Extents().Size() : TNumericLimits<FReal>::Max();
-					const FReal Particle1Size = Contact.Particle[1]->HasBounds() ? 0.5f * Contact.Particle[1]->LocalBounds().Extents().Size() : TNumericLimits<FReal>::Max();
+					const FReal Particle0Size = Contact.GetParticle0()->HasBounds() ? 0.5f * Contact.GetParticle0()->LocalBounds().Extents().Size() : TNumericLimits<FReal>::Max();
+					const FReal Particle1Size = Contact.GetParticle1()->HasBounds() ? 0.5f * Contact.GetParticle1()->LocalBounds().Extents().Size() : TNumericLimits<FReal>::Max();
 					if (!FDebugDrawQueue::GetInstance().IsInRegionOfInterest(WorldActorTransform0.GetLocation(), Particle0Size))
 					{
 						return;
@@ -858,10 +858,10 @@ namespace Chaos
 
 						const int32 ContactPlaneOwner = 1;
 						const int32 ContactPointOwner = 1 - ContactPlaneOwner;
-						const FRigidTransform3& PlaneTransform = (ContactPlaneOwner == 0) ? Contact.ImplicitTransform[0] * WorldActorTransform0 : Contact.ImplicitTransform[1] * WorldActorTransform1;
-						const FRigidTransform3& PointTransform = (ContactPlaneOwner == 0) ? Contact.ImplicitTransform[1] * WorldActorTransform1 : Contact.ImplicitTransform[0] * WorldActorTransform0;
+						const FRigidTransform3& PlaneTransform = (ContactPlaneOwner == 0) ? Contact.GetShapeRelativeTransform0() * WorldActorTransform0 : Contact.GetShapeRelativeTransform1() * WorldActorTransform1;
+						const FRigidTransform3& PointTransform = (ContactPlaneOwner == 0) ? Contact.GetShapeRelativeTransform1() * WorldActorTransform1 : Contact.GetShapeRelativeTransform0() * WorldActorTransform0;
 						const FConstGenericParticleHandle PlaneParticle = (ContactPlaneOwner == 0) ? Particle0 : Particle1;
-						const FVec3 PlaneNormal = ManifoldPoint.ContactPoint.Normal;
+						const FVec3 PlaneNormal = PlaneTransform.TransformVectorNoScale(ManifoldPoint.ContactPoint.ShapeContactNormal);
 						const FVec3 PointLocation = PointTransform.TransformPosition(ManifoldPoint.ContactPoint.ShapeContactPoints[ContactPointOwner]);
 						const FVec3 PlaneLocation = PlaneTransform.TransformPosition(ManifoldPoint.ContactPoint.ShapeContactPoints[ContactPlaneOwner]);
 						const FVec3 PointPlaneLocation = PointLocation - FVec3::DotProduct(PointLocation - PlaneLocation, PlaneNormal) * PlaneNormal;
@@ -909,8 +909,8 @@ namespace Chaos
 						FDebugDrawQueue::GetInstance().DrawDebugCircle(WorldPointLocation, 0.5f * Settings.DrawScale * Settings.ContactWidth, 12, DiscColor, false, KINDA_SMALL_NUMBER, Settings.DrawPriority, Settings.LineThickness, Axes.GetUnitAxis(EAxis::Y), Axes.GetUnitAxis(EAxis::Z), false);
 
 						// Previous points
-						const FVec3 WorldPrevPointLocation = SpaceTransform.TransformPosition(ManifoldPoint.WorldContactPoints[ContactPointOwner]);
-						const FVec3 WorldPrevPlaneLocation = SpaceTransform.TransformPosition(ManifoldPoint.WorldContactPoints[ContactPlaneOwner]);
+						const FVec3 WorldPrevPointLocation = SpaceTransform.TransformPosition(PointTransform.TransformPosition(ManifoldPoint.ShapeAnchorPoints[ContactPointOwner]));
+						const FVec3 WorldPrevPlaneLocation = SpaceTransform.TransformPosition(PlaneTransform.TransformPosition(ManifoldPoint.ShapeAnchorPoints[ContactPlaneOwner]));
 						FDebugDrawQueue::GetInstance().DrawDebugLine(WorldPrevPointLocation, WorldPointLocation, FColor::White, false, KINDA_SMALL_NUMBER, Settings.DrawPriority, Settings.LineThickness);
 						FDebugDrawQueue::GetInstance().DrawDebugLine(WorldPrevPlaneLocation, WorldPlaneLocation, FColor::White, false, KINDA_SMALL_NUMBER, Settings.DrawPriority, Settings.LineThickness);
 
@@ -929,20 +929,20 @@ namespace Chaos
 				}
 				else
 				{
-					const FVec3 Location = SpaceTransform.TransformPosition(Contact.GetLocation());
-					const FVec3 Normal = SpaceTransform.TransformVector(Contact.GetNormal());
+					const FVec3 Location = SpaceTransform.TransformPosition(Contact.CalculateWorldContactLocation());
+					const FVec3 Normal = SpaceTransform.TransformVector(Contact.CalculateWorldContactNormal());
 					DrawCollisionImpl(Location, Normal, Contact.GetPhi(), FVec3(0), FColor(200, 0, 0), FColor(200, 0, 0), FColor(200, 0, 0), ColorScale, Settings);
 				}
 			}
 			if (Settings.ContactOwnerWidth > 0)
 			{
-				const FVec3 Location = SpaceTransform.TransformPosition(Contact.GetLocation());
-				const FVec3 Normal = SpaceTransform.TransformVector(Contact.GetNormal());
+				const FVec3 Location = SpaceTransform.TransformPosition(Contact.CalculateWorldContactLocation());
+				const FVec3 Normal = SpaceTransform.TransformVector(Contact.CalculateWorldContactNormal());
 
 				const FColor C3 = (ColorScale * FColor(128, 128, 128)).ToFColor(false);
 				const FMatrix Axes = FRotationMatrix::MakeFromX(Normal);
-				const FVec3 P0 = SpaceTransform.TransformPosition(Contact.Particle[0]->X());
-				const FVec3 P1 = SpaceTransform.TransformPosition(Contact.Particle[1]->X());
+				const FVec3 P0 = SpaceTransform.TransformPosition(Contact.GetParticle0()->X());
+				const FVec3 P1 = SpaceTransform.TransformPosition(Contact.GetParticle1()->X());
 				FDebugDrawQueue::GetInstance().DrawDebugLine(Location, P0, C3, false, KINDA_SMALL_NUMBER, Settings.DrawPriority, Settings.LineThickness * 0.5f);
 				FDebugDrawQueue::GetInstance().DrawDebugLine(Location, P1, C3, false, KINDA_SMALL_NUMBER, Settings.DrawPriority, Settings.LineThickness * 0.5f);
 			}
@@ -962,18 +962,18 @@ namespace Chaos
 			for (int32 ConstraintIndex = 0; ConstraintIndex < Collisions.NumConstraints(); ++ConstraintIndex)
 			{
 				const FPBDCollisionConstraint* PointConstraint = &Collisions.GetConstraint(ConstraintIndex);
-				if (!PointConstraint->GetDisabled() && (PointConstraint->Manifold.Phi < TNumericLimits<FReal>::Max()))
+				if (!PointConstraint->GetDisabled() && (PointConstraint->GetPhi() < TNumericLimits<FReal>::Max()))
 				{
-					const FImplicitObject* Implicit0 = PointConstraint->Manifold.Implicit[0];
-					const FImplicitObject* Implicit1 = PointConstraint->Manifold.Implicit[1];
+					const FImplicitObject* Implicit0 = PointConstraint->GetImplicit0();
+					const FImplicitObject* Implicit1 = PointConstraint->GetImplicit1();
 					if ((Implicit0 != nullptr) && (Implicit1 != nullptr))
 					{
-						FConstGenericParticleHandle Particle0 = PointConstraint->Particle[0];
-						FConstGenericParticleHandle Particle1 = PointConstraint->Particle[1];
+						FConstGenericParticleHandle Particle0 = PointConstraint->GetParticle0();
+						FConstGenericParticleHandle Particle1 = PointConstraint->GetParticle1();
 						const FRigidTransform3 WorldActorTransform0 = FParticleUtilities::GetActorWorldTransform(Particle0);
 						const FRigidTransform3 WorldActorTransform1 = FParticleUtilities::GetActorWorldTransform(Particle1);
-						const FRigidTransform3 ShapeWorldTransform0 = PointConstraint->ImplicitTransform[0] * WorldActorTransform0;
-						const FRigidTransform3 ShapeWorldTransform1 = PointConstraint->ImplicitTransform[1] * WorldActorTransform1;
+						const FRigidTransform3 ShapeWorldTransform0 = PointConstraint->GetShapeRelativeTransform0() * WorldActorTransform0;
+						const FRigidTransform3 ShapeWorldTransform1 = PointConstraint->GetShapeRelativeTransform1() * WorldActorTransform1;
 
 						if (!Shapes.Contains(Implicit0))
 						{
@@ -1138,14 +1138,16 @@ namespace Chaos
 		{
 			auto DrawGraphCollision = [&](const FRigidTransform3& SpaceTransform, const FPBDCollisionConstraint* Constraint,  int32 IslandIndex, int32 LevelIndex, int32 ColorIndex, int32 OrderIndex, bool bIsUsed, const FChaosDebugDrawSettings& Settings)
 			{
-				FVec3 ContactPos = Constraint->GetContactLocation();
+				const FRigidTransform3& ShapeTransform0 = Constraint->GetShapeWorldTransform0();
+				const FRigidTransform3& ShapeTransform1 = Constraint->GetShapeWorldTransform1();
+				FVec3 ContactPos = FReal(0.5) * (ShapeTransform0.GetLocation() + ShapeTransform1.GetLocation());
 				if (Constraint->GetManifoldPoints().Num() > 0)
 				{
 					ContactPos = FVec3(0);
 					for (const FManifoldPoint& ManifoldPoint : Constraint->GetManifoldPoints())
 					{
-						ContactPos += SpaceTransform.TransformPosition(ManifoldPoint.WorldContactPoints[0]);
-						ContactPos += SpaceTransform.TransformPosition(ManifoldPoint.WorldContactPoints[1]);
+						ContactPos += SpaceTransform.TransformPosition(ShapeTransform0.TransformPositionNoScale(ManifoldPoint.ContactPoint.ShapeContactPoints[0]));
+						ContactPos += SpaceTransform.TransformPosition(ShapeTransform1.TransformPositionNoScale(ManifoldPoint.ContactPoint.ShapeContactPoints[1]));
 					}
 					ContactPos /= (FReal)(2 * Constraint->GetManifoldPoints().Num());
 				}
diff --git a/Engine/Source/Runtime/Experimental/Chaos/Private/Chaos/Collision/CapsuleConvexContactPoint.cpp b/Engine/Source/Runtime/Experimental/Chaos/Private/Chaos/Collision/CapsuleConvexContactPoint.cpp
--- a/Engine/Source/Runtime/Experimental/Chaos/Private/Chaos/Collision/CapsuleConvexContactPoint.cpp
+++ b/Engine/Source/Runtime/Experimental/Chaos/Private/Chaos/Collision/CapsuleConvexContactPoint.cpp
@@ -36,8 +36,6 @@ namespace Chaos
 				ContactPoint.ShapeContactPoints[0] = PosConvex;
 				ContactPoint.ShapeContactPoints[1] = PosCapsule;
 				ContactPoint.ShapeContactNormal = -NormalConvex;
-				ContactPoint.Location = ConvexTransform.TransformPosition(ContactPoint.ShapeContactPoints[0]);
-				ContactPoint.Normal = ConvexTransform.TransformVectorNoScale(-NormalConvex);
 				ContactPoint.Phi = Phi;
 			}
 		}
diff --git a/Engine/Source/Runtime/Experimental/Chaos/Private/Chaos/Collision/ContactPointsMiscShapes.cpp b/Engine/Source/Runtime/Experimental/Chaos/Private/Chaos/Collision/ContactPointsMiscShapes.cpp
--- a/Engine/Source/Runtime/Experimental/Chaos/Private/Chaos/Collision/ContactPointsMiscShapes.cpp
+++ b/Engine/Source/Runtime/Experimental/Chaos/Private/Chaos/Collision/ContactPointsMiscShapes.cpp
@@ -45,9 +45,8 @@ namespace Chaos
 					Contact.ShapeContactPoints[0] = AToBTM.InverseTransformPosition(Location);
 					Contact.ShapeContactPoints[1] = Location;
 					Contact.ShapeContactNormal = Normal;
-					Contact.Location = BTransform.TransformPosition(Location);
-					Contact.Normal = BTransform.TransformVectorNoScale(Normal);
-					ComputeSweptContactPhiAndTOIHelper(Contact.Normal, Dir, Length, OutTime, TOI, Contact.Phi);
+					const FVec3 ContactNormal = BTransform.TransformVectorNoScale(Normal);
+					ComputeSweptContactPhiAndTOIHelper(ContactNormal, Dir, Length, OutTime, TOI, Contact.Phi);
 				}
 			});
 
@@ -92,8 +91,6 @@ namespace Chaos
 				Contact.ShapeContactPoints[1] = Location - ContactPhi * Normal;
 				Contact.ShapeContactNormal = Normal;
 				Contact.Phi = ContactPhi;
-				Contact.Location = BTransform.TransformPosition(Location);
-				Contact.Normal = BTransform.TransformVectorNoScale(Normal);
 			}
 		}
 		else if (const TImplicitObjectInstanced<GeometryA>* InstancedConvexImplicit = A.template GetObject<const TImplicitObjectInstanced<GeometryA> >())
@@ -106,8 +103,6 @@ namespace Chaos
 					Contact.ShapeContactPoints[1] = Location - ContactPhi * Normal;
 					Contact.ShapeContactNormal = Normal;
 					Contact.Phi = ContactPhi;
-					Contact.Location = BTransform.TransformPosition(Location);
-					Contact.Normal = BTransform.TransformVectorNoScale(Normal);
 				}
 			}
 		}
@@ -119,8 +114,6 @@ namespace Chaos
 				Contact.ShapeContactPoints[1] = Location - ContactPhi * Normal;
 				Contact.ShapeContactNormal = Normal;
 				Contact.Phi = ContactPhi;
-				Contact.Location = BTransform.TransformPosition(Location);
-				Contact.Normal = BTransform.TransformVectorNoScale(Normal);
 			}
 		}
 
@@ -152,11 +145,8 @@ namespace Chaos
 		for (FContactPoint& ContactPoint : ContactPoints)
 		{
 			// Calculate the shape contact points
-			ContactPoint.ShapeContactPoints[0] = AToBTM.InverseTransformPosition(ContactPoint.Location);
-			ContactPoint.ShapeContactPoints[1] = ContactPoint.Location - ContactPoint.Phi * ContactPoint.ShapeContactNormal;
-			ContactPoint.Normal = BTransform.TransformVectorNoScale(ContactPoint.ShapeContactNormal);
-			// Get the Location in world space
-			ContactPoint.Location = BTransform.TransformPosition(ContactPoint.Location);
+			ContactPoint.ShapeContactPoints[0] = AToBTM.InverseTransformPosition(ContactPoint.ShapeContactPoints[1]);
+			ContactPoint.ShapeContactPoints[1] = ContactPoint.ShapeContactPoints[1] - ContactPoint.Phi * ContactPoint.ShapeContactNormal;
 		}
 	}
 
@@ -182,9 +172,11 @@ namespace Chaos
 				// @todo(chaos): handle Instanced with margin
 				if (B.LowLevelSweepGeom(ADowncast, AToBTM, LocalDir, Length, OutTime, Location, Normal, FaceIndex, 0.0f, true))
 				{
-					Contact.Location = BTransform.TransformPosition(Location);
-					Contact.Normal = BTransform.TransformVectorNoScale(Normal);
-					ComputeSweptContactPhiAndTOIHelper(Contact.Normal, Dir, Length, OutTime, TOI, Contact.Phi);
+					Contact.ShapeContactPoints[1] = Location;
+					Contact.ShapeContactNormal = Normal;
+
+					const FVec3& ContactNormal = BTransform.TransformVectorNoScale(Normal);
+					ComputeSweptContactPhiAndTOIHelper(ContactNormal, Dir, Length, OutTime, TOI, Contact.Phi);
 				}
 			});
 
@@ -215,8 +207,6 @@ namespace Chaos
 			Result.ShapeContactPoints[1] = Sphere2.GetCenter() + Sphere2Transform.InverseTransformVector(R2 * Normal);
 			Result.ShapeContactNormal = Sphere2Transform.InverseTransformVector(Normal);
 			Result.Phi = NewPhi;
-			Result.Normal = Normal;
-			Result.Location = Center2 + R2 * Result.Normal;
 		}
 		return Result;
 	}
@@ -239,8 +229,6 @@ namespace Chaos
 		Result.ShapeContactPoints[1] = PlaneTransform.InverseTransformPosition(Location - Phi * NormalWorld);
 		Result.ShapeContactNormal = PlaneTransform.InverseTransformVector(NormalWorld);
 		Result.Phi = Phi;
-		Result.Normal = NormalWorld;
-		Result.Location = Location;
 
 		return Result;
 	}
@@ -263,8 +251,6 @@ namespace Chaos
 		Result.ShapeContactPoints[1] = BoxTransform.InverseTransformPosition(LocationWorld - Phi * NormalWorld);
 		Result.ShapeContactNormal = NormalBox;
 		Result.Phi = Phi;
-		Result.Normal = NormalWorld;
-		Result.Location = LocationWorld;
 		return Result;
 	}
 
@@ -291,8 +277,6 @@ namespace Chaos
 			Result.ShapeContactPoints[1] = BTransform.InverseTransformPosition(LocationB);
 			Result.ShapeContactNormal = BTransform.InverseTransformVector(Normal);
 			Result.Phi = NewPhi;
-			Result.Normal = Normal;
-			Result.Location = FReal(0.5) * (LocationA + LocationB);
 		}
 
 		return Result;
@@ -421,8 +405,6 @@ namespace Chaos
 			Result.ShapeContactPoints[1] = BTransform.InverseTransformPosition(LocationB);
 			Result.ShapeContactNormal = BTransform.InverseTransformVector(Normal);
 			Result.Phi = NewPhi;
-			Result.Normal = Normal;
-			Result.Location = FReal(0.5) * (LocationA + LocationB);
 		}
 
 		return Result;
diff --git a/Engine/Source/Runtime/Experimental/Chaos/Private/Chaos/Collision/GJKContactPoint.h b/Engine/Source/Runtime/Experimental/Chaos/Private/Chaos/Collision/GJKContactPoint.h
--- a/Engine/Source/Runtime/Experimental/Chaos/Private/Chaos/Collision/GJKContactPoint.h
+++ b/Engine/Source/Runtime/Experimental/Chaos/Private/Chaos/Collision/GJKContactPoint.h
@@ -31,8 +31,6 @@ namespace Chaos
 			Contact.ShapeContactPoints[0] = ClosestA;
 			Contact.ShapeContactPoints[1] = BToATM.InverseTransformPosition(ClosestBInA);
 			Contact.ShapeContactNormal = -BToATM.InverseTransformVector(Normal);
-			Contact.Location = ATM.TransformPosition(ClosestA + ThicknessA * Normal);
-			Contact.Normal = -ATM.TransformVectorNoScale(Normal);
 			Contact.Phi = -Penetration;
 		}
 
diff --git a/Engine/Source/Runtime/Experimental/Chaos/Private/Chaos/Collision/GJKContactPointSwept.cpp b/Engine/Source/Runtime/Experimental/Chaos/Private/Chaos/Collision/GJKContactPointSwept.cpp
--- a/Engine/Source/Runtime/Experimental/Chaos/Private/Chaos/Collision/GJKContactPointSwept.cpp
+++ b/Engine/Source/Runtime/Experimental/Chaos/Private/Chaos/Collision/GJKContactPointSwept.cpp
@@ -47,8 +47,6 @@ namespace Chaos
 			Contact.ShapeContactPoints[0] = ATOITM.InverseTransformPosition(WorldLocation);
 			Contact.ShapeContactPoints[1] = Location;
 			Contact.ShapeContactNormal = Normal;
-			Contact.Location = WorldLocation;
-			Contact.Normal = WorldNormal;
 
 			if (OutTime > 0.f)
 			{
diff --git a/Engine/Source/Runtime/Experimental/Chaos/Private/Chaos/Collision/PBDCollisionConstraint.cpp b/Engine/Source/Runtime/Experimental/Chaos/Private/Chaos/Collision/PBDCollisionConstraint.cpp
--- a/Engine/Source/Runtime/Experimental/Chaos/Private/Chaos/Collision/PBDCollisionConstraint.cpp
+++ b/Engine/Source/Runtime/Experimental/Chaos/Private/Chaos/Collision/PBDCollisionConstraint.cpp
@@ -141,19 +141,19 @@ namespace Chaos
 	FPBDCollisionConstraint::FPBDCollisionConstraint()
 		: ImplicitTransform{ FRigidTransform3(), FRigidTransform3() }
 		, Particle{ nullptr, nullptr }
-		, AccumulatedImpulse(0)
 		, Manifold()
+		, AccumulatedImpulse(0)
 		, TimeOfImpact(0)
 		, ContainerCookie()
 		, CCDType(ECollisionCCDType::Disabled)
 		, Stiffness(FReal(1))
-		, ManifoldPoints()
-		, SavedManifoldPoints()
 		, CullDistance(TNumericLimits<FReal>::Max())
 		, CollisionMargins{ 0, 0 }
 		, CollisionTolerance(0)
 		, Flags()
 		, SolverBodies{ nullptr, nullptr }
+		, ManifoldPoints()
+		, SavedManifoldPoints()
 		, GJKWarmStartData()
 		, ShapeWorldTransform0()
 		, ShapeWorldTransform1()
@@ -177,19 +177,19 @@ namespace Chaos
 		const FBVHParticles* Simplicial1)
 		: ImplicitTransform{ FRigidTransform3(), FRigidTransform3() }
 		, Particle{ Particle0, Particle1 }
-		, AccumulatedImpulse(0)
 		, Manifold()
+		, AccumulatedImpulse(0)
 		, TimeOfImpact(0)
 		, ContainerCookie()
 		, CCDType(ECollisionCCDType::Disabled)
 		, Stiffness(FReal(1))
-		, ManifoldPoints()
-		, SavedManifoldPoints()
 		, CullDistance(TNumericLimits<FReal>::Max())
 		, CollisionMargins{ 0, 0 }
 		, CollisionTolerance(0)
 		, Flags()
 		, SolverBodies{ nullptr, nullptr }
+		, ManifoldPoints()
+		, SavedManifoldPoints()
 		, GJKWarmStartData()
 		, ShapeWorldTransform0()
 		, ShapeWorldTransform1()
@@ -282,6 +282,26 @@ namespace Chaos
 		ContainerCookie.MidPhase->SetIsSleeping(bInIsSleeping);
 	}
 
+	FVec3 FPBDCollisionConstraint::CalculateWorldContactLocation() const
+	{
+		if (ClosestManifoldPointIndex != INDEX_NONE)
+		{
+			const FVec3 WorldContact0 = GetShapeWorldTransform0().TransformPositionNoScale(ManifoldPoints[ClosestManifoldPointIndex].ContactPoint.ShapeContactPoints[0]);
+			const FVec3 WorldContact1 = GetShapeWorldTransform1().TransformPositionNoScale(ManifoldPoints[ClosestManifoldPointIndex].ContactPoint.ShapeContactPoints[1]);
+			return FReal(0.5) * (WorldContact0 + WorldContact1);
+		}
+		return FVec3(0);
+	}
+
+	FVec3 FPBDCollisionConstraint::CalculateWorldContactNormal() const
+	{
+		if (ClosestManifoldPointIndex != INDEX_NONE)
+		{
+			return GetShapeWorldTransform1().TransformVectorNoScale(ManifoldPoints[ClosestManifoldPointIndex].ContactPoint.ShapeContactNormal);
+		}
+		return FVec3(0, 0, 1);
+	}
+
 	// Are the two manifold points the same point?
 	// Ideally a contact is considered the same as one from the previous iteration if
 	//		The contact is Vertex - Face and there was a prior iteration collision on the same Vertex
@@ -371,6 +391,15 @@ namespace Chaos
 		return BestMatchIndex;
 	}
 
+	void FPBDCollisionConstraint::UpdateManifoldPointPhi(const int32 ManifoldPointIndex)
+	{
+		FManifoldPoint& ManifoldPoint = ManifoldPoints[ManifoldPointIndex];
+		const FVec3 WorldContact0 = GetShapeWorldTransform0().TransformPositionNoScale(ManifoldPoint.ContactPoint.ShapeContactPoints[0]);
+		const FVec3 WorldContact1 = GetShapeWorldTransform1().TransformPositionNoScale(ManifoldPoint.ContactPoint.ShapeContactPoints[1]);
+		const FVec3 WorldContactNormal = GetShapeWorldTransform1().TransformVectorNoScale(ManifoldPoint.ContactPoint.ShapeContactNormal);
+		ManifoldPoint.ContactPoint.Phi = FVec3::DotProduct(WorldContact0 - WorldContact1, WorldContactNormal);
+	}
+
 	void FPBDCollisionConstraint::UpdateManifoldContacts()
 	{
 		// This is only used when calling collision detection in the solver loop, which is only for incremental manifolds.
@@ -381,20 +410,20 @@ namespace Chaos
 			ShapeWorldTransform1 = FRigidTransform3(GetSolverBody1()->CorrectedP(), GetSolverBody1()->CorrectedQ());
 		}
 
+		ClosestManifoldPointIndex = INDEX_NONE;
 		Manifold.Reset();
 
 		for (int32 ManifoldPointIndex = 0; ManifoldPointIndex < ManifoldPoints.Num(); ManifoldPointIndex++)
 		{
 			FManifoldPoint& ManifoldPoint = ManifoldPoints[ManifoldPointIndex];
 
-			UpdateManifoldPointFromContact(ManifoldPointIndex);
+			UpdateManifoldPointPhi(ManifoldPointIndex);
 
-			ManifoldPoint.Flags.bInsideStaticFrictionCone = Flags.bUseManifold;
+			ManifoldPoint.Flags.bInsideStaticFrictionCone = false;
 
-			// Copy currently active point
-			if (ManifoldPoint.ContactPoint.Phi < Manifold.Phi)
+			if (ManifoldPoint.ContactPoint.Phi < GetPhi())
 			{
-				SetActiveContactPoint(ManifoldPoint.ContactPoint);
+				ClosestManifoldPointIndex = ManifoldPointIndex;
 			}
 		}
 	}
@@ -410,10 +439,9 @@ namespace Chaos
 			
 			int32 ManifoldPointIndex = AddManifoldPoint(ContactPoint);
 
-			// Copy currently active point
-			if (ManifoldPoints[ManifoldPointIndex].ContactPoint.Phi < Manifold.Phi)
+			if (ManifoldPoints[ManifoldPointIndex].ContactPoint.Phi < GetPhi())
 			{
-				SetActiveContactPoint(ManifoldPoints[ManifoldPointIndex].ContactPoint);
+				ClosestManifoldPointIndex = ManifoldPointIndex;
 			}
 
 			Flags.bUseIncrementalManifold = false;
@@ -435,8 +463,6 @@ namespace Chaos
 			{
 				// This contact point is already in the manifold - update the state
 				ManifoldPoints[ManifoldPointIndex].ContactPoint = ContactPoint;
-
-				UpdateManifoldPointFromContact(ManifoldPointIndex);
 			}
 			else
 			{
@@ -444,10 +470,9 @@ namespace Chaos
 				ManifoldPointIndex = AddManifoldPoint(ContactPoint);
 			}
 
-			// Copy currently active point
-			if (ManifoldPoints[ManifoldPointIndex].ContactPoint.Phi < Manifold.Phi)
+			if (ManifoldPoints[ManifoldPointIndex].ContactPoint.Phi < GetPhi())
 			{
-				SetActiveContactPoint(ManifoldPoints[ManifoldPointIndex].ContactPoint);
+				ClosestManifoldPointIndex = ManifoldPointIndex;
 			}
 		}
 		else 
@@ -458,7 +483,7 @@ namespace Chaos
 
 			InitManifoldPoint(0);
 
-			SetActiveContactPoint(ManifoldPoints[0].ContactPoint);
+			ClosestManifoldPointIndex = 0;
 		}
 
 		Flags.bUseIncrementalManifold = true;
@@ -466,22 +491,13 @@ namespace Chaos
 
 	void FPBDCollisionConstraint::InitManifoldPoint(const int32 ManifoldPointIndex)
 	{
-		if ((Particle[0] == nullptr) || (Particle[1] == nullptr))
-		{
-			// @todo(chaos): This is just for unit tests testing one-shot manifolds - remove it somehow... 
-			// maybe ConstructConvexConvexOneShotManifold should not take a Constraint
-			return;
-		}
-
 		FManifoldPoint& ManifoldPoint = ManifoldPoints[ManifoldPointIndex];
 		ManifoldPoint.InitialShapeContactPoints[0] = ManifoldPoint.ContactPoint.ShapeContactPoints[0];
 		ManifoldPoint.InitialShapeContactPoints[1] = ManifoldPoint.ContactPoint.ShapeContactPoints[1];
 
-		// Initialize the previous contact transforms if the data is available, otherwise reset them to current
-		TryRestoreFrictionData(ManifoldPointIndex);
-
-		// Update the derived contact state (World-Space data)
-		UpdateManifoldPointFromContact(ManifoldPointIndex);
+		ManifoldPoints[ManifoldPointIndex].Flags.Reset();
+		ManifoldPoints[ManifoldPointIndex].NetPushOut = FVec3(0);
+		ManifoldPoints[ManifoldPointIndex].NetImpulse = FVec3(0);
 	}
 
 	int32 FPBDCollisionConstraint::AddManifoldPoint(const FContactPoint& ContactPoint)
@@ -495,22 +511,6 @@ namespace Chaos
 		return ManifoldPointIndex;
 	}
 
-	void FPBDCollisionConstraint::UpdateManifoldPointFromContact(const int32 ManifoldPointIndex)
-	{
-		FManifoldPoint& ManifoldPoint = ManifoldPoints[ManifoldPointIndex];
-
-		ManifoldPoint.WorldContactPoints[0] = ShapeWorldTransform0.TransformPositionNoScale(ManifoldPoint.ShapeAnchorPoints[0]);
-		ManifoldPoint.WorldContactPoints[1] = ShapeWorldTransform1.TransformPositionNoScale(ManifoldPoint.ShapeAnchorPoints[1]);
-	}
-
-	void FPBDCollisionConstraint::SetActiveContactPoint(const FContactPoint& ContactPoint)
-	{
-		// @todo(chaos): once we settle on manifolds we should just store the index
-		Manifold.Location = ContactPoint.Location;
-		Manifold.Normal = ContactPoint.Normal;
-		Manifold.Phi = ContactPoint.Phi;
-	}
-
 	bool FPBDCollisionConstraint::CanUseManifold(FGeometryParticleHandle* Particle0, FGeometryParticleHandle* Particle1) const
 	{
 		// Do not use manifolds when a body is connected by a joint to another. Manifolds do not work when the bodies may be moved
@@ -526,6 +526,7 @@ namespace Chaos
 
 	void FPBDCollisionConstraint::ResetActiveManifoldContacts()
 	{
+		ClosestManifoldPointIndex = INDEX_NONE;
 		Manifold.Reset();
 		ManifoldPoints.Reset();
 		ExpectedNumManifoldPoints = 0;
@@ -542,7 +543,7 @@ namespace Chaos
 		}
 		else
 		{
-			// The bodies have barely moved so we keep the same shape-space contacts. We still need to update world-space contacts though.
+			// The bodies have barely moved so we keep the same shape-space contacts. We still need to update ManifoldPoint Phis
 			UpdateManifoldContacts();
 		}
 
@@ -571,7 +572,7 @@ namespace Chaos
 		const FRigidTransform3 Shape0ToShape1Transform = ShapeWorldTransform0.GetRelativeTransformNoScale(ShapeWorldTransform1);
 		const FVec3 Contact0In1 = Shape0ToShape1Transform.TransformPositionNoScale(ManifoldPoint.InitialShapeContactPoints[0]);
 		const FVec3& Contact1In1 = ManifoldPoint.InitialShapeContactPoints[1];
-		const FVec3 ContactNormalIn1 = ShapeWorldTransform1.InverseTransformVectorNoScale(ManifoldPoint.ContactPoint.Normal);
+		const FVec3& ContactNormalIn1 = ManifoldPoint.ContactPoint.ShapeContactNormal;
 
 		const FVec3 ContactDeltaIn1 = Contact0In1 - Contact1In1;
 		const FReal ContactPhi = FVec3::DotProduct(ContactDeltaIn1, ContactNormalIn1);
@@ -580,17 +581,11 @@ namespace Chaos
 		const FVec3 ShapeContactPoint1 = Contact0In1 - ContactPhi * ContactNormalIn1;
 		ManifoldPoint.ContactPoint.ShapeContactPoints[1] = ShapeContactPoint1;
 		ManifoldPoint.ContactPoint.Phi = ContactPhi;
-
-		// Restore friction anchors if we have them for this point
-		TryRestoreFrictionData(ManifoldPointIndex);
-
-		// Update world-space contact locations
-		UpdateManifoldPointFromContact(ManifoldPointIndex);
-		ManifoldPoint.ContactPoint.Location = FReal(0.5) * (ManifoldPoint.WorldContactPoints[0] + ManifoldPoint.WorldContactPoints[1]);
 	}
 
 	void FPBDCollisionConstraint::ReprojectManifoldContacts()
 	{
+		ClosestManifoldPointIndex = INDEX_NONE;
 		Manifold.Reset();
 
 		for (int32 ManifoldPointIndex = 0; ManifoldPointIndex < ManifoldPoints.Num(); ManifoldPointIndex++)
@@ -599,26 +594,27 @@ namespace Chaos
 
 			ReprojectManifoldPoint(ManifoldPointIndex);
 
-			ManifoldPoint.Flags.bInsideStaticFrictionCone = Flags.bUseManifold;
+			ManifoldPoint.Flags.bInsideStaticFrictionCone = false;
 			ManifoldPoint.Flags.bWasRestored = true;
 
 			// Copy currently active point
-			if (ManifoldPoint.ContactPoint.Phi < Manifold.Phi)
+			if (ManifoldPoint.ContactPoint.Phi < GetPhi())
 			{
-				SetActiveContactPoint(ManifoldPoint.ContactPoint);
+				ClosestManifoldPointIndex = ManifoldPointIndex;
 			}
 		}
 	}
 
 	bool FPBDCollisionConstraint::UpdateAndTryRestoreManifold()
 	{
-		const FCollisionTolerances& Tolerances = Chaos_Manifold_Tolerances;
+		const FCollisionTolerances Tolerances = Chaos_Manifold_Tolerances;
 		const FReal ContactPositionTolerance = Tolerances.ContactPositionToleranceScale * CollisionTolerance;
 		const FReal ShapePositionTolerance = (ManifoldPoints.Num() > 0) ? Tolerances.ShapePositionToleranceScaleN * CollisionTolerance : Tolerances.ShapePositionToleranceScale0 * CollisionTolerance;
 		const FReal ShapeRotationThreshold = (ManifoldPoints.Num() > 0) ? Tolerances.ShapeRotationThresholdN : Tolerances.ShapeRotationThreshold0;
 		const FReal ContactPositionToleranceSq = FMath::Square(ContactPositionTolerance);
 
 		// Reset current closest point
+		ClosestManifoldPointIndex = INDEX_NONE;
 		Manifold.Reset();
 
 		// How many manifold points we expect. E.g., for Box-box this will be 4 or 1 depending on whether
@@ -644,7 +640,7 @@ namespace Chaos
 				// @todo(chaos): this should use the normal owner. Currently we assume body 1 is the owner
 				const FVec3 Contact0In1 = Shape0ToShape1Transform.TransformPositionNoScale(ManifoldPoint.InitialShapeContactPoints[0]);
 				const FVec3& Contact1In1 = ManifoldPoint.InitialShapeContactPoints[1];
-				const FVec3 ContactNormalIn1 = ShapeWorldTransform1.InverseTransformVectorNoScale(ManifoldPoint.ContactPoint.Normal);
+				const FVec3& ContactNormalIn1 = ManifoldPoint.ContactPoint.ShapeContactNormal;
 
 				const FVec3 ContactDeltaIn1 = Contact0In1 - Contact1In1;
 				const FReal ContactPhi = FVec3::DotProduct(ContactDeltaIn1, ContactNormalIn1);
@@ -659,6 +655,10 @@ namespace Chaos
 					const FVec3 ShapeContactPoint1 = Contact0In1 - ContactPhi * ContactNormalIn1;
 					ManifoldPoint.ContactPoint.ShapeContactPoints[1] = ShapeContactPoint1;
 					ManifoldPoint.ContactPoint.Phi = ContactPhi;
+					if (ManifoldPoint.ContactPoint.Phi < GetPhi())
+					{
+						ClosestManifoldPointIndex = ManifoldPointIndex;
+					}
 				}
 				else if ((ManifoldPointToRemove == INDEX_NONE) && bChaos_Collision_EnableManifoldInject)
 				{
@@ -677,26 +677,10 @@ namespace Chaos
 			if (ManifoldPointToRemove != INDEX_NONE)
 			{
 				ManifoldPoints.RemoveAt(ManifoldPointToRemove);
-			}
-
-			// Update world-space state for the points we kept
-			for (int32 ManifoldPointIndex = 0; ManifoldPointIndex < ManifoldPoints.Num(); ++ManifoldPointIndex)
-			{
-				FManifoldPoint& ManifoldPoint = ManifoldPoints[ManifoldPointIndex];
-
-				// Restore friction anchors if we have them for this point
-				TryRestoreFrictionData(ManifoldPointIndex);
-
-				// Update world-space contact locations
-				UpdateManifoldPointFromContact(ManifoldPointIndex);
-				ManifoldPoint.ContactPoint.Location = FReal(0.5) * (ManifoldPoint.WorldContactPoints[0] + ManifoldPoint.WorldContactPoints[1]);
-
-				ManifoldPoint.Flags.bWasRestored = true;
-
-				if (ManifoldPoint.ContactPoint.Phi < Manifold.Phi)
+				if ((ManifoldPointToRemove < ClosestManifoldPointIndex) && (ClosestManifoldPointIndex != INDEX_NONE))
 				{
-					// Update closest point
-					SetActiveContactPoint(ManifoldPoint.ContactPoint);
+					--ClosestManifoldPointIndex;
+					check(ClosestManifoldPointIndex >= 0);
 				}
 			}
 		}
@@ -748,7 +732,7 @@ namespace Chaos
 		{
 			FManifoldPoint& ManifoldPoint = ManifoldPoints[ManifoldPointIndex];
 
-			const FReal NormalOverlap = FVec3::DotProduct(ManifoldPoint.ContactPoint.Normal, NewContactPoint.Normal);
+			const FReal NormalOverlap = FVec3::DotProduct(ManifoldPoint.ContactPoint.ShapeContactNormal, NewContactPoint.ShapeContactNormal);
 			if (NormalOverlap < NormalThreshold)
 			{
 				return false;
@@ -773,11 +757,9 @@ namespace Chaos
 					ManifoldPoint.InitialShapeContactPoints[0] = NewContactPoint.ShapeContactPoints[0];
 					ManifoldPoint.InitialShapeContactPoints[1] = NewContactPoint.ShapeContactPoints[1];
 					ManifoldPoint.Flags.bWasRestored = false;
-					TryRestoreFrictionData(ManifoldPointIndex);
-					UpdateManifoldPointFromContact(ManifoldPointIndex);
 					if (NewContactPoint.Phi < GetPhi())
 					{
-						SetActiveContactPoint(ManifoldPoint.ContactPoint);
+						ClosestManifoldPointIndex = ManifoldPointIndex;
 					}
 				}
 
@@ -921,10 +903,9 @@ namespace Chaos
 				ManifoldPoints[ManifoldPointIndex].InitialShapeContactPoints[0] = NewContactPoint.ShapeContactPoints[0];
 				ManifoldPoints[ManifoldPointIndex].InitialShapeContactPoints[1] = NewContactPoint.ShapeContactPoints[1];
 				ManifoldPoints[ManifoldPointIndex].Flags.bWasRestored = false;
-				UpdateManifoldPointFromContact(ManifoldPointIndex);
-				if (NewContactPoint.Phi < Manifold.Phi)
+				if (NewContactPoint.Phi < GetPhi())
 				{
-					SetActiveContactPoint(NewContactPoint);
+					ClosestManifoldPointIndex = ManifoldPointIndex;
 				}
 			}
 		}
@@ -996,31 +977,6 @@ namespace Chaos
 		return nullptr;
 	}
 
-	void FPBDCollisionConstraint::TryRestoreFrictionData(const int32 ManifoldPointIndex)
-	{
-		FManifoldPoint& ManifoldPoint = ManifoldPoints[ManifoldPointIndex];
-
-		// Assume we have no matching point from the previous tick, but that we can hold friction from now on
-		// Not supported for non-manifolds yet (hopefully we don't need to)
-		ManifoldPoint.Flags.bInsideStaticFrictionCone = Flags.bUseManifold;
-
-		// Find the previous manifold point that matches if there is one
-		const FSavedManifoldPoint* SavedManifoldPoint = FindSavedManifoldPoint(ManifoldPoint);
-		if (SavedManifoldPoint != nullptr)
-		{
-			// We have data from the previous tick and static friction was enabled - restore the data
-			ManifoldPoint.ShapeAnchorPoints[0] = SavedManifoldPoint->ShapeContactPoints[0];
-			ManifoldPoint.ShapeAnchorPoints[1] = SavedManifoldPoint->ShapeContactPoints[1];
-			ManifoldPoint.Flags.bWasFrictionRestored = true;
-		}
-		else
-		{
-			ManifoldPoint.ShapeAnchorPoints[0] = ManifoldPoint.ContactPoint.ShapeContactPoints[0];
-			ManifoldPoint.ShapeAnchorPoints[1] = ManifoldPoint.ContactPoint.ShapeContactPoints[1];
-			ManifoldPoint.Flags.bWasFrictionRestored = false;
-		}
-	}
-
 	ECollisionConstraintDirection FPBDCollisionConstraint::GetConstraintDirection(const FReal Dt) const
 	{
 		if (GetDisabled())
@@ -1030,7 +986,7 @@ namespace Chaos
 		// D\tau is the chacteristic time (as in GBF paper Sec 8.1)
 		const FReal Dtau = Dt * Chaos_GBFCharacteristicTimeRatio; 
 
-		const FVec3 Normal = GetNormal();
+		const FVec3 Normal = CalculateWorldContactNormal();
 		const FReal Phi = GetPhi();
 		if (GetPhi() >= GetCullDistance())
 		{
diff --git a/Engine/Source/Runtime/Experimental/Chaos/Private/Chaos/Collision/PBDCollisionConstraintHandle.cpp b/Engine/Source/Runtime/Experimental/Chaos/Private/Chaos/Collision/PBDCollisionConstraintHandle.cpp
--- a/Engine/Source/Runtime/Experimental/Chaos/Private/Chaos/Collision/PBDCollisionConstraintHandle.cpp
+++ b/Engine/Source/Runtime/Experimental/Chaos/Private/Chaos/Collision/PBDCollisionConstraintHandle.cpp
@@ -42,11 +42,6 @@ namespace Chaos
 		return !GetContact().GetDisabled();
 	}
 
-	FVec3 FPBDCollisionConstraintHandle::GetContactLocation() const
-	{
-		return GetContact().GetLocation();
-	}
-
 	FVec3 FPBDCollisionConstraintHandle::GetAccumulatedImpulse() const
 	{
 		return GetContact().AccumulatedImpulse;
@@ -54,12 +49,12 @@ namespace Chaos
 
 	TVector<const TGeometryParticleHandle<FReal, 3>*, 2> FPBDCollisionConstraintHandle::GetConstrainedParticles() const
 	{
-		return { GetContact().Particle[0], GetContact().Particle[1] };
+		return { GetContact().GetParticle0(), GetContact().GetParticle1() };
 	}
 
 	TVector<TGeometryParticleHandle<FReal, 3>*, 2> FPBDCollisionConstraintHandle::GetConstrainedParticles()
 	{
-		return { GetContact().Particle[0], GetContact().Particle[1] };
+		return { GetContact().GetParticle0(), GetContact().GetParticle1() };
 	}
 
 	void FPBDCollisionConstraintHandle::GatherInput(FReal Dt, const int32 Particle0Level, const int32 Particle1Level, FPBDIslandSolverData& SolverData)
diff --git a/Engine/Source/Runtime/Experimental/Chaos/Private/Chaos/Collision/PBDCollisionSolver.cpp b/Engine/Source/Runtime/Experimental/Chaos/Private/Chaos/Collision/PBDCollisionSolver.cpp
--- a/Engine/Source/Runtime/Experimental/Chaos/Private/Chaos/Collision/PBDCollisionSolver.cpp
+++ b/Engine/Source/Runtime/Experimental/Chaos/Private/Chaos/Collision/PBDCollisionSolver.cpp
@@ -59,11 +59,24 @@ namespace Chaos
 		const FSolverReal Dt,
 		const FConstraintSolverBody& Body0,
 		const FConstraintSolverBody& Body1,
-		const FVec3& InWorldAnchorPoint0,
-		const FVec3& InWorldAnchorPoint1,
+		const FSolverVec3& InRelativeContactPosition0,
+		const FSolverVec3& InRelativeContactPosition1,
 		const FSolverVec3& InWorldContactNormal,
-		const bool bAddVelocityToFriction)
+		const FSolverVec3& InWorldContactTangentU,
+		const FSolverVec3& InWorldContactTangentV,
+		const FSolverReal InWorldContactDeltaNormal,
+		const FSolverReal InWorldContactDeltaTangentU,
+		const FSolverReal InWorldContactDeltaTangentV)
 	{
+		RelativeContactPosition0 = InRelativeContactPosition0;
+		RelativeContactPosition1 = InRelativeContactPosition1;
+		WorldContactNormal = InWorldContactNormal;
+		WorldContactTangentU = InWorldContactTangentU;
+		WorldContactTangentV = InWorldContactTangentV;
+		WorldContactDeltaNormal = InWorldContactDeltaNormal;
+		WorldContactDeltaTangentU = InWorldContactDeltaTangentU;
+		WorldContactDeltaTangentV = InWorldContactDeltaTangentV;
+
 		NetPushOutNormal = FSolverReal(0);
 		NetPushOutTangentU = FSolverReal(0);
 		NetPushOutTangentV = FSolverReal(0);
@@ -71,8 +84,6 @@ namespace Chaos
 		NetImpulseTangentU = FSolverReal(0);
 		NetImpulseTangentV = FSolverReal(0);
 
-		UpdateContact(Dt, Body0, Body1, InWorldAnchorPoint0, InWorldAnchorPoint1, InWorldContactNormal, bAddVelocityToFriction);
-
 		UpdateMass(Body0, Body1);
 	}
 
@@ -97,51 +108,6 @@ namespace Chaos
 
 	}
 
-	inline void FPBDCollisionSolverManifoldPoint::UpdateContact(
-		const FSolverReal Dt,
-		const FConstraintSolverBody& Body0, 
-		const FConstraintSolverBody& Body1,
-		const FVec3& InWorldAnchorPoint0,
-		const FVec3& InWorldAnchorPoint1, 
-		const FSolverVec3& InWorldContactNormal,
-		const bool bAddVelocityToFriction)
-	{
-		// The world-space point where we apply impulses/corrections (same world-space point for momentum conservation)
-		const FVec3 WorldContactPosition = FReal(0.5) * (InWorldAnchorPoint0 + InWorldAnchorPoint1);
-		RelativeContactPosition0 = FSolverVec3(WorldContactPosition - Body0.P());
-		RelativeContactPosition1 = FSolverVec3(WorldContactPosition - Body1.P());
-
-		// The world-space contact normal
-		WorldContactNormal = InWorldContactNormal;
-
-		// World-space contact tangents. We are treating the normal as the constraint-space Z axis
-		// and the Tangent U and V as the constraint-space X and Y axes resepctively
-		WorldContactTangentU = FSolverVec3::CrossProduct(FSolverVec3(0, 1, 0), InWorldContactNormal);
-		if (!WorldContactTangentU.Normalize(FSolverReal(KINDA_SMALL_NUMBER)))
-		{
-			WorldContactTangentU = FSolverVec3::CrossProduct(FSolverVec3(1, 0, 0), InWorldContactNormal);
-			WorldContactTangentU = WorldContactTangentU.GetUnsafeNormal();
-		}
-		WorldContactTangentV = FSolverVec3::CrossProduct(InWorldContactNormal, WorldContactTangentU);
-
-		// The contact point error we are trying to correct in this solver
-		const FSolverVec3 WorldContactDelta = InWorldAnchorPoint0 - InWorldAnchorPoint1;
-		WorldContactDeltaNormal = FSolverVec3::DotProduct(WorldContactDelta, WorldContactNormal);
-		WorldContactDeltaTangentU = FSolverVec3::DotProduct(WorldContactDelta, WorldContactTangentU);
-		WorldContactDeltaTangentV = FSolverVec3::DotProduct(WorldContactDelta, WorldContactTangentV);
-
-		// For rolling friction to work (e.g., on spheres) we need to include the angular velocity in the friction calculation
-		if (bAddVelocityToFriction)
-		{
-			// Contact velocity from rotation for quadratic surfaces. Corrects for rolling friction.
-			FVec3 ContactVel = (Body0.V() + FVec3::CrossProduct(Body0.W(), RelativeContactPosition0)) - (Body1.V() + FVec3::CrossProduct(Body1.W(), RelativeContactPosition1));
-			const FSolverReal ContactVelocityU = FSolverVec3::DotProduct(ContactVel, WorldContactTangentU);
-			const FSolverReal ContactVelocityV = FSolverVec3::DotProduct(ContactVel, WorldContactTangentV);
-			WorldContactDeltaTangentU += ContactVelocityU * Dt;
-			WorldContactDeltaTangentV += ContactVelocityV * Dt;
-		}
-	}
-
 	inline void FPBDCollisionSolverManifoldPoint::UpdateMass(const FConstraintSolverBody& Body0, const FConstraintSolverBody& Body1)
 	{
 		FSolverReal ContactMassInvNormal = FSolverReal(0);
@@ -272,32 +238,40 @@ namespace Chaos
 
 	void FPBDCollisionSolver::InitContact(
 		const int32 ManifoldPoiontIndex,
-		const FReal Dt,
-		const FVec3& InWorldAnchorPoint0,
-		const FVec3& InWorldAnchorPoint1,
-		const FVec3& InWorldContactNormal,
-		const bool bAddVelocityToFriction)
+		const FSolverReal Dt,
+		const FSolverVec3& InRelativeContactPosition0,
+		const FSolverVec3& InRelativeContactPosition1,
+		const FSolverVec3& InWorldContactNormal,
+		const FSolverVec3& InWorldContactTangentU,
+		const FSolverVec3& InWorldContactTangentV,
+		const FSolverReal InWorldContactDeltaNormal,
+		const FSolverReal InWorldContactDeltaTangentU,
+		const FSolverReal InWorldContactDeltaTangentV)
 	{
 		State.ManifoldPoints[ManifoldPoiontIndex].InitContact(
 			FSolverReal(Dt),
 			State.SolverBodies[0], 
 			State.SolverBodies[1], 
-			InWorldAnchorPoint0, 
-			InWorldAnchorPoint1, 
-			FSolverVec3(InWorldContactNormal),
-			bAddVelocityToFriction);
+			InRelativeContactPosition0,
+			InRelativeContactPosition1,
+			InWorldContactNormal,
+			InWorldContactTangentU,
+			InWorldContactTangentV,
+			InWorldContactDeltaNormal,
+			InWorldContactDeltaTangentU,
+			InWorldContactDeltaTangentV);
 	}
 
 	void FPBDCollisionSolver::InitMaterial(
 		const int32 ManifoldPoiontIndex,
-		const FReal InRestitution,
-		const FReal InRestitutionVelocityThreshold)
+		const FSolverReal InRestitution,
+		const FSolverReal InRestitutionVelocityThreshold)
 	{
 		State.ManifoldPoints[ManifoldPoiontIndex].InitMaterial(
 			State.SolverBodies[0], 
 			State.SolverBodies[1], 
-			FSolverReal(InRestitution), 
-			FSolverReal(InRestitutionVelocityThreshold));
+			InRestitution, 
+			InRestitutionVelocityThreshold);
 	}
 
 	void FPBDCollisionSolver::SolveVelocityAverage(const FSolverReal Dt)
diff --git a/Engine/Source/Runtime/Experimental/Chaos/Private/Chaos/Collision/PBDCollisionSolver.h b/Engine/Source/Runtime/Experimental/Chaos/Private/Chaos/Collision/PBDCollisionSolver.h
--- a/Engine/Source/Runtime/Experimental/Chaos/Private/Chaos/Collision/PBDCollisionSolver.h
+++ b/Engine/Source/Runtime/Experimental/Chaos/Private/Chaos/Collision/PBDCollisionSolver.h
@@ -46,10 +46,14 @@ namespace Chaos
 			const FSolverReal Dt,
 			const FConstraintSolverBody& Body0,
 			const FConstraintSolverBody& Body1,
-			const FVec3& InWorldAnchorPoint0,
-			const FVec3& InWorldAnchorPoint1,
+			const FSolverVec3& InRelativeContactPosition0,
+			const FSolverVec3& InRelativeContactPosition1,
 			const FSolverVec3& InWorldContactNormal,
-			const bool bAddVelocityToFriction);
+			const FSolverVec3& InWorldContactTangentU,
+			const FSolverVec3& InWorldContactTangentV,
+			const FSolverReal InWorldContactDeltaNormal,
+			const FSolverReal InWorldContactDeltaTangentU,
+			const FSolverReal InWorldContactDeltaTangentV);
 
 		/**
 		 * @brief Initialize the material related properties of the contact
@@ -60,18 +64,6 @@ namespace Chaos
 			const FSolverReal InRestitution,
 			const FSolverReal InRestitutionVelocityThreshold);
 
-		/**
-		 * @brief Update the world-space relative contact points based on current body transforms and body-space contact positions
-		*/
-		void UpdateContact(
-			const FSolverReal Dt,
-			const FConstraintSolverBody& Body0,
-			const FConstraintSolverBody& Body1,
-			const FVec3& InWorldAnchorPoint0,
-			const FVec3& InWorldAnchorPoint1,
-			const FSolverVec3& InWorldContactNormal,
-			const bool bAddVelocityToFriction);
-
 		/**
 		 * @brief Update the cached mass properties based on the current body transforms
 		*/
@@ -170,16 +162,16 @@ namespace Chaos
 		FSolverReal DynamicFriction() const { return State.DynamicFriction; }
 		FSolverReal VelocityFriction() const { return State.VelocityFriction; }
 
-		void SetFriction(const FReal InStaticFriction, const FReal InDynamicFriction, const FReal InVelocityFriction)
+		void SetFriction(const FSolverReal InStaticFriction, const FSolverReal InDynamicFriction, const FSolverReal InVelocityFriction)
 		{
-			State.StaticFriction = FSolverReal(InStaticFriction);
-			State.DynamicFriction = FSolverReal(InDynamicFriction);
-			State.VelocityFriction = FSolverReal(InVelocityFriction);
+			State.StaticFriction = InStaticFriction;
+			State.DynamicFriction = InDynamicFriction;
+			State.VelocityFriction = InVelocityFriction;
 		}
 
-		void SetStiffness(const FReal InStiffness)
+		void SetStiffness(const FSolverReal InStiffness)
 		{
-			State.Stiffness = FSolverReal(InStiffness);
+			State.Stiffness = InStiffness;
 		}
 
 		void SetSolverBodies(FSolverBody* SolverBody0, FSolverBody* SolverBody1)
@@ -213,16 +205,20 @@ namespace Chaos
 
 		void InitContact(
 			const int32 ManifoldPoiontIndex,
-			const FReal Dt,
-			const FVec3& InWorldAnchorPoint0,
-			const FVec3& InWorldAnchorPoint1,
-			const FVec3& InWorldContactNormal,
-			const bool bAddVelocityToFriction);
+			const FSolverReal Dt,
+			const FSolverVec3& InRelativeContactPosition0,
+			const FSolverVec3& InRelativeContactPosition1,
+			const FSolverVec3& InWorldContactNormal,
+			const FSolverVec3& InWorldContactTangentU,
+			const FSolverVec3& InWorldContactTangentV,
+			const FSolverReal InWorldContactDeltaNormal,
+			const FSolverReal InWorldContactDeltaTangentU,
+			const FSolverReal InWorldContactDeltaTangentV);
 
 		void InitMaterial(
 			const int32 ManifoldPoiontIndex,
-			const FReal InRestitution,
-			const FReal InRestitutionVelocityThreshold);
+			const FSolverReal InRestitution,
+			const FSolverReal InRestitutionVelocityThreshold);
 
 		/**
 		 * @brief Get the first (decaorated) solver body
diff --git a/Engine/Source/Runtime/Experimental/Chaos/Private/Chaos/Collision/SolverCollisionContainer.cpp b/Engine/Source/Runtime/Experimental/Chaos/Private/Chaos/Collision/SolverCollisionContainer.cpp
--- a/Engine/Source/Runtime/Experimental/Chaos/Private/Chaos/Collision/SolverCollisionContainer.cpp
+++ b/Engine/Source/Runtime/Experimental/Chaos/Private/Chaos/Collision/SolverCollisionContainer.cpp
@@ -70,8 +70,8 @@ namespace Chaos
 
 			// Find the solver bodies for the particles we constrain. This will add them to the container
 			// if they aren't there already, and ensure that they are populated with the latest data.
-			FSolverBody* Body0 = SolverBodyContainer.FindOrAdd(Constraint->Particle[0]);
-			FSolverBody* Body1 = SolverBodyContainer.FindOrAdd(Constraint->Particle[1]);
+			FSolverBody* Body0 = SolverBodyContainer.FindOrAdd(Constraint->GetParticle0());
+			FSolverBody* Body1 = SolverBodyContainer.FindOrAdd(Constraint->GetParticle1());
 
 			Body0->SetLevel(Particle0Level);
 			Body1->SetLevel(Particle1Level);
@@ -79,53 +79,111 @@ namespace Chaos
 			// Friction values. Static and Dynamic friction are applied in the position solve for most shapes.
 			// For quadratic shapes, we run dynamic friction in the velocity solve for better rolling behaviour.
 			// @todo(chaos): fix static/dynamic friction for quadratic shapes
-			const FReal StaticFriction = FMath::Max(Constraint->Manifold.AngularFriction, Constraint->Manifold.Friction);
-			const FReal DynamicFriction = Constraint->Manifold.Friction;
+			const FSolverReal StaticFriction = FSolverReal(Constraint->GetStaticFriction());
+			const FSolverReal DynamicFriction = FSolverReal(Constraint->GetDynamicFriction());
 			const bool bIsQuadratic = Constraint->HasQuadraticShape();
 			if (!bIsQuadratic)
 			{
-				Solver.SetFriction(StaticFriction, DynamicFriction, FReal(0));
+				Solver.SetFriction(StaticFriction, DynamicFriction, FSolverReal(0));
 			}
 			else
 			{
 				Solver.SetFriction(StaticFriction, FReal(0), DynamicFriction);
 			}
 
-			Solver.SetStiffness(Constraint->GetStiffness());
+			Solver.SetStiffness(FSolverReal(Constraint->GetStiffness()));
 
 			Solver.SetSolverBodies(Body0, Body1);
 
-			GatherManifoldPoints(Dt);
+			GatherManifoldPoints(Dt, Body0, Body1);
 
 			// We should try to remove this - the Constraint should not need to know about solver objects
 			Constraint->SetSolverBodies(Body0, Body1);
 		}
 
-		void GatherManifoldPoints(const FReal Dt)
+		void GatherManifoldPoints(
+			const FReal InDt,
+			FSolverBody* Body0,
+			FSolverBody* Body1)
 		{
+			FSolverReal Dt = FSolverReal(InDt);
+
 			// We handle incremental manifolds by just collecting any new contacts
 			const int32 BeginPointIndex = Solver.NumManifoldPoints();
 			const int32 EndPointIndex = Solver.SetNumManifoldPoints(Constraint->GetManifoldPoints().Num());
 
-			const FReal RestitutionVelocityThreshold = Constraint->Manifold.RestitutionThreshold * Dt;
-			const FReal Restitution = Constraint->Manifold.Restitution;
+			const FSolverReal RestitutionVelocityThreshold = FSolverReal(Constraint->GetRestitutionThreshold()) * Dt;
+			const FSolverReal Restitution = FSolverReal(Constraint->GetRestitution());
+
+			const FRigidTransform3& ShapeWorldTransform0 = Constraint->GetShapeWorldTransform0();
+			const FRigidTransform3& ShapeWorldTransform1 = Constraint->GetShapeWorldTransform1();
 
 			for (int32 ManifoldPointIndex = BeginPointIndex; ManifoldPointIndex < EndPointIndex; ++ManifoldPointIndex)
 			{
 				TArrayView<FManifoldPoint> ManifoldPoints = Constraint->GetManifoldPoints();
-				const FManifoldPoint& ManifoldPoint = ManifoldPoints[ManifoldPointIndex];
-
-				// If we didn't have friction data restored, we estimate the static friction target by adding velocity to the contact error				
-				const bool bAddVelocityToFriction = !ManifoldPoint.Flags.bWasFrictionRestored;
+				FManifoldPoint& ManifoldPoint = ManifoldPoints[ManifoldPointIndex];
+
+				const FVec3 WorldContactPoint0 = ShapeWorldTransform0.TransformPositionNoScale(ManifoldPoint.ContactPoint.ShapeContactPoints[0]);
+				const FVec3 WorldContactPoint1 = ShapeWorldTransform1.TransformPositionNoScale(ManifoldPoint.ContactPoint.ShapeContactPoints[1]);
+				const FVec3 WorldContactPoint = FReal(0.5) * (WorldContactPoint0 + WorldContactPoint1);
+
+				const FSolverVec3 WorldContactNormal = FSolverVec3(ShapeWorldTransform1.TransformVectorNoScale(ManifoldPoint.ContactPoint.ShapeContactNormal));
+				const FSolverVec3 RelativeContactPosition0 = FSolverVec3(WorldContactPoint - Body0->P());
+				const FSolverVec3 RelativeContactPosition1 = FSolverVec3(WorldContactPoint - Body1->P());
+
+				// If we have contact data from a previous tick, use it to calculate the lateral position delta we need
+				// to apply to move the contacts back to their original relative locations (i.e., to enforce static friction)
+				// @todo(chaos): we should not be writing back to the constraint here - find a better way to update the friction anchor. See FPBDCollisionConstraint::SetSolverResults
+				FSolverVec3 WorldFrictionDelta = FVec3(0);
+				const FSavedManifoldPoint* SavedManifoldPoint = Constraint->FindSavedManifoldPoint(ManifoldPoint);
+				if (SavedManifoldPoint != nullptr)
+				{
+					const FSolverVec3 FrictionDelta0 = FSolverVec3(SavedManifoldPoint->ShapeContactPoints[0] - ManifoldPoint.ContactPoint.ShapeContactPoints[0]);
+					const FSolverVec3 FrictionDelta1 = FSolverVec3(SavedManifoldPoint->ShapeContactPoints[1] - ManifoldPoint.ContactPoint.ShapeContactPoints[1]);
+					WorldFrictionDelta = ShapeWorldTransform0.TransformVectorNoScale(FrictionDelta0) - ShapeWorldTransform1.TransformVectorNoScale(FrictionDelta1);
+
+					ManifoldPoint.ShapeAnchorPoints[0] = SavedManifoldPoint->ShapeContactPoints[0];
+					ManifoldPoint.ShapeAnchorPoints[1] = SavedManifoldPoint->ShapeContactPoints[1];
+				}
+				else
+				{
+					const FSolverVec3 ContactVel0 = Body0->V() + FSolverVec3::CrossProduct(Body0->W(), RelativeContactPosition0);
+					const FSolverVec3 ContactVel1 = Body1->V() + FSolverVec3::CrossProduct(Body1->W(), RelativeContactPosition1);
+					const FSolverVec3 ContactVel = ContactVel0 - ContactVel1;
+					WorldFrictionDelta = ContactVel * Dt;
+
+					ManifoldPoint.ShapeAnchorPoints[0] = ManifoldPoint.ContactPoint.ShapeContactPoints[0];
+					ManifoldPoint.ShapeAnchorPoints[1] = ManifoldPoint.ContactPoint.ShapeContactPoints[1];
+				}
+
+				// World-space contact tangents. We are treating the normal as the constraint-space Z axis
+				// and the Tangent U and V as the constraint-space X and Y axes respectively
+				FSolverVec3 WorldContactTangentU = FSolverVec3::CrossProduct(FSolverVec3(0, 1, 0), WorldContactNormal);
+				if (!WorldContactTangentU.Normalize(FSolverReal(KINDA_SMALL_NUMBER)))
+				{
+					WorldContactTangentU = FSolverVec3::CrossProduct(FSolverVec3(1, 0, 0), WorldContactNormal);
+					WorldContactTangentU = WorldContactTangentU.GetUnsafeNormal();
+				}
+				const FSolverVec3 WorldContactTangentV = FSolverVec3::CrossProduct(WorldContactNormal, WorldContactTangentU);
+
+				// The contact point error we are trying to correct in this solver
+				const FSolverVec3 WorldContactDelta = FSolverVec3(WorldContactPoint0 - WorldContactPoint1);
+				const FSolverReal WorldContactDeltaNormal = FSolverVec3::DotProduct(WorldContactDelta, WorldContactNormal);
+				const FSolverReal WorldContactDeltaTangentU = FSolverVec3::DotProduct(WorldContactDelta + WorldFrictionDelta, WorldContactTangentU);
+				const FSolverReal WorldContactDeltaTangentV = FSolverVec3::DotProduct(WorldContactDelta + WorldFrictionDelta, WorldContactTangentV);
 
 				// Initialize the structural data in the contact (relative contact points, contact mass etc)
 				Solver.InitContact(
 					ManifoldPointIndex,
 					Dt,
-					ManifoldPoint.WorldContactPoints[0],
-					ManifoldPoint.WorldContactPoints[1],
-					ManifoldPoint.ContactPoint.Normal,
-					bAddVelocityToFriction);
+					RelativeContactPosition0,
+					RelativeContactPosition1,
+					WorldContactNormal,
+					WorldContactTangentU,
+					WorldContactTangentV,
+					WorldContactDeltaNormal,
+					WorldContactDeltaTangentU,
+					WorldContactDeltaTangentV);
 
 				// Initialize the material properties (restitution and friction related)
 				Solver.InitMaterial(
@@ -319,7 +377,7 @@ namespace Chaos
 			if (CollisionSolver.IsIncrementalManifold())
 			{
 				Collisions::Update(*CollisionSolver.GetConstraint(), Dt);
-				CollisionSolver.GatherManifoldPoints(Dt);
+				CollisionSolver.GatherManifoldPoints(Dt, &CollisionSolver.GetSolver().SolverBody0().SolverBody(), &CollisionSolver.GetSolver().SolverBody1().SolverBody());
 			}
 			if (bApplyStaticFriction)
 			{
diff --git a/Engine/Source/Runtime/Experimental/Chaos/Private/Chaos/Collision/SphereConvexContactPoint.cpp b/Engine/Source/Runtime/Experimental/Chaos/Private/Chaos/Collision/SphereConvexContactPoint.cpp
--- a/Engine/Source/Runtime/Experimental/Chaos/Private/Chaos/Collision/SphereConvexContactPoint.cpp
+++ b/Engine/Source/Runtime/Experimental/Chaos/Private/Chaos/Collision/SphereConvexContactPoint.cpp
@@ -80,8 +80,6 @@ namespace Chaos
 				ContactPoint.ShapeContactPoints[0] = PosConvex;
 				ContactPoint.ShapeContactPoints[1] = PosSphere;
 				ContactPoint.ShapeContactNormal = -NormalConvex;
-				ContactPoint.Location = ConvexTransform.TransformPosition(ContactPoint.ShapeContactPoints[0]);
-				ContactPoint.Normal = ConvexTransform.TransformVectorNoScale(-NormalConvex);
 				ContactPoint.Phi = Phi;
 			}
 		}
diff --git a/Engine/Source/Runtime/Experimental/Chaos/Private/Chaos/CollisionOneShotManifolds.cpp b/Engine/Source/Runtime/Experimental/Chaos/Private/Chaos/CollisionOneShotManifolds.cpp
--- a/Engine/Source/Runtime/Experimental/Chaos/Private/Chaos/CollisionOneShotManifolds.cpp
+++ b/Engine/Source/Runtime/Experimental/Chaos/Private/Chaos/CollisionOneShotManifolds.cpp
@@ -133,204 +133,7 @@ namespace Chaos
 			const FReal Dt,
 			FPBDCollisionConstraint& Constraint)
 		{
-			if (bChaos_Collision_Manifold_BoxAsConvex)
-			{
-				ConstructConvexConvexOneShotManifold(Box1, Box1Transform, Box2, Box2Transform, Dt, Constraint);
-				return;
-			}
-
-			const uint32 SpaceDimension = 3;
-
-			// We only build one shot manifolds once
-			// All boxes are prescaled
-			check(Constraint.GetManifoldPoints().Num() == 0);
-			check(Box1Transform.GetScale3D() == FVec3(1.0f, 1.0f, 1.0f));
-			check(Box2Transform.GetScale3D() == FVec3(1.0f, 1.0f, 1.0f));
-
-			const uint32 MaxContactPointCount = 8;
-			uint32 ContactPointCount = 0;
-
-			// Use GJK only once
-			FContactPoint GJKContactPoint = BoxBoxContactPoint(Box1, Box2, Box1Transform, Box2Transform, Constraint.Manifold.RestitutionPadding);
-
-			FRigidTransform3 Box1TransformCenter = Box1Transform;
-			Box1TransformCenter.SetTranslation(Box1Transform.TransformPositionNoScale(Box1.GetCenter()));
-			FRigidTransform3 Box2TransformCenter = Box2Transform;
-			Box2TransformCenter.SetTranslation(Box2Transform.TransformPositionNoScale(Box2.GetCenter()));
-
-			// GJK does not give us any face information yet, so find the best reference face here
-			uint32 BestFaceNormalAxisBox1 = 0; // Note: Face normals are axis aligned due to coordinates. This is an element of {0 , 1, 2}
-			int32 BestFaceNormalAxisDirectionBox1 = 1; // {-1, 1}
-			FReal BestFaceNormalSizeInDirectionBox1 = -1.0f;
-			const FVec3 SeparationDirectionLocalBox1 = Box1TransformCenter.InverseTransformVectorNoScale(GJKContactPoint.Normal);   // Todo: Check if we can have Non uniform scale here
-			// Box1: Iterating through 2 faces at a time here
-			for (uint32 FaceNormalAxis = 0; FaceNormalAxis < SpaceDimension; FaceNormalAxis++)
-			{
-				const FReal AbsSeparationDirectionDotFaceNormal = FMath::Abs(SeparationDirectionLocalBox1[FaceNormalAxis]);
-				if (AbsSeparationDirectionDotFaceNormal > BestFaceNormalSizeInDirectionBox1)
-				{
-					BestFaceNormalAxisBox1 = FaceNormalAxis;
-					BestFaceNormalSizeInDirectionBox1 = AbsSeparationDirectionDotFaceNormal;
-					BestFaceNormalAxisDirectionBox1 = SeparationDirectionLocalBox1[FaceNormalAxis] >= 0.0f ? -1 : 1;
-				}
-			}
-
-			// Now for Box2
-			uint32 BestFaceNormalAxisBox2 = 0; // Note: Face normals are axis aligned due to coordinates. This is a n element of {0, 1, 2}
-			int32 BestFaceNormalAxisDirectionBox2 = 1; // {-1, 1}
-			FReal BestFaceNormalSizeInDirectionBox2 = -1.0f;
-			const FVec3 SeparationDirectionLocalBox2 = Box2TransformCenter.InverseTransformVectorNoScale(GJKContactPoint.Normal);   // Todo: Check if we can have Non uniform scale here
-			for (uint32 FaceNormalAxis = 0; FaceNormalAxis < SpaceDimension; FaceNormalAxis++)
-			{
-				FVec3 FaceNormal(FVec3::ZeroVector);
-				FaceNormal[FaceNormalAxis] = 1.0f;
-				const FReal SeparationDirectionDotFaceNormal = FMath::Abs(SeparationDirectionLocalBox2[FaceNormalAxis]);
-
-				if (SeparationDirectionDotFaceNormal > BestFaceNormalSizeInDirectionBox2)
-				{
-					BestFaceNormalAxisBox2 = FaceNormalAxis;
-					BestFaceNormalSizeInDirectionBox2 = SeparationDirectionDotFaceNormal;
-					BestFaceNormalAxisDirectionBox2 = SeparationDirectionLocalBox2[FaceNormalAxis] >= 0.0f ? 1 : -1;  // Note opposite of box1
-				}
-			}
-
-			const FReal SmallBiasToPreventFeatureFlipping = 0.002f; // This improves frame coherence by penalizing box 1 in favour of box 2
-			bool ReferenceFaceBox1 = true; // Is the reference face on box1 or box2?
-			if (BestFaceNormalSizeInDirectionBox2 + SmallBiasToPreventFeatureFlipping > BestFaceNormalSizeInDirectionBox1)
-			{
-				ReferenceFaceBox1 = false;
-			}
-
-			// Is this a vertex-plane or edge-edge contact? 
-			const FReal PlaneContactNormalEpsilon = Chaos_Collision_Manifold_PlaneContactNormalEpsilon;
-			const bool bIsPlaneContact = FMath::IsNearlyEqual(BestFaceNormalSizeInDirectionBox1, (FReal)1., PlaneContactNormalEpsilon) || FMath::IsNearlyEqual(BestFaceNormalSizeInDirectionBox2, (FReal)1., PlaneContactNormalEpsilon);
-
-			// For edge-edge contacts we find the edges involved and project the contact onto the edges
-			if (!bIsPlaneContact)
-			{
-				FVec3 ShapeEdgePos1 = Box1.GetClosestEdgePosition(INDEX_NONE, GJKContactPoint.ShapeContactPoints[0]);
-				FVec3 ShapeEdgePos2 = Box2.GetClosestEdgePosition(INDEX_NONE, GJKContactPoint.ShapeContactPoints[1]);
-				FVec3 EdgePos1 = Box1Transform.TransformPosition(ShapeEdgePos1);
-				FVec3 EdgePos2 = Box2Transform.TransformPosition(ShapeEdgePos2);
-				FReal EdgePhi = FVec3::DotProduct(EdgePos1 - EdgePos2, GJKContactPoint.Normal);
-
-				GJKContactPoint.ShapeContactPoints[0] = ShapeEdgePos1;
-				GJKContactPoint.ShapeContactPoints[1] = ShapeEdgePos2;
-				GJKContactPoint.Phi = EdgePhi;
-				GJKContactPoint.Location = 0.5f * (EdgePos1 + EdgePos2);
-
-				Constraint.AddOneshotManifoldContact(GJKContactPoint);
-				return;
-			}
-
-
-			// For vertex-plane contacts, we use a convex face as the manifold plane
-			// Setup pointers to other box and reference box
-			const FRigidTransform3* RefBoxTM;
-			const FRigidTransform3* OtherBoxTM;
-			const FImplicitBox3* RefBox;
-			const FImplicitBox3* OtherBox;
-
-			if (ReferenceFaceBox1)
-			{
-				RefBoxTM = &Box1TransformCenter;
-				OtherBoxTM = &Box2TransformCenter;
-				RefBox = &Box1;
-				OtherBox = &Box2;
-			}
-			else
-			{
-				RefBoxTM = &Box2TransformCenter;
-				OtherBoxTM = &Box1TransformCenter;
-				RefBox = &Box2;
-				OtherBox = &Box1;
-			}
-
-			// Populate the clipped vertices by the other face's vertices
-
-			// Populate initial clipping vertices with a face from the other box
-			FVec3 otherBoxHalfExtents = 0.5f * OtherBox->Extents();
-			uint32 ConstantCoordinateIndex = ReferenceFaceBox1 ? BestFaceNormalAxisBox2 : BestFaceNormalAxisBox1;
-			FReal ConstantCoordinate = otherBoxHalfExtents[ConstantCoordinateIndex] * (FReal)(ReferenceFaceBox1 ? BestFaceNormalAxisDirectionBox2 : BestFaceNormalAxisDirectionBox1);
-
-			uint32 VariableCoordinateIndices[2];
-			FReal VariableCoordinates[2];
-
-			uint32 VariableCoordinateCount = 0;
-			for (uint32 Coordinate = 0; Coordinate < 3; ++Coordinate)
-			{
-				if (Coordinate != ConstantCoordinateIndex)
-				{
-					VariableCoordinateIndices[VariableCoordinateCount] = Coordinate;
-					VariableCoordinates[VariableCoordinateCount] = otherBoxHalfExtents[Coordinate];
-					++VariableCoordinateCount;
-				}
-			}
-
-			ContactPointCount = 4; // Number of face vertices
-			const uint32 GreyCode[4] = { 0, 1, 3, 2 }; // Grey code to make sure we add vertices in correct order
-			FVec3 ClippedVertices[MaxContactPointCount];
-			// Add the vertices in an order that will form a closed loop
-			const FRigidTransform3 BoxOtherToRef = OtherBoxTM->GetRelativeTransform(*RefBoxTM);
-			for (uint32 Vertex = 0; Vertex < ContactPointCount; Vertex++)
-			{
-				ClippedVertices[Vertex][ConstantCoordinateIndex] = ConstantCoordinate;
-				ClippedVertices[Vertex][VariableCoordinateIndices[0]] = (GreyCode[Vertex] & (1 << 0)) ? VariableCoordinates[0] : -VariableCoordinates[0];
-				ClippedVertices[Vertex][VariableCoordinateIndices[1]] = (GreyCode[Vertex] & (1 << 1)) ? VariableCoordinates[1] : -VariableCoordinates[1];
-				ClippedVertices[Vertex] = BoxOtherToRef.TransformPositionNoScale(ClippedVertices[Vertex]);
-			}
-
-			// Now clip against all planes that belong to the reference plane's, edges
-			FVec3 ClippedVertices2[MaxContactPointCount]; // We will use a double buffer as an optimization
-			FVec3* VertexBuffer1 = ClippedVertices;
-			FVec3* VertexBuffer2 = ClippedVertices2;
-
-			FVec3 refBoxHalfExtents = 0.5f * RefBox->Extents();
-			uint32 RefPlaneCoordinateIndex = ReferenceFaceBox1 ? BestFaceNormalAxisBox1 : BestFaceNormalAxisBox2;
-			for (uint32 Coordinate = 0; Coordinate < 3; ++Coordinate)
-			{
-				if (Coordinate != RefPlaneCoordinateIndex)
-				{
-					ContactPointCount = BoxBoxClipVerticesAgainstPlane(VertexBuffer1, VertexBuffer2, ContactPointCount, static_cast<int32>(Coordinate), refBoxHalfExtents[Coordinate]);
-					ContactPointCount = BoxBoxClipVerticesAgainstPlane(VertexBuffer2, VertexBuffer1, ContactPointCount, static_cast<int32>(Coordinate), -refBoxHalfExtents[Coordinate]);
-				}
-			}
-
-			// Reduce number of contacts to a maximum of 4
-			if (ContactPointCount > 4)
-			{
-				FRotation3 RotateSeperationToZ = FRotation3::FromRotatedVector(ReferenceFaceBox1 ? SeparationDirectionLocalBox1 : SeparationDirectionLocalBox2, FVec3(0.0f, 0.0f, 1.0f));
-				for (uint32 ContactPointIndex = 0; ContactPointIndex < ContactPointCount; ++ContactPointIndex)
-				{
-					ClippedVertices[ContactPointIndex] = RotateSeperationToZ * ClippedVertices[ContactPointIndex];
-				}
-
-				ContactPointCount = ReduceManifoldContactPoints(ClippedVertices, ContactPointCount);
-
-				for (uint32 ContactPointIndex = 0; ContactPointIndex < ContactPointCount; ++ContactPointIndex)
-				{
-					ClippedVertices[ContactPointIndex] = RotateSeperationToZ.Inverse() * ClippedVertices[ContactPointIndex];
-				}
-			}
-
-			// Generate the contact points from the clipped vertices
-			for (uint32 ContactPointIndex = 0; ContactPointIndex < ContactPointCount; ++ContactPointIndex)
-			{
-				FContactPoint ContactPoint;
-				const FVec3 VertexInReferenceCubeCoordinates = ClippedVertices[ContactPointIndex];
-				FVec3 PointProjectedOntoReferenceFace = VertexInReferenceCubeCoordinates;
-				PointProjectedOntoReferenceFace[RefPlaneCoordinateIndex] = refBoxHalfExtents[RefPlaneCoordinateIndex] * (FReal)(ReferenceFaceBox1 ? BestFaceNormalAxisDirectionBox1 : BestFaceNormalAxisDirectionBox2);
-				FVec3 ClippedPointInOtherCubeCoordinates = BoxOtherToRef.InverseTransformPositionNoScale(VertexInReferenceCubeCoordinates);
-
-				ContactPoint.ShapeContactPoints[0] = ReferenceFaceBox1 ? PointProjectedOntoReferenceFace + RefBox->GetCenter() : ClippedPointInOtherCubeCoordinates + OtherBox->GetCenter();
-				ContactPoint.ShapeContactPoints[1] = ReferenceFaceBox1 ? ClippedPointInOtherCubeCoordinates + OtherBox->GetCenter() : PointProjectedOntoReferenceFace + RefBox->GetCenter();
-				ContactPoint.ShapeContactNormal = SeparationDirectionLocalBox2;
-				ContactPoint.Location = RefBoxTM->TransformPositionNoScale(PointProjectedOntoReferenceFace);
-				ContactPoint.Normal = GJKContactPoint.Normal;
-				ContactPoint.Phi = FVec3::DotProduct(PointProjectedOntoReferenceFace - VertexInReferenceCubeCoordinates, ReferenceFaceBox1 ? SeparationDirectionLocalBox1 : -SeparationDirectionLocalBox2);
-
-				Constraint.AddOneshotManifoldContact(ContactPoint);
-			}
+			ConstructConvexConvexOneShotManifold(Box1, Box1Transform, Box2, Box2Transform, Dt, Constraint);
 		}
 
 		/////////////////////////////
@@ -638,10 +441,6 @@ namespace Chaos
 				Contact.ShapeContactPoints[0] = ClosestA;
 				Contact.ShapeContactPoints[1] = ClosestB;
 				Contact.ShapeContactNormal = -NormalB;	// We want normal pointing from B to A
-				const FVec3 WorldLocationA = ATM.TransformPositionNoScale(ClosestA);
-				const FVec3 WorldLocationB = BTM.TransformPositionNoScale(ClosestB);
-				Contact.Location = FReal(0.5) * (WorldLocationA + WorldLocationB);
-				Contact.Normal = -ATM.TransformVectorNoScale(NormalA);
 				Contact.Phi = -Penetration;
 			}
 
@@ -679,6 +478,44 @@ namespace Chaos
 			}
 		}
 
+		// Specialization for scaled convex. Avoids the instantiation of a scaled plane object
+		// which almost doubles the cost of the function.
+		template <>
+		void FindBestPlane<TImplicitObjectScaled<FConvex>>(
+			const TImplicitObjectScaled<FConvex>& ScaledConvex,
+			const FVec3& X,
+			const FVec3& N,
+			const FReal MaxDistance,
+			const int32 PlaneIndex,
+			int32& BestPlaneIndex,
+			FReal& BestPlaneDot)
+		{
+			const FConvex* UnscaledConvex = ScaledConvex.GetInnerObject()->template GetObject<FConvex>();
+			const TPlaneConcrete<FReal, 3>& UnscaledPlane = UnscaledConvex->GetPlane(PlaneIndex);
+
+			const FVec3 ScaledPlaneX = FVec3(UnscaledPlane.X()) * ScaledConvex.GetScale();
+			FVec3 ScaledPlaneN = FVec3(UnscaledPlane.Normal()) / ScaledConvex.GetScale();
+			if (ScaledPlaneN.Normalize())
+			{
+				// Ignore planes that do not oppose N
+				const FReal PlaneNormalDotN = FVec3::DotProduct(N, ScaledPlaneN);
+				if (PlaneNormalDotN <= -SMALL_NUMBER)
+				{
+					// Reject planes farther than MaxDistance
+					const FReal PlaneDistance = FVec3::DotProduct(X - ScaledPlaneX, ScaledPlaneN);
+					if (FMath::Abs(PlaneDistance) <= MaxDistance)
+					{
+						// Keep the most opposing plane
+						if (PlaneNormalDotN < BestPlaneDot)
+						{
+							BestPlaneDot = PlaneNormalDotN;
+							BestPlaneIndex = PlaneIndex;
+						}
+					}
+				}
+			}
+		}
+
 		// Select one of the planes on the convex to use as the contact plane, given an estimated contact position and opposing 
 		// normal from GJK with margins (which gives the shapes rounded corners/edges).
 		template <typename ConvexImplicitType>
@@ -784,14 +621,16 @@ namespace Chaos
 
 			PHYSICS_CSV_CUSTOM_EXPENSIVE(PhysicsCounters, NumManifoldsCreated, 1, ECsvCustomStatOp::Accumulate);
 
+			const FRigidTransform3 Convex2ToConvex1Transform = Convex2Transform.GetRelativeTransformNoScale(Convex1Transform);
+
 			// @todo(chaos): get the vertex index from GJK and use to to get the plane
-			const FVec3 SeparationDirectionLocalConvex1 = Convex1Transform.InverseTransformVectorNoScale(GJKContactPoint.Normal);
+			const FVec3 SeparationDirectionLocalConvex1 = Convex2ToConvex1Transform.TransformVectorNoScale(GJKContactPoint.ShapeContactNormal);
 			const int32 MostOpposingPlaneIndexConvex1 = SelectContactPlane(Convex1, GJKContactPoint.ShapeContactPoints[0], SeparationDirectionLocalConvex1, Margin1, VertexIndexA);
 			const TPlaneConcrete<FReal, 3> BestPlaneConvex1 = Convex1.GetPlane(MostOpposingPlaneIndexConvex1);
 			const FReal BestPlaneDotNormalConvex1 = !bConvex1IsCapsule ? FMath::Abs(FVec3::DotProduct(-SeparationDirectionLocalConvex1, BestPlaneConvex1.Normal())) : -FLT_MAX;
 
 			// Now for Convex2
-			const FVec3 SeparationDirectionLocalConvex2 = Convex2Transform.InverseTransformVectorNoScale(GJKContactPoint.Normal);
+			const FVec3 SeparationDirectionLocalConvex2 = GJKContactPoint.ShapeContactNormal;
 			const int32 MostOpposingPlaneIndexConvex2 = SelectContactPlane(Convex2, GJKContactPoint.ShapeContactPoints[1], -SeparationDirectionLocalConvex2, Margin2, VertexIndexB);
 			const TPlaneConcrete<FReal, 3> BestPlaneConvex2 = Convex2.GetPlane(MostOpposingPlaneIndexConvex2);
 			const FReal BestPlaneDotNormalConvex2 = FMath::Abs(FVec3::DotProduct(SeparationDirectionLocalConvex2, BestPlaneConvex2.Normal()));
@@ -826,16 +665,14 @@ namespace Chaos
 					ShapeEdgePos1 -= Margin1 * SeparationDirectionLocalConvex1;
 				}
 
-				const FVec3 EdgePos1 = Convex1Transform.TransformPosition(ShapeEdgePos1);
-				const FVec3 EdgePos2 = Convex2Transform.TransformPosition(ShapeEdgePos2);
-				const FReal EdgePhi = FVec3::DotProduct(EdgePos1 - EdgePos2, GJKContactPoint.Normal);
-				const FVec3 WorldPos = FReal(0.5) * (EdgePos1 + EdgePos2);
-				const FVec3& WorldNormal = GJKContactPoint.Normal;
+				const FVec3 EdgePos1In2 = Convex2ToConvex1Transform.InverseTransformPositionNoScale(ShapeEdgePos1);
+				const FVec3 EdgePos2In2 = ShapeEdgePos2;
+				const FReal EdgePhi = FVec3::DotProduct(EdgePos1In2 - EdgePos2In2, GJKContactPoint.ShapeContactNormal);
+				const FVec3 EdgePosIn2 = FReal(0.5) * (EdgePos1In2 + EdgePos2In2);
 
-				GJKContactPoint.ShapeContactPoints[0] = Convex1Transform.InverseTransformPositionNoScale(WorldPos + FReal(0.5) * EdgePhi * WorldNormal);
-				GJKContactPoint.ShapeContactPoints[1] = Convex2Transform.InverseTransformPositionNoScale(WorldPos - FReal(0.5) * EdgePhi * WorldNormal);
+				GJKContactPoint.ShapeContactPoints[0] = Convex2ToConvex1Transform.TransformPositionNoScale(EdgePosIn2 + FReal(0.5) * EdgePhi * GJKContactPoint.ShapeContactNormal);
+				GJKContactPoint.ShapeContactPoints[1] = EdgePosIn2 - FReal(0.5) * EdgePhi * GJKContactPoint.ShapeContactNormal;
 				GJKContactPoint.Phi = EdgePhi;
-				GJKContactPoint.Location = WorldPos;
 				// Normal unchanged from GJK result
 
 				Constraint.AddOneshotManifoldContact(GJKContactPoint);
@@ -862,7 +699,7 @@ namespace Chaos
 				if (ReferenceFaceConvex1)
 				{
 					RefConvexTM = &Convex1Transform;
-					ConvexOtherToRef = Convex2Transform.GetRelativeTransform(Convex1Transform);
+					ConvexOtherToRef = Convex2ToConvex1Transform;
 
 					ClippedVertices = GenerateConvexManifoldClippedVertices(
 						Convex1,
@@ -880,7 +717,7 @@ namespace Chaos
 				else
 				{
 					RefConvexTM = &Convex2Transform;
-					ConvexOtherToRef = Convex1Transform.GetRelativeTransform(Convex2Transform);
+					ConvexOtherToRef = Convex1Transform.GetRelativeTransformNoScale(Convex2Transform);
 
 					ClippedVertices = GenerateConvexManifoldClippedVertices(
 						Convex2,
@@ -941,8 +778,6 @@ namespace Chaos
 					ContactPoint.ShapeContactPoints[0] = ReferenceFaceConvex1 ? PointProjectedOntoReferenceFace : ClippedPointInOtherCoordinates;
 					ContactPoint.ShapeContactPoints[1] = ReferenceFaceConvex1 ? ClippedPointInOtherCoordinates : PointProjectedOntoReferenceFace;
 					ContactPoint.ShapeContactNormal = SeparationDirectionLocalConvex2;
-					ContactPoint.Location = ReferenceFaceConvex1 ? RefConvexTM->TransformPositionNoScale(VertexInReferenceCoordinates) : RefConvexTM->TransformPositionNoScale(PointProjectedOntoReferenceFace);
-					ContactPoint.Normal = GJKContactPoint.Normal;
 					ContactPoint.Phi = FVec3::DotProduct(PointProjectedOntoReferenceFace - VertexInReferenceCoordinates, ReferenceFaceConvex1 ? SeparationDirectionLocalConvex1 : -SeparationDirectionLocalConvex2);				
 
 					Constraint.AddOneshotManifoldContact(ContactPoint);
diff --git a/Engine/Source/Runtime/Experimental/Chaos/Private/Chaos/CollisionOneShotManifoldsMiscShapes.cpp b/Engine/Source/Runtime/Experimental/Chaos/Private/Chaos/CollisionOneShotManifoldsMiscShapes.cpp
--- a/Engine/Source/Runtime/Experimental/Chaos/Private/Chaos/CollisionOneShotManifoldsMiscShapes.cpp
+++ b/Engine/Source/Runtime/Experimental/Chaos/Private/Chaos/CollisionOneShotManifoldsMiscShapes.cpp
@@ -43,7 +43,7 @@ namespace Chaos
 			// @todo(chaos): support manifold maintenance
 			Constraint.ResetActiveManifoldContacts();
 
-			FContactPoint ContactPoint = SphereSphereContactPoint(SphereA, SphereATransform, SphereB, SphereBTransform, Constraint.GetCullDistance(), Constraint.Manifold.RestitutionPadding);
+			FContactPoint ContactPoint = SphereSphereContactPoint(SphereA, SphereATransform, SphereB, SphereBTransform, Constraint.GetCullDistance(), Constraint.GetRestitutionPadding());
 			if (ContactPoint.Phi < Constraint.GetCullDistance())
 			{
 				Constraint.AddOneshotManifoldContact(ContactPoint);
@@ -66,7 +66,7 @@ namespace Chaos
 			// @todo(chaos): support manifold maintenance
 			Constraint.ResetActiveManifoldContacts();
 
-			FContactPoint ContactPoint = SpherePlaneContactPoint(Sphere, SphereTransform, Plane, PlaneTransform, Constraint.Manifold.RestitutionPadding);
+			FContactPoint ContactPoint = SpherePlaneContactPoint(Sphere, SphereTransform, Plane, PlaneTransform, Constraint.GetRestitutionPadding());
 			if (ContactPoint.Phi < Constraint.GetCullDistance())
 			{
 				Constraint.AddOneshotManifoldContact(ContactPoint);
@@ -83,7 +83,7 @@ namespace Chaos
 			// @todo(chaos): support manifold maintenance
 			Constraint.ResetActiveManifoldContacts();
 
-			FContactPoint ContactPoint = SphereBoxContactPoint(Sphere, SphereTransform, Box, BoxTransform, Constraint.Manifold.RestitutionPadding);
+			FContactPoint ContactPoint = SphereBoxContactPoint(Sphere, SphereTransform, Box, BoxTransform, Constraint.GetRestitutionPadding());
 			if (ContactPoint.Phi < Constraint.GetCullDistance())
 			{
 				Constraint.AddOneshotManifoldContact(ContactPoint);
@@ -100,7 +100,7 @@ namespace Chaos
 			// @todo(chaos): support manifold maintenance
 			Constraint.ResetActiveManifoldContacts();
 
-			FContactPoint ContactPoint = SphereCapsuleContactPoint(Sphere, SphereTransform, Capsule, CapsuleTransform, Constraint.Manifold.RestitutionPadding);
+			FContactPoint ContactPoint = SphereCapsuleContactPoint(Sphere, SphereTransform, Capsule, CapsuleTransform, Constraint.GetRestitutionPadding());
 			if (ContactPoint.Phi < Constraint.GetCullDistance())
 			{
 				Constraint.AddOneshotManifoldContact(ContactPoint);
@@ -239,7 +239,6 @@ namespace Chaos
 
 			FContactPoint ContactPoint;
 			ContactPoint.ShapeContactNormal = CapsuleBTransform.InverseTransformVector(Normal);
-			ContactPoint.Normal = Normal;
 
 			auto AddManifoldPoint = [&](FReal ClippedCoord)
 			{
@@ -253,7 +252,6 @@ namespace Chaos
 				{
 					ContactPoint.ShapeContactPoints[0] = CapsuleATransform.InverseTransformPosition(LocationA);
 					ContactPoint.ShapeContactPoints[1] = CapsuleBTransform.InverseTransformPosition(LocationB);
-					ContactPoint.Location = 0.5f * (LocationA + LocationB);
 					ContactPoint.Phi = Phi;
 					Constraint.AddOneshotManifoldContact(ContactPoint);
 				}
diff --git a/Engine/Source/Runtime/Experimental/Chaos/Private/Chaos/CollisionResolution.cpp b/Engine/Source/Runtime/Experimental/Chaos/Private/Chaos/CollisionResolution.cpp
--- a/Engine/Source/Runtime/Experimental/Chaos/Private/Chaos/CollisionResolution.cpp
+++ b/Engine/Source/Runtime/Experimental/Chaos/Private/Chaos/CollisionResolution.cpp
@@ -146,19 +146,6 @@ namespace Chaos
 			static const bool bImmediateUpdate = B_IMMEDIATEUPDATE;
 		};
 
-		// GJKPenetration (EPA) can return arbitrarily wrong normals when Phi approaches zero (it ends up renormalizing
-		// a very small vector) so attempt to fix the normal by calling PhiWithNormal.
-		template<typename T_SHAPE>
-		FORCEINLINE void FixGJKPenetrationNormal(FContactPoint& ContactPoint, const T_SHAPE& Shape1, const FTransform& WorldTransform1)
-		{
-			FVec3 NormalLocal1;
-			const FReal Phi = Shape1.PhiWithNormal(ContactPoint.ShapeContactPoints[1], NormalLocal1);
-
-			ContactPoint.ShapeContactNormal = NormalLocal1;
-			ContactPoint.Normal = WorldTransform1.TransformVectorNoScale(NormalLocal1);
-			ContactPoint.Phi = Phi;
-		}
-
 		/** Determines if body should use CCD. If using CCD, computes Dir and Length of sweep.
 		* If either of the object is CCD-enabled and fast moving, we enable CCD for the constraint.
 		* We compare P-X and BoundingBox().Extents().Min() to determine if an object is fast-moving or not. Here is how this works:
@@ -242,7 +229,7 @@ namespace Chaos
 		{
 			CONDITIONAL_SCOPE_CYCLE_COUNTER(STAT_Collisions_SetSweptConstraintTOI, ConstraintsDetailedStats);
 
-			if (SweptConstraint.Manifold.Phi > 0.0f)
+			if (SweptConstraint.GetPhi() > 0.0f)
 			{
 				SweptConstraint.TimeOfImpact = TNumericLimits<FReal>::Max();
 				return;
@@ -255,7 +242,7 @@ namespace Chaos
 			}
 
 			const FReal MinBounds = Particle->Geometry()->BoundingBox().Extents().Min();
-			const FReal Depth = -SweptConstraint.Manifold.Phi;
+			const FReal Depth = -SweptConstraint.GetPhi();
 			const FReal MaxDepth = MinBounds * CCDAllowedDepthBoundsScale;
 
 			if (TOI <= 0.0f)
@@ -273,11 +260,10 @@ namespace Chaos
 				else if (Depth > MaxDepth)
 				{
 					// Compute ExtraT that will take us from TOI and 0 phi to slightly past TOI w/ phi at maximum allowed depth.
-					const FReal DirDotNormal = FMath::Abs(FVec3::DotProduct(Dir, SweptConstraint.Manifold.Normal));
+					const FReal DirDotNormal = FMath::Abs(FVec3::DotProduct(Dir, SweptConstraint.CalculateWorldContactNormal()));
 					const FReal ExtraT = MaxDepth / (DirDotNormal * Length);
 
 					SweptConstraint.TimeOfImpact = TOI + ExtraT;
-					SweptConstraint.Manifold.Phi = -MaxDepth;
 				}
 			}
 		}
@@ -521,7 +507,7 @@ namespace Chaos
 			}
 			else
 			{
-				UpdateContactPoint(Constraint, BoxBoxContactPoint(Box1, Box2, Box1Transform, Box2Transform, Constraint.Manifold.RestitutionPadding), Dt);
+				UpdateContactPoint(Constraint, BoxBoxContactPoint(Box1, Box2, Box1Transform, Box2Transform, Constraint.GetRestitutionPadding()), Dt);
 			}
 		}
 
@@ -627,9 +613,9 @@ namespace Chaos
 						NumConstraints = 0;
 					}
 					Contact.Phi = NewPhi;
-					Contact.Normal = PlaneTransform.TransformVectorNoScale(Normal);
-					Contact.Location = PlaneTransform.TransformPosition(Corners[i]);
-					PotentialConstraints[NumConstraints++] = Contact.Location;
+					Contact.ShapeContactPoints[1] = Corners[i];
+					Contact.ShapeContactNormal = Normal;
+					PotentialConstraints[NumConstraints++] = Corners[i];
 				}
 			}
 			if (NumConstraints > 1)
@@ -639,12 +625,11 @@ namespace Chaos
 				{
 					AverageLocation += PotentialConstraints[ConstraintIdx];
 				}
-				Contact.Location = AverageLocation / static_cast<FReal>(NumConstraints);
+				Contact.ShapeContactPoints[1] = AverageLocation / static_cast<FReal>(NumConstraints);
 			}
 
-			Contact.ShapeContactPoints[0] = BoxTransform.InverseTransformPosition(Contact.Location);
-			Contact.ShapeContactPoints[1] = PlaneTransform.InverseTransformPosition(Contact.Location - Contact.Phi * Contact.Normal);
-			Contact.ShapeContactNormal = PlaneTransform.InverseTransformVector(Contact.Normal);
+			Contact.ShapeContactPoints[0] = BoxToPlaneTransform.InverseTransformPositionNoScale(Contact.ShapeContactPoints[1]);
+			Contact.ShapeContactPoints[1] = Contact.ShapeContactPoints[1] - Contact.Phi * Contact.ShapeContactNormal;
 
 			return Contact;
 		}
@@ -739,7 +724,7 @@ namespace Chaos
 			}
 			else
 			{
-				UpdateContactPoint(Constraint, SphereSphereContactPoint(Sphere1, Sphere1Transform, Sphere2, Sphere2Transform, Constraint.GetCullDistance(), Constraint.Manifold.RestitutionPadding), Dt);
+				UpdateContactPoint(Constraint, SphereSphereContactPoint(Sphere1, Sphere1Transform, Sphere2, Sphere2Transform, Constraint.GetCullDistance(), Constraint.GetRestitutionPadding()), Dt);
 			}
 		}
 
@@ -827,7 +812,7 @@ namespace Chaos
 			}
 			else
 			{
-				UpdateContactPoint(Constraint, SpherePlaneContactPoint(Sphere, SphereTransform, Plane, PlaneTransform, Constraint.Manifold.RestitutionPadding), Dt);
+				UpdateContactPoint(Constraint, SpherePlaneContactPoint(Sphere, SphereTransform, Plane, PlaneTransform, Constraint.GetRestitutionPadding()), Dt);
 			}
 		}
 
@@ -862,7 +847,7 @@ namespace Chaos
 			}
 			else
 			{
-				UpdateContactPoint(Constraint, SphereBoxContactPoint(Sphere, SphereTransform, Box, BoxTransform, Constraint.Manifold.RestitutionPadding), Dt);
+				UpdateContactPoint(Constraint, SphereBoxContactPoint(Sphere, SphereTransform, Box, BoxTransform, Constraint.GetRestitutionPadding()), Dt);
 			}
 		}
 
@@ -898,7 +883,7 @@ namespace Chaos
 			}
 			else
 			{
-				UpdateContactPoint(Constraint, SphereCapsuleContactPoint(A, ATransform, B, BTransform, Constraint.Manifold.RestitutionPadding), Dt);
+				UpdateContactPoint(Constraint, SphereCapsuleContactPoint(A, ATransform, B, BTransform, Constraint.GetRestitutionPadding()), Dt);
 			}
 		}
 
@@ -1053,7 +1038,7 @@ namespace Chaos
 			}
 			else
 			{
-				UpdateContactPoint(Constraint, CapsuleCapsuleContactPoint(A, ATransform, B, BTransform, Constraint.Manifold.RestitutionPadding), Dt);
+				UpdateContactPoint(Constraint, CapsuleCapsuleContactPoint(A, ATransform, B, BTransform, Constraint.GetRestitutionPadding()), Dt);
 			}
 		}
 
@@ -1086,8 +1071,8 @@ namespace Chaos
 			}
 			else
 			{
-				const FVec3 InitialDir = ATransform.GetRotation().Inverse() * -Constraint.GetNormal();
-				UpdateContactPoint(Constraint, CapsuleBoxContactPoint(A, ATransform, B, BTransform, InitialDir, Constraint.Manifold.RestitutionPadding), Dt);
+				const FVec3 InitialDir = FVec3(1,0,0);
+				UpdateContactPoint(Constraint, CapsuleBoxContactPoint(A, ATransform, B, BTransform, InitialDir, Constraint.GetRestitutionPadding()), Dt);
 			}
 		}
 
@@ -1345,7 +1330,7 @@ namespace Chaos
 		{
 			CONDITIONAL_SCOPE_CYCLE_COUNTER(STAT_Collisions_UpdateGenericConvexConvexConstraint, ConstraintsDetailedStats);
 
-			UpdateGenericConvexConvexConstraintHelper(Implicit0, WorldTransform0, Implicit1, WorldTransform1, Dt, Constraint.Manifold.RestitutionPadding, Constraint);
+			UpdateGenericConvexConvexConstraintHelper(Implicit0, WorldTransform0, Implicit1, WorldTransform1, Dt, Constraint.GetRestitutionPadding(), Constraint);
 		}
 
 		void UpdateGenericConvexConvexConstraintSwept(TGeometryParticleHandle<FReal, 3>* Particle0, const FImplicitObject& Implicit0, const FRigidTransform3& StartWorldTransform0, TGeometryParticleHandle<FReal, 3>* Particle1, const FImplicitObject& Implicit1, const FRigidTransform3& StartWorldTransform1, const FVec3& Dir, const FReal Length, const FReal Dt, FPBDCollisionConstraint& Constraint)
@@ -1563,13 +1548,13 @@ namespace Chaos
 
 			// NOTE: We are assuming that if only one body has collision particles, it is the first one. This implies
 			// that the first body is dynamic (currently this is true)
-			const FBVHParticles* SampleParticles = Constraint.Manifold.Simplicial[0];
+			const FBVHParticles* SampleParticles = Constraint.GetCollisionParticles0();
 			if (SampleParticles)
 			{
 				// Sample particles are in Particle(Actor) space, but SolverBody state in in CoM space
 				// The transforms passed in are the Particle transforms but concatenatated with the shape transform(s) so we need
 				// to regenerate the Particle transform without those shape transforms. We should be able to eliminate this...
-				FGenericParticleHandle Particle0 = Constraint.Particle[0];
+				FGenericParticleHandle Particle0 = Constraint.GetParticle0();
 				// @todo(chaos): we should not be accessing the particle here (or anywhere in collision detection)
 				// The two paths here are required for unit testing. Units tests don't always set up the solver bodies, but this function is
 				// usually called in the solver loop where solver bodies must be used.
@@ -1586,17 +1571,8 @@ namespace Chaos
 
 			if (SampleParticles)
 			{
-				const FImplicitObject* Obj1 = Constraint.Manifold.Implicit[1];
+				const FImplicitObject* Obj1 = Constraint.GetImplicit1();
 				FContactPoint ContactPoint = SampleObject<UpdateType>(*Obj1, LevelsetTM, *SampleParticles, ParticlesTM, Constraint.GetCullDistance());
-
-				if (ContactPoint.IsSet())
-				{
-					// @todo(chaos): this work is partially duplicated from SampleObject - fix this
-					ContactPoint.ShapeContactPoints[0] = WorldTransform0.InverseTransformPosition(ContactPoint.Location + FReal(0.5) * ContactPoint.Phi * ContactPoint.Normal);
-					ContactPoint.ShapeContactPoints[1] = WorldTransform1.InverseTransformPosition(ContactPoint.Location - FReal(0.5) * ContactPoint.Phi * ContactPoint.Normal);
-					ContactPoint.ShapeContactNormal = WorldTransform1.InverseTransformVector(ContactPoint.Normal);
-				}
-
 				UpdateContactPoint(Constraint, ContactPoint, Dt);
 			}
 		}
@@ -1612,36 +1588,32 @@ namespace Chaos
 		void ConstructLevelsetLevelsetConstraints(TGeometryParticleHandle<FReal, 3>* Particle0, TGeometryParticleHandle<FReal, 3>* Particle1, const FImplicitObject* Implicit0, const FBVHParticles* Simplicial0, const FImplicitObject* Implicit1, const FBVHParticles* Simplicial1, const FRigidTransform3& ParticleWorldTransform0, const FRigidTransform3& LocalTransform0, const FRigidTransform3& ParticleWorldTransform1, const FRigidTransform3& LocalTransform1, const FReal CullDistance, const FReal Dt, const FCollisionContext& Context)
 		{
 			const bool bAllowManifolds = Context.bAllowManifolds && bChaos_Collision_AllowLevelsetManifolds;
-			FPBDCollisionConstraint* Constraint = Context.CollisionAllocator->FindOrCreateConstraint(Particle0, Implicit0, Simplicial0, LocalTransform0, Particle1, Implicit1, Simplicial1, LocalTransform1, CullDistance, EContactShapesType::LevelSetLevelSet, bAllowManifolds);
-
-			// @todo(chaos): support manifold restoration for Levelsets and remove this
-			Constraint->ResetActiveManifoldContacts();
+			FPBDCollisionConstraint* Constraint = nullptr;
 
-			// @todo(chaos): the order test should be done befor we create the constraint so we can do it right first time, rather than patching it after
-			bool bIsParticleDynamic0 = Particle0->CastToRigidParticle() && Particle0->ObjectState() == EObjectStateType::Dynamic;
+			// @todo(chaos): We do not really support a runtime decision for the particle order now that we have persistent collisions. Fix the logic below...?
+			// Also clean this logic up so it is readable. When is Simplical0 not equal to CollisionParticles0?
+			bool bIsParticleDynamic0 = Particle0->CastToRigidParticle() && ((Particle0->ObjectState() == EObjectStateType::Dynamic) || (Particle0->ObjectState() == EObjectStateType::Sleeping));
 			int32 P0NumCollisionParticles = Simplicial0 ? Simplicial0->Size() : Particle0->CastToRigidParticle()->CollisionParticlesSize();
 			if (!Particle1->Geometry() || (bIsParticleDynamic0 && !P0NumCollisionParticles && Particle0->Geometry() && !Particle0->Geometry()->IsUnderlyingUnion()))
 			{
-				Constraint->Particle[0] = Particle1;
-				Constraint->Particle[1] = Particle0;
-				Constraint->ImplicitTransform[0] = LocalTransform1;
-				Constraint->ImplicitTransform[1] = LocalTransform0;
-				Constraint->SetManifold(Implicit1, Simplicial1, Implicit0, Simplicial0);
+				Constraint = Context.CollisionAllocator->FindOrCreateConstraint(Particle1, Implicit1, Simplicial1, LocalTransform1, Particle0, Implicit0, Simplicial0, LocalTransform0, CullDistance, EContactShapesType::LevelSetLevelSet, bAllowManifolds);
 			}
 			else
 			{
-				Constraint->Particle[0] = Particle0;
-				Constraint->Particle[1] = Particle1;
-				Constraint->ImplicitTransform[0] = LocalTransform0;
-				Constraint->ImplicitTransform[1] = LocalTransform1;
-				Constraint->SetManifold(Implicit0, Simplicial0, Implicit1, Simplicial1);
+				Constraint = Context.CollisionAllocator->FindOrCreateConstraint(Particle0, Implicit0, Simplicial0, LocalTransform0, Particle1, Implicit1, Simplicial1, LocalTransform1, CullDistance, EContactShapesType::LevelSetLevelSet, bAllowManifolds);
 			}
 
-			if (T_TRAITS::bImmediateUpdate && !Constraint->WasManifoldRestored())
+			if (Constraint != nullptr)
 			{
-				FRigidTransform3 WorldTransform0 = Constraint->ImplicitTransform[0] * ParticleWorldTransform0;
-				FRigidTransform3 WorldTransform1 = Constraint->ImplicitTransform[1] * ParticleWorldTransform1;
-				UpdateLevelsetLevelsetConstraint<ECollisionUpdateType::Deepest>(WorldTransform0, WorldTransform1, Dt, *Constraint);
+				// @todo(chaos): support manifold restoration for Levelsets and remove this
+				Constraint->ResetActiveManifoldContacts();
+
+				if (T_TRAITS::bImmediateUpdate && !Constraint->WasManifoldRestored())
+				{
+					FRigidTransform3 WorldTransform0 = Constraint->GetShapeRelativeTransform0() * ParticleWorldTransform0;
+					FRigidTransform3 WorldTransform1 = Constraint->GetShapeRelativeTransform1() * ParticleWorldTransform1;
+					UpdateLevelsetLevelsetConstraint<ECollisionUpdateType::Deepest>(WorldTransform0, WorldTransform1, Dt, *Constraint);
+				}
 			}
 		}
 
@@ -1665,16 +1637,16 @@ namespace Chaos
 			// @todo(chaos): remove
 			//const FVec3 OriginalContactPositionLocal0 = WorldTransform0.InverseTransformPosition(Constraint.Manifold.Location);
 			//const FVec3 OriginalContactPositionLocal1 = WorldTransform1.InverseTransformPosition(Constraint.Manifold.Location);
-			if (!Constraint.Manifold.Implicit[0] || !Constraint.Manifold.Implicit[1])
+			if (!Constraint.GetImplicit0() || !Constraint.GetImplicit1())
 			{
 				return;
 			}
 
-			const FImplicitObject& Implicit0 = *Constraint.Manifold.Implicit[0];
-			const FImplicitObject& Implicit1 = *Constraint.Manifold.Implicit[1];
+			const FImplicitObject& Implicit0 = *Constraint.GetImplicit0();
+			const FImplicitObject& Implicit1 = *Constraint.GetImplicit1();
 
 
-			switch (Constraint.Manifold.ShapesType)
+			switch (Constraint.GetShapesType())
 			{
 			case EContactShapesType::SphereSphere:
 				UpdateSphereSphereConstraint(*Implicit0.template GetObject<TSphere<FReal, 3>>(), WorldTransform0, *Implicit1.template GetObject<TSphere<FReal, 3>>(), WorldTransform1, Dt, Constraint);
@@ -1794,10 +1766,10 @@ namespace Chaos
 		{
 			INC_DWORD_STAT(STAT_ChaosCollisionCounter_NumContactUpdates);
 
-			TGeometryParticleHandle<FReal, 3>* Particle0 = Constraint.Particle[0];
-			TGeometryParticleHandle<FReal, 3>* Particle1 = Constraint.Particle[1];
-			const FImplicitObject& Implicit0 = *Constraint.Manifold.Implicit[0];
-			const FImplicitObject& Implicit1 = *Constraint.Manifold.Implicit[1];
+			TGeometryParticleHandle<FReal, 3>* Particle0 = Constraint.GetParticle0();
+			TGeometryParticleHandle<FReal, 3>* Particle1 = Constraint.GetParticle1();
+			const FImplicitObject& Implicit0 = *Constraint.GetImplicit0();
+			const FImplicitObject& Implicit1 = *Constraint.GetImplicit1();
 
 			FReal LengthCCD = 0.0f;
 			FVec3 DirCCD(0.0f);
@@ -1805,7 +1777,7 @@ namespace Chaos
 
 			if (bUseCCD)
 			{
-				const EContactShapesType ShapePairType = Constraint.Manifold.ShapesType;
+				const EContactShapesType ShapePairType = Constraint.GetShapesType();
 				switch (ShapePairType)
 				{
 				case EContactShapesType::SphereSphere:
@@ -2221,16 +2193,16 @@ namespace Chaos
 		template<ECollisionUpdateType UpdateType>
 		void UpdateConstraintFromGeometry(FPBDCollisionConstraint& Constraint, const FRigidTransform3& ParticleWorldTransform0, const FRigidTransform3& ParticleWorldTransform1, const FReal Dt)
 		{
-			const FRigidTransform3 WorldTransform0 = Constraint.ImplicitTransform[0] * ParticleWorldTransform0;
-			const FRigidTransform3 WorldTransform1 = Constraint.ImplicitTransform[1] * ParticleWorldTransform1;
+			const FRigidTransform3 WorldTransform0 = Constraint.GetShapeRelativeTransform0() * ParticleWorldTransform0;
+			const FRigidTransform3 WorldTransform1 = Constraint.GetShapeRelativeTransform1() * ParticleWorldTransform1;
 			UpdateConstraintFromGeometryImpl<UpdateType>(Constraint, WorldTransform0, WorldTransform1, Dt);
 		}
 
 		template<ECollisionUpdateType UpdateType>
 		bool UpdateConstraintFromGeometrySwept(FPBDCollisionConstraint& Constraint, const FRigidTransform3& ParticleWorldTransform0, const FRigidTransform3& ParticleWorldTransform1, const FReal Dt)
 		{
-			const FRigidTransform3 WorldTransform0 = Constraint.ImplicitTransform[0] * ParticleWorldTransform0;
-			const FRigidTransform3 WorldTransform1 = Constraint.ImplicitTransform[1] * ParticleWorldTransform1;
+			const FRigidTransform3 WorldTransform0 = Constraint.GetShapeRelativeTransform0() * ParticleWorldTransform0;
+			const FRigidTransform3 WorldTransform1 = Constraint.GetShapeRelativeTransform1() * ParticleWorldTransform1;
 			return UpdateConstraintFromGeometrySweptImpl<UpdateType>(Constraint, WorldTransform0, WorldTransform1, Dt);
 		}
 
diff --git a/Engine/Source/Runtime/Experimental/Chaos/Private/Chaos/CollisionResolutionUtil.cpp b/Engine/Source/Runtime/Experimental/Chaos/Private/Chaos/CollisionResolutionUtil.cpp
--- a/Engine/Source/Runtime/Experimental/Chaos/Private/Chaos/CollisionResolutionUtil.cpp
+++ b/Engine/Source/Runtime/Experimental/Chaos/Private/Chaos/CollisionResolutionUtil.cpp
@@ -116,22 +116,6 @@ namespace Chaos
 			return Numerator < Denominator ? (Impulse * Numerator / Denominator) : Impulse;
 		}
 
-		bool SampleObjectHelper(const FImplicitObject& Object, const FRigidTransform3& ObjectTransform, const FRigidTransform3& SampleToObjectTransform, const FVec3& SampleParticle, FReal Thickness, FContactPoint& Contact)
-		{
-			FVec3 LocalPoint = SampleToObjectTransform.TransformPositionNoScale(SampleParticle);
-			FVec3 LocalNormal;
-			FReal LocalPhi = Object.PhiWithNormal(LocalPoint, LocalNormal);
-
-			if (LocalPhi < Contact.Phi)
-			{
-				Contact.Phi = LocalPhi;
-				Contact.Normal = ObjectTransform.TransformVectorNoScale(LocalNormal);
-				Contact.Location = ObjectTransform.TransformPositionNoScale(LocalPoint);
-				return true;
-			}
-			return false;
-		}
-
 		bool SampleObjectNoNormal(const FImplicitObject& Object, const FRigidTransform3& ObjectTransform, const FRigidTransform3& SampleToObjectTransform, const FVec3& SampleParticle, FReal Thickness, FContactPoint& Contact)
 		{
 			FVec3 LocalPoint = SampleToObjectTransform.TransformPositionNoScale(SampleParticle);
@@ -155,7 +139,7 @@ namespace Chaos
 
 			if (LocalThickness < -KINDA_SMALL_NUMBER)
 			{
-				Contact.Location += LocalPoint * LocalThickness;
+				Contact.ShapeContactPoints[0] += LocalPoint * LocalThickness;
 				TotalThickness += LocalThickness;
 				return true;
 			}
@@ -186,10 +170,12 @@ namespace Chaos
 			FContactPoint Contact;
 			FContactPoint AvgContact;
 
-			Contact.Location = FVec3::ZeroVector;
-			Contact.Normal = FVec3::ZeroVector;
-			AvgContact.Location = FVec3::ZeroVector;
-			AvgContact.Normal = FVec3::ZeroVector;
+			Contact.ShapeContactPoints[0] = FVec3(0);
+			Contact.ShapeContactPoints[1] = FVec3(0);
+			Contact.ShapeContactNormal = FVec3(0);
+			AvgContact.ShapeContactPoints[0] = FVec3(0);
+			AvgContact.ShapeContactPoints[1] = FVec3(0);
+			AvgContact.ShapeContactNormal = FVec3(0);
 			AvgContact.Phi = CullingDistance;
 			FReal WeightSum = FReal(0); // Sum of weights used for averaging (negative)
 
@@ -230,7 +216,7 @@ namespace Chaos
 								&PotentialParticles[0],
 								CullingDistance,
 								WeightSum,
-								(ispc::FVector&)AvgContact.Location,
+								(ispc::FVector&)AvgContact.ShapeContactPoints[0],
 								PotentialParticles.Num());
 						}
 						else
@@ -272,7 +258,7 @@ namespace Chaos
 								&PotentialParticles[0],
 								CullingDistance,
 								WeightSum,
-								(ispc::FVector&)AvgContact.Location,
+								(ispc::FVector&)AvgContact.ShapeContactPoints[0],
 								PotentialParticles.Num());
 						}
 						else
@@ -339,7 +325,7 @@ namespace Chaos
 							(ispc::FVector*)&SampleParticles.XArray()[0],
 							CullingDistance,
 							WeightSum,
-							(ispc::FVector&)AvgContact.Location,
+							(ispc::FVector&)AvgContact.ShapeContactPoints[0],
 							NumParticles);
 					}
 					else
@@ -376,7 +362,7 @@ namespace Chaos
 							(ispc::FVector*)&SampleParticles.XArray()[0],
 							CullingDistance,
 							WeightSum,
-							(ispc::FVector&)AvgContact.Location,
+							(ispc::FVector&)AvgContact.ShapeContactPoints[0],
 							NumParticles);
 					}
 					else
@@ -410,7 +396,7 @@ namespace Chaos
 							(ispc::FVector*) & SampleParticles.XArray()[0],
 							CullingDistance,
 							WeightSum,
-							(ispc::FVector&)AvgContact.Location,
+							(ispc::FVector&)AvgContact.ShapeContactPoints[0],
 							NumParticles);
 					}
 					else
@@ -446,7 +432,7 @@ namespace Chaos
 							(ispc::FVector*)&SampleParticles.XArray()[0],
 							CullingDistance,
 							WeightSum,
-							(ispc::FVector&)AvgContact.Location,
+							(ispc::FVector&)AvgContact.ShapeContactPoints[0],
 							NumParticles);
 					}
 					else
@@ -496,14 +482,15 @@ namespace Chaos
 			{
 				if (WeightSum < -KINDA_SMALL_NUMBER)
 				{
-					FVec3 LocalPoint = AvgContact.Location / WeightSum;
+					FVec3 LocalPoint = AvgContact.ShapeContactPoints[0] / WeightSum;
 					FVec3 LocalNormal;
 					const FReal NewPhi = Object.PhiWithNormal(LocalPoint, LocalNormal);
 					if (NewPhi < Contact.Phi)
 					{
 						Contact.Phi = NewPhi;
-						Contact.Location = ObjectTransform.TransformPositionNoScale(LocalPoint);
-						Contact.Normal = ObjectTransform.TransformVectorNoScale(LocalNormal);
+						Contact.ShapeContactPoints[0] = SampleToObjectTM.InverseTransformPositionNoScale(LocalPoint);
+						Contact.ShapeContactPoints[1] = LocalPoint - NewPhi * LocalNormal;
+						Contact.ShapeContactNormal = SampleToObjectTM.InverseTransformVectorNoScale(LocalNormal);
 					}
 				}
 				else
@@ -517,8 +504,9 @@ namespace Chaos
 				FVec3 LocalPoint = SampleToObjectTM.TransformPositionNoScale(SampleParticles.X(DeepestParticle));
 				FVec3 LocalNormal;
 				Contact.Phi = Object.PhiWithNormal(LocalPoint, LocalNormal);
-				Contact.Location = ObjectTransform.TransformPositionNoScale(LocalPoint);
-				Contact.Normal = ObjectTransform.TransformVectorNoScale(LocalNormal);
+				Contact.ShapeContactPoints[0] = SampleParticles.X(DeepestParticle);
+				Contact.ShapeContactPoints[1] = LocalPoint - Contact.Phi * LocalNormal;
+				Contact.ShapeContactNormal = SampleToObjectTM.InverseTransformVectorNoScale(LocalNormal);
 			}
 
 			return Contact;
@@ -532,10 +520,12 @@ namespace Chaos
 			FContactPoint Contact;
 			FContactPoint AvgContact;
 
-			Contact.Location = FVec3::ZeroVector;
-			Contact.Normal = FVec3::ZeroVector;
-			AvgContact.Location = FVec3::ZeroVector;
-			AvgContact.Normal = FVec3::ZeroVector;
+			Contact.ShapeContactPoints[0] = FVec3(0);
+			Contact.ShapeContactPoints[1] = FVec3(0);
+			Contact.ShapeContactNormal = FVec3(0);
+			AvgContact.ShapeContactPoints[0] = FVec3(0);
+			AvgContact.ShapeContactPoints[1] = FVec3(0);
+			AvgContact.ShapeContactNormal = FVec3(0);
 			AvgContact.Phi = CullingDistance;
 			FReal WeightSum = FReal(0); // Sum of weights used for averaging (negative)
 
@@ -604,14 +594,15 @@ namespace Chaos
 			{
 				if (WeightSum < -KINDA_SMALL_NUMBER)
 				{
-					FVec3 LocalPoint = AvgContact.Location / WeightSum;
+					FVec3 LocalPoint = AvgContact.ShapeContactPoints[0] / WeightSum;
 					FVec3 LocalNormal;
 					const FReal NewPhi = Object.PhiWithNormal(LocalPoint, LocalNormal);
 					if (NewPhi < Contact.Phi)
 					{
 						Contact.Phi = NewPhi;
-						Contact.Location = ObjectTransform.TransformPositionNoScale(LocalPoint - FReal(0.5) * NewPhi * LocalNormal);
-						Contact.Normal = ObjectTransform.TransformVectorNoScale(LocalNormal);
+						Contact.ShapeContactPoints[0] = SampleToObjectTM.InverseTransformPositionNoScale(LocalPoint);
+						Contact.ShapeContactPoints[1] = LocalPoint - NewPhi * LocalNormal; 
+						Contact.ShapeContactNormal = SampleToObjectTM.InverseTransformVectorNoScale(LocalNormal);
 					}
 				}
 				else
@@ -625,9 +616,9 @@ namespace Chaos
 				FVec3 LocalPoint = SampleToObjectTM.TransformPositionNoScale(SampleParticles.X(DeepestParticle));
 				FVec3 LocalNormal;
 				Contact.Phi = Object.PhiWithNormal(LocalPoint, LocalNormal);
-				Contact.Location = ObjectTransform.TransformPositionNoScale(LocalPoint - FReal(0.5) * Contact.Phi * LocalNormal);
-				Contact.Normal = ObjectTransform.TransformVectorNoScale(LocalNormal);
-
+				Contact.ShapeContactPoints[0] = SampleParticles.X(DeepestParticle);
+				Contact.ShapeContactPoints[1] = LocalPoint - Contact.Phi * LocalNormal;
+				Contact.ShapeContactNormal = SampleToObjectTM.InverseTransformVectorNoScale(LocalNormal);
 			}
 
 			return Contact;
diff --git a/Engine/Source/Runtime/Experimental/Chaos/Private/Chaos/ContactModification.cpp b/Engine/Source/Runtime/Experimental/Chaos/Private/Chaos/ContactModification.cpp
--- a/Engine/Source/Runtime/Experimental/Chaos/Private/Chaos/ContactModification.cpp
+++ b/Engine/Source/Runtime/Experimental/Chaos/Private/Chaos/ContactModification.cpp
@@ -46,7 +46,7 @@ namespace Chaos
 
 	const FImplicitObject* FContactPairModifier::GetContactGeometry(int32 ParticleIdx)
 	{
-		return Constraint->Manifold.Implicit[ParticleIdx];
+		return Constraint->GetImplicit(ParticleIdx);
 	}
 
 	FRigidTransform3 FContactPairModifier::GetShapeToWorld(int32 ParticleIdx) const
@@ -57,10 +57,10 @@ namespace Chaos
 		if (FPBDRigidParticleHandle* Rigid = Particle->CastToRigidParticle())
 		{
 			// Use PQ for rigids.
-			return Constraint->ImplicitTransform[ParticleIdx] * FParticleUtilitiesPQ::GetActorWorldTransform(Rigid);
+			return Constraint->GetShapeRelativeTransform(ParticleIdx) * FParticleUtilitiesPQ::GetActorWorldTransform(Rigid);
 		}
 
-		return Constraint->ImplicitTransform[ParticleIdx] * FParticleUtilitiesXR::GetActorWorldTransform(Particle);
+		return Constraint->GetShapeRelativeTransform(ParticleIdx) * FParticleUtilitiesXR::GetActorWorldTransform(Particle);
 	}
 
 	FReal FContactPairModifier::GetSeparation(int32 ContactPointIdx) const
@@ -92,26 +92,40 @@ namespace Chaos
 
 	FVec3 FContactPairModifier::GetWorldNormal(int32 ContactPointIdx) const
 	{
-		return Constraint->GetManifoldPoints()[ContactPointIdx].ContactPoint.Normal;
+		const FRigidTransform3& ShapeTransform1 = Constraint->GetShapeWorldTransform1();
+
+		TArrayView<FManifoldPoint> ManifoldPoints = Constraint->GetManifoldPoints();
+		FManifoldPoint& ManifoldPoint = ManifoldPoints[ContactPointIdx];
+
+		return ShapeTransform1.TransformVectorNoScale(ManifoldPoint.ContactPoint.ShapeContactNormal);
 	}
 
 	void FContactPairModifier::ModifyWorldNormal(const FVec3& Normal, int32 ContactPointIdx)
 	{
+		FVec3 WorldContactPoint0, WorldContactPoint1;
+		GetWorldContactLocations(ContactPointIdx, WorldContactPoint0, WorldContactPoint1);
+
 		TArrayView<FManifoldPoint> ManifoldPoints = Constraint->GetManifoldPoints();
 		FManifoldPoint& ManifoldPoint = ManifoldPoints[ContactPointIdx];
 
-		ManifoldPoint.ContactPoint.Normal = Normal;
-		ManifoldPoint.ContactPoint.Phi = FVec3::DotProduct(ManifoldPoint.WorldContactPoints[0] - ManifoldPoint.WorldContactPoints[1], ManifoldPoint.ContactPoint.Normal);
+		const FRigidTransform3& ShapeTransform1 = Constraint->GetShapeWorldTransform1();
+		const FVec3 ShapeNormal = ShapeTransform1.InverseTransformVectorNoScale(Normal);
+
+		ManifoldPoint.ContactPoint.ShapeContactNormal = ShapeNormal;
+		ManifoldPoint.ContactPoint.Phi = FVec3::DotProduct(WorldContactPoint0 - WorldContactPoint1, Normal);
 
 		Modifier->MarkConstraintForManifoldUpdate(*Constraint);
 	}
 
 	void FContactPairModifier::GetWorldContactLocations(int32 ContactPointIdx, FVec3& OutLocation0, FVec3& OutLocation1) const
 	{
+		const FRigidTransform3& ShapeTransform0 = Constraint->GetShapeWorldTransform0();
+		const FRigidTransform3& ShapeTransform1 = Constraint->GetShapeWorldTransform1();
+
 		TArrayView<FManifoldPoint> ManifoldPoints = Constraint->GetManifoldPoints();
 		FManifoldPoint& ManifoldPoint = ManifoldPoints[ContactPointIdx];
-		OutLocation0 = ManifoldPoint.WorldContactPoints[0];
-		OutLocation1 = ManifoldPoint.WorldContactPoints[1];
+		OutLocation0 = ShapeTransform0.TransformPositionNoScale(ManifoldPoint.ContactPoint.ShapeContactPoints[0]);
+		OutLocation1 = ShapeTransform1.TransformPositionNoScale(ManifoldPoint.ContactPoint.ShapeContactPoints[1]);
 	}
 
 	FVec3 FContactPairModifier::GetWorldContactLocation(int32 ContactPointIdx) const
@@ -128,59 +142,55 @@ namespace Chaos
 
 		ManifoldPoint.ContactPoint.ShapeContactPoints[0] = Constraint->GetShapeWorldTransform0().InverseTransformPositionNoScale(Location0);
 		ManifoldPoint.ContactPoint.ShapeContactPoints[1] = Constraint->GetShapeWorldTransform1().InverseTransformPositionNoScale(Location1);
-		// @todo(chaos): Overwriting ShapeAnchorPoints disables static friction for this tick - we might want to do something better here
-		ManifoldPoint.ShapeAnchorPoints[0] = ManifoldPoint.ContactPoint.ShapeContactPoints[0];
-		ManifoldPoint.ShapeAnchorPoints[1] = ManifoldPoint.ContactPoint.ShapeContactPoints[1];
-		ManifoldPoint.WorldContactPoints[0] = Location0;
-		ManifoldPoint.WorldContactPoints[1] = Location1;
-		ManifoldPoint.ContactPoint.Location = 0.5 * (Location0 + Location1);
+		// @todo(chaos): We may also need to do something with the SavedManifoldPoints on the Constraint otherwise static friction
+		// may still try to push the points back to their original relative positions
 
 		Modifier->MarkConstraintForManifoldUpdate(*Constraint);
 	}
 
 	FReal FContactPairModifier::GetRestitution() const
 	{
-		return Constraint->GetManifold().Restitution;
+		return Constraint->GetRestitution();
 	}
 
 	void FContactPairModifier::ModifyRestitution(FReal Restitution)
 	{
-		Constraint->Manifold.Restitution = Restitution;
+		Constraint->SetRestitution(Restitution);
 	}
 
 	FReal FContactPairModifier::GetRestitutionThreshold() const
 	{
-		return Constraint->Manifold.RestitutionThreshold;
+		return Constraint->GetRestitutionThreshold();
 	}
 
 	void FContactPairModifier::ModifyRestitutionThreshold(FReal Threshold)
 	{
-		Constraint->Manifold.RestitutionThreshold = Threshold;
+		Constraint->SetRestitutionThreshold(Threshold);
 	}
 
 	FReal FContactPairModifier::GetDynamicFriction() const
 	{
-		return Constraint->Manifold.Friction;
+		return Constraint->GetDynamicFriction();
 	}
 
 	void FContactPairModifier::ModifyDynamicFriction(FReal DynamicFriction)
 	{
-		Constraint->Manifold.Friction = DynamicFriction;
+		Constraint->SetDynamicFriction(DynamicFriction);
 	}
 
 	FReal FContactPairModifier::GetStaticFriction() const
 	{
-		return Constraint->Manifold.AngularFriction;
+		return Constraint->GetStaticFriction();
 	}
 
 	void FContactPairModifier::ModifyStaticFriction(FReal StaticFriction)
 	{
-		Constraint->Manifold.AngularFriction = StaticFriction;
+		Constraint->SetStaticFriction(StaticFriction);
 	}
 
 	FVec3 FContactPairModifier::GetParticleVelocity(int32 ParticleIdx) const
 	{
-		const FGeometryParticleHandle* Particle = Constraint->Particle[ParticleIdx];
+		const FGeometryParticleHandle* Particle = Constraint->GetParticle(ParticleIdx);
 		const FKinematicGeometryParticleHandle* KinematicHandle = Particle->CastToKinematicParticle();
 		if (!ensure(KinematicHandle))
 		{
@@ -193,7 +203,7 @@ namespace Chaos
 
 	void FContactPairModifier::ModifyParticleVelocity(FVec3 Velocity, int32 ParticleIdx)
 	{
-		FGeometryParticleHandle* Particle = Constraint->Particle[ParticleIdx];
+		FGeometryParticleHandle* Particle = Constraint->GetParticle(ParticleIdx);
 		FKinematicGeometryParticleHandle* KinematicHandle = Particle->CastToKinematicParticle();
 		if (!ensure(KinematicHandle))
 		{
@@ -217,7 +227,7 @@ namespace Chaos
 
 	FVec3 FContactPairModifier::GetParticleAngularVelocity(int32 ParticleIdx) const
 	{
-		const FGeometryParticleHandle* Particle = Constraint->Particle[ParticleIdx];
+		const FGeometryParticleHandle* Particle = Constraint->GetParticle(ParticleIdx);
 		const FKinematicGeometryParticleHandle* KinematicHandle = Particle->CastToKinematicParticle();
 		if (!ensure(KinematicHandle))
 		{
@@ -230,7 +240,7 @@ namespace Chaos
 
 	void FContactPairModifier::ModifyParticleAngularVelocity(FVec3 AngularVelocity, int32 ParticleIdx)
 	{
-		FGeometryParticleHandle* Particle = Constraint->Particle[ParticleIdx];
+		FGeometryParticleHandle* Particle = Constraint->GetParticle(ParticleIdx);
 		FKinematicGeometryParticleHandle* KinematicHandle = Particle->CastToKinematicParticle();
 		if (!ensure(KinematicHandle))
 		{
@@ -252,7 +262,7 @@ namespace Chaos
 
 	FVec3 FContactPairModifier::GetParticlePosition(int32 ParticleIdx) const
 	{
-		const FGeometryParticleHandle* Particle = Constraint->Particle[ParticleIdx];
+		const FGeometryParticleHandle* Particle = Constraint->GetParticle(ParticleIdx);
 		const FPBDRigidParticleHandle* RigidHandle = Particle->CastToRigidParticle();
 
 		if (RigidHandle)
@@ -272,7 +282,7 @@ namespace Chaos
 
 	void FContactPairModifier::ModifyParticlePosition(FVec3 Position, bool bMaintainVelocity, int32 ParticleIdx)
 	{
-		FGeometryParticleHandle* Particle = Constraint->Particle[ParticleIdx];
+		FGeometryParticleHandle* Particle = Constraint->GetParticle(ParticleIdx);
 
 		Modifier->MarkConstraintForManifoldUpdate(*Constraint);
 
@@ -320,7 +330,7 @@ namespace Chaos
 
 	FRotation3 FContactPairModifier::GetParticleRotation(int32 ParticleIdx) const
 	{
-		const FGeometryParticleHandle* Particle = Constraint->Particle[ParticleIdx];
+		const FGeometryParticleHandle* Particle = Constraint->GetParticle(ParticleIdx);
 		const FPBDRigidParticleHandle* RigidHandle = Particle->CastToRigidParticle();
 
 
@@ -335,7 +345,7 @@ namespace Chaos
 
 	void FContactPairModifier::ModifyParticleRotation(FRotation3 Rotation, bool bMaintainVelocity, int32 ParticleIdx)
 	{
-		FGeometryParticleHandle* Particle = Constraint->Particle[ParticleIdx];
+		FGeometryParticleHandle* Particle = Constraint->GetParticle(ParticleIdx);
 
 		Modifier->MarkConstraintForManifoldUpdate(*Constraint);
 
@@ -382,29 +392,41 @@ namespace Chaos
 
 	FReal FContactPairModifier::GetInvInertiaScale(int32 ParticleIdx) const
 	{
-		return ParticleIdx == 0 ? Constraint->Manifold.InvInertiaScale0 : Constraint->Manifold.InvInertiaScale1;
+		return ParticleIdx == 0 ? Constraint->GetInvInertiaScale0() : Constraint->GetInvInertiaScale1();
 	}
 
 	void FContactPairModifier::ModifyInvInertiaScale(FReal InInvInertiaScale, int32 ParticleIdx)
 	{
-		FReal& InvInertiaScale = ParticleIdx == 0 ? Constraint->Manifold.InvInertiaScale0 : Constraint->Manifold.InvInertiaScale1;
-		InvInertiaScale = InInvInertiaScale;
+		if (ParticleIdx == 0)
+		{
+			Constraint->SetInvInertiaScale0(InInvInertiaScale);
+		}
+		else
+		{
+			Constraint->SetInvInertiaScale1(InInvInertiaScale);
+		}
 	}
 
 	FReal FContactPairModifier::GetInvMassScale(int32 ParticleIdx) const
 	{
-		return ParticleIdx == 0 ? Constraint->Manifold.InvMassScale0 : Constraint->Manifold.InvMassScale1;
+		return ParticleIdx == 0 ? Constraint->GetInvMassScale0() : Constraint->GetInvMassScale1();
 	}
 
 	void FContactPairModifier::ModifyInvMassScale(FReal InInvMassScale, int32 ParticleIdx)
 	{
-		FReal& InvMassScale = ParticleIdx == 0 ? Constraint->Manifold.InvMassScale0 : Constraint->Manifold.InvMassScale1;
-		InvMassScale = InInvMassScale;
+		if (ParticleIdx == 0)
+		{
+			Constraint->SetInvMassScale0(InInvMassScale);
+		}
+		else
+		{
+			Constraint->SetInvMassScale1(InInvMassScale);
+		}
 	}
 
 	TVec2<FGeometryParticleHandle*> FContactPairModifier::GetParticlePair() const
 	{
-		return { Constraint->Particle[0], Constraint->Particle[1] };
+		return { Constraint->GetParticle0(), Constraint->GetParticle1() };
 	}
 
 	void FContactPairModifierIterator::SeekValidContact()
@@ -459,6 +481,7 @@ namespace Chaos
 
 	void FCollisionContactModifier::UpdateConstraintManifolds()
 	{
+		// Update derived state that depends on transforms etc
 		for (FPBDCollisionConstraint* Constraint : NeedsManifoldUpdate)
 		{
 			Constraint->UpdateManifoldContacts();
diff --git a/Engine/Source/Runtime/Experimental/Chaos/Private/Chaos/Island/SolverIsland.cpp b/Engine/Source/Runtime/Experimental/Chaos/Private/Chaos/Island/SolverIsland.cpp
--- a/Engine/Source/Runtime/Experimental/Chaos/Private/Chaos/Island/SolverIsland.cpp
+++ b/Engine/Source/Runtime/Experimental/Chaos/Private/Chaos/Island/SolverIsland.cpp
@@ -111,8 +111,8 @@ inline bool ConstraintSortPredicate(const FConstraintHandle& L, const FConstrain
 			return CollisionConstraintL->GetCCDType() < CollisionConstraintR->GetCCDType();
 		}
 
-		const FParticleID ParticleIdxsL[] = { CollisionConstraintL->Particle[0]->ParticleID(), CollisionConstraintL->Particle[1]->ParticleID() };
-		const FParticleID ParticleIdxsR[] = { CollisionConstraintR->Particle[0]->ParticleID(), CollisionConstraintL->Particle[1]->ParticleID() };
+		const FParticleID ParticleIdxsL[] = { CollisionConstraintL->GetParticle0()->ParticleID(), CollisionConstraintL->GetParticle1()->ParticleID() };
+		const FParticleID ParticleIdxsR[] = { CollisionConstraintR->GetParticle0()->ParticleID(), CollisionConstraintL->GetParticle1()->ParticleID() };
 
 		const int32 MinIdxL = ParticleIdxsL[0] < ParticleIdxsL[1] ? 0 : 1;
 		const int32 MinIdxR = ParticleIdxsR[0] < ParticleIdxsR[1] ? 0 : 1;
diff --git a/Engine/Source/Runtime/Experimental/Chaos/Private/Chaos/PBDCollisionConstraintUtil.cpp b/Engine/Source/Runtime/Experimental/Chaos/Private/Chaos/PBDCollisionConstraintUtil.cpp
--- a/Engine/Source/Runtime/Experimental/Chaos/Private/Chaos/PBDCollisionConstraintUtil.cpp
+++ b/Engine/Source/Runtime/Experimental/Chaos/Private/Chaos/PBDCollisionConstraintUtil.cpp
@@ -51,7 +51,7 @@ namespace Chaos
 		int32 NumberOfCellsXYZ = NumberOfCellsX * NumberOfCellsY * NumberOfCellsZ;
 		for (int32 IdxConstraint = 0; IdxConstraint < ConstraintsArray.Num(); ++IdxConstraint)
 		{
-			FVector Location = (FVector)ConstraintsArray[IdxConstraint].GetLocation() - BoundingBox.Min + FVector(0.5f * CellSize);
+			FVector Location = (FVector)ConstraintsArray[IdxConstraint].CalculateWorldContactLocation() - BoundingBox.Min + FVector(0.5f * CellSize);
 			int32 HashTableIdx = (int32)(Location.X * CellSizeInv) +
 								 (int32)(Location.Y * CellSizeInv) * NumberOfCellsX +
 								 (int32)(Location.Z * CellSizeInv) * NumberOfCellsXY;
diff --git a/Engine/Source/Runtime/Experimental/Chaos/Private/Chaos/PBDCollisionConstraints.cpp b/Engine/Source/Runtime/Experimental/Chaos/Private/Chaos/PBDCollisionConstraints.cpp
--- a/Engine/Source/Runtime/Experimental/Chaos/Private/Chaos/PBDCollisionConstraints.cpp
+++ b/Engine/Source/Runtime/Experimental/Chaos/Private/Chaos/PBDCollisionConstraints.cpp
@@ -84,12 +84,6 @@ namespace Chaos
 	bool CollisionsAllowParticleTracking = true;
 	FAutoConsoleVariableRef CVarCollisionsAllowParticleTracking(TEXT("p.Chaos.Collision.AllowParticleTracking"), CollisionsAllowParticleTracking, TEXT("Allow particles to track their collisions constraints when their DoBufferCollisions flag is enable [def:true]"));
 
-#if INTEL_ISPC
-	bool bChaos_Collision_ISPC_Enabled = false;
-	FAutoConsoleVariableRef CVarChaosCollisionISPCEnabled(TEXT("p.Chaos.Collision.ISPC"), bChaos_Collision_ISPC_Enabled, TEXT("Whether to use ISPC optimizations in the Collision Solver"));
-#endif
-
-
 	DECLARE_CYCLE_STAT(TEXT("Collisions::Reset"), STAT_Collisions_Reset, STATGROUP_ChaosCollision);
 	DECLARE_CYCLE_STAT(TEXT("Collisions::UpdatePointConstraints"), STAT_Collisions_UpdatePointConstraints, STATGROUP_ChaosCollision);
 	DECLARE_CYCLE_STAT(TEXT("Collisions::BeginDetect"), STAT_Collisions_BeginDetect, STATGROUP_ChaosCollision);
@@ -126,12 +120,6 @@ namespace Chaos
 		, MaxPushOutVelocity(0)
 		, SolverType(EConstraintSolverType::QuasiPbd)
 	{
-#if INTEL_ISPC
-		if (bRealTypeCompatibleWithISPC && bChaos_Collision_ISPC_Enabled)
-		{
-			check(sizeof(FCollisionContact) == ispc::SizeofFCollisionContact());
-		}
-#endif
 	}
 
 	FPBDCollisionConstraints::~FPBDCollisionConstraints()
diff --git a/Engine/Source/Runtime/Experimental/Chaos/Private/Chaos/PBDCollisionConstraints.ispc b/Engine/Source/Runtime/Experimental/Chaos/Private/Chaos/PBDCollisionConstraints.ispc
--- a/Engine/Source/Runtime/Experimental/Chaos/Private/Chaos/PBDCollisionConstraints.ispc
+++ b/Engine/Source/Runtime/Experimental/Chaos/Private/Chaos/PBDCollisionConstraints.ispc
@@ -649,27 +649,6 @@ export uniform bool SampleSphereNoNormalAll(const uniform float &ObjectRadius,
 //
 //
 
-struct FCollisionContact
-{
-	bool bDisabled;
-	FVector Normal;
-	FVector Location;
-	float Phi;
-
-	float Friction;
-	float AngularFriction;
-	float Restitution;
-	float InvInertiaScale0;
-	float InvInertiaScale1;
-	int ShapesType;
-	int* Implicit[2];
-};
-
-export uniform int32 SizeofFCollisionContact()
-{
-	return sizeof(uniform FCollisionContact);
-}
-
 // Get the CoM position given the Actor position
 static inline uniform FVector GetCoMWorldPosition(
 	const uniform FVector& PActorWorld,
@@ -741,115 +720,3 @@ static inline uniform FMatrix33 ComputeJointFactorMatrix(const uniform FVector&
 	return SetMatrix33(R1.V[0], R1.V[1], R1.V[2], R2.V[0], R2.V[1], R2.V[2]);
 }
 
-// Apply the contact force. Update the actors' positions and rotations.
-export void ApplyContact2(
-	const uniform FCollisionContact* uniform Contact,
-	uniform FVector& PActor0,
-	uniform FVector4& QActor0,
-	uniform FVector& PActor1,
-	uniform FVector4& QActor1,
-	const uniform float InvM0,
-	const uniform FVector& InvILocal0,
-	const uniform FVector& PCoMLocal0,
-	const uniform FVector4& QCoMLocal0,
-	const uniform float InvM1,
-	const uniform FVector& InvILocal1,
-	const uniform FVector& PCoMLocal1,
-	const uniform FVector4& QCoMLocal1)
-{
-	uniform FVector P0 = GetCoMWorldPosition(PActor0, QActor0, PCoMLocal0);
-	uniform FVector P1 = GetCoMWorldPosition(PActor1, QActor1, PCoMLocal1);
-	uniform FVector4 Q0 = GetCoMWorldRotation(QActor0, QCoMLocal0);
-	uniform FVector4 Q1 = GetCoMWorldRotation(QActor1, QCoMLocal1);
-	uniform FVector VectorToPoint0 = Contact->Location - P0;
-	uniform FVector VectorToPoint1 = Contact->Location - P1;
-	
-	uniform FMatrix33 ContactInvI = Matrix33Zero;
-	uniform FMatrix33 InvI0 = Matrix33Zero;
-	uniform FMatrix33 InvI1 = Matrix33Zero;
-	if (InvM0 > 0.0f)
-	{
-		InvI0 = Multiply(ComputeWorldSpaceInertia(Q0, InvILocal0), Contact->InvInertiaScale0);
-		ContactInvI = AddAB(ContactInvI, ComputeJointFactorMatrix(VectorToPoint0, InvI0, InvM0));
-	}
-	if (InvM1 > 0.0f)
-	{
-		InvI1 = Multiply(ComputeWorldSpaceInertia(Q1, InvILocal1), Contact->InvInertiaScale1);
-		ContactInvI = AddAB(ContactInvI, ComputeJointFactorMatrix(VectorToPoint1, InvI1, InvM1));
-	}
-	
-	// Calculate the normal correction
-	// @todo(chaos): support restitution in this mode (requires a small-impact cutoff threshold)
-	uniform float Restitution = 0.0f;//Contact->Restitution;
-	uniform FVector NormalError = Contact->Phi * Contact->Normal;
-	uniform float NormalImpulseDenominator = VectorDot(Contact->Normal, Multiply(ContactInvI, Contact->Normal));
-	uniform FVector NormalImpulseNumerator = -(1.0f + Restitution) * NormalError;
-	uniform FVector NormalCorrection = NormalImpulseNumerator / NormalImpulseDenominator;
-	
-	// Calculate the lateral correction
-	uniform FVector LateralCorrection = ZeroVector;
-	//if (Contact->Friction > 0)
-	//{
-	//	// Get contact velocity
-	//	uniform FVector X0 = FParticleUtilitiesXR::GetCoMWorldPosition(Particle0);
-	//	uniform FVector X1 = FParticleUtilitiesXR::GetCoMWorldPosition(Particle1);
-	//	FRotation3 R0 = FParticleUtilitiesXR::GetCoMWorldRotation(Particle0);
-	//	FRotation3 R1 = FParticleUtilitiesXR::GetCoMWorldRotation(Particle1);
-	//	uniform FVector V0 = FVec3::CalculateVelocity(X0, P0, IterationParameters.Dt);
-	//	uniform FVector W0 = FRotation3::CalculateAngularVelocity(R0, Q0, IterationParameters.Dt);
-	//	uniform FVector V1 = FVec3::CalculateVelocity(X1, P1, IterationParameters.Dt);
-	//	uniform FVector W1 = FRotation3::CalculateAngularVelocity(R1, Q1, IterationParameters.Dt);
-	//	uniform FVector CV0 = V0 + FVec3::CrossProduct(W0, VectorToPoint0);
-	//	uniform FVector CV1 = V1 + FVec3::CrossProduct(W1, VectorToPoint1);
-	//	uniform FVector CV = CV0 - CV1;
-	//
-	//	// Calculate lateral correction, clamped to the friction cone. Kinda.
-	//	uniform float CVNormalMag = FVec3::DotProduct(CV, Contact->Normal);
-	//	if (CVNormalMag < 0.0f)
-	//	{
-	//		uniform FVector CVLateral = CV - CVNormalMag * Contact->Normal;
-	//		uniform float CVLateralMag = CVLateral.Size();
-	//		if (CVLateralMag > KINDA_SMALL_NUMBER)
-	//		{
-	//			uniform FVector DirLateral = CVLateral / CVLateralMag;
-	//			uniform FVector LateralImpulseNumerator = -CVLateral * IterationParameters.Dt;
-	//			uniform float LateralImpulseDenominator = FVec3::DotProduct(DirLateral, ContactInvI * DirLateral);
-	//			LateralCorrection = LateralImpulseNumerator / LateralImpulseDenominator;
-	//			uniform float LateralImpulseMag = LateralCorrection.Size();
-	//			uniform float NormalImpulseMag = NormalCorrection.Size();
-	//			if (LateralImpulseMag > Contact->Friction * NormalImpulseMag)
-	//			{
-	//				LateralCorrection *= Contact->Friction * NormalImpulseMag / LateralImpulseMag;
-	//			}
-	//		}
-	//	}
-	//}
-
-	// Net Correction
-	uniform FVector DX = NormalCorrection + LateralCorrection;
-				
-	if (InvM0 > 0.0f)
-	{
-		uniform FVector DP0 = InvM0 * DX;
-		uniform FVector DR0 = Multiply(InvI0, VectorCross(VectorToPoint0, DX));
-		uniform FVector4 DQ0 = VectorQuaternionMultiply2(SetVector4(DR0, 0.0f), Q0) * 0.5f;
-	
-		P0 = P0 + DP0;
-		Q0 = VectorNormalizeQuaternion(Q0 + DQ0);
-		
-		SetCoMWorldTransform(PActor0, QActor0, P0, Q0, PCoMLocal0, QCoMLocal0);
-	}
-
-	if (InvM1 > 0.0f)
-	{
-		uniform FVector DXNeg = VectorNegate(DX);
-		uniform FVector DP1 = InvM1 * DXNeg;
-		uniform FVector DR1 = Multiply(InvI1, VectorCross(VectorToPoint1, DXNeg));
-		uniform FVector4 DQ1 = VectorQuaternionMultiply2(SetVector4(DR1, 0.0f), Q1) * 0.5f;
-		
-		P1 = P1 + DP1;
-		Q1 = VectorNormalizeQuaternion(Q1 + DQ1);
-		
-		SetCoMWorldTransform(PActor1, QActor1, P1, Q1, PCoMLocal1, QCoMLocal1);
-	}
-}
diff --git a/Engine/Source/Runtime/Experimental/Chaos/Private/Chaos/PBDCollisionConstraintsContact.cpp b/Engine/Source/Runtime/Experimental/Chaos/Private/Chaos/PBDCollisionConstraintsContact.cpp
--- a/Engine/Source/Runtime/Experimental/Chaos/Private/Chaos/PBDCollisionConstraintsContact.cpp
+++ b/Engine/Source/Runtime/Experimental/Chaos/Private/Chaos/PBDCollisionConstraintsContact.cpp
@@ -17,10 +17,6 @@
 
 namespace Chaos
 {
-#if INTEL_ISPC
-	extern bool bChaos_Collision_ISPC_Enabled;
-#endif
-
 	namespace CVars
 	{
 		int32 Chaos_Collision_EnergyClampEnabled = 1;
@@ -73,7 +69,7 @@ namespace Chaos
 			// \todo(chaos): see if we can easily switch to CoM transforms now in collision loop (shapes are held in actor space)
 			const FSolverBody& Body0 = *Constraint.GetSolverBody0();
 			const FSolverBody& Body1 = *Constraint.GetSolverBody1();
-			FGenericParticleHandle Particle0 = FGenericParticleHandle(Constraint.Particle[0]);
+			FGenericParticleHandle Particle0 = FGenericParticleHandle(Constraint.GetParticle0());
 			const FRigidTransform3 TransformXQ0(Body0.X() - Body0.ActorQ().RotateVector(Particle0->CenterOfMass()), Body0.ActorQ());
 			const FRigidTransform3 Transform1 = FRigidTransform3(Body1.ActorP(), Body1.ActorQ());
 
@@ -90,20 +86,26 @@ namespace Chaos
 		{
 			FVec3 AccumulatedImpulse(0);
 
-			FCollisionContact& Contact = Constraint.Manifold;
 			FSolverBody& Body0 = *Constraint.GetSolverBody0();
 			FSolverBody& Body1 = *Constraint.GetSolverBody1();
 
-			FVec3 VectorToPoint0 = Contact.Location - Body0.P();
-			FVec3 VectorToPoint1 = Contact.Location - Body1.P();
+			const FVec3 ContactLocation = Constraint.CalculateWorldContactLocation();
+			const FVec3 ContactNormal = Constraint.CalculateWorldContactNormal();
+			const FReal ContactFriction = Constraint.GetDynamicFriction();
+			const FReal ContactRestitution = Constraint.GetRestitution();
+			FReal ContactPhi = Constraint.GetPhi();
+			FReal ContactRestitutionPadding = Constraint.GetRestitutionPadding();
 
-			if (Contact.Phi < 0)
+			FVec3 VectorToPoint0 = ContactLocation - Body0.P();
+			FVec3 VectorToPoint1 = ContactLocation - Body1.P();
+
+			if (ContactPhi < 0)
 			{
 				*IterationParameters.NeedsAnotherIteration = true;
 
-				bool bApplyResitution = (Contact.Restitution > 0.0f);
-				bool bHaveRestitutionPadding = (Contact.RestitutionPadding > 0.0f);
-				bool bApplyFriction = (Contact.Friction > 0) && (IterationParameters.Dt > SMALL_NUMBER);
+				bool bApplyResitution = (ContactRestitution > 0.0f);
+				bool bHaveRestitutionPadding = (ContactRestitutionPadding > 0.0f);
+				bool bApplyFriction = (ContactFriction > 0) && (IterationParameters.Dt > SMALL_NUMBER);
 
 				// If we have restitution, pad the constraint by an amount that enforces the outgoing velocity constraint
 				// Really this should be per contact point, not per constraint.
@@ -119,13 +121,15 @@ namespace Chaos
 					FVec3 CV0 = BodyV0 + FVec3::CrossProduct(BodyW0, VectorToPoint0);
 					FVec3 CV1 = BodyV1 + FVec3::CrossProduct(BodyW1, VectorToPoint1);
 					FVec3 CV = CV0 - CV1;
-					FReal CVNormal = FVec3::DotProduct(CV, Contact.Normal);
+					FReal CVNormal = FVec3::DotProduct(CV, ContactNormal);
 
 					// No restitution below threshold normal velocity (CVNormal is negative here)
 					if (CVNormal < -ParticleParameters.RestitutionVelocityThreshold)
 					{
-						Contact.RestitutionPadding = -(1.0f + Contact.Restitution) * CVNormal * IterationParameters.Dt + Contact.Phi;
-						Contact.Phi -= Contact.RestitutionPadding;
+						ContactRestitutionPadding = -(1.0f + ContactRestitution) * CVNormal * IterationParameters.Dt + ContactPhi;
+						ContactPhi -= ContactRestitutionPadding;
+
+						Constraint.SetRestitutionPadding(ContactRestitutionPadding);
 					}
 				}
 			
@@ -134,8 +138,8 @@ namespace Chaos
 					(Body1.IsDynamic() ? ComputeFactorMatrix3(VectorToPoint1, FMatrix33(Body1.InvI()), FReal(Body1.InvM())) : FMatrix33(0));
 
 				// Calculate the normal correction
-				FVec3 NormalError = Contact.Phi * Contact.Normal;
-				FReal NormalImpulseDenominator = FVec3::DotProduct(Contact.Normal, ContactInvI * Contact.Normal);
+				FVec3 NormalError = ContactPhi * ContactNormal;
+				FReal NormalImpulseDenominator = FVec3::DotProduct(ContactNormal, ContactInvI * ContactNormal);
 				FVec3 NormalImpulseNumerator = -NormalError;
 				FVec3 NormalCorrection = NormalImpulseNumerator / NormalImpulseDenominator;
 
@@ -151,10 +155,10 @@ namespace Chaos
 					FVec3 CV0 = V0 + FVec3::CrossProduct(W0, VectorToPoint0);
 					FVec3 CV1 = V1 + FVec3::CrossProduct(W1, VectorToPoint1);
 					FVec3 CV = CV0 - CV1;
-					FReal CVNormal = FVec3::DotProduct(CV, Contact.Normal);
+					FReal CVNormal = FVec3::DotProduct(CV, ContactNormal);
 					if (CVNormal < 0.0f)
 					{
-						FVec3 CVLateral = CV - CVNormal * Contact.Normal;
+						FVec3 CVLateral = CV - CVNormal * ContactNormal;
 						FReal CVLateralMag = CVLateral.Size();
 						if (CVLateralMag > KINDA_SMALL_NUMBER)
 						{
@@ -164,9 +168,9 @@ namespace Chaos
 							LateralCorrection = LateralImpulseNumerator / LateralImpulseDenominator;
 							FReal LateralImpulseMag = LateralCorrection.Size();
 							FReal NormalImpulseMag = NormalCorrection.Size();
-							if (LateralImpulseMag > Contact.Friction * NormalImpulseMag)
+							if (LateralImpulseMag > ContactFriction * NormalImpulseMag)
 							{
-								LateralCorrection *= Contact.Friction * NormalImpulseMag / LateralImpulseMag;
+								LateralCorrection *= ContactFriction * NormalImpulseMag / LateralImpulseMag;
 							}
 						}
 					}
@@ -232,8 +236,8 @@ namespace Chaos
 				// to work, it should be set via the Scatter method of the SolverBody
 				if (ParticleParameters.Collided)
 				{
-					FGenericParticleHandle Particle0 = FGenericParticleHandle(Constraint.Particle[0]);
-					FGenericParticleHandle Particle1 = FGenericParticleHandle(Constraint.Particle[1]);
+					FGenericParticleHandle Particle0 = FGenericParticleHandle(Constraint.GetParticle0());
+					FGenericParticleHandle Particle1 = FGenericParticleHandle(Constraint.GetParticle1());
 					Particle0->AuxilaryValue(*ParticleParameters.Collided) = true;
 					Particle1->AuxilaryValue(*ParticleParameters.Collided) = true;
 				}
diff --git a/Engine/Source/Runtime/Experimental/Chaos/Private/Chaos/PBDConstraintRule.cpp b/Engine/Source/Runtime/Experimental/Chaos/Private/Chaos/PBDConstraintRule.cpp
--- a/Engine/Source/Runtime/Experimental/Chaos/Private/Chaos/PBDConstraintRule.cpp
+++ b/Engine/Source/Runtime/Experimental/Chaos/Private/Chaos/PBDConstraintRule.cpp
@@ -76,6 +76,7 @@ namespace Chaos
 	{
 		if(SolverData)
 		{
+			SolverData->GetConstraintIndex(Constraints.GetContainerId()) = 0;
 			Constraints.SetNumIslandConstraints(Constraints.NumConstraints(), *SolverData);
 			Constraints.GatherInput(Dt, *SolverData);
 		}
diff --git a/Engine/Source/Runtime/Experimental/Chaos/Private/Chaos/PBDRigidClustering.cpp b/Engine/Source/Runtime/Experimental/Chaos/Private/Chaos/PBDRigidClustering.cpp
--- a/Engine/Source/Runtime/Experimental/Chaos/Private/Chaos/PBDRigidClustering.cpp
+++ b/Engine/Source/Runtime/Experimental/Chaos/Private/Chaos/PBDRigidClustering.cpp
@@ -1113,7 +1113,7 @@ namespace Chaos
 				const FVec3 V0 = Rigid0 ? Rigid0->PreV() : FVec3(0);
 				const FVec3 V1 = Rigid1 ? Rigid1->PreV() : FVec3(0);
 				const FVec3 DeltaV = V0 - V1;
-				const FReal SpeedAlongNormal = FVec3::DotProduct(DeltaV, ContactHandle->GetContact().GetNormal());
+				const FReal SpeedAlongNormal = FVec3::DotProduct(DeltaV, ContactHandle->GetContact().CalculateWorldContactNormal());
 
 				// If we're not approaching at more than the min speed, reject the contact
 				if (SpeedAlongNormal > -MinContactSpeedForStrainEval && ContactHandle->GetAccumulatedImpulse().SizeSquared() > FReal(0))
@@ -1131,7 +1131,7 @@ namespace Chaos
 				const TArray<FPBDRigidParticleHandle*>& ParentToChildren)
 			{
 				const FRigidTransform3 WorldToClusterTM = FRigidTransform3(Cluster->P(), Cluster->Q());
-				const FVec3 ContactLocationClusterLocal = WorldToClusterTM.InverseTransformPosition(ContactHandle->GetContactLocation());
+				const FVec3 ContactLocationClusterLocal = WorldToClusterTM.InverseTransformPosition(ContactHandle->GetContact().CalculateWorldContactLocation());
 				FAABB3 ContactBox(ContactLocationClusterLocal, ContactLocationClusterLocal);
 				ContactBox.Thicken(ClusterDistanceThreshold);
 				if (Cluster->ChildrenSpatial())
diff --git a/Engine/Source/Runtime/Experimental/Chaos/Private/EventDefaults.cpp b/Engine/Source/Runtime/Experimental/Chaos/Private/EventDefaults.cpp
--- a/Engine/Source/Runtime/Experimental/Chaos/Private/EventDefaults.cpp
+++ b/Engine/Source/Runtime/Experimental/Chaos/Private/EventDefaults.cpp
@@ -122,15 +122,15 @@ namespace Chaos
 
 					if (ensure(!Constraint.AccumulatedImpulse.ContainsNaN() && FMath::IsFinite(Constraint.GetPhi())))
 					{
-						FGeometryParticleHandle* Particle0 = Constraint.Particle[0];
-						FGeometryParticleHandle* Particle1 = Constraint.Particle[1];
-						FKinematicGeometryParticleHandle* Body0 = Particle0->CastToKinematicParticle();
+						const FGeometryParticleHandle* Particle0 = Constraint.GetParticle0();
+						const FGeometryParticleHandle* Particle1 = Constraint.GetParticle1();
+						const FKinematicGeometryParticleHandle* Body0 = Particle0->CastToKinematicParticle();
 
 						// presently when a rigidbody or kinematic hits static geometry then Body1 is null
-						FKinematicGeometryParticleHandle* Body1 = Particle1->CastToKinematicParticle();
+						const FKinematicGeometryParticleHandle* Body1 = Particle1->CastToKinematicParticle();
 
-						const FImplicitObject* Implicit0 = Constraint.Manifold.Implicit[0];
-						const FImplicitObject* Implicit1 = Constraint.Manifold.Implicit[1];
+						const FImplicitObject* Implicit0 = Constraint.GetImplicit0();
+						const FImplicitObject* Implicit1 = Constraint.GetImplicit1();
 
 						const FPerShapeData* Shape0 = Particle0->GetImplicitShape(Implicit0);
 						const FPerShapeData* Shape1 = Particle1->GetImplicitShape(Implicit1);
@@ -145,13 +145,13 @@ namespace Chaos
 							continue;
 						}
 
-						FKinematicGeometryParticleHandle* Primary = Body0 ? Body0 : Body1;
-						FKinematicGeometryParticleHandle* Secondary = Body0 ? Body1 : Body0;
+						const FKinematicGeometryParticleHandle* Primary = Body0 ? Body0 : Body1;
+						const FKinematicGeometryParticleHandle* Secondary = Body0 ? Body1 : Body0;
 
 						if (!Constraint.AccumulatedImpulse.IsZero() && Primary)
 						{
-							if (ensure(!Constraint.GetLocation().ContainsNaN() &&
-								!Constraint.GetNormal().ContainsNaN()) &&
+							if (ensure(!Constraint.CalculateWorldContactLocation().ContainsNaN() &&
+								!Constraint.CalculateWorldContactNormal().ContainsNaN()) &&
 								!Primary->V().ContainsNaN() &&
 								!Primary->W().ContainsNaN() &&
 								(Secondary == nullptr || ((!Secondary->V().ContainsNaN()) && !Secondary->W().ContainsNaN())))
@@ -171,17 +171,18 @@ namespace Chaos
 					{
 						const FPBDCollisionConstraint& Constraint = ValidCollisionHandles[IdxCollision]->GetContact();
 
-						FGeometryParticleHandle* Particle0 = Constraint.Particle[0];
-						FGeometryParticleHandle* Particle1 = Constraint.Particle[1];
+						const FGeometryParticleHandle* Particle0 = Constraint.GetParticle0();
+						const FGeometryParticleHandle* Particle1 = Constraint.GetParticle1();
 
 						FCollidingData Data;
-						Data.Location = Constraint.GetLocation();
+						Data.Location = Constraint.CalculateWorldContactLocation();
 						Data.AccumulatedImpulse = Constraint.AccumulatedImpulse;
-						Data.Normal = Constraint.GetNormal();
+						Data.Normal = Constraint.CalculateWorldContactNormal();
 						Data.PenetrationDepth = Constraint.GetPhi();
-							
-						Data.Proxy1 = Particle0 ? Particle0->PhysicsProxy() : nullptr;
-						Data.Proxy2 = Particle1 ? Particle1->PhysicsProxy() : nullptr;
+						
+						// @todo(chaos): fix this casting
+						Data.Proxy1 = Particle0 ? const_cast<IPhysicsProxyBase*>(Particle0->PhysicsProxy()) : nullptr;
+						Data.Proxy2 = Particle1 ? const_cast<IPhysicsProxyBase*>(Particle1->PhysicsProxy()) : nullptr;
 
 						const FPerShapeData* Shape0 = Particle0 ? Particle0->GetImplicitShape(Constraint.GetImplicit0()) : nullptr;
 						const FPerShapeData* Shape1 = Particle1 ? Particle1->GetImplicitShape(Constraint.GetImplicit1()) : nullptr;
@@ -195,18 +196,18 @@ namespace Chaos
 							continue;
 						}
 
-						if (FPBDRigidParticleHandle * Rigid0 = Particle0->CastToRigidParticle())
+						if (const FPBDRigidParticleHandle * Rigid0 = Particle0->CastToRigidParticle())
 						{
 							Data.DeltaVelocity1 = Rigid0->V() - Rigid0->PreV();
 						}
-						if (FPBDRigidParticleHandle * Rigid1 = Particle1->CastToRigidParticle())
+						if (const FPBDRigidParticleHandle * Rigid1 = Particle1->CastToRigidParticle())
 						{
 							Data.DeltaVelocity2 = Rigid1->V() - Rigid1->PreV();
 						}
 
 						// todo: do we need these anymore now we are storing the particles you can access all of this stuff from there
 						// do we still need these now we have pointers to particles returned?
-						FPBDRigidParticleHandle* PBDRigid0 = Particle0->CastToRigidParticle();
+						const FPBDRigidParticleHandle* PBDRigid0 = Particle0->CastToRigidParticle();
 						if (PBDRigid0 && PBDRigid0->ObjectState() == EObjectStateType::Dynamic)
 						{
 							Data.Velocity1 = PBDRigid0->V();
@@ -214,7 +215,7 @@ namespace Chaos
 							Data.Mass1 = PBDRigid0->M();
 						}
 
-						FPBDRigidParticleHandle* PBDRigid1 = Particle1->CastToRigidParticle();
+						const FPBDRigidParticleHandle* PBDRigid1 = Particle1->CastToRigidParticle();
 						if (PBDRigid1 && PBDRigid1->ObjectState() == EObjectStateType::Dynamic)
 						{
 							Data.Velocity2 = PBDRigid1->V();
@@ -222,8 +223,8 @@ namespace Chaos
 							Data.Mass2 = PBDRigid1->M();
 						}
 
-						IPhysicsProxyBase* const PhysicsProxy = Particle0->PhysicsProxy();
-						IPhysicsProxyBase* const OtherPhysicsProxy = Particle1->PhysicsProxy();
+						IPhysicsProxyBase* const PhysicsProxy = const_cast<IPhysicsProxyBase*>(Particle0->PhysicsProxy());
+						IPhysicsProxyBase* const OtherPhysicsProxy = const_cast<IPhysicsProxyBase*>(Particle1->PhysicsProxy());
 							
 						const FSolverCollisionEventFilter* SolverCollisionEventFilter = Solver->GetEventFilters()->GetCollisionFilter();
 						if (!SolverCollisionEventFilter->Enabled() || SolverCollisionEventFilter->Pass(Data))
diff --git a/Engine/Source/Runtime/Experimental/Chaos/Public/Chaos/Collision/ContactPoint.h b/Engine/Source/Runtime/Experimental/Chaos/Public/Chaos/Collision/ContactPoint.h
--- a/Engine/Source/Runtime/Experimental/Chaos/Public/Chaos/Collision/ContactPoint.h
+++ b/Engine/Source/Runtime/Experimental/Chaos/Public/Chaos/Collision/ContactPoint.h
@@ -31,12 +31,6 @@ namespace Chaos
 		// Shape-space contact normal on the second shape with direction that points away from shape 1
 		FVec3 ShapeContactNormal;
 
-		// World-space contact location
-		FVec3 Location;
-
-		// World-space contact normal with direction that points away from shape 1
-		FVec3 Normal;
-
 		// Contact separation (negative for overlap)
 		FReal Phi;
 
@@ -59,7 +53,6 @@ namespace Chaos
 			{
 				Swap(ShapeContactPoints[0], ShapeContactPoints[1]);
 				ShapeContactNormal = -ShapeContactNormal;
-				Normal = -Normal;
 			}
 			return *this;
 		}
diff --git a/Engine/Source/Runtime/Experimental/Chaos/Public/Chaos/Collision/PBDCollisionConstraint.h b/Engine/Source/Runtime/Experimental/Chaos/Public/Chaos/Collision/PBDCollisionConstraint.h
--- a/Engine/Source/Runtime/Experimental/Chaos/Public/Chaos/Collision/PBDCollisionConstraint.h
+++ b/Engine/Source/Runtime/Experimental/Chaos/Public/Chaos/Collision/PBDCollisionConstraint.h
@@ -45,7 +45,6 @@ namespace Chaos
 			: ContactPoint()
 			, ShapeAnchorPoints{ FVec3(0), FVec3(0) }
 			, InitialShapeContactPoints{ FVec3(0), FVec3(0) }
-			, WorldContactPoints{ FVec3(0), FVec3(0) }
 			, NetPushOut(0)
 			, NetImpulse(0)
 			, Flags()
@@ -55,7 +54,6 @@ namespace Chaos
 			: ContactPoint(InContactPoint)
 			, ShapeAnchorPoints{ FVec3(0), FVec3(0) }
 			, InitialShapeContactPoints{ FVec3(0), FVec3(0) }
-			, WorldContactPoints{ FVec3(0), FVec3(0) }
 			, NetPushOut(0)
 			, NetImpulse(0)
 			, Flags()
@@ -64,20 +62,22 @@ namespace Chaos
 		FContactPoint ContactPoint;			// Contact point results of low-level collision detection
 		FVec3 ShapeAnchorPoints[2];
 		FVec3 InitialShapeContactPoints[2];	// ShapeContactPoints when the constraint was first initialized. Used to track reusablility
-		FVec3 WorldContactPoints[2];		// World-space contact points on the two bodies
 		FVec3 NetPushOut;					// Total pushout applied at this contact point
 		FVec3 NetImpulse;					// Total impulse applied by this contact point
 
 		union FFlags
 		{
-			FFlags() : Bits(0) {}
+			FFlags() { Reset(); }
+			
+			void Reset() { Bits = 0; }
+
 			struct
 			{
-				uint32 bInsideStaticFrictionCone : 1;		// Whether we are inside the static friction cone (used in PushOut)
-				uint32 bWasRestored : 1;
-				uint32 bWasFrictionRestored : 1;
+				uint8 bWasRestored : 1;
+				uint8 bWasFrictionRestored : 1;
+				uint8 bInsideStaticFrictionCone : 1;		// Whether we are inside the static friction cone
 			};
-			uint32 Bits;
+			uint8 Bits;
 		} Flags;
 	};
 
@@ -99,9 +99,6 @@ namespace Chaos
 	public:
 		FCollisionContact(const FImplicitObject* InImplicit0 = nullptr, const FBVHParticles* InSimplicial0 = nullptr, const FImplicitObject* InImplicit1 = nullptr, const FBVHParticles* InSimplicial1 = nullptr)
 			: bDisabled(false)
-			, Normal(0)
-			, Location(0)
-			, Phi(TNumericLimits<FReal>::Max())
 			, Friction(0)
 			, AngularFriction(0)
 			, Restitution(0)
@@ -121,10 +118,6 @@ namespace Chaos
 
 		bool bDisabled;
 
-		FVec3 Normal;
-		FVec3 Location;
-		FReal Phi;
-
 		FReal Friction;			// @todo(chaos): rename DynamicFriction
 		FReal AngularFriction;	// @todo(chaos): rename StaticFriction
 		FReal Restitution;
@@ -139,15 +132,9 @@ namespace Chaos
 		void Reset()
 		{
 			bDisabled = false;
-			Phi = TNumericLimits<FReal>::Max();
 			RestitutionPadding = 0;
 		}
 
-		FString ToString() const
-		{
-			return FString::Printf(TEXT("Location:%s, Normal:%s, Phi:%f"), *Location.ToString(), *Normal.ToString(), Phi);
-		}
-
 		const FImplicitObject* Implicit[2]; // {Of Particle[0], Of Particle[1]}
 		const FBVHParticles* Simplicial[2]; // {Of Particle[0], Of Particle[1]}
 	};
@@ -296,10 +283,17 @@ namespace Chaos
 		const FGeometryParticleHandle* GetParticle0() const { return Particle[0]; }
 		FGeometryParticleHandle* GetParticle1() { return Particle[1]; }
 		const FGeometryParticleHandle* GetParticle1() const { return Particle[1]; }
+		FGeometryParticleHandle* GetParticle(const int32 ParticleIndex) { check((ParticleIndex >= 0) && (ParticleIndex < 2)); return Particle[ParticleIndex]; }
+		const FGeometryParticleHandle* GetParticle(const int32 ParticleIndex) const { check((ParticleIndex >= 0) && (ParticleIndex < 2)); return Particle[ParticleIndex]; }
+
 		const FImplicitObject* GetImplicit0() const { return Manifold.Implicit[0]; }
 		const FImplicitObject* GetImplicit1() const { return Manifold.Implicit[1]; }
+		const FImplicitObject* GetImplicit(const int32 ParticleIndex) const { check((ParticleIndex >= 0) && (ParticleIndex < 2)); return Manifold.Implicit[ParticleIndex]; }
+
 		const FBVHParticles* GetCollisionParticles0() const { return Manifold.Simplicial[0]; }
 		const FBVHParticles* GetCollisionParticles1() const { return Manifold.Simplicial[1]; }
+		const FBVHParticles* GetCollisionParticles(const int32 ParticleIndex) const { check((ParticleIndex >= 0) && (ParticleIndex < 2)); return Manifold.Simplicial[ParticleIndex]; }
+
 		const FReal GetCollisionMargin0() const { return CollisionMargins[0]; }
 		const FReal GetCollisionMargin1() const { return CollisionMargins[1]; }
 		const FReal GetCollisionTolerance() const { return CollisionTolerance; }
@@ -312,24 +306,21 @@ namespace Chaos
 
 		// @todo(chaos): half of this API is wrong for the new multi-point manifold constraints. Remove it
 
-		const FCollisionContact& GetManifold() const { return Manifold; }
-
-		void ResetPhi(FReal InPhi) { SetPhi(InPhi); }
-
-		void SetPhi(FReal InPhi) { Manifold.Phi = InPhi; }
-		FReal GetPhi() const { return Manifold.Phi; }
+		void ResetPhi(FReal InPhi) { ClosestManifoldPointIndex = INDEX_NONE; }
+		FReal GetPhi() const { return (ClosestManifoldPointIndex != INDEX_NONE) ? ManifoldPoints[ClosestManifoldPointIndex].ContactPoint.Phi : TNumericLimits<FReal>::Max(); }
 
 		void SetDisabled(bool bInDisabled) { Manifold.bDisabled = bInDisabled; }
 		bool GetDisabled() const { return Manifold.bDisabled; }
 
 		virtual void SetIsSleeping(const bool bInIsSleeping) override;
-		//virtual bool IsSleeping() const override { return ContainerCookie.bIsSleeping; }
 
-		void SetNormal(const FVec3& InNormal) { Manifold.Normal = InNormal; }
-		FVec3 GetNormal() const { return Manifold.Normal; }
+		// Get the world-space normal of the closest manifold point
+		// @todo(chaos): remove (used by legacy RBAN collision solver)
+		FVec3 CalculateWorldContactNormal() const;
 
-		void SetLocation(const FVec3& InLocation) { Manifold.Location = InLocation; }
-		FVec3 GetLocation() const { return Manifold.Location; }
+		// Get the world-space contact location of the closest manifold point
+		// @todo(chaos): remove (used by legacy RBAN collision solver)
+		FVec3 CalculateWorldContactLocation() const;
 
 		void SetInvMassScale0(const FReal InInvMassScale) { Manifold.InvMassScale0 = InInvMassScale; }
 		FReal GetInvMassScale0() const { return Manifold.InvMassScale0; }
@@ -346,6 +337,23 @@ namespace Chaos
 		void SetStiffness(FReal InStiffness) { Stiffness = InStiffness; }
 		FReal GetStiffness() const { return Stiffness; }
 
+		void SetRestitution(const FReal InRestitution) { Manifold.Restitution = InRestitution; }
+		FReal GetRestitution() const { return Manifold.Restitution; }
+
+		void SetRestitutionThreshold(const FReal InRestitutionThreshold) { Manifold.RestitutionThreshold = InRestitutionThreshold; }
+		FReal GetRestitutionThreshold() const { return Manifold.RestitutionThreshold; }
+
+		void SetRestitutionPadding(const FReal InRestitutionPadding) { Manifold.RestitutionPadding = InRestitutionPadding; }
+		FReal GetRestitutionPadding() const { return Manifold.RestitutionPadding; }
+
+		void SetStaticFriction(const FReal InStaticFriction) { Manifold.AngularFriction = InStaticFriction; }
+		FReal GetStaticFriction() const { return FMath::Max(Manifold.AngularFriction, Manifold.Friction); }
+
+		void SetDynamicFriction(const FReal InDynamicFriction) { Manifold.Friction = InDynamicFriction; }
+		FReal GetDynamicFriction() const { return Manifold.Friction; }
+
+		EContactShapesType GetShapesType() const { return Manifold.ShapesType; }
+
 		FString ToString() const;
 
 		FReal GetCullDistance() const { return CullDistance; }
@@ -364,12 +372,20 @@ namespace Chaos
 		TArrayView<FManifoldPoint> GetManifoldPoints() { return MakeArrayView(ManifoldPoints.begin(), ManifoldPoints.Num()); }
 		TArrayView<const FManifoldPoint> GetManifoldPoints() const { return MakeArrayView(ManifoldPoints.begin(), ManifoldPoints.Num()); }
 
+		int32 NumManifoldPoints() const { return ManifoldPoints.Num(); }
+		FManifoldPoint& GetManifoldPoint(const int32 PointIndex) { return ManifoldPoints[PointIndex]; }
+		const FManifoldPoint& GetManifoldPoint(const int32 PointIndex) const { return ManifoldPoints[PointIndex]; }
+		const FManifoldPoint* GetClosestManifoldPoint() const { return (ClosestManifoldPointIndex != INDEX_NONE) ? &ManifoldPoints[ClosestManifoldPointIndex] : nullptr; }
+		const FSavedManifoldPoint* FindSavedManifoldPoint(const FManifoldPoint& ManifoldPoint) const;
+
 		void AddIncrementalManifoldContact(const FContactPoint& ContactPoint);
 		void AddOneshotManifoldContact(const FContactPoint& ContactPoint);
 		void UpdateManifoldContacts();
 
+		// Particle-relative transform of each collision shape in the constraint
 		const FRigidTransform3& GetShapeRelativeTransform0() const { return ImplicitTransform[0]; }
 		const FRigidTransform3& GetShapeRelativeTransform1() const { return ImplicitTransform[1]; }
+		const FRigidTransform3& GetShapeRelativeTransform(const int32 ParticleIndex) const { check((ParticleIndex >= 0) && (ParticleIndex < 2)); return ImplicitTransform[ParticleIndex]; }
 
 		const FRigidTransform3& GetShapeWorldTransform0() const { return ShapeWorldTransform0; }
 		const FRigidTransform3& GetShapeWorldTransform1() const { return ShapeWorldTransform1; }
@@ -391,7 +407,6 @@ namespace Chaos
 				ManifoldPoints[ManifoldPointIndex] = InManifoldPoints[ManifoldPointIndex];
 			}
 		}
-		void UpdateManifoldPointFromContact(const int32 ManifoldPointIndex);
 
 		// The GJK warm-start data. This is updated directly in the narrow phase
 		FGJKSimplexData& GetGJKWarmStartData() { return GJKWarmStartData; }
@@ -407,12 +422,6 @@ namespace Chaos
 			SolverBodies[1] = InSolverBody1;
 		}
 
-		/**
-		 * @brief Whether this constraint was newly created this tick (as opposed to restored from a previous tick)
-		 * @see IsRestored()
-		*/
-		//bool IsNew() const { return ContainerCookie.CreationEpoch == ContainerCookie.LastUsedEpoch; }
-
 		/**
 		 * @brief Whether this constraint was fully restored from a previous tick, and the manifold should be reused as-is
 		*/
@@ -466,25 +475,24 @@ namespace Chaos
 				const int32 SavedIndex = SavedManifoldPoints.Add();
 				FSavedManifoldPoint& SavedManifoldPoint = SavedManifoldPoints[SavedIndex];
 
-				if (FMath::IsNearlyEqual(StaticFrictionRatio, FReal(1)))
+				if (StaticFrictionRatio >= FReal(1.0f - KINDA_SMALL_NUMBER))
 				{
-					// Static friction held - we keep the same contacts points as-is for use next frame
+					// StaticFrictionRatio ~= 1: Static friction held - we keep the same contacts points as-is for use next frame
 					SavedManifoldPoint.ShapeContactPoints[0] = ManifoldPoint.ShapeAnchorPoints[0];
 					SavedManifoldPoint.ShapeContactPoints[1] = ManifoldPoint.ShapeAnchorPoints[1];
 				}
-				else if (FMath::IsNearlyEqual(StaticFrictionRatio, FReal(0)))
+				else if (StaticFrictionRatio < FReal(KINDA_SMALL_NUMBER))
 				{
-					// No friction - discard the friction anchors
+					// StaticFrictionRatio ~= 0: No friction (or no contact) - discard the friction anchors
 					SavedManifoldPoint.ShapeContactPoints[0] = ManifoldPoint.ContactPoint.ShapeContactPoints[0];
 					SavedManifoldPoint.ShapeContactPoints[1] = ManifoldPoint.ContactPoint.ShapeContactPoints[1];
 				}
 				else
 				{
-					// We exceeded the friction cone. Slide the friction anchor toward the last-detected contact position
-					// so that it sits at the edge of the friction cone.
-					const FReal Alpha = FMath::Clamp(StaticFrictionRatio, FReal(0), FReal(1));
-					SavedManifoldPoint.ShapeContactPoints[0] = FVec3::Lerp(ManifoldPoint.ContactPoint.ShapeContactPoints[0], ManifoldPoint.ShapeAnchorPoints[0], Alpha);
-					SavedManifoldPoint.ShapeContactPoints[1] = FVec3::Lerp(ManifoldPoint.ContactPoint.ShapeContactPoints[1], ManifoldPoint.ShapeAnchorPoints[1], Alpha);
+					// 0 < StaticFrictionRatio < 1: We exceeded the friction cone. Slide the friction anchor 
+					// toward the last-detected contact position so that it sits at the edge of the friction cone.
+					SavedManifoldPoint.ShapeContactPoints[0] = FVec3::Lerp(ManifoldPoint.ContactPoint.ShapeContactPoints[0], ManifoldPoint.ShapeAnchorPoints[0], StaticFrictionRatio);
+					SavedManifoldPoint.ShapeContactPoints[1] = FVec3::Lerp(ManifoldPoint.ContactPoint.ShapeContactPoints[1], ManifoldPoint.ShapeAnchorPoints[1], StaticFrictionRatio);
 				}
 			}
 		}
@@ -524,28 +532,27 @@ namespace Chaos
 		int32 FindManifoldPoint(const FContactPoint& ContactPoint) const;
 		int32 AddManifoldPoint(const FContactPoint& ContactPoint);
 		void InitManifoldPoint(const int32 ManifoldPointIndex);
-		void SetActiveContactPoint(const FContactPoint& ContactPoint);
 
-		/**
-		 * @brief Restore data used for static friction from the previous point(s) to the current
-		*/
-		void TryRestoreFrictionData(const int32 ManifoldPointIndex);
+		// Update the store Phi for the manifold point based on current world-space shape transforms
+		// @todo(chaos): Only intended for use by the legacy solvers - remove it
+		void UpdateManifoldPointPhi(const int32 ManifoldPointIndex);
 
 		void ReprojectManifoldContacts();
 		void ReprojectManifoldPoint(const int32 ManifoldPointIndex);
 
 		void InitMarginsAndTolerances(const EImplicitObjectType ImplicitType0, const EImplicitObjectType ImplicitType1, const FReal Margin0, const FReal Margin1);
 
-		const FSavedManifoldPoint* FindSavedManifoldPoint(const FManifoldPoint& ManifoldPoint) const;
+	private:
 		FReal CalculateSavedManifoldPointScore(const FSavedManifoldPoint& SavedManifoldPoint, const FManifoldPoint& ManifoldPoint, const FReal DistanceToleranceSq) const;
 
-	public:
 		//@todo(chaos): make this stuff private
 		FRigidTransform3 ImplicitTransform[2];		// Local-space transforms of the shape (relative to particle)
 		FGeometryParticleHandle* Particle[2];
-		FVec3 AccumulatedImpulse;					// @todo(chaos): we need to accumulate angular impulse separately
 		FCollisionContact Manifold;// @todo(chaos): rename
 
+	public:
+		FVec3 AccumulatedImpulse;					// @todo(chaos): we need to accumulate angular impulse separately
+
 		// Value in range [0,1] used to interpolate P between [X,P] that we will rollback to when solving at time of impact.
 		FReal TimeOfImpact;
 
@@ -554,10 +561,6 @@ namespace Chaos
 		ECollisionCCDType CCDType;
 		FReal Stiffness;
 
-		TCArray<FManifoldPoint, MaxManifoldPoints> ManifoldPoints;
-
-		// The manifold points from the previous tick when we don't reuse the manifold. Used by static friction.
-		TCArray<FSavedManifoldPoint, MaxManifoldPoints> SavedManifoldPoints;
 		FReal CullDistance;
 
 		// The margins to use during collision detection. We don't always use the margins on the shapes directly.
@@ -570,6 +573,8 @@ namespace Chaos
 
 		FReal FrictionPositionTolerance;
 
+		int32 ClosestManifoldPointIndex;
+
 		union FFlags
 		{
 			FFlags() : Bits(0) {}
@@ -587,6 +592,9 @@ namespace Chaos
 		// These are only needed here while we still have the legacy solvers (not QuasiPBD)
 		FSolverBody* SolverBodies[2];
 
+		TCArray<FManifoldPoint, MaxManifoldPoints> ManifoldPoints;
+		TCArray<FSavedManifoldPoint, MaxManifoldPoints> SavedManifoldPoints;
+
 		// Simplex data from the last call to GJK, used to warm-start GJK
 		FGJKSimplexData GJKWarmStartData;
 
diff --git a/Engine/Source/Runtime/Experimental/Chaos/Public/Chaos/Collision/PBDCollisionConstraintHandle.h b/Engine/Source/Runtime/Experimental/Chaos/Public/Chaos/Collision/PBDCollisionConstraintHandle.h
--- a/Engine/Source/Runtime/Experimental/Chaos/Public/Chaos/Collision/PBDCollisionConstraintHandle.h
+++ b/Engine/Source/Runtime/Experimental/Chaos/Public/Chaos/Collision/PBDCollisionConstraintHandle.h
@@ -79,7 +79,7 @@ namespace Chaos
 
 		virtual bool IsEnabled() const override;
 
-		FVec3 GetContactLocation() const;
+		//FVec3 GetContactLocation() const;
 
 		FVec3 GetAccumulatedImpulse() const;
 
diff --git a/Engine/Source/Runtime/Experimental/Chaos/Public/Chaos/CollisionResolutionUtil.h b/Engine/Source/Runtime/Experimental/Chaos/Public/Chaos/CollisionResolutionUtil.h
--- a/Engine/Source/Runtime/Experimental/Chaos/Public/Chaos/CollisionResolutionUtil.h
+++ b/Engine/Source/Runtime/Experimental/Chaos/Public/Chaos/CollisionResolutionUtil.h
@@ -73,9 +73,6 @@ namespace Chaos
 			const FVec3& ContactVelocity0,
 			const FVec3& ContactVelocity1);
 
-		bool 
-		SampleObjectHelper(const FImplicitObject& Object, const FRigidTransform3& ObjectTransform, const FRigidTransform3& SampleToObjectTransform, const FVec3& SampleParticle, FReal Thickness, FContactPoint& ContactPoint);
-
 		bool 
 		SampleObjectNoNormal(const FImplicitObject& Object, const FRigidTransform3& ObjectTransform, const FRigidTransform3& SampleToObjectTransform, const FVec3& SampleParticle, FReal Thickness, FContactPoint& ContactPoint);
 
diff --git a/Engine/Source/Runtime/Experimental/Chaos/Public/Chaos/EvolutionResimCache.h b/Engine/Source/Runtime/Experimental/Chaos/Public/Chaos/EvolutionResimCache.h
--- a/Engine/Source/Runtime/Experimental/Chaos/Public/Chaos/EvolutionResimCache.h
+++ b/Engine/Source/Runtime/Experimental/Chaos/Public/Chaos/EvolutionResimCache.h
@@ -59,10 +59,10 @@ namespace Chaos
 				WeakPairs.Empty(Constraints.Num());
 				for(FPBDCollisionConstraint& Constraint : Constraints)
 				{
-					WeakPairs.Add(FWeakConstraintPair{Constraint.Particle[0]->WeakParticleHandle(),Constraint.Particle[1]->WeakParticleHandle()});
+					WeakPairs.Add(FWeakConstraintPair{Constraint.GetParticle0()->WeakParticleHandle(),Constraint.GetParticle1()->WeakParticleHandle()});
 
-					auto* A = Constraint.Particle[0];
-					auto* B = Constraint.Particle[1];
+					auto* A = Constraint.GetParticle0();
+					auto* B = Constraint.GetParticle1();
 
 					//Need to do this on save for any new constraints
 					MarkSoftIfDesync(*A,*B);
