Chaos - fix CollisionKey when the game uses both Local and Global IDs for particles

diff --git a/Engine/Source/Runtime/Experimental/Chaos/Public/Chaos/Collision/CollisionKeys.h b/Engine/Source/Runtime/Experimental/Chaos/Public/Chaos/Collision/CollisionKeys.h
--- a/Engine/Source/Runtime/Experimental/Chaos/Public/Chaos/Collision/CollisionKeys.h
+++ b/Engine/Source/Runtime/Experimental/Chaos/Public/Chaos/Collision/CollisionKeys.h
@@ -32,6 +32,7 @@ namespace Chaos
 	/**
 	 * @brief A key which uniquely identifes a particle pair for use by the collision detection system
 	 * This key will be the same if particles order is reversed.
+	 * @note This uses ParticleID and truncates it from 32 to 31 bits
 	*/
 	class FCollisionParticlePairKey
 	{
@@ -56,27 +57,45 @@ namespace Chaos
 	private:
 		void GenerateKey(const FGeometryParticleHandle* Particle0, const FGeometryParticleHandle* Particle1)
 		{
-			int32 ID0 = (Particle0->ParticleID().LocalID != INDEX_NONE) ? Particle0->ParticleID().LocalID : Particle0->ParticleID().GlobalID;
-			int32 ID1 = (Particle1->ParticleID().LocalID != INDEX_NONE) ? Particle1->ParticleID().LocalID : Particle1->ParticleID().GlobalID;
+			const bool bIsLocalID0 = Particle0->ParticleID().LocalID != INDEX_NONE;
+			const bool bIsLocalID1 = Particle1->ParticleID().LocalID != INDEX_NONE;
+			const uint32 ID0 = uint32((bIsLocalID0) ? Particle0->ParticleID().LocalID : Particle0->ParticleID().GlobalID);
+			const uint32 ID1 = uint32((bIsLocalID1) ? Particle1->ParticleID().LocalID : Particle1->ParticleID().GlobalID);
 
 			if (ID0 < ID1)
 			{
-				Key.Key32s[0] = ID0;
-				Key.Key32s[1] = ID1;
+				Key.Key32s[0].Key31 = ID0;
+				Key.Key32s[0].IsLocal = bIsLocalID0;
+				Key.Key32s[1].Key31 = ID1;
+				Key.Key32s[1].IsLocal = bIsLocalID1;
 			}
 			else
 			{
-				Key.Key32s[0] = ID1;
-				Key.Key32s[1] = ID0;
+				Key.Key32s[0].Key31 = ID1;
+				Key.Key32s[0].IsLocal = bIsLocalID1;
+				Key.Key32s[1].Key31 = ID0;
+				Key.Key32s[1].IsLocal = bIsLocalID0;
 			}
 		}
 
+		struct FParticleIDKey
+		{
+			uint32 Key31 : 31;
+			uint32 IsLocal : 1;
+		};
 		union FIDKey
 		{
 			uint64 Key64;
-			int32 Key32s[2];
+			FParticleIDKey Key32s[2];
 		};
 
+		// This class is sensitive to changes in FParticleID - try to catch that here...
+		static_assert(sizeof(FParticleID) == 8, "FParticleID size does not match FCollisionParticlePairKey (expected 64 bits)");
+		static_assert(sizeof(FParticleID::GlobalID) == 4, "FParticleID::GlobalID size does not match FCollisionParticlePairKey (expected 32 bits)");
+		static_assert(sizeof(FParticleID::LocalID) == 4, "FParticleID::LocalID size does not match FCollisionParticlePairKey (expected 32 bits)");
+		static_assert(sizeof(FParticleIDKey) == 4, "FCollisionParticlePairKey::FParticleIDKey size is not 32 bits");
+		static_assert(sizeof(FIDKey) == 8, "FCollisionParticlePairKey::FIDKey size is not 64 bits");
+
 		FIDKey Key;
 	};
 
