Chaos - convex-convex optimizations
- change some functions back to floats
- optimize clipping
- reorder some tests in plane seacrh
- specialize SelectContactPlane for scaled convex
- reorder operations in Manifold Restore test

diff --git a/Engine/Source/Runtime/Experimental/Chaos/Private/Chaos/Collision/PBDCollisionConstraint.cpp b/Engine/Source/Runtime/Experimental/Chaos/Private/Chaos/Collision/PBDCollisionConstraint.cpp
--- a/Engine/Source/Runtime/Experimental/Chaos/Private/Chaos/Collision/PBDCollisionConstraint.cpp
+++ b/Engine/Source/Runtime/Experimental/Chaos/Private/Chaos/Collision/PBDCollisionConstraint.cpp
@@ -428,26 +428,6 @@ namespace Chaos
 		}
 	}
 
-	void FPBDCollisionConstraint::AddOneshotManifoldContact(const FContactPoint& ContactPoint)
-	{
-		if (ContactPoint.IsSet())
-		{
-			if (ManifoldPoints.IsFull())
-			{
-				return;
-			}
-			
-			int32 ManifoldPointIndex = AddManifoldPoint(ContactPoint);
-
-			if (ManifoldPoints[ManifoldPointIndex].ContactPoint.Phi < GetPhi())
-			{
-				ClosestManifoldPointIndex = ManifoldPointIndex;
-			}
-
-			Flags.bUseIncrementalManifold = false;
-		}
-	}
-
 	void FPBDCollisionConstraint::AddIncrementalManifoldContact(const FContactPoint& ContactPoint)
 	{
 		if (ManifoldPoints.IsFull())
@@ -489,28 +469,6 @@ namespace Chaos
 		Flags.bUseIncrementalManifold = true;
 	}
 
-	void FPBDCollisionConstraint::InitManifoldPoint(const int32 ManifoldPointIndex)
-	{
-		FManifoldPoint& ManifoldPoint = ManifoldPoints[ManifoldPointIndex];
-		ManifoldPoint.InitialShapeContactPoints[0] = ManifoldPoint.ContactPoint.ShapeContactPoints[0];
-		ManifoldPoint.InitialShapeContactPoints[1] = ManifoldPoint.ContactPoint.ShapeContactPoints[1];
-
-		ManifoldPoints[ManifoldPointIndex].Flags.Reset();
-		ManifoldPoints[ManifoldPointIndex].NetPushOut = FVec3(0);
-		ManifoldPoints[ManifoldPointIndex].NetImpulse = FVec3(0);
-	}
-
-	int32 FPBDCollisionConstraint::AddManifoldPoint(const FContactPoint& ContactPoint)
-	{
-		int32 ManifoldPointIndex = ManifoldPoints.Add();
-
-		ManifoldPoints[ManifoldPointIndex].ContactPoint = ContactPoint;
-
-		InitManifoldPoint(ManifoldPointIndex);
-
-		return ManifoldPointIndex;
-	}
-
 	bool FPBDCollisionConstraint::CanUseManifold(FGeometryParticleHandle* Particle0, FGeometryParticleHandle* Particle1) const
 	{
 		// Do not use manifolds when a body is connected by a joint to another. Manifolds do not work when the bodies may be moved
@@ -623,6 +581,34 @@ namespace Chaos
 		ExpectedNumManifoldPoints = ManifoldPoints.Num();
 		Flags.bWasManifoldRestored = false;
 
+		// If we did not remove any contact points and we have not moved or rotated much we may reuse the manifold as-is.
+		bool bMovedBeyondTolerance = true;
+		if ((ShapePositionTolerance > 0) && (ShapeRotationThreshold > 0))
+		{
+			// The transform check is necessary regardless of how many points we have left in the manifold because
+			// as a body moves/rotates we may have to change which faces/edges are colliding. We can't know if the face/edge
+			// will change until we run the closest-point checks (GJK) in the narrow phase.
+			const FVec3 Shape1ToShape0Translation = ShapeWorldTransform0.GetTranslation() - ShapeWorldTransform1.GetTranslation();
+			const FVec3 OriginalShape1ToShape0Translation = LastShapeWorldTransform0.GetTranslation() - LastShapeWorldTransform1.GetTranslation();
+			const FVec3 TranslationDelta = Shape1ToShape0Translation - OriginalShape1ToShape0Translation;
+			if (TranslationDelta.IsNearlyZero(ShapePositionTolerance))
+			{
+				const FRotation3 Shape1toShape0Rotation = ShapeWorldTransform0.GetRotation().Inverse() * ShapeWorldTransform1.GetRotation();
+				const FRotation3 OriginalShape1toShape0Rotation = LastShapeWorldTransform0.GetRotation().Inverse() * LastShapeWorldTransform1.GetRotation();
+				const FReal RotationOverlap = FRotation3::DotProduct(Shape1toShape0Rotation, OriginalShape1toShape0Rotation);
+				if (RotationOverlap > ShapeRotationThreshold)
+				{
+					bMovedBeyondTolerance = false;
+				}
+			}
+		}
+
+		if (bMovedBeyondTolerance)
+		{
+			ResetActiveManifoldContacts();
+			return false;
+		}
+
 		// Either update or remove each manifold point depending on how far it has moved from its initial relative point
 		// NOTE: We do not reset if we have 0 points - we can still "restore" a zero point manifold if the bodies have not moved
 		int32 ManifoldPointToRemove = INDEX_NONE;
@@ -655,6 +641,7 @@ namespace Chaos
 					const FVec3 ShapeContactPoint1 = Contact0In1 - ContactPhi * ContactNormalIn1;
 					ManifoldPoint.ContactPoint.ShapeContactPoints[1] = ShapeContactPoint1;
 					ManifoldPoint.ContactPoint.Phi = ContactPhi;
+					ManifoldPoint.Flags.bWasRestored = true;
 					if (ManifoldPoint.ContactPoint.Phi < GetPhi())
 					{
 						ClosestManifoldPointIndex = ManifoldPointIndex;
@@ -673,7 +660,7 @@ namespace Chaos
 				}
 			}
 
-			// Remove points - only one point removal suport required (see above)
+			// Remove points - only one point removal support required (see above)
 			if (ManifoldPointToRemove != INDEX_NONE)
 			{
 				ManifoldPoints.RemoveAt(ManifoldPointToRemove);
@@ -682,31 +669,12 @@ namespace Chaos
 					--ClosestManifoldPointIndex;
 					check(ClosestManifoldPointIndex >= 0);
 				}
+				return false;
 			}
 		}
 
-		// If we did not remove any contact points and we have not moved or rotated much we may reuse the manifold as-is.
-		if ((ManifoldPointToRemove == INDEX_NONE) && (ShapePositionTolerance > 0) && (ShapeRotationThreshold > 0))
-		{
-			// The transform check is necessary regardless of how many points we have left in the manifold because
-			// as a body moves/rotates we may have to change which faces/edges are colliding. We can't know if the face/edge
-			// will change until we run the closest-point checks (GJK) in the narrow phase.
-			const FVec3 Shape1ToShape0Translation = ShapeWorldTransform0.GetTranslation() - ShapeWorldTransform1.GetTranslation();
-			const FVec3 OriginalShape1ToShape0Translation = LastShapeWorldTransform0.GetTranslation() - LastShapeWorldTransform1.GetTranslation();
-			const FVec3 TranslationDelta = Shape1ToShape0Translation - OriginalShape1ToShape0Translation;
-			if (TranslationDelta.IsNearlyZero(ShapePositionTolerance))
-			{
-				const FRotation3 Shape1toShape0Rotation = ShapeWorldTransform0.GetRotation().Inverse() * ShapeWorldTransform1.GetRotation();
-				const FRotation3 OriginalShape1toShape0Rotation = LastShapeWorldTransform0.GetRotation().Inverse() * LastShapeWorldTransform1.GetRotation();
-				const FReal RotationOverlap = FRotation3::DotProduct(Shape1toShape0Rotation, OriginalShape1toShape0Rotation);
-				if (RotationOverlap > ShapeRotationThreshold)
-				{
-					return true;
-				}
-			}
-		}
-
-		return false;
+		Flags.bWasManifoldRestored = true;
+		return true;
 	}
 
 	bool FPBDCollisionConstraint::TryAddManifoldContact(const FContactPoint& NewContactPoint)
diff --git a/Engine/Source/Runtime/Experimental/Chaos/Private/Chaos/CollisionOneShotManifolds.cpp b/Engine/Source/Runtime/Experimental/Chaos/Private/Chaos/CollisionOneShotManifolds.cpp
--- a/Engine/Source/Runtime/Experimental/Chaos/Private/Chaos/CollisionOneShotManifolds.cpp
+++ b/Engine/Source/Runtime/Experimental/Chaos/Private/Chaos/CollisionOneShotManifolds.cpp
@@ -312,56 +312,55 @@ namespace Chaos
 		// more vertices may be added to outputVertexBuffer by this function
 		// This is the core of the Sutherland-Hodgman algorithm
 		// Plane Normals face outwards 
-		uint32 ClipVerticesAgainstPlane(const FVec3* InputVertexBuffer, FVec3* outputVertexBuffer, uint32 ClipPointCount, uint32 MaxNumberOfOutputPoints, FVec3 ClippingPlaneNormal, FReal PlaneDistance)
+		inline uint32 ClipVerticesAgainstPlane(const FVec3* InputVertexBuffer, FVec3* OutputVertexBuffer, const uint32 ClipPointCount, const uint32 MaxNumberOfOutputPoints, const FVec3 ClippingPlaneNormal, const FReal PlaneDistance)
 		{
+			check(ClipPointCount > 0);
 
-			auto CalculateIntersect = [=](const FVec3& Point1, const FVec3& Point2) -> FVec3
-			{
-				// Only needs to be valid if the line connecting Point1 with Point2 actually intersects
-				FVec3 Result;
+			uint32 NewClipPointCount = 0;
+
+			FVec3 PrevClipPoint;
+			FReal PrevClipPointDotNormal;
+			FVec3 CurrentClipPoint = InputVertexBuffer[ClipPointCount - 1];;
+			FReal CurrentClipPointDotNormal = FVec3::DotProduct(CurrentClipPoint, ClippingPlaneNormal);
+			const FReal PlaneClipDistance = PlaneDistance + PlaneDistance * SMALL_NUMBER;
 
-				FReal Denominator = FVec3::DotProduct(Point2 - Point1, ClippingPlaneNormal); // Can be negative
+			auto CalculateIntersect = [&PrevClipPoint, &CurrentClipPoint, &PlaneDistance](const FReal Dot1, const FReal Dot2)
+			{
+				const FReal Denominator = Dot2 - Dot1; // Can be negative
 				if (FMath::Abs(Denominator) < SMALL_NUMBER)
 				{
-					Result = Point1;
+					return PrevClipPoint;
 				}
 				else
 				{
-					FReal Alpha = (PlaneDistance - FVec3::DotProduct(Point1, ClippingPlaneNormal)) / Denominator;
-					Result = FMath::Lerp(Point1, Point2, Alpha);
+					const FReal Alpha = (PlaneDistance - Dot1) / Denominator;
+					return FVec3::Lerp(PrevClipPoint, CurrentClipPoint, Alpha);
 				}
-				return Result;
-			};
-
-			auto InsideClipFace = [=](const FVec3& Point) -> bool
-			{
-				// Epsilon is there so that previously clipped points will still be inside the plane
-				return FVec3::DotProduct(Point, ClippingPlaneNormal) <= PlaneDistance + PlaneDistance * SMALL_NUMBER; 
 			};
 
-			uint32 NewClipPointCount = 0;
-
 			for (uint32 ClipPointIndex = 0; ClipPointIndex < ClipPointCount; ClipPointIndex++)
 			{
-				FVec3 CurrentClipPoint = InputVertexBuffer[ClipPointIndex];
-				FVec3 PrevClipPoint = InputVertexBuffer[(ClipPointIndex + ClipPointCount - 1) % ClipPointCount];
-				FVec3 InterSect = CalculateIntersect(PrevClipPoint, CurrentClipPoint);
+				PrevClipPoint = CurrentClipPoint;
+				PrevClipPointDotNormal = CurrentClipPointDotNormal;
+				CurrentClipPoint = InputVertexBuffer[ClipPointIndex];
+				CurrentClipPointDotNormal = FVec3::DotProduct(CurrentClipPoint, ClippingPlaneNormal);
 
-				if (InsideClipFace(CurrentClipPoint))
+				// Epsilon is there so that previously clipped points will still be inside the plane
+				if (CurrentClipPointDotNormal <= PlaneClipDistance)
 				{
-					if (!InsideClipFace(PrevClipPoint))
+					if (PrevClipPointDotNormal > PlaneClipDistance)
 					{
-						outputVertexBuffer[NewClipPointCount++] = InterSect;
+						OutputVertexBuffer[NewClipPointCount++] = CalculateIntersect(PrevClipPointDotNormal, CurrentClipPointDotNormal);
 						if (NewClipPointCount >= MaxNumberOfOutputPoints)
 						{
 							break;
 						}
 					}
-					outputVertexBuffer[NewClipPointCount++] = CurrentClipPoint;
+					OutputVertexBuffer[NewClipPointCount++] = CurrentClipPoint;
 				}
-				else if (InsideClipFace(PrevClipPoint))
+				else if (PrevClipPointDotNormal < PlaneClipDistance)
 				{
-					outputVertexBuffer[NewClipPointCount++] = InterSect;
+					OutputVertexBuffer[NewClipPointCount++] = CalculateIntersect(PrevClipPointDotNormal, CurrentClipPointDotNormal);
 				}
 
 				if (NewClipPointCount >= MaxNumberOfOutputPoints)
@@ -404,7 +403,7 @@ namespace Chaos
 			const int32 RefConvexFaceVerticesNum = RefConvex.NumPlaneVertices(RefPlaneIndex);
 			int32 ClippingPlaneCount = RefConvexFaceVerticesNum;
 			FVec3 PrevPoint = RefConvex.GetVertex(RefConvex.GetPlaneVertex(RefPlaneIndex, ClippingPlaneCount - 1));
-			for (int32 ClippingPlaneIndex = 0; ClippingPlaneIndex < ClippingPlaneCount; ++ClippingPlaneIndex)
+			for (int32 ClippingPlaneIndex = 0; (ClippingPlaneIndex < ClippingPlaneCount) && (ContactPointCount > 1); ++ClippingPlaneIndex)
 			{
 				FVec3 CurrentPoint = RefConvex.GetVertex(RefConvex.GetPlaneVertex(RefPlaneIndex, ClippingPlaneIndex));
 				FVec3 ClippingPlaneNormal = RefWindingOrder * FVec3::CrossProduct(RefPlaneNormal, PrevPoint - CurrentPoint);
@@ -459,14 +458,14 @@ namespace Chaos
 			FReal& BestPlaneDot)
 		{
 			const TPlaneConcrete<FReal, 3> Plane = Convex.GetPlane(PlaneIndex);
-			const FReal PlaneNormalDotN = FVec3::DotProduct(N, Plane.Normal());
 				
-			// Ignore planes that do not oppose N
-			if (PlaneNormalDotN <= -SMALL_NUMBER)
+			// Reject planes farther than MaxDistance
+			const FReal PlaneDistance = Plane.SignedDistance(X);
+			if (FMath::Abs(PlaneDistance) <= MaxDistance)
 			{
-				// Reject planes farther than MaxDistance
-				const FReal PlaneDistance = Plane.SignedDistance(X);
-				if (FMath::Abs(PlaneDistance) <= MaxDistance)
+				// Ignore planes that do not oppose N
+				const FReal PlaneNormalDotN = FVec3::DotProduct(N, Plane.Normal());
+				if (PlaneNormalDotN <= -SMALL_NUMBER)
 				{
 					// Keep the most opposing plane
 					if (PlaneNormalDotN < BestPlaneDot)
@@ -480,36 +479,42 @@ namespace Chaos
 
 		// Specialization for scaled convex. Avoids the instantiation of a scaled plane object
 		// which almost doubles the cost of the function.
-		template <>
-		void FindBestPlane<TImplicitObjectScaled<FConvex>>(
+		void FindBestPlaneScaledConvex(
 			const TImplicitObjectScaled<FConvex>& ScaledConvex,
-			const FVec3& X,
-			const FVec3& N,
-			const FReal MaxDistance,
+			const FConvex::FVec3Type& X,
+			const FConvex::FVec3Type& N,
+			const FConvex::FVec3Type& Scale,
+			const FConvex::FVec3Type& ScaleInv,
+			const FConvex::FRealType MaxDistance,
 			const int32 PlaneIndex,
-			int32& BestPlaneIndex,
-			FReal& BestPlaneDot)
+			int32& OutBestPlaneIndex,
+			FConvex::FRealType& InOutBestPlaneDot)
 		{
+			using FConvexReal = FConvex::FRealType;
+			using FConvexVec3 = FConvex::FVec3Type;
+			using FConvexPlane = FConvex::FPlaneType;
+
 			const FConvex* UnscaledConvex = ScaledConvex.GetInnerObject()->template GetObject<FConvex>();
-			const TPlaneConcrete<FReal, 3>& UnscaledPlane = UnscaledConvex->GetPlane(PlaneIndex);
+			const FConvexPlane& UnscaledPlane = UnscaledConvex->GetPlaneRaw(PlaneIndex);
+
+			const FConvexVec3 ScaledPlaneX = UnscaledPlane.X() * Scale;
+			FConvexVec3 ScaledPlaneN = UnscaledPlane.Normal() * ScaleInv;
 
-			const FVec3 ScaledPlaneX = FVec3(UnscaledPlane.X()) * ScaledConvex.GetScale();
-			FVec3 ScaledPlaneN = FVec3(UnscaledPlane.Normal()) / ScaledConvex.GetScale();
 			if (ScaledPlaneN.Normalize())
 			{
-				// Ignore planes that do not oppose N
-				const FReal PlaneNormalDotN = FVec3::DotProduct(N, ScaledPlaneN);
-				if (PlaneNormalDotN <= -SMALL_NUMBER)
+				// Reject planes farther than MaxDistance
+				const FConvexReal PlaneDistance = FConvexVec3::DotProduct(X - ScaledPlaneX, ScaledPlaneN);
+				if (FMath::Abs(PlaneDistance) <= MaxDistance)
 				{
-					// Reject planes farther than MaxDistance
-					const FReal PlaneDistance = FVec3::DotProduct(X - ScaledPlaneX, ScaledPlaneN);
-					if (FMath::Abs(PlaneDistance) <= MaxDistance)
+					// Ignore planes that do not oppose N
+					const FConvexReal PlaneNormalDotN = FConvexVec3::DotProduct(N, ScaledPlaneN);
+					if (PlaneNormalDotN <= -SMALL_NUMBER)
 					{
 						// Keep the most opposing plane
-						if (PlaneNormalDotN < BestPlaneDot)
+						if (PlaneNormalDotN < InOutBestPlaneDot)
 						{
-							BestPlaneDot = PlaneNormalDotN;
-							BestPlaneIndex = PlaneIndex;
+							InOutBestPlaneDot = PlaneNormalDotN;
+							OutBestPlaneIndex = PlaneIndex;
 						}
 					}
 				}
@@ -569,6 +574,71 @@ namespace Chaos
 			return BestPlaneIndex;
 		}
 
+		template <>
+		int32 SelectContactPlane<TImplicitObjectScaled<FConvex>>(
+			const TImplicitObjectScaled<FConvex>& ScaledConvex,
+			const FVec3 InX,
+			const FVec3 InN,
+			const FReal InMaxDistance,
+			const int32 VertexIndex)
+		{
+			using FConvexReal = FConvex::FRealType;
+			using FConvexVec3 = FConvex::FVec3Type;
+			using FConvexPlane = FConvex::FPlaneType;
+
+			// Handle InMaxDistance = 0. We expect that the X is actually on the surface in this case, so the search distance just needs to be some reasonable tolerance.
+			// @todo(chaos): this should probable be dependent on the size of the objects...
+			const FConvexReal MinFaceSearchDistance = Chaos_Collision_Manifold_MinFaceSearchDistance;
+			const FConvexReal MaxDistance = FMath::Max(FConvexReal(InMaxDistance), MinFaceSearchDistance);
+
+			// LWC_TODO: Precision Loss
+			// Scale precision is ok as long as we only support large positions and not large sizes...
+			// N precision is ok with floats (normalized)
+			// X precision is ok as long as we don't try to find the exact separation of objects that are separated by LWC distances
+			// which we don't in collision detection (objects must pass bounds checks first).
+			const FConvexVec3 Scale = FConvexVec3(ScaledConvex.GetScale());
+			const FConvexVec3 ScaleInv = FConvexVec3(ScaledConvex.GetInvScale());
+			const FConvexVec3 X = FConvexVec3(InX);
+			const FConvexVec3 N = FConvexVec3(InN);
+
+			int32 BestPlaneIndex = INDEX_NONE;
+			FConvexReal BestPlaneDot = 1.0f;
+			{
+				int32 PlaneIndices[3] = { INDEX_NONE, INDEX_NONE, INDEX_NONE };
+				int32 NumPlanes = ScaledConvex.GetVertexPlanes3(VertexIndex, PlaneIndices[0], PlaneIndices[1], PlaneIndices[2]);
+
+				// If we have more than 3 planes we iterate over the full set of planes since it is faster than using the half edge structure
+				if (NumPlanes > 3)
+				{
+					NumPlanes = ScaledConvex.NumPlanes();
+					for (int32 PlaneIndex = 0; PlaneIndex < NumPlanes; ++PlaneIndex)
+					{
+						FindBestPlaneScaledConvex(ScaledConvex, X, N, Scale, ScaleInv, MaxDistance, PlaneIndex, BestPlaneIndex, BestPlaneDot);
+					}
+				}
+				// Otherwise we iterate over the cached planes
+				else
+				{
+					for (int32 PlaneIndex = 0; PlaneIndex < NumPlanes; ++PlaneIndex)
+					{
+						FindBestPlaneScaledConvex(ScaledConvex, X, N, Scale, ScaleInv, MaxDistance, PlaneIndices[PlaneIndex], BestPlaneIndex, BestPlaneDot);
+					}
+				}
+			}
+
+			// Malformed convexes or half-spaces or capsules could have all planes rejected above.
+			// If that happens, select the most opposing plane including those that
+			// may point the same direction as N. 
+			if (BestPlaneIndex == INDEX_NONE)
+			{
+				// This always returns a valid plane.
+				BestPlaneIndex = ScaledConvex.GetMostOpposingPlane(N);
+			}
+
+			check(BestPlaneIndex != INDEX_NONE);
+			return BestPlaneIndex;
+		}
+
 
 		template <typename ConvexImplicitType1, typename ConvexImplicitType2>
 		void ConstructConvexConvexOneShotManifold(
@@ -688,11 +758,11 @@ namespace Chaos
 			// We will use a double buffer as an optimization
 			const uint32 MaxContactPointCount = 32; // This should be tuned
 			uint32 ContactPointCount = 0;
-			FVec3 ClippedVertices1[MaxContactPointCount];
-			FVec3 ClippedVertices2[MaxContactPointCount];
 			FVec3* ClippedVertices = nullptr;
 			const FRigidTransform3* RefConvexTM;
 			FRigidTransform3 ConvexOtherToRef;
+			FVec3 ClippedVertices1[MaxContactPointCount];
+			FVec3 ClippedVertices2[MaxContactPointCount];
 
 			{
 				SCOPE_CYCLE_COUNTER_MANIFOLD_CLIP();
diff --git a/Engine/Source/Runtime/Experimental/Chaos/Private/Chaos/Convex.cpp b/Engine/Source/Runtime/Experimental/Chaos/Private/Chaos/Convex.cpp
--- a/Engine/Source/Runtime/Experimental/Chaos/Private/Chaos/Convex.cpp
+++ b/Engine/Source/Runtime/Experimental/Chaos/Private/Chaos/Convex.cpp
@@ -151,23 +151,31 @@ namespace Chaos
 		return MostOpposingIdx;
 	}
 
-	FVec3 FConvex::GetClosestEdgePosition(int32 PlaneIndex, const FVec3& Position) const
+	FVec3 FConvex::GetClosestEdgePosition(int32 PlaneIndex, const FVec3& InPosition) const
 	{
-		FVec3 ClosestEdgePosition = FVec3(0);
-		FReal ClosestDistanceSq = FLT_MAX;
+		FVec3Type ClosestEdgePosition = FVec3Type(0);
+		FRealType ClosestDistanceSq = FLT_MAX;
+		FVec3Type Position = FVec3Type(InPosition);
 
 		const int32 PlaneVerticesNum = NumPlaneVertices(PlaneIndex);
 		if (PlaneVerticesNum > 0)
 		{
-			FVec3 P0 = GetVertex(GetPlaneVertex(PlaneIndex, PlaneVerticesNum - 1));
+			FVec3Type P0 = GetVertex(GetPlaneVertex(PlaneIndex, PlaneVerticesNum - 1));
 			for (int32 PlaneVertexIndex = 0; PlaneVertexIndex < PlaneVerticesNum; ++PlaneVertexIndex)
 			{
 				const int32 VertexIndex = GetPlaneVertex(PlaneIndex, PlaneVertexIndex);
-				const FVec3 P1 = GetVertex(VertexIndex);
+				const FVec3Type P1 = GetVertex(VertexIndex);
 				
-				const FVec3 EdgePosition = FMath::ClosestPointOnLine(P0, P1, Position);
-				const FReal EdgeDistanceSq = (EdgePosition - Position).SizeSquared();
-
+				// LWC_TODO: low-precision version of FMath::ClosestPointOnLine
+				// See FMath::ClosestPointOnLine (which does not have a float version when LWC is enbled)
+				//const FVec3Type EdgePosition = FMath::ClosestPointOnLine(P0, P1, Position);
+				const FVec3Type DP = P1 - P0;
+				const FRealType A = FVec3Type::DotProduct((P0 - Position), DP);
+				const FRealType B = DP.SizeSquared();
+				const FRealType T = FMath::Clamp(-A / B, FRealType(0), FRealType(1));
+				const FVec3Type EdgePosition = P0 + (T * DP);
+				
+				const FRealType EdgeDistanceSq = (EdgePosition - Position).SizeSquared();
 				if (EdgeDistanceSq < ClosestDistanceSq)
 				{
 					ClosestDistanceSq = EdgeDistanceSq;
@@ -178,7 +186,7 @@ namespace Chaos
 			}
 		}
 
-		return ClosestEdgePosition;
+		return FVec3(ClosestEdgePosition);
 	}
 
 	bool FConvex::GetClosestEdgeVertices(int32 PlaneIndex, const FVec3& Position, int32& OutVertexIndex0, int32& OutVertexIndex1) const
diff --git a/Engine/Source/Runtime/Experimental/Chaos/Public/Chaos/Collision/PBDCollisionConstraint.h b/Engine/Source/Runtime/Experimental/Chaos/Public/Chaos/Collision/PBDCollisionConstraint.h
--- a/Engine/Source/Runtime/Experimental/Chaos/Public/Chaos/Collision/PBDCollisionConstraint.h
+++ b/Engine/Source/Runtime/Experimental/Chaos/Public/Chaos/Collision/PBDCollisionConstraint.h
@@ -379,7 +379,21 @@ namespace Chaos
 		const FSavedManifoldPoint* FindSavedManifoldPoint(const FManifoldPoint& ManifoldPoint) const;
 
 		void AddIncrementalManifoldContact(const FContactPoint& ContactPoint);
-		void AddOneshotManifoldContact(const FContactPoint& ContactPoint);
+
+		inline void AddOneshotManifoldContact(const FContactPoint& ContactPoint)
+		{
+			if (ContactPoint.IsSet() && !ManifoldPoints.IsFull())
+			{
+				int32 ManifoldPointIndex = AddManifoldPoint(ContactPoint);
+				if (ManifoldPoints[ManifoldPointIndex].ContactPoint.Phi < GetPhi())
+				{
+					ClosestManifoldPointIndex = ManifoldPointIndex;
+				}
+				Flags.bUseIncrementalManifold = false;
+			}
+		}
+
+
 		void UpdateManifoldContacts();
 
 		// Particle-relative transform of each collision shape in the constraint
@@ -530,8 +544,24 @@ namespace Chaos
 
 		bool AreMatchingContactPoints(const FContactPoint& A, const FContactPoint& B, FReal& OutScore) const;
 		int32 FindManifoldPoint(const FContactPoint& ContactPoint) const;
-		int32 AddManifoldPoint(const FContactPoint& ContactPoint);
-		void InitManifoldPoint(const int32 ManifoldPointIndex);
+
+		inline void InitManifoldPoint(const int32 ManifoldPointIndex)
+		{
+			FManifoldPoint& ManifoldPoint = ManifoldPoints[ManifoldPointIndex];
+			ManifoldPoint.InitialShapeContactPoints[0] = ManifoldPoint.ContactPoint.ShapeContactPoints[0];
+			ManifoldPoint.InitialShapeContactPoints[1] = ManifoldPoint.ContactPoint.ShapeContactPoints[1];
+			ManifoldPoint.Flags.Reset();
+			ManifoldPoint.NetPushOut = FVec3(0);
+			ManifoldPoint.NetImpulse = FVec3(0);
+		}
+
+		inline int32 AddManifoldPoint(const FContactPoint& ContactPoint)
+		{
+			int32 ManifoldPointIndex = ManifoldPoints.Add();	// Note: no initialization (see TCArray)
+			ManifoldPoints[ManifoldPointIndex].ContactPoint = ContactPoint;
+			InitManifoldPoint(ManifoldPointIndex);
+			return ManifoldPointIndex;
+		}
 
 		// Update the store Phi for the manifold point based on current world-space shape transforms
 		// @todo(chaos): Only intended for use by the legacy solvers - remove it
diff --git a/Engine/Source/Runtime/Experimental/Chaos/Public/Chaos/Convex.h b/Engine/Source/Runtime/Experimental/Chaos/Public/Chaos/Convex.h
--- a/Engine/Source/Runtime/Experimental/Chaos/Public/Chaos/Convex.h
+++ b/Engine/Source/Runtime/Experimental/Chaos/Public/Chaos/Convex.h
@@ -418,6 +418,11 @@ namespace Chaos
 			return TPlaneConcrete<FReal, 3>::MakeFrom(Planes[FaceIndex]);
 		}
 
+		const FPlaneType& GetPlaneRaw(int32 FaceIndex) const
+		{
+			return Planes[FaceIndex];
+		}
+
 		// Get the vertex at the specified index (e.g., indices from GetPlaneVertexs)
 		const FVec3Type& GetVertex(int32 VertexIndex) const
 		{
