Chaos - fix levelset collision

diff --git a/Engine/Source/Runtime/Experimental/Chaos/Private/Chaos/CollisionResolutionUtil.cpp b/Engine/Source/Runtime/Experimental/Chaos/Private/Chaos/CollisionResolutionUtil.cpp
--- a/Engine/Source/Runtime/Experimental/Chaos/Private/Chaos/CollisionResolutionUtil.cpp
+++ b/Engine/Source/Runtime/Experimental/Chaos/Private/Chaos/CollisionResolutionUtil.cpp
@@ -160,7 +160,9 @@ namespace Chaos
 		FAutoConsoleVariableRef CVarSampleMinParticlesForAcceleration(TEXT("p.SampleMinParticlesForAcceleration"), SampleMinParticlesForAcceleration, TEXT("The minimum number of particles needed before using an acceleration structure when sampling"));
 
 
-#if INTEL_ISPC
+#if INTEL_ISPC && 0
+		// LWC_TODO: get ISPC levelset collision working again
+
 		template<ECollisionUpdateType UpdateType>
 		FContactPoint SampleObject(const FImplicitObject& Object, const FRigidTransform3& ObjectTransform, const FBVHParticles& SampleParticles, const FRigidTransform3& SampleParticlesTransform, FReal CullingDistance)
 		{
@@ -490,7 +492,7 @@ namespace Chaos
 						Contact.Phi = NewPhi;
 						Contact.ShapeContactPoints[0] = SampleToObjectTM.InverseTransformPositionNoScale(LocalPoint);
 						Contact.ShapeContactPoints[1] = LocalPoint - NewPhi * LocalNormal;
-						Contact.ShapeContactNormal = SampleToObjectTM.InverseTransformVectorNoScale(LocalNormal);
+						Contact.ShapeContactNormal = LocalNormal;
 					}
 				}
 				else
@@ -506,7 +508,7 @@ namespace Chaos
 				Contact.Phi = Object.PhiWithNormal(LocalPoint, LocalNormal);
 				Contact.ShapeContactPoints[0] = SampleParticles.X(DeepestParticle);
 				Contact.ShapeContactPoints[1] = LocalPoint - Contact.Phi * LocalNormal;
-				Contact.ShapeContactNormal = SampleToObjectTM.InverseTransformVectorNoScale(LocalNormal);
+				Contact.ShapeContactNormal = LocalNormal;
 			}
 
 			return Contact;
@@ -602,7 +604,7 @@ namespace Chaos
 						Contact.Phi = NewPhi;
 						Contact.ShapeContactPoints[0] = SampleToObjectTM.InverseTransformPositionNoScale(LocalPoint);
 						Contact.ShapeContactPoints[1] = LocalPoint - NewPhi * LocalNormal; 
-						Contact.ShapeContactNormal = SampleToObjectTM.InverseTransformVectorNoScale(LocalNormal);
+						Contact.ShapeContactNormal = LocalNormal;
 					}
 				}
 				else
@@ -618,7 +620,7 @@ namespace Chaos
 				Contact.Phi = Object.PhiWithNormal(LocalPoint, LocalNormal);
 				Contact.ShapeContactPoints[0] = SampleParticles.X(DeepestParticle);
 				Contact.ShapeContactPoints[1] = LocalPoint - Contact.Phi * LocalNormal;
-				Contact.ShapeContactNormal = SampleToObjectTM.InverseTransformVectorNoScale(LocalNormal);
+				Contact.ShapeContactNormal = LocalNormal;
 			}
 
 			return Contact;
