Fix bad constraint removal at solver destruction

diff --git a/Engine/Source/Programs/HeadlessChaos/Private/HeadlessChaosTestGraph.cpp b/Engine/Source/Programs/HeadlessChaos/Private/HeadlessChaosTestGraph.cpp
--- a/Engine/Source/Programs/HeadlessChaos/Private/HeadlessChaosTestGraph.cpp
+++ b/Engine/Source/Programs/HeadlessChaos/Private/HeadlessChaosTestGraph.cpp
@@ -588,6 +588,7 @@ namespace ChaosTest {
 				
 				++Index;
 			}
+			Graph.ResetIndices();
 		}
 	}
 
diff --git a/Engine/Source/Runtime/Experimental/Chaos/Private/Chaos/Island/IslandManager.cpp b/Engine/Source/Runtime/Experimental/Chaos/Private/Chaos/Island/IslandManager.cpp
--- a/Engine/Source/Runtime/Experimental/Chaos/Private/Chaos/Island/IslandManager.cpp
+++ b/Engine/Source/Runtime/Experimental/Chaos/Private/Chaos/Island/IslandManager.cpp
@@ -277,6 +277,9 @@ FPBDIslandManager::FPBDIslandManager(const TParticleView<FPBDRigidParticles>& PB
 }
 
 FPBDIslandManager::~FPBDIslandManager()
+{}
+
+void FPBDIslandManager::ResetIndices()
 {
 	// Reset of all the particles / constraint graph index to INDEX_NONE
 	for(auto& ItemNode : IslandGraph->ItemNodes)
@@ -337,6 +340,33 @@ void FPBDIslandManager::InitializeGraph(const TParticleView<FPBDRigidParticles>&
 	}
 }
 
+void FPBDIslandManager::RemoveConstraints()
+{
+	for(auto& ItemEdge : IslandGraph->ItemEdges)
+	{
+		ItemEdge.Key->SetConstraintGraphIndex(INDEX_NONE);
+	}
+	IslandGraph->ItemEdges.Reset();
+	IslandGraph->GraphEdges.Reset();
+	
+	for(auto& GraphIsland : IslandGraph->GraphIslands)
+	{
+		GraphIsland.NumEdges = 0;
+	}
+	for(auto& GraphNode: IslandGraph->GraphNodes)
+	{
+		GraphNode.NodeEdges.Reset();
+	}
+	for(auto& IslandSolver : IslandSolvers)
+	{
+		IslandSolver->ClearConstraints();
+	}
+	for(auto& IslandGroup : IslandGroups)
+	{
+		IslandGroup->NumConstraints() = 0;
+	}
+}
+
 int32 FPBDIslandManager::ReserveParticles(const int32 NumParticles)
 {
 	const int32 MaxIndex = IslandGraph->NumNodes();
diff --git a/Engine/Source/Runtime/Experimental/Chaos/Private/Chaos/PBDRigidsEvolutionGBF.cpp b/Engine/Source/Runtime/Experimental/Chaos/Private/Chaos/PBDRigidsEvolutionGBF.cpp
--- a/Engine/Source/Runtime/Experimental/Chaos/Private/Chaos/PBDRigidsEvolutionGBF.cpp
+++ b/Engine/Source/Runtime/Experimental/Chaos/Private/Chaos/PBDRigidsEvolutionGBF.cpp
@@ -709,9 +709,10 @@ FPBDRigidsEvolutionGBF::FPBDRigidsEvolutionGBF(FPBDRigidsSOAs& InParticles,THand
 
 FPBDRigidsEvolutionGBF::~FPBDRigidsEvolutionGBF()
 {
+	// Not sure we need to reset the indices since all the particles and constraints are going to be destroyed
+	//GetConstraintGraph().ResetIndices();
 }
 
-
 void FPBDRigidsEvolutionGBF::Serialize(FChaosArchive& Ar)
 {
 	Base::Serialize(Ar);
diff --git a/Engine/Source/Runtime/Experimental/Chaos/Public/Chaos/Island/IslandManager.h b/Engine/Source/Runtime/Experimental/Chaos/Public/Chaos/Island/IslandManager.h
--- a/Engine/Source/Runtime/Experimental/Chaos/Public/Chaos/Island/IslandManager.h
+++ b/Engine/Source/Runtime/Experimental/Chaos/Public/Chaos/Island/IslandManager.h
@@ -51,6 +51,16 @@ public:
 	*/
 	~FPBDIslandManager();
 
+	/**
+	* Reset nodes and edges graph indices 
+	*/
+	void ResetIndices();
+
+	/**
+	* Remove all the constraints from the graph
+	*/
+	void RemoveConstraints();
+
 	/**
 	  * Default island manager constructor
 	  * @param Particles List of particles to be used to fill the graph nodes
diff --git a/Engine/Source/Runtime/Experimental/Chaos/Public/Chaos/PBDRigidsEvolution.h b/Engine/Source/Runtime/Experimental/Chaos/Public/Chaos/PBDRigidsEvolution.h
--- a/Engine/Source/Runtime/Experimental/Chaos/Public/Chaos/PBDRigidsEvolution.h
+++ b/Engine/Source/Runtime/Experimental/Chaos/Public/Chaos/PBDRigidsEvolution.h
@@ -577,6 +577,9 @@ public:
 
 	CHAOS_API void ResetConstraints()
 	{
+		// Remove all the constraints from the graph
+		GetConstraintGraph().RemoveConstraints();
+		
 		// Clear all particle lists of collisions and constraints
 		// (this could be performed by the constraint containers
 		// but it would be unnecessarily expensive to remove them
