Fix for static analysis warning
https://horde.devtools.epicgames.com/log/61e6af45837b79f7ccdbf866?lineindex=5129

diff --git a/Engine/Source/Runtime/Experimental/Chaos/Private/Chaos/CCDUtilities.cpp b/Engine/Source/Runtime/Experimental/Chaos/Private/Chaos/CCDUtilities.cpp
--- a/Engine/Source/Runtime/Experimental/Chaos/Private/Chaos/CCDUtilities.cpp
+++ b/Engine/Source/Runtime/Experimental/Chaos/Private/Chaos/CCDUtilities.cpp
@@ -222,8 +222,12 @@ namespace Chaos
 			if (Island == INDEX_NONE)
 			{
 				// non-dynamic pairs are already ignored in Init() so if Particle 0 is null the second one should not be 
-				ensure(CCDConstraint.Particle[1] != nullptr); 
-				Island = CCDConstraint.Particle[1]->Island;
+				ensure(CCDConstraint.Particle[1] != nullptr);
+
+				if (CCDConstraint.Particle[1])
+				{
+					Island = CCDConstraint.Particle[1]->Island;
+				}	
 			}
 			CCDConstraint.Island = Island;
 			IslandConstraintNum[Island]++;
