Chaos - minor optimizations and cleanup
- reorder some properties in collision constraint
- remove some unused data and functions in midphase and constraint
- WIP to re-enable GJK point injection

#okforgithuib public

diff --git a/Engine/Source/Runtime/Experimental/Chaos/Private/Chaos/ChaosDebugDraw.cpp b/Engine/Source/Runtime/Experimental/Chaos/Private/Chaos/ChaosDebugDraw.cpp
--- a/Engine/Source/Runtime/Experimental/Chaos/Private/Chaos/ChaosDebugDraw.cpp
+++ b/Engine/Source/Runtime/Experimental/Chaos/Private/Chaos/ChaosDebugDraw.cpp
@@ -918,12 +918,17 @@ namespace Chaos
 						if (Contact.WasManifoldRestored())
 						{
 							const FReal BoxScale = Settings.DrawScale * Settings.ContactWidth;
-							FDebugDrawQueue::GetInstance().DrawDebugBox(WorldPlaneLocation, FVec3(BoxScale, BoxScale, FReal(0.01)), FRotation3(FRotationMatrix::MakeFromZ(WorldPlaneNormal)), FColor::Blue, false, KINDA_SMALL_NUMBER, Settings.DrawPriority, 0.5f * Settings.LineThickness);
+							FDebugDrawQueue::GetInstance().DrawDebugBox(WorldPlaneLocation, FVec3(BoxScale, BoxScale, FReal(0.01)), FRotation3(FRotationMatrix::MakeFromZ(WorldPlaneNormal)), FColor::Green, false, KINDA_SMALL_NUMBER, Settings.DrawPriority, 0.5f * Settings.LineThickness);
 						}
 						else if (ManifoldPoint.Flags.bWasRestored)
 						{
 							const FReal BoxScale = Settings.DrawScale * Settings.ContactWidth;
-							FDebugDrawQueue::GetInstance().DrawDebugBox(WorldPlaneLocation, FVec3(BoxScale, BoxScale, FReal(0.01)), FRotation3(FRotationMatrix::MakeFromZ(WorldPlaneNormal)), FColor::Purple, false, KINDA_SMALL_NUMBER, Settings.DrawPriority, 0.5f * Settings.LineThickness);
+							FDebugDrawQueue::GetInstance().DrawDebugBox(WorldPlaneLocation, FVec3(BoxScale, BoxScale, FReal(0.01)), FRotation3(FRotationMatrix::MakeFromZ(WorldPlaneNormal)), FColor::Yellow, false, KINDA_SMALL_NUMBER, Settings.DrawPriority, 0.5f * Settings.LineThickness);
+						}
+						else if (ManifoldPoint.Flags.bWasReplaced)
+						{
+							const FReal BoxScale = Settings.DrawScale * Settings.ContactWidth;
+							FDebugDrawQueue::GetInstance().DrawDebugBox(WorldPlaneLocation, FVec3(BoxScale, BoxScale, FReal(0.01)), FRotation3(FRotationMatrix::MakeFromZ(WorldPlaneNormal)), FColor::Orange, false, KINDA_SMALL_NUMBER, Settings.DrawPriority, 0.5f * Settings.LineThickness);
 						}
 					}
 				}
diff --git a/Engine/Source/Runtime/Experimental/Chaos/Private/Chaos/Collision/PBDCollisionConstraint.cpp b/Engine/Source/Runtime/Experimental/Chaos/Private/Chaos/Collision/PBDCollisionConstraint.cpp
--- a/Engine/Source/Runtime/Experimental/Chaos/Private/Chaos/Collision/PBDCollisionConstraint.cpp
+++ b/Engine/Source/Runtime/Experimental/Chaos/Private/Chaos/Collision/PBDCollisionConstraint.cpp
@@ -18,7 +18,8 @@
 
 namespace Chaos
 {
-	extern bool bChaos_Collision_EnableManifoldInject;
+	extern bool bChaos_Collision_EnableManifoldGJKInject;
+	extern bool bChaos_Collision_EnableManifoldGJKReplace;
 
 	FRealSingle Chaos_Manifold_MatchPositionTolerance = 0.3f;		// Fraction of object size position tolerance
 	FRealSingle Chaos_Manifold_MatchNormalTolerance = 0.02f;		// Dot product tolerance
@@ -98,6 +99,30 @@ namespace Chaos
 		return false;
 	}
 
+	void FPBDCollisionConstraint::MakeInline(
+		FGeometryParticleHandle* Particle0,
+		const FImplicitObject* Implicit0,
+		const FBVHParticles* Simplicial0,
+		const FRigidTransform3& ImplicitLocalTransform0,
+		FGeometryParticleHandle* Particle1,
+		const FImplicitObject* Implicit1,
+		const FBVHParticles* Simplicial1,
+		const FRigidTransform3& ImplicitLocalTransform1,
+		const FReal InCullDistance,
+		const bool bInUseManifold,
+		const EContactShapesType ShapesType,
+		FPBDCollisionConstraint& OutConstraint)
+	{
+		OutConstraint.Particle[0] = Particle0;
+		OutConstraint.Particle[1] = Particle1;
+		OutConstraint.Implicit[0] = Implicit0;
+		OutConstraint.Implicit[1] = Implicit1;
+		OutConstraint.Simplicial[0] = Simplicial0;
+		OutConstraint.Simplicial[1] = Simplicial1;
+
+		OutConstraint.Setup(ECollisionCCDType::Disabled, ShapesType, ImplicitLocalTransform0, ImplicitLocalTransform1, InCullDistance, bInUseManifold);
+	}
+
 	TUniquePtr<FPBDCollisionConstraint> FPBDCollisionConstraint::Make(
 		FGeometryParticleHandle* Particle0,
 		const FImplicitObject* Implicit0,
@@ -141,31 +166,30 @@ namespace Chaos
 	FPBDCollisionConstraint::FPBDCollisionConstraint()
 		: ImplicitTransform{ FRigidTransform3(), FRigidTransform3() }
 		, Particle{ nullptr, nullptr }
+		, Implicit{ nullptr, nullptr }
+		, Simplicial{ nullptr, nullptr }
 		, Manifold()
+		, Stiffness(1)
 		, AccumulatedImpulse(0)
 		, TimeOfImpact(0)
 		, ContainerCookie()
+		, ShapesType(EContactShapesType::Unknown)
 		, CCDType(ECollisionCCDType::Disabled)
-		, Stiffness(FReal(1))
+		, ShapeWorldTransform0()
+		, ShapeWorldTransform1()
 		, CullDistance(TNumericLimits<FReal>::Max())
 		, CollisionMargins{ 0, 0 }
 		, CollisionTolerance(0)
+		, ClosestManifoldPointIndex(INDEX_NONE)
+		, ExpectedNumManifoldPoints(0)
 		, Flags()
+		, LastShapeWorldPositionDelta()
+		, LastShapeWorldRotationDelta()
 		, SolverBodies{ nullptr, nullptr }
+		, GJKWarmStartData()
 		, ManifoldPoints()
 		, SavedManifoldPoints()
-		, GJKWarmStartData()
-		, ShapeWorldTransform0()
-		, ShapeWorldTransform1()
-		, LastShapeWorldTransform0()
-		, LastShapeWorldTransform1()
-		, ExpectedNumManifoldPoints(0)
 	{
-		Manifold.Implicit[0] = nullptr;
-		Manifold.Implicit[1] = nullptr;
-		Manifold.Simplicial[0] = nullptr;
-		Manifold.Simplicial[1] = nullptr;
-		Manifold.ShapesType = EContactShapesType::Unknown;
 	}
 
 	FPBDCollisionConstraint::FPBDCollisionConstraint(
@@ -177,31 +201,30 @@ namespace Chaos
 		const FBVHParticles* Simplicial1)
 		: ImplicitTransform{ FRigidTransform3(), FRigidTransform3() }
 		, Particle{ Particle0, Particle1 }
+		, Implicit{ Implicit0, Implicit1 }
+		, Simplicial{ Simplicial0, Simplicial1 }
 		, Manifold()
+		, Stiffness(1)
 		, AccumulatedImpulse(0)
 		, TimeOfImpact(0)
 		, ContainerCookie()
+		, ShapesType(EContactShapesType::Unknown)
 		, CCDType(ECollisionCCDType::Disabled)
-		, Stiffness(FReal(1))
+		, ShapeWorldTransform0()
+		, ShapeWorldTransform1()
 		, CullDistance(TNumericLimits<FReal>::Max())
 		, CollisionMargins{ 0, 0 }
 		, CollisionTolerance(0)
+		, ClosestManifoldPointIndex(INDEX_NONE)
+		, ExpectedNumManifoldPoints(0)
 		, Flags()
+		, LastShapeWorldPositionDelta()
+		, LastShapeWorldRotationDelta()
 		, SolverBodies{ nullptr, nullptr }
+		, GJKWarmStartData()
 		, ManifoldPoints()
 		, SavedManifoldPoints()
-		, GJKWarmStartData()
-		, ShapeWorldTransform0()
-		, ShapeWorldTransform1()
-		, LastShapeWorldTransform0()
-		, LastShapeWorldTransform1()
-		, ExpectedNumManifoldPoints(0)
 	{
-		Manifold.Implicit[0] = Implicit0;
-		Manifold.Implicit[1] = Implicit1;
-		Manifold.Simplicial[0] = Simplicial0;
-		Manifold.Simplicial[1] = Simplicial1;
-		Manifold.ShapesType = EContactShapesType::Unknown;
 	}
 
 	void FPBDCollisionConstraint::Setup(
@@ -214,7 +237,7 @@ namespace Chaos
 	{
 		CCDType = InCCDType;
 
-		Manifold.ShapesType = InShapesType;
+		ShapesType = InShapesType;
 
 		ImplicitTransform[0] = InImplicitLocalTransform0;
 		ImplicitTransform[1] = InImplicitLocalTransform1;
@@ -410,6 +433,7 @@ namespace Chaos
 			ShapeWorldTransform1 = FRigidTransform3(GetSolverBody1()->CorrectedP(), GetSolverBody1()->CorrectedQ());
 		}
 
+		Flags.bDisabled = false;
 		ClosestManifoldPointIndex = INDEX_NONE;
 		Manifold.Reset();
 
@@ -488,84 +512,13 @@ namespace Chaos
 		Manifold.Reset();
 		ManifoldPoints.Reset();
 		ExpectedNumManifoldPoints = 0;
+		Flags.bDisabled = false;
 		Flags.bWasManifoldRestored = false;
 	}
 
-	void FPBDCollisionConstraint::RestoreManifold(const bool bReproject)
-	{
-		// We want to restore the manifold as-is and skip the narrow phase, but...
-		if (bReproject)
-		{
-			// The bodies may have moved a moderate amount so we need to update the shape-space (and subsequently the world-space) contact points
-			ReprojectManifoldContacts();
-		}
-		else
-		{
-			// The bodies have barely moved so we keep the same shape-space contacts. We still need to update ManifoldPoint Phis
-			UpdateManifoldContacts();
-		}
-
-		Flags.bWasManifoldRestored = true;
-	}
-
-
-	void FPBDCollisionConstraint::SetShapeWorldTransforms(const FRigidTransform3& InShapeWorldTransform0, const FRigidTransform3& InShapeWorldTransform1)
-	{
-		ShapeWorldTransform0 = InShapeWorldTransform0;
-		ShapeWorldTransform1 = InShapeWorldTransform1;
-	}
-
-	void FPBDCollisionConstraint::SetLastShapeWorldTransforms(const FRigidTransform3& InShapeWorldTransform0, const FRigidTransform3& InShapeWorldTransform1)
-	{
-		LastShapeWorldTransform0 = InShapeWorldTransform0;
-		LastShapeWorldTransform1 = InShapeWorldTransform1;
-	}
-
-	void FPBDCollisionConstraint::ReprojectManifoldPoint(const int32 ManifoldPointIndex)
-	{
-		FManifoldPoint& ManifoldPoint = ManifoldPoints[ManifoldPointIndex];
-
-		// Calculate the world-space contact location and separation at the current shape transforms
-		// @todo(chaos): this should use the normal owner. Currently we assume body 1 is the owner
-		const FRigidTransform3 Shape0ToShape1Transform = ShapeWorldTransform0.GetRelativeTransformNoScale(ShapeWorldTransform1);
-		const FVec3 Contact0In1 = Shape0ToShape1Transform.TransformPositionNoScale(ManifoldPoint.InitialShapeContactPoints[0]);
-		const FVec3& Contact1In1 = ManifoldPoint.InitialShapeContactPoints[1];
-		const FVec3& ContactNormalIn1 = ManifoldPoint.ContactPoint.ShapeContactNormal;
-
-		const FVec3 ContactDeltaIn1 = Contact0In1 - Contact1In1;
-		const FReal ContactPhi = FVec3::DotProduct(ContactDeltaIn1, ContactNormalIn1);
-
-		// Recalculate the contact points at the new location
-		const FVec3 ShapeContactPoint1 = Contact0In1 - ContactPhi * ContactNormalIn1;
-		ManifoldPoint.ContactPoint.ShapeContactPoints[1] = ShapeContactPoint1;
-		ManifoldPoint.ContactPoint.Phi = ContactPhi;
-	}
-
-	void FPBDCollisionConstraint::ReprojectManifoldContacts()
-	{
-		ClosestManifoldPointIndex = INDEX_NONE;
-		Manifold.Reset();
-
-		for (int32 ManifoldPointIndex = 0; ManifoldPointIndex < ManifoldPoints.Num(); ManifoldPointIndex++)
-		{
-			FManifoldPoint& ManifoldPoint = ManifoldPoints[ManifoldPointIndex];
-
-			ReprojectManifoldPoint(ManifoldPointIndex);
-
-			ManifoldPoint.Flags.bInsideStaticFrictionCone = false;
-			ManifoldPoint.Flags.bWasRestored = true;
-
-			// Copy currently active point
-			if (ManifoldPoint.ContactPoint.Phi < GetPhi())
-			{
-				ClosestManifoldPointIndex = ManifoldPointIndex;
-			}
-		}
-	}
-
 	bool FPBDCollisionConstraint::UpdateAndTryRestoreManifold()
 	{
-		const FCollisionTolerances Tolerances = Chaos_Manifold_Tolerances;
+		const FCollisionTolerances Tolerances = FCollisionTolerances();//Chaos_Manifold_Tolerances;
 		const FReal ContactPositionTolerance = Tolerances.ContactPositionToleranceScale * CollisionTolerance;
 		const FReal ShapePositionTolerance = (ManifoldPoints.Num() > 0) ? Tolerances.ShapePositionToleranceScaleN * CollisionTolerance : Tolerances.ShapePositionToleranceScale0 * CollisionTolerance;
 		const FReal ShapeRotationThreshold = (ManifoldPoints.Num() > 0) ? Tolerances.ShapeRotationThresholdN : Tolerances.ShapeRotationThreshold0;
@@ -573,6 +526,7 @@ namespace Chaos
 
 		// Reset current closest point
 		ClosestManifoldPointIndex = INDEX_NONE;
+		Flags.bDisabled = false;
 		Manifold.Reset();
 
 		// How many manifold points we expect. E.g., for Box-box this will be 4 or 1 depending on whether
@@ -581,7 +535,7 @@ namespace Chaos
 		ExpectedNumManifoldPoints = ManifoldPoints.Num();
 		Flags.bWasManifoldRestored = false;
 
-		// If we did not remove any contact points and we have not moved or rotated much we may reuse the manifold as-is.
+		// If we have not moved or rotated much we may reuse some of the manifold points, as long as they have not moved far as well (see below)
 		bool bMovedBeyondTolerance = true;
 		if ((ShapePositionTolerance > 0) && (ShapeRotationThreshold > 0))
 		{
@@ -589,13 +543,12 @@ namespace Chaos
 			// as a body moves/rotates we may have to change which faces/edges are colliding. We can't know if the face/edge
 			// will change until we run the closest-point checks (GJK) in the narrow phase.
 			const FVec3 Shape1ToShape0Translation = ShapeWorldTransform0.GetTranslation() - ShapeWorldTransform1.GetTranslation();
-			const FVec3 OriginalShape1ToShape0Translation = LastShapeWorldTransform0.GetTranslation() - LastShapeWorldTransform1.GetTranslation();
-			const FVec3 TranslationDelta = Shape1ToShape0Translation - OriginalShape1ToShape0Translation;
+			const FVec3 TranslationDelta = Shape1ToShape0Translation - LastShapeWorldPositionDelta;
 			if (TranslationDelta.IsNearlyZero(ShapePositionTolerance))
 			{
 				const FRotation3 Shape1toShape0Rotation = ShapeWorldTransform0.GetRotation().Inverse() * ShapeWorldTransform1.GetRotation();
-				const FRotation3 OriginalShape1toShape0Rotation = LastShapeWorldTransform0.GetRotation().Inverse() * LastShapeWorldTransform1.GetRotation();
-				const FReal RotationOverlap = FRotation3::DotProduct(Shape1toShape0Rotation, OriginalShape1toShape0Rotation);
+				const FRotation3 OriginalShape1toShape0Rotation = LastShapeWorldRotationDelta;
+				const FReal RotationOverlap = FRotation3::DotProduct(Shape1toShape0Rotation, LastShapeWorldRotationDelta);
 				if (RotationOverlap > ShapeRotationThreshold)
 				{
 					bMovedBeyondTolerance = false;
@@ -642,12 +595,13 @@ namespace Chaos
 					ManifoldPoint.ContactPoint.ShapeContactPoints[1] = ShapeContactPoint1;
 					ManifoldPoint.ContactPoint.Phi = ContactPhi;
 					ManifoldPoint.Flags.bWasRestored = true;
+					ManifoldPoint.Flags.bWasReplaced = false;
 					if (ManifoldPoint.ContactPoint.Phi < GetPhi())
 					{
 						ClosestManifoldPointIndex = ManifoldPointIndex;
 					}
 				}
-				else if ((ManifoldPointToRemove == INDEX_NONE) && bChaos_Collision_EnableManifoldInject)
+				else if ((ManifoldPointToRemove == INDEX_NONE) && (bChaos_Collision_EnableManifoldGJKReplace || bChaos_Collision_EnableManifoldGJKInject))
 				{
 					// We can reject up to 1 point (if we have GJK point injection enabled)
 					ManifoldPointToRemove = ManifoldPointIndex;
@@ -684,9 +638,13 @@ namespace Chaos
 		const FReal NormalThreshold = Tolerances.ManifoldPointNormalThreshold;
 
 		// We must end up with a full manifold after this if we want to reuse it
-		if ((ManifoldPoints.Num() < ExpectedNumManifoldPoints - 1) || (ExpectedNumManifoldPoints == 0))
+		//if ((ManifoldPoints.Num() < ExpectedNumManifoldPoints - 1) || (ExpectedNumManifoldPoints == 0))
+		//{
+		//	// We need to add more than 1 point to restore the manifold so we must rebuild it from scratch
+		//	return false;
+		//}
+		if (ManifoldPoints.Num() == 0)
 		{
-			// We need to add more than 1 point to restore the manifold so we must rebuild it from scratch
 			return false;
 		}
 
@@ -710,21 +668,17 @@ namespace Chaos
 			const FVec3 DR1 = ManifoldPoint.ContactPoint.ShapeContactPoints[1] - NewContactPoint.ShapeContactPoints[1];
 			if ((DR0.SizeSquared() < PositionToleranceSq) && (DR1.SizeSquared() < PositionToleranceSq))
 			{
-				// If we should replace a point but will then have too few points we abort
-				if (ManifoldPoints.Num() < ExpectedNumManifoldPoints)
-				{
-					return false;
-				}
-
 				// If the existing point has a deeper penetration, just re-use it. This is common when we have a GJK
 				// result on an edge or corner - the contact created when generating the manifold is on the
 				// surface shape rather than the rounded (margin-reduced) shape.
+				// If the new point is deeper, use it.
 				if (ManifoldPoint.ContactPoint.Phi > NewContactPoint.Phi)
 				{
 					ManifoldPoint.ContactPoint = NewContactPoint;
 					ManifoldPoint.InitialShapeContactPoints[0] = NewContactPoint.ShapeContactPoints[0];
 					ManifoldPoint.InitialShapeContactPoints[1] = NewContactPoint.ShapeContactPoints[1];
 					ManifoldPoint.Flags.bWasRestored = false;
+					ManifoldPoint.Flags.bWasReplaced = true;
 					if (NewContactPoint.Phi < GetPhi())
 					{
 						ClosestManifoldPointIndex = ManifoldPointIndex;
@@ -736,7 +690,7 @@ namespace Chaos
 		}
 
 		// If we have a full manifold, see if we can use or reject the GJK point
-		if (ManifoldPoints.Num() == 4)
+		if ((ManifoldPoints.Num() == 4) && bChaos_Collision_EnableManifoldGJKInject)
 		{
 			return TryInsertManifoldContact(NewContactPoint);
 		}
diff --git a/Engine/Source/Runtime/Experimental/Chaos/Private/Chaos/Collision/ParticlePairMidPhase.cpp b/Engine/Source/Runtime/Experimental/Chaos/Private/Chaos/Collision/ParticlePairMidPhase.cpp
--- a/Engine/Source/Runtime/Experimental/Chaos/Private/Chaos/Collision/ParticlePairMidPhase.cpp
+++ b/Engine/Source/Runtime/Experimental/Chaos/Private/Chaos/Collision/ParticlePairMidPhase.cpp
@@ -20,20 +20,7 @@ namespace Chaos
 		bool bChaos_Collision_MidPhase_EnableBoundsChecks = true;
 		FAutoConsoleVariableRef CVarChaos_Collision_EnableBoundsChecks(TEXT("p.Chaos.Collision.EnableBoundsChecks"), bChaos_Collision_MidPhase_EnableBoundsChecks, TEXT(""));
 
-		Chaos::FRealSingle Chaos_Collision_RestoreTolerance_NoContact_Position = 0.005f;	// About 0.5cm for a meter cube
-		Chaos::FRealSingle Chaos_Collision_RestoreTolerance_NoContact_Rotation = 0.1f;		// About 10deg
-		Chaos::FRealSingle Chaos_Collision_RestoreTolerance_Contact_Position = 0.02f;		// About 2cm for a meter cube
-		Chaos::FRealSingle Chaos_Collision_RestoreTolerance_Contact_Rotation = 0.1f;		// About 10deg
-		Chaos::FRealSingle Chaos_Collision_RestoreTolerance_ReprojectFactor = 0.0f;
-		FAutoConsoleVariableRef CVarChaos_Collision_RestoreTolerance_NoContact_Position(TEXT("p.Chaos.Collision.RestoreTolerance.NoContact.Position"), Chaos_Collision_RestoreTolerance_NoContact_Position, TEXT("Fraction of Size. Particle pairs that move less than this may have their contacts reinstated"));
-		FAutoConsoleVariableRef CVarChaos_Collision_RestoreTolerance_NoContact_Rotation(TEXT("p.Chaos.Collision.RestoreTolerance.NoContact.Rotation"), Chaos_Collision_RestoreTolerance_NoContact_Rotation, TEXT("Quaternion Dot Product Limit. Particle pairs that move less than this may have their contacts reinstated"));
-		FAutoConsoleVariableRef CVarChaos_Collision_RestoreTolerance_Contact_Position(TEXT("p.Chaos.Collision.RestoreTolerance.WithContact.Position"), Chaos_Collision_RestoreTolerance_Contact_Position, TEXT("Fraction of Size. Particle pairs that move less than this may have their contacts reinstated"));
-		FAutoConsoleVariableRef CVarChaos_Collision_RestoreTolerance_Contact_Rotation(TEXT("p.Chaos.Collision.RestoreTolerance.WithContact.Rotation"), Chaos_Collision_RestoreTolerance_Contact_Rotation, TEXT("Quaternion Dot Product Limit. Particle pairs that move less than this may have their contacts reinstated"));
-		FAutoConsoleVariableRef CVarChaos_Collision_RestoreTolerance_ReprojectFactor(TEXT("p.Chaos.Collision.RestoreTolerance.ReprojectFactor"), Chaos_Collision_RestoreTolerance_ReprojectFactor, TEXT("Quaternion Dot Product Limit. Particle pairs that move less than this may have their contacts reinstated"));
-
-		bool bChaos_Collision_EnableManifoldRestore = false;	// This causes subtle problems with friction unless points are reprojected, but then it's basically the same as EnableManifoldUpdate
 		bool bChaos_Collision_EnableManifoldUpdate = true;
-		FAutoConsoleVariableRef CVarChaos_Collision_EnableManifoldRestore(TEXT("p.Chaos.Collision.EnableManifoldRestore"), bChaos_Collision_EnableManifoldRestore, TEXT(""));
 		FAutoConsoleVariableRef CVarChaos_Collision_EnableManifoldUpdate(TEXT("p.Chaos.Collision.EnableManifoldUpdate"), bChaos_Collision_EnableManifoldUpdate, TEXT(""));
 
 		Chaos::FRealSingle Chaos_Collision_CullDistanceScaleInverseSize = 0.01f;	// 100cm
@@ -144,7 +131,7 @@ namespace Chaos
 		}
 
 		// Do not try to reuse manifold points for capsules or spheres (against anything)
-		Flags.bEnableManifoldCheck = bChaos_Collision_EnableManifoldUpdate && !bIsSphere0 && !bIsSphere1 && !bIsCapsule0 && !bIsCapsule1;
+		Flags.bEnableManifoldUpdate = bChaos_Collision_EnableManifoldUpdate && !bIsSphere0 && !bIsSphere1 && !bIsCapsule0 && !bIsCapsule1;
 	}
 
 	FSingleShapePairCollisionDetector::~FSingleShapePairCollisionDetector()
@@ -297,7 +284,7 @@ namespace Chaos
 			}
 
 			bool bWasManifoldRestored = false;
-			if (Flags.bEnableManifoldCheck)
+			if (Flags.bEnableManifoldUpdate)
 			{
 				// Update the existing manifold. We can re-use as-is if none of the points have moved much and the bodies have not moved much
 				// NOTE: this can succeed in "restoring" even if we have no manifold points
@@ -349,31 +336,6 @@ namespace Chaos
 		return 0;
 	}
 
-	int32 FSingleShapePairCollisionDetector::RestoreCollision(const FReal CullDistance, const bool bReproject)
-	{
-		// Only restore constraints if active last tick. Any older than that and the shapes were separated for a bit
-		const int32 CurrentEpoch = MidPhase.GetCollisionAllocator().GetCurrentEpoch();
-		const int32 LastEpoch = CurrentEpoch - 1;
-		if ((Constraint != nullptr) && IsUsedSince(LastEpoch))
-		{
-			const FRigidTransform3& ShapeWorldTransform0 = Shape0->GetLeafWorldTransform();
-			const FRigidTransform3& ShapeWorldTransform1 = Shape1->GetLeafWorldTransform();
-			Constraint->SetShapeWorldTransforms(ShapeWorldTransform0, ShapeWorldTransform1);
-
-			Constraint->RestoreManifold(bReproject);
-
-			if (Constraint->GetPhi() <= CullDistance)
-			{
-				if (MidPhase.GetCollisionAllocator().ActivateConstraint(Constraint.Get()))
-				{
-					LastUsedEpoch = CurrentEpoch;
-					return 1;
-				}
-			}
-		}
-		return 0;
-	}
-
 	void FSingleShapePairCollisionDetector::WakeCollision(const int32 SleepEpoch)
 	{
 		if ((Constraint != nullptr) && IsUsedSince(SleepEpoch))
@@ -604,41 +566,6 @@ namespace Chaos
 		return Constraints.Add(Key.GetKey(), MoveTemp(Constraint)).Get();
 	}
 
-	int32 FMultiShapePairCollisionDetector::RestoreCollisions(const FReal CullDistance, const bool bReproject)
-	{
-		if (Constraints.Num() == 0)
-		{
-			return 0;
-		}
-
-		int32 NumRestored = 0;
-		const int32 LastEpoch = MidPhase.GetCollisionAllocator().GetCurrentEpoch() - 1;
-		const FRigidTransform3 ParticleTransform0 = FParticleUtilitiesPQ::GetActorWorldTransform(FConstGenericParticleHandle(Particle0));
-		const FRigidTransform3 ParticleTransform1 = FParticleUtilitiesPQ::GetActorWorldTransform(FConstGenericParticleHandle(Particle1));
-
-		for (auto& KVP : Constraints)
-		{
-			TUniquePtr<FPBDCollisionConstraint>& Constraint = KVP.Value;
-			if ((Constraint != nullptr) && (Constraint->GetContainerCookie().LastUsedEpoch >= LastEpoch))
-			{
-				const FRigidTransform3 ShapeWorldTransform0 = Constraint->ImplicitTransform[0] * ParticleTransform0;
-				const FRigidTransform3 ShapeWorldTransform1 = Constraint->ImplicitTransform[1] * ParticleTransform1;
-				Constraint->SetShapeWorldTransforms(ShapeWorldTransform0, ShapeWorldTransform1);
-
-				Constraint->RestoreManifold(bReproject);
-
-				if (Constraint->GetPhi() < CullDistance)
-				{
-					if (MidPhase.GetCollisionAllocator().ActivateConstraint(Constraint.Get()))
-					{
-						++NumRestored;
-					}
-				}
-			}
-		}
-		return NumRestored;
-	}
-
 	void FMultiShapePairCollisionDetector::WakeCollisions(const int32 SleepEpoch)
 	{
 		const int32 CurrentEpoch = MidPhase.GetCollisionAllocator().GetCurrentEpoch();
@@ -724,12 +651,6 @@ namespace Chaos
 		, NumActiveConstraints(0)
 		, ParticleCollisionsIndex0(INDEX_NONE)
 		, ParticleCollisionsIndex1(INDEX_NONE)
-		, RestoreThresholdZeroContacts()
-		, RestoreThreshold()
-		, RestoreParticleP0(FVec3(0))
-		, RestoreParticleP1(FVec3(0))
-		, RestoreParticleQ0(FRotation3::FromIdentity())
-		, RestoreParticleQ1(FRotation3::FromIdentity())
 		, CullDistanceScale(1)
 	{
 	}
@@ -760,7 +681,6 @@ namespace Chaos
 
 		Flags.bIsCCD = false;
 		Flags.bIsInitialized = false;
-		Flags.bRestorable = false;
 		Flags.bIsSleeping = false;
 	}
 
@@ -779,8 +699,6 @@ namespace Chaos
 
 		Flags.bIsCCD = FConstGenericParticleHandle(Particle0)->CCDEnabled() || FConstGenericParticleHandle(Particle1)->CCDEnabled();
 
-		Flags.bRestorable = true;
-
 		BuildDetectors();
 
 		InitThresholds();
@@ -803,14 +721,6 @@ namespace Chaos
 					TryAddShapePair(Shape0, Shape1);
 				}
 			}
-
-			// If we have only 1 shape pair, we may disable Manifold Restore functionality if it is spheres or capsules
-			// @todo(chaos): make manifold restoration work better with rolling shapes (spheres and capsules)
-			const bool bDisableManifoldRestore = (ShapePairDetectors.Num() == 1) && !ShapePairDetectors[0].EnableManifoldRestore();
-			if (bDisableManifoldRestore)
-			{
-				Flags.bRestorable = false;
-			}
 		}
 	}
 
@@ -848,15 +758,6 @@ namespace Chaos
 					MultiShapePairDetectors.Emplace(FMultiShapePairCollisionDetector(Particle0, Shape0, Particle1, Shape1, *this));
 				}
 			}
-
-			// We don't allow full constraint restoration for LevelSets or Unions because small changes in 
-			// transform can change what conatct points are generated.
-			// @todo(chaos): LevelSets require one-shot manifolds to suport full restore
-			// @todo(chaos): Unions may need to reactivate constraints that were not used last frame to support full restore
-			if ((ShapePairType == EContactShapesType::LevelSetLevelSet) || (ShapePairType == EContactShapesType::Unknown))
-			{
-				Flags.bRestorable = false;
-			}
 		}
 	}
 
@@ -884,16 +785,7 @@ namespace Chaos
 		// @todo(chaos): Spheres and capsules need smaller position tolerance - the restore test doesn't work well with rolling
 		bool bIsDynamic0 = FConstGenericParticleHandle(Particle0)->IsDynamic();
 		bool bIsDynamic1 = FConstGenericParticleHandle(Particle1)->IsDynamic();
-		const FReal MinBoundsSize0 = bIsDynamic0 ? Particle0->LocalBounds().Extents().GetMin() : TNumericLimits<FReal>::Max();
-		const FReal MinBoundsSize1 = bIsDynamic1 ? Particle1->LocalBounds().Extents().GetMin() : TNumericLimits<FReal>::Max();
-		const FReal ThresholdSize = FMath::Min(MinBoundsSize0, MinBoundsSize1);
-
-		RestoreThresholdZeroContacts.PositionThreshold = Chaos_Collision_RestoreTolerance_NoContact_Position * ThresholdSize;
-		RestoreThresholdZeroContacts.RotationThreshold = Chaos_Collision_RestoreTolerance_NoContact_Rotation;
 
-		RestoreThreshold.PositionThreshold = Chaos_Collision_RestoreTolerance_Contact_Position * ThresholdSize;
-		RestoreThreshold.RotationThreshold = Chaos_Collision_RestoreTolerance_Contact_Rotation;
-	
 		// NOTE: If CullDistance ends up smaller than the thresholds used to restore collisions, we can end up missing
 		// collisions as the objects move if we restore a "zero contact" manifold after movement greater than the cull distance. 
 		// Currently this should not happen, but it is not explicitly ensured by the way the thresholds and CullDistanceScale are calculated.
@@ -907,100 +799,14 @@ namespace Chaos
 		CullDistanceScale = FMath::Max3(CullDistanceScale0, CullDistanceScale1, MinCullDistanceScale);
 	}
 
-	ECollisionRestoreType FParticlePairMidPhase::ShouldRestoreConstraints(const FReal Dt)
-	{
-		const FVec3& ParticleP0 = FConstGenericParticleHandle(Particle0)->P();
-		const FRotation3& ParticleQ0 = FConstGenericParticleHandle(Particle0)->Q();
-		const FVec3& ParticleP1 = FConstGenericParticleHandle(Particle1)->P();
-		const FRotation3& ParticleQ1 = FConstGenericParticleHandle(Particle1)->Q();
-
-		// We can only restore collisions if they were created or updated last tick
-		ECollisionRestoreType RestoreType = ECollisionRestoreType::None;
-		if (IsUsedSince(CollisionAllocator->GetCurrentEpoch() - 1))
-		{
-			const FReal PositionReprojectThreshold = (NumActiveConstraints == 0) ? RestoreThresholdZeroContacts.PositionThreshold : RestoreThreshold.PositionThreshold;
-			const FReal PositionReuseThreshold = (NumActiveConstraints == 0) ? PositionReprojectThreshold : Chaos_Collision_RestoreTolerance_ReprojectFactor * PositionReprojectThreshold;
-
-			const FReal MaxDP0 = FVec3(ParticleP0 - RestoreParticleP0).GetAbsMax();
-			const FReal MaxDP1 = FVec3(ParticleP1 - RestoreParticleP1).GetAbsMax();
-			const FReal MaxDP = FMath::Max(MaxDP0, MaxDP1);
-
-			// Check position delta and calculate the best-case reuse option (ignoring rotation for now)
-			if (MaxDP < PositionReuseThreshold)
-			{
-				// We moved a tiny amount - reuse the manifold exactly as-is
-				RestoreType = ECollisionRestoreType::Reuse;
-			}
-			else if (MaxDP < PositionReprojectThreshold)
-			{
-				// We moved a bit but within threshold - reuse the manifold but reproject the vertex contacts onto the planes
-				RestoreType = ECollisionRestoreType::Reproject;
-			}
-
-			// If we passed the position test, check the rotation threshold and maybe veto the restore decision
-			if (RestoreType != ECollisionRestoreType::None)
-			{
-				const FReal RotationThreshold = (NumActiveConstraints == 0) ? RestoreThresholdZeroContacts.RotationThreshold : RestoreThreshold.RotationThreshold;
-				const FReal MaxDQ0 = FRotation3(ParticleQ0 - RestoreParticleQ0).GetAbsMax();
-				const FReal MaxDQ1 = FRotation3(ParticleQ1 - RestoreParticleQ1).GetAbsMax();
-				const FReal MaxDQ = FMath::Max(MaxDQ0, MaxDQ1);
-				if (MaxDQ > RotationThreshold)
-				{
-					RestoreType = ECollisionRestoreType::None;
-				}
-			}
-		}
-
-		if (RestoreType != ECollisionRestoreType::None)
-		{
-			// We have moved and will rebuild the manifold so
-			// update the current manifold transforms for future restore checks
-			RestoreParticleP0 = ParticleP0;
-			RestoreParticleP1 = ParticleP1;
-			RestoreParticleQ0 = ParticleQ0;
-			RestoreParticleQ1 = ParticleQ1;
-		}
-
-		return RestoreType;
-	}
-
-	bool FParticlePairMidPhase::TryRestoreConstraints(const FReal Dt, const FReal CullDistance)
-	{
-		SCOPE_CYCLE_COUNTER(STAT_Collisions_Restore);
-		PHYSICS_CSV_SCOPED_EXPENSIVE(PhysicsVerbose, DetectCollisions_RestoreCollision);
-
-		// If the particles haven't moved relative to each other, we can just reuse the constraint as-is
-		const ECollisionRestoreType CollisionRestoreType = ShouldRestoreConstraints(Dt);
-		if (CollisionRestoreType != ECollisionRestoreType::None)
-		{
-			const bool bReproject = (CollisionRestoreType == ECollisionRestoreType::Reproject);
-
-			int32 NumRestored = 0;
-			for (FSingleShapePairCollisionDetector& ShapePair : ShapePairDetectors)
-			{
-				NumRestored += ShapePair.RestoreCollision(CullDistance, bReproject);
-			}
-			for (FMultiShapePairCollisionDetector& MultiShapePair : MultiShapePairDetectors)
-			{
-				NumRestored += MultiShapePair.RestoreCollisions(CullDistance, bReproject);
-			}
-
-			NumActiveConstraints = NumRestored;
-
-			PHYSICS_CSV_CUSTOM_EXPENSIVE(PhysicsCounters, NumRestoredContacts, NumRestored, ECsvCustomStatOp::Accumulate);
-
-			// NOTE: We return restored as true, even if we didn't have any constraints to restore. This is for Bodies 
-			// that were separated by more than CullDistance last tick and have not moved more than the tolerances.
-			return true;
-		}
-		return false;
-	}
-
 	void FParticlePairMidPhase::GenerateCollisions(
 		const FReal InCullDistance,
 		const FReal Dt,
 		FCollisionContext& Context)
 	{
+		SCOPE_CYCLE_COUNTER(STAT_Collisions_GenerateCollisions);
+		PHYSICS_CSV_SCOPED_EXPENSIVE(PhysicsVerbose, DetectCollisions_NarrowPhase);
+
 		if (!IsValid())
 		{
 			return;
@@ -1012,28 +818,15 @@ namespace Chaos
 		// Enable CCD?
 		const bool bUseCCD = Flags.bIsCCD && ShouldEnableCCD(Dt);
 
-		// If the bodies have not moved, we will reuse the constraints as-is
-		const bool bCanRestore = bChaos_Collision_EnableManifoldRestore && Flags.bRestorable && !bUseCCD;
-		const bool bWasRestored = bCanRestore && TryRestoreConstraints(Dt, CullDistance);
-
-		// If the bodies have moved we need to create or update the constraints
-		if (!bWasRestored)
+		// Run collision detection on all potentially colliding shape pairs
+		NumActiveConstraints = 0;
+		for (FSingleShapePairCollisionDetector& ShapePair : ShapePairDetectors)
 		{
-			SCOPE_CYCLE_COUNTER(STAT_Collisions_GenerateCollisions);
-			PHYSICS_CSV_SCOPED_EXPENSIVE(PhysicsVerbose, DetectCollisions_NarrowPhase);
-
-			NumActiveConstraints = 0;
-
-			// Run collision detection on all potentially colliding shape pairs
-			for (FSingleShapePairCollisionDetector& ShapePair : ShapePairDetectors)
-			{
-				NumActiveConstraints += ShapePair.GenerateCollision(CullDistance, bUseCCD, Dt);
-			}
-
-			for (FMultiShapePairCollisionDetector& MultiShapePair : MultiShapePairDetectors)
-			{
-				NumActiveConstraints += MultiShapePair.GenerateCollisions(CullDistance, bUseCCD, Dt, Context);
-			}
+			NumActiveConstraints += ShapePair.GenerateCollision(CullDistance, bUseCCD, Dt);
+		}
+		for (FMultiShapePairCollisionDetector& MultiShapePair : MultiShapePairDetectors)
+		{
+			NumActiveConstraints += MultiShapePair.GenerateCollisions(CullDistance, bUseCCD, Dt, Context);
 		}
 
 		LastUsedEpoch = CollisionAllocator->GetCurrentEpoch();
diff --git a/Engine/Source/Runtime/Experimental/Chaos/Private/Chaos/CollisionOneShotManifolds.cpp b/Engine/Source/Runtime/Experimental/Chaos/Private/Chaos/CollisionOneShotManifolds.cpp
--- a/Engine/Source/Runtime/Experimental/Chaos/Private/Chaos/CollisionOneShotManifolds.cpp
+++ b/Engine/Source/Runtime/Experimental/Chaos/Private/Chaos/CollisionOneShotManifolds.cpp
@@ -40,8 +40,10 @@ namespace Chaos
 	FAutoConsoleVariableRef CVarForceOneShotManifoldEdgeEdgeCaseZeroCullDistance(TEXT("p.Chaos.Collision.Manifold.ForceOneShotManifoldEdgeEdgeCaseZeroCullDistance"), ForceOneShotManifoldEdgeEdgeCaseZeroCullDistance,
 	TEXT("If enabled, if one shot manifold hits edge/edge case, we will force a cull distance of zero. That means edge/edge contacts will be thrown out if separated at all. Only applies to Convex/Convex oneshot impl."));
 
-	bool bChaos_Collision_EnableManifoldInject = false;
-	FAutoConsoleVariableRef CVarChaos_Collision_EnableManifoldInject(TEXT("p.Chaos.Collision.EnableManifolInject"), bChaos_Collision_EnableManifoldInject, TEXT(""));
+	bool bChaos_Collision_EnableManifoldGJKReplace = false;
+	bool bChaos_Collision_EnableManifoldGJKInject = false;
+	FAutoConsoleVariableRef CVarChaos_Collision_EnableManifoldReplace(TEXT("p.Chaos.Collision.EnableManifoldGJKReplace"), bChaos_Collision_EnableManifoldGJKReplace, TEXT(""));
+	FAutoConsoleVariableRef CVarChaos_Collision_EnableManifoldInject(TEXT("p.Chaos.Collision.EnableManifoldGJKInject"), bChaos_Collision_EnableManifoldGJKInject, TEXT(""));
 
 	
 	namespace Collisions
@@ -671,7 +673,7 @@ namespace Chaos
 			FContactPoint GJKContactPoint = GJKContactPointMargin(Convex1, Convex2, Convex1Transform, Convex2Transform, Margin1, Margin2, Constraint.GetGJKWarmStartData(), MaxMarginDelta, VertexIndexA, VertexIndexB);
 			PHYSICS_CSV_CUSTOM_EXPENSIVE(PhysicsCounters, NumManifoldsGJKCalled, 1, ECsvCustomStatOp::Accumulate);
 
-			const bool bCanUpdateManifold = bChaos_Collision_EnableManifoldInject;
+			const bool bCanUpdateManifold = bChaos_Collision_EnableManifoldGJKReplace;
 			if (bCanUpdateManifold && Constraint.TryAddManifoldContact(GJKContactPoint))
 			{
 				PHYSICS_CSV_CUSTOM_EXPENSIVE(PhysicsCounters, NumManifoldsMaintained, 1, ECsvCustomStatOp::Accumulate);
diff --git a/Engine/Source/Runtime/Experimental/Chaos/Private/Chaos/PBDCollisionConstraints.cpp b/Engine/Source/Runtime/Experimental/Chaos/Private/Chaos/PBDCollisionConstraints.cpp
--- a/Engine/Source/Runtime/Experimental/Chaos/Private/Chaos/PBDCollisionConstraints.cpp
+++ b/Engine/Source/Runtime/Experimental/Chaos/Private/Chaos/PBDCollisionConstraints.cpp
@@ -183,8 +183,8 @@ namespace Chaos
 
 	void FPBDCollisionConstraints::UpdateConstraintMaterialProperties(FPBDCollisionConstraint& Constraint)
 	{
-		const FChaosPhysicsMaterial* PhysicsMaterial0 = GetPhysicsMaterial(Constraint.Particle[0], Constraint.Manifold.Implicit[0], MPhysicsMaterials, MPerParticlePhysicsMaterials);
-		const FChaosPhysicsMaterial* PhysicsMaterial1 = GetPhysicsMaterial(Constraint.Particle[1], Constraint.Manifold.Implicit[1], MPhysicsMaterials, MPerParticlePhysicsMaterials);
+		const FChaosPhysicsMaterial* PhysicsMaterial0 = GetPhysicsMaterial(Constraint.Particle[0], Constraint.Implicit[0], MPhysicsMaterials, MPerParticlePhysicsMaterials);
+		const FChaosPhysicsMaterial* PhysicsMaterial1 = GetPhysicsMaterial(Constraint.Particle[1], Constraint.Implicit[1], MPhysicsMaterials, MPerParticlePhysicsMaterials);
 
 		FCollisionContact& Contact = Constraint.Manifold;
 		if (PhysicsMaterial0 && PhysicsMaterial1)
diff --git a/Engine/Source/Runtime/Experimental/Chaos/Public/Chaos/Collision/PBDCollisionConstraint.h b/Engine/Source/Runtime/Experimental/Chaos/Public/Chaos/Collision/PBDCollisionConstraint.h
--- a/Engine/Source/Runtime/Experimental/Chaos/Public/Chaos/Collision/PBDCollisionConstraint.h
+++ b/Engine/Source/Runtime/Experimental/Chaos/Public/Chaos/Collision/PBDCollisionConstraint.h
@@ -74,6 +74,7 @@ namespace Chaos
 			struct
 			{
 				uint8 bWasRestored : 1;
+				uint8 bWasReplaced : 1;
 				uint8 bWasFrictionRestored : 1;
 				uint8 bInsideStaticFrictionCone : 1;		// Whether we are inside the static friction cone
 			};
@@ -97,9 +98,8 @@ namespace Chaos
 	class CHAOS_API FCollisionContact
 	{
 	public:
-		FCollisionContact(const FImplicitObject* InImplicit0 = nullptr, const FBVHParticles* InSimplicial0 = nullptr, const FImplicitObject* InImplicit1 = nullptr, const FBVHParticles* InSimplicial1 = nullptr)
-			: bDisabled(false)
-			, Friction(0)
+		FCollisionContact()
+			: Friction(0)
 			, AngularFriction(0)
 			, Restitution(0)
 			, RestitutionPadding(0)
@@ -108,16 +108,9 @@ namespace Chaos
 			, InvMassScale1(1.f)
 			, InvInertiaScale0(1.f)
 			, InvInertiaScale1(1.f)
-			, ShapesType(EContactShapesType::Unknown)
 		{
-			Implicit[0] = InImplicit0;
-			Implicit[1] = InImplicit1;
-			Simplicial[0] = InSimplicial0;
-			Simplicial[1] = InSimplicial1;
 		}
 
-		bool bDisabled;
-
 		FReal Friction;			// @todo(chaos): rename DynamicFriction
 		FReal AngularFriction;	// @todo(chaos): rename StaticFriction
 		FReal Restitution;
@@ -127,16 +120,11 @@ namespace Chaos
 		FReal InvMassScale1;
 		FReal InvInertiaScale0;
 		FReal InvInertiaScale1;
-		EContactShapesType ShapesType;
 
 		void Reset()
 		{
-			bDisabled = false;
 			RestitutionPadding = 0;
 		}
-
-		const FImplicitObject* Implicit[2]; // {Of Particle[0], Of Particle[1]}
-		const FBVHParticles* Simplicial[2]; // {Of Particle[0], Of Particle[1]}
 	};
 
 
@@ -221,6 +209,25 @@ namespace Chaos
 
 		static const int32 MaxManifoldPoints = 4;
 
+		/**
+		 * @brief Create a contact constraint
+		 * Initializes a constraint stored inline in an object. Only intended to be called once.
+		 * Does not reinitialize all data so not intended to reset a constraint for reuse with different particles etc.
+		*/
+		static void MakeInline(
+			FGeometryParticleHandle* Particle0,
+			const FImplicitObject* Implicit0,
+			const FBVHParticles* Simplicial0,
+			const FRigidTransform3& ImplicitLocalTransform0,
+			FGeometryParticleHandle* Particle1,
+			const FImplicitObject* Implicit1,
+			const FBVHParticles* Simplicial1,
+			const FRigidTransform3& ImplicitLocalTransform1,
+			const FReal InCullDistance,
+			const bool bInUseManifold,
+			const EContactShapesType ShapesType,
+			FPBDCollisionConstraint& OutConstraint);
+
 		/**
 		 * @brief Create a contact constraint
 		 * Allocates a constraint on the heap, with a permanent address.
@@ -266,13 +273,15 @@ namespace Chaos
 
 		bool ContainsManifold(const FImplicitObject* A, const FBVHParticles* AS, const FImplicitObject* B, const FBVHParticles* BS) const
 		{
-			return A == Manifold.Implicit[0] && B == Manifold.Implicit[1] && AS == Manifold.Simplicial[0] && BS == Manifold.Simplicial[1];
+			return A == Implicit[0] && B == Implicit[1] && AS == Simplicial[0] && BS == Simplicial[1];
 		}
 
 		void SetManifold(const FImplicitObject* A, const FBVHParticles* AS, const FImplicitObject* B, const FBVHParticles* BS)
 		{
-			Manifold.Implicit[0] = A; Manifold.Implicit[1] = B;
-			Manifold.Simplicial[0] = AS; Manifold.Simplicial[1] = BS;
+			Implicit[0] = A; 
+			Implicit[1] = B;
+			Simplicial[0] = AS; 
+			Simplicial[1] = BS;
 		}
 
 		//
@@ -286,13 +295,13 @@ namespace Chaos
 		FGeometryParticleHandle* GetParticle(const int32 ParticleIndex) { check((ParticleIndex >= 0) && (ParticleIndex < 2)); return Particle[ParticleIndex]; }
 		const FGeometryParticleHandle* GetParticle(const int32 ParticleIndex) const { check((ParticleIndex >= 0) && (ParticleIndex < 2)); return Particle[ParticleIndex]; }
 
-		const FImplicitObject* GetImplicit0() const { return Manifold.Implicit[0]; }
-		const FImplicitObject* GetImplicit1() const { return Manifold.Implicit[1]; }
-		const FImplicitObject* GetImplicit(const int32 ParticleIndex) const { check((ParticleIndex >= 0) && (ParticleIndex < 2)); return Manifold.Implicit[ParticleIndex]; }
+		const FImplicitObject* GetImplicit0() const { return Implicit[0]; }
+		const FImplicitObject* GetImplicit1() const { return Implicit[1]; }
+		const FImplicitObject* GetImplicit(const int32 ParticleIndex) const { check((ParticleIndex >= 0) && (ParticleIndex < 2)); return Implicit[ParticleIndex]; }
 
-		const FBVHParticles* GetCollisionParticles0() const { return Manifold.Simplicial[0]; }
-		const FBVHParticles* GetCollisionParticles1() const { return Manifold.Simplicial[1]; }
-		const FBVHParticles* GetCollisionParticles(const int32 ParticleIndex) const { check((ParticleIndex >= 0) && (ParticleIndex < 2)); return Manifold.Simplicial[ParticleIndex]; }
+		const FBVHParticles* GetCollisionParticles0() const { return Simplicial[0]; }
+		const FBVHParticles* GetCollisionParticles1() const { return Simplicial[1]; }
+		const FBVHParticles* GetCollisionParticles(const int32 ParticleIndex) const { check((ParticleIndex >= 0) && (ParticleIndex < 2)); return Simplicial[ParticleIndex]; }
 
 		const FReal GetCollisionMargin0() const { return CollisionMargins[0]; }
 		const FReal GetCollisionMargin1() const { return CollisionMargins[1]; }
@@ -309,8 +318,8 @@ namespace Chaos
 		void ResetPhi(FReal InPhi) { ClosestManifoldPointIndex = INDEX_NONE; }
 		FReal GetPhi() const { return (ClosestManifoldPointIndex != INDEX_NONE) ? ManifoldPoints[ClosestManifoldPointIndex].ContactPoint.Phi : TNumericLimits<FReal>::Max(); }
 
-		void SetDisabled(bool bInDisabled) { Manifold.bDisabled = bInDisabled; }
-		bool GetDisabled() const { return Manifold.bDisabled; }
+		void SetDisabled(bool bInDisabled) { Flags.bDisabled = bInDisabled; }
+		bool GetDisabled() const { return Flags.bDisabled; }
 
 		virtual void SetIsSleeping(const bool bInIsSleeping) override;
 
@@ -352,7 +361,7 @@ namespace Chaos
 		void SetDynamicFriction(const FReal InDynamicFriction) { Manifold.Friction = InDynamicFriction; }
 		FReal GetDynamicFriction() const { return Manifold.Friction; }
 
-		EContactShapesType GetShapesType() const { return Manifold.ShapesType; }
+		EContactShapesType GetShapesType() const { return ShapesType; }
 
 		FString ToString() const;
 
@@ -404,8 +413,17 @@ namespace Chaos
 		const FRigidTransform3& GetShapeWorldTransform0() const { return ShapeWorldTransform0; }
 		const FRigidTransform3& GetShapeWorldTransform1() const { return ShapeWorldTransform1; }
 
-		void SetShapeWorldTransforms(const FRigidTransform3& InShapeWorldTransform0, const FRigidTransform3& InShapeWorldTransform1);
-		void SetLastShapeWorldTransforms(const FRigidTransform3& InShapeWorldTransform0, const FRigidTransform3& InShapeWorldTransform1);
+		void SetShapeWorldTransforms(const FRigidTransform3& InShapeWorldTransform0, const FRigidTransform3& InShapeWorldTransform1)
+		{
+			ShapeWorldTransform0 = InShapeWorldTransform0;
+			ShapeWorldTransform1 = InShapeWorldTransform1;
+		}
+
+		void SetLastShapeWorldTransforms(const FRigidTransform3& InShapeWorldTransform0, const FRigidTransform3& InShapeWorldTransform1)
+		{
+			LastShapeWorldPositionDelta = InShapeWorldTransform0.GetTranslation() - InShapeWorldTransform1.GetTranslation();
+			LastShapeWorldRotationDelta = InShapeWorldTransform0.GetRotation().Inverse() * InShapeWorldTransform1.GetRotation();
+		}
 
 		bool UpdateAndTryRestoreManifold();
 		void ResetActiveManifoldContacts();
@@ -441,12 +459,6 @@ namespace Chaos
 		*/
 		bool WasManifoldRestored() const { return Flags.bWasManifoldRestored; }
 
-		/**
-		 * @brief Restore the contact manifold (assumes relative motion of the two bodies is small)
-		 * @see IsWithinManifoldRestorationThreshold
-		*/
-		void RestoreManifold(const bool bReproject);
-
 		/**
 		 * Determine the constraint direction based on Normal and Phi.
 		 * This function assumes that the constraint is update-to-date.
@@ -567,18 +579,19 @@ namespace Chaos
 		// @todo(chaos): Only intended for use by the legacy solvers - remove it
 		void UpdateManifoldPointPhi(const int32 ManifoldPointIndex);
 
-		void ReprojectManifoldContacts();
-		void ReprojectManifoldPoint(const int32 ManifoldPointIndex);
-
 		void InitMarginsAndTolerances(const EImplicitObjectType ImplicitType0, const EImplicitObjectType ImplicitType1, const FReal Margin0, const FReal Margin1);
 
 	private:
 		FReal CalculateSavedManifoldPointScore(const FSavedManifoldPoint& SavedManifoldPoint, const FManifoldPoint& ManifoldPoint, const FReal DistanceToleranceSq) const;
 
-		//@todo(chaos): make this stuff private
-		FRigidTransform3 ImplicitTransform[2];		// Local-space transforms of the shape (relative to particle)
+		// Local-space transforms of the shape (relative to particle)
+		FRigidTransform3 ImplicitTransform[2];
+		
 		FGeometryParticleHandle* Particle[2];
-		FCollisionContact Manifold;// @todo(chaos): rename
+		const FImplicitObject* Implicit[2];
+		const FBVHParticles* Simplicial[2];
+		FCollisionContact Manifold;// @todo(chaos): rename to FCollisionMaterial or something
+		FReal Stiffness;
 
 	public:
 		FVec3 AccumulatedImpulse;					// @todo(chaos): we need to accumulate angular impulse separately
@@ -588,9 +601,14 @@ namespace Chaos
 
 	private:
 		FPBDCollisionConstraintContainerCookie ContainerCookie;
+		EContactShapesType ShapesType;
 		ECollisionCCDType CCDType;
-		FReal Stiffness;
 
+		// The shape transforms at the current particle transforms
+		FRigidTransform3 ShapeWorldTransform0;
+		FRigidTransform3 ShapeWorldTransform1;
+
+		// The separation distance at which we don't track contacts
 		FReal CullDistance;
 
 		// The margins to use during collision detection. We don't always use the margins on the shapes directly.
@@ -601,38 +619,40 @@ namespace Chaos
 		// margins of the two shapes, as well as their types
 		FReal CollisionTolerance;
 
-		FReal FrictionPositionTolerance;
-
+		// The index into ManifoldPoints of the point with the lowest Phi
 		int32 ClosestManifoldPointIndex;
 
+		// Used by manifold point injection to see how many points were in the manifold before UpdateAndTryRestore
+		int32 ExpectedNumManifoldPoints;
+
 		union FFlags
 		{
 			FFlags() : Bits(0) {}
 			struct
 			{
-				bool bUseManifold;
-				bool bUseIncrementalManifold;
-				bool bWasManifoldRestored;
-				bool bIsQuadratic0;
-				bool bIsQuadratic1;
+				uint32 bDisabled : 1;
+				uint32 bUseManifold : 1;
+				uint32 bUseIncrementalManifold : 1;
+				uint32 bWasManifoldRestored : 1;
+				uint32 bIsQuadratic0 : 1;
+				uint32 bIsQuadratic1 : 1;
 			};
 			uint32 Bits;
 		} Flags;
 
+		// Relative transform the last time we ran the narrow phase
+		// Used to detect when the bodies have moved too far to reues the manifold
+		FVec3 LastShapeWorldPositionDelta;
+		FRotation3 LastShapeWorldRotationDelta;
+
 		// These are only needed here while we still have the legacy solvers (not QuasiPBD)
 		FSolverBody* SolverBodies[2];
 
-		TCArray<FManifoldPoint, MaxManifoldPoints> ManifoldPoints;
-		TCArray<FSavedManifoldPoint, MaxManifoldPoints> SavedManifoldPoints;
-
 		// Simplex data from the last call to GJK, used to warm-start GJK
 		FGJKSimplexData GJKWarmStartData;
 
-		FRigidTransform3 ShapeWorldTransform0;
-		FRigidTransform3 ShapeWorldTransform1;
-		FRigidTransform3 LastShapeWorldTransform0;
-		FRigidTransform3 LastShapeWorldTransform1;
-		int32 ExpectedNumManifoldPoints;
+		TCArray<FManifoldPoint, MaxManifoldPoints> ManifoldPoints;
+		TCArray<FSavedManifoldPoint, MaxManifoldPoints> SavedManifoldPoints;
 	};
 
 
diff --git a/Engine/Source/Runtime/Experimental/Chaos/Public/Chaos/Collision/ParticlePairMidPhase.h b/Engine/Source/Runtime/Experimental/Chaos/Public/Chaos/Collision/ParticlePairMidPhase.h
--- a/Engine/Source/Runtime/Experimental/Chaos/Public/Chaos/Collision/ParticlePairMidPhase.h
+++ b/Engine/Source/Runtime/Experimental/Chaos/Public/Chaos/Collision/ParticlePairMidPhase.h
@@ -21,21 +21,6 @@ namespace Chaos
 	class FPerShapeData;
 	class FSingleShapePairCollisionDetector;
 
-	/**
-	 * @brief The result of the test to see whether we can reuse a manifold (when the bodies have not moved much)
-	*/
-	enum class ECollisionRestoreType
-	{
-		// Do not retore the manifold - bodies have moved beyond the threshold
-		None,
-
-		// Restore the manifold exactly as-is - the bodies have barely moved
-		Reuse,
-
-		// Restore the manifold, but the bodies have moved a bit so reproject the contacts
-		Reproject,
-	};
-
 	/**
 	 * @brief Handles collision detection for a pair of simple shapes (i.e., not compound shapes)
 	 * 
@@ -83,13 +68,7 @@ namespace Chaos
 			const FReal Dt);
 
 		/**
-		 * @brief Reactivate the collision exactly as it was last frame
-		 * @return The number of collisions constraints that were restored
-		*/
-		int32 RestoreCollision(const FReal CullDistance, const bool bReproject);
-
-		/**
-		 * @brief Reactivate the constraint (essentially the same as Restore but slightly optimized)
+		 * @brief Reactivate the constraint
 		 * @parame SleepEpoch The tick on which the particle went to sleep.
 		 * Only constraints that were active when the particle went to sleep should be reactivated.
 		*/
@@ -101,15 +80,6 @@ namespace Chaos
 		*/
 		void SetCollision(const FPBDCollisionConstraint& Constraint);
 
-		/**
-		 * @brief Whether we want to use manifold restore for this shape type
-		 * @see FParticlePairMidPhase::ShouldRestoreConstraints
-		*/
-		bool EnableManifoldRestore() const
-		{
-			return Flags.bEnableManifoldCheck;
-		}
-
 	private:
 		int32 GenerateCollisionImpl(const FReal CullDistance, const bool bUseCCD, const FReal Dt);
 
@@ -141,7 +111,7 @@ namespace Chaos
 				uint32 bEnableAABBCheck : 1;
 				uint32 bEnableOBBCheck0 : 1;
 				uint32 bEnableOBBCheck1 : 1;
-				uint32 bEnableManifoldCheck : 1;
+				uint32 bEnableManifoldUpdate : 1;
 			};
 			uint32 Bits;
 		} Flags;
@@ -211,13 +181,7 @@ namespace Chaos
 			EContactShapesType ShapePairType);
 
 		/**
-		 * @brief Reactivate the collision exactly as it was last frame
-		 * @return The number of collisions constraints that were restored
-		*/
-		int32 RestoreCollisions(const FReal CullDistance, const bool bReproject);
-
-		/**
-		 * @brief Reactivate the constraint (essentially the same as Restore but slightly optimized)
+		 * @brief Reactivate the constraint
 		 * @parame SleepEpoch The tick on which the particle went to sleep.
 		 * Only constraints that were active when the particle went to sleep should be reactivated.
 		*/
@@ -254,20 +218,6 @@ namespace Chaos
 		const FPerShapeData* Shape1;
 	};
 
-
-	class FMidPhaseRestoreThresholds
-	{
-	public:
-		FMidPhaseRestoreThresholds()
-			: PositionThreshold(0)
-			, RotationThreshold(0)
-		{
-		}
-
-		FReal PositionThreshold;	// cm
-		FReal RotationThreshold;	// rad
-	};
-
 	/**
 	 * @brief Produce collisions for a particle pair
 	 * A FParticlePairMidPhase object is created for every particle pair whose bounds overlap. It is 
@@ -427,18 +377,6 @@ namespace Chaos
 
 		void InitThresholds();
 
-		/**
-		 * @brief Whether we should reuse the constraint as-is and skip the narrow phase
-		 * This will be true if neither particle has moved much. 
-		 * This is non-const because it updates some position tracking data.
-		*/
-		ECollisionRestoreType ShouldRestoreConstraints(const FReal Dt);
-
-		/**
-		 * @brief If the particles have not moved muc, reactivate all the colisions and skip the narrow phase.
-		*/
-		bool TryRestoreConstraints(const FReal Dt, const FReal CullDistance);
-
 
 		FGeometryParticleHandle* Particle0;
 		FGeometryParticleHandle* Particle1;
@@ -455,7 +393,6 @@ namespace Chaos
 			{
 				uint32 bIsCCD : 1;
 				uint32 bIsInitialized : 1;
-				uint32 bRestorable : 1;
 				uint32 bIsSleeping : 1;
 			};
 			uint32 Bits;
@@ -467,14 +404,6 @@ namespace Chaos
 		int32 ParticleCollisionsIndex0;
 		int32 ParticleCollisionsIndex1;
 
-		// The particle transforms the last time the collisions were updated (used to determine whether we can restore contacts)
-		FMidPhaseRestoreThresholds RestoreThresholdZeroContacts;
-		FMidPhaseRestoreThresholds RestoreThreshold;
-		FVec3 RestoreParticleP0;
-		FVec3 RestoreParticleP1;
-		FRotation3 RestoreParticleQ0;
-		FRotation3 RestoreParticleQ1;
-
 		// A number based on the size of the dynamic objects used to scale cull distance
 		FReal CullDistanceScale;
 	};
diff --git a/Engine/Source/Runtime/Experimental/Chaos/Public/Chaos/Framework/UncheckedArray.h b/Engine/Source/Runtime/Experimental/Chaos/Public/Chaos/Framework/UncheckedArray.h
--- a/Engine/Source/Runtime/Experimental/Chaos/Public/Chaos/Framework/UncheckedArray.h
+++ b/Engine/Source/Runtime/Experimental/Chaos/Public/Chaos/Framework/UncheckedArray.h
@@ -166,8 +166,9 @@ namespace Chaos
 		}
 
 	private:
-		ElementType Elements[MaxElements];
+		// Element count before array for better cache behaviour
 		int32 NumElements;
+		ElementType Elements[MaxElements];
 	};
 
 
