Fix bad awake state when particles are stationary

diff --git a/Engine/Source/Runtime/Experimental/Chaos/Private/Chaos/Island/IslandGraph.cpp b/Engine/Source/Runtime/Experimental/Chaos/Private/Chaos/Island/IslandGraph.cpp
--- a/Engine/Source/Runtime/Experimental/Chaos/Private/Chaos/Island/IslandGraph.cpp
+++ b/Engine/Source/Runtime/Experimental/Chaos/Private/Chaos/Island/IslandGraph.cpp
@@ -31,16 +31,19 @@ void FIslandGraph<NodeType, EdgeType, IslandType>::ReserveNodes(const int32 NumN
 }
 
 template<typename NodeType, typename EdgeType, typename IslandType>
-void FIslandGraph<NodeType, EdgeType, IslandType>::ParentIslands(const int32 FirstIsland, const int32 SecondIsland)
+void FIslandGraph<NodeType, EdgeType, IslandType>::ParentIslands(const int32 FirstIsland, const int32 SecondIsland, const bool bIsEdgeMoving)
 {
 	if (GraphIslands.IsValidIndex(FirstIsland) && GraphIslands.IsValidIndex(SecondIsland) && (FirstIsland != SecondIsland))
 	{
 		GraphIslands[SecondIsland].ChildrenIslands.Add(FirstIsland);
 		GraphIslands[FirstIsland].ChildrenIslands.Add(SecondIsland);
-		
-		// If we are adding a constraint the island is no longer persistent
-		GraphIslands[FirstIsland].bIsPersistent = false;
-		GraphIslands[SecondIsland].bIsPersistent = false;
+
+		if(bIsEdgeMoving)
+		{
+			// If we are adding a constraint the island is no longer persistent
+			GraphIslands[FirstIsland].bIsPersistent = false;
+			GraphIslands[SecondIsland].bIsPersistent = false;
+		}
 	}
 }
 	
@@ -59,10 +62,14 @@ void FIslandGraph<NodeType, EdgeType, IslandType>::UpdateNode(const NodeType& No
 			for (int32& EdgeIndex : GraphNode.NodeEdges)
 			{
 				FGraphEdge& GraphEdge = GraphEdges[EdgeIndex];
+				const bool bIsEdgeMoving = IsEdgeMoving(EdgeIndex);
 				
-				ParentIslands(MasterIsland, GraphEdge.IslandIndex);
+				ParentIslands(MasterIsland, GraphEdge.IslandIndex, bIsEdgeMoving);
 				MasterIsland = GraphEdge.IslandIndex;
-				GraphIslands[GraphEdge.IslandIndex].bIsPersistent = false;
+				if(bIsEdgeMoving)
+				{
+					GraphIslands[GraphEdge.IslandIndex].bIsPersistent = false;
+				}
 			}
 		}
 		// If node is not valid and was valid before we check that all the node edges
@@ -133,7 +140,10 @@ void FIslandGraph<NodeType, EdgeType, IslandType>::RemoveNode(const NodeType& No
 			// If only one node and zero edges, we invalidate the node island
 			if(GraphNode.NodeEdges.Num() == 0  && GraphIslands.IsValidIndex(GraphNode.IslandIndex))
 			{
-				GraphIslands[GraphNode.IslandIndex].bIsPersistent = false;
+				if(!GraphNode.bStationaryNode)
+				{
+					GraphIslands[GraphNode.IslandIndex].bIsPersistent = false;
+				}
 			}
 			else
 			{
@@ -143,7 +153,10 @@ void FIslandGraph<NodeType, EdgeType, IslandType>::RemoveNode(const NodeType& No
 					if (GraphNode.NodeEdges.IsValidIndex(NodeEdgeIndex))
 					{
 						const int32 GraphEdgeIndex = GraphNode.NodeEdges[NodeEdgeIndex];
-						GraphIslands[GraphEdges[GraphEdgeIndex].IslandIndex].bIsPersistent = false;
+						if(IsEdgeMoving(GraphEdgeIndex))
+						{
+							GraphIslands[GraphEdges[GraphEdgeIndex].IslandIndex].bIsPersistent = false;
+						}
 						RemoveEdge(GraphEdgeIndex);
 					}
 				}
@@ -187,6 +200,8 @@ void FIslandGraph<NodeType, EdgeType, IslandType>::AttachIslands(const int32 Edg
 			const bool bFirstValidIsland = GraphIslands.IsValidIndex(FirstNode.IslandIndex) && FirstNode.bValidNode;
 			const bool bSecondValidIsland = GraphIslands.IsValidIndex(SecondNode.IslandIndex) && SecondNode.bValidNode;
 
+			const bool bIsEdgeMoving = IsEdgeMoving(EdgeIndex);
+
 			// We check if one of the 2 nodes have invalid island
 			// if yes we set the invalid node island index 
 			// and the edge one to be the valid one 
@@ -195,7 +210,7 @@ void FIslandGraph<NodeType, EdgeType, IslandType>::AttachIslands(const int32 Edg
 			{
 				GraphEdge.IslandIndex = FirstNode.IslandIndex;
 				SecondNode.IslandIndex = GraphEdge.IslandIndex;
-				if(SecondNode.bValidNode)
+				if(SecondNode.bValidNode && bIsEdgeMoving)
 				{
 					GraphIslands[GraphEdge.IslandIndex].bIsPersistent = false;
 				}
@@ -204,7 +219,7 @@ void FIslandGraph<NodeType, EdgeType, IslandType>::AttachIslands(const int32 Edg
 			{
 				GraphEdge.IslandIndex = SecondNode.IslandIndex;
 				FirstNode.IslandIndex = GraphEdge.IslandIndex;
-				if(FirstNode.bValidNode)
+				if(FirstNode.bValidNode && bIsEdgeMoving)
 				{
 					GraphIslands[GraphEdge.IslandIndex].bIsPersistent = false;
 				}
@@ -224,7 +239,7 @@ void FIslandGraph<NodeType, EdgeType, IslandType>::AttachIslands(const int32 Edg
 				// In order to do that we build an island graph and we will 
 				// merge recursively the children islands onto the parent one
 				GraphEdge.IslandIndex = FMath::Min(FirstNode.IslandIndex, SecondNode.IslandIndex);
-				ParentIslands(FirstNode.IslandIndex, SecondNode.IslandIndex);
+				ParentIslands(FirstNode.IslandIndex, SecondNode.IslandIndex, bIsEdgeMoving);
 			}
 		}
 		else if (GraphNodes.IsValidIndex(GraphEdge.FirstNode) && !GraphNodes.IsValidIndex(GraphEdge.SecondNode) && GraphNodes[GraphEdge.FirstNode].bValidNode)
diff --git a/Engine/Source/Runtime/Experimental/Chaos/Public/Chaos/Island/IslandGraph.h b/Engine/Source/Runtime/Experimental/Chaos/Public/Chaos/Island/IslandGraph.h
--- a/Engine/Source/Runtime/Experimental/Chaos/Public/Chaos/Island/IslandGraph.h
+++ b/Engine/Source/Runtime/Experimental/Chaos/Public/Chaos/Island/IslandGraph.h
@@ -222,8 +222,9 @@ public:
 	 * Link two islands to each other for future island traversal
 	 * @param FirstIsland First island index to be linked
 	 * @param SecondIsland Second island index to be linked
+	 * @param bNonStationary Boolean to check if at least one of the 2 nodes are non stationary
 	 */
-	void ParentIslands(const int32 FirstIsland, const int32 SecondIsland);
+	void ParentIslands(const int32 FirstIsland, const int32 SecondIsland,const bool bNonStationary);
 
 	/**
 	 * Given an edge index we try to attach islands. 
@@ -306,6 +307,24 @@ public:
 	* @return First available color 
 	*/
 	int32 PickColor(const FGraphNode& GraphNode, const int32 OtherIndex, TQueue<int32>& NodeQueue);
+
+	/**
+	* Check if at least one of the 2 edge nodes is moving
+	* @param EdgeIndex Index of the edge to check
+	* @return Boolean to check if the edge is moving or not
+	*/
+	bool IsEdgeMoving(const int32 EdgeIndex)
+	{
+		if(GraphEdges.IsValidIndex(EdgeIndex))
+		{
+			const FGraphEdge& GraphEdge = GraphEdges[EdgeIndex];
+			if(GraphNodes.IsValidIndex(GraphEdge.FirstNode) && GraphNodes.IsValidIndex(GraphEdge.SecondNode))
+			{
+				return !GraphNodes[GraphEdge.FirstNode].bStationaryNode || !GraphNodes[GraphEdge.SecondNode].bStationaryNode;
+			}
+		}
+		return false;
+	}
 	
 	/** List of graph nodes */
 	TSparseArray<FGraphNode> GraphNodes;
