Chaos : Clustering Connection Graph
- Add ability to display using debug command (p.Chaos.Solver.DebugDraw.Cluster.Constraints 1)
- Changed default to MinimalSpanningSubsetDelaunayTriangulation

diff --git a/Engine/Source/Runtime/Experimental/Chaos/Private/Chaos/ChaosDebugDraw.cpp b/Engine/Source/Runtime/Experimental/Chaos/Private/Chaos/ChaosDebugDraw.cpp
--- a/Engine/Source/Runtime/Experimental/Chaos/Private/Chaos/ChaosDebugDraw.cpp
+++ b/Engine/Source/Runtime/Experimental/Chaos/Private/Chaos/ChaosDebugDraw.cpp
@@ -20,6 +20,7 @@
 #include "Chaos/PBDConstraintColor.h"
 #include "Chaos/PBDConstraintGraph.h"
 #include "Chaos/PBDJointConstraints.h"
+#include "Chaos/PBDRigidClustering.h"
 #include "Chaos/PBDRigidParticles.h"
 #include "Chaos/PBDSuspensionConstraints.h"
 #include "Chaos/Sphere.h"
@@ -1216,6 +1217,119 @@ namespace Chaos
 			}
 		}
 
+		struct FOrderedClusterConnection {
+			const FRigidClustering::FClusterHandle Parent;
+			const FRigidClustering::FClusterHandle Handle0;
+			const FRigidClustering::FClusterHandle Handle1;
+			FReal Strain;
+
+			FOrderedClusterConnection()
+				: Parent(nullptr)
+				, Handle0(nullptr)
+				, Handle1(nullptr)
+				, Strain(0)
+			{}
+
+			FOrderedClusterConnection(const FRigidClustering::FClusterHandle InParent,
+				const FRigidClustering::FClusterHandle In0,
+				const FRigidClustering::FClusterHandle In1,
+				FReal InStrain)
+				: Parent(InParent)
+				, Handle0(In0)
+				, Handle1(In1)
+				, Strain(InStrain)
+			{
+				check(In0 < In1);
+			}
+			bool operator==(const FOrderedClusterConnection& R) const
+			{
+				return Equals(R);
+			}
+
+			bool Equals(const FOrderedClusterConnection& R) const
+			{
+				return Handle0 == R.Handle0 && Handle1 == R.Handle1;
+			}
+		};
+
+		uint32 GetTypeHash(const FOrderedClusterConnection& Object)
+		{
+			uint32 Hash = FCrc::MemCrc32(&Object, sizeof(FOrderedClusterConnection));
+			return Hash;
+		}
+
+		void DrawConnectionGraphImpl(const FRigidClustering& Clustering, const FChaosDebugDrawSettings& Settings)
+		{
+			auto DrawConnections = [&](TSet<FOrderedClusterConnection>& Connections, FReal MaxStrain, const FChaosDebugDrawSettings& Settings)
+			{
+				if (FMath::IsNearlyZero(MaxStrain))
+				{
+					MaxStrain = (FReal)1;
+				}
+
+				for (auto& Connection : Connections)
+				{
+					FRigidTransform3 ClusterTransform(Connection.Parent->X(), Connection.Parent->R());
+					FVec3 Pos0 = ClusterTransform.TransformPosition(Connection.Handle0->ChildToParent().GetLocation());
+					FVec3 Pos1 = ClusterTransform.TransformPosition(Connection.Handle1->ChildToParent().GetLocation());
+
+					FColor Color = FColor::Green; // FMath::Lerp(FColor::Green, FColor::Red, Connection.Strain / MaxStrain);
+
+					FDebugDrawQueue::GetInstance().DrawDebugLine(Pos0, Pos1, Color, false, KINDA_SMALL_NUMBER, Settings.DrawPriority, Settings.LineThickness * 3);
+				}
+
+			};
+
+			//if (bChaosDebugDrawClusterConnectionGraph)
+			{
+				FReal MaxStrain = -FLT_MAX;
+				int32 NumConnections = 0;
+				const TArrayCollectionArray<FReal>& Strain = Clustering.GetStrainArray();
+
+				for (auto& ActiveCluster : Clustering.GetChildrenMap())
+				{
+					if (!ActiveCluster.Key->Disabled())
+					{
+						for (const FRigidClustering::FRigidHandle& Rigid : ActiveCluster.Value)
+						{
+							if (const FRigidClustering::FClusterHandle& Child = Rigid->CastToClustered())
+							{
+								MaxStrain = FMath::Max(MaxStrain, Child->Strain());
+							}
+						}
+						NumConnections += ActiveCluster.Value.Num();
+					}
+				}
+
+				if (NumConnections)
+				{
+					TSet<FOrderedClusterConnection> Connections;
+					Connections.Reserve(NumConnections);
+
+					for (auto& ActiveCluster : Clustering.GetChildrenMap())
+					{
+						if (!ActiveCluster.Key->Disabled())
+						{
+							for (const FRigidClustering::FRigidHandle& Rigid : ActiveCluster.Value)
+							{
+								if (const FRigidClustering::FClusterHandle& Child = Rigid->CastToClustered())
+								{
+									const FConnectivityEdgeArray& Edges = Child->ConnectivityEdges();
+									for (const TConnectivityEdge<FReal>& Edge : Edges)
+									{
+										const FRigidClustering::FRigidHandle A = (Child < Edge.Sibling) ? Child : Edge.Sibling;
+										const FRigidClustering::FRigidHandle B = (Child < Edge.Sibling) ? Edge.Sibling : Child;
+										Connections.Add(FOrderedClusterConnection(ActiveCluster.Key, A->CastToClustered(), B->CastToClustered(), Edge.Strain));
+									}
+								}
+							}
+						}
+					}
+					DrawConnections(Connections, MaxStrain, Settings);
+				}
+			}
+		}
+
 		void DrawSuspensionConstraintsImpl(const FRigidTransform3& SpaceTransform, const FPBDSuspensionConstraints& Constraints, int32 ConstraintIndex, const FChaosDebugDrawSettings& Settings)
 		{
 			FConstGenericParticleHandle Particle = Constraints.GetConstrainedParticles(ConstraintIndex)[0];
@@ -1444,6 +1558,14 @@ namespace Chaos
 			}
 		}
 
+		void DrawConnectionGraph(const FRigidClustering& Clustering, const FChaosDebugDrawSettings* Settings)
+		{
+			if (FDebugDrawQueue::IsDebugDrawingEnabled())
+			{
+				DrawConnectionGraphImpl(Clustering, GetChaosDebugDrawSettings(Settings));
+			}
+		}
+
 		void DrawCollidingShapes(const FRigidTransform3& SpaceTransform, const FPBDCollisionConstraints& Collisions, FRealSingle ColorScale, const FChaosDebugDrawSettings* Settings)
 		{
 			if (FDebugDrawQueue::IsDebugDrawingEnabled())
diff --git a/Engine/Source/Runtime/Experimental/Chaos/Private/PBDRigidsSolver.cpp b/Engine/Source/Runtime/Experimental/Chaos/Private/PBDRigidsSolver.cpp
--- a/Engine/Source/Runtime/Experimental/Chaos/Private/PBDRigidsSolver.cpp
+++ b/Engine/Source/Runtime/Experimental/Chaos/Private/PBDRigidsSolver.cpp
@@ -82,6 +82,7 @@ namespace Chaos
 		int32 ChaosSolverDebugDrawSpatialAccelerationStructureShowLeaves = 0;
 		int32 ChaosSolverDebugDrawSpatialAccelerationStructureShowNodes = 0;
 		int32 ChaosSolverDebugDrawSuspensionConstraints = 0;
+		int32 ChaosSolverDrawClusterConstraints = 0;
 		DebugDraw::FChaosDebugDrawJointFeatures ChaosSolverDrawJointFeatures = DebugDraw::FChaosDebugDrawJointFeatures::MakeDefault();
 		FAutoConsoleVariableRef CVarChaosSolverDrawShapes(TEXT("p.Chaos.Solver.DebugDrawShapes"), ChaosSolverDebugDrawShapes, TEXT("Draw Shapes (0 = never; 1 = end of frame)."));
 		FAutoConsoleVariableRef CVarChaosSolverDrawCollisions(TEXT("p.Chaos.Solver.DebugDrawCollisions"), ChaosSolverDebugDrawCollisions, TEXT("Draw Collisions (0 = never; 1 = end of frame)."));
@@ -105,6 +106,7 @@ namespace Chaos
 		FAutoConsoleVariableRef CVarChaosSolverDrawJointFeaturesColor(TEXT("p.Chaos.Solver.DebugDraw.JointFeatures.Color"), ChaosSolverDrawJointFeatures.bColor, TEXT("Joint features mask (see FDebugDrawJointFeatures)."));
 		FAutoConsoleVariableRef CVarChaosSolverDrawJointFeaturesIsland(TEXT("p.Chaos.Solver.DebugDraw.JointFeatures.Island"), ChaosSolverDrawJointFeatures.bIsland, TEXT("Joint features mask (see FDebugDrawJointFeatures)."));
 		FAutoConsoleVariableRef CVarChaosSolverDrawSuspensionConstraints(TEXT("p.Chaos.Solver.DebugDrawSuspension"), ChaosSolverDebugDrawSuspensionConstraints, TEXT("Draw Suspension (0 = never; 1 = end of frame)."));
+		FAutoConsoleVariableRef CVarChaosSolverDrawClusterConstraints(TEXT("p.Chaos.Solver.DebugDraw.Cluster.Constraints"), ChaosSolverDrawClusterConstraints, TEXT("Draw Active Cluster Constraints (0 = never; 1 = end of frame)."));
 
 
 		DebugDraw::FChaosDebugDrawSettings ChaosSolverDebugDebugDrawSettings(
@@ -1510,6 +1512,10 @@ CSV_CUSTOM_STAT(PhysicsCounters, Name, Value, ECsvCustomStatOp::Set);
 		{
 			DebugDraw::DrawConstraintGraph(FRigidTransform3(), GetEvolution()->GetConstraintGraph(), &ChaosSolverDebugDebugDrawSettings);
 		}
+		if (ChaosSolverDrawClusterConstraints == 1)
+		{
+			DebugDraw::DrawConnectionGraph(MEvolution->GetRigidClustering(), &ChaosSolverDebugDebugDrawSettings);
+		}
 		if (ChaosSolverDebugDrawCollidingShapes == 1)
 		{
 			DebugDraw::DrawCollidingShapes(FRigidTransform3(), GetEvolution()->GetCollisionConstraints(), 1.f, &ChaosSolverDebugDebugDrawSettings);
diff --git a/Engine/Source/Runtime/Experimental/Chaos/Public/Chaos/ChaosDebugDraw.h b/Engine/Source/Runtime/Experimental/Chaos/Public/Chaos/ChaosDebugDraw.h
--- a/Engine/Source/Runtime/Experimental/Chaos/Public/Chaos/ChaosDebugDraw.h
+++ b/Engine/Source/Runtime/Experimental/Chaos/Public/Chaos/ChaosDebugDraw.h
@@ -13,6 +13,7 @@ namespace Chaos
 	class FCollisionConstraintAllocator;
 	template <typename PayloadType, typename T, int d> class ISpatialAcceleration;
 	class FPBDSuspensionConstraints;
+	class FRigidClustering;
 
 	namespace DebugDraw
 	{
@@ -209,6 +210,7 @@ namespace Chaos
 		CHAOS_API void DrawSimulationSpace(const FSimulationSpace& SimSpace, const FChaosDebugDrawSettings* Settings = nullptr);
 		CHAOS_API void DrawShape(const FRigidTransform3& ShapeTransform, const FImplicitObject* Shape, const FColor& Color, const FChaosDebugDrawSettings* Settings = nullptr);
 		CHAOS_API void DrawConstraintGraph(const FRigidTransform3& ShapeTransform, const FPBDConstraintGraph& Graph, const FChaosDebugDrawSettings* Settings = nullptr);
+		CHAOS_API void DrawConnectionGraph(const FRigidClustering& Clustering, const FChaosDebugDrawSettings* Settings = nullptr);
 		CHAOS_API void DrawCollidingShapes(const FRigidTransform3& SpaceTransform, const FPBDCollisionConstraints& Collisions, FRealSingle ColorScale, const FChaosDebugDrawSettings* Settings = nullptr);
 		CHAOS_API void DrawSpatialAccelerationStructure(const ISpatialAcceleration<FAccelerationStructureHandle, FReal, 3>& SpatialAccelerationStructure, const FChaosDebugDrawSettings* Settings = nullptr);
 		CHAOS_API void DrawSuspensionConstraints(const FRigidTransform3& SpaceTransform, const FPBDSuspensionConstraints& Constraints, const FChaosDebugDrawSettings* Settings = nullptr);
diff --git a/Engine/Source/Runtime/Experimental/Chaos/Public/Chaos/PBDRigidClusteredParticles.h b/Engine/Source/Runtime/Experimental/Chaos/Public/Chaos/PBDRigidClusteredParticles.h
--- a/Engine/Source/Runtime/Experimental/Chaos/Public/Chaos/PBDRigidClusteredParticles.h
+++ b/Engine/Source/Runtime/Experimental/Chaos/Public/Chaos/PBDRigidClusteredParticles.h
@@ -63,6 +63,8 @@ struct TConnectivityEdge
 	TPBDRigidParticleHandle<T, 3>* Sibling;
 	T Strain;
 };
+typedef TConnectivityEdge<FReal> FConnectivityEdge;
+typedef TArray<FConnectivityEdge> FConnectivityEdgeArray;
 
 template<class T, int d>
 class TPBDRigidClusteredParticles : public TPBDRigidParticles<T, d>
diff --git a/Engine/Source/Runtime/Experimental/Chaos/Public/Chaos/PBDRigidClustering.h b/Engine/Source/Runtime/Experimental/Chaos/Public/Chaos/PBDRigidClustering.h
--- a/Engine/Source/Runtime/Experimental/Chaos/Public/Chaos/PBDRigidClustering.h
+++ b/Engine/Source/Runtime/Experimental/Chaos/Public/Chaos/PBDRigidClustering.h
@@ -212,7 +212,8 @@ public:
 	*    the cluster body, can be updated during the evaluation of the simulation.
 	*/
 	TArrayCollectionArray<FReal>& GetStrainArray() { return MParticles.StrainsArray(); }
-
+	const TArrayCollectionArray<FReal>& GetStrainArray() const { return MParticles.StrainsArray(); }
+		
 	/**
 	*  GetParentToChildren
 	*    The parent to children map stores the currently active cluster ids (Particle Indices) as
diff --git a/Engine/Source/Runtime/Experimental/GeometryCollectionEngine/Private/GeometryCollection/GeometryCollectionObject.cpp b/Engine/Source/Runtime/Experimental/GeometryCollectionEngine/Private/GeometryCollection/GeometryCollectionObject.cpp
--- a/Engine/Source/Runtime/Experimental/GeometryCollectionEngine/Private/GeometryCollection/GeometryCollectionObject.cpp
+++ b/Engine/Source/Runtime/Experimental/GeometryCollectionEngine/Private/GeometryCollection/GeometryCollectionObject.cpp
@@ -63,7 +63,7 @@ UGeometryCollection::UGeometryCollection(const FObjectInitializer& ObjectInitial
 	, ClusterGroupIndex(0)
 	, MaxClusterLevel(100)
 	, DamageThreshold({ 250.0 })
-	, ClusterConnectionType(EClusterConnectionTypeEnum::Chaos_PointImplicit)
+	, ClusterConnectionType(EClusterConnectionTypeEnum::Chaos_MinimalSpanningSubsetDelaunayTriangulation)
 	, bUseFullPrecisionUVs(false)
 	, bStripOnCook(false)
 	, EnableNanite(false)
