Chaos : Clustering
- Add connection graph to the bottom level clusters.

diff --git a/Engine/Source/Runtime/Experimental/Chaos/Private/PhysicsProxy/GeometryCollectionPhysicsProxy.cpp b/Engine/Source/Runtime/Experimental/Chaos/Private/PhysicsProxy/GeometryCollectionPhysicsProxy.cpp
--- a/Engine/Source/Runtime/Experimental/Chaos/Private/PhysicsProxy/GeometryCollectionPhysicsProxy.cpp
+++ b/Engine/Source/Runtime/Experimental/Chaos/Private/PhysicsProxy/GeometryCollectionPhysicsProxy.cpp
@@ -976,7 +976,7 @@ FGeometryCollectionPhysicsProxy::BuildClusters(
 	//Since the mass orientation is computed in world space in both cases we'd end up with a diagonal inertia matrix and identity rotation that looks like this: [big, small, big] or [small, big, big].
 	//Because of this we need to know how to rotate collision particles and geometry to match with original computation. If it was just geometry we could transform it before passing, but we need collision particles as well
 	Chaos::FClusterCreationParameters ClusterCreationParameters = ClusterParameters;
-	ClusterCreationParameters.bGenerateConnectionGraph = false;
+	ClusterCreationParameters.bGenerateConnectionGraph = true;
 	// fix... ClusterCreationParameters.CollisionParticles = Simplicials[CollectionClusterIndex];
 	ClusterCreationParameters.ConnectionMethod = Parameters.ClusterConnectionMethod;
 	if (ClusterCreationParameters.CollisionParticles)
