Chaos - fix shape world transform for incremental manifolds on shapes wiht offsets

diff --git a/Engine/Source/Runtime/Experimental/Chaos/Private/Chaos/Collision/PBDCollisionConstraint.cpp b/Engine/Source/Runtime/Experimental/Chaos/Private/Chaos/Collision/PBDCollisionConstraint.cpp
--- a/Engine/Source/Runtime/Experimental/Chaos/Private/Chaos/Collision/PBDCollisionConstraint.cpp
+++ b/Engine/Source/Runtime/Experimental/Chaos/Private/Chaos/Collision/PBDCollisionConstraint.cpp
@@ -429,8 +429,8 @@ namespace Chaos
 		// @todo(chaos): Remove this when we don't need to support incremental manifolds (this will only be called on creation/restore)
 		if ((GetSolverBody0() != nullptr) && (GetSolverBody1() != nullptr))
 		{
-			ShapeWorldTransform0 = FRigidTransform3(GetSolverBody0()->CorrectedP(), GetSolverBody0()->CorrectedQ());
-			ShapeWorldTransform1 = FRigidTransform3(GetSolverBody1()->CorrectedP(), GetSolverBody1()->CorrectedQ());
+			ShapeWorldTransform0 = GetShapeRelativeTransform0() * FRigidTransform3(GetSolverBody0()->CorrectedP(), GetSolverBody0()->CorrectedQ());
+			ShapeWorldTransform1 = GetShapeRelativeTransform1() * FRigidTransform3(GetSolverBody1()->CorrectedP(), GetSolverBody1()->CorrectedQ());
 		}
 
 		Flags.bDisabled = false;
