Chaos - fix RBAN collision for offset shapes

diff --git a/Engine/Source/Runtime/Experimental/Chaos/Private/Chaos/Collision/PBDCollisionConstraint.cpp b/Engine/Source/Runtime/Experimental/Chaos/Private/Chaos/Collision/PBDCollisionConstraint.cpp
--- a/Engine/Source/Runtime/Experimental/Chaos/Private/Chaos/Collision/PBDCollisionConstraint.cpp
+++ b/Engine/Source/Runtime/Experimental/Chaos/Private/Chaos/Collision/PBDCollisionConstraint.cpp
@@ -425,12 +425,34 @@ namespace Chaos
 
 	void FPBDCollisionConstraint::UpdateManifoldContacts()
 	{
-		// This is only used when calling collision detection in the solver loop, which is only for incremental manifolds.
-		// @todo(chaos): Remove this when we don't need to support incremental manifolds (this will only be called on creation/restore)
 		if ((GetSolverBody0() != nullptr) && (GetSolverBody1() != nullptr))
 		{
-			ShapeWorldTransform0 = GetShapeRelativeTransform0() * FRigidTransform3(GetSolverBody0()->CorrectedP(), GetSolverBody0()->CorrectedQ());
-			ShapeWorldTransform1 = GetShapeRelativeTransform1() * FRigidTransform3(GetSolverBody1()->CorrectedP(), GetSolverBody1()->CorrectedQ());
+			// This is only entered when calling collision detection in a legacy solver (RBAN).
+			// We need to update the contact Phi for the current iteration based on what the body transforms would be if
+			// we applied the corrections accumulated so far.
+			// @todo(chaos): It is extremely expensicve! Remove this when RBAN uses the QuasiPBD solver.
+			// NOTE: ShapeRelativeTransforms are in actor-space. The SolverBodies give CoM transforms.
+			const FConstGenericParticleHandle P0 = Particle[0];
+			const FConstGenericParticleHandle P1 = Particle[1];
+
+			const FRigidTransform3 ParticleCoMTransform0 = FRigidTransform3(
+				GetSolverBody0()->CorrectedP(), 
+				GetSolverBody0()->CorrectedQ());
+
+			const FRigidTransform3 ParticleCoMTransform1 = FRigidTransform3(
+				GetSolverBody1()->CorrectedP(), 
+				GetSolverBody1()->CorrectedQ());
+
+			const FRigidTransform3 ShapeCoMRelativeTransform0 = FRigidTransform3(
+				P0->RotationOfMass().UnrotateVector(ImplicitTransform[0].GetLocation() - P0->CenterOfMass()),
+				P0->RotationOfMass().Inverse() * ImplicitTransform[0].GetRotation());
+
+			const FRigidTransform3 ShapeCoMRelativeTransform1 = FRigidTransform3(
+				P1->RotationOfMass().UnrotateVector(ImplicitTransform[1].GetLocation() - P1->CenterOfMass()),
+				P1->RotationOfMass().Inverse() * ImplicitTransform[1].GetRotation());
+
+			ShapeWorldTransform0 = ShapeCoMRelativeTransform0 * ParticleCoMTransform0;
+			ShapeWorldTransform1 = ShapeCoMRelativeTransform1 * ParticleCoMTransform1;
 		}
 
 		Flags.bDisabled = false;
