Chaos - debug draw settings for manifold point info

diff --git a/Engine/Source/Runtime/Engine/Private/PhysicsEngine/ImmediatePhysics/ImmediatePhysicsChaos/ImmediatePhysicsSimulation_Chaos.cpp b/Engine/Source/Runtime/Engine/Private/PhysicsEngine/ImmediatePhysics/ImmediatePhysicsChaos/ImmediatePhysicsSimulation_Chaos.cpp
--- a/Engine/Source/Runtime/Engine/Private/PhysicsEngine/ImmediatePhysics/ImmediatePhysicsChaos/ImmediatePhysicsSimulation_Chaos.cpp
+++ b/Engine/Source/Runtime/Engine/Private/PhysicsEngine/ImmediatePhysics/ImmediatePhysicsChaos/ImmediatePhysicsSimulation_Chaos.cpp
@@ -179,6 +179,7 @@ Chaos::DebugDraw::FChaosDebugDrawSettings ChaosImmPhysDebugDebugDrawSettings(
 	/* ContactLen =					*/ 4.0f,
 	/* ContactWidth =				*/ 2.0f,
 	/* ContactPhiWidth =			*/ 0.0f,
+	/* ContactInfoWidth				*/ 2.0f,
 	/* ContactOwnerWidth =			*/ 0.0f,
 	/* ConstraintAxisLen =			*/ 5.0f,
 	/* JointComSize =				*/ 2.0f,
@@ -207,6 +208,7 @@ FAutoConsoleVariableRef CVarChaosImmPhysArrowSize(TEXT("p.Chaos.ImmPhys.DebugDra
 FAutoConsoleVariableRef CVarChaosImmPhysBodyAxisLen(TEXT("p.Chaos.ImmPhys.DebugDraw.BodyAxisLen"), ChaosImmPhysDebugDebugDrawSettings.BodyAxisLen, TEXT("BodyAxisLen."));
 FAutoConsoleVariableRef CVarChaosImmPhysContactLen(TEXT("p.Chaos.ImmPhys.DebugDraw.ContactLen"), ChaosImmPhysDebugDebugDrawSettings.ContactLen, TEXT("ContactLen."));
 FAutoConsoleVariableRef CVarChaosImmPhysContactWidth(TEXT("p.Chaos.ImmPhys.DebugDraw.ContactWidth"), ChaosImmPhysDebugDebugDrawSettings.ContactWidth, TEXT("ContactWidth."));
+FAutoConsoleVariableRef CVarChaosImmPhysContactInfoWidth(TEXT("p.Chaos.ImmPhys.DebugDraw.ContactInfoWidth"), ChaosImmPhysDebugDebugDrawSettings.ContactInfoWidth, TEXT("ContactInfoWidth."));
 FAutoConsoleVariableRef CVarChaosImmPhysContactPhiWidth(TEXT("p.Chaos.ImmPhys.DebugDraw.ContactPhiWidth"), ChaosImmPhysDebugDebugDrawSettings.ContactPhiWidth, TEXT("ContactPhiWidth."));
 FAutoConsoleVariableRef CVarChaosImmPhysContactOwnerWidth(TEXT("p.Chaos.ImmPhys.DebugDraw.ContactOwnerWidth"), ChaosImmPhysDebugDebugDrawSettings.ContactOwnerWidth, TEXT("ContactOwnerWidth."));
 FAutoConsoleVariableRef CVarChaosImmPhysConstraintAxisLen(TEXT("p.Chaos.ImmPhys.DebugDraw.ConstraintAxisLen"), ChaosImmPhysDebugDebugDrawSettings.ConstraintAxisLen, TEXT("ConstraintAxisLen."));
diff --git a/Engine/Source/Runtime/Experimental/Chaos/Private/Chaos/ChaosDebugDraw.cpp b/Engine/Source/Runtime/Experimental/Chaos/Private/Chaos/ChaosDebugDraw.cpp
--- a/Engine/Source/Runtime/Experimental/Chaos/Private/Chaos/ChaosDebugDraw.cpp
+++ b/Engine/Source/Runtime/Experimental/Chaos/Private/Chaos/ChaosDebugDraw.cpp
@@ -90,6 +90,7 @@ namespace Chaos
 			/* ContactLen =					*/ 4.0f,
 			/* ContactWidth =				*/ 2.0f,
 			/* ContactPhiWidth =			*/ 0.0f,
+			/* ContactInfoWidth				*/ 2.0f,
 			/* ContactOwnerWidth =			*/ 0.0f,
 			/* ConstraintAxisLen =			*/ 5.0f,
 			/* JointComSize =				*/ 2.0f,
@@ -801,144 +802,127 @@ namespace Chaos
 			}
 		}
 
-		void DrawCollisionImpl(const FVec3& Location, const FVec3& Normal, FReal Phi, const FVec3& Impulse, const FColor& DiscColor, const FColor& NormalColor, const FColor& ImpulseColor, FRealSingle ColorScale, const FChaosDebugDrawSettings& Settings)
-		{
-			FMatrix Axes = FRotationMatrix::MakeFromX(Normal);
-			if (Settings.ContactWidth > 0)
-			{
-				FColor C0 = (ColorScale * DiscColor).ToFColor(false);
-				FDebugDrawQueue::GetInstance().DrawDebugCircle(Location, Settings.DrawScale * Settings.ContactWidth, 12, C0, false, KINDA_SMALL_NUMBER, Settings.DrawPriority, Settings.LineThickness, Axes.GetUnitAxis(EAxis::Y), Axes.GetUnitAxis(EAxis::Z), false);
-			}
-			if (Settings.ContactLen > 0)
-			{
-				FColor C1 = (ColorScale * NormalColor).ToFColor(false);
-				FDebugDrawQueue::GetInstance().DrawDebugLine(Location, Location + Settings.DrawScale * Settings.ContactLen * Normal, C1, false, KINDA_SMALL_NUMBER, Settings.DrawPriority, Settings.LineThickness);
-			}
-			if (Settings.ContactPhiWidth > 0 && (Phi < FLT_MAX))
-			{
-				FColor C2 = (ColorScale * FColor(128, 128, 0)).ToFColor(false);
-				FDebugDrawQueue::GetInstance().DrawDebugCircle(Location - Phi * Normal, Settings.DrawScale * Settings.ContactPhiWidth, 12, C2, false, KINDA_SMALL_NUMBER, Settings.DrawPriority, Settings.LineThickness, Axes.GetUnitAxis(EAxis::Y), Axes.GetUnitAxis(EAxis::Z), false);
-			}
-			if ((Settings.ImpulseScale > 0) && !Impulse.IsNearlyZero())
-			{
-				FColor C3 = (ColorScale * ImpulseColor).ToFColor(false);
-				FDebugDrawQueue::GetInstance().DrawDebugLine(Location, Location + Settings.DrawScale * Settings.ImpulseScale * Impulse, C3, false, KINDA_SMALL_NUMBER, Settings.DrawPriority, Settings.LineThickness);
-			}
-		}
-
 		void DrawCollisionImpl(const FRigidTransform3& SpaceTransform, const FPBDCollisionConstraint& Contact, FRealSingle ColorScale, const FChaosDebugDrawSettings& Settings)
 		{
-			if ((Settings.ContactWidth > 0) || (Settings.ContactLen > 0) || (Settings.ImpulseScale > 0.0f))
+			if ((Settings.ContactWidth > 0) || (Settings.ContactLen > 0) || (Settings.ContactInfoWidth > 0) || (Settings.ImpulseScale > 0.0f))
 			{
-				if (Contact.GetUseManifold())
+				const FConstGenericParticleHandle Particle0 = Contact.GetParticle0();
+				const FConstGenericParticleHandle Particle1 = Contact.GetParticle1();
+				const FRigidTransform3 WorldActorTransform0 = FParticleUtilities::GetActorWorldTransform(Particle0);
+				const FRigidTransform3 WorldActorTransform1 = FParticleUtilities::GetActorWorldTransform(Particle1);
+
+				// Are we within the region of interest?
+				const FReal Particle0Size = Contact.GetParticle0()->HasBounds() ? 0.5f * Contact.GetParticle0()->LocalBounds().Extents().Size() : TNumericLimits<FReal>::Max();
+				const FReal Particle1Size = Contact.GetParticle1()->HasBounds() ? 0.5f * Contact.GetParticle1()->LocalBounds().Extents().Size() : TNumericLimits<FReal>::Max();
+				if (!FDebugDrawQueue::GetInstance().IsInRegionOfInterest(WorldActorTransform0.GetLocation(), Particle0Size))
 				{
-					const FConstGenericParticleHandle Particle0 = Contact.GetParticle0();
-					const FConstGenericParticleHandle Particle1 = Contact.GetParticle1();
-					const FRigidTransform3 WorldActorTransform0 = FParticleUtilities::GetActorWorldTransform(Particle0);
-					const FRigidTransform3 WorldActorTransform1 = FParticleUtilities::GetActorWorldTransform(Particle1);
-
-					// Are we within the region of interest?
-					const FReal Particle0Size = Contact.GetParticle0()->HasBounds() ? 0.5f * Contact.GetParticle0()->LocalBounds().Extents().Size() : TNumericLimits<FReal>::Max();
-					const FReal Particle1Size = Contact.GetParticle1()->HasBounds() ? 0.5f * Contact.GetParticle1()->LocalBounds().Extents().Size() : TNumericLimits<FReal>::Max();
-					if (!FDebugDrawQueue::GetInstance().IsInRegionOfInterest(WorldActorTransform0.GetLocation(), Particle0Size))
+					return;
+				}
+				if (!FDebugDrawQueue::GetInstance().IsInRegionOfInterest(WorldActorTransform1.GetLocation(), Particle1Size))
+				{
+					return;
+				}
+
+				for (const FManifoldPoint& ManifoldPoint : Contact.GetManifoldPoints())
+				{
+					const bool bIsActive = !ManifoldPoint.NetPushOut.IsNearlyZero() || !ManifoldPoint.NetImpulse.IsNearlyZero();
+					if (!bIsActive && !bChaosDebugDebugDrawInactiveContacts)
 					{
-						return;
+						continue;
 					}
-					if (!FDebugDrawQueue::GetInstance().IsInRegionOfInterest(WorldActorTransform1.GetLocation(), Particle1Size))
+
+					const int32 ContactPlaneOwner = 1;
+					const int32 ContactPointOwner = 1 - ContactPlaneOwner;
+					const FRigidTransform3& PlaneTransform = (ContactPlaneOwner == 0) ? Contact.GetShapeRelativeTransform0() * WorldActorTransform0 : Contact.GetShapeRelativeTransform1() * WorldActorTransform1;
+					const FRigidTransform3& PointTransform = (ContactPlaneOwner == 0) ? Contact.GetShapeRelativeTransform1() * WorldActorTransform1 : Contact.GetShapeRelativeTransform0() * WorldActorTransform0;
+					const FConstGenericParticleHandle PlaneParticle = (ContactPlaneOwner == 0) ? Particle0 : Particle1;
+					const FVec3 PlaneNormal = PlaneTransform.TransformVectorNoScale(ManifoldPoint.ContactPoint.ShapeContactNormal);
+					const FVec3 PointLocation = PointTransform.TransformPosition(ManifoldPoint.ContactPoint.ShapeContactPoints[ContactPointOwner]);
+					const FVec3 PlaneLocation = PlaneTransform.TransformPosition(ManifoldPoint.ContactPoint.ShapeContactPoints[ContactPlaneOwner]);
+					const FVec3 PointPlaneLocation = PointLocation - FVec3::DotProduct(PointLocation - PlaneLocation, PlaneNormal) * PlaneNormal;
+
+					// Dynamic friction, restitution = red
+					// Static friction, no restitution = green
+					// Inactive = gray
+					FColor DiscColor = FColor(200, 0, 0);
+					FColor NormalColor = FColor(200, 0, 0);
+					FColor ImpulseColor = FColor(0, 0, 200);
+					FColor PushOutColor = FColor(200, 200, 0);
+					FColor PushOutImpusleColor = FColor(0, 200, 200);
+					if (ManifoldPoint.Flags.bInsideStaticFrictionCone)
 					{
-						return;
+						DiscColor = FColor(150, 200, 0);
 					}
-
-					for (const FManifoldPoint& ManifoldPoint : Contact.GetManifoldPoints())
+					if (!bIsActive)
 					{
-						const bool bIsActive = !ManifoldPoint.NetPushOut.IsNearlyZero() || !ManifoldPoint.NetImpulse.IsNearlyZero();
-						if (!bIsActive && !bChaosDebugDebugDrawInactiveContacts)
-						{
-							continue;
-						}
-
-						const int32 ContactPlaneOwner = 1;
-						const int32 ContactPointOwner = 1 - ContactPlaneOwner;
-						const FRigidTransform3& PlaneTransform = (ContactPlaneOwner == 0) ? Contact.GetShapeRelativeTransform0() * WorldActorTransform0 : Contact.GetShapeRelativeTransform1() * WorldActorTransform1;
-						const FRigidTransform3& PointTransform = (ContactPlaneOwner == 0) ? Contact.GetShapeRelativeTransform1() * WorldActorTransform1 : Contact.GetShapeRelativeTransform0() * WorldActorTransform0;
-						const FConstGenericParticleHandle PlaneParticle = (ContactPlaneOwner == 0) ? Particle0 : Particle1;
-						const FVec3 PlaneNormal = PlaneTransform.TransformVectorNoScale(ManifoldPoint.ContactPoint.ShapeContactNormal);
-						const FVec3 PointLocation = PointTransform.TransformPosition(ManifoldPoint.ContactPoint.ShapeContactPoints[ContactPointOwner]);
-						const FVec3 PlaneLocation = PlaneTransform.TransformPosition(ManifoldPoint.ContactPoint.ShapeContactPoints[ContactPlaneOwner]);
-						const FVec3 PointPlaneLocation = PointLocation - FVec3::DotProduct(PointLocation - PlaneLocation, PlaneNormal) * PlaneNormal;
-
-						// Dynamic friction, restitution = red
-						// Static friction, no restitution = green
-						// Inactive = gray
-						FColor DiscColor = FColor(200, 0, 0);
-						FColor NormalColor = FColor(200, 0, 0);
-						FColor ImpulseColor = FColor(0, 0, 200);
-						FColor PushOutColor = FColor(200, 200, 0);
-						FColor PushOutImpusleColor = FColor(0, 200, 200);
-						if (ManifoldPoint.Flags.bInsideStaticFrictionCone)
-						{
-							DiscColor = FColor(150, 200, 0);
-						}
-						if (!bIsActive)
-						{
-							DiscColor = FColor(100, 100, 100);
-							NormalColor = FColor(100, 100, 100);
-						}
+						DiscColor = FColor(100, 100, 100);
+						NormalColor = FColor(100, 100, 100);
+					}
 
-						const FVec3 WorldPointLocation = SpaceTransform.TransformPosition(PointLocation);
-						const FVec3 WorldPlaneLocation = SpaceTransform.TransformPosition(PlaneLocation);
-						const FVec3 WorldPointPlaneLocation = SpaceTransform.TransformPosition(PointPlaneLocation);
-						const FVec3 WorldPlaneNormal = SpaceTransform.TransformVectorNoScale(PlaneNormal);
+					const FVec3 WorldPointLocation = SpaceTransform.TransformPosition(PointLocation);
+					const FVec3 WorldPlaneLocation = SpaceTransform.TransformPosition(PlaneLocation);
+					const FVec3 WorldPointPlaneLocation = SpaceTransform.TransformPosition(PointPlaneLocation);
+					const FVec3 WorldPlaneNormal = SpaceTransform.TransformVectorNoScale(PlaneNormal);
+					const FMatrix Axes = FRotationMatrix::MakeFromX(WorldPlaneNormal);
 
-						// Pushout
-						if ((Settings.PushOutScale > 0) && !ManifoldPoint.NetPushOut.IsNearlyZero())
-						{
-							FColor Color = (ColorScale * PushOutImpusleColor).ToFColor(false);
-							FDebugDrawQueue::GetInstance().DrawDebugLine(WorldPointPlaneLocation, WorldPointPlaneLocation + Settings.DrawScale * Settings.PushOutScale * ManifoldPoint.NetPushOut, Color, false, KINDA_SMALL_NUMBER, Settings.DrawPriority, Settings.LineThickness);
-						}
-						if ((Settings.ImpulseScale > 0) && !ManifoldPoint.NetImpulse.IsNearlyZero())
-						{
-							FColor Color = (ColorScale * PushOutImpusleColor).ToFColor(false);
-							FDebugDrawQueue::GetInstance().DrawDebugLine(WorldPointPlaneLocation, WorldPointPlaneLocation + Settings.DrawScale * Settings.ImpulseScale * ManifoldPoint.NetImpulse, Color, false, KINDA_SMALL_NUMBER, Settings.DrawPriority, Settings.LineThickness);
-						}
+					// Pushout
+					if ((Settings.PushOutScale > 0) && !ManifoldPoint.NetPushOut.IsNearlyZero())
+					{
+						FColor Color = (ColorScale * PushOutImpusleColor).ToFColor(false);
+						FDebugDrawQueue::GetInstance().DrawDebugLine(WorldPointPlaneLocation, WorldPointPlaneLocation + Settings.DrawScale * Settings.PushOutScale * ManifoldPoint.NetPushOut, Color, false, KINDA_SMALL_NUMBER, Settings.DrawPriority, Settings.LineThickness);
+					}
+					if ((Settings.ImpulseScale > 0) && !ManifoldPoint.NetImpulse.IsNearlyZero())
+					{
+						FColor Color = (ColorScale * ImpulseColor).ToFColor(false);
+						FDebugDrawQueue::GetInstance().DrawDebugLine(WorldPointPlaneLocation, WorldPointPlaneLocation + Settings.DrawScale * Settings.ImpulseScale * ManifoldPoint.NetImpulse, Color, false, KINDA_SMALL_NUMBER, Settings.DrawPriority, Settings.LineThickness);
+					}
 
-						// Manifold plane and normal
-						DrawCollisionImpl(WorldPlaneLocation, WorldPlaneNormal, ManifoldPoint.ContactPoint.Phi, ManifoldPoint.NetImpulse, DiscColor, NormalColor, ImpulseColor, ColorScale, Settings);
+					// Manifold plane and normal
+					if (Settings.ContactWidth > 0)
+					{
+						FColor C0 = (ColorScale * DiscColor).ToFColor(false);
+						FDebugDrawQueue::GetInstance().DrawDebugCircle(WorldPlaneLocation, Settings.DrawScale * Settings.ContactWidth, 12, C0, false, KINDA_SMALL_NUMBER, Settings.DrawPriority, Settings.LineThickness, Axes.GetUnitAxis(EAxis::Y), Axes.GetUnitAxis(EAxis::Z), false);
+					}
+					if (Settings.ContactLen > 0)
+					{
+						FColor C1 = (ColorScale * NormalColor).ToFColor(false);
+						FDebugDrawQueue::GetInstance().DrawDebugLine(WorldPlaneLocation, WorldPlaneLocation + Settings.DrawScale * Settings.ContactLen * WorldPlaneNormal, C1, false, KINDA_SMALL_NUMBER, Settings.DrawPriority, Settings.LineThickness);
+					}
+					if (Settings.ContactPhiWidth > 0 && (ManifoldPoint.ContactPoint.Phi < FLT_MAX))
+					{
+						FColor C2 = (ColorScale * FColor(128, 128, 0)).ToFColor(false);
+						FDebugDrawQueue::GetInstance().DrawDebugCircle(WorldPlaneLocation - ManifoldPoint.ContactPoint.Phi * WorldPlaneNormal, Settings.DrawScale * Settings.ContactPhiWidth, 12, C2, false, KINDA_SMALL_NUMBER, Settings.DrawPriority, Settings.LineThickness, Axes.GetUnitAxis(EAxis::Y), Axes.GetUnitAxis(EAxis::Z), false);
+					}
 
-						// Manifold point
-						FMatrix Axes = FRotationMatrix::MakeFromX(WorldPlaneNormal);
-						FDebugDrawQueue::GetInstance().DrawDebugCircle(WorldPointLocation, 0.5f * Settings.DrawScale * Settings.ContactWidth, 12, DiscColor, false, KINDA_SMALL_NUMBER, Settings.DrawPriority, Settings.LineThickness, Axes.GetUnitAxis(EAxis::Y), Axes.GetUnitAxis(EAxis::Z), false);
+					// Manifold point
+					FDebugDrawQueue::GetInstance().DrawDebugCircle(WorldPointLocation, 0.5f * Settings.DrawScale * Settings.ContactWidth, 12, DiscColor, false, KINDA_SMALL_NUMBER, Settings.DrawPriority, Settings.LineThickness, Axes.GetUnitAxis(EAxis::Y), Axes.GetUnitAxis(EAxis::Z), false);
 
-						// Previous points
-						const FVec3 WorldPrevPointLocation = SpaceTransform.TransformPosition(PointTransform.TransformPosition(ManifoldPoint.ShapeAnchorPoints[ContactPointOwner]));
-						const FVec3 WorldPrevPlaneLocation = SpaceTransform.TransformPosition(PlaneTransform.TransformPosition(ManifoldPoint.ShapeAnchorPoints[ContactPlaneOwner]));
-						FDebugDrawQueue::GetInstance().DrawDebugLine(WorldPrevPointLocation, WorldPointLocation, FColor::White, false, KINDA_SMALL_NUMBER, Settings.DrawPriority, Settings.LineThickness);
-						FDebugDrawQueue::GetInstance().DrawDebugLine(WorldPrevPlaneLocation, WorldPlaneLocation, FColor::White, false, KINDA_SMALL_NUMBER, Settings.DrawPriority, Settings.LineThickness);
+					// Previous points
+					const FVec3 WorldPrevPointLocation = SpaceTransform.TransformPosition(PointTransform.TransformPosition(ManifoldPoint.ShapeAnchorPoints[ContactPointOwner]));
+					const FVec3 WorldPrevPlaneLocation = SpaceTransform.TransformPosition(PlaneTransform.TransformPosition(ManifoldPoint.ShapeAnchorPoints[ContactPlaneOwner]));
+					FDebugDrawQueue::GetInstance().DrawDebugLine(WorldPrevPointLocation, WorldPointLocation, FColor::White, false, KINDA_SMALL_NUMBER, Settings.DrawPriority, Settings.LineThickness);
+					FDebugDrawQueue::GetInstance().DrawDebugLine(WorldPrevPlaneLocation, WorldPlaneLocation, FColor::White, false, KINDA_SMALL_NUMBER, Settings.DrawPriority, Settings.LineThickness);
 
-						// Whether restored
+					// Whether restored
+					if (Settings.ContactInfoWidth > 0)
+					{
 						if (Contact.WasManifoldRestored())
 						{
-							const FReal BoxScale = Settings.DrawScale * Settings.ContactWidth;
-							FDebugDrawQueue::GetInstance().DrawDebugBox(WorldPlaneLocation, FVec3(BoxScale, BoxScale, FReal(0.01)), FRotation3(FRotationMatrix::MakeFromZ(WorldPlaneNormal)), FColor::Green, false, KINDA_SMALL_NUMBER, Settings.DrawPriority, 0.5f * Settings.LineThickness);
+							const FReal BoxScale = Settings.DrawScale * Settings.ContactInfoWidth;
+							FDebugDrawQueue::GetInstance().DrawDebugBox(WorldPlaneLocation, FVec3(BoxScale, BoxScale, FReal(0.01)), FRotation3(FRotationMatrix::MakeFromZ(WorldPlaneNormal)), FColor::Blue, false, KINDA_SMALL_NUMBER, Settings.DrawPriority, 0.5f * Settings.LineThickness);
 						}
 						else if (ManifoldPoint.Flags.bWasRestored)
 						{
-							const FReal BoxScale = Settings.DrawScale * Settings.ContactWidth;
-							FDebugDrawQueue::GetInstance().DrawDebugBox(WorldPlaneLocation, FVec3(BoxScale, BoxScale, FReal(0.01)), FRotation3(FRotationMatrix::MakeFromZ(WorldPlaneNormal)), FColor::Yellow, false, KINDA_SMALL_NUMBER, Settings.DrawPriority, 0.5f * Settings.LineThickness);
+							const FReal BoxScale = Settings.DrawScale * Settings.ContactInfoWidth;
+							FDebugDrawQueue::GetInstance().DrawDebugBox(WorldPlaneLocation, FVec3(BoxScale, BoxScale, FReal(0.01)), FRotation3(FRotationMatrix::MakeFromZ(WorldPlaneNormal)), FColor::Purple, false, KINDA_SMALL_NUMBER, Settings.DrawPriority, 0.5f * Settings.LineThickness);
 						}
 						else if (ManifoldPoint.Flags.bWasReplaced)
 						{
-							const FReal BoxScale = Settings.DrawScale * Settings.ContactWidth;
+							const FReal BoxScale = Settings.DrawScale * Settings.ContactInfoWidth;
 							FDebugDrawQueue::GetInstance().DrawDebugBox(WorldPlaneLocation, FVec3(BoxScale, BoxScale, FReal(0.01)), FRotation3(FRotationMatrix::MakeFromZ(WorldPlaneNormal)), FColor::Orange, false, KINDA_SMALL_NUMBER, Settings.DrawPriority, 0.5f * Settings.LineThickness);
 						}
 					}
 				}
-				else
-				{
-					const FVec3 Location = SpaceTransform.TransformPosition(Contact.CalculateWorldContactLocation());
-					const FVec3 Normal = SpaceTransform.TransformVector(Contact.CalculateWorldContactNormal());
-					DrawCollisionImpl(Location, Normal, Contact.GetPhi(), FVec3(0), FColor(200, 0, 0), FColor(200, 0, 0), FColor(200, 0, 0), ColorScale, Settings);
-				}
 			}
 			if (Settings.ContactOwnerWidth > 0)
 			{
diff --git a/Engine/Source/Runtime/Experimental/Chaos/Private/PBDRigidsSolver.cpp b/Engine/Source/Runtime/Experimental/Chaos/Private/PBDRigidsSolver.cpp
--- a/Engine/Source/Runtime/Experimental/Chaos/Private/PBDRigidsSolver.cpp
+++ b/Engine/Source/Runtime/Experimental/Chaos/Private/PBDRigidsSolver.cpp
@@ -47,6 +47,7 @@ DECLARE_DWORD_COUNTER_STAT(TEXT("NumRestoredConstraints"), STAT_ChaosCounter_Num
 DECLARE_DWORD_COUNTER_STAT(TEXT("NumManifoldPoints"), STAT_ChaosCounter_NumManifoldPoints, STATGROUP_ChaosCounters);
 DECLARE_DWORD_COUNTER_STAT(TEXT("NumActiveManifoldPoints"), STAT_ChaosCounter_NumActiveManifoldPoints, STATGROUP_ChaosCounters);
 DECLARE_DWORD_COUNTER_STAT(TEXT("NumRestoredManifoldPoints"), STAT_ChaosCounter_NumRestoredManifoldPoints, STATGROUP_ChaosCounters);
+DECLARE_DWORD_COUNTER_STAT(TEXT("NumUpdatedManifoldPoints"), STAT_ChaosCounter_NumUpdatedManifoldPoints, STATGROUP_ChaosCounters);
 DECLARE_DWORD_COUNTER_STAT(TEXT("NumJoints"), STAT_ChaosCounter_NumJoints, STATGROUP_ChaosCounters);
 
 // Stat Iteration counters
@@ -115,6 +116,7 @@ namespace Chaos
 			/* ContactLen =					*/ 30.0f,
 			/* ContactWidth =				*/ 6.0f,
 			/* ContactPhiWidth =			*/ 0.0f,
+			/* ContactInfoWidth				*/ 6.0f,
 			/* ContactOwnerWidth =			*/ 0.0f,
 			/* ConstraintAxisLen =			*/ 30.0f,
 			/* JointComSize =				*/ 2.0f,
@@ -143,6 +145,7 @@ namespace Chaos
 		FAutoConsoleVariableRef CVarChaosSolverBodyAxisLen(TEXT("p.Chaos.Solver.DebugDraw.BodyAxisLen"), ChaosSolverDebugDebugDrawSettings.BodyAxisLen, TEXT("BodyAxisLen."));
 		FAutoConsoleVariableRef CVarChaosSolverContactLen(TEXT("p.Chaos.Solver.DebugDraw.ContactLen"), ChaosSolverDebugDebugDrawSettings.ContactLen, TEXT("ContactLen."));
 		FAutoConsoleVariableRef CVarChaosSolverContactWidth(TEXT("p.Chaos.Solver.DebugDraw.ContactWidth"), ChaosSolverDebugDebugDrawSettings.ContactWidth, TEXT("ContactWidth."));
+		FAutoConsoleVariableRef CVarChaosSolverContactInfoWidth(TEXT("p.Chaos.Solver.DebugDraw.ContactInfoWidth"), ChaosSolverDebugDebugDrawSettings.ContactInfoWidth, TEXT("ContactInfoWidth."));
 		FAutoConsoleVariableRef CVarChaosSolverContactPhiWidth(TEXT("p.Chaos.Solver.DebugDraw.ContactPhiWidth"), ChaosSolverDebugDebugDrawSettings.ContactPhiWidth, TEXT("ContactPhiWidth."));
 		FAutoConsoleVariableRef CVarChaosSolverContactOwnerWidth(TEXT("p.Chaos.Solver.DebugDraw.ContactOwnerWidth"), ChaosSolverDebugDebugDrawSettings.ContactOwnerWidth, TEXT("ContactOwnerWidth."));
 		FAutoConsoleVariableRef CVarChaosSolverConstraintAxisLen(TEXT("p.Chaos.Solver.DebugDraw.ConstraintAxisLen"), ChaosSolverDebugDebugDrawSettings.ConstraintAxisLen, TEXT("ConstraintAxisLen."));
@@ -1435,6 +1438,7 @@ CSV_CUSTOM_STAT(PhysicsCounters, Name, Value, ECsvCustomStatOp::Set);
 		int32 NumManifoldPoints = 0;
 		int32 NumActiveManifoldPoints = 0;
 		int32 NumRestoredManifoldPoints = 0;
+		int32 NumUpdatedManifoldPoints = 0;
 		for (const FPBDCollisionConstraintHandle* Collision : GetEvolution()->GetCollisionConstraints().GetConstraints())
 		{
 			if (Collision->GetContact().IsEnabled())
@@ -1458,6 +1462,10 @@ CSV_CUSTOM_STAT(PhysicsCounters, Name, Value, ECsvCustomStatOp::Set);
 					{
 						++NumRestoredManifoldPoints;
 					}
+					if (ManifoldPoint.Flags.bWasReplaced)
+					{
+						++NumUpdatedManifoldPoints;
+					}
 					if (!ManifoldPoint.NetPushOut.IsNearlyZero())
 					{
 						++NumActiveManifoldPoints;
@@ -1473,6 +1481,7 @@ CSV_CUSTOM_STAT(PhysicsCounters, Name, Value, ECsvCustomStatOp::Set);
 		CHAOS_COUNTER_STAT(NumManifoldPoints, NumManifoldPoints);
 		CHAOS_COUNTER_STAT(NumActiveManifoldPoints, NumActiveManifoldPoints);
 		CHAOS_COUNTER_STAT(NumRestoredManifoldPoints, NumRestoredManifoldPoints);
+		CHAOS_COUNTER_STAT(NumUpdatedManifoldPoints, NumUpdatedManifoldPoints);
 	}
 
 	void FPBDRigidsSolver::PostTickDebugDraw(FReal Dt) const
diff --git a/Engine/Source/Runtime/Experimental/Chaos/Public/Chaos/ChaosDebugDraw.h b/Engine/Source/Runtime/Experimental/Chaos/Public/Chaos/ChaosDebugDraw.h
--- a/Engine/Source/Runtime/Experimental/Chaos/Public/Chaos/ChaosDebugDraw.h
+++ b/Engine/Source/Runtime/Experimental/Chaos/Public/Chaos/ChaosDebugDraw.h
@@ -65,6 +65,7 @@ namespace Chaos
 				FRealSingle InContactLen,
 				FRealSingle InContactWidth,
 				FRealSingle InContactPhiWidth,
+				FRealSingle InContactInfoWidth,
 				FRealSingle InContactOwnerWidth,
 				FRealSingle InConstraintAxisLen,
 				FRealSingle InJointComSize,
@@ -93,6 +94,7 @@ namespace Chaos
 				, ContactLen(InContactLen)
 				, ContactWidth(InContactWidth)
 				, ContactPhiWidth(InContactPhiWidth)
+				, ContactInfoWidth(InContactInfoWidth)
 				, ContactOwnerWidth(InContactOwnerWidth)
 				, ConstraintAxisLen(InConstraintAxisLen)
 				, JointComSize(InJointComSize)
@@ -123,6 +125,7 @@ namespace Chaos
 			FRealSingle ContactLen;
 			FRealSingle ContactWidth;
 			FRealSingle ContactPhiWidth;
+			FRealSingle ContactInfoWidth;
 			FRealSingle ContactOwnerWidth;
 			FRealSingle ConstraintAxisLen;
 			FRealSingle JointComSize;
