Chaos - disable manifold reuse for convex vs non-convex collision

diff --git a/Engine/Source/Runtime/Experimental/Chaos/Private/Chaos/Collision/ContactPointsMiscShapes.cpp b/Engine/Source/Runtime/Experimental/Chaos/Private/Chaos/Collision/ContactPointsMiscShapes.cpp
--- a/Engine/Source/Runtime/Experimental/Chaos/Private/Chaos/Collision/ContactPointsMiscShapes.cpp
+++ b/Engine/Source/Runtime/Experimental/Chaos/Private/Chaos/Collision/ContactPointsMiscShapes.cpp
@@ -140,14 +140,6 @@ namespace Chaos
 		{
 			B.ContactManifold(*ConvexImplicit, AToBTM, CullDistance, ContactPoints);
 		}
-
-		// Note: All results we got are in ShapeB's local space
-		for (FContactPoint& ContactPoint : ContactPoints)
-		{
-			// Calculate the shape contact points
-			ContactPoint.ShapeContactPoints[0] = AToBTM.InverseTransformPosition(ContactPoint.ShapeContactPoints[1]);
-			ContactPoint.ShapeContactPoints[1] = ContactPoint.ShapeContactPoints[1] - ContactPoint.Phi * ContactPoint.ShapeContactNormal;
-		}
 	}
 
 	// A is the implicit here, we want to return a contact point on B (trimesh)
diff --git a/Engine/Source/Runtime/Experimental/Chaos/Private/Chaos/Collision/PBDCollisionConstraint.cpp b/Engine/Source/Runtime/Experimental/Chaos/Private/Chaos/Collision/PBDCollisionConstraint.cpp
--- a/Engine/Source/Runtime/Experimental/Chaos/Private/Chaos/Collision/PBDCollisionConstraint.cpp
+++ b/Engine/Source/Runtime/Experimental/Chaos/Private/Chaos/Collision/PBDCollisionConstraint.cpp
@@ -41,6 +41,12 @@ namespace Chaos
 	bool bChaos_Manifold_EnableFrictionRestore = true;
 	FAutoConsoleVariableRef CVarChaos_Manifold_EnableFrictionRestore(TEXT("p.Chaos.Collision.Manifold.EnableFrictionRestore"), bChaos_Manifold_EnableFrictionRestore, TEXT(""));
 	
+	// The margin to use when we are colliding a convex shape against a zero-margin shape. E.g., Box-Triangle.
+	// When both shapes have a margin we use the minimum margin, but we don't want to use a zero margin because we hit the EPA degenerate case
+	// NOTE: This is currently disabled - margins for convex-trimesh cause bigger problems than the EPA issue
+	FRealSingle Chaos_Collision_ConvexZeroMargin = 0.0f;
+	FAutoConsoleVariableRef CVarChaos_Collision_ConvexZeroMargin(TEXT("p.Chaos.Collision.ConvexZeroMargin"), Chaos_Collision_ConvexZeroMargin, TEXT(""));
+
 	struct FCollisionTolerances
 	{
 		// Multiplied by the contact margin to produce a distance within which contacts are considered to be the same point
@@ -260,7 +266,7 @@ namespace Chaos
 		// One shape in a collision will always have a margin. Only triangles have zero margin and we don't 
 		// collide two triangles. If we have a triangle, it is always the second shape.
 		// The collision tolerance is used for knowing whether a new contact matches an existing one.
-		// If we have two non-quadratic shapes, we use the smallest margin on both shapes.
+		// If we have two non-quadratic shapes, we use the smallest margin on both shapes (unless one shape has zero margin).
 		// If we have a quadratic shape versus a non-quadratic, we don't need a margin on the non-quadratic.
 		// For non-quadratics the collision tolerance is the smallest non-zero margin. 
 		// For quadratic shapes we want a collision tolerance much smaller than the radius.
@@ -269,9 +275,23 @@ namespace Chaos
 		const FReal QuadraticToleranceScale = 0.05f;
 		if (!bIsQuadratic0 && !bIsQuadratic1)
 		{
-			CollisionMargins[0] = FMath::Min(Margin0, Margin1);
-			CollisionMargins[1] = CollisionMargins[0];
-			CollisionTolerance = ((Margin0 < Margin1) || (Margin1 == 0)) ? Margin0 : Margin1;
+			// If one shape has a zero margin, enforce a minimum margin to avoid the EPA degenerate case. E.g., Box-Triangle
+			// If both shapes have a margin, use the smaller margin on both shapes. E.g., Box-Box
+			// We should never see both shapes with zero margin, but if we did we'd end up with a zero margin
+			if ((Margin0 == FReal(0)) || (Margin1 == FReal(0)))
+			{
+				const FReal MinMargin = Chaos_Collision_ConvexZeroMargin;
+				const FReal NonZeroMargin = FMath::Max(Margin0, Margin1);
+				CollisionMargins[0] = FMath::Min(MinMargin, NonZeroMargin);
+				CollisionMargins[1] = CollisionMargins[0];
+				CollisionTolerance = NonZeroMargin;
+			}
+			else
+			{
+				CollisionMargins[0] = FMath::Min(Margin0, Margin1);
+				CollisionMargins[1] = CollisionMargins[0];
+				CollisionTolerance = CollisionMargins[0];
+			}
 		}
 		else if (bIsQuadratic0 && bIsQuadratic1)
 		{
diff --git a/Engine/Source/Runtime/Experimental/Chaos/Private/Chaos/Collision/ParticlePairMidPhase.cpp b/Engine/Source/Runtime/Experimental/Chaos/Private/Chaos/Collision/ParticlePairMidPhase.cpp
--- a/Engine/Source/Runtime/Experimental/Chaos/Private/Chaos/Collision/ParticlePairMidPhase.cpp
+++ b/Engine/Source/Runtime/Experimental/Chaos/Private/Chaos/Collision/ParticlePairMidPhase.cpp
@@ -119,6 +119,8 @@ namespace Chaos
 		const bool bIsSphere1 = (ImplicitType1 == ImplicitObjectType::Sphere);
 		const bool bIsCapsule0 = (ImplicitType0 == ImplicitObjectType::Capsule);
 		const bool bIsCapsule1 = (ImplicitType1 == ImplicitObjectType::Capsule);
+		const bool bIsTriangle0 = (ImplicitType0 == ImplicitObjectType::TriangleMesh) || (ImplicitType0 == ImplicitObjectType::HeightField);
+		const bool bIsTriangle1 = (ImplicitType1 == ImplicitObjectType::TriangleMesh) || (ImplicitType1 == ImplicitObjectType::HeightField);
 
 		const bool bAllowBoundsChecked = bChaos_Collision_MidPhase_EnableBoundsChecks && bHasBounds0 && bHasBounds1;
 		Flags.bEnableAABBCheck = bAllowBoundsChecked && !(bIsSphere0 && bIsSphere1);	// No AABB test if both are spheres
@@ -131,7 +133,7 @@ namespace Chaos
 		}
 
 		// Do not try to reuse manifold points for capsules or spheres (against anything)
-		Flags.bEnableManifoldUpdate = bChaos_Collision_EnableManifoldUpdate && !bIsSphere0 && !bIsSphere1 && !bIsCapsule0 && !bIsCapsule1;
+		Flags.bEnableManifoldUpdate = bChaos_Collision_EnableManifoldUpdate && !bIsSphere0 && !bIsSphere1 && !bIsCapsule0 && !bIsCapsule1 && !bIsTriangle0 && !bIsTriangle1;
 	}
 
 	FSingleShapePairCollisionDetector::~FSingleShapePairCollisionDetector()
@@ -277,7 +279,7 @@ namespace Chaos
 
 			// If the constraint was not used last frame, it needs to be reset. 
 			// Otherwise we will try to reuse it below
-			if (!bWasUpdatedLastTick || bUseCCD)
+			if (!bWasUpdatedLastTick  || (Constraint->GetManifoldPoints().Num() == 0) || bUseCCD)
 			{
 				// Clear all manifold data including saved contact data
 				Constraint->ResetManifold();
