Interchange code cleanup, move some header code into cpp file

diff --git a/Engine/Plugins/Enterprise/DatasmithImporter/Source/DatasmithImporter/Private/DatasmithTextureImporter.cpp b/Engine/Plugins/Enterprise/DatasmithImporter/Source/DatasmithImporter/Private/DatasmithTextureImporter.cpp
--- a/Engine/Plugins/Enterprise/DatasmithImporter/Source/DatasmithImporter/Private/DatasmithTextureImporter.cpp
+++ b/Engine/Plugins/Enterprise/DatasmithImporter/Source/DatasmithImporter/Private/DatasmithTextureImporter.cpp
@@ -409,11 +409,11 @@ UE::Interchange::FAssetImportResultRef FDatasmithTextureImporter::CreateTextureA
 	return UInterchangeManager::GetInterchangeManager().ImportAssetAsync( ContentPath, ScopedSourceData.GetSourceData(), ImportAssetParameters );
 }
 
-bool UDatasmithTexturePipeline::ExecutePreImportPipeline(UInterchangeBaseNodeContainer* BaseNodeContainer, const TArray<UInterchangeSourceData*>& SourceDatas)
+void UDatasmithTexturePipeline::ExecutePreImportPipeline(UInterchangeBaseNodeContainer* BaseNodeContainer, const TArray<UInterchangeSourceData*>& SourceDatas)
 {
 	if ( !TextureElement.IsValid() )
 	{
-		return false;
+		return;
 	}
 
 	TArray< FString > Nodes;
@@ -421,13 +421,13 @@ bool UDatasmithTexturePipeline::ExecutePreImportPipeline(UInterchangeBaseNodeCon
 
 	if ( Nodes.Num() <= 0)
 	{
-		return false;
+		return;
 	}
 
 	UInterchangeTextureNode* TextureTranslatedNode = Cast<UInterchangeTextureNode >( BaseNodeContainer->GetNode( Nodes[0] ) );
 	if (!ensure(TextureTranslatedNode))
 	{
-		return false;
+		return;
 	}
 
 	FString DisplayLabel = TextureTranslatedNode->GetDisplayLabel();
@@ -435,7 +435,7 @@ bool UDatasmithTexturePipeline::ExecutePreImportPipeline(UInterchangeBaseNodeCon
 	UInterchangeTextureFactoryNode* TextureFactoryNode = NewObject<UInterchangeTextureFactoryNode>(BaseNodeContainer, NAME_None);
 	if (!ensure(TextureFactoryNode))
 	{
-		return false;
+		return;
 	}
 	//Creating a UTexture2D
 	TextureFactoryNode->InitializeTextureNode(NodeUID, DisplayLabel, UTexture2D::StaticClass()->GetName(), TextureTranslatedNode->GetDisplayLabel());
@@ -549,8 +549,6 @@ bool UDatasmithTexturePipeline::ExecutePreImportPipeline(UInterchangeBaseNodeCon
 	{
 		TextureFactoryNode->SetCustomFilter( TexFilter.GetValue() );
 	}
-
-	return true;
 }
 
 #undef LOCTEXT_NAMESPACE
diff --git a/Engine/Plugins/Enterprise/DatasmithImporter/Source/DatasmithImporter/Public/DatasmithTextureImporter.h b/Engine/Plugins/Enterprise/DatasmithImporter/Source/DatasmithImporter/Public/DatasmithTextureImporter.h
--- a/Engine/Plugins/Enterprise/DatasmithImporter/Source/DatasmithImporter/Public/DatasmithTextureImporter.h
+++ b/Engine/Plugins/Enterprise/DatasmithImporter/Source/DatasmithImporter/Public/DatasmithTextureImporter.h
@@ -46,7 +46,7 @@ class UDatasmithTexturePipeline : public UInterchangePipelineBase
 	GENERATED_BODY()
 
 public:
-	virtual bool ExecutePreImportPipeline(UInterchangeBaseNodeContainer* BaseNodeContainer, const TArray<UInterchangeSourceData*>& SourceDatas) override;
+	virtual void ExecutePreImportPipeline(UInterchangeBaseNodeContainer* BaseNodeContainer, const TArray<UInterchangeSourceData*>& SourceDatas) override;
 
 	bool IsScripted() override
 	{
diff --git a/Engine/Plugins/Experimental/Interchange/Editor/Source/Pipelines/Private/InterchangeGraphInspectorPipeline.cpp b/Engine/Plugins/Experimental/Interchange/Editor/Source/Pipelines/Private/InterchangeGraphInspectorPipeline.cpp
--- a/Engine/Plugins/Experimental/Interchange/Editor/Source/Pipelines/Private/InterchangeGraphInspectorPipeline.cpp
+++ b/Engine/Plugins/Experimental/Interchange/Editor/Source/Pipelines/Private/InterchangeGraphInspectorPipeline.cpp
@@ -13,7 +13,7 @@
 #include "UObject/ObjectMacros.h"
 #include "Widgets/SWindow.h"
 
-bool UInterchangeGraphInspectorPipeline::ExecutePreImportPipeline(UInterchangeBaseNodeContainer* BaseNodeContainer, const TArray<UInterchangeSourceData*>& SourceDatas)
+void UInterchangeGraphInspectorPipeline::ExecutePreImportPipeline(UInterchangeBaseNodeContainer* BaseNodeContainer, const TArray<UInterchangeSourceData*>& SourceDatas)
 {
 	//Create and show the graph inspector UI dialog
 	TSharedPtr<SWindow> ParentWindow;
@@ -35,5 +35,4 @@ bool UInterchangeGraphInspectorPipeline::ExecutePreImportPipeline(UInterchangeBa
 	);
 
 	FSlateApplication::Get().AddModalWindow(Window, ParentWindow, false);
-	return true;
 }
diff --git a/Engine/Plugins/Experimental/Interchange/Editor/Source/Pipelines/Public/InterchangeGraphInspectorPipeline.h b/Engine/Plugins/Experimental/Interchange/Editor/Source/Pipelines/Public/InterchangeGraphInspectorPipeline.h
--- a/Engine/Plugins/Experimental/Interchange/Editor/Source/Pipelines/Public/InterchangeGraphInspectorPipeline.h
+++ b/Engine/Plugins/Experimental/Interchange/Editor/Source/Pipelines/Public/InterchangeGraphInspectorPipeline.h
@@ -35,7 +35,7 @@ public:
 
 protected:
 
-	virtual bool ExecutePreImportPipeline(UInterchangeBaseNodeContainer* BaseNodeContainer, const TArray<UInterchangeSourceData*>& SourceDatas) override;
+	virtual void ExecutePreImportPipeline(UInterchangeBaseNodeContainer* BaseNodeContainer, const TArray<UInterchangeSourceData*>& SourceDatas) override;
 
 	//virtual bool ExecutePostImportPipeline(const UInterchangeBaseNodeContainer* BaseNodeContainer, const FName& NodeKey, UObject* CreatedAsset) override;
 
diff --git a/Engine/Plugins/Experimental/Interchange/Runtime/Source/Import/Private/Texture/InterchangeUDIMTranslator.cpp b/Engine/Plugins/Experimental/Interchange/Runtime/Source/Import/Private/Texture/InterchangeUDIMTranslator.cpp
--- a/Engine/Plugins/Experimental/Interchange/Runtime/Source/Import/Private/Texture/InterchangeUDIMTranslator.cpp
+++ b/Engine/Plugins/Experimental/Interchange/Runtime/Source/Import/Private/Texture/InterchangeUDIMTranslator.cpp
@@ -181,7 +181,7 @@ bool UInterchangeUDIMTranslator::Translate(UInterchangeBaseNodeContainer& BaseNo
 
 	FString DisplayLabel = PreUDIMName + PostUDIMName;
 
-	TextureNode->InitializeNode(Filename, DisplayLabel, EInterchangeNodeContainerType::NodeContainerType_TranslatedAsset);
+	TextureNode->InitializeNode(Filename, DisplayLabel, EInterchangeNodeContainerType::TranslatedAsset);
 	TextureNode->SetPayLoadKey(Filename);
 
 	UInterchangeManager& InterchangeManager = UInterchangeManager::GetInterchangeManager();
diff --git a/Engine/Plugins/Experimental/Interchange/Runtime/Source/Import/Private/Texture/TextureTranslatorUtilities.cpp b/Engine/Plugins/Experimental/Interchange/Runtime/Source/Import/Private/Texture/TextureTranslatorUtilities.cpp
--- a/Engine/Plugins/Experimental/Interchange/Runtime/Source/Import/Private/Texture/TextureTranslatorUtilities.cpp
+++ b/Engine/Plugins/Experimental/Interchange/Runtime/Source/Import/Private/Texture/TextureTranslatorUtilities.cpp
@@ -39,7 +39,7 @@ namespace UE::Interchange::Private
 			return false;
 		}
 
-		TextureNode->InitializeNode(NodeUID, DisplayLabel, EInterchangeNodeContainerType::NodeContainerType_TranslatedAsset);
+		TextureNode->InitializeNode(NodeUID, DisplayLabel, EInterchangeNodeContainerType::TranslatedAsset);
 		TextureNode->SetPayLoadKey(Filename);
 
 		BaseNodeContainer.AddNode(TextureNode);
diff --git a/Engine/Plugins/Experimental/Interchange/Runtime/Source/Nodes/Private/InterchangeMaterialFactoryNode.cpp b/Engine/Plugins/Experimental/Interchange/Runtime/Source/Nodes/Private/InterchangeMaterialFactoryNode.cpp
new file mode 100644
--- /dev/null
+++ b/Engine/Plugins/Experimental/Interchange/Runtime/Source/Nodes/Private/InterchangeMaterialFactoryNode.cpp
@@ -0,0 +1,237 @@
+// Copyright Epic Games, Inc. All Rights Reserved.
+
+#include "InterchangeMaterialFactoryNode.h"
+
+#if WITH_ENGINE
+
+#include "Materials/Material.h"
+#include "Materials/MaterialInstance.h"
+#include "Materials/MaterialInterface.h"
+
+#endif
+
+
+UInterchangeMaterialFactoryNode::UInterchangeMaterialFactoryNode()
+{
+#if WITH_ENGINE
+	AssetClass = nullptr;
+#endif
+	TextureDependencies.Initialize(Attributes, TextureDependenciesKey.Key);
+}
+
+void UInterchangeMaterialFactoryNode::InitializeMaterialNode(const FString& UniqueID, const FString& DisplayLabel, const FString& InAssetClass)
+{
+	bIsMaterialNodeClassInitialized = false;
+	InitializeNode(UniqueID, DisplayLabel, EInterchangeNodeContainerType::FactoryData);
+	FString OperationName = GetTypeName() + TEXT(".SetAssetClassName");
+	InterchangePrivateNodeBase::SetCustomAttribute<FString>(*Attributes, ClassNameAttributeKey, OperationName, InAssetClass);
+	FillAssetClassFromAttribute();
+}
+
+FString UInterchangeMaterialFactoryNode::GetTypeName() const
+{
+	const FString TypeName = TEXT("MaterialFactoryNode");
+	return TypeName;
+}
+
+UClass* UInterchangeMaterialFactoryNode::GetObjectClass() const
+{
+	ensure(bIsMaterialNodeClassInitialized);
+#if WITH_ENGINE
+	return AssetClass.Get() != nullptr ? AssetClass.Get() : UMaterial::StaticClass();
+#else
+	return nullptr;
+#endif
+}
+
+FString UInterchangeMaterialFactoryNode::GetKeyDisplayName(const UE::Interchange::FAttributeKey& NodeAttributeKey) const
+{
+	FString KeyDisplayName = NodeAttributeKey.ToString();
+	if (NodeAttributeKey == TextureDependenciesKey)
+	{
+		return KeyDisplayName = TEXT("Texture Dependencies count");
+	}
+	else if (NodeAttributeKey.Key.StartsWith(TextureDependenciesKey.Key))
+	{
+		KeyDisplayName = TEXT("Texture Dependencies Index ");
+		const FString IndexKey = UE::Interchange::TArrayAttributeHelper<FString>::IndexKey();
+		int32 IndexPosition = NodeAttributeKey.Key.Find(IndexKey) + IndexKey.Len();
+		if (IndexPosition < NodeAttributeKey.Key.Len())
+		{
+			KeyDisplayName += NodeAttributeKey.Key.RightChop(IndexPosition);
+		}
+		return KeyDisplayName;
+	}
+	return Super::GetKeyDisplayName(NodeAttributeKey);
+}
+
+FString UInterchangeMaterialFactoryNode::GetAttributeCategory(const UE::Interchange::FAttributeKey& NodeAttributeKey) const
+{
+	if (NodeAttributeKey == Macro_CustomTranslatedMaterialNodeUidKey
+		|| NodeAttributeKey == Macro_CustomMaterialUsageKey
+		|| NodeAttributeKey == Macro_CustomBlendModeKey)
+	{
+		return FString(TEXT("MaterialFactory"));
+	}
+	else if (NodeAttributeKey.Key.StartsWith(TextureDependenciesKey.Key))
+	{
+		return FString(TEXT("TextureDependencies"));
+	}
+	return Super::GetAttributeCategory(NodeAttributeKey);
+}
+
+FString UInterchangeMaterialFactoryNode::GetMaterialFactoryNodeUidFromMaterialNodeUid(const FString& TranslatedNodeUid)
+{
+	FString NewUid = TEXT("Factory_") + TranslatedNodeUid;
+	return NewUid;
+}
+
+int32 UInterchangeMaterialFactoryNode::GetTextureDependeciesCount() const
+{
+	return TextureDependencies.GetCount();
+}
+
+void UInterchangeMaterialFactoryNode::GetTextureDependencies(TArray<FString>& OutDependencies) const
+{
+	TextureDependencies.GetItems(OutDependencies);
+}
+
+void UInterchangeMaterialFactoryNode::GetTextureDependency(const int32 Index, FString& OutDependency) const
+{
+	TextureDependencies.GetItem(Index, OutDependency);
+}
+
+bool UInterchangeMaterialFactoryNode::SetTextureDependencyUid(const FString& DependencyUid)
+{
+	return TextureDependencies.AddItem(DependencyUid);
+}
+
+bool UInterchangeMaterialFactoryNode::RemoveTextureDependencyUid(const FString& DependencyUid)
+{
+	return TextureDependencies.RemoveItem(DependencyUid);
+}
+
+bool UInterchangeMaterialFactoryNode::GetCustomTranslatedMaterialNodeUid(FString& AttributeValue) const
+{
+	IMPLEMENT_NODE_ATTRIBUTE_GETTER(TranslatedMaterialNodeUid, FString);
+}
+
+bool UInterchangeMaterialFactoryNode::SetCustomTranslatedMaterialNodeUid(const FString& AttributeValue)
+{
+	IMPLEMENT_NODE_ATTRIBUTE_SETTER_NODELEGATE(TranslatedMaterialNodeUid, FString);
+}
+
+bool UInterchangeMaterialFactoryNode::GetCustomMaterialUsage(uint8& AttributeValue) const
+{
+	IMPLEMENT_NODE_ATTRIBUTE_GETTER(MaterialUsage, uint8);
+}
+
+bool UInterchangeMaterialFactoryNode::SetCustomMaterialUsage(const uint8& AttributeValue, bool bAddApplyDelegate /*= true*/)
+{
+	IMPLEMENT_NODE_ATTRIBUTE_SETTER(UInterchangeMaterialFactoryNode, MaterialUsage, uint8, UMaterialInterface);
+}
+
+bool UInterchangeMaterialFactoryNode::GetCustomBlendMode(uint8& AttributeValue) const
+{
+	IMPLEMENT_NODE_ATTRIBUTE_GETTER(BlendMode, uint8);
+}
+
+bool UInterchangeMaterialFactoryNode::SetCustomBlendMode(const uint8& AttributeValue, bool bAddApplyDelegate /*= true*/)
+{
+	IMPLEMENT_NODE_ATTRIBUTE_SETTER(UInterchangeMaterialFactoryNode, BlendMode, uint8, UMaterialInterface);
+}
+
+void UInterchangeMaterialFactoryNode::Serialize(FArchive& Ar)
+{
+	Super::Serialize(Ar);
+#if WITH_ENGINE
+	if (Ar.IsLoading())
+	{
+		//Make sure the class is properly set when we compile with engine, this will set the
+		//bIsMaterialNodeClassInitialized to true.
+		SetMaterialNodeClassFromClassAttribute();
+	}
+#endif //#if WITH_ENGINE
+}
+
+void UInterchangeMaterialFactoryNode::FillAssetClassFromAttribute()
+{
+#if WITH_ENGINE
+	FString OperationName = GetTypeName() + TEXT(".GetAssetClassName");
+	FString ClassName;
+	InterchangePrivateNodeBase::GetCustomAttribute<FString>(*Attributes, ClassNameAttributeKey, OperationName, ClassName);
+	if (ClassName.Equals(UMaterial::StaticClass()->GetName()))
+	{
+		AssetClass = UMaterial::StaticClass();
+		bIsMaterialNodeClassInitialized = true;
+	}
+	else if (ClassName.Equals(UMaterialInstance::StaticClass()->GetName()))
+	{
+		AssetClass = UMaterialInstance::StaticClass();
+		bIsMaterialNodeClassInitialized = true;
+	}
+#endif
+}
+
+bool UInterchangeMaterialFactoryNode::SetMaterialNodeClassFromClassAttribute()
+{
+	if (!bIsMaterialNodeClassInitialized)
+	{
+		FillAssetClassFromAttribute();
+	}
+	return bIsMaterialNodeClassInitialized;
+}
+
+#if WITH_ENGINE
+
+bool UInterchangeMaterialFactoryNode::ApplyCustomMaterialUsageToAsset(UObject* Asset) const
+{
+	if (!Asset)
+	{
+		return false;
+	}
+	UMaterialInterface* TypedObject = Cast<UMaterialInterface>(Asset);
+	if (!TypedObject)
+	{
+		return false;
+	}
+	uint8 ValueData;
+	if (GetCustomMaterialUsage(ValueData))
+	{
+		TypedObject->CheckMaterialUsage_Concurrent(EMaterialUsage(ValueData));
+		return true;
+	}
+	return false;
+}
+
+bool UInterchangeMaterialFactoryNode::FillCustomMaterialUsageFromAsset(UObject* Asset)
+{
+	if (!Asset)
+	{
+		return false;
+	}
+	UMaterialInterface* TypedObject = Cast<UMaterialInterface>(Asset);
+	if (!TypedObject)
+	{
+		return false;
+	}
+	const UMaterial* Material = TypedObject->GetMaterial_Concurrent();
+	if (!Material)
+	{
+		return false;
+	}
+	for (int32 Usage = 0; Usage < MATUSAGE_MAX; Usage++)
+	{
+		const EMaterialUsage UsageEnum = (EMaterialUsage)Usage;
+		if (Material->GetUsageByFlag(UsageEnum))
+		{
+			if (SetCustomMaterialUsage((int8)UsageEnum, false))
+			{
+				return true;
+			}
+		}
+	}
+	return false;
+}
+
+#endif //WITH_ENGINE
\ No newline at end of file
diff --git a/Engine/Plugins/Experimental/Interchange/Runtime/Source/Nodes/Private/InterchangeMaterialNode.cpp b/Engine/Plugins/Experimental/Interchange/Runtime/Source/Nodes/Private/InterchangeMaterialNode.cpp
new file mode 100644
--- /dev/null
+++ b/Engine/Plugins/Experimental/Interchange/Runtime/Source/Nodes/Private/InterchangeMaterialNode.cpp
@@ -0,0 +1,213 @@
+// Copyright Epic Games, Inc. All Rights Reserved.
+
+#include "InterchangeMaterialNode.h"
+
+UInterchangeMaterialNode::UInterchangeMaterialNode()
+{
+	TextureDependencies.Initialize(Attributes, TextureDependenciesKey.Key);
+}
+
+/**
+ * Return the node type name of the class, we use this when reporting error
+ */
+FString UInterchangeMaterialNode::GetTypeName() const
+{
+	const FString TypeName = TEXT("MaterialNode");
+	return TypeName;
+}
+
+FString UInterchangeMaterialNode::GetKeyDisplayName(const UE::Interchange::FAttributeKey& NodeAttributeKey) const
+{
+	FString KeyDisplayName = NodeAttributeKey.ToString();
+	if (NodeAttributeKey == UE::Interchange::FMaterialNodeStaticData::PayloadSourceFileKey())
+	{
+		return KeyDisplayName = TEXT("Payload Source Key");
+	}
+	else if (NodeAttributeKey == TextureDependenciesKey)
+	{
+		return KeyDisplayName = TEXT("Texture Dependencies count");
+	}
+	else if (NodeAttributeKey.Key.StartsWith(TextureDependenciesKey.Key))
+	{
+		KeyDisplayName = TEXT("Texture Dependencies Index ");
+		const FString IndexKey = UE::Interchange::TArrayAttributeHelper<FString>::IndexKey();
+		int32 IndexPosition = NodeAttributeKey.Key.Find(IndexKey) + IndexKey.Len();
+		if (IndexPosition < NodeAttributeKey.Key.Len())
+		{
+			KeyDisplayName += NodeAttributeKey.Key.RightChop(IndexPosition);
+		}
+		return KeyDisplayName;
+	}
+	return Super::GetKeyDisplayName(NodeAttributeKey);
+}
+
+FString UInterchangeMaterialNode::GetAttributeCategory(const UE::Interchange::FAttributeKey& NodeAttributeKey) const
+{
+	if (NodeAttributeKey.Key.StartsWith(TextureDependenciesKey.Key))
+	{
+		return FString(TEXT("TextureDependencies"));
+	}
+	return Super::GetAttributeCategory(NodeAttributeKey);
+}
+
+int32 UInterchangeMaterialNode::GetTextureDependeciesCount() const
+{
+	return TextureDependencies.GetCount();
+}
+
+void UInterchangeMaterialNode::GetTextureDependencies(TArray<FString>& OutDependencies) const
+{
+	TextureDependencies.GetItems(OutDependencies);
+}
+
+void UInterchangeMaterialNode::GetTextureDependency(const int32 Index, FString& OutDependency) const
+{
+	TextureDependencies.GetItem(Index, OutDependency);
+}
+
+bool UInterchangeMaterialNode::SetTextureDependencyUid(const FString& DependencyUid)
+{
+	return TextureDependencies.AddItem(DependencyUid);
+}
+
+bool UInterchangeMaterialNode::RemoveTextureDependencyUid(const FString& DependencyUid)
+{
+	return TextureDependencies.RemoveItem(DependencyUid);
+}
+
+const TOptional<FString> UInterchangeMaterialNode::GetPayLoadKey() const
+{
+	if (!Attributes->ContainAttribute(UE::Interchange::FMaterialNodeStaticData::PayloadSourceFileKey()))
+	{
+		return TOptional<FString>();
+	}
+	UE::Interchange::FAttributeStorage::TAttributeHandle<FString> AttributeHandle = Attributes->GetAttributeHandle<FString>(UE::Interchange::FMaterialNodeStaticData::PayloadSourceFileKey());
+	if (!AttributeHandle.IsValid())
+	{
+		return TOptional<FString>();
+	}
+	FString PayloadKey;
+	UE::Interchange::EAttributeStorageResult Result = AttributeHandle.Get(PayloadKey);
+	if (!IsAttributeStorageResultSuccess(Result))
+	{
+		LogAttributeStorageErrors(Result, TEXT("UInterchangeMaterialNode.GetPayLoadKey"), UE::Interchange::FMaterialNodeStaticData::PayloadSourceFileKey());
+		return TOptional<FString>();
+	}
+	return TOptional<FString>(PayloadKey);
+}
+
+void UInterchangeMaterialNode::SetPayLoadKey(const FString& PayloadKey)
+{
+	UE::Interchange::EAttributeStorageResult Result = Attributes->RegisterAttribute(UE::Interchange::FMaterialNodeStaticData::PayloadSourceFileKey(), PayloadKey);
+	if (!IsAttributeStorageResultSuccess(Result))
+	{
+		LogAttributeStorageErrors(Result, TEXT("UInterchangeMaterialNode.SetPayLoadKey"), UE::Interchange::FMaterialNodeStaticData::PayloadSourceFileKey());
+	}
+}
+
+void UInterchangeMaterialNode::AddTextureParameterData(const EInterchangeMaterialNodeParameterName ParameterName, const FString& TextureUid, const int32 UVSetIndex, const float ScaleU, const float ScaleV)
+{
+	FParameterData& ParameterData = ParameterDatas.FindOrAdd(ParameterName);
+	ParameterData.bIsTextureParameter = true;
+	ParameterData.bIsVectorParameter = false;
+	ParameterData.bIsScalarParameter = false;
+	ParameterData.TextureUid = TextureUid;
+	ParameterData.UVSetIndex = UVSetIndex;
+	ParameterData.ScaleU = ScaleU;
+	ParameterData.ScaleV = ScaleV;
+}
+
+bool UInterchangeMaterialNode::GetTextureParameterData(const EInterchangeMaterialNodeParameterName ParameterName, FString& OutTextureUid, int32& OutUVSetIndex, float& OutScaleU, float& OutScaleV) const
+{
+	if (const FParameterData* ParameterData = ParameterDatas.Find(ParameterName))
+	{
+		if (!ParameterData->bIsTextureParameter)
+		{
+			return false;
+		}
+		OutTextureUid = ParameterData->TextureUid;
+		OutUVSetIndex = ParameterData->UVSetIndex;
+		OutScaleU = ParameterData->ScaleU;
+		OutScaleV = ParameterData->ScaleV;
+		return true;
+	}
+
+	return false;
+}
+
+void UInterchangeMaterialNode::AddVectorParameterData(const EInterchangeMaterialNodeParameterName ParameterName, const FVector& VectorData)
+{
+	FParameterData& ParameterData = ParameterDatas.FindOrAdd(ParameterName);
+	ParameterData.bIsTextureParameter = false;
+	ParameterData.bIsVectorParameter = true;
+	ParameterData.bIsScalarParameter = false;
+	ParameterData.VectorParameter = VectorData;
+}
+
+bool UInterchangeMaterialNode::GetVectorParameterData(const EInterchangeMaterialNodeParameterName ParameterName, FVector& OutVectorData) const
+{
+	if (const FParameterData* ParameterData = ParameterDatas.Find(ParameterName))
+	{
+		if (!ParameterData->bIsVectorParameter)
+		{
+			return false;
+		}
+		OutVectorData = ParameterData->VectorParameter;
+		return true;
+	}
+	return false;
+}
+
+void UInterchangeMaterialNode::AddScalarParameterData(const EInterchangeMaterialNodeParameterName ParameterName, const float ScalarData)
+{
+	FParameterData& ParameterData = ParameterDatas.FindOrAdd(ParameterName);
+	ParameterData.bIsTextureParameter = false;
+	ParameterData.bIsVectorParameter = false;
+	ParameterData.bIsScalarParameter = true;
+	ParameterData.ScalarParameter = ScalarData;
+}
+
+bool UInterchangeMaterialNode::GetScalarParameterData(const EInterchangeMaterialNodeParameterName ParameterName, float& OutScalarData) const
+{
+	if (const FParameterData* ParameterData = ParameterDatas.Find(ParameterName))
+	{
+		if (!ParameterData->bIsScalarParameter)
+		{
+			return false;
+		}
+		OutScalarData = ParameterData->ScalarParameter;
+		return true;
+	}
+	return false;
+}
+
+void UInterchangeMaterialNode::Serialize(FArchive& Ar)
+{
+	Super::Serialize(Ar);
+
+	if (Ar.IsLoading())
+	{
+		ParameterDatas.Reset();
+	}
+	int32 ParamCount = ParameterDatas.Num();
+	Ar << ParamCount;
+
+	if (Ar.IsSaving())
+	{
+		for (TPair<EInterchangeMaterialNodeParameterName, FParameterData> ParamPair : ParameterDatas)
+		{
+			Ar << ParamPair.Key;
+			Ar << ParamPair.Value;
+		}
+	}
+	else if (Ar.IsLoading())
+	{
+		for (int32 ParamIndex = 0; ParamIndex < ParamCount; ++ParamIndex)
+		{
+			EInterchangeMaterialNodeParameterName ParamKey;
+			Ar << ParamKey;
+			FParameterData& ParamValue = ParameterDatas.FindOrAdd(ParamKey);
+			Ar << ParamValue;
+		}
+	}
+}
\ No newline at end of file
diff --git a/Engine/Plugins/Experimental/Interchange/Runtime/Source/Nodes/Private/InterchangeMeshNode.cpp b/Engine/Plugins/Experimental/Interchange/Runtime/Source/Nodes/Private/InterchangeMeshNode.cpp
new file mode 100644
--- /dev/null
+++ b/Engine/Plugins/Experimental/Interchange/Runtime/Source/Nodes/Private/InterchangeMeshNode.cpp
@@ -0,0 +1,508 @@
+// Copyright Epic Games, Inc. All Rights Reserved.
+
+#include "InterchangeMeshNode.h"
+
+//Interchange namespace
+namespace UE
+{
+	namespace Interchange
+	{
+
+		const FAttributeKey& FMeshNodeStaticData::PayloadSourceFileKey()
+		{
+			static FAttributeKey AttributeKey(TEXT("__PayloadSourceFile__"));
+			return AttributeKey;
+		}
+
+		const FAttributeKey& FMeshNodeStaticData::IsSkinnedMeshKey()
+		{
+			static FAttributeKey AttributeKey(TEXT("__IsSkinnedMeshKey__"));
+			return AttributeKey;
+		}
+
+		const FAttributeKey& FMeshNodeStaticData::IsBlendShapeKey()
+		{
+			static FAttributeKey AttributeKey(TEXT("__IsBlendShapeKey__"));
+			return AttributeKey;
+		}
+
+		const FAttributeKey& FMeshNodeStaticData::BlendShapeNameKey()
+		{
+			static FAttributeKey AttributeKey(TEXT("__BlendShapeNameKey__"));
+			return AttributeKey;
+		}
+
+		const FString& FMeshNodeStaticData::GetSkeletonDependenciesKey()
+		{
+			static FString Dependencies_BaseKey = TEXT("__MeshSkeletonDependencies__");
+			return Dependencies_BaseKey;
+		}
+
+		const FString& FMeshNodeStaticData::GetMaterialDependenciesKey()
+		{
+			static FString Dependencies_BaseKey = TEXT("__MeshMaterialDependencies__");
+			return Dependencies_BaseKey;
+		}
+
+		const FString& FMeshNodeStaticData::GetShapeDependenciesKey()
+		{
+			static FString Dependencies_BaseKey = TEXT("__MeshShapeDependencies__");
+			return Dependencies_BaseKey;
+		}
+
+		const FString& FMeshNodeStaticData::GetSceneInstancesUidsKey()
+		{
+			static FString SceneInstanceUids_BaseKey = TEXT("__MeshSceneInstancesUids__");
+			return SceneInstanceUids_BaseKey;
+		}
+
+	}//ns Interchange
+}//ns UE
+
+UInterchangeMeshNode::UInterchangeMeshNode()
+{
+	SkeletonDependencies.Initialize(Attributes, UE::Interchange::FMeshNodeStaticData::GetSkeletonDependenciesKey());
+	MaterialDependencies.Initialize(Attributes, UE::Interchange::FMeshNodeStaticData::GetMaterialDependenciesKey());
+	ShapeDependencies.Initialize(Attributes, UE::Interchange::FMeshNodeStaticData::GetShapeDependenciesKey());
+	SceneInstancesUids.Initialize(Attributes, UE::Interchange::FMeshNodeStaticData::GetSceneInstancesUidsKey());
+}
+
+FString UInterchangeMeshNode::GetKeyDisplayName(const UE::Interchange::FAttributeKey& NodeAttributeKey) const
+{
+	FString KeyDisplayName = NodeAttributeKey.ToString();
+	if (NodeAttributeKey == UE::Interchange::FMeshNodeStaticData::PayloadSourceFileKey())
+	{
+		return KeyDisplayName = TEXT("Payload Source Key");
+	}
+	else if (NodeAttributeKey == UE::Interchange::FMeshNodeStaticData::IsSkinnedMeshKey())
+	{
+		return KeyDisplayName = TEXT("Is a Skinned Mesh");
+	}
+	else if (NodeAttributeKey == UE::Interchange::FMeshNodeStaticData::IsBlendShapeKey())
+	{
+		return KeyDisplayName = TEXT("Is a Blend Shape");
+	}
+	else if (NodeAttributeKey == UE::Interchange::FMeshNodeStaticData::BlendShapeNameKey())
+	{
+		return KeyDisplayName = TEXT("Blend Shape Name");
+	}
+	else if (NodeAttributeKey.Key.Equals(UE::Interchange::FMeshNodeStaticData::GetSkeletonDependenciesKey()))
+	{
+		return KeyDisplayName = TEXT("Skeleton Dependencies count");
+	}
+	else if (NodeAttributeKey.Key.StartsWith(UE::Interchange::FMeshNodeStaticData::GetSkeletonDependenciesKey()))
+	{
+		KeyDisplayName = TEXT("Skeleton Dependencies Index ");
+		const FString IndexKey = UE::Interchange::TArrayAttributeHelper<FString>::IndexKey();
+		int32 IndexPosition = NodeAttributeKey.Key.Find(IndexKey) + IndexKey.Len();
+		if (IndexPosition < NodeAttributeKey.Key.Len())
+		{
+			KeyDisplayName += NodeAttributeKey.Key.RightChop(IndexPosition);
+		}
+		return KeyDisplayName;
+	}
+	else if (NodeAttributeKey.Key.Equals(UE::Interchange::FMeshNodeStaticData::GetMaterialDependenciesKey()))
+	{
+		return KeyDisplayName = TEXT("Material Dependencies count");
+	}
+	else if (NodeAttributeKey.Key.StartsWith(UE::Interchange::FMeshNodeStaticData::GetMaterialDependenciesKey()))
+	{
+		KeyDisplayName = TEXT("Material Dependencies Index ");
+		const FString IndexKey = UE::Interchange::TArrayAttributeHelper<FString>::IndexKey();
+		int32 IndexPosition = NodeAttributeKey.Key.Find(IndexKey) + IndexKey.Len();
+		if (IndexPosition < NodeAttributeKey.Key.Len())
+		{
+			KeyDisplayName += NodeAttributeKey.Key.RightChop(IndexPosition);
+		}
+		return KeyDisplayName;
+	}
+	else if (NodeAttributeKey.Key.Equals(UE::Interchange::FMeshNodeStaticData::GetShapeDependenciesKey()))
+	{
+		return KeyDisplayName = TEXT("Shape Dependencies count");
+	}
+	else if (NodeAttributeKey.Key.StartsWith(UE::Interchange::FMeshNodeStaticData::GetShapeDependenciesKey()))
+	{
+		KeyDisplayName = TEXT("Shape Dependencies Index ");
+		const FString IndexKey = UE::Interchange::TArrayAttributeHelper<FString>::IndexKey();
+		int32 IndexPosition = NodeAttributeKey.Key.Find(IndexKey) + IndexKey.Len();
+		if (IndexPosition < NodeAttributeKey.Key.Len())
+		{
+			KeyDisplayName += NodeAttributeKey.Key.RightChop(IndexPosition);
+		}
+		return KeyDisplayName;
+	}
+	else if (NodeAttributeKey.Key.Equals(UE::Interchange::FMeshNodeStaticData::GetSceneInstancesUidsKey()))
+	{
+		return KeyDisplayName = TEXT("Scene mesh instances count");
+	}
+	else if (NodeAttributeKey.Key.StartsWith(UE::Interchange::FMeshNodeStaticData::GetSceneInstancesUidsKey()))
+	{
+		KeyDisplayName = TEXT("Scene mesh instances Index ");
+		const FString IndexKey = UE::Interchange::TArrayAttributeHelper<FString>::IndexKey();
+		int32 IndexPosition = NodeAttributeKey.Key.Find(IndexKey) + IndexKey.Len();
+		if (IndexPosition < NodeAttributeKey.Key.Len())
+		{
+			KeyDisplayName += NodeAttributeKey.Key.RightChop(IndexPosition);
+		}
+		return KeyDisplayName;
+	}
+	return Super::GetKeyDisplayName(NodeAttributeKey);
+}
+
+FString UInterchangeMeshNode::GetAttributeCategory(const UE::Interchange::FAttributeKey& NodeAttributeKey) const
+{
+	if (NodeAttributeKey.Key.StartsWith(UE::Interchange::FMeshNodeStaticData::GetSkeletonDependenciesKey()))
+	{
+		return FString(TEXT("SkeletonDependencies"));
+	}
+	else if (NodeAttributeKey.Key.StartsWith(UE::Interchange::FMeshNodeStaticData::GetMaterialDependenciesKey()))
+	{
+		return FString(TEXT("MaterialDependencies"));
+	}
+	else if (NodeAttributeKey.Key.StartsWith(UE::Interchange::FMeshNodeStaticData::GetShapeDependenciesKey()))
+	{
+		return FString(TEXT("ShapeDependencies"));
+	}
+	else if (NodeAttributeKey.Key.StartsWith(UE::Interchange::FMeshNodeStaticData::GetSceneInstancesUidsKey()))
+	{
+		return FString(TEXT("SceneInstances"));
+	}
+	else if (NodeAttributeKey == Macro_CustomVertexCountKey
+				|| NodeAttributeKey == Macro_CustomPolygonCountKey
+				|| NodeAttributeKey == Macro_CustomBoundingBoxKey
+				|| NodeAttributeKey == Macro_CustomHasVertexNormalKey
+				|| NodeAttributeKey == Macro_CustomHasVertexBinormalKey
+				|| NodeAttributeKey == Macro_CustomHasVertexTangentKey
+				|| NodeAttributeKey == Macro_CustomHasSmoothGroupKey
+				|| NodeAttributeKey == Macro_CustomHasVertexColorKey
+				|| NodeAttributeKey == Macro_CustomUVCountKey)
+	{
+		return FString(TEXT("MeshInfo"));
+	}
+	return Super::GetAttributeCategory(NodeAttributeKey);
+}
+
+FString UInterchangeMeshNode::GetTypeName() const
+{
+	const FString TypeName = TEXT("MeshNode");
+	return TypeName;
+}
+
+FName UInterchangeMeshNode::GetIconName() const
+{
+	FString MeshIconName = TEXT("MeshIcon.");
+	if (IsSkinnedMesh())
+	{
+		MeshIconName += TEXT("Skinned");
+	}
+	else
+	{
+		MeshIconName += TEXT("Static");
+	}
+		
+	return FName(*MeshIconName);
+}
+
+bool UInterchangeMeshNode::IsSkinnedMesh() const
+{
+	if (!Attributes->ContainAttribute(UE::Interchange::FMeshNodeStaticData::IsSkinnedMeshKey()))
+	{
+		return false;
+	}
+
+	UE::Interchange::FAttributeStorage::TAttributeHandle<bool> Handle = Attributes->GetAttributeHandle<bool>(UE::Interchange::FMeshNodeStaticData::IsSkinnedMeshKey());
+	if (Handle.IsValid())
+	{
+		bool bValue = false;
+		Handle.Get(bValue);
+		return bValue;
+	}
+	return false;
+}
+
+bool UInterchangeMeshNode::SetSkinnedMesh(const bool bIsSkinnedMesh)
+{
+	UE::Interchange::EAttributeStorageResult Result = Attributes->RegisterAttribute(UE::Interchange::FMeshNodeStaticData::IsSkinnedMeshKey(), bIsSkinnedMesh);
+	if (IsAttributeStorageResultSuccess(Result))
+	{
+		UE::Interchange::FAttributeStorage::TAttributeHandle<bool> Handle = Attributes->GetAttributeHandle<bool>(UE::Interchange::FMeshNodeStaticData::IsSkinnedMeshKey());
+		return Handle.IsValid();
+	}
+	return false;
+}
+
+bool UInterchangeMeshNode::IsBlendShape() const
+{
+	if (!Attributes->ContainAttribute(UE::Interchange::FMeshNodeStaticData::IsBlendShapeKey()))
+	{
+		return false;
+	}
+
+	UE::Interchange::FAttributeStorage::TAttributeHandle<bool> Handle = Attributes->GetAttributeHandle<bool>(UE::Interchange::FMeshNodeStaticData::IsBlendShapeKey());
+	if (Handle.IsValid())
+	{
+		bool bValue = false;
+		Handle.Get(bValue);
+		return bValue;
+	}
+	return false;
+}
+
+bool UInterchangeMeshNode::SetBlendShape(const bool bIsBlendShape)
+{
+	UE::Interchange::EAttributeStorageResult Result = Attributes->RegisterAttribute(UE::Interchange::FMeshNodeStaticData::IsBlendShapeKey(), bIsBlendShape);
+	if (IsAttributeStorageResultSuccess(Result))
+	{
+		UE::Interchange::FAttributeStorage::TAttributeHandle<bool> Handle = Attributes->GetAttributeHandle<bool>(UE::Interchange::FMeshNodeStaticData::IsBlendShapeKey());
+		return Handle.IsValid();
+	}
+	return false;
+}
+
+bool UInterchangeMeshNode::GetBlendShapeName(FString& OutBlendShapeName) const
+{
+	OutBlendShapeName.Empty();
+	if (!Attributes->ContainAttribute(UE::Interchange::FMeshNodeStaticData::BlendShapeNameKey()))
+	{
+		return false;
+	}
+
+	UE::Interchange::FAttributeStorage::TAttributeHandle<FString> Handle = Attributes->GetAttributeHandle<FString>(UE::Interchange::FMeshNodeStaticData::BlendShapeNameKey());
+	if (Handle.IsValid())
+	{
+		Handle.Get(OutBlendShapeName);
+		return true;
+	}
+	return false;
+}
+
+bool UInterchangeMeshNode::SetBlendShapeName(const FString& BlendShapeName)
+{
+	UE::Interchange::EAttributeStorageResult Result = Attributes->RegisterAttribute(UE::Interchange::FMeshNodeStaticData::BlendShapeNameKey(), BlendShapeName);
+	if (IsAttributeStorageResultSuccess(Result))
+	{
+		UE::Interchange::FAttributeStorage::TAttributeHandle<FString> Handle = Attributes->GetAttributeHandle<FString>(UE::Interchange::FMeshNodeStaticData::BlendShapeNameKey());
+		return Handle.IsValid();
+	}
+	return false;
+}
+
+const TOptional<FString> UInterchangeMeshNode::GetPayLoadKey() const
+{
+	if (!Attributes->ContainAttribute(UE::Interchange::FMeshNodeStaticData::PayloadSourceFileKey()))
+	{
+		return TOptional<FString>();
+	}
+	UE::Interchange::FAttributeStorage::TAttributeHandle<FString> AttributeHandle = Attributes->GetAttributeHandle<FString>(UE::Interchange::FMeshNodeStaticData::PayloadSourceFileKey());
+	if (!AttributeHandle.IsValid())
+	{
+		return TOptional<FString>();
+	}
+	FString PayloadKey;
+	UE::Interchange::EAttributeStorageResult Result = AttributeHandle.Get(PayloadKey);
+	if (!IsAttributeStorageResultSuccess(Result))
+	{
+		LogAttributeStorageErrors(Result, TEXT("UInterchangeMeshNode.GetPayLoadKey"), UE::Interchange::FMeshNodeStaticData::PayloadSourceFileKey());
+		return TOptional<FString>();
+	}
+	return TOptional<FString>(PayloadKey);
+}
+
+void UInterchangeMeshNode::SetPayLoadKey(const FString& PayloadKey)
+{
+	UE::Interchange::EAttributeStorageResult Result = Attributes->RegisterAttribute(UE::Interchange::FMeshNodeStaticData::PayloadSourceFileKey(), PayloadKey);
+	if (!IsAttributeStorageResultSuccess(Result))
+	{
+		LogAttributeStorageErrors(Result, TEXT("UInterchangeMeshNode.SetPayLoadKey"), UE::Interchange::FMeshNodeStaticData::PayloadSourceFileKey());
+	}
+}
+	
+bool UInterchangeMeshNode::GetCustomVertexCount(int32& AttributeValue) const
+{
+	IMPLEMENT_NODE_ATTRIBUTE_GETTER(VertexCount, int32);
+}
+
+bool UInterchangeMeshNode::SetCustomVertexCount(const int32& AttributeValue)
+{
+	IMPLEMENT_NODE_ATTRIBUTE_SETTER_NODELEGATE(VertexCount, int32);
+}
+
+bool UInterchangeMeshNode::GetCustomPolygonCount(int32& AttributeValue) const
+{
+	IMPLEMENT_NODE_ATTRIBUTE_GETTER(PolygonCount, int32);
+}
+
+bool UInterchangeMeshNode::SetCustomPolygonCount(const int32& AttributeValue)
+{
+	IMPLEMENT_NODE_ATTRIBUTE_SETTER_NODELEGATE(PolygonCount, int32);
+}
+
+bool UInterchangeMeshNode::GetCustomBoundingBox(FBox& AttributeValue) const
+{
+	IMPLEMENT_NODE_ATTRIBUTE_GETTER(BoundingBox, FBox);
+}
+
+bool UInterchangeMeshNode::SetCustomBoundingBox(const FBox& AttributeValue)
+{
+	IMPLEMENT_NODE_ATTRIBUTE_SETTER_NODELEGATE(BoundingBox, FBox);
+}
+
+bool UInterchangeMeshNode::GetCustomHasVertexNormal(bool& AttributeValue) const
+{
+	IMPLEMENT_NODE_ATTRIBUTE_GETTER(HasVertexNormal, bool);
+}
+
+bool UInterchangeMeshNode::SetCustomHasVertexNormal(const bool& AttributeValue)
+{
+	IMPLEMENT_NODE_ATTRIBUTE_SETTER_NODELEGATE(HasVertexNormal, bool);
+}
+
+bool UInterchangeMeshNode::GetCustomHasVertexBinormal(bool& AttributeValue) const
+{
+	IMPLEMENT_NODE_ATTRIBUTE_GETTER(HasVertexBinormal, bool);
+}
+
+bool UInterchangeMeshNode::SetCustomHasVertexBinormal(const bool& AttributeValue)
+{
+	IMPLEMENT_NODE_ATTRIBUTE_SETTER_NODELEGATE(HasVertexBinormal, bool);
+}
+
+bool UInterchangeMeshNode::GetCustomHasVertexTangent(bool& AttributeValue) const
+{
+	IMPLEMENT_NODE_ATTRIBUTE_GETTER(HasVertexTangent, bool);
+}
+
+bool UInterchangeMeshNode::SetCustomHasVertexTangent(const bool& AttributeValue)
+{
+	IMPLEMENT_NODE_ATTRIBUTE_SETTER_NODELEGATE(HasVertexTangent, bool);
+}
+
+bool UInterchangeMeshNode::GetCustomHasSmoothGroup(bool& AttributeValue) const
+{
+	IMPLEMENT_NODE_ATTRIBUTE_GETTER(HasSmoothGroup, bool);
+}
+
+bool UInterchangeMeshNode::SetCustomHasSmoothGroup(const bool& AttributeValue)
+{
+	IMPLEMENT_NODE_ATTRIBUTE_SETTER_NODELEGATE(HasSmoothGroup, bool);
+}
+
+bool UInterchangeMeshNode::GetCustomHasVertexColor(bool& AttributeValue) const
+{
+	IMPLEMENT_NODE_ATTRIBUTE_GETTER(HasVertexColor, bool);
+}
+
+bool UInterchangeMeshNode::SetCustomHasVertexColor(const bool& AttributeValue)
+{
+	IMPLEMENT_NODE_ATTRIBUTE_SETTER_NODELEGATE(HasVertexColor, bool);
+}
+
+bool UInterchangeMeshNode::GetCustomUVCount(int32& AttributeValue) const
+{
+	IMPLEMENT_NODE_ATTRIBUTE_GETTER(UVCount, int32);
+}
+
+bool UInterchangeMeshNode::SetCustomUVCount(const int32& AttributeValue)
+{
+	IMPLEMENT_NODE_ATTRIBUTE_SETTER_NODELEGATE(UVCount, int32);
+}
+
+int32 UInterchangeMeshNode::GetSkeletonDependeciesCount() const
+{
+	return SkeletonDependencies.GetCount();
+}
+
+void UInterchangeMeshNode::GetSkeletonDependencies(TArray<FString>& OutDependencies) const
+{
+	SkeletonDependencies.GetItems(OutDependencies);
+}
+
+void UInterchangeMeshNode::GetSkeletonDependency(const int32 Index, FString& OutDependency) const
+{
+	SkeletonDependencies.GetItem(Index, OutDependency);
+}
+
+bool UInterchangeMeshNode::SetSkeletonDependencyUid(const FString& DependencyUid)
+{ 
+	return SkeletonDependencies.AddItem(DependencyUid);
+}
+
+bool UInterchangeMeshNode::RemoveSkeletonDependencyUid(const FString& DependencyUid)
+{
+	return SkeletonDependencies.RemoveItem(DependencyUid);
+}
+
+int32 UInterchangeMeshNode::GetMaterialDependeciesCount() const
+{
+	return MaterialDependencies.GetCount();
+}
+
+void UInterchangeMeshNode::GetMaterialDependencies(TArray<FString>& OutDependencies) const
+{
+	MaterialDependencies.GetItems(OutDependencies);
+}
+
+void UInterchangeMeshNode::GetMaterialDependency(const int32 Index, FString& OutDependency) const
+{
+	MaterialDependencies.GetItem(Index, OutDependency);
+}
+
+bool UInterchangeMeshNode::SetMaterialDependencyUid(const FString& DependencyUid)
+{
+	return MaterialDependencies.AddItem(DependencyUid);
+}
+
+bool UInterchangeMeshNode::RemoveMaterialDependencyUid(const FString& DependencyUid)
+{
+	return MaterialDependencies.RemoveItem(DependencyUid);
+}
+
+int32 UInterchangeMeshNode::GetShapeDependeciesCount() const
+{
+	return ShapeDependencies.GetCount();
+}
+
+void UInterchangeMeshNode::GetShapeDependencies(TArray<FString>& OutDependencies) const
+{
+	ShapeDependencies.GetItems(OutDependencies);
+}
+
+void UInterchangeMeshNode::GetShapeDependency(const int32 Index, FString& OutDependency) const
+{
+	ShapeDependencies.GetItem(Index, OutDependency);
+}
+
+bool UInterchangeMeshNode::SetShapeDependencyUid(const FString& DependencyUid)
+{
+	return ShapeDependencies.AddItem(DependencyUid);
+}
+
+bool UInterchangeMeshNode::RemoveShapeDependencyUid(const FString& DependencyUid)
+{
+	return ShapeDependencies.RemoveItem(DependencyUid);
+}
+
+int32 UInterchangeMeshNode::GetSceneInstanceUidsCount() const
+{
+	return SceneInstancesUids.GetCount();
+}
+
+void UInterchangeMeshNode::GetSceneInstanceUids(TArray<FString>& OutDependencies) const
+{
+	SceneInstancesUids.GetItems(OutDependencies);
+}
+
+void UInterchangeMeshNode::GetSceneInstanceUid(const int32 Index, FString& OutDependency) const
+{
+	SceneInstancesUids.GetItem(Index, OutDependency);
+}
+
+bool UInterchangeMeshNode::SetSceneInstanceUid(const FString& DependencyUid)
+{
+	return SceneInstancesUids.AddItem(DependencyUid);
+}
+
+bool UInterchangeMeshNode::RemoveSceneInstanceUid(const FString& DependencyUid)
+{
+	return SceneInstancesUids.RemoveItem(DependencyUid);
+}
\ No newline at end of file
diff --git a/Engine/Plugins/Experimental/Interchange/Runtime/Source/Nodes/Private/InterchangeSceneNode.cpp b/Engine/Plugins/Experimental/Interchange/Runtime/Source/Nodes/Private/InterchangeSceneNode.cpp
new file mode 100644
--- /dev/null
+++ b/Engine/Plugins/Experimental/Interchange/Runtime/Source/Nodes/Private/InterchangeSceneNode.cpp
@@ -0,0 +1,264 @@
+// Copyright Epic Games, Inc. All Rights Reserved.
+
+#include "InterchangeSceneNode.h"
+
+//Interchange namespace
+namespace UE
+{
+	namespace Interchange
+	{
+
+		const FString& FSceneNodeStaticData::GetNodeSpecializeTypeBaseKey()
+		{
+			static FString SceneNodeSpecializeType_BaseKey = TEXT("SceneNodeSpecializeType");
+			return SceneNodeSpecializeType_BaseKey;
+		}
+
+		const FString& FSceneNodeStaticData::GetMaterialDependencyUidsBaseKey()
+		{
+			static FString MaterialDependencyUids_BaseKey = TEXT("__MaterialDependencyUidsBaseKey__");
+			return MaterialDependencyUids_BaseKey;
+		}
+
+		const FString& FSceneNodeStaticData::GetJointSpecializeTypeString()
+		{
+			static FString JointSpecializeTypeString = TEXT("Joint");
+			return JointSpecializeTypeString;
+		}
+
+		const FString& FSceneNodeStaticData::GetLodGroupSpecializeTypeString()
+		{
+			static FString JointSpecializeTypeString = TEXT("LodGroup");
+			return JointSpecializeTypeString;
+		}
+	}//ns Interchange
+}//ns UE
+
+UInterchangeSceneNode::UInterchangeSceneNode()
+{
+	NodeSpecializeTypes.Initialize(Attributes, UE::Interchange::FSceneNodeStaticData::GetNodeSpecializeTypeBaseKey());
+	MaterialDependencyUids.Initialize(Attributes, UE::Interchange::FSceneNodeStaticData::GetMaterialDependencyUidsBaseKey());
+}
+
+/**
+	* Return the node type name of the class, we use this when reporting error
+	*/
+FString UInterchangeSceneNode::GetTypeName() const
+{
+	const FString TypeName = TEXT("SceneNode");
+	return TypeName;
+}
+
+FString UInterchangeSceneNode::GetKeyDisplayName(const UE::Interchange::FAttributeKey& NodeAttributeKey) const
+{
+	FString KeyDisplayName = NodeAttributeKey.Key;
+	if (NodeAttributeKey.Key.Equals(UE::Interchange::FSceneNodeStaticData::GetNodeSpecializeTypeBaseKey()))
+	{
+		KeyDisplayName = TEXT("Specialized type count");
+		return KeyDisplayName;
+	}
+	else if (NodeAttributeKey.Key.StartsWith(UE::Interchange::FSceneNodeStaticData::GetNodeSpecializeTypeBaseKey()))
+	{
+		KeyDisplayName = TEXT("Specialized type index ");
+		const FString IndexKey = UE::Interchange::TArrayAttributeHelper<FString>::IndexKey();
+		int32 IndexPosition = NodeAttributeKey.Key.Find(IndexKey) + IndexKey.Len();
+		if (IndexPosition < NodeAttributeKey.Key.Len())
+		{
+			KeyDisplayName += NodeAttributeKey.Key.RightChop(IndexPosition);
+		}
+		return KeyDisplayName;
+	}
+	else if (NodeAttributeKey.Key.Equals(UE::Interchange::FSceneNodeStaticData::GetMaterialDependencyUidsBaseKey()))
+	{
+		KeyDisplayName = TEXT("Material dependencies count");
+		return KeyDisplayName;
+	}
+	else if (NodeAttributeKey.Key.StartsWith(UE::Interchange::FSceneNodeStaticData::GetMaterialDependencyUidsBaseKey()))
+	{
+		KeyDisplayName = TEXT("Material dependency index ");
+		const FString IndexKey = UE::Interchange::TArrayAttributeHelper<FString>::IndexKey();
+		int32 IndexPosition = NodeAttributeKey.Key.Find(IndexKey) + IndexKey.Len();
+		if (IndexPosition < NodeAttributeKey.Key.Len())
+		{
+			KeyDisplayName += NodeAttributeKey.Key.RightChop(IndexPosition);
+		}
+		return KeyDisplayName;
+	}
+	return Super::GetKeyDisplayName(NodeAttributeKey);
+}
+
+FString UInterchangeSceneNode::GetAttributeCategory(const UE::Interchange::FAttributeKey& NodeAttributeKey) const
+{
+	if (NodeAttributeKey.Key.StartsWith(UE::Interchange::FSceneNodeStaticData::GetNodeSpecializeTypeBaseKey()))
+	{
+		return FString(TEXT("SpecializeType"));
+	}
+	else if (NodeAttributeKey.Key.StartsWith(UE::Interchange::FSceneNodeStaticData::GetMaterialDependencyUidsBaseKey()))
+	{
+		return FString(TEXT("MaterialDependencies"));
+	}
+	else if (NodeAttributeKey == Macro_CustomLocalTransformKey
+		|| NodeAttributeKey == Macro_CustomGlobalTransformKey
+		|| NodeAttributeKey == Macro_CustomAssetInstanceUidKey)
+	{
+		return FString(TEXT("Scene"));
+	}
+	return Super::GetAttributeCategory(NodeAttributeKey);
+}
+
+FName UInterchangeSceneNode::GetIconName() const
+{
+	FString SpecializedType;
+	GetSpecializedType(0, SpecializedType);
+	if (SpecializedType.IsEmpty())
+	{
+		return NAME_None;
+	}
+	SpecializedType = TEXT("SceneGraphIcon.") + SpecializedType;
+	return FName(*SpecializedType);
+}
+
+bool UInterchangeSceneNode::IsSpecializedTypeContains(const FString& SpecializedType) const
+{
+	TArray<FString> SpecializedTypes;
+	GetSpecializedTypes(SpecializedTypes);
+	for (const FString& SpecializedTypeRef : SpecializedTypes)
+	{
+		if (SpecializedTypeRef.Equals(SpecializedType))
+		{
+			return true;
+		}
+	}
+	return false;
+}
+
+int32 UInterchangeSceneNode::GetSpecializedTypeCount() const
+{
+	return NodeSpecializeTypes.GetCount();
+}
+
+void UInterchangeSceneNode::GetSpecializedType(const int32 Index, FString& OutSpecializedType) const
+{
+	NodeSpecializeTypes.GetItem(Index, OutSpecializedType);
+}
+
+void UInterchangeSceneNode::GetSpecializedTypes(TArray<FString>& OutSpecializedTypes) const
+{
+	NodeSpecializeTypes.GetItems(OutSpecializedTypes);
+}
+
+bool UInterchangeSceneNode::AddSpecializedType(const FString& SpecializedType)
+{
+	return NodeSpecializeTypes.AddItem(SpecializedType);
+}
+
+bool UInterchangeSceneNode::RemoveSpecializedType(const FString& SpecializedType)
+{
+	return NodeSpecializeTypes.RemoveItem(SpecializedType);
+}
+
+int32 UInterchangeSceneNode::GetMaterialDependencyUidsCount() const
+{
+	return MaterialDependencyUids.GetCount();
+}
+
+void UInterchangeSceneNode::GetMaterialDependencyUid(const int32 Index, FString& OutMaterialDependencyUid) const
+{
+	MaterialDependencyUids.GetItem(Index, OutMaterialDependencyUid);
+}
+
+void UInterchangeSceneNode::GetMaterialDependencyUids(TArray<FString>& OutMaterialDependencyUids) const
+{
+	MaterialDependencyUids.GetItems(OutMaterialDependencyUids);
+}
+
+bool UInterchangeSceneNode::AddMaterialDependencyUid(const FString& MaterialDependencyUid)
+{
+	return MaterialDependencyUids.AddItem(MaterialDependencyUid);
+}
+
+bool UInterchangeSceneNode::RemoveMaterialDependencyUid(const FString& MaterialDependencyUid)
+{
+	return MaterialDependencyUids.RemoveItem(MaterialDependencyUid);
+}
+
+bool UInterchangeSceneNode::GetCustomLocalTransform(FTransform& AttributeValue) const
+{
+	IMPLEMENT_NODE_ATTRIBUTE_GETTER(LocalTransform, FTransform);
+}
+
+bool UInterchangeSceneNode::SetCustomLocalTransform(const FTransform& AttributeValue)
+{
+	IMPLEMENT_NODE_ATTRIBUTE_SETTER_NODELEGATE(LocalTransform, FTransform);
+}
+
+bool UInterchangeSceneNode::GetCustomGlobalTransform(FTransform& AttributeValue) const
+{
+	IMPLEMENT_NODE_ATTRIBUTE_GETTER(GlobalTransform, FTransform);
+}
+
+bool UInterchangeSceneNode::SetCustomGlobalTransform(const FTransform& AttributeValue)
+{
+	IMPLEMENT_NODE_ATTRIBUTE_SETTER_NODELEGATE(GlobalTransform, FTransform);
+}
+
+bool UInterchangeSceneNode::GetCustomBindPoseLocalTransform(FTransform& AttributeValue) const
+{
+	IMPLEMENT_NODE_ATTRIBUTE_GETTER(BindPoseLocalTransform, FTransform);
+}
+
+bool UInterchangeSceneNode::SetCustomBindPoseLocalTransform(const FTransform& AttributeValue)
+{
+	IMPLEMENT_NODE_ATTRIBUTE_SETTER_NODELEGATE(BindPoseLocalTransform, FTransform);
+}
+
+bool UInterchangeSceneNode::GetCustomBindPoseGlobalTransform(FTransform& AttributeValue) const
+{
+	IMPLEMENT_NODE_ATTRIBUTE_GETTER(BindPoseGlobalTransform, FTransform);
+}
+
+bool UInterchangeSceneNode::SetCustomBindPoseGlobalTransform(const FTransform& AttributeValue)
+{
+	IMPLEMENT_NODE_ATTRIBUTE_SETTER_NODELEGATE(BindPoseGlobalTransform, FTransform);
+}
+
+bool UInterchangeSceneNode::GetCustomTimeZeroLocalTransform(FTransform& AttributeValue) const
+{
+	IMPLEMENT_NODE_ATTRIBUTE_GETTER(TimeZeroLocalTransform, FTransform);
+}
+
+bool UInterchangeSceneNode::SetCustomTimeZeroLocalTransform(const FTransform& AttributeValue)
+{
+	IMPLEMENT_NODE_ATTRIBUTE_SETTER_NODELEGATE(TimeZeroLocalTransform, FTransform);
+}
+
+bool UInterchangeSceneNode::GetCustomTimeZeroGlobalTransform(FTransform& AttributeValue) const
+{
+	IMPLEMENT_NODE_ATTRIBUTE_GETTER(TimeZeroGlobalTransform, FTransform);
+}
+
+bool UInterchangeSceneNode::SetCustomTimeZeroGlobalTransform(const FTransform& AttributeValue)
+{
+	IMPLEMENT_NODE_ATTRIBUTE_SETTER_NODELEGATE(TimeZeroGlobalTransform, FTransform);
+}
+
+bool UInterchangeSceneNode::GetCustomGeometricTransform(FTransform& AttributeValue) const
+{
+	IMPLEMENT_NODE_ATTRIBUTE_GETTER(GeometricTransform, FTransform);
+}
+
+bool UInterchangeSceneNode::SetCustomGeometricTransform(const FTransform& AttributeValue)
+{
+	IMPLEMENT_NODE_ATTRIBUTE_SETTER_NODELEGATE(GeometricTransform, FTransform);
+}
+
+bool UInterchangeSceneNode::GetCustomAssetInstanceUid(FString& AttributeValue) const
+{
+	IMPLEMENT_NODE_ATTRIBUTE_GETTER(AssetInstanceUid, FString);
+}
+
+bool UInterchangeSceneNode::SetCustomAssetInstanceUid(const FString& AttributeValue)
+{
+	IMPLEMENT_NODE_ATTRIBUTE_SETTER_NODELEGATE(AssetInstanceUid, FString);
+}
+
diff --git a/Engine/Plugins/Experimental/Interchange/Runtime/Source/Nodes/Private/InterchangeSkeletalMeshFactoryNode.cpp b/Engine/Plugins/Experimental/Interchange/Runtime/Source/Nodes/Private/InterchangeSkeletalMeshFactoryNode.cpp
new file mode 100644
--- /dev/null
+++ b/Engine/Plugins/Experimental/Interchange/Runtime/Source/Nodes/Private/InterchangeSkeletalMeshFactoryNode.cpp
@@ -0,0 +1,211 @@
+// Copyright Epic Games, Inc. All Rights Reserved.
+
+#include "InterchangeSkeletalMeshFactoryNode.h"
+
+#if WITH_ENGINE
+#include "Engine/SkeletalMesh.h"
+#endif
+
+//Interchange namespace
+namespace UE
+{
+	namespace Interchange
+	{
+		const FString& FSkeletalMeshNodeStaticData::GetLodDependenciesBaseKey()
+		{
+			static FString LodDependencies_BaseKey = TEXT("Lod_Dependencies");
+			return LodDependencies_BaseKey;
+		}
+	}//ns Interchange
+}//ns UE
+
+UInterchangeSkeletalMeshFactoryNode::UInterchangeSkeletalMeshFactoryNode()
+{
+#if WITH_ENGINE
+	AssetClass = nullptr;
+#endif
+	LodDependencies.Initialize(Attributes, UE::Interchange::FSkeletalMeshNodeStaticData::GetLodDependenciesBaseKey());
+}
+
+void UInterchangeSkeletalMeshFactoryNode::InitializeSkeletalMeshNode(const FString& UniqueID, const FString& DisplayLabel, const FString& InAssetClass)
+{
+	bIsNodeClassInitialized = false;
+	InitializeNode(UniqueID, DisplayLabel, EInterchangeNodeContainerType::FactoryData);
+
+	FString OperationName = GetTypeName() + TEXT(".SetAssetClassName");
+	InterchangePrivateNodeBase::SetCustomAttribute<FString>(*Attributes, ClassNameAttributeKey, OperationName, InAssetClass);
+	FillAssetClassFromAttribute();
+}
+
+void UInterchangeSkeletalMeshFactoryNode::Serialize(FArchive& Ar)
+{
+	Super::Serialize(Ar);
+#if WITH_ENGINE
+	if (Ar.IsLoading())
+	{
+		//Make sure the class is properly set when we compile with engine, this will set the
+		//bIsNodeClassInitialized to true.
+		SetNodeClassFromClassAttribute();
+	}
+#endif //#if WITH_ENGINE
+}
+
+FString UInterchangeSkeletalMeshFactoryNode::GetTypeName() const
+{
+	const FString TypeName = TEXT("SkeletalMeshNode");
+	return TypeName;
+}
+
+FString UInterchangeSkeletalMeshFactoryNode::GetKeyDisplayName(const UE::Interchange::FAttributeKey& NodeAttributeKey) const
+{
+	FString KeyDisplayName = NodeAttributeKey.Key;
+	if (NodeAttributeKey.Key.Equals(UE::Interchange::FSkeletalMeshNodeStaticData::GetLodDependenciesBaseKey()))
+	{
+		KeyDisplayName = TEXT("LOD Dependencies Count");
+		return KeyDisplayName;
+	}
+	else if (NodeAttributeKey.Key.StartsWith(UE::Interchange::FSkeletalMeshNodeStaticData::GetLodDependenciesBaseKey()))
+	{
+		KeyDisplayName = TEXT("LOD Dependencies Index ");
+		const FString IndexKey = UE::Interchange::TArrayAttributeHelper<FString>::IndexKey();
+		int32 IndexPosition = NodeAttributeKey.Key.Find(IndexKey) + IndexKey.Len();
+		if (IndexPosition < NodeAttributeKey.Key.Len())
+		{
+			KeyDisplayName += NodeAttributeKey.Key.RightChop(IndexPosition);
+		}
+		return KeyDisplayName;
+	}
+	return Super::GetKeyDisplayName(NodeAttributeKey);
+}
+
+UClass* UInterchangeSkeletalMeshFactoryNode::GetObjectClass() const
+{
+	ensure(bIsNodeClassInitialized);
+#if WITH_ENGINE
+	return AssetClass.Get() != nullptr ? AssetClass.Get() : USkeletalMesh::StaticClass();
+#else
+	return nullptr;
+#endif
+}
+
+int32 UInterchangeSkeletalMeshFactoryNode::GetLodDataCount() const
+{
+	return LodDependencies.GetCount();
+}
+
+void UInterchangeSkeletalMeshFactoryNode::GetLodDataUniqueIds(TArray<FString>& OutLodDataUniqueIds) const
+{
+	LodDependencies.GetItems(OutLodDataUniqueIds);
+}
+
+bool UInterchangeSkeletalMeshFactoryNode::AddLodDataUniqueId(const FString& LodDataUniqueId)
+{
+	return LodDependencies.AddItem(LodDataUniqueId);
+}
+
+bool UInterchangeSkeletalMeshFactoryNode::RemoveLodDataUniqueId(const FString& LodDataUniqueId)
+{
+	return LodDependencies.RemoveItem(LodDataUniqueId);
+}
+
+bool UInterchangeSkeletalMeshFactoryNode::GetCustomSkeletonSoftObjectPath(FSoftObjectPath& AttributeValue) const
+{
+	IMPLEMENT_NODE_ATTRIBUTE_GETTER(SkeletonSoftObjectPath, FSoftObjectPath)
+}
+
+bool UInterchangeSkeletalMeshFactoryNode::SetCustomSkeletonSoftObjectPath(const FSoftObjectPath& AttributeValue)
+{
+	IMPLEMENT_NODE_ATTRIBUTE_SETTER_NODELEGATE(SkeletonSoftObjectPath, FSoftObjectPath)
+}
+
+bool UInterchangeSkeletalMeshFactoryNode::GetCustomImportMorphTarget(bool& AttributeValue) const
+{
+	IMPLEMENT_NODE_ATTRIBUTE_GETTER(ImportMorphTarget, bool)
+}
+
+bool UInterchangeSkeletalMeshFactoryNode::SetCustomImportMorphTarget(const bool& AttributeValue)
+{
+	IMPLEMENT_NODE_ATTRIBUTE_SETTER_NODELEGATE(ImportMorphTarget, bool)
+}
+
+bool UInterchangeSkeletalMeshFactoryNode::GetCustomCreatePhysicsAsset(bool& AttributeValue) const
+{
+	IMPLEMENT_NODE_ATTRIBUTE_GETTER(CreatePhysicsAsset, bool)
+}
+
+bool UInterchangeSkeletalMeshFactoryNode::SetCustomCreatePhysicsAsset(const bool& AttributeValue)
+{
+	IMPLEMENT_NODE_ATTRIBUTE_SETTER_NODELEGATE(CreatePhysicsAsset, bool)
+}
+
+bool UInterchangeSkeletalMeshFactoryNode::GetCustomPhysicAssetSoftObjectPath(FSoftObjectPath& AttributeValue) const
+{
+	IMPLEMENT_NODE_ATTRIBUTE_GETTER(PhysicAssetSoftObjectPath, FSoftObjectPath)
+}
+
+bool UInterchangeSkeletalMeshFactoryNode::SetCustomPhysicAssetSoftObjectPath(const FSoftObjectPath& AttributeValue)
+{
+	IMPLEMENT_NODE_ATTRIBUTE_SETTER_NODELEGATE(PhysicAssetSoftObjectPath, FSoftObjectPath)
+}
+
+bool UInterchangeSkeletalMeshFactoryNode::GetCustomVertexColorReplace(bool& AttributeValue) const
+{
+	IMPLEMENT_NODE_ATTRIBUTE_GETTER(VertexColorReplace, bool)
+}
+
+bool UInterchangeSkeletalMeshFactoryNode::SetCustomVertexColorReplace(const bool& AttributeValue)
+{
+	IMPLEMENT_NODE_ATTRIBUTE_SETTER_NODELEGATE(VertexColorReplace, bool)
+}
+
+bool UInterchangeSkeletalMeshFactoryNode::GetCustomVertexColorIgnore(bool& AttributeValue) const
+{
+	IMPLEMENT_NODE_ATTRIBUTE_GETTER(VertexColorIgnore, bool)
+}
+
+bool UInterchangeSkeletalMeshFactoryNode::SetCustomVertexColorIgnore(const bool& AttributeValue)
+{
+	IMPLEMENT_NODE_ATTRIBUTE_SETTER_NODELEGATE(VertexColorIgnore, bool)
+}
+
+bool UInterchangeSkeletalMeshFactoryNode::GetCustomVertexColorOverride(FColor& AttributeValue) const
+{
+	IMPLEMENT_NODE_ATTRIBUTE_GETTER(VertexColorOverride, FColor)
+}
+
+bool UInterchangeSkeletalMeshFactoryNode::SetCustomVertexColorOverride(const FColor& AttributeValue)
+{
+	IMPLEMENT_NODE_ATTRIBUTE_SETTER_NODELEGATE(VertexColorOverride, FColor)
+}
+
+void UInterchangeSkeletalMeshFactoryNode::FillAssetClassFromAttribute()
+{
+#if WITH_ENGINE
+	FString OperationName = GetTypeName() + TEXT(".GetAssetClassName");
+	FString ClassName;
+	InterchangePrivateNodeBase::GetCustomAttribute<FString>(*Attributes, ClassNameAttributeKey, OperationName, ClassName);
+	if (ClassName.Equals(USkeletalMesh::StaticClass()->GetName()))
+	{
+		AssetClass = USkeletalMesh::StaticClass();
+		bIsNodeClassInitialized = true;
+	}
+#endif
+}
+
+bool UInterchangeSkeletalMeshFactoryNode::SetNodeClassFromClassAttribute()
+{
+	if (!bIsNodeClassInitialized)
+	{
+		FillAssetClassFromAttribute();
+	}
+	return bIsNodeClassInitialized;
+}
+
+bool UInterchangeSkeletalMeshFactoryNode::IsEditorOnlyDataDefined()
+{
+#if WITH_EDITORONLY_DATA
+	return true;
+#else
+	return false;
+#endif
+}
diff --git a/Engine/Plugins/Experimental/Interchange/Runtime/Source/Nodes/Private/InterchangeSkeletalMeshLodDataNode.cpp b/Engine/Plugins/Experimental/Interchange/Runtime/Source/Nodes/Private/InterchangeSkeletalMeshLodDataNode.cpp
new file mode 100644
--- /dev/null
+++ b/Engine/Plugins/Experimental/Interchange/Runtime/Source/Nodes/Private/InterchangeSkeletalMeshLodDataNode.cpp
@@ -0,0 +1,110 @@
+// Copyright Epic Games, Inc. All Rights Reserved.
+
+#include "InterchangeSkeletalMeshLodDataNode.h"
+
+//Interchange namespace
+namespace UE
+{
+	namespace Interchange
+	{
+		const FString& FSkeletalMeshNodeLodDataStaticData::GetMeshUidsBaseKey()
+		{
+			static FString MeshUids_BaseKey = TEXT("__MeshUids__Key");
+			return MeshUids_BaseKey;
+		}
+	}//ns Interchange
+}//ns UE
+
+UInterchangeSkeletalMeshLodDataNode::UInterchangeSkeletalMeshLodDataNode()
+{
+	MeshUids.Initialize(Attributes, UE::Interchange::FSkeletalMeshNodeLodDataStaticData::GetMeshUidsBaseKey());
+}
+
+/**
+	* Return the node type name of the class, we use this when reporting error
+	*/
+FString UInterchangeSkeletalMeshLodDataNode::GetTypeName() const
+{
+	const FString TypeName = TEXT("SkeletalMeshLodDataNode");
+	return TypeName;
+}
+
+FString UInterchangeSkeletalMeshLodDataNode::GetKeyDisplayName(const UE::Interchange::FAttributeKey& NodeAttributeKey) const
+{
+	FString KeyDisplayName = NodeAttributeKey.Key;
+	if (NodeAttributeKey.Key.Equals(UE::Interchange::FSkeletalMeshNodeLodDataStaticData::GetMeshUidsBaseKey()))
+	{
+		KeyDisplayName = TEXT("Mesh count");
+		return KeyDisplayName;
+	}
+	else if (NodeAttributeKey.Key.StartsWith(UE::Interchange::FSkeletalMeshNodeLodDataStaticData::GetMeshUidsBaseKey()))
+	{
+		KeyDisplayName = TEXT("Mesh index ");
+		const FString IndexKey = UE::Interchange::TArrayAttributeHelper<FString>::IndexKey();
+		int32 IndexPosition = NodeAttributeKey.Key.Find(IndexKey) + IndexKey.Len();
+		if (IndexPosition < NodeAttributeKey.Key.Len())
+		{
+			KeyDisplayName += NodeAttributeKey.Key.RightChop(IndexPosition);
+		}
+		return KeyDisplayName;
+	}
+	else if (NodeAttributeKey == Macro_CustomSkeletonUidKey)
+	{
+		KeyDisplayName = TEXT("Skeleton factory node");
+		return KeyDisplayName;
+	}
+	return Super::GetKeyDisplayName(NodeAttributeKey);
+}
+
+FString UInterchangeSkeletalMeshLodDataNode::GetAttributeCategory(const UE::Interchange::FAttributeKey& NodeAttributeKey) const
+{
+	if (NodeAttributeKey.Key.StartsWith(UE::Interchange::FSkeletalMeshNodeLodDataStaticData::GetMeshUidsBaseKey()))
+	{
+		return FString(TEXT("Meshes"));
+	}
+	return Super::GetAttributeCategory(NodeAttributeKey);
+}
+
+bool UInterchangeSkeletalMeshLodDataNode::GetCustomSkeletonUid(FString& AttributeValue) const
+{
+	IMPLEMENT_NODE_ATTRIBUTE_GETTER(SkeletonUid, FString);
+}
+
+bool UInterchangeSkeletalMeshLodDataNode::SetCustomSkeletonUid(const FString& AttributeValue)
+{
+	IMPLEMENT_NODE_ATTRIBUTE_SETTER_NODELEGATE(SkeletonUid, FString)
+}
+
+int32 UInterchangeSkeletalMeshLodDataNode::GetMeshUidsCount() const
+{
+	return MeshUids.GetCount();
+}
+
+void UInterchangeSkeletalMeshLodDataNode::GetMeshUids(TArray<FString>& OutMeshNames) const
+{
+	MeshUids.GetItems(OutMeshNames);
+}
+
+bool UInterchangeSkeletalMeshLodDataNode::AddMeshUid(const FString& MeshName)
+{
+	return MeshUids.AddItem(MeshName);
+}
+
+bool UInterchangeSkeletalMeshLodDataNode::RemoveMeshUid(const FString& MeshName)
+{
+	return MeshUids.RemoveItem(MeshName);
+}
+
+bool UInterchangeSkeletalMeshLodDataNode::RemoveAllMeshes()
+{
+	return MeshUids.RemoveAllItems();
+}
+
+bool UInterchangeSkeletalMeshLodDataNode::IsEditorOnlyDataDefined()
+{
+#if WITH_EDITORONLY_DATA
+	return true;
+#else
+	return false;
+#endif
+}
diff --git a/Engine/Plugins/Experimental/Interchange/Runtime/Source/Nodes/Private/InterchangeStaticMeshFactoryNode.cpp b/Engine/Plugins/Experimental/Interchange/Runtime/Source/Nodes/Private/InterchangeStaticMeshFactoryNode.cpp
new file mode 100644
--- /dev/null
+++ b/Engine/Plugins/Experimental/Interchange/Runtime/Source/Nodes/Private/InterchangeStaticMeshFactoryNode.cpp
@@ -0,0 +1,170 @@
+// Copyright Epic Games, Inc. All Rights Reserved.
+
+#include "InterchangeStaticMeshFactoryNode.h"
+
+#if WITH_ENGINE
+#include "Engine/StaticMesh.h"
+#endif
+
+namespace UE
+{
+	namespace Interchange
+	{
+		const FString& FStaticMeshNodeStaticData::GetLodDependenciesBaseKey()
+		{
+			static FString LodDependencies_BaseKey = TEXT("Lod_Dependencies");
+			return LodDependencies_BaseKey;
+		}
+	} // namespace Interchange
+} // namespace UE
+
+
+UInterchangeStaticMeshFactoryNode::UInterchangeStaticMeshFactoryNode()
+{
+#if WITH_ENGINE
+	AssetClass = nullptr;
+#endif
+	LodDependencies.Initialize(Attributes, UE::Interchange::FStaticMeshNodeStaticData::GetLodDependenciesBaseKey());
+}
+
+void UInterchangeStaticMeshFactoryNode::InitializeStaticMeshNode(const FString& UniqueID, const FString& DisplayLabel, const FString& InAssetClass)
+{
+	bIsNodeClassInitialized = false;
+	InitializeNode(UniqueID, DisplayLabel, EInterchangeNodeContainerType::FactoryData);
+
+	FString OperationName = GetTypeName() + TEXT(".SetAssetClassName");
+	InterchangePrivateNodeBase::SetCustomAttribute<FString>(*Attributes, ClassNameAttributeKey, OperationName, InAssetClass);
+	FillAssetClassFromAttribute();
+}
+
+void UInterchangeStaticMeshFactoryNode::Serialize(FArchive& Ar)
+{
+	Super::Serialize(Ar);
+#if WITH_ENGINE
+	if (Ar.IsLoading())
+	{
+		// Make sure the class is properly set when we compile with engine, this will set the bIsNodeClassInitialized to true.
+		SetNodeClassFromClassAttribute();
+	}
+#endif //#if WITH_ENGINE
+}
+
+FString UInterchangeStaticMeshFactoryNode::GetTypeName() const
+{
+	const FString TypeName = TEXT("StaticMeshNode");
+	return TypeName;
+}
+
+FString UInterchangeStaticMeshFactoryNode::GetKeyDisplayName(const UE::Interchange::FAttributeKey& NodeAttributeKey) const
+{
+	FString KeyDisplayName = NodeAttributeKey.Key;
+	if (NodeAttributeKey.Key.Equals(UE::Interchange::FStaticMeshNodeStaticData::GetLodDependenciesBaseKey()))
+	{
+		KeyDisplayName = TEXT("LOD Dependencies Count");
+		return KeyDisplayName;
+	}
+	else if (NodeAttributeKey.Key.StartsWith(UE::Interchange::FStaticMeshNodeStaticData::GetLodDependenciesBaseKey()))
+	{
+		KeyDisplayName = TEXT("LOD Dependencies Index ");
+		const FString IndexKey = UE::Interchange::TArrayAttributeHelper<FString>::IndexKey();
+		int32 IndexPosition = NodeAttributeKey.Key.Find(IndexKey) + IndexKey.Len();
+		if (IndexPosition < NodeAttributeKey.Key.Len())
+		{
+			KeyDisplayName += NodeAttributeKey.Key.RightChop(IndexPosition);
+		}
+		return KeyDisplayName;
+	}
+	return Super::GetKeyDisplayName(NodeAttributeKey);
+}
+
+UClass* UInterchangeStaticMeshFactoryNode::GetObjectClass() const
+{
+	ensure(bIsNodeClassInitialized);
+#if WITH_ENGINE
+	return AssetClass.Get() != nullptr ? AssetClass.Get() : UStaticMesh::StaticClass();
+#else
+	return nullptr;
+#endif
+}
+
+int32 UInterchangeStaticMeshFactoryNode::GetLodDataCount() const
+{
+	return LodDependencies.GetCount();
+}
+
+void UInterchangeStaticMeshFactoryNode::GetLodDataUniqueIds(TArray<FString>& OutLodDataUniqueIds) const
+{
+	LodDependencies.GetItems(OutLodDataUniqueIds);
+}
+
+bool UInterchangeStaticMeshFactoryNode::AddLodDataUniqueId(const FString& LodDataUniqueId)
+{
+	return LodDependencies.AddItem(LodDataUniqueId);
+}
+
+bool UInterchangeStaticMeshFactoryNode::RemoveLodDataUniqueId(const FString& LodDataUniqueId)
+{
+	return LodDependencies.RemoveItem(LodDataUniqueId);
+}
+
+bool UInterchangeStaticMeshFactoryNode::GetCustomVertexColorReplace(bool& AttributeValue) const
+{
+	IMPLEMENT_NODE_ATTRIBUTE_GETTER(VertexColorReplace, bool)
+}
+
+bool UInterchangeStaticMeshFactoryNode::SetCustomVertexColorReplace(const bool& AttributeValue)
+{
+	IMPLEMENT_NODE_ATTRIBUTE_SETTER_NODELEGATE(VertexColorReplace, bool)
+}
+
+bool UInterchangeStaticMeshFactoryNode::GetCustomVertexColorIgnore(bool& AttributeValue) const
+{
+	IMPLEMENT_NODE_ATTRIBUTE_GETTER(VertexColorIgnore, bool)
+}
+
+bool UInterchangeStaticMeshFactoryNode::SetCustomVertexColorIgnore(const bool& AttributeValue)
+{
+	IMPLEMENT_NODE_ATTRIBUTE_SETTER_NODELEGATE(VertexColorIgnore, bool)
+}
+
+bool UInterchangeStaticMeshFactoryNode::GetCustomVertexColorOverride(FColor& AttributeValue) const
+{
+	IMPLEMENT_NODE_ATTRIBUTE_GETTER(VertexColorOverride, FColor)
+}
+
+bool UInterchangeStaticMeshFactoryNode::SetCustomVertexColorOverride(const FColor& AttributeValue)
+{
+	IMPLEMENT_NODE_ATTRIBUTE_SETTER_NODELEGATE(VertexColorOverride, FColor)
+}
+
+void UInterchangeStaticMeshFactoryNode::FillAssetClassFromAttribute()
+{
+#if WITH_ENGINE
+	FString OperationName = GetTypeName() + TEXT(".GetAssetClassName");
+	FString ClassName;
+	InterchangePrivateNodeBase::GetCustomAttribute<FString>(*Attributes, ClassNameAttributeKey, OperationName, ClassName);
+	if (ClassName.Equals(UStaticMesh::StaticClass()->GetName()))
+	{
+		AssetClass = UStaticMesh::StaticClass();
+		bIsNodeClassInitialized = true;
+	}
+#endif
+}
+
+bool UInterchangeStaticMeshFactoryNode::SetNodeClassFromClassAttribute()
+{
+	if (!bIsNodeClassInitialized)
+	{
+		FillAssetClassFromAttribute();
+	}
+	return bIsNodeClassInitialized;
+}
+
+bool UInterchangeStaticMeshFactoryNode::IsEditorOnlyDataDefined()
+{
+#if WITH_EDITORONLY_DATA
+	return true;
+#else
+	return false;
+#endif
+}
diff --git a/Engine/Plugins/Experimental/Interchange/Runtime/Source/Nodes/Private/InterchangeStaticMeshLodDataNode.cpp b/Engine/Plugins/Experimental/Interchange/Runtime/Source/Nodes/Private/InterchangeStaticMeshLodDataNode.cpp
new file mode 100644
--- /dev/null
+++ b/Engine/Plugins/Experimental/Interchange/Runtime/Source/Nodes/Private/InterchangeStaticMeshLodDataNode.cpp
@@ -0,0 +1,96 @@
+// Copyright Epic Games, Inc. All Rights Reserved.
+
+#include "InterchangeStaticMeshLodDataNode.h"
+
+namespace UE
+{
+	namespace Interchange
+	{
+		const FString& FStaticMeshNodeLodDataStaticData::GetMeshUidsBaseKey()
+		{
+			static FString MeshUids_BaseKey = TEXT("__MeshUids__Key");
+			return MeshUids_BaseKey;
+		}
+	} // namespace Interchange
+} // namespace UE
+
+
+UInterchangeStaticMeshLodDataNode::UInterchangeStaticMeshLodDataNode()
+{
+	MeshUids.Initialize(Attributes, UE::Interchange::FStaticMeshNodeLodDataStaticData::GetMeshUidsBaseKey());
+}
+
+/**
+	* Return the node type name of the class, we use this when reporting error
+	*/
+FString UInterchangeStaticMeshLodDataNode::GetTypeName() const
+{
+	const FString TypeName = TEXT("StaticMeshLodDataNode");
+	return TypeName;
+}
+
+FString UInterchangeStaticMeshLodDataNode::GetKeyDisplayName(const UE::Interchange::FAttributeKey& NodeAttributeKey) const
+{
+	FString KeyDisplayName = NodeAttributeKey.Key;
+	if (NodeAttributeKey.Key.Equals(UE::Interchange::FStaticMeshNodeLodDataStaticData::GetMeshUidsBaseKey()))
+	{
+		KeyDisplayName = TEXT("Mesh count");
+		return KeyDisplayName;
+	}
+	else if (NodeAttributeKey.Key.StartsWith(UE::Interchange::FStaticMeshNodeLodDataStaticData::GetMeshUidsBaseKey()))
+	{
+		KeyDisplayName = TEXT("Mesh index ");
+		const FString IndexKey = UE::Interchange::TArrayAttributeHelper<FString>::IndexKey();
+		int32 IndexPosition = NodeAttributeKey.Key.Find(IndexKey) + IndexKey.Len();
+		if (IndexPosition < NodeAttributeKey.Key.Len())
+		{
+			KeyDisplayName += NodeAttributeKey.Key.RightChop(IndexPosition);
+		}
+		return KeyDisplayName;
+	}
+
+	return Super::GetKeyDisplayName(NodeAttributeKey);
+}
+
+FString UInterchangeStaticMeshLodDataNode::GetAttributeCategory(const UE::Interchange::FAttributeKey& NodeAttributeKey) const
+{
+	if (NodeAttributeKey.Key.StartsWith(UE::Interchange::FStaticMeshNodeLodDataStaticData::GetMeshUidsBaseKey()))
+	{
+		return FString(TEXT("Meshes"));
+	}
+	return Super::GetAttributeCategory(NodeAttributeKey);
+}
+
+int32 UInterchangeStaticMeshLodDataNode::GetMeshUidsCount() const
+{
+	return MeshUids.GetCount();
+}
+
+void UInterchangeStaticMeshLodDataNode::GetMeshUids(TArray<FString>& OutMeshNames) const
+{
+	MeshUids.GetItems(OutMeshNames);
+}
+
+bool UInterchangeStaticMeshLodDataNode::AddMeshUid(const FString& MeshName)
+{
+	return MeshUids.AddItem(MeshName);
+}
+
+bool UInterchangeStaticMeshLodDataNode::RemoveMeshUid(const FString& MeshName)
+{
+	return MeshUids.RemoveItem(MeshName);
+}
+
+bool UInterchangeStaticMeshLodDataNode::RemoveAllMeshes()
+{
+	return MeshUids.RemoveAllItems();
+}
+
+bool UInterchangeStaticMeshLodDataNode::IsEditorOnlyDataDefined()
+{
+#if WITH_EDITORONLY_DATA
+	return true;
+#else
+	return false;
+#endif
+}
diff --git a/Engine/Plugins/Experimental/Interchange/Runtime/Source/Nodes/Public/InterchangeMaterialFactoryNode.h b/Engine/Plugins/Experimental/Interchange/Runtime/Source/Nodes/Public/InterchangeMaterialFactoryNode.h
--- a/Engine/Plugins/Experimental/Interchange/Runtime/Source/Nodes/Public/InterchangeMaterialFactoryNode.h
+++ b/Engine/Plugins/Experimental/Interchange/Runtime/Source/Nodes/Public/InterchangeMaterialFactoryNode.h
@@ -4,14 +4,7 @@
 
 #include "CoreMinimal.h"
 #include "Nodes/InterchangeBaseNode.h"
-
-#if WITH_ENGINE
-
-#include "Materials/Material.h"
-#include "Materials/MaterialInstance.h"
-#include "Materials/MaterialInterface.h"
-
-#endif
+#include "UObject/UObjectGlobals.h"
 
 #include "InterchangeMaterialFactoryNode.generated.h"
 
@@ -21,14 +14,7 @@ class INTERCHANGENODES_API UInterchangeMaterialFactoryNode : public UInterchange
 	GENERATED_BODY()
 
 public:
-	UInterchangeMaterialFactoryNode()
-	:UInterchangeBaseNode()
-	{
-#if WITH_ENGINE
-		AssetClass = nullptr;
-#endif
-		TextureDependencies.Initialize(Attributes, TextureDependenciesKey.Key);
-	}
+	UInterchangeMaterialFactoryNode();
 
 	/**
 	 * Initialize node data
@@ -37,217 +23,89 @@ public:
 	 * @param InAssetClass - The class the material factory will create for this node.
 	 */
 	UFUNCTION(BlueprintCallable, Category = "Interchange | Node | MaterialFactory")
-	void InitializeMaterialNode(const FString& UniqueID, const FString& DisplayLabel, const FString& InAssetClass)
-	{
-		bIsMaterialNodeClassInitialized = false;
-		InitializeNode(UniqueID, DisplayLabel, EInterchangeNodeContainerType::NodeContainerType_FactoryData);
-		FString OperationName = GetTypeName() + TEXT(".SetAssetClassName");
-		InterchangePrivateNodeBase::SetCustomAttribute<FString>(*Attributes, ClassNameAttributeKey, OperationName, InAssetClass);
-		FillAssetClassFromAttribute();
-	}
+	void InitializeMaterialNode(const FString& UniqueID, const FString& DisplayLabel, const FString& InAssetClass);
 
 	/**
 	 * Return the node type name of the class, we use this when reporting error
 	 */
-	virtual FString GetTypeName() const override
-	{
-		const FString TypeName = TEXT("MaterialFactoryNode");
-		return TypeName;
-	}
+	virtual FString GetTypeName() const override;
 
 	/** Get the class this node want to create */
 	UFUNCTION(BlueprintCallable, Category = "Interchange | Node | MaterialFactory")
-	virtual class UClass* GetObjectClass() const override
-	{
-		ensure(bIsMaterialNodeClassInitialized);
-#if WITH_ENGINE
-		return AssetClass.Get() != nullptr ? AssetClass.Get() : UMaterial::StaticClass();
-#else
-		return nullptr;
-#endif
-	}
-
-	virtual FString GetKeyDisplayName(const UE::Interchange::FAttributeKey& NodeAttributeKey) const override
-	{
-		FString KeyDisplayName = NodeAttributeKey.ToString();
-		if (NodeAttributeKey == TextureDependenciesKey)
-		{
-			return KeyDisplayName = TEXT("Texture Dependencies count");
-		}
-		else if (NodeAttributeKey.Key.StartsWith(TextureDependenciesKey.Key))
-		{
-			KeyDisplayName = TEXT("Texture Dependencies Index ");
-			const FString IndexKey = UE::Interchange::FNameAttributeArrayHelper::IndexKey();
-			int32 IndexPosition = NodeAttributeKey.Key.Find(IndexKey) + IndexKey.Len();
-			if (IndexPosition < NodeAttributeKey.Key.Len())
-			{
-				KeyDisplayName += NodeAttributeKey.Key.RightChop(IndexPosition);
-			}
-			return KeyDisplayName;
-		}
-		return Super::GetKeyDisplayName(NodeAttributeKey);
-	}
-
-	virtual FString GetAttributeCategory(const UE::Interchange::FAttributeKey& NodeAttributeKey) const override
-	{
-		if (NodeAttributeKey == Macro_CustomTranslatedMaterialNodeUidKey
-			|| NodeAttributeKey == Macro_CustomMaterialUsageKey
-			|| NodeAttributeKey == Macro_CustomBlendModeKey)
-		{
-			return FString(TEXT("MaterialFactory"));
-		}
-		else if (NodeAttributeKey.Key.StartsWith(TextureDependenciesKey.Key))
-		{
-			return FString(TEXT("TextureDependencies"));
-		}
-		return Super::GetAttributeCategory(NodeAttributeKey);
-	}
-
-	virtual FGuid GetHash() const override
-	{
-		return Attributes->GetStorageHash();
-	}
-
-	static FString GetMaterialFactoryNodeUidFromMaterialNodeUid(const FString& TranslatedNodeUid)
-	{
-		FString NewUid = TEXT("Factory_") + TranslatedNodeUid;
-		return NewUid;
-	}
+	virtual class UClass* GetObjectClass() const override;
+
+	virtual FString GetKeyDisplayName(const UE::Interchange::FAttributeKey& NodeAttributeKey) const override;
+
+	virtual FString GetAttributeCategory(const UE::Interchange::FAttributeKey& NodeAttributeKey) const override;
+
+	static FString GetMaterialFactoryNodeUidFromMaterialNodeUid(const FString& TranslatedNodeUid);
 
 	/**
 	 * This function allow to retrieve the number of Texture dependencies for this object.
 	 *
 	 */
 	UFUNCTION(BlueprintCallable, Category = "Interchange | Node | MaterialFactory")
-	int32 GetTextureDependeciesCount() const
-	{
-		return TextureDependencies.GetCount();
-	}
+	int32 GetTextureDependeciesCount() const;
 
 	/**
 	 * This function allow to retrieve the Texture dependency for this object.
 	 *
 	 */
 	UFUNCTION(BlueprintCallable, Category = "Interchange | Node | MaterialFactory")
-	void GetTextureDependencies(TArray<FString>& OutDependencies) const
-	{
-		TextureDependencies.GetNames(OutDependencies);
-	}
+	void GetTextureDependencies(TArray<FString>& OutDependencies) const;
 
 	/**
 	 * This function allow to retrieve one Texture dependency for this object.
 	 *
 	 */
 	UFUNCTION(BlueprintCallable, Category = "Interchange | Node | MaterialFactory")
-	void GetTextureDependency(const int32 Index, FString& OutDependency) const
-	{
-		TextureDependencies.GetName(Index, OutDependency);
-	}
+	void GetTextureDependency(const int32 Index, FString& OutDependency) const;
 
 	/**
 	 * Add one Texture dependency to this object.
 	 */
 	UFUNCTION(BlueprintCallable, Category = "Interchange | Node | MaterialFactory")
-	bool SetTextureDependencyUid(const FString& DependencyUid)
-	{
-		return TextureDependencies.AddName(DependencyUid);
-	}
+	bool SetTextureDependencyUid(const FString& DependencyUid);
 
 	/**
 	 * Remove one Texture dependency from this object.
 	 */
 	UFUNCTION(BlueprintCallable, Category = "Interchange | Node | MaterialFactory")
-	bool RemoveTextureDependencyUid(const FString& DependencyUid)
-	{
-		return TextureDependencies.RemoveName(DependencyUid);
-	}
+	bool RemoveTextureDependencyUid(const FString& DependencyUid);
 
 	/**
 	 * Get the translated texture node unique ID.
 	 */
 	UFUNCTION(BlueprintCallable, Category = "Interchange | Node | MaterialFactory")
-	bool GetCustomTranslatedMaterialNodeUid(FString& AttributeValue) const
-	{
-		IMPLEMENT_NODE_ATTRIBUTE_GETTER(TranslatedMaterialNodeUid, FString);
-	}
+	bool GetCustomTranslatedMaterialNodeUid(FString& AttributeValue) const;
 
 	/**
 	 * Set the translated texture node unique ID. This is the reference to the node that was create by the translator and this node is needed to get the texture payload.
 	 */
 	UFUNCTION(BlueprintCallable, Category = "Interchange | Node | MaterialFactory")
-	bool SetCustomTranslatedMaterialNodeUid(const FString& AttributeValue)
-	{
-		IMPLEMENT_NODE_ATTRIBUTE_SETTER_NODELEGATE(TranslatedMaterialNodeUid, FString);
-	}
+	bool SetCustomTranslatedMaterialNodeUid(const FString& AttributeValue);
 
 	/** Return false if the Attribute was not set previously.*/
 	UFUNCTION(BlueprintCallable, Category = "Interchange | Node | MaterialFactory")
-	bool GetCustomMaterialUsage(uint8& AttributeValue) const
-	{
-		IMPLEMENT_NODE_ATTRIBUTE_GETTER(MaterialUsage, uint8);
-	}
+	bool GetCustomMaterialUsage(uint8& AttributeValue) const;
 
 	UFUNCTION(BlueprintCallable, Category = "Interchange | Node | MaterialFactory")
-	bool SetCustomMaterialUsage(const uint8& AttributeValue, bool bAddApplyDelegate = true)
-	{
-		IMPLEMENT_NODE_ATTRIBUTE_SETTER(UInterchangeMaterialFactoryNode, MaterialUsage, uint8, UMaterialInterface);
-	}
+	bool SetCustomMaterialUsage(const uint8& AttributeValue, bool bAddApplyDelegate = true);
 
 	/** Return false if the Attribute was not set previously.*/
 	UFUNCTION(BlueprintCallable, Category = "Interchange | Node | MaterialFactory")
-	bool GetCustomBlendMode(uint8& AttributeValue) const
-	{
-		IMPLEMENT_NODE_ATTRIBUTE_GETTER(BlendMode, uint8);
-	}
+	bool GetCustomBlendMode(uint8& AttributeValue) const;
 
 	UFUNCTION(BlueprintCallable, Category = "Interchange | Node | MaterialFactory")
-	bool SetCustomBlendMode(const uint8& AttributeValue, bool bAddApplyDelegate = true)
-	{
-		IMPLEMENT_NODE_ATTRIBUTE_SETTER(UInterchangeMaterialFactoryNode, BlendMode, uint8, UMaterialInterface);
-	}
-
-	virtual void Serialize(FArchive& Ar) override
-	{
-		Super::Serialize(Ar);
-#if WITH_ENGINE
-		if (Ar.IsLoading())
-		{
-			//Make sure the class is properly set when we compile with engine, this will set the
-			//bIsMaterialNodeClassInitialized to true.
-			SetMaterialNodeClassFromClassAttribute();
-		}
-#endif //#if WITH_ENGINE
-	}
+	bool SetCustomBlendMode(const uint8& AttributeValue, bool bAddApplyDelegate = true);
+
+	virtual void Serialize(FArchive& Ar) override;
 
 private:
 
-	void FillAssetClassFromAttribute()
-	{
-#if WITH_ENGINE
-		FString OperationName = GetTypeName() + TEXT(".GetAssetClassName");
-		FString ClassName;
-		InterchangePrivateNodeBase::GetCustomAttribute<FString>(*Attributes, ClassNameAttributeKey, OperationName, ClassName);
-		if (ClassName.Equals(UMaterial::StaticClass()->GetName()))
-		{
-			AssetClass = UMaterial::StaticClass();
-			bIsMaterialNodeClassInitialized = true;
-		}
-		else if (ClassName.Equals(UMaterialInstance::StaticClass()->GetName()))
-		{
-			AssetClass = UMaterialInstance::StaticClass();
-			bIsMaterialNodeClassInitialized = true;
-		}
-#endif
-	}
+	void FillAssetClassFromAttribute();
 
-	bool SetMaterialNodeClassFromClassAttribute()
-	{
-		if (!bIsMaterialNodeClassInitialized)
-		{
-			FillAssetClassFromAttribute();
-		}
-		return bIsMaterialNodeClassInitialized;
-	}
+	bool SetMaterialNodeClassFromClassAttribute();
 
 	const UE::Interchange::FAttributeKey ClassNameAttributeKey = UE::Interchange::FBaseNodeStaticData::ClassTypeAttributeKey();
 	
@@ -265,55 +123,9 @@ private:
 
 	//Material Ussage is not handle properly, you can have several usage check
 	//Currently I can set only one and retrieve the first, TODO make this work (maybe one key per usage...)
-	bool ApplyCustomMaterialUsageToAsset(UObject * Asset) const
-	{
-		if (!Asset)
-		{
-			return false;
-		}
-		UMaterialInterface* TypedObject = Cast<UMaterialInterface>(Asset);
-		if (!TypedObject)
-		{
-			return false;
-		}
-		uint8 ValueData;
-		if (GetCustomMaterialUsage(ValueData))
-		{
-			TypedObject->CheckMaterialUsage_Concurrent(EMaterialUsage(ValueData));
-			return true;
-		}
-		return false;
-	}
-
-	bool FillCustomMaterialUsageFromAsset(UObject* Asset)
-	{
-		if (!Asset)
-		{
-			return false;
-		}
-		UMaterialInterface* TypedObject = Cast<UMaterialInterface>(Asset);
-		if (!TypedObject)
-		{
-			return false;
-		}
-		const UMaterial* Material = TypedObject->GetMaterial_Concurrent();
-		if (!Material)
-		{
-			return false;
-		}
-		for (int32 Usage = 0; Usage < MATUSAGE_MAX; Usage++)
-		{
-			const EMaterialUsage UsageEnum = (EMaterialUsage)Usage;
-			if (Material->GetUsageByFlag(UsageEnum))
-			{
-				if (SetCustomMaterialUsage((int8)UsageEnum, false))
-				{
-					return true;
-				}
-			}
-		}
-		return false;
-	}
+	bool ApplyCustomMaterialUsageToAsset(UObject* Asset) const;
+
+	bool FillCustomMaterialUsageFromAsset(UObject* Asset);
 #endif
 
 	IMPLEMENT_NODE_ATTRIBUTE_APPLY_UOBJECT(BlendMode, uint8, UMaterial);
@@ -324,5 +136,5 @@ protected:
 #endif
 	bool bIsMaterialNodeClassInitialized = false;
 
-	UE::Interchange::FNameAttributeArrayHelper TextureDependencies;
+	UE::Interchange::TArrayAttributeHelper<FString> TextureDependencies;
 };
diff --git a/Engine/Plugins/Experimental/Interchange/Runtime/Source/Nodes/Public/InterchangeMaterialNode.h b/Engine/Plugins/Experimental/Interchange/Runtime/Source/Nodes/Public/InterchangeMaterialNode.h
--- a/Engine/Plugins/Experimental/Interchange/Runtime/Source/Nodes/Public/InterchangeMaterialNode.h
+++ b/Engine/Plugins/Experimental/Interchange/Runtime/Source/Nodes/Public/InterchangeMaterialNode.h
@@ -4,15 +4,7 @@
 
 #include "CoreMinimal.h"
 #include "Nodes/InterchangeBaseNode.h"
-
-#if WITH_ENGINE
-
-#include "Materials/Material.h"
-#include "Materials/MaterialInstance.h"
-#include "Materials/MaterialInterface.h"
-#include "Engine/Texture.h"
-
-#endif
+#include "UObject/UObjectGlobals.h"
 
 #include "InterchangeMaterialNode.generated.h"
 
@@ -21,7 +13,7 @@ namespace UE
 {
 	namespace Interchange
 	{
-		struct FMaterialNodeStaticData : public FBaseNodeStaticData
+		struct INTERCHANGENODES_API FMaterialNodeStaticData : public FBaseNodeStaticData
 		{
 			static const FAttributeKey& PayloadSourceFileKey()
 			{
@@ -54,138 +46,54 @@ class INTERCHANGENODES_API UInterchangeMaterialNode : public UInterchangeBaseNod
 	GENERATED_BODY()
 
 public:
-	UInterchangeMaterialNode()
-	:UInterchangeBaseNode()
-	{
-		TextureDependencies.Initialize(Attributes, TextureDependenciesKey.Key);
-	}
+	UInterchangeMaterialNode();
 
 	/**
 	 * Return the node type name of the class, we use this when reporting error
 	 */
-	virtual FString GetTypeName() const override
-	{
-		const FString TypeName = TEXT("MaterialNode");
-		return TypeName;
-	}
+	virtual FString GetTypeName() const override;
 
-	virtual FString GetKeyDisplayName(const UE::Interchange::FAttributeKey& NodeAttributeKey) const override
-	{
-		FString KeyDisplayName = NodeAttributeKey.ToString();
-		if (NodeAttributeKey == UE::Interchange::FMaterialNodeStaticData::PayloadSourceFileKey())
-		{
-			return KeyDisplayName = TEXT("Payload Source Key");
-		}
-		else if (NodeAttributeKey == TextureDependenciesKey)
-		{
-			return KeyDisplayName = TEXT("Texture Dependencies count");
-		}
-		else if (NodeAttributeKey.Key.StartsWith(TextureDependenciesKey.Key))
-		{
-			KeyDisplayName = TEXT("Texture Dependencies Index ");
-			const FString IndexKey = UE::Interchange::FNameAttributeArrayHelper::IndexKey();
-			int32 IndexPosition = NodeAttributeKey.Key.Find(IndexKey) + IndexKey.Len();
-			if (IndexPosition < NodeAttributeKey.Key.Len())
-			{
-				KeyDisplayName += NodeAttributeKey.Key.RightChop(IndexPosition);
-			}
-			return KeyDisplayName;
-		}
-		return Super::GetKeyDisplayName(NodeAttributeKey);
-	}
+	virtual FString GetKeyDisplayName(const UE::Interchange::FAttributeKey& NodeAttributeKey) const override;
 
-	virtual FString GetAttributeCategory(const UE::Interchange::FAttributeKey& NodeAttributeKey) const override
-	{
-		if (NodeAttributeKey.Key.StartsWith(TextureDependenciesKey.Key))
-		{
-			return FString(TEXT("TextureDependencies"));
-		}
-		return Super::GetAttributeCategory(NodeAttributeKey);
-	}
-
-	virtual FGuid GetHash() const override
-	{
-		return Attributes->GetStorageHash();
-	}
+	virtual FString GetAttributeCategory(const UE::Interchange::FAttributeKey& NodeAttributeKey) const override;
 
 	/**
 	 * This function allow to retrieve the number of Texture dependencies for this object.
 	 *
 	 */
 	UFUNCTION(BlueprintCallable, Category = "Interchange | Node | Material")
-	int32 GetTextureDependeciesCount() const
-	{
-		return TextureDependencies.GetCount();
-	}
+	int32 GetTextureDependeciesCount() const;
 
 	/**
 	 * This function allow to retrieve the Texture dependency for this object.
 	 *
 	 */
 	UFUNCTION(BlueprintCallable, Category = "Interchange | Node | Material")
-	void GetTextureDependencies(TArray<FString>& OutDependencies) const
-	{
-		TextureDependencies.GetNames(OutDependencies);
-	}
+	void GetTextureDependencies(TArray<FString>& OutDependencies) const;
 
 	/**
 	 * This function allow to retrieve one Texture dependency for this object.
 	 *
 	 */
 	UFUNCTION(BlueprintCallable, Category = "Interchange | Node | Material")
-	void GetTextureDependency(const int32 Index, FString& OutDependency) const
-	{
-		TextureDependencies.GetName(Index, OutDependency);
-	}
+	void GetTextureDependency(const int32 Index, FString& OutDependency) const;
 
 	/**
 	 * Add one Texture dependency to this object.
 	 */
 	UFUNCTION(BlueprintCallable, Category = "Interchange | Node | Material")
-	bool SetTextureDependencyUid(const FString& DependencyUid)
-	{
-		return TextureDependencies.AddName(DependencyUid);
-	}
+	bool SetTextureDependencyUid(const FString& DependencyUid);
 
 	/**
 	 * Remove one Texture dependency from this object.
 	 */
 	UFUNCTION(BlueprintCallable, Category = "Interchange | Node | Material")
-	bool RemoveTextureDependencyUid(const FString& DependencyUid)
-	{
-		return TextureDependencies.RemoveName(DependencyUid);
-	}
+	bool RemoveTextureDependencyUid(const FString& DependencyUid);
 
-	virtual const TOptional<FString> GetPayLoadKey() const
-	{
-		if (!Attributes->ContainAttribute(UE::Interchange::FMaterialNodeStaticData::PayloadSourceFileKey()))
-		{
-			return TOptional<FString>();
-		}
-		UE::Interchange::FAttributeStorage::TAttributeHandle<FString> AttributeHandle = Attributes->GetAttributeHandle<FString>(UE::Interchange::FMaterialNodeStaticData::PayloadSourceFileKey());
-		if (!AttributeHandle.IsValid())
-		{
-			return TOptional<FString>();
-		}
-		FString PayloadKey;
-		UE::Interchange::EAttributeStorageResult Result = AttributeHandle.Get(PayloadKey);
-		if (!IsAttributeStorageResultSuccess(Result))
-		{
-			LogAttributeStorageErrors(Result, TEXT("UInterchangeMaterialNode.GetPayLoadKey"), UE::Interchange::FMaterialNodeStaticData::PayloadSourceFileKey());
-			return TOptional<FString>();
-		}
-		return TOptional<FString>(PayloadKey);
-	}
+	virtual const TOptional<FString> GetPayLoadKey() const;
 
 	UFUNCTION(BlueprintCallable, Category = "Interchange | Node | Material")
-	virtual void SetPayLoadKey(const FString& PayloadKey)
-	{
-		UE::Interchange::EAttributeStorageResult Result = Attributes->RegisterAttribute(UE::Interchange::FMaterialNodeStaticData::PayloadSourceFileKey(), PayloadKey);
-		if (!IsAttributeStorageResultSuccess(Result))
-		{
-			LogAttributeStorageErrors(Result, TEXT("UInterchangeMaterialNode.SetPayLoadKey"), UE::Interchange::FMaterialNodeStaticData::PayloadSourceFileKey());
-		}
-	}
+	virtual void SetPayLoadKey(const FString& PayloadKey);
 
 	//////////////////////////////////////////////////////////////////////////
 	/**
@@ -204,117 +112,24 @@ public:
 	 * @note - A parameter name can have only one of the 3 type of input set, the last input type set is the one that will be created.
 	 */
 	UFUNCTION(BlueprintCallable, Category = "Interchange | Node | Material")
-	void AddTextureParameterData(const EInterchangeMaterialNodeParameterName ParameterName, const FString& TextureUid, const int32 UVSetIndex, const float ScaleU, const float ScaleV)
-	{
-		FParameterData& ParameterData = ParameterDatas.FindOrAdd(ParameterName);
-		ParameterData.bIsTextureParameter = true;
-		ParameterData.bIsVectorParameter = false;
-		ParameterData.bIsScalarParameter = false;
-		ParameterData.TextureUid = TextureUid;
-		ParameterData.UVSetIndex = UVSetIndex;
-		ParameterData.ScaleU = ScaleU;
-		ParameterData.ScaleV = ScaleV;
-	}
+	void AddTextureParameterData(const EInterchangeMaterialNodeParameterName ParameterName, const FString& TextureUid, const int32 UVSetIndex, const float ScaleU, const float ScaleV);
 
 	UFUNCTION(BlueprintCallable, Category = "Interchange | Node | Material")
-	bool GetTextureParameterData(const EInterchangeMaterialNodeParameterName ParameterName, FString& OutTextureUid, int32& OutUVSetIndex, float& OutScaleU, float& OutScaleV) const
-	{
-		if (const FParameterData* ParameterData = ParameterDatas.Find(ParameterName))
-		{
-			if (!ParameterData->bIsTextureParameter)
-			{
-				return false;
-			}
-			OutTextureUid = ParameterData->TextureUid;
-			OutUVSetIndex = ParameterData->UVSetIndex;
-			OutScaleU = ParameterData->ScaleU;
-			OutScaleV = ParameterData->ScaleV;
-			return true;
-		}
-
-		return false;
-	}
+	bool GetTextureParameterData(const EInterchangeMaterialNodeParameterName ParameterName, FString& OutTextureUid, int32& OutUVSetIndex, float& OutScaleU, float& OutScaleV) const;
 
 	UFUNCTION(BlueprintCallable, Category = "Interchange | Node | Material")
-	void AddVectorParameterData(const EInterchangeMaterialNodeParameterName ParameterName, const FVector& VectorData)
-	{
-		FParameterData& ParameterData = ParameterDatas.FindOrAdd(ParameterName);
-		ParameterData.bIsTextureParameter = false;
-		ParameterData.bIsVectorParameter = true;
-		ParameterData.bIsScalarParameter = false;
-		ParameterData.VectorParameter = VectorData;
-	}
+	void AddVectorParameterData(const EInterchangeMaterialNodeParameterName ParameterName, const FVector& VectorData);
 
 	UFUNCTION(BlueprintCallable, Category = "Interchange | Node | Material")
-	bool GetVectorParameterData(const EInterchangeMaterialNodeParameterName ParameterName, FVector& OutVectorData) const
-	{
-		if (const FParameterData* ParameterData = ParameterDatas.Find(ParameterName))
-		{
-			if (!ParameterData->bIsVectorParameter)
-			{
-				return false;
-			}
-			OutVectorData = ParameterData->VectorParameter;
-			return true;
-		}
-		return false;
-	}
+	bool GetVectorParameterData(const EInterchangeMaterialNodeParameterName ParameterName, FVector& OutVectorData) const;
 
 	UFUNCTION(BlueprintCallable, Category = "Interchange | Node | Material")
-	void AddScalarParameterData(const EInterchangeMaterialNodeParameterName ParameterName, const float ScalarData)
-	{
-		FParameterData& ParameterData = ParameterDatas.FindOrAdd(ParameterName);
-		ParameterData.bIsTextureParameter = false;
-		ParameterData.bIsVectorParameter = false;
-		ParameterData.bIsScalarParameter = true;
-		ParameterData.ScalarParameter = ScalarData;
-	}
+	void AddScalarParameterData(const EInterchangeMaterialNodeParameterName ParameterName, const float ScalarData);
 
 	UFUNCTION(BlueprintCallable, Category = "Interchange | Node | Material")
-	bool GetScalarParameterData(const EInterchangeMaterialNodeParameterName ParameterName, float& OutScalarData) const
-	{
-		if (const FParameterData* ParameterData = ParameterDatas.Find(ParameterName))
-		{
-			if (!ParameterData->bIsScalarParameter)
-			{
-				return false;
-			}
-			OutScalarData = ParameterData->ScalarParameter;
-			return true;
-		}
-		return false;
-	}
-
-	virtual void Serialize(FArchive& Ar) override
-	{
-		Super::Serialize(Ar);
-
-		if (Ar.IsLoading())
-		{
-			ParameterDatas.Reset();
-		}
-		int32 ParamCount = ParameterDatas.Num();
-		Ar << ParamCount;
+	bool GetScalarParameterData(const EInterchangeMaterialNodeParameterName ParameterName, float& OutScalarData) const;
 
-		if (Ar.IsSaving())
-		{
-			for (TPair<EInterchangeMaterialNodeParameterName, FParameterData> ParamPair : ParameterDatas)
-			{
-				Ar << ParamPair.Key;
-				Ar << ParamPair.Value;
-			}
-		}
-		else if (Ar.IsLoading())
-		{
-			for (int32 ParamIndex = 0; ParamIndex < ParamCount; ++ParamIndex)
-			{
-				EInterchangeMaterialNodeParameterName ParamKey;
-				Ar << ParamKey;
-				FParameterData& ParamValue = ParameterDatas.FindOrAdd(ParamKey);
-				Ar << ParamValue;
-			}
-		}
-	}
+	virtual void Serialize(FArchive& Ar) override;
 
 private:
 	const UE::Interchange::FAttributeKey TextureDependenciesKey = UE::Interchange::FAttributeKey(TEXT("__TextureDependenciesKey__"));
@@ -353,5 +168,5 @@ private:
 	//This member is serialize manually in the serialize override
 	TMap<EInterchangeMaterialNodeParameterName, FParameterData> ParameterDatas;
 
-	UE::Interchange::FNameAttributeArrayHelper TextureDependencies;
+	UE::Interchange::TArrayAttributeHelper<FString> TextureDependencies;
 };
diff --git a/Engine/Plugins/Experimental/Interchange/Runtime/Source/Nodes/Public/InterchangeMeshNode.h b/Engine/Plugins/Experimental/Interchange/Runtime/Source/Nodes/Public/InterchangeMeshNode.h
--- a/Engine/Plugins/Experimental/Interchange/Runtime/Source/Nodes/Public/InterchangeMeshNode.h
+++ b/Engine/Plugins/Experimental/Interchange/Runtime/Source/Nodes/Public/InterchangeMeshNode.h
@@ -13,53 +13,16 @@ namespace UE
 	namespace Interchange
 	{
 
-		struct FMeshNodeStaticData : public FBaseNodeStaticData
-		{
-			static const FAttributeKey& PayloadSourceFileKey()
-			{
-				static FAttributeKey AttributeKey(TEXT("__PayloadSourceFile__"));
-				return AttributeKey;
-			}
-
-			static const FAttributeKey& IsSkinnedMeshKey()
-			{
-				static FAttributeKey AttributeKey(TEXT("__IsSkinnedMeshKey__"));
-				return AttributeKey;
-			}
-
-			static const FAttributeKey& IsBlendShapeKey()
-			{
-				static FAttributeKey AttributeKey(TEXT("__IsBlendShapeKey__"));
-				return AttributeKey;
-			}
-
-			static const FAttributeKey& BlendShapeNameKey()
-			{
-				static FAttributeKey AttributeKey(TEXT("__BlendShapeNameKey__"));
-				return AttributeKey;
-			}
-			static const FString& GetSkeletonDependenciesKey()
-			{
-				static FString Dependencies_BaseKey = TEXT("__MeshSkeletonDependencies__");
-				return Dependencies_BaseKey;
-			}
-			static const FString& GetMaterialDependenciesKey()
-			{
-				static FString Dependencies_BaseKey = TEXT("__MeshMaterialDependencies__");
-				return Dependencies_BaseKey;
-			}
-
-			static const FString& GetShapeDependenciesKey()
-			{
-				static FString Dependencies_BaseKey = TEXT("__MeshShapeDependencies__");
-				return Dependencies_BaseKey;
-			}
-
-			static const FString& GetSceneInstancesUidsKey()
-			{
-				static FString SceneInstanceUids_BaseKey = TEXT("__MeshSceneInstancesUids__");
-				return SceneInstanceUids_BaseKey;
-			}
+		struct INTERCHANGENODES_API FMeshNodeStaticData : public FBaseNodeStaticData
+		{
+			static const FAttributeKey& PayloadSourceFileKey();
+			static const FAttributeKey& IsSkinnedMeshKey();
+			static const FAttributeKey& IsBlendShapeKey();
+			static const FAttributeKey& BlendShapeNameKey();
+			static const FString& GetSkeletonDependenciesKey();
+			static const FString& GetMaterialDependenciesKey();
+			static const FString& GetShapeDependenciesKey();
+			static const FString& GetSceneInstancesUidsKey();
 		};
 
 	}//ns Interchange
@@ -71,610 +34,257 @@ class INTERCHANGENODES_API UInterchangeMeshNode : public UInterchangeBaseNode
 	GENERATED_BODY()
 
 public:
-	UInterchangeMeshNode()
-		:UInterchangeBaseNode()
-	{
-		SkeletonDependencies.Initialize(Attributes, UE::Interchange::FMeshNodeStaticData::GetSkeletonDependenciesKey());
-		MaterialDependencies.Initialize(Attributes, UE::Interchange::FMeshNodeStaticData::GetMaterialDependenciesKey());
-		ShapeDependencies.Initialize(Attributes, UE::Interchange::FMeshNodeStaticData::GetShapeDependenciesKey());
-		SceneInstancesUids.Initialize(Attributes, UE::Interchange::FMeshNodeStaticData::GetSceneInstancesUidsKey());
-	}
+	UInterchangeMeshNode();
 
-	virtual FString GetKeyDisplayName(const UE::Interchange::FAttributeKey& NodeAttributeKey) const override
-	{
-		FString KeyDisplayName = NodeAttributeKey.ToString();
-		if (NodeAttributeKey == UE::Interchange::FMeshNodeStaticData::PayloadSourceFileKey())
-		{
-			return KeyDisplayName = TEXT("Payload Source Key");
-		}
-		else if (NodeAttributeKey == UE::Interchange::FMeshNodeStaticData::IsSkinnedMeshKey())
-		{
-			return KeyDisplayName = TEXT("Is a Skinned Mesh");
-		}
-		else if (NodeAttributeKey == UE::Interchange::FMeshNodeStaticData::IsBlendShapeKey())
-		{
-			return KeyDisplayName = TEXT("Is a Blend Shape");
-		}
-		else if (NodeAttributeKey == UE::Interchange::FMeshNodeStaticData::BlendShapeNameKey())
-		{
-			return KeyDisplayName = TEXT("Blend Shape Name");
-		}
-		else if (NodeAttributeKey.Key.Equals(UE::Interchange::FMeshNodeStaticData::GetSkeletonDependenciesKey()))
-		{
-			return KeyDisplayName = TEXT("Skeleton Dependencies count");
-		}
-		else if (NodeAttributeKey.Key.StartsWith(UE::Interchange::FMeshNodeStaticData::GetSkeletonDependenciesKey()))
-		{
-			KeyDisplayName = TEXT("Skeleton Dependencies Index ");
-			const FString IndexKey = UE::Interchange::FNameAttributeArrayHelper::IndexKey();
-			int32 IndexPosition = NodeAttributeKey.Key.Find(IndexKey) + IndexKey.Len();
-			if (IndexPosition < NodeAttributeKey.Key.Len())
-			{
-				KeyDisplayName += NodeAttributeKey.Key.RightChop(IndexPosition);
-			}
-			return KeyDisplayName;
-		}
-		else if (NodeAttributeKey.Key.Equals(UE::Interchange::FMeshNodeStaticData::GetMaterialDependenciesKey()))
-		{
-			return KeyDisplayName = TEXT("Material Dependencies count");
-		}
-		else if (NodeAttributeKey.Key.StartsWith(UE::Interchange::FMeshNodeStaticData::GetMaterialDependenciesKey()))
-		{
-			KeyDisplayName = TEXT("Material Dependencies Index ");
-			const FString IndexKey = UE::Interchange::FNameAttributeArrayHelper::IndexKey();
-			int32 IndexPosition = NodeAttributeKey.Key.Find(IndexKey) + IndexKey.Len();
-			if (IndexPosition < NodeAttributeKey.Key.Len())
-			{
-				KeyDisplayName += NodeAttributeKey.Key.RightChop(IndexPosition);
-			}
-			return KeyDisplayName;
-		}
-		else if (NodeAttributeKey.Key.Equals(UE::Interchange::FMeshNodeStaticData::GetShapeDependenciesKey()))
-		{
-			return KeyDisplayName = TEXT("Shape Dependencies count");
-		}
-		else if (NodeAttributeKey.Key.StartsWith(UE::Interchange::FMeshNodeStaticData::GetShapeDependenciesKey()))
-		{
-			KeyDisplayName = TEXT("Shape Dependencies Index ");
-			const FString IndexKey = UE::Interchange::FNameAttributeArrayHelper::IndexKey();
-			int32 IndexPosition = NodeAttributeKey.Key.Find(IndexKey) + IndexKey.Len();
-			if (IndexPosition < NodeAttributeKey.Key.Len())
-			{
-				KeyDisplayName += NodeAttributeKey.Key.RightChop(IndexPosition);
-			}
-			return KeyDisplayName;
-		}
-		else if (NodeAttributeKey.Key.Equals(UE::Interchange::FMeshNodeStaticData::GetSceneInstancesUidsKey()))
-		{
-			return KeyDisplayName = TEXT("Scene mesh instances count");
-		}
-		else if (NodeAttributeKey.Key.StartsWith(UE::Interchange::FMeshNodeStaticData::GetSceneInstancesUidsKey()))
-		{
-			KeyDisplayName = TEXT("Scene mesh instances Index ");
-			const FString IndexKey = UE::Interchange::FNameAttributeArrayHelper::IndexKey();
-			int32 IndexPosition = NodeAttributeKey.Key.Find(IndexKey) + IndexKey.Len();
-			if (IndexPosition < NodeAttributeKey.Key.Len())
-			{
-				KeyDisplayName += NodeAttributeKey.Key.RightChop(IndexPosition);
-			}
-			return KeyDisplayName;
-		}
-		return Super::GetKeyDisplayName(NodeAttributeKey);
-	}
-
-	virtual FString GetAttributeCategory(const UE::Interchange::FAttributeKey& NodeAttributeKey) const override
-	{
-		if (NodeAttributeKey.Key.StartsWith(UE::Interchange::FMeshNodeStaticData::GetSkeletonDependenciesKey()))
-		{
-			return FString(TEXT("SkeletonDependencies"));
-		}
-		else if (NodeAttributeKey.Key.StartsWith(UE::Interchange::FMeshNodeStaticData::GetMaterialDependenciesKey()))
-		{
-			return FString(TEXT("MaterialDependencies"));
-		}
-		else if (NodeAttributeKey.Key.StartsWith(UE::Interchange::FMeshNodeStaticData::GetShapeDependenciesKey()))
-		{
-			return FString(TEXT("ShapeDependencies"));
-		}
-		else if (NodeAttributeKey.Key.StartsWith(UE::Interchange::FMeshNodeStaticData::GetSceneInstancesUidsKey()))
-		{
-			return FString(TEXT("SceneInstances"));
-		}
-		else if (NodeAttributeKey == Macro_CustomVertexCountKey
-				 || NodeAttributeKey == Macro_CustomPolygonCountKey
-				 || NodeAttributeKey == Macro_CustomBoundingBoxKey
-				 || NodeAttributeKey == Macro_CustomHasVertexNormalKey
-				 || NodeAttributeKey == Macro_CustomHasVertexBinormalKey
-				 || NodeAttributeKey == Macro_CustomHasVertexTangentKey
-				 || NodeAttributeKey == Macro_CustomHasSmoothGroupKey
-				 || NodeAttributeKey == Macro_CustomHasVertexColorKey
-				 || NodeAttributeKey == Macro_CustomUVCountKey)
-		{
-			return FString(TEXT("MeshInfo"));
-		}
-		return Super::GetAttributeCategory(NodeAttributeKey);
-	}
+	virtual FString GetKeyDisplayName(const UE::Interchange::FAttributeKey& NodeAttributeKey) const override;
+
+	virtual FString GetAttributeCategory(const UE::Interchange::FAttributeKey& NodeAttributeKey) const override;
 
 	/**
 	 * Return the node type name of the class, we use this when reporting error
 	 */
-	virtual FString GetTypeName() const override
-	{
-		const FString TypeName = TEXT("MeshNode");
-		return TypeName;
-	}
-
-	virtual FGuid GetHash() const override
-	{
-		return Attributes->GetStorageHash();
-	}
+	virtual FString GetTypeName() const override;
 
 	/**
 	 * Icon name are simply create by adding "InterchangeIcon_" in front of the specialized type. If there is no special type the function will return NAME_None which will use the default icon.
 	 */
-	virtual FName GetIconName() const override
-	{
-		FString MeshIconName = TEXT("MeshIcon.");
-		if (IsSkinnedMesh())
-		{
-			MeshIconName += TEXT("Skinned");
-		}
-		else
-		{
-			MeshIconName += TEXT("Static");
-		}
-		
-		return FName(*MeshIconName);
-	}
+	virtual FName GetIconName() const override;
 
 	/**
 	 * Return true if this node represent a skinned mesh
 	 */
 	UFUNCTION(BlueprintCallable, Category = "Interchange | Node | Mesh")
-	bool IsSkinnedMesh() const
-	{
-		if (!Attributes->ContainAttribute(UE::Interchange::FMeshNodeStaticData::IsSkinnedMeshKey()))
-		{
-			return false;
-		}
-
-		UE::Interchange::FAttributeStorage::TAttributeHandle<bool> Handle = Attributes->GetAttributeHandle<bool>(UE::Interchange::FMeshNodeStaticData::IsSkinnedMeshKey());
-		if (Handle.IsValid())
-		{
-			bool bValue = false;
-			Handle.Get(bValue);
-			return bValue;
-		}
-		return false;
-	}
+	bool IsSkinnedMesh() const;
 
 	/**
 	 * Set the IsSkinnedMesh attribute to determine if this node represent a skinned mesh.
 	 */
 	UFUNCTION(BlueprintCallable, Category = "Interchange | Node | Mesh")
-	bool SetSkinnedMesh(const bool bIsSkinnedMesh)
-	{
-		UE::Interchange::EAttributeStorageResult Result = Attributes->RegisterAttribute(UE::Interchange::FMeshNodeStaticData::IsSkinnedMeshKey(), bIsSkinnedMesh);
-		if (IsAttributeStorageResultSuccess(Result))
-		{
-			UE::Interchange::FAttributeStorage::TAttributeHandle<bool> Handle = Attributes->GetAttributeHandle<bool>(UE::Interchange::FMeshNodeStaticData::IsSkinnedMeshKey());
-			return Handle.IsValid();
-		}
-		return false;
-	}
+	bool SetSkinnedMesh(const bool bIsSkinnedMesh);
 
 	/**
-	 * Return true if this node represent a skinned mesh
+	 * Return true if this node represent a blend shape
 	 */
 	UFUNCTION(BlueprintCallable, Category = "Interchange | Node | Mesh")
-	bool IsBlendShape() const
-	{
-		if (!Attributes->ContainAttribute(UE::Interchange::FMeshNodeStaticData::IsBlendShapeKey()))
-		{
-			return false;
-		}
-
-		UE::Interchange::FAttributeStorage::TAttributeHandle<bool> Handle = Attributes->GetAttributeHandle<bool>(UE::Interchange::FMeshNodeStaticData::IsBlendShapeKey());
-		if (Handle.IsValid())
-		{
-			bool bValue = false;
-			Handle.Get(bValue);
-			return bValue;
-		}
-		return false;
-	}
+	bool IsBlendShape() const;
 
 	/**
 	 * Set the IsBlendShape attribute to determine if this node represent a blend shape.
 	 */
 	UFUNCTION(BlueprintCallable, Category = "Interchange | Node | Mesh")
-	bool SetBlendShape(const bool bIsBlendShape)
-	{
-		UE::Interchange::EAttributeStorageResult Result = Attributes->RegisterAttribute(UE::Interchange::FMeshNodeStaticData::IsBlendShapeKey(), bIsBlendShape);
-		if (IsAttributeStorageResultSuccess(Result))
-		{
-			UE::Interchange::FAttributeStorage::TAttributeHandle<bool> Handle = Attributes->GetAttributeHandle<bool>(UE::Interchange::FMeshNodeStaticData::IsBlendShapeKey());
-			return Handle.IsValid();
-		}
-		return false;
-	}
+	bool SetBlendShape(const bool bIsBlendShape);
 
 	/**
+	 * Get the blend shape name.
 	 * Return true if we successfully query the BlendShapeName attribute
 	 */
 	UFUNCTION(BlueprintCallable, Category = "Interchange | Node | Mesh")
-	bool GetBlendShapeName(FString& OutBlendShapeName) const
-	{
-		OutBlendShapeName.Empty();
-		if (!Attributes->ContainAttribute(UE::Interchange::FMeshNodeStaticData::BlendShapeNameKey()))
-		{
-			return false;
-		}
-
-		UE::Interchange::FAttributeStorage::TAttributeHandle<FString> Handle = Attributes->GetAttributeHandle<FString>(UE::Interchange::FMeshNodeStaticData::BlendShapeNameKey());
-		if (Handle.IsValid())
-		{
-			Handle.Get(OutBlendShapeName);
-			return true;
-		}
-		return false;
-	}
+	bool GetBlendShapeName(FString& OutBlendShapeName) const;
 
 	/**
 	 * Set the BlendShapeName attribute to determine if this node represent a blend shape.
 	 */
 	UFUNCTION(BlueprintCallable, Category = "Interchange | Node | Mesh")
-	bool SetBlendShapeName(const FString& BlendShapeName)
-	{
-		UE::Interchange::EAttributeStorageResult Result = Attributes->RegisterAttribute(UE::Interchange::FMeshNodeStaticData::BlendShapeNameKey(), BlendShapeName);
-		if (IsAttributeStorageResultSuccess(Result))
-		{
-			UE::Interchange::FAttributeStorage::TAttributeHandle<FString> Handle = Attributes->GetAttributeHandle<FString>(UE::Interchange::FMeshNodeStaticData::BlendShapeNameKey());
-			return Handle.IsValid();
-		}
-		return false;
-	}
+	bool SetBlendShapeName(const FString& BlendShapeName);
 
 	/** Mesh node Interface Begin */
-	virtual const TOptional<FString> GetPayLoadKey() const
-	{
-		if (!Attributes->ContainAttribute(UE::Interchange::FMeshNodeStaticData::PayloadSourceFileKey()))
-		{
-			return TOptional<FString>();
-		}
-		UE::Interchange::FAttributeStorage::TAttributeHandle<FString> AttributeHandle = Attributes->GetAttributeHandle<FString>(UE::Interchange::FMeshNodeStaticData::PayloadSourceFileKey());
-		if (!AttributeHandle.IsValid())
-		{
-			return TOptional<FString>();
-		}
-		FString PayloadKey;
-		UE::Interchange::EAttributeStorageResult Result = AttributeHandle.Get(PayloadKey);
-		if (!IsAttributeStorageResultSuccess(Result))
-		{
-			LogAttributeStorageErrors(Result, TEXT("UInterchangeMeshNode.GetPayLoadKey"), UE::Interchange::FMeshNodeStaticData::PayloadSourceFileKey());
-			return TOptional<FString>();
-		}
-		return TOptional<FString>(PayloadKey);
-	}
+	virtual const TOptional<FString> GetPayLoadKey() const;
 
 	UFUNCTION(BlueprintCallable, Category = "Interchange | Node | Mesh")
-	virtual void SetPayLoadKey(const FString& PayloadKey)
-	{
-		UE::Interchange::EAttributeStorageResult Result = Attributes->RegisterAttribute(UE::Interchange::FMeshNodeStaticData::PayloadSourceFileKey(), PayloadKey);
-		if (!IsAttributeStorageResultSuccess(Result))
-		{
-			LogAttributeStorageErrors(Result, TEXT("UInterchangeMeshNode.SetPayLoadKey"), UE::Interchange::FMeshNodeStaticData::PayloadSourceFileKey());
-		}
-	}
+	virtual void SetPayLoadKey(const FString& PayloadKey);
 	
-	/** Return false if the Attribute was not set previously.*/
+	/** Query this mesh vertices count. Return false if the attribute was not set.*/
 	UFUNCTION(BlueprintCallable, Category = "Interchange | Node | Mesh")
-	bool GetCustomVertexCount(int32& AttributeValue) const
-	{
-		IMPLEMENT_NODE_ATTRIBUTE_GETTER(VertexCount, int32);
-	}
+	bool GetCustomVertexCount(int32& AttributeValue) const;
 
+	/** Set this mesh vertices count. Return false if the attribute cannot be set.*/
 	UFUNCTION(BlueprintCallable, Category = "Interchange | Node | Mesh")
-	bool SetCustomVertexCount(const int32& AttributeValue)
-	{
-		IMPLEMENT_NODE_ATTRIBUTE_SETTER_NODELEGATE(VertexCount, int32);
-	}
+	bool SetCustomVertexCount(const int32& AttributeValue);
 
-	/** Return false if the Attribute was not set previously.*/
+	/** Query this mesh polygon count. Return false if the attribute was not set.*/
 	UFUNCTION(BlueprintCallable, Category = "Interchange | Node | Mesh")
-	bool GetCustomPolygonCount(int32& AttributeValue) const
-	{
-		IMPLEMENT_NODE_ATTRIBUTE_GETTER(PolygonCount, int32);
-	}
+	bool GetCustomPolygonCount(int32& AttributeValue) const;
 
+	/** Set this mesh polygon count. Return false if the attribute cannot be set.*/
 	UFUNCTION(BlueprintCallable, Category = "Interchange | Node | Mesh")
-	bool SetCustomPolygonCount(const int32& AttributeValue)
-	{
-		IMPLEMENT_NODE_ATTRIBUTE_SETTER_NODELEGATE(PolygonCount, int32);
-	}
+	bool SetCustomPolygonCount(const int32& AttributeValue);
 
-	/** Return false if the Attribute was not set previously.*/
+	/** Query this mesh bounding box. Return false if the attribute was not set.*/
 	UFUNCTION(BlueprintCallable, Category = "Interchange | Node | Mesh")
-	bool GetCustomBoundingBox(FBox& AttributeValue) const
-	{
-		IMPLEMENT_NODE_ATTRIBUTE_GETTER(BoundingBox, FBox);
-	}
+	bool GetCustomBoundingBox(FBox& AttributeValue) const;
 
+	/** Set this mesh bounding box. Return false if the attribute cannot be set.*/
 	UFUNCTION(BlueprintCallable, Category = "Interchange | Node | Mesh")
-	bool SetCustomBoundingBox(const FBox& AttributeValue)
-	{
-		IMPLEMENT_NODE_ATTRIBUTE_SETTER_NODELEGATE(BoundingBox, FBox);
-	}
+	bool SetCustomBoundingBox(const FBox& AttributeValue);
 
-	/** Return false if the Attribute was not set previously.*/
+	/** Query if this mesh has vertex normal. Return false if the attribute was not set.*/
 	UFUNCTION(BlueprintCallable, Category = "Interchange | Node | Mesh")
-	bool GetCustomHasVertexNormal(bool& AttributeValue) const
-	{
-		IMPLEMENT_NODE_ATTRIBUTE_GETTER(HasVertexNormal, bool);
-	}
+	bool GetCustomHasVertexNormal(bool& AttributeValue) const;
 
+	/** Set this mesh has vertex normal attribute. Return false if the attribute cannot be set.*/
 	UFUNCTION(BlueprintCallable, Category = "Interchange | Node | Mesh")
-	bool SetCustomHasVertexNormal(const bool& AttributeValue)
-	{
-		IMPLEMENT_NODE_ATTRIBUTE_SETTER_NODELEGATE(HasVertexNormal, bool);
-	}
+	bool SetCustomHasVertexNormal(const bool& AttributeValue);
 
-	/** Return false if the Attribute was not set previously.*/
+	/** Query if this mesh has vertex bi-normal. Return false if the attribute was not set.*/
 	UFUNCTION(BlueprintCallable, Category = "Interchange | Node | Mesh")
-	bool GetCustomHasVertexBinormal(bool& AttributeValue) const
-	{
-		IMPLEMENT_NODE_ATTRIBUTE_GETTER(HasVertexBinormal, bool);
-	}
+	bool GetCustomHasVertexBinormal(bool& AttributeValue) const;
 
+	/** Set this mesh has vertex bi-normal attribute. Return false if the attribute cannot be set.*/
 	UFUNCTION(BlueprintCallable, Category = "Interchange | Node | Mesh")
-	bool SetCustomHasVertexBinormal(const bool& AttributeValue)
-	{
-		IMPLEMENT_NODE_ATTRIBUTE_SETTER_NODELEGATE(HasVertexBinormal, bool);
-	}
+	bool SetCustomHasVertexBinormal(const bool& AttributeValue);
 
-	/** Return false if the Attribute was not set previously.*/
+	/** Query if this mesh has vertex tangent. Return false if the attribute was not set.*/
 	UFUNCTION(BlueprintCallable, Category = "Interchange | Node | Mesh")
-	bool GetCustomHasVertexTangent(bool& AttributeValue) const
-	{
-		IMPLEMENT_NODE_ATTRIBUTE_GETTER(HasVertexTangent, bool);
-	}
+	bool GetCustomHasVertexTangent(bool& AttributeValue) const;
 
+	/** Set this mesh has vertex tangent attribute. Return false if the attribute cannot be set.*/
 	UFUNCTION(BlueprintCallable, Category = "Interchange | Node | Mesh")
-	bool SetCustomHasVertexTangent(const bool& AttributeValue)
-	{
-		IMPLEMENT_NODE_ATTRIBUTE_SETTER_NODELEGATE(HasVertexTangent, bool);
-	}
+	bool SetCustomHasVertexTangent(const bool& AttributeValue);
 
-	/** Return false if the Attribute was not set previously.*/
+	/** Query if this mesh has smooth group. Return false if the attribute was not set.*/
 	UFUNCTION(BlueprintCallable, Category = "Interchange | Node | Mesh")
-	bool GetCustomHasSmoothGroup(bool& AttributeValue) const
-	{
-		IMPLEMENT_NODE_ATTRIBUTE_GETTER(HasSmoothGroup, bool);
-	}
+	bool GetCustomHasSmoothGroup(bool& AttributeValue) const;
 
+	/** Set this mesh has smooth group attribute. Return false if the attribute cannot be set.*/
 	UFUNCTION(BlueprintCallable, Category = "Interchange | Node | Mesh")
-	bool SetCustomHasSmoothGroup(const bool& AttributeValue)
-	{
-		IMPLEMENT_NODE_ATTRIBUTE_SETTER_NODELEGATE(HasSmoothGroup, bool);
-	}
+	bool SetCustomHasSmoothGroup(const bool& AttributeValue);
 
-	/** Return false if the Attribute was not set previously.*/
+	/** Query if this mesh has vertex color. Return false if the attribute was not set.*/
 	UFUNCTION(BlueprintCallable, Category = "Interchange | Node | Mesh")
-	bool GetCustomHasVertexColor(bool& AttributeValue) const
-	{
-		IMPLEMENT_NODE_ATTRIBUTE_GETTER(HasVertexColor, bool);
-	}
+	bool GetCustomHasVertexColor(bool& AttributeValue) const;
 
+	/** Set this mesh has vertex color attribute. Return false if the attribute cannot be set.*/
 	UFUNCTION(BlueprintCallable, Category = "Interchange | Node | Mesh")
-	bool SetCustomHasVertexColor(const bool& AttributeValue)
-	{
-		IMPLEMENT_NODE_ATTRIBUTE_SETTER_NODELEGATE(HasVertexColor, bool);
-	}
+	bool SetCustomHasVertexColor(const bool& AttributeValue);
 
-	/** Return false if the Attribute was not set previously.*/
+	/** Query this mesh UV count. Return false if the attribute was not set.*/
 	UFUNCTION(BlueprintCallable, Category = "Interchange | Node | Mesh")
-	bool GetCustomUVCount(int32& AttributeValue) const
-	{
-		IMPLEMENT_NODE_ATTRIBUTE_GETTER(UVCount, int32);
-	}
+	bool GetCustomUVCount(int32& AttributeValue) const;
 
+	/** Set this mesh UV count attribute. Return false if the attribute cannot be set.*/
 	UFUNCTION(BlueprintCallable, Category = "Interchange | Node | Mesh")
-	bool SetCustomUVCount(const int32& AttributeValue)
-	{
-		IMPLEMENT_NODE_ATTRIBUTE_SETTER_NODELEGATE(UVCount, int32);
-	}
+	bool SetCustomUVCount(const int32& AttributeValue);
 
 	/**
 	 * This function allow to retrieve the number of skeleton dependencies for this object.
-	 *
 	 */
 	UFUNCTION(BlueprintCallable, Category = "Interchange | Node | Mesh")
-	int32 GetSkeletonDependeciesCount() const
-	{
-		return SkeletonDependencies.GetCount();
-	}
+	int32 GetSkeletonDependeciesCount() const;
 
 	/**
 	 * This function allow to retrieve the skeleton dependency for this object.
-	 *
 	 */
 	UFUNCTION(BlueprintCallable, Category = "Interchange | Node | Mesh")
-	void GetSkeletonDependencies(TArray<FString>& OutDependencies) const
-	{
-		SkeletonDependencies.GetNames(OutDependencies);
-	}
+	void GetSkeletonDependencies(TArray<FString>& OutDependencies) const;
 
 	/**
 	 * This function allow to retrieve one skeleton dependency for this object.
-	 *
 	 */
 	UFUNCTION(BlueprintCallable, Category = "Interchange | Node | Mesh")
-	void GetSkeletonDependency(const int32 Index, FString& OutDependency) const
-	{
-		SkeletonDependencies.GetName(Index, OutDependency);
-	}
+	void GetSkeletonDependency(const int32 Index, FString& OutDependency) const;
 
 	/**
 	 * Add one skeleton dependency to this object.
 	 */
 	UFUNCTION(BlueprintCallable, Category = "Interchange | Node | Mesh")
-	bool SetSkeletonDependencyUid(const FString& DependencyUid)
-	{ 
-		return SkeletonDependencies.AddName(DependencyUid);
-	}
+	bool SetSkeletonDependencyUid(const FString& DependencyUid);
 
 	/**
 	 * Remove one skeleton dependency from this object.
 	 */
 	UFUNCTION(BlueprintCallable, Category = "Interchange | Node | Mesh")
-	bool RemoveSkeletonDependencyUid(const FString& DependencyUid)
-	{
-		return SkeletonDependencies.RemoveName(DependencyUid);
-	}
+	bool RemoveSkeletonDependencyUid(const FString& DependencyUid);
 
 	/**
 	 * This function allow to retrieve the number of Material dependencies for this object.
-	 *
 	 */
 	UFUNCTION(BlueprintCallable, Category = "Interchange | Node | Mesh")
-	int32 GetMaterialDependeciesCount() const
-	{
-		return MaterialDependencies.GetCount();
-	}
+	int32 GetMaterialDependeciesCount() const;
 
 	/**
 	 * This function allow to retrieve the Material dependency for this object.
-	 *
 	 */
 	UFUNCTION(BlueprintCallable, Category = "Interchange | Node | Mesh")
-	void GetMaterialDependencies(TArray<FString>& OutDependencies) const
-	{
-		MaterialDependencies.GetNames(OutDependencies);
-	}
+	void GetMaterialDependencies(TArray<FString>& OutDependencies) const;
 
 	/**
 	 * This function allow to retrieve one Material dependency for this object.
-	 *
 	 */
 	UFUNCTION(BlueprintCallable, Category = "Interchange | Node | Mesh")
-	void GetMaterialDependency(const int32 Index, FString& OutDependency) const
-	{
-		MaterialDependencies.GetName(Index, OutDependency);
-	}
+	void GetMaterialDependency(const int32 Index, FString& OutDependency) const;
 
 	/**
 	 * Add one Material dependency to this object.
 	 */
 	UFUNCTION(BlueprintCallable, Category = "Interchange | Node | Mesh")
-	bool SetMaterialDependencyUid(const FString& DependencyUid)
-	{
-		return MaterialDependencies.AddName(DependencyUid);
-	}
+	bool SetMaterialDependencyUid(const FString& DependencyUid);
 
 	/**
 	 * Remove one Material dependency from this object.
 	 */
 	UFUNCTION(BlueprintCallable, Category = "Interchange | Node | Mesh")
-	bool RemoveMaterialDependencyUid(const FString& DependencyUid)
-	{
-		return MaterialDependencies.RemoveName(DependencyUid);
-	}
+	bool RemoveMaterialDependencyUid(const FString& DependencyUid);
 
 	/**
 	 * This function allow to retrieve the number of Shape dependencies for this object.
-	 *
 	 */
 	UFUNCTION(BlueprintCallable, Category = "Interchange | Node | Mesh")
-	int32 GetShapeDependeciesCount() const
-	{
-		return ShapeDependencies.GetCount();
-	}
+	int32 GetShapeDependeciesCount() const;
 
 	/**
 	 * This function allow to retrieve the Shape dependency for this object.
-	 *
 	 */
 	UFUNCTION(BlueprintCallable, Category = "Interchange | Node | Mesh")
-	void GetShapeDependencies(TArray<FString>& OutDependencies) const
-	{
-		ShapeDependencies.GetNames(OutDependencies);
-	}
+	void GetShapeDependencies(TArray<FString>& OutDependencies) const;
 
 	/**
 	 * This function allow to retrieve one Shape dependency for this object.
-	 *
 	 */
 	UFUNCTION(BlueprintCallable, Category = "Interchange | Node | Mesh")
-	void GetShapeDependency(const int32 Index, FString& OutDependency) const
-	{
-		ShapeDependencies.GetName(Index, OutDependency);
-	}
+	void GetShapeDependency(const int32 Index, FString& OutDependency) const;
 
 	/**
 	 * Add one Shape dependency to this object.
 	 */
 	UFUNCTION(BlueprintCallable, Category = "Interchange | Node | Mesh")
-	bool SetShapeDependencyUid(const FString& DependencyUid)
-	{
-		return ShapeDependencies.AddName(DependencyUid);
-	}
+	bool SetShapeDependencyUid(const FString& DependencyUid);
 
 	/**
 	 * Remove one Shape dependency from this object.
 	 */
 	UFUNCTION(BlueprintCallable, Category = "Interchange | Node | Mesh")
-	bool RemoveShapeDependencyUid(const FString& DependencyUid)
-	{
-		return ShapeDependencies.RemoveName(DependencyUid);
-	}
+	bool RemoveShapeDependencyUid(const FString& DependencyUid);
 
 	/**
 	 * This function allow to retrieve the number of scene node instancing this mesh.
-	 *
 	 */
 	UFUNCTION(BlueprintCallable, Category = "Interchange | Node | Mesh")
-	int32 GetSceneInstanceUidsCount() const
-	{
-		return SceneInstancesUids.GetCount();
-	}
+	int32 GetSceneInstanceUidsCount() const;
 
 	/**
 	 * This function allow to retrieve the Shape dependency for this object.
-	 *
 	 */
 	UFUNCTION(BlueprintCallable, Category = "Interchange | Node | Mesh")
-	void GetSceneInstanceUids(TArray<FString>& OutDependencies) const
-	{
-		SceneInstancesUids.GetNames(OutDependencies);
-	}
+	void GetSceneInstanceUids(TArray<FString>& OutDependencies) const;
 
 	/**
 	 * This function allow to retrieve one Shape dependency for this object.
-	 *
 	 */
 	UFUNCTION(BlueprintCallable, Category = "Interchange | Node | Mesh")
-	void GetSceneInstanceUid(const int32 Index, FString& OutDependency) const
-	{
-		SceneInstancesUids.GetName(Index, OutDependency);
-	}
+	void GetSceneInstanceUid(const int32 Index, FString& OutDependency) const;
 
 	/**
 	 * Add one Shape dependency to this object.
 	 */
 	UFUNCTION(BlueprintCallable, Category = "Interchange | Node | Mesh")
-	bool SetSceneInstanceUid(const FString& DependencyUid)
-	{
-		return SceneInstancesUids.AddName(DependencyUid);
-	}
+	bool SetSceneInstanceUid(const FString& DependencyUid);
 
 	/**
 	 * Remove one Shape dependency from this object.
 	 */
 	UFUNCTION(BlueprintCallable, Category = "Interchange | Node | Mesh")
-	bool RemoveSceneInstanceUid(const FString& DependencyUid)
-	{
-		return SceneInstancesUids.RemoveName(DependencyUid);
-	}
+	bool RemoveSceneInstanceUid(const FString& DependencyUid);
+
 private:
 	const UE::Interchange::FAttributeKey Macro_CustomVertexCountKey = UE::Interchange::FAttributeKey(TEXT("VertexCount"));
 	const UE::Interchange::FAttributeKey Macro_CustomPolygonCountKey = UE::Interchange::FAttributeKey(TEXT("PolygonCount"));
@@ -686,8 +296,8 @@ private:
 	const UE::Interchange::FAttributeKey Macro_CustomHasVertexColorKey = UE::Interchange::FAttributeKey(TEXT("HasVertexColor"));
 	const UE::Interchange::FAttributeKey Macro_CustomUVCountKey = UE::Interchange::FAttributeKey(TEXT("UVCount"));
 
-	UE::Interchange::FNameAttributeArrayHelper SkeletonDependencies;
-	UE::Interchange::FNameAttributeArrayHelper MaterialDependencies;
-	UE::Interchange::FNameAttributeArrayHelper ShapeDependencies;
-	UE::Interchange::FNameAttributeArrayHelper SceneInstancesUids;
+	UE::Interchange::TArrayAttributeHelper<FString> SkeletonDependencies;
+	UE::Interchange::TArrayAttributeHelper<FString> MaterialDependencies;
+	UE::Interchange::TArrayAttributeHelper<FString> ShapeDependencies;
+	UE::Interchange::TArrayAttributeHelper<FString> SceneInstancesUids;
 };
diff --git a/Engine/Plugins/Experimental/Interchange/Runtime/Source/Nodes/Public/InterchangePhysicsAssetFactoryNode.h b/Engine/Plugins/Experimental/Interchange/Runtime/Source/Nodes/Public/InterchangePhysicsAssetFactoryNode.h
--- a/Engine/Plugins/Experimental/Interchange/Runtime/Source/Nodes/Public/InterchangePhysicsAssetFactoryNode.h
+++ b/Engine/Plugins/Experimental/Interchange/Runtime/Source/Nodes/Public/InterchangePhysicsAssetFactoryNode.h
@@ -36,7 +36,7 @@ public:
 	void InitializePhysicsAssetNode(const FString& UniqueID, const FString& DisplayLabel, const FString& InAssetClass)
 	{
 		bIsNodeClassInitialized = false;
-		InitializeNode(UniqueID, DisplayLabel, EInterchangeNodeContainerType::NodeContainerType_FactoryData);
+		InitializeNode(UniqueID, DisplayLabel, EInterchangeNodeContainerType::FactoryData);
 		FString OperationName = GetTypeName() + TEXT(".SetAssetClassName");
 		InterchangePrivateNodeBase::SetCustomAttribute<FString>(*Attributes, ClassNameAttributeKey, OperationName, InAssetClass);
 		FillAssetClassFromAttribute();
diff --git a/Engine/Plugins/Experimental/Interchange/Runtime/Source/Nodes/Public/InterchangeSceneNode.h b/Engine/Plugins/Experimental/Interchange/Runtime/Source/Nodes/Public/InterchangeSceneNode.h
--- a/Engine/Plugins/Experimental/Interchange/Runtime/Source/Nodes/Public/InterchangeSceneNode.h
+++ b/Engine/Plugins/Experimental/Interchange/Runtime/Source/Nodes/Public/InterchangeSceneNode.h
@@ -13,31 +13,12 @@ namespace UE
 	namespace Interchange
 	{
 
-		struct FSceneNodeStaticData : public FBaseNodeStaticData
+		struct INTERCHANGENODES_API FSceneNodeStaticData : public FBaseNodeStaticData
 		{
-			static const FString& GetNodeSpecializeTypeBaseKey()
-			{
-				static FString SceneNodeSpecializeType_BaseKey = TEXT("SceneNodeSpecializeType");
-				return SceneNodeSpecializeType_BaseKey;
-			}
-
-			static const FString& GetMaterialDependencyUidsBaseKey()
-			{
-				static FString MaterialDependencyUids_BaseKey = TEXT("__MaterialDependencyUidsBaseKey__");
-				return MaterialDependencyUids_BaseKey;
-			}
-
-			static const FString& GetJointSpecializeTypeString()
-			{
-				static FString JointSpecializeTypeString = TEXT("Joint");
-				return JointSpecializeTypeString;
-			}
-
-			static const FString& GetLodGroupSpecializeTypeString()
-			{
-				static FString JointSpecializeTypeString = TEXT("LodGroup");
-				return JointSpecializeTypeString;
-			}
+			static const FString& GetNodeSpecializeTypeBaseKey();
+			static const FString& GetMaterialDependencyUidsBaseKey();
+			static const FString& GetJointSpecializeTypeString();
+			static const FString& GetLodGroupSpecializeTypeString();
 		};
 
 	}//ns Interchange
@@ -49,296 +30,128 @@ class INTERCHANGENODES_API UInterchangeSceneNode : public UInterchangeBaseNode
 	GENERATED_BODY()
 
 public:
-	UInterchangeSceneNode()
-	:UInterchangeBaseNode()
-	{
-		NodeSpecializeTypes.Initialize(Attributes, UE::Interchange::FSceneNodeStaticData::GetNodeSpecializeTypeBaseKey());
-		MaterialDependencyUids.Initialize(Attributes, UE::Interchange::FSceneNodeStaticData::GetMaterialDependencyUidsBaseKey());
-	}
+	UInterchangeSceneNode();
 
 	/**
 	 * Return the node type name of the class, we use this when reporting error
 	 */
-	virtual FString GetTypeName() const override
-	{
-		const FString TypeName = TEXT("SceneNode");
-		return TypeName;
-	}
+	virtual FString GetTypeName() const override;
 
-	virtual FString GetKeyDisplayName(const UE::Interchange::FAttributeKey& NodeAttributeKey) const override
-	{
-		FString KeyDisplayName = NodeAttributeKey.Key;
-		if (NodeAttributeKey.Key.Equals(UE::Interchange::FSceneNodeStaticData::GetNodeSpecializeTypeBaseKey()))
-		{
-			KeyDisplayName = TEXT("Specialized type count");
-			return KeyDisplayName;
-		}
-		else if (NodeAttributeKey.Key.StartsWith(UE::Interchange::FSceneNodeStaticData::GetNodeSpecializeTypeBaseKey()))
-		{
-			KeyDisplayName = TEXT("Specialized type index ");
-			const FString IndexKey = UE::Interchange::FNameAttributeArrayHelper::IndexKey();
-			int32 IndexPosition = NodeAttributeKey.Key.Find(IndexKey) + IndexKey.Len();
-			if (IndexPosition < NodeAttributeKey.Key.Len())
-			{
-				KeyDisplayName += NodeAttributeKey.Key.RightChop(IndexPosition);
-			}
-			return KeyDisplayName;
-		}
-		else if (NodeAttributeKey.Key.Equals(UE::Interchange::FSceneNodeStaticData::GetMaterialDependencyUidsBaseKey()))
-		{
-			KeyDisplayName = TEXT("Material dependencies count");
-			return KeyDisplayName;
-		}
-		else if (NodeAttributeKey.Key.StartsWith(UE::Interchange::FSceneNodeStaticData::GetMaterialDependencyUidsBaseKey()))
-		{
-			KeyDisplayName = TEXT("Material dependency index ");
-			const FString IndexKey = UE::Interchange::FNameAttributeArrayHelper::IndexKey();
-			int32 IndexPosition = NodeAttributeKey.Key.Find(IndexKey) + IndexKey.Len();
-			if (IndexPosition < NodeAttributeKey.Key.Len())
-			{
-				KeyDisplayName += NodeAttributeKey.Key.RightChop(IndexPosition);
-			}
-			return KeyDisplayName;
-		}
-		return Super::GetKeyDisplayName(NodeAttributeKey);
-	}
-
-	virtual FString GetAttributeCategory(const UE::Interchange::FAttributeKey& NodeAttributeKey) const override
-	{
-		if (NodeAttributeKey.Key.StartsWith(UE::Interchange::FSceneNodeStaticData::GetNodeSpecializeTypeBaseKey()))
-		{
-			return FString(TEXT("SpecializeType"));
-		}
-		else if (NodeAttributeKey.Key.StartsWith(UE::Interchange::FSceneNodeStaticData::GetMaterialDependencyUidsBaseKey()))
-		{
-			return FString(TEXT("MaterialDependencies"));
-		}
-		else if (NodeAttributeKey == Macro_CustomLocalTransformKey
-				 || NodeAttributeKey == Macro_CustomGlobalTransformKey
-				 || NodeAttributeKey == Macro_CustomAssetInstanceUidKey)
-		{
-			return FString(TEXT("Scene"));
-		}
-		return Super::GetAttributeCategory(NodeAttributeKey);
-	}
+	virtual FString GetKeyDisplayName(const UE::Interchange::FAttributeKey& NodeAttributeKey) const override;
 
-	virtual FGuid GetHash() const override
-	{
-		return Attributes->GetStorageHash();
-	}
+	virtual FString GetAttributeCategory(const UE::Interchange::FAttributeKey& NodeAttributeKey) const override;
 
 	/**
 	 * Icon name are simply create by adding "InterchangeIcon_" in front of the specialized type. If there is no special type the function will return NAME_None which will use the default icon.
 	 */
-	virtual FName GetIconName() const override
-	{
-		FString SpecializedType;
-		GetSpecializedType(0, SpecializedType);
-		if (SpecializedType.IsEmpty())
-		{
-			return NAME_None;
-		}
-		SpecializedType = TEXT("SceneGraphIcon.") + SpecializedType;
-		return FName(*SpecializedType);
-	}
+	virtual FName GetIconName() const override;
 
-	/** Return true if this node contains the specialized type parameter.*/
 	UFUNCTION(BlueprintCallable, Category = "Interchange | Node | Scene")
-	bool IsSpecializedTypeContains(const FString& SpecializedType) const
-	{
-		TArray<FString> SpecializedTypes;
-		GetSpecializedTypes(SpecializedTypes);
-		for (const FString& SpecializedTypeRef : SpecializedTypes)
-		{
-			if (SpecializedTypeRef.Equals(SpecializedType))
-			{
-				return true;
-			}
-		}
-		return false;
-	}
+	bool IsSpecializedTypeContains(const FString& SpecializedType) const;
 
 	/** Specialized type are scene node special type like (Joint or LODGroup).*/
 	UFUNCTION(BlueprintCallable, Category = "Interchange | Node | Scene")
-	int32 GetSpecializedTypeCount() const
-	{
-		return NodeSpecializeTypes.GetCount();
-	}
+	int32 GetSpecializedTypeCount() const;
 
 	UFUNCTION(BlueprintCallable, Category = "Interchange | Node | Scene")
-	void GetSpecializedType(const int32 Index, FString& OutSpecializedType) const
-	{
-		NodeSpecializeTypes.GetName(Index, OutSpecializedType);
-	}
+	void GetSpecializedType(const int32 Index, FString& OutSpecializedType) const;
 
 	UFUNCTION(BlueprintCallable, Category = "Interchange | Node | Scene")
-	void GetSpecializedTypes(TArray<FString>& OutSpecializedTypes) const
-	{
-		NodeSpecializeTypes.GetNames(OutSpecializedTypes);
-	}
+	void GetSpecializedTypes(TArray<FString>& OutSpecializedTypes) const;
 
 	UFUNCTION(BlueprintCallable, Category = "Interchange | Node | Scene")
-	bool AddSpecializedType(const FString& SpecializedType)
-	{
-		return NodeSpecializeTypes.AddName(SpecializedType);
-	}
+	bool AddSpecializedType(const FString& SpecializedType);
 
 	UFUNCTION(BlueprintCallable, Category = "Interchange | Node | Scene")
-	bool RemoveSpecializedType(const FString& SpecializedType)
-	{
-		return NodeSpecializeTypes.RemoveName(SpecializedType);
-	}
+	bool RemoveSpecializedType(const FString& SpecializedType);
 
 	/** Asset dependencies are the asset on which this node depend.*/
 	UFUNCTION(BlueprintCallable, Category = "Interchange | Node | Scene")
-	int32 GetMaterialDependencyUidsCount() const
-	{
-		return MaterialDependencyUids.GetCount();
-	}
+	int32 GetMaterialDependencyUidsCount() const;
 
 	UFUNCTION(BlueprintCallable, Category = "Interchange | Node | Scene")
-	void GetMaterialDependencyUid(const int32 Index, FString& OutMaterialDependencyUid) const
-	{
-		MaterialDependencyUids.GetName(Index, OutMaterialDependencyUid);
-	}
+	void GetMaterialDependencyUid(const int32 Index, FString& OutMaterialDependencyUid) const;
 
 	UFUNCTION(BlueprintCallable, Category = "Interchange | Node | Scene")
-	void GetMaterialDependencyUids(TArray<FString>& OutMaterialDependencyUids) const
-	{
-		MaterialDependencyUids.GetNames(OutMaterialDependencyUids);
-	}
+	void GetMaterialDependencyUids(TArray<FString>& OutMaterialDependencyUids) const;
 
 	UFUNCTION(BlueprintCallable, Category = "Interchange | Node | Scene")
-	bool AddMaterialDependencyUid(const FString& MaterialDependencyUid)
-	{
-		return MaterialDependencyUids.AddName(MaterialDependencyUid);
-	}
+	bool AddMaterialDependencyUid(const FString& MaterialDependencyUid);
 
 	UFUNCTION(BlueprintCallable, Category = "Interchange | Node | Scene")
-	bool RemoveMaterialDependencyUid(const FString& MaterialDependencyUid)
-	{
-		return MaterialDependencyUids.RemoveName(MaterialDependencyUid);
-	}
+	bool RemoveMaterialDependencyUid(const FString& MaterialDependencyUid);
 
 	//Default transform is the transform we have in the node (no bind pose, no time evaluation).
 
 	/** Return the default scene node local transform. */
 	UFUNCTION(BlueprintCallable, Category = "Interchange | Node | Scene")
-	bool GetCustomLocalTransform(FTransform& AttributeValue) const
-	{
-		IMPLEMENT_NODE_ATTRIBUTE_GETTER(LocalTransform, FTransform);
-	}
+	bool GetCustomLocalTransform(FTransform& AttributeValue) const;
 
 	/** Store the default scene node local transform. */
 	UFUNCTION(BlueprintCallable, Category = "Interchange | Node | Scene")
-	bool SetCustomLocalTransform(const FTransform& AttributeValue)
-	{
-		IMPLEMENT_NODE_ATTRIBUTE_SETTER_NODELEGATE(LocalTransform, FTransform);
-	}
+	bool SetCustomLocalTransform(const FTransform& AttributeValue);
 
 	/** Return the default scene node global transform. */
 	UFUNCTION(BlueprintCallable, Category = "Interchange | Node | Scene")
-	bool GetCustomGlobalTransform(FTransform& AttributeValue) const
-	{
-		IMPLEMENT_NODE_ATTRIBUTE_GETTER(GlobalTransform, FTransform);
-	}
+	bool GetCustomGlobalTransform(FTransform& AttributeValue) const;
 
 	/** Store the default scene node global transform. */
 	UFUNCTION(BlueprintCallable, Category = "Interchange | Node | Scene")
-	bool SetCustomGlobalTransform(const FTransform& AttributeValue)
-	{
-		IMPLEMENT_NODE_ATTRIBUTE_SETTER_NODELEGATE(GlobalTransform, FTransform);
-	}
+	bool SetCustomGlobalTransform(const FTransform& AttributeValue);
 
 	//Bind pose transform is the transform of the joint when the binding with the mesh was done.
 	//This attribute should be set only if we have a joint.
 
 	/** Return the bind pose scene node local transform. */
 	UFUNCTION(BlueprintCallable, Category = "Interchange | Node | Joint")
-	bool GetCustomBindPoseLocalTransform(FTransform& AttributeValue) const
-	{
-		IMPLEMENT_NODE_ATTRIBUTE_GETTER(BindPoseLocalTransform, FTransform);
-	}
+	bool GetCustomBindPoseLocalTransform(FTransform& AttributeValue) const;
 
 	/** Store the bind pose scene node local transform. */
 	UFUNCTION(BlueprintCallable, Category = "Interchange | Node | Joint")
-	bool SetCustomBindPoseLocalTransform(const FTransform& AttributeValue)
-	{
-		IMPLEMENT_NODE_ATTRIBUTE_SETTER_NODELEGATE(BindPoseLocalTransform, FTransform);
-	}
+	bool SetCustomBindPoseLocalTransform(const FTransform& AttributeValue);
 
 	/** Return the bind pose scene node global transform. */
 	UFUNCTION(BlueprintCallable, Category = "Interchange | Node | Joint")
-	bool GetCustomBindPoseGlobalTransform(FTransform& AttributeValue) const
-	{
-		IMPLEMENT_NODE_ATTRIBUTE_GETTER(BindPoseGlobalTransform, FTransform);
-	}
+	bool GetCustomBindPoseGlobalTransform(FTransform& AttributeValue) const;
 
 	/** Store the bind pose scene node global transform. */
 	UFUNCTION(BlueprintCallable, Category = "Interchange | Node | Joint")
-	bool SetCustomBindPoseGlobalTransform(const FTransform& AttributeValue)
-	{
-		IMPLEMENT_NODE_ATTRIBUTE_SETTER_NODELEGATE(BindPoseGlobalTransform, FTransform);
-	}
+	bool SetCustomBindPoseGlobalTransform(const FTransform& AttributeValue);
 
 	//Time zero transform is the transform of the node at time zero.
 	//This is useful when there is no bind pose or when we import rigid mesh.
 
 	/** Return the time zero scene node local transform. */
 	UFUNCTION(BlueprintCallable, Category = "Interchange | Node | Joint")
-	bool GetCustomTimeZeroLocalTransform(FTransform& AttributeValue) const
-	{
-		IMPLEMENT_NODE_ATTRIBUTE_GETTER(TimeZeroLocalTransform, FTransform);
-	}
+	bool GetCustomTimeZeroLocalTransform(FTransform& AttributeValue) const;
 
 	/** Store the time zero scene node local transform. */
 	UFUNCTION(BlueprintCallable, Category = "Interchange | Node | Joint")
-	bool SetCustomTimeZeroLocalTransform(const FTransform& AttributeValue)
-	{
-		IMPLEMENT_NODE_ATTRIBUTE_SETTER_NODELEGATE(TimeZeroLocalTransform, FTransform);
-	}
+	bool SetCustomTimeZeroLocalTransform(const FTransform& AttributeValue);
 
 	/** Return the time zero scene node global transform. */
 	UFUNCTION(BlueprintCallable, Category = "Interchange | Node | Joint")
-	bool GetCustomTimeZeroGlobalTransform(FTransform& AttributeValue) const
-	{
-		IMPLEMENT_NODE_ATTRIBUTE_GETTER(TimeZeroGlobalTransform, FTransform);
-	}
+	bool GetCustomTimeZeroGlobalTransform(FTransform& AttributeValue) const;
 
 	/** Store the time zero scene node global transform. */
 	UFUNCTION(BlueprintCallable, Category = "Interchange | Node | Joint")
-	bool SetCustomTimeZeroGlobalTransform(const FTransform& AttributeValue)
-	{
-		IMPLEMENT_NODE_ATTRIBUTE_SETTER_NODELEGATE(TimeZeroGlobalTransform, FTransform);
-	}
+	bool SetCustomTimeZeroGlobalTransform(const FTransform& AttributeValue);
 
 	/** Return the geometric offset. Any mesh attach to this scene node will be offset using this transform. */
 	UFUNCTION(BlueprintCallable, Category = "Interchange | Node | Joint")
-	bool GetCustomGeometricTransform(FTransform& AttributeValue) const
-	{
-		IMPLEMENT_NODE_ATTRIBUTE_GETTER(GeometricTransform, FTransform);
-	}
+	bool GetCustomGeometricTransform(FTransform& AttributeValue) const;
 
 	/** Store the geometric offset. Any mesh attach to this scene node will be offset using this transform. */
 	UFUNCTION(BlueprintCallable, Category = "Interchange | Node | Joint")
-	bool SetCustomGeometricTransform(const FTransform& AttributeValue)
-	{
-		IMPLEMENT_NODE_ATTRIBUTE_SETTER_NODELEGATE(GeometricTransform, FTransform);
-	}
+	bool SetCustomGeometricTransform(const FTransform& AttributeValue);
 
 	/** Return false if the Attribute was not set previously.*/
 	UFUNCTION(BlueprintCallable, Category = "Interchange | Node | Scene")
-	bool GetCustomAssetInstanceUid(FString& AttributeValue) const
-	{
-		IMPLEMENT_NODE_ATTRIBUTE_GETTER(AssetInstanceUid, FString);
-	}
+	bool GetCustomAssetInstanceUid(FString& AttributeValue) const;
 
 	/** Tells which asset, if any, a scene node is instantiating */
 	UFUNCTION(BlueprintCallable, Category = "Interchange | Node | Scene")
-	bool SetCustomAssetInstanceUid(const FString& AttributeValue)
-	{
-		IMPLEMENT_NODE_ATTRIBUTE_SETTER_NODELEGATE(AssetInstanceUid, FString);
-	}
+	bool SetCustomAssetInstanceUid(const FString& AttributeValue);
 
 private:
 	//Scene Attribute Keys
@@ -351,6 +164,6 @@ private:
 	const UE::Interchange::FAttributeKey Macro_CustomGeometricTransformKey = UE::Interchange::FAttributeKey(TEXT("GeometricTransform"));
 	const UE::Interchange::FAttributeKey Macro_CustomAssetInstanceUidKey = UE::Interchange::FAttributeKey(TEXT("AssetInstanceUid"));
 
-	UE::Interchange::FNameAttributeArrayHelper NodeSpecializeTypes;
-	UE::Interchange::FNameAttributeArrayHelper MaterialDependencyUids;
+	UE::Interchange::TArrayAttributeHelper<FString> NodeSpecializeTypes;
+	UE::Interchange::TArrayAttributeHelper<FString> MaterialDependencyUids;
 };
diff --git a/Engine/Plugins/Experimental/Interchange/Runtime/Source/Nodes/Public/InterchangeSkeletalMeshFactoryNode.h b/Engine/Plugins/Experimental/Interchange/Runtime/Source/Nodes/Public/InterchangeSkeletalMeshFactoryNode.h
--- a/Engine/Plugins/Experimental/Interchange/Runtime/Source/Nodes/Public/InterchangeSkeletalMeshFactoryNode.h
+++ b/Engine/Plugins/Experimental/Interchange/Runtime/Source/Nodes/Public/InterchangeSkeletalMeshFactoryNode.h
@@ -5,10 +5,6 @@
 #include "CoreMinimal.h"
 #include "Nodes/InterchangeBaseNode.h"
 
-#if WITH_ENGINE
-#include "Engine/SkeletalMesh.h"
-#endif
-
 #include "InterchangeSkeletalMeshFactoryNode.generated.h"
 
 //Interchange namespace
@@ -17,13 +13,9 @@ namespace UE
 	namespace Interchange
 	{
 
-		struct FSkeletalMeshNodeStaticData : public FBaseNodeStaticData
+		struct INTERCHANGENODES_API FSkeletalMeshNodeStaticData : public FBaseNodeStaticData
 		{
-			static const FString& GetLodDependenciesBaseKey()
-			{
-				static FString LodDependencies_BaseKey = TEXT("Lod_Dependencies");
-				return LodDependencies_BaseKey;
-			}
+			static const FString& GetLodDependenciesBaseKey();
 		};
 
 	}//ns Interchange
@@ -35,14 +27,7 @@ class INTERCHANGENODES_API UInterchangeSkeletalMeshFactoryNode : public UInterch
 	GENERATED_BODY()
 
 public:
-	UInterchangeSkeletalMeshFactoryNode()
-		:UInterchangeBaseNode()
-	{
-#if WITH_ENGINE
-		AssetClass = nullptr;
-#endif
-		LodDependencies.Initialize(Attributes, UE::Interchange::FSkeletalMeshNodeStaticData::GetLodDependenciesBaseKey());
-	}
+	UInterchangeSkeletalMeshFactoryNode();
 
 	/**
 	 * Initialize node data
@@ -51,226 +36,97 @@ public:
 	 * @param InAssetClass - The class the SkeletalMesh factory will create for this node.
 	 */
 	UFUNCTION(BlueprintCallable, Category = "Interchange | Node | SkeletalMesh")
-	void InitializeSkeletalMeshNode(const FString& UniqueID, const FString& DisplayLabel, const FString& InAssetClass)
-	{
-		bIsNodeClassInitialized = false;
-		InitializeNode(UniqueID, DisplayLabel, EInterchangeNodeContainerType::NodeContainerType_FactoryData);
-
-		FString OperationName = GetTypeName() + TEXT(".SetAssetClassName");
-		InterchangePrivateNodeBase::SetCustomAttribute<FString>(*Attributes, ClassNameAttributeKey, OperationName, InAssetClass);
-		FillAssetClassFromAttribute();
-	}
+	void InitializeSkeletalMeshNode(const FString& UniqueID, const FString& DisplayLabel, const FString& InAssetClass);
 
-	virtual void Serialize(FArchive& Ar) override
-	{
-		Super::Serialize(Ar);
-#if WITH_ENGINE
-		if (Ar.IsLoading())
-		{
-			//Make sure the class is properly set when we compile with engine, this will set the
-			//bIsNodeClassInitialized to true.
-			SetNodeClassFromClassAttribute();
-		}
-#endif //#if WITH_ENGINE
-	}
+	virtual void Serialize(FArchive& Ar) override;
 
 	/**
 	 * Return the node type name of the class, we use this when reporting error
 	 */
-	virtual FString GetTypeName() const override
-	{
-		const FString TypeName = TEXT("SkeletalMeshNode");
-		return TypeName;
-	}
+	virtual FString GetTypeName() const override;
 
-	virtual FString GetKeyDisplayName(const UE::Interchange::FAttributeKey& NodeAttributeKey) const override
-	{
-		FString KeyDisplayName = NodeAttributeKey.Key;
-		if (NodeAttributeKey.Key.Equals(UE::Interchange::FSkeletalMeshNodeStaticData::GetLodDependenciesBaseKey()))
-		{
-			KeyDisplayName = TEXT("LOD Dependencies Count");
-			return KeyDisplayName;
-		}
-		else if (NodeAttributeKey.Key.StartsWith(UE::Interchange::FSkeletalMeshNodeStaticData::GetLodDependenciesBaseKey()))
-		{
-			KeyDisplayName = TEXT("LOD Dependencies Index ");
-			const FString IndexKey = UE::Interchange::FNameAttributeArrayHelper::IndexKey();
-			int32 IndexPosition = NodeAttributeKey.Key.Find(IndexKey) + IndexKey.Len();
-			if (IndexPosition < NodeAttributeKey.Key.Len())
-			{
-				KeyDisplayName += NodeAttributeKey.Key.RightChop(IndexPosition);
-			}
-			return KeyDisplayName;
-		}
-		return Super::GetKeyDisplayName(NodeAttributeKey);
-	}
+	virtual FString GetKeyDisplayName(const UE::Interchange::FAttributeKey& NodeAttributeKey) const override;
 
 	/** Get the class this node want to create */
 	UFUNCTION(BlueprintCallable, Category = "Interchange | Node | SkeletalMesh")
-	virtual class UClass* GetObjectClass() const override
-	{
-		ensure(bIsNodeClassInitialized);
-#if WITH_ENGINE
-		return AssetClass.Get() != nullptr ? AssetClass.Get() : USkeletalMesh::StaticClass();
-#else
-		return nullptr;
-#endif
-	}
-
-	virtual FGuid GetHash() const override
-	{
-		return Attributes->GetStorageHash();
-	}
+	virtual class UClass* GetObjectClass() const override;
 
 public:
 	/** Return The number of LOD this skeletalmesh has.*/
 	UFUNCTION(BlueprintCallable, Category = "Interchange | Node | SkeletalMesh")
-	int32 GetLodDataCount() const
-	{
-		return LodDependencies.GetCount();
-	}
+	int32 GetLodDataCount() const;
 
 	UFUNCTION(BlueprintCallable, Category = "Interchange | Node | SkeletalMesh")
-	void GetLodDataUniqueIds(TArray<FString>& OutLodDataUniqueIds) const
-	{
-		LodDependencies.GetNames(OutLodDataUniqueIds);
-	}
+	void GetLodDataUniqueIds(TArray<FString>& OutLodDataUniqueIds) const;
 
 	UFUNCTION(BlueprintCallable, Category = "Interchange | Node | SkeletalMesh")
-	bool AddLodDataUniqueId(const FString& LodDataUniqueId)
-	{
-		return LodDependencies.AddName(LodDataUniqueId);
-	}
+	bool AddLodDataUniqueId(const FString& LodDataUniqueId);
 
 	UFUNCTION(BlueprintCallable, Category = "Interchange | Node | SkeletalMesh")
-	bool RemoveLodDataUniqueId(const FString& LodDataUniqueId)
-	{
-		return LodDependencies.RemoveName(LodDataUniqueId);
-	}
+	bool RemoveLodDataUniqueId(const FString& LodDataUniqueId);
 	
-	/** Return false if the Attribute was not set previously.*/
+	/** Query the skeletal mesh factory skeleton UObject. Return false if the attribute was not set.*/
 	UFUNCTION(BlueprintCallable, Category = "Interchange | Node | SkeletalMesh")
-	bool GetCustomSkeletonSoftObjectPath(FSoftObjectPath& AttributeValue) const
-	{
-		IMPLEMENT_NODE_ATTRIBUTE_GETTER(SkeletonSoftObjectPath, FSoftObjectPath)
-	}
+	bool GetCustomSkeletonSoftObjectPath(FSoftObjectPath& AttributeValue) const;
 
+	/** Set the skeletal mesh factory skeleton UObject. Return false if the attribute cannot be set.*/
 	UFUNCTION(BlueprintCallable, Category = "Interchange | Node | SkeletalMesh")
-	bool SetCustomSkeletonSoftObjectPath(const FSoftObjectPath& AttributeValue)
-	{
-		IMPLEMENT_NODE_ATTRIBUTE_SETTER_NODELEGATE(SkeletonSoftObjectPath, FSoftObjectPath)
-	}
+	bool SetCustomSkeletonSoftObjectPath(const FSoftObjectPath& AttributeValue);
 
-	/** Return false if the Attribute was not set previously.*/
+	/** Query weather the skeletal mesh factory should create the morph target. Return false if the attribute was not set.*/
 	UFUNCTION(BlueprintCallable, Category = "Interchange | Node | SkeletalMesh")
-	bool GetCustomImportMorphTarget(bool& AttributeValue) const
-	{
-		IMPLEMENT_NODE_ATTRIBUTE_GETTER(ImportMorphTarget, bool)
-	}
+	bool GetCustomImportMorphTarget(bool& AttributeValue) const;
 
+	/** Set weather the skeletal mesh factory should create the morph target. Return false if the attribute cannot be set.*/
 	UFUNCTION(BlueprintCallable, Category = "Interchange | Node | SkeletalMesh")
-	bool SetCustomImportMorphTarget(const bool& AttributeValue)
-	{
-		IMPLEMENT_NODE_ATTRIBUTE_SETTER_NODELEGATE(ImportMorphTarget, bool)
-	}
+	bool SetCustomImportMorphTarget(const bool& AttributeValue);
 
-	/** Return false if the Attribute was not set previously.*/
+	/** Query weather the skeletal mesh factory should create a physics asset. Return false if the attribute was not set.*/
 	UFUNCTION(BlueprintCallable, Category = "Interchange | Node | SkeletalMesh")
-	bool GetCustomCreatePhysicsAsset(bool& AttributeValue) const
-	{
-		IMPLEMENT_NODE_ATTRIBUTE_GETTER(CreatePhysicsAsset, bool)
-	}
+	bool GetCustomCreatePhysicsAsset(bool& AttributeValue) const;
 
+	/** Set weather the skeletal mesh factory should create a physics asset. Return false if the attribute cannot be set.*/
 	UFUNCTION(BlueprintCallable, Category = "Interchange | Node | SkeletalMesh")
-	bool SetCustomCreatePhysicsAsset(const bool& AttributeValue)
-	{
-		IMPLEMENT_NODE_ATTRIBUTE_SETTER_NODELEGATE(CreatePhysicsAsset, bool)
-	}
+	bool SetCustomCreatePhysicsAsset(const bool& AttributeValue);
 
-	/** Return false if the Attribute was not set previously.*/
+	/** Query a physics asset the skeletal mesh factory should use. Return false if the attribute was not set.*/
 	UFUNCTION(BlueprintCallable, Category = "Interchange | Node | SkeletalMesh")
-	bool GetCustomPhysicAssetSoftObjectPath(FSoftObjectPath& AttributeValue) const
-	{
-		IMPLEMENT_NODE_ATTRIBUTE_GETTER(PhysicAssetSoftObjectPath, FSoftObjectPath)
-	}
+	bool GetCustomPhysicAssetSoftObjectPath(FSoftObjectPath& AttributeValue) const;
 
+	/** Set a physics asset the skeletal mesh factory should use. Return false if the attribute cannot be set.*/
 	UFUNCTION(BlueprintCallable, Category = "Interchange | Node | SkeletalMesh")
-	bool SetCustomPhysicAssetSoftObjectPath(const FSoftObjectPath& AttributeValue)
-	{
-		IMPLEMENT_NODE_ATTRIBUTE_SETTER_NODELEGATE(PhysicAssetSoftObjectPath, FSoftObjectPath)
-	}
+	bool SetCustomPhysicAssetSoftObjectPath(const FSoftObjectPath& AttributeValue);
 
-	/** Return false if the Attribute was not set previously.*/
+	/** Query weather the skeletal mesh factory should replace the vertex color. Return false if the attribute was not set.*/
 	UFUNCTION(BlueprintCallable, Category = "Interchange | Node | SkeletalMesh")
-	bool GetCustomVertexColorReplace(bool& AttributeValue) const
-	{
-		IMPLEMENT_NODE_ATTRIBUTE_GETTER(VertexColorReplace, bool)
-	}
+	bool GetCustomVertexColorReplace(bool& AttributeValue) const;
 
+	/** Set weather the skeletal mesh factory should replace the vertex color. Return false if the attribute cannot be set.*/
 	UFUNCTION(BlueprintCallable, Category = "Interchange | Node | SkeletalMesh")
-	bool SetCustomVertexColorReplace(const bool& AttributeValue)
-	{
-		IMPLEMENT_NODE_ATTRIBUTE_SETTER_NODELEGATE(VertexColorReplace, bool)
-	}
+	bool SetCustomVertexColorReplace(const bool& AttributeValue);
 
-	/** Return false if the Attribute was not set previously.*/
+	/** Query weather the skeletal mesh factory should ignore the vertex color. Return false if the attribute was not set.*/
 	UFUNCTION(BlueprintCallable, Category = "Interchange | Node | SkeletalMesh")
-	bool GetCustomVertexColorIgnore(bool& AttributeValue) const
-	{
-		IMPLEMENT_NODE_ATTRIBUTE_GETTER(VertexColorIgnore, bool)
-	}
+	bool GetCustomVertexColorIgnore(bool& AttributeValue) const;
 
+	/** Set weather the skeletal mesh factory should ignore the vertex color. Return false if the attribute cannot be set.*/
 	UFUNCTION(BlueprintCallable, Category = "Interchange | Node | SkeletalMesh")
-	bool SetCustomVertexColorIgnore(const bool& AttributeValue)
-	{
-		IMPLEMENT_NODE_ATTRIBUTE_SETTER_NODELEGATE(VertexColorIgnore, bool)
-	}
+	bool SetCustomVertexColorIgnore(const bool& AttributeValue);
 
-	/** Return false if the Attribute was not set previously.*/
+	/** Query weather the skeletal mesh factory should override the vertex color. Return false if the attribute was not set.*/
 	UFUNCTION(BlueprintCallable, Category = "Interchange | Node | SkeletalMesh")
-	bool GetCustomVertexColorOverride(FColor& AttributeValue) const
-	{
-		IMPLEMENT_NODE_ATTRIBUTE_GETTER(VertexColorOverride, FColor)
-	}
+	bool GetCustomVertexColorOverride(FColor& AttributeValue) const;
 
+	/** Set weather the skeletal mesh factory should override the vertex color. Return false if the attribute cannot be set.*/
 	UFUNCTION(BlueprintCallable, Category = "Interchange | Node | SkeletalMesh")
-	bool SetCustomVertexColorOverride(const FColor& AttributeValue)
-	{
-		IMPLEMENT_NODE_ATTRIBUTE_SETTER_NODELEGATE(VertexColorOverride, FColor)
-	}
+	bool SetCustomVertexColorOverride(const FColor& AttributeValue);
 private:
 
-	void FillAssetClassFromAttribute()
-	{
-#if WITH_ENGINE
-		FString OperationName = GetTypeName() + TEXT(".GetAssetClassName");
-		FString ClassName;
-		InterchangePrivateNodeBase::GetCustomAttribute<FString>(*Attributes, ClassNameAttributeKey, OperationName, ClassName);
-		if (ClassName.Equals(USkeletalMesh::StaticClass()->GetName()))
-		{
-			AssetClass = USkeletalMesh::StaticClass();
-			bIsNodeClassInitialized = true;
-		}
-#endif
-	}
+	void FillAssetClassFromAttribute();
 
-	bool SetNodeClassFromClassAttribute()
-	{
-		if (!bIsNodeClassInitialized)
-		{
-			FillAssetClassFromAttribute();
-		}
-		return bIsNodeClassInitialized;
-	}
+	bool SetNodeClassFromClassAttribute();
 
-	bool IsEditorOnlyDataDefined()
-	{
-#if WITH_EDITORONLY_DATA
-		return true;
-#else
-		return false;
-#endif
-	}
+	bool IsEditorOnlyDataDefined();
 
 	const UE::Interchange::FAttributeKey ClassNameAttributeKey = UE::Interchange::FBaseNodeStaticData::ClassTypeAttributeKey();
 	const UE::Interchange::FAttributeKey Macro_CustomImportMorphTargetKey = UE::Interchange::FAttributeKey(TEXT("ImportMorphTarget"));
@@ -281,10 +137,10 @@ private:
 	const UE::Interchange::FAttributeKey Macro_CustomVertexColorIgnoreKey = UE::Interchange::FAttributeKey(TEXT("VertexColorIgnore"));
 	const UE::Interchange::FAttributeKey Macro_CustomVertexColorOverrideKey = UE::Interchange::FAttributeKey(TEXT("VertexColorOverride"));
 
-	UE::Interchange::FNameAttributeArrayHelper LodDependencies;
+	UE::Interchange::TArrayAttributeHelper<FString> LodDependencies;
 protected:
 #if WITH_ENGINE
-	TSubclassOf<USkeletalMesh> AssetClass = nullptr;
+	TSubclassOf<class USkeletalMesh> AssetClass = nullptr;
 #endif
 	bool bIsNodeClassInitialized = false;
 };
diff --git a/Engine/Plugins/Experimental/Interchange/Runtime/Source/Nodes/Public/InterchangeSkeletalMeshLodDataNode.h b/Engine/Plugins/Experimental/Interchange/Runtime/Source/Nodes/Public/InterchangeSkeletalMeshLodDataNode.h
--- a/Engine/Plugins/Experimental/Interchange/Runtime/Source/Nodes/Public/InterchangeSkeletalMeshLodDataNode.h
+++ b/Engine/Plugins/Experimental/Interchange/Runtime/Source/Nodes/Public/InterchangeSkeletalMeshLodDataNode.h
@@ -13,13 +13,9 @@ namespace UE
 	namespace Interchange
 	{
 
-		struct FSkeletalMeshNodeLodDataStaticData : public FBaseNodeStaticData
+		struct INTERCHANGENODES_API FSkeletalMeshNodeLodDataStaticData : public FBaseNodeStaticData
 		{
-			static const FString& GetMeshUidsBaseKey()
-			{
-				static FString MeshUids_BaseKey = TEXT("__MeshUids__Key");
-				return MeshUids_BaseKey;
-			}
+			static const FString& GetMeshUidsBaseKey();
 		};
 
 	}//ns Interchange
@@ -31,123 +27,52 @@ class INTERCHANGENODES_API UInterchangeSkeletalMeshLodDataNode : public UInterch
 	GENERATED_BODY()
 
 public:
-	UInterchangeSkeletalMeshLodDataNode()
-		:UInterchangeBaseNode()
-	{
-		MeshUids.Initialize(Attributes, UE::Interchange::FSkeletalMeshNodeLodDataStaticData::GetMeshUidsBaseKey());
-	}
+	UInterchangeSkeletalMeshLodDataNode();
 
 	/**
 	 * Return the node type name of the class, we use this when reporting error
 	 */
-	virtual FString GetTypeName() const override
-	{
-		const FString TypeName = TEXT("SkeletalMeshLodDataNode");
-		return TypeName;
-	}
+	virtual FString GetTypeName() const override;
 
-	virtual FString GetKeyDisplayName(const UE::Interchange::FAttributeKey& NodeAttributeKey) const override
-	{
-		FString KeyDisplayName = NodeAttributeKey.Key;
-		if (NodeAttributeKey.Key.Equals(UE::Interchange::FSkeletalMeshNodeLodDataStaticData::GetMeshUidsBaseKey()))
-		{
-			KeyDisplayName = TEXT("Mesh count");
-			return KeyDisplayName;
-		}
-		else if (NodeAttributeKey.Key.StartsWith(UE::Interchange::FSkeletalMeshNodeLodDataStaticData::GetMeshUidsBaseKey()))
-		{
-			KeyDisplayName = TEXT("Mesh index ");
-			const FString IndexKey = UE::Interchange::FNameAttributeArrayHelper::IndexKey();
-			int32 IndexPosition = NodeAttributeKey.Key.Find(IndexKey) + IndexKey.Len();
-			if (IndexPosition < NodeAttributeKey.Key.Len())
-			{
-				KeyDisplayName += NodeAttributeKey.Key.RightChop(IndexPosition);
-			}
-			return KeyDisplayName;
-		}
-		else if (NodeAttributeKey == Macro_CustomSkeletonUidKey)
-		{
-			KeyDisplayName = TEXT("Skeleton factory node");
-			return KeyDisplayName;
-		}
-		return Super::GetKeyDisplayName(NodeAttributeKey);
-	}
-
-	virtual FString GetAttributeCategory(const UE::Interchange::FAttributeKey& NodeAttributeKey) const override
-	{
-		if (NodeAttributeKey.Key.StartsWith(UE::Interchange::FSkeletalMeshNodeLodDataStaticData::GetMeshUidsBaseKey()))
-		{
-			return FString(TEXT("Meshes"));
-		}
-		return Super::GetAttributeCategory(NodeAttributeKey);
-	}
-
-	virtual FGuid GetHash() const override
-	{
-		return Attributes->GetStorageHash();
-	}
+	virtual FString GetKeyDisplayName(const UE::Interchange::FAttributeKey& NodeAttributeKey) const override;
 
+	virtual FString GetAttributeCategory(const UE::Interchange::FAttributeKey& NodeAttributeKey) const override;
 
 public:
-	/** Return false if the Attribute was not set previously.*/
+	/** Query the LOD skeletal mesh factory skeleton reference. Return false if the attribute was not set.*/
 	UFUNCTION(BlueprintCallable, Category = "Interchange | Node | SkeletalMeshLodData")
-	bool GetCustomSkeletonUid(FString& AttributeValue) const
-	{
-		IMPLEMENT_NODE_ATTRIBUTE_GETTER(SkeletonUid, FString);
-	}
+	bool GetCustomSkeletonUid(FString& AttributeValue) const;
 
+	/** Set the LOD skeletal mesh factory skeleton reference. Return false if the attribute cannot be set.*/
 	UFUNCTION(BlueprintCallable, Category = "Interchange | Node | SkeletalMeshLodData")
-	bool SetCustomSkeletonUid(const FString& AttributeValue)
-	{
-		IMPLEMENT_NODE_ATTRIBUTE_SETTER_NODELEGATE(SkeletonUid, FString)
-	}
-
-	/* Mesh Uids: It can be either a scene or a mesh node uid. If its a scene it mean we want the mesh factory to bake the geo payload with the global transform of the scene node. */
+	bool SetCustomSkeletonUid(const FString& AttributeValue);
 
+	/* Return the number of mesh geometry this LOD will be made of. Mesh uids can be either a scene or a mesh node. If its a scene it mean we want the mesh factory to bake the geo payload with the global transform of the scene node. */
 	UFUNCTION(BlueprintCallable, Category = "Interchange | Node | SkeletalMeshLodData")
-	int32 GetMeshUidsCount() const
-	{
-		return MeshUids.GetCount();
-	}
+	int32 GetMeshUidsCount() const;
 
+	/* Query all mesh geometry this LOD will be made of. Mesh uids can be either a scene or a mesh node. If its a scene it mean we want the mesh factory to bake the geo payload with the global transform of the scene node. */
 	UFUNCTION(BlueprintCallable, Category = "Interchange | Node | SkeletalMeshLodData")
-	void GetMeshUids(TArray<FString>& OutBlendShapeNames) const
-	{
-		MeshUids.GetNames(OutBlendShapeNames);
-	}
+	void GetMeshUids(TArray<FString>& OutMeshNames) const;
 
+	/* Add one mesh geometry use to create this LOD geometry. Mesh uids can be either a scene or a mesh node. If its a scene it mean we want the mesh factory to bake the geo payload with the global transform of the scene node. */
 	UFUNCTION(BlueprintCallable, Category = "Interchange | Node | SkeletalMeshLodData")
-	bool AddMeshUid(const FString& BlendShapeName)
-	{
-		return MeshUids.AddName(BlendShapeName);
-	}
+	bool AddMeshUid(const FString& MeshName);
 
+	/* Remove one mesh geometry use to create this LOD geometry. Mesh uids can be either a scene or a mesh node. If its a scene it mean we want the mesh factory to bake the geo payload with the global transform of the scene node. */
 	UFUNCTION(BlueprintCallable, Category = "Interchange | Node | SkeletalMeshLodData")
-	bool RemoveMeshUid(const FString& BlendShapeName)
-	{
-		return MeshUids.RemoveName(BlendShapeName);
-	}
+	bool RemoveMeshUid(const FString& MeshName);
 
+	/* Remove all mesh geometry use to create this LOD geometry. Mesh uids can be either a scene or a mesh node. If its a scene it mean we want the mesh factory to bake the geo payload with the global transform of the scene node. */
 	UFUNCTION(BlueprintCallable, Category = "Interchange | Node | SkeletalMeshLodData")
-	bool RemoveAllMeshes()
-	{
-		return MeshUids.RemoveAllNames();
-	}
+	bool RemoveAllMeshes();
 
 private:
 
-	bool IsEditorOnlyDataDefined()
-	{
-#if WITH_EDITORONLY_DATA
-		return true;
-#else
-		return false;
-#endif
-	}
-
-	//SkeletalMesh
+	bool IsEditorOnlyDataDefined();
+
 	const UE::Interchange::FAttributeKey Macro_CustomSkeletonUidKey = UE::Interchange::FAttributeKey(TEXT("__SkeletonUid__Key"));
 
-	UE::Interchange::FNameAttributeArrayHelper MeshUids;
+	UE::Interchange::TArrayAttributeHelper<FString> MeshUids;
 protected:
 };
diff --git a/Engine/Plugins/Experimental/Interchange/Runtime/Source/Nodes/Public/InterchangeSkeletonFactoryNode.h b/Engine/Plugins/Experimental/Interchange/Runtime/Source/Nodes/Public/InterchangeSkeletonFactoryNode.h
--- a/Engine/Plugins/Experimental/Interchange/Runtime/Source/Nodes/Public/InterchangeSkeletonFactoryNode.h
+++ b/Engine/Plugins/Experimental/Interchange/Runtime/Source/Nodes/Public/InterchangeSkeletonFactoryNode.h
@@ -35,7 +35,7 @@ public:
 	void InitializeSkeletonNode(const FString& UniqueID, const FString& DisplayLabel, const FString& InAssetClass)
 	{
 		bIsNodeClassInitialized = false;
-		InitializeNode(UniqueID, DisplayLabel, EInterchangeNodeContainerType::NodeContainerType_FactoryData);
+		InitializeNode(UniqueID, DisplayLabel, EInterchangeNodeContainerType::FactoryData);
 		FString OperationName = GetTypeName() + TEXT(".SetAssetClassName");
 		InterchangePrivateNodeBase::SetCustomAttribute<FString>(*Attributes, ClassNameAttributeKey, OperationName, InAssetClass);
 		FillAssetClassFromAttribute();
diff --git a/Engine/Plugins/Experimental/Interchange/Runtime/Source/Nodes/Public/InterchangeStaticMeshFactoryNode.h b/Engine/Plugins/Experimental/Interchange/Runtime/Source/Nodes/Public/InterchangeStaticMeshFactoryNode.h
--- a/Engine/Plugins/Experimental/Interchange/Runtime/Source/Nodes/Public/InterchangeStaticMeshFactoryNode.h
+++ b/Engine/Plugins/Experimental/Interchange/Runtime/Source/Nodes/Public/InterchangeStaticMeshFactoryNode.h
@@ -5,10 +5,6 @@
 #include "CoreMinimal.h"
 #include "Nodes/InterchangeBaseNode.h"
 
-#if WITH_ENGINE
-#include "Engine/StaticMesh.h"
-#endif
-
 #include "InterchangeStaticMeshFactoryNode.generated.h"
 
 
@@ -16,15 +12,10 @@ namespace UE
 {
 	namespace Interchange
 	{
-		struct FStaticMeshNodeStaticData : public FBaseNodeStaticData
+		struct INTERCHANGENODES_API FStaticMeshNodeStaticData : public FBaseNodeStaticData
 		{
-			static const FString& GetLodDependenciesBaseKey()
-			{
-				static FString LodDependencies_BaseKey = TEXT("Lod_Dependencies");
-				return LodDependencies_BaseKey;
-			}
+			static const FString& GetLodDependenciesBaseKey();
 		};
-
 	} // namespace Interchange
 } // namespace UE
 
@@ -35,13 +26,7 @@ class INTERCHANGENODES_API UInterchangeStaticMeshFactoryNode : public UInterchan
 	GENERATED_BODY()
 
 public:
-	UInterchangeStaticMeshFactoryNode()
-	{
-#if WITH_ENGINE
-		AssetClass = nullptr;
-#endif
-		LodDependencies.Initialize(Attributes, UE::Interchange::FStaticMeshNodeStaticData::GetLodDependenciesBaseKey());
-	}
+	UInterchangeStaticMeshFactoryNode();
 
 	/**
 	 * Initialize node data
@@ -50,183 +35,74 @@ public:
 	 * @param InAssetClass - The class the StaticMesh factory will create for this node.
 	 */
 	UFUNCTION(BlueprintCallable, Category = "Interchange | Node | StaticMesh")
-	void InitializeStaticMeshNode(const FString& UniqueID, const FString& DisplayLabel, const FString& InAssetClass)
-	{
-		bIsNodeClassInitialized = false;
-		InitializeNode(UniqueID, DisplayLabel, EInterchangeNodeContainerType::NodeContainerType_FactoryData);
+	void InitializeStaticMeshNode(const FString& UniqueID, const FString& DisplayLabel, const FString& InAssetClass);
 
-		FString OperationName = GetTypeName() + TEXT(".SetAssetClassName");
-		InterchangePrivateNodeBase::SetCustomAttribute<FString>(*Attributes, ClassNameAttributeKey, OperationName, InAssetClass);
-		FillAssetClassFromAttribute();
-	}
-
-	virtual void Serialize(FArchive& Ar) override
-	{
-		Super::Serialize(Ar);
-#if WITH_ENGINE
-		if (Ar.IsLoading())
-		{
-			// Make sure the class is properly set when we compile with engine, this will set the bIsNodeClassInitialized to true.
-			SetNodeClassFromClassAttribute();
-		}
-#endif //#if WITH_ENGINE
-	}
+	virtual void Serialize(FArchive& Ar) override;
 
 	/**
 	 * Return the node type name of the class, we use this when reporting error
 	 */
-	virtual FString GetTypeName() const override
-	{
-		const FString TypeName = TEXT("StaticMeshNode");
-		return TypeName;
-	}
+	virtual FString GetTypeName() const override;
 
-	virtual FString GetKeyDisplayName(const UE::Interchange::FAttributeKey& NodeAttributeKey) const override
-	{
-		FString KeyDisplayName = NodeAttributeKey.Key;
-		if (NodeAttributeKey.Key.Equals(UE::Interchange::FStaticMeshNodeStaticData::GetLodDependenciesBaseKey()))
-		{
-			KeyDisplayName = TEXT("LOD Dependencies Count");
-			return KeyDisplayName;
-		}
-		else if (NodeAttributeKey.Key.StartsWith(UE::Interchange::FStaticMeshNodeStaticData::GetLodDependenciesBaseKey()))
-		{
-			KeyDisplayName = TEXT("LOD Dependencies Index ");
-			const FString IndexKey = UE::Interchange::FNameAttributeArrayHelper::IndexKey();
-			int32 IndexPosition = NodeAttributeKey.Key.Find(IndexKey) + IndexKey.Len();
-			if (IndexPosition < NodeAttributeKey.Key.Len())
-			{
-				KeyDisplayName += NodeAttributeKey.Key.RightChop(IndexPosition);
-			}
-			return KeyDisplayName;
-		}
-		return Super::GetKeyDisplayName(NodeAttributeKey);
-	}
+	virtual FString GetKeyDisplayName(const UE::Interchange::FAttributeKey& NodeAttributeKey) const override;
 
 	/** Get the class this node want to create */
 	UFUNCTION(BlueprintCallable, Category = "Interchange | Node | StaticMesh")
-	virtual class UClass* GetObjectClass() const override
-	{
-		ensure(bIsNodeClassInitialized);
-#if WITH_ENGINE
-		return AssetClass.Get() != nullptr ? AssetClass.Get() : UStaticMesh::StaticClass();
-#else
-		return nullptr;
-#endif
-	}
-
-	virtual FGuid GetHash() const override
-	{
-		return Attributes->GetStorageHash();
-	}
+	virtual class UClass* GetObjectClass() const override;
 
 public:
 	/** Return The number of LOD this static mesh has.*/
 	UFUNCTION(BlueprintCallable, Category = "Interchange | Node | StaticMesh")
-	int32 GetLodDataCount() const
-	{
-		return LodDependencies.GetCount();
-	}
+	int32 GetLodDataCount() const;
 
 	UFUNCTION(BlueprintCallable, Category = "Interchange | Node | StaticMesh")
-	void GetLodDataUniqueIds(TArray<FString>& OutLodDataUniqueIds) const
-	{
-		LodDependencies.GetNames(OutLodDataUniqueIds);
-	}
+	void GetLodDataUniqueIds(TArray<FString>& OutLodDataUniqueIds) const;
 
 	UFUNCTION(BlueprintCallable, Category = "Interchange | Node | StaticMesh")
-	bool AddLodDataUniqueId(const FString& LodDataUniqueId)
-	{
-		return LodDependencies.AddName(LodDataUniqueId);
-	}
+	bool AddLodDataUniqueId(const FString& LodDataUniqueId);
 
 	UFUNCTION(BlueprintCallable, Category = "Interchange | Node | StaticMesh")
-	bool RemoveLodDataUniqueId(const FString& LodDataUniqueId)
-	{
-		return LodDependencies.RemoveName(LodDataUniqueId);
-	}
+	bool RemoveLodDataUniqueId(const FString& LodDataUniqueId);
 
-	/** Return false if the Attribute was not set previously.*/
+	/** Query weather the static mesh factory should replace the vertex color. Return false if the attribute was not set.*/
 	UFUNCTION(BlueprintCallable, Category = "Interchange | Node | StaticMesh")
-	bool GetCustomVertexColorReplace(bool& AttributeValue) const
-	{
-		IMPLEMENT_NODE_ATTRIBUTE_GETTER(VertexColorReplace, bool)
-	}
+	bool GetCustomVertexColorReplace(bool& AttributeValue) const;
 
+	/** Set weather the static mesh factory should replace the vertex color. Return false if the attribute cannot be set.*/
 	UFUNCTION(BlueprintCallable, Category = "Interchange | Node | StaticMesh")
-	bool SetCustomVertexColorReplace(const bool& AttributeValue)
-	{
-		IMPLEMENT_NODE_ATTRIBUTE_SETTER_NODELEGATE(VertexColorReplace, bool)
-	}
+	bool SetCustomVertexColorReplace(const bool& AttributeValue);
 
-	/** Return false if the Attribute was not set previously.*/
+	/** Query weather the static mesh factory should ignore the vertex color. Return false if the attribute was not set.*/
 	UFUNCTION(BlueprintCallable, Category = "Interchange | Node | StaticMesh")
-	bool GetCustomVertexColorIgnore(bool& AttributeValue) const
-	{
-		IMPLEMENT_NODE_ATTRIBUTE_GETTER(VertexColorIgnore, bool)
-	}
+	bool GetCustomVertexColorIgnore(bool& AttributeValue) const;
 
+	/** Set weather the static mesh factory should ignore the vertex color. Return false if the attribute cannot be set.*/
 	UFUNCTION(BlueprintCallable, Category = "Interchange | Node | StaticMesh")
-	bool SetCustomVertexColorIgnore(const bool& AttributeValue)
-	{
-		IMPLEMENT_NODE_ATTRIBUTE_SETTER_NODELEGATE(VertexColorIgnore, bool)
-	}
+	bool SetCustomVertexColorIgnore(const bool& AttributeValue);
 
-	/** Return false if the Attribute was not set previously.*/
+	/** Query weather the static mesh factory should override the vertex color. Return false if the attribute was not set.*/
 	UFUNCTION(BlueprintCallable, Category = "Interchange | Node | StaticMesh")
-	bool GetCustomVertexColorOverride(FColor& AttributeValue) const
-	{
-		IMPLEMENT_NODE_ATTRIBUTE_GETTER(VertexColorOverride, FColor)
-	}
+	bool GetCustomVertexColorOverride(FColor& AttributeValue) const;
 
+	/** Set weather the static mesh factory should override the vertex color. Return false if the attribute cannot be set.*/
 	UFUNCTION(BlueprintCallable, Category = "Interchange | Node | StaticMesh")
-	bool SetCustomVertexColorOverride(const FColor& AttributeValue)
-	{
-		IMPLEMENT_NODE_ATTRIBUTE_SETTER_NODELEGATE(VertexColorOverride, FColor)
-	}
-private:
+	bool SetCustomVertexColorOverride(const FColor& AttributeValue);
 
-	void FillAssetClassFromAttribute()
-	{
-#if WITH_ENGINE
-		FString OperationName = GetTypeName() + TEXT(".GetAssetClassName");
-		FString ClassName;
-		InterchangePrivateNodeBase::GetCustomAttribute<FString>(*Attributes, ClassNameAttributeKey, OperationName, ClassName);
-		if (ClassName.Equals(UStaticMesh::StaticClass()->GetName()))
-		{
-			AssetClass = UStaticMesh::StaticClass();
-			bIsNodeClassInitialized = true;
-		}
-#endif
-	}
-
-	bool SetNodeClassFromClassAttribute()
-	{
-		if (!bIsNodeClassInitialized)
-		{
-			FillAssetClassFromAttribute();
-		}
-		return bIsNodeClassInitialized;
-	}
+private:
 
-	bool IsEditorOnlyDataDefined()
-	{
-#if WITH_EDITORONLY_DATA
-		return true;
-#else
-		return false;
-#endif
-	}
+	void FillAssetClassFromAttribute();
+	bool SetNodeClassFromClassAttribute();
+	bool IsEditorOnlyDataDefined();
 
 	const UE::Interchange::FAttributeKey ClassNameAttributeKey = UE::Interchange::FBaseNodeStaticData::ClassTypeAttributeKey();
 	const UE::Interchange::FAttributeKey Macro_CustomVertexColorReplaceKey = UE::Interchange::FAttributeKey(TEXT("VertexColorReplace"));
 	const UE::Interchange::FAttributeKey Macro_CustomVertexColorIgnoreKey = UE::Interchange::FAttributeKey(TEXT("VertexColorIgnore"));
 	const UE::Interchange::FAttributeKey Macro_CustomVertexColorOverrideKey = UE::Interchange::FAttributeKey(TEXT("VertexColorOverride"));
 
-	UE::Interchange::FNameAttributeArrayHelper LodDependencies;
+	UE::Interchange::TArrayAttributeHelper<FString> LodDependencies;
 protected:
 #if WITH_ENGINE
-	TSubclassOf<UStaticMesh> AssetClass = nullptr;
+	TSubclassOf<class UStaticMesh> AssetClass = nullptr;
 #endif
 	bool bIsNodeClassInitialized = false;
 };
diff --git a/Engine/Plugins/Experimental/Interchange/Runtime/Source/Nodes/Public/InterchangeStaticMeshLodDataNode.h b/Engine/Plugins/Experimental/Interchange/Runtime/Source/Nodes/Public/InterchangeStaticMeshLodDataNode.h
--- a/Engine/Plugins/Experimental/Interchange/Runtime/Source/Nodes/Public/InterchangeStaticMeshLodDataNode.h
+++ b/Engine/Plugins/Experimental/Interchange/Runtime/Source/Nodes/Public/InterchangeStaticMeshLodDataNode.h
@@ -14,13 +14,8 @@ namespace UE
 	{
 		struct FStaticMeshNodeLodDataStaticData : public FBaseNodeStaticData
 		{
-			static const FString& GetMeshUidsBaseKey()
-			{
-				static FString MeshUids_BaseKey = TEXT("__MeshUids__Key");
-				return MeshUids_BaseKey;
-			}
+			static const FString& GetMeshUidsBaseKey();
 		};
-
 	} // namespace Interchange
 } // namespace UE
 
@@ -31,102 +26,39 @@ class INTERCHANGENODES_API UInterchangeStaticMeshLodDataNode : public UInterchan
 	GENERATED_BODY()
 
 public:
-	UInterchangeStaticMeshLodDataNode()
-	{
-		MeshUids.Initialize(Attributes, UE::Interchange::FStaticMeshNodeLodDataStaticData::GetMeshUidsBaseKey());
-	}
+	UInterchangeStaticMeshLodDataNode();
 
 	/**
 	 * Return the node type name of the class, we use this when reporting error
 	 */
-	virtual FString GetTypeName() const override
-	{
-		const FString TypeName = TEXT("StaticMeshLodDataNode");
-		return TypeName;
-	}
+	virtual FString GetTypeName() const override;
 
-	virtual FString GetKeyDisplayName(const UE::Interchange::FAttributeKey& NodeAttributeKey) const override
-	{
-		FString KeyDisplayName = NodeAttributeKey.Key;
-		if (NodeAttributeKey.Key.Equals(UE::Interchange::FStaticMeshNodeLodDataStaticData::GetMeshUidsBaseKey()))
-		{
-			KeyDisplayName = TEXT("Mesh count");
-			return KeyDisplayName;
-		}
-		else if (NodeAttributeKey.Key.StartsWith(UE::Interchange::FStaticMeshNodeLodDataStaticData::GetMeshUidsBaseKey()))
-		{
-			KeyDisplayName = TEXT("Mesh index ");
-			const FString IndexKey = UE::Interchange::FNameAttributeArrayHelper::IndexKey();
-			int32 IndexPosition = NodeAttributeKey.Key.Find(IndexKey) + IndexKey.Len();
-			if (IndexPosition < NodeAttributeKey.Key.Len())
-			{
-				KeyDisplayName += NodeAttributeKey.Key.RightChop(IndexPosition);
-			}
-			return KeyDisplayName;
-		}
-
-		return Super::GetKeyDisplayName(NodeAttributeKey);
-	}
-
-	virtual FString GetAttributeCategory(const UE::Interchange::FAttributeKey& NodeAttributeKey) const override
-	{
-		if (NodeAttributeKey.Key.StartsWith(UE::Interchange::FStaticMeshNodeLodDataStaticData::GetMeshUidsBaseKey()))
-		{
-			return FString(TEXT("Meshes"));
-		}
-		return Super::GetAttributeCategory(NodeAttributeKey);
-	}
-
-	virtual FGuid GetHash() const override
-	{
-		return Attributes->GetStorageHash();
-	}
+	virtual FString GetKeyDisplayName(const UE::Interchange::FAttributeKey& NodeAttributeKey) const override;
 
+	virtual FString GetAttributeCategory(const UE::Interchange::FAttributeKey& NodeAttributeKey) const override;
 
 public:
 	/* Mesh Uids: It can be either a scene or a mesh node uid. If its a scene it mean we want the mesh factory to bake the geo payload with the global transform of the scene node. */
 
 	UFUNCTION(BlueprintCallable, Category = "Interchange | Node | StaticMeshLodData")
-	int32 GetMeshUidsCount() const
-	{
-		return MeshUids.GetCount();
-	}
+	int32 GetMeshUidsCount() const;
 
 	UFUNCTION(BlueprintCallable, Category = "Interchange | Node | StaticMeshLodData")
-	void GetMeshUids(TArray<FString>& OutMeshNames) const
-	{
-		MeshUids.GetNames(OutMeshNames);
-	}
+	void GetMeshUids(TArray<FString>& OutMeshNames) const;
 
 	UFUNCTION(BlueprintCallable, Category = "Interchange | Node | StaticMeshLodData")
-	bool AddMeshUid(const FString& MeshName)
-	{
-		return MeshUids.AddName(MeshName);
-	}
+	bool AddMeshUid(const FString& MeshName);
 
 	UFUNCTION(BlueprintCallable, Category = "Interchange | Node | StaticMeshLodData")
-	bool RemoveMeshUid(const FString& MeshName)
-	{
-		return MeshUids.RemoveName(MeshName);
-	}
+	bool RemoveMeshUid(const FString& MeshName);
 
 	UFUNCTION(BlueprintCallable, Category = "Interchange | Node | StaticMeshLodData")
-	bool RemoveAllMeshes()
-	{
-		return MeshUids.RemoveAllNames();
-	}
+	bool RemoveAllMeshes();
 
 private:
 
-	bool IsEditorOnlyDataDefined()
-	{
-#if WITH_EDITORONLY_DATA
-		return true;
-#else
-		return false;
-#endif
-	}
-
-	UE::Interchange::FNameAttributeArrayHelper MeshUids;
+	bool IsEditorOnlyDataDefined();
+
+	UE::Interchange::TArrayAttributeHelper<FString> MeshUids;
 protected:
 };
diff --git a/Engine/Plugins/Experimental/Interchange/Runtime/Source/Nodes/Public/InterchangeTextureFactoryNode.h b/Engine/Plugins/Experimental/Interchange/Runtime/Source/Nodes/Public/InterchangeTextureFactoryNode.h
--- a/Engine/Plugins/Experimental/Interchange/Runtime/Source/Nodes/Public/InterchangeTextureFactoryNode.h
+++ b/Engine/Plugins/Experimental/Interchange/Runtime/Source/Nodes/Public/InterchangeTextureFactoryNode.h
@@ -38,7 +38,7 @@ public:
 	{
 		//Initialize
 		bIsTextureNodeClassInitialized = false;
-		InitializeNode(UniqueID, DisplayLabel, EInterchangeNodeContainerType::NodeContainerType_FactoryData);
+		InitializeNode(UniqueID, DisplayLabel, EInterchangeNodeContainerType::FactoryData);
 		
 		//Set the class
 		FString OperationName = GetTypeName() + TEXT(".SetAssetClassName");
diff --git a/Engine/Plugins/Experimental/Interchange/Runtime/Source/Parsers/Fbx/Private/FbxCamera.cpp b/Engine/Plugins/Experimental/Interchange/Runtime/Source/Parsers/Fbx/Private/FbxCamera.cpp
--- a/Engine/Plugins/Experimental/Interchange/Runtime/Source/Parsers/Fbx/Private/FbxCamera.cpp
+++ b/Engine/Plugins/Experimental/Interchange/Runtime/Source/Parsers/Fbx/Private/FbxCamera.cpp
@@ -51,7 +51,7 @@ namespace UE
 					return nullptr;
 				}
 
-				CameraNode->InitializeNode(NodeUid, NodeName, EInterchangeNodeContainerType::NodeContainerType_TranslatedAsset);
+				CameraNode->InitializeNode(NodeUid, NodeName, EInterchangeNodeContainerType::TranslatedAsset);
 				NodeContainer.AddNode(CameraNode);
 
 				return CameraNode;
diff --git a/Engine/Plugins/Experimental/Interchange/Runtime/Source/Parsers/Fbx/Private/FbxLight.cpp b/Engine/Plugins/Experimental/Interchange/Runtime/Source/Parsers/Fbx/Private/FbxLight.cpp
--- a/Engine/Plugins/Experimental/Interchange/Runtime/Source/Parsers/Fbx/Private/FbxLight.cpp
+++ b/Engine/Plugins/Experimental/Interchange/Runtime/Source/Parsers/Fbx/Private/FbxLight.cpp
@@ -50,7 +50,7 @@ namespace UE
 					return nullptr;
 				}
 
-				LightNode->InitializeNode(NodeUid, NodeName, EInterchangeNodeContainerType::NodeContainerType_TranslatedAsset);
+				LightNode->InitializeNode(NodeUid, NodeName, EInterchangeNodeContainerType::TranslatedAsset);
 				NodeContainer.AddNode(LightNode);
 
 				return LightNode;
diff --git a/Engine/Plugins/Experimental/Interchange/Runtime/Source/Parsers/Fbx/Private/FbxMaterial.cpp b/Engine/Plugins/Experimental/Interchange/Runtime/Source/Parsers/Fbx/Private/FbxMaterial.cpp
--- a/Engine/Plugins/Experimental/Interchange/Runtime/Source/Parsers/Fbx/Private/FbxMaterial.cpp
+++ b/Engine/Plugins/Experimental/Interchange/Runtime/Source/Parsers/Fbx/Private/FbxMaterial.cpp
@@ -34,7 +34,7 @@ namespace UE
 					return nullptr;
 				}
 				// Creating a UMaterialInterface
-				MaterialNode->InitializeNode(NodeUid, NodeName, EInterchangeNodeContainerType::NodeContainerType_TranslatedAsset);
+				MaterialNode->InitializeNode(NodeUid, NodeName, EInterchangeNodeContainerType::TranslatedAsset);
 				MaterialNode->SetPayLoadKey(NodeUid);
 				NodeContainer.AddNode(MaterialNode);
 				return MaterialNode;
@@ -51,7 +51,7 @@ namespace UE
 					return nullptr;
 				}
 				// Creating a UTexture2D
-				TextureNode->InitializeNode(NodeUid, DisplayLabel, EInterchangeNodeContainerType::NodeContainerType_TranslatedAsset);
+				TextureNode->InitializeNode(NodeUid, DisplayLabel, EInterchangeNodeContainerType::TranslatedAsset);
 				//All texture translator expect a file has the payload key
 				FString NormalizeFilePath = TextureFilePath;
 				FPaths::NormalizeFilename(NormalizeFilePath);
diff --git a/Engine/Plugins/Experimental/Interchange/Runtime/Source/Parsers/Fbx/Private/FbxMesh.cpp b/Engine/Plugins/Experimental/Interchange/Runtime/Source/Parsers/Fbx/Private/FbxMesh.cpp
--- a/Engine/Plugins/Experimental/Interchange/Runtime/Source/Parsers/Fbx/Private/FbxMesh.cpp
+++ b/Engine/Plugins/Experimental/Interchange/Runtime/Source/Parsers/Fbx/Private/FbxMesh.cpp
@@ -1570,7 +1570,7 @@ namespace UE
 					return nullptr;
 				}
 				// Creating a UMaterialInterface
-				MeshNode->InitializeNode(NodeUid, DisplayLabel, EInterchangeNodeContainerType::NodeContainerType_TranslatedAsset);
+				MeshNode->InitializeNode(NodeUid, DisplayLabel, EInterchangeNodeContainerType::TranslatedAsset);
 				NodeContainer.AddNode(MeshNode);
 				return MeshNode;
 			}
diff --git a/Engine/Plugins/Experimental/Interchange/Runtime/Source/Parsers/Fbx/Private/FbxScene.cpp b/Engine/Plugins/Experimental/Interchange/Runtime/Source/Parsers/Fbx/Private/FbxScene.cpp
--- a/Engine/Plugins/Experimental/Interchange/Runtime/Source/Parsers/Fbx/Private/FbxScene.cpp
+++ b/Engine/Plugins/Experimental/Interchange/Runtime/Source/Parsers/Fbx/Private/FbxScene.cpp
@@ -264,7 +264,7 @@ namespace UE
 					return nullptr;
 				}
 				// Creating a UMaterialInterface
-				TransformNode->InitializeNode(NodeUid, DisplayLabel, EInterchangeNodeContainerType::NodeContainerType_TranslatedScene);
+				TransformNode->InitializeNode(NodeUid, DisplayLabel, EInterchangeNodeContainerType::TranslatedScene);
 				NodeContainer.AddNode(TransformNode);
 				return TransformNode;
 			}
diff --git a/Engine/Plugins/Experimental/Interchange/Runtime/Source/Pipelines/Private/InterchangeGenericAssetsPipeline.cpp b/Engine/Plugins/Experimental/Interchange/Runtime/Source/Pipelines/Private/InterchangeGenericAssetsPipeline.cpp
--- a/Engine/Plugins/Experimental/Interchange/Runtime/Source/Pipelines/Private/InterchangeGenericAssetsPipeline.cpp
+++ b/Engine/Plugins/Experimental/Interchange/Runtime/Source/Pipelines/Private/InterchangeGenericAssetsPipeline.cpp
@@ -32,12 +32,12 @@ void UInterchangeGenericAssetsPipeline::PreDialogCleanup(const FName PipelineSta
 	SaveSettings(PipelineStackName);
 }
 
-bool UInterchangeGenericAssetsPipeline::ExecutePreImportPipeline(UInterchangeBaseNodeContainer* InBaseNodeContainer, const TArray<UInterchangeSourceData*>& InSourceDatas)
+void UInterchangeGenericAssetsPipeline::ExecutePreImportPipeline(UInterchangeBaseNodeContainer* InBaseNodeContainer, const TArray<UInterchangeSourceData*>& InSourceDatas)
 {
 	if (!InBaseNodeContainer)
 	{
 		UE_LOG(LogInterchangePipeline, Warning, TEXT("UInterchangeGenericAssetsPipeline: Cannot execute pre-import pipeline because InBaseNodeContrainer is null"));
-		return false;
+		return;
 	}
 
 	BaseNodeContainer = InBaseNodeContainer;
@@ -54,7 +54,7 @@ bool UInterchangeGenericAssetsPipeline::ExecutePreImportPipeline(UInterchangeBas
 	{
 		switch(Node->GetNodeContainerType())
 		{
-			case EInterchangeNodeContainerType::NodeContainerType_TranslatedAsset:
+			case EInterchangeNodeContainerType::TranslatedAsset:
 			{
 				if (UInterchangeTextureNode* TextureNode = Cast<UInterchangeTextureNode>(Node))
 				{
@@ -222,23 +222,21 @@ bool UInterchangeGenericAssetsPipeline::ExecutePreImportPipeline(UInterchangeBas
 	}
 
 	ImplementUseSourceNameForAssetOption();
-
-	return true;
 }
 
-bool UInterchangeGenericAssetsPipeline::ExecutePostImportPipeline(const UInterchangeBaseNodeContainer* InBaseNodeContainer, const FString& NodeKey, UObject* CreatedAsset, bool bIsAReimport)
+void UInterchangeGenericAssetsPipeline::ExecutePostImportPipeline(const UInterchangeBaseNodeContainer* InBaseNodeContainer, const FString& NodeKey, UObject* CreatedAsset, bool bIsAReimport)
 {
 	//We do not use the provided base container since ExecutePreImportPipeline cache it
 	//We just make sure the same one is pass in parameter
 	if (!InBaseNodeContainer || !ensure(BaseNodeContainer == InBaseNodeContainer) || !CreatedAsset)
 	{
-		return false;
+		return;
 	}
 
 	UInterchangeBaseNode* Node = BaseNodeContainer->GetNode(NodeKey);
 	if (!Node)
 	{
-		return false;
+		return;
 	}
 
 	PostImportSkeletalMesh(CreatedAsset, Node);
@@ -247,8 +245,6 @@ bool UInterchangeGenericAssetsPipeline::ExecutePostImportPipeline(const UInterch
 	PostImportPhysicsAssetImport(CreatedAsset, Node);
 
 	PostImportTextureAssetImport(CreatedAsset, bIsAReimport);
-
-	return true;
 }
 
 UInterchangeMaterialFactoryNode* UInterchangeGenericAssetsPipeline::CreateMaterialFactoryNode(const UInterchangeMaterialNode* MaterialNode)
@@ -447,7 +443,7 @@ UInterchangeStaticMeshLodDataNode* UInterchangeGenericAssetsPipeline::CreateStat
 		return nullptr;
 	}
 
-	StaticMeshLodDataNode->InitializeNode(NodeUID, DisplayLabel, EInterchangeNodeContainerType::NodeContainerType_FactoryData);
+	StaticMeshLodDataNode->InitializeNode(NodeUID, DisplayLabel, EInterchangeNodeContainerType::FactoryData);
 	BaseNodeContainer->AddNode(StaticMeshLodDataNode);
 	return StaticMeshLodDataNode;
 }
diff --git a/Engine/Plugins/Experimental/Interchange/Runtime/Source/Pipelines/Private/InterchangeGenericScenesPipeline.cpp b/Engine/Plugins/Experimental/Interchange/Runtime/Source/Pipelines/Private/InterchangeGenericScenesPipeline.cpp
--- a/Engine/Plugins/Experimental/Interchange/Runtime/Source/Pipelines/Private/InterchangeGenericScenesPipeline.cpp
+++ b/Engine/Plugins/Experimental/Interchange/Runtime/Source/Pipelines/Private/InterchangeGenericScenesPipeline.cpp
@@ -18,12 +18,12 @@
 #include "Engine/SpotLight.h"
 #include "Engine/StaticMeshActor.h"
 
-bool UInterchangeGenericLevelPipeline::ExecutePreImportPipeline(UInterchangeBaseNodeContainer* InBaseNodeContainer, const TArray<UInterchangeSourceData*>& InSourceDatas)
+void UInterchangeGenericLevelPipeline::ExecutePreImportPipeline(UInterchangeBaseNodeContainer* InBaseNodeContainer, const TArray<UInterchangeSourceData*>& InSourceDatas)
 {
 	if (!InBaseNodeContainer)
 	{
 		UE_LOG(LogInterchangePipeline, Warning, TEXT("UInterchangeGenericAssetsPipeline: Cannot execute pre-import pipeline because InBaseNodeContrainer is null"));
-		return false;
+		return;
 	}
 
 	TArray<UInterchangeSceneNode*> SceneNodes;
@@ -33,7 +33,7 @@ bool UInterchangeGenericLevelPipeline::ExecutePreImportPipeline(UInterchangeBase
 	{
 		switch(Node->GetNodeContainerType())
 		{
-		case EInterchangeNodeContainerType::NodeContainerType_TranslatedScene:
+		case EInterchangeNodeContainerType::TranslatedScene:
 		{
 			if (UInterchangeSceneNode* SceneNode = Cast<UInterchangeSceneNode>(Node))
 			{
@@ -52,8 +52,6 @@ bool UInterchangeGenericLevelPipeline::ExecutePreImportPipeline(UInterchangeBase
 			CreateActorFactoryNode(SceneNode, InBaseNodeContainer);
 		}
 	}
-
-	return true;
 }
 
 void UInterchangeGenericLevelPipeline::CreateActorFactoryNode(const UInterchangeSceneNode* SceneNode, UInterchangeBaseNodeContainer* FactoryNodeContainer)
@@ -86,7 +84,7 @@ void UInterchangeGenericLevelPipeline::CreateActorFactoryNode(const UInterchange
 		return;
 	}
 
-	ActorFactoryNode->InitializeNode(TEXT("Factory_") + SceneNode->GetUniqueID(), SceneNode->GetDisplayLabel(), EInterchangeNodeContainerType::NodeContainerType_FactoryData);
+	ActorFactoryNode->InitializeNode(TEXT("Factory_") + SceneNode->GetUniqueID(), SceneNode->GetDisplayLabel(), EInterchangeNodeContainerType::FactoryData);
 
 	if (!SceneNode->GetParentUid().IsEmpty())
 	{
diff --git a/Engine/Plugins/Experimental/Interchange/Runtime/Source/Pipelines/Private/InterchangeGenericSkeletalMeshPipeline.cpp b/Engine/Plugins/Experimental/Interchange/Runtime/Source/Pipelines/Private/InterchangeGenericSkeletalMeshPipeline.cpp
--- a/Engine/Plugins/Experimental/Interchange/Runtime/Source/Pipelines/Private/InterchangeGenericSkeletalMeshPipeline.cpp
+++ b/Engine/Plugins/Experimental/Interchange/Runtime/Source/Pipelines/Private/InterchangeGenericSkeletalMeshPipeline.cpp
@@ -240,18 +240,18 @@ namespace UE::Interchange::SkeletalMeshGenericPipeline
 	}
 }
 
-bool UInterchangeGenericAssetsPipeline::ExecutePreImportPipelineSkeletalMesh()
+void UInterchangeGenericAssetsPipeline::ExecutePreImportPipelineSkeletalMesh()
 {
 	if (!bImportSkeletalMeshes)
 	{
 		//Nothing to import
-		return true;
+		return;
 	}
 
 	if (ForceAllMeshHasType != EInterchangeForceMeshType::IFMT_None && ForceAllMeshHasType != EInterchangeForceMeshType::IFMT_SkeletalMesh)
 	{
 		//Nothing to import
-		return true;
+		return;
 	}
 	const bool bConvertStaticMeshToSkeletalMesh = (ForceAllMeshHasType == EInterchangeForceMeshType::IFMT_SkeletalMesh);
 	TMap<FString, TArray<FString>> SkeletalMeshFactoryDependencyOrderPerSkeletonRootNodeUid;
@@ -413,7 +413,6 @@ bool UInterchangeGenericAssetsPipeline::ExecutePreImportPipelineSkeletalMesh()
 			CreatePerSkeletonRootUidSkinnedMesh(bUseMeshInstance);
 		}
 	}
-	return true;
 }
 
 
@@ -643,7 +642,7 @@ UInterchangeSkeletalMeshLodDataNode* UInterchangeGenericAssetsPipeline::CreateSk
 		return nullptr;
 	}
 	// Creating a UMaterialInterface
-	SkeletalMeshLodDataNode->InitializeNode(NodeUID, DisplayLabel, EInterchangeNodeContainerType::NodeContainerType_FactoryData);
+	SkeletalMeshLodDataNode->InitializeNode(NodeUID, DisplayLabel, EInterchangeNodeContainerType::FactoryData);
 	BaseNodeContainer->AddNode(SkeletalMeshLodDataNode);
 	return SkeletalMeshLodDataNode;
 }
diff --git a/Engine/Plugins/Experimental/Interchange/Runtime/Source/Pipelines/Private/InterchangePipelineMeshesUtilities.cpp b/Engine/Plugins/Experimental/Interchange/Runtime/Source/Pipelines/Private/InterchangePipelineMeshesUtilities.cpp
--- a/Engine/Plugins/Experimental/Interchange/Runtime/Source/Pipelines/Private/InterchangePipelineMeshesUtilities.cpp
+++ b/Engine/Plugins/Experimental/Interchange/Runtime/Source/Pipelines/Private/InterchangePipelineMeshesUtilities.cpp
@@ -21,7 +21,7 @@ UInterchangePipelineMeshesUtilities* UInterchangePipelineMeshesUtilities::Create
 	//Find all translated node we need for this pipeline
 	BaseNodeContainer->IterateNodes([&PipelineMeshesUtilities, &BaseNodeContainer, &SkeletonRootNodeUids](const FString& NodeUid, UInterchangeBaseNode* Node)
 	{
-		if (Node->GetNodeContainerType() == EInterchangeNodeContainerType::NodeContainerType_TranslatedAsset)
+		if (Node->GetNodeContainerType() == EInterchangeNodeContainerType::TranslatedAsset)
 		{
 			if (UInterchangeMeshNode* MeshNode = Cast<UInterchangeMeshNode>(Node))
 			{
@@ -35,7 +35,7 @@ UInterchangePipelineMeshesUtilities* UInterchangePipelineMeshesUtilities::Create
 	//Find all translated scene node we need for this pipeline
 	BaseNodeContainer->IterateNodes([&PipelineMeshesUtilities, &BaseNodeContainer, &SkeletonRootNodeUids](const FString& NodeUid, UInterchangeBaseNode* Node)
 	{
-		if (Node->GetNodeContainerType() == EInterchangeNodeContainerType::NodeContainerType_TranslatedScene)
+		if (Node->GetNodeContainerType() == EInterchangeNodeContainerType::TranslatedScene)
 		{
 			if (UInterchangeSceneNode* SceneNode = Cast<UInterchangeSceneNode>(Node))
 			{
diff --git a/Engine/Plugins/Experimental/Interchange/Runtime/Source/Pipelines/Public/InterchangeGenericAssetsPipeline.h b/Engine/Plugins/Experimental/Interchange/Runtime/Source/Pipelines/Public/InterchangeGenericAssetsPipeline.h
--- a/Engine/Plugins/Experimental/Interchange/Runtime/Source/Pipelines/Public/InterchangeGenericAssetsPipeline.h
+++ b/Engine/Plugins/Experimental/Interchange/Runtime/Source/Pipelines/Public/InterchangeGenericAssetsPipeline.h
@@ -214,9 +214,9 @@ public:
 
 protected:
 
-	virtual bool ExecutePreImportPipeline(UInterchangeBaseNodeContainer* InBaseNodeContainer, const TArray<UInterchangeSourceData*>& InSourceDatas) override;
+	virtual void ExecutePreImportPipeline(UInterchangeBaseNodeContainer* InBaseNodeContainer, const TArray<UInterchangeSourceData*>& InSourceDatas) override;
 
-	virtual bool ExecutePostImportPipeline(const UInterchangeBaseNodeContainer* BaseNodeContainer, const FString& NodeKey, UObject* CreatedAsset, bool bIsAReimport) override;
+	virtual void ExecutePostImportPipeline(const UInterchangeBaseNodeContainer* BaseNodeContainer, const FString& NodeKey, UObject* CreatedAsset, bool bIsAReimport) override;
 
 	virtual bool CanExecuteOnAnyThread(EInterchangePipelineTask PipelineTask) override
 	{
@@ -277,7 +277,7 @@ private:
 	/**
 	 * This function will create any skeletalmesh we need to create according to the pipeline options
 	 */
-	bool ExecutePreImportPipelineSkeletalMesh();
+	void ExecutePreImportPipelineSkeletalMesh();
 
 	/**
 	 * This function will finish creating the skeletalmesh asset
diff --git a/Engine/Plugins/Experimental/Interchange/Runtime/Source/Pipelines/Public/InterchangeGenericScenesPipeline.h b/Engine/Plugins/Experimental/Interchange/Runtime/Source/Pipelines/Public/InterchangeGenericScenesPipeline.h
--- a/Engine/Plugins/Experimental/Interchange/Runtime/Source/Pipelines/Public/InterchangeGenericScenesPipeline.h
+++ b/Engine/Plugins/Experimental/Interchange/Runtime/Source/Pipelines/Public/InterchangeGenericScenesPipeline.h
@@ -24,7 +24,7 @@ public:
 
 protected:
 
-	virtual bool ExecutePreImportPipeline(UInterchangeBaseNodeContainer* InBaseNodeContainer, const TArray<UInterchangeSourceData*>& InSourceDatas) override;
+	virtual void ExecutePreImportPipeline(UInterchangeBaseNodeContainer* InBaseNodeContainer, const TArray<UInterchangeSourceData*>& InSourceDatas) override;
 
 	virtual bool CanExecuteOnAnyThread(EInterchangePipelineTask PipelineTask) override
 	{
diff --git a/Engine/Source/Runtime/Experimental/Interchange/Core/Private/Nodes/InterchangeBaseNode.cpp b/Engine/Source/Runtime/Experimental/Interchange/Core/Private/Nodes/InterchangeBaseNode.cpp
--- a/Engine/Source/Runtime/Experimental/Interchange/Core/Private/Nodes/InterchangeBaseNode.cpp
+++ b/Engine/Source/Runtime/Experimental/Interchange/Core/Private/Nodes/InterchangeBaseNode.cpp
@@ -7,6 +7,67 @@
 #include "UObject/Object.h"
 #include "UObject/ObjectMacros.h"
 
+namespace UE
+{
+	namespace Interchange
+	{
+
+		const FAttributeKey& FBaseNodeStaticData::UniqueIDKey()
+		{
+			static FAttributeKey AttributeKey(TEXT("__UNQ_ID_"));
+			return AttributeKey;
+		}
+
+		const FAttributeKey& FBaseNodeStaticData::DisplayLabelKey()
+		{
+			static FAttributeKey AttributeKey(TEXT("__DSPL_LBL_"));
+			return AttributeKey;
+		}
+
+		const FAttributeKey& FBaseNodeStaticData::ParentIDKey()
+		{
+			static FAttributeKey AttributeKey(TEXT("__PARENT_UID_"));
+			return AttributeKey;
+		}
+
+		const FAttributeKey& FBaseNodeStaticData::IsEnabledKey()
+		{
+			static FAttributeKey AttributeKey(TEXT("__IS_NBLD_"));
+			return AttributeKey;
+		}
+
+		const FString& FBaseNodeStaticData::TargetAssetIDsKey()
+		{
+			static FString Key(TEXT("__TARGET_ASSET_IDS_"));
+			return Key;
+		}
+
+		const FString& FBaseNodeStaticData::GetFactoryDependenciesBaseKey()
+		{
+			static FString BaseNodeFactoryDependencies_BaseKey = TEXT("__BaseNodeFactoryDependencies__");
+			return BaseNodeFactoryDependencies_BaseKey;
+		}
+
+		const FAttributeKey& FBaseNodeStaticData::ClassTypeAttributeKey()
+		{
+			static FAttributeKey AttributeKey(TEXT("__ClassTypeAttribute__"));
+			return AttributeKey;
+		}
+
+		const FAttributeKey& FBaseNodeStaticData::AssetNameKey()
+		{
+			static FAttributeKey AttributeKey(TEXT("__Asset_Name_Key__"));
+			return AttributeKey;
+		}
+
+		const FAttributeKey& FBaseNodeStaticData::NodeContainerTypeKey()
+		{
+			static FAttributeKey AttributeKey(TEXT("__Node_Container_Type_Key__"));
+			return AttributeKey;
+		}
+
+	} //ns Interchange
+} //ns UE
 
 UInterchangeBaseNode::UInterchangeBaseNode()
 {
@@ -14,7 +75,7 @@ UInterchangeBaseNode::UInterchangeBaseNode()
 	FactoryDependencies.Initialize(Attributes, UE::Interchange::FBaseNodeStaticData::GetFactoryDependenciesBaseKey());
 	TargetNodes.Initialize(Attributes, UE::Interchange::FBaseNodeStaticData::TargetAssetIDsKey());
 	RegisterAttribute<bool>(UE::Interchange::FBaseNodeStaticData::IsEnabledKey(), true);
-	RegisterAttribute<uint8>(UE::Interchange::FBaseNodeStaticData::NodeContainerTypeKey(), static_cast<uint8>(EInterchangeNodeContainerType::NodeContainerType_None));
+	RegisterAttribute<uint8>(UE::Interchange::FBaseNodeStaticData::NodeContainerTypeKey(), static_cast<uint8>(EInterchangeNodeContainerType::None));
 }
 
 void UInterchangeBaseNode::InitializeNode(const FString& UniqueID, const FString& DisplayLabel, const EInterchangeNodeContainerType NodeContainerType)
@@ -46,11 +107,148 @@ FString UInterchangeBaseNode::GetTypeName() const
 	return TypeName;
 }
 
+FName UInterchangeBaseNode::GetIconName() const
+{
+	return NAME_None;
+}
+
+FString UInterchangeBaseNode::GetKeyDisplayName(const UE::Interchange::FAttributeKey& NodeAttributeKey) const
+{
+	FString KeyDisplayName = NodeAttributeKey.ToString();
+	if (NodeAttributeKey == UE::Interchange::FBaseNodeStaticData::ParentIDKey())
+	{
+		KeyDisplayName = TEXT("Parent Unique ID");
+	}
+	else if (NodeAttributeKey == UE::Interchange::FBaseNodeStaticData::DisplayLabelKey())
+	{
+		KeyDisplayName = TEXT("Name");
+	}
+	else if (NodeAttributeKey == UE::Interchange::FBaseNodeStaticData::IsEnabledKey())
+	{
+		KeyDisplayName = TEXT("Enabled");
+	}
+	else if (NodeAttributeKey == UE::Interchange::FBaseNodeStaticData::UniqueIDKey())
+	{
+		KeyDisplayName = TEXT("Unique ID");
+	}
+	else if (NodeAttributeKey == UE::Interchange::FBaseNodeStaticData::ClassTypeAttributeKey())
+	{
+		KeyDisplayName = TEXT("Node Class Type");
+	}
+	else if (NodeAttributeKey == UE::Interchange::FBaseNodeStaticData::AssetNameKey())
+	{
+		KeyDisplayName = TEXT("Imported Asset Name");
+	}
+	else if (NodeAttributeKey.Key.Equals(UE::Interchange::FBaseNodeStaticData::GetFactoryDependenciesBaseKey()))
+	{
+		KeyDisplayName = TEXT("Factory Dependencies Count");
+	}
+	else if (NodeAttributeKey.Key.StartsWith(UE::Interchange::FBaseNodeStaticData::GetFactoryDependenciesBaseKey()))
+	{
+		KeyDisplayName = TEXT("Factory Dependencies Index ");
+		const FString IndexKey = UE::Interchange::TArrayAttributeHelper<FString>::IndexKey();
+		int32 IndexPosition = NodeAttributeKey.Key.Find(IndexKey) + IndexKey.Len();
+		if (IndexPosition < NodeAttributeKey.Key.Len())
+		{
+			KeyDisplayName += NodeAttributeKey.Key.RightChop(IndexPosition);
+		}
+	}
+	else if (NodeAttributeKey.Key.Equals(UE::Interchange::FBaseNodeStaticData::TargetAssetIDsKey()))
+	{
+		KeyDisplayName = TEXT("Target Asset Count");
+	}
+	else if (NodeAttributeKey.Key.StartsWith(UE::Interchange::FBaseNodeStaticData::TargetAssetIDsKey()))
+	{
+		KeyDisplayName = TEXT("Target Asset Index ");
+		const FString IndexKey = UE::Interchange::TArrayAttributeHelper<FString>::IndexKey();
+		int32 IndexPosition = NodeAttributeKey.Key.Find(IndexKey) + IndexKey.Len();
+		if (IndexPosition < NodeAttributeKey.Key.Len())
+		{
+			KeyDisplayName += NodeAttributeKey.Key.RightChop(IndexPosition);
+		}
+	}
+	else if (NodeAttributeKey == UE::Interchange::FBaseNodeStaticData::NodeContainerTypeKey())
+	{
+		KeyDisplayName = TEXT("Node Container Type");
+	}
+	return KeyDisplayName;
+}
+
+bool UInterchangeBaseNode::ShouldHideAttribute(const UE::Interchange::FAttributeKey& NodeAttributeKey) const
+{
+	return false;
+}
+
+FString UInterchangeBaseNode::GetAttributeCategory(const UE::Interchange::FAttributeKey& NodeAttributeKey) const
+{
+	FString CategoryName = TEXT("Attributes");
+	if (NodeAttributeKey.Key.StartsWith(UE::Interchange::FBaseNodeStaticData::GetFactoryDependenciesBaseKey()))
+	{
+		CategoryName = TEXT("FactoryDependencies");
+	}
+	return CategoryName;
+}
+
 bool UInterchangeBaseNode::HasAttribute(const UE::Interchange::FAttributeKey& NodeAttributeKey) const
 {
 	return Attributes->ContainAttribute(NodeAttributeKey);
 }
 
+UE::Interchange::EAttributeTypes UInterchangeBaseNode::GetAttributeType(const UE::Interchange::FAttributeKey& NodeAttributeKey) const
+{
+	return Attributes->GetAttributeType(NodeAttributeKey);
+}
+
+void UInterchangeBaseNode::GetAttributeKeys(TArray<UE::Interchange::FAttributeKey>& AttributeKeys)
+{
+	Attributes->GetAttributeKeys(AttributeKeys);
+}
+
+bool UInterchangeBaseNode::RemoveAttribute(const FString& NodeAttributeKey)
+{
+	Attributes->UnregisterAttribute(UE::Interchange::FAttributeKey(NodeAttributeKey));
+	return !HasAttribute(UE::Interchange::FAttributeKey(NodeAttributeKey));
+}
+
+bool UInterchangeBaseNode::AddBooleanAttribute(const FString& NodeAttributeKey, const bool& Value)
+{
+	INTERCHANGE_BASE_NODE_ADD_ATTRIBUTE(bool);
+}
+
+bool UInterchangeBaseNode::GetBooleanAttribute(const FString& NodeAttributeKey, bool& OutValue)
+{
+	INTERCHANGE_BASE_NODE_GET_ATTRIBUTE(bool);
+}
+
+bool UInterchangeBaseNode::AddInt32Attribute(const FString& NodeAttributeKey, const int32& Value)
+{
+	INTERCHANGE_BASE_NODE_ADD_ATTRIBUTE(int32);
+}
+
+bool UInterchangeBaseNode::GetInt32Attribute(const FString& NodeAttributeKey, int32& OutValue)
+{
+	INTERCHANGE_BASE_NODE_GET_ATTRIBUTE(int32);
+}
+
+bool UInterchangeBaseNode::AddFloatAttribute(const FString& NodeAttributeKey, const float& Value)
+{
+	INTERCHANGE_BASE_NODE_ADD_ATTRIBUTE(float);
+}
+
+bool UInterchangeBaseNode::GetFloatAttribute(const FString& NodeAttributeKey, float& OutValue)
+{
+	INTERCHANGE_BASE_NODE_GET_ATTRIBUTE(float);
+}
+
+bool UInterchangeBaseNode::AddStringAttribute(const FString& NodeAttributeKey, const FString& Value)
+{
+	INTERCHANGE_BASE_NODE_ADD_ATTRIBUTE(FString);
+}
+
+bool UInterchangeBaseNode::GetStringAttribute(const FString& NodeAttributeKey, FString& OutValue)
+{
+	INTERCHANGE_BASE_NODE_GET_ATTRIBUTE(FString);
+}
 
 FString UInterchangeBaseNode::GetUniqueID() const
 {
@@ -115,22 +313,22 @@ int32 UInterchangeBaseNode::GetFactoryDependenciesCount() const
 
 void UInterchangeBaseNode::GetFactoryDependency(const int32 Index, FString& OutDependency) const
 {
-	FactoryDependencies.GetName(Index, OutDependency);
+	FactoryDependencies.GetItem(Index, OutDependency);
 }
 
 void UInterchangeBaseNode::GetFactoryDependencies(TArray<FString>& OutDependencies ) const
 {
-	FactoryDependencies.GetNames(OutDependencies);
+	FactoryDependencies.GetItems(OutDependencies);
 }
 
 bool UInterchangeBaseNode::AddFactoryDependencyUid(const FString& DependencyUid)
 {
-	return FactoryDependencies.AddName(DependencyUid);
+	return FactoryDependencies.AddItem(DependencyUid);
 }
 
 bool UInterchangeBaseNode::RemoveFactoryDependencyUid(const FString& DependencyUid)
 {
-	return FactoryDependencies.RemoveName(DependencyUid);
+	return FactoryDependencies.RemoveItem(DependencyUid);
 }
 
 bool UInterchangeBaseNode::IsEnabled() const
@@ -165,16 +363,16 @@ EInterchangeNodeContainerType UInterchangeBaseNode::GetNodeContainerType() const
 {
 	if (!Attributes->ContainAttribute(UE::Interchange::FBaseNodeStaticData::NodeContainerTypeKey()))
 	{
-		return EInterchangeNodeContainerType::NodeContainerType_None;
+		return EInterchangeNodeContainerType::None;
 	}
 	UE::Interchange::FAttributeStorage::TAttributeHandle<uint8> Handle = Attributes->GetAttributeHandle<uint8>(UE::Interchange::FBaseNodeStaticData::NodeContainerTypeKey());
 	if (Handle.IsValid())
 	{
-		uint8 Value = static_cast<uint8>(EInterchangeNodeContainerType::NodeContainerType_None);
+		uint8 Value = static_cast<uint8>(EInterchangeNodeContainerType::None);
 		Handle.Get(Value);
 		return static_cast<EInterchangeNodeContainerType>(Value);
 	}
-	return EInterchangeNodeContainerType::NodeContainerType_None;
+	return EInterchangeNodeContainerType::None;
 }
 
 FGuid UInterchangeBaseNode::GetHash() const
@@ -223,17 +421,17 @@ int32 UInterchangeBaseNode::GetTargetNodeCount() const
 
 void UInterchangeBaseNode::GetTargetNodeUids(TArray<FString>& OutTargetAssets) const
 {
-	TargetNodes.GetNames(OutTargetAssets);
+	TargetNodes.GetItems(OutTargetAssets);
 }
 
 bool UInterchangeBaseNode::AddTargetNodeUid(const FString& AssetUid) const
 {
-	return TargetNodes.AddName(AssetUid);
+	return TargetNodes.AddItem(AssetUid);
 }
 
 bool UInterchangeBaseNode::RemoveTargetNodeUid(const FString& AssetUid) const
 {
-	return TargetNodes.RemoveName(AssetUid);
+	return TargetNodes.RemoveItem(AssetUid);
 }
 
 FString UInterchangeBaseNode::InvalidNodeUid()
@@ -290,6 +488,21 @@ void UInterchangeBaseNode::Serialize(FArchive& Ar)
 	}
 }
 
+void UInterchangeBaseNode::CompareNodeStorage(UInterchangeBaseNode* NodeA, const UInterchangeBaseNode* NodeB, TArray<UE::Interchange::FAttributeKey>& RemovedAttributes, TArray<UE::Interchange::FAttributeKey>& AddedAttributes, TArray<UE::Interchange::FAttributeKey>& ModifiedAttributes)
+{
+	UE::Interchange::FAttributeStorage::CompareStorage(*(NodeA->Attributes), *(NodeB->Attributes), RemovedAttributes, AddedAttributes, ModifiedAttributes);
+}
+
+void UInterchangeBaseNode::CopyStorageAttributes(const UInterchangeBaseNode* SourceNode, UInterchangeBaseNode* DestinationNode, TArray<UE::Interchange::FAttributeKey>& AttributeKeys)
+{
+	UE::Interchange::FAttributeStorage::CopyStorageAttributes(*(SourceNode->Attributes), *(DestinationNode->Attributes), AttributeKeys);
+}
+
+void UInterchangeBaseNode::CopyStorage(const UInterchangeBaseNode* SourceNode, UInterchangeBaseNode* DestinationNode)
+{
+	*(DestinationNode->Attributes) = *(SourceNode->Attributes);
+}
+
 FProperty* InterchangePrivateNodeBase::FindPropertyByPathChecked(TVariant<UObject*, uint8*>& Container, UStruct* Outer, FStringView PropertyPath)
 {
 	int32 SeparatorIndex;
diff --git a/Engine/Source/Runtime/Experimental/Interchange/Core/Public/InterchangePipelineBase.h b/Engine/Source/Runtime/Experimental/Interchange/Core/Public/InterchangePipelineBase.h
--- a/Engine/Source/Runtime/Experimental/Interchange/Core/Public/InterchangePipelineBase.h
+++ b/Engine/Source/Runtime/Experimental/Interchange/Core/Public/InterchangePipelineBase.h
@@ -33,12 +33,12 @@ public:
 	 * @note - the Interchange manager is calling this function not the virtual one that is call by the default implementation.
 	 */
 	UFUNCTION(BlueprintNativeEvent, BlueprintCallable, Category = "Interchange | Pipeline")
-	bool ScriptedExecutePreImportPipeline(UInterchangeBaseNodeContainer* BaseNodeContainer, const TArray<UInterchangeSourceData*>& SourceDatas);
+	void ScriptedExecutePreImportPipeline(UInterchangeBaseNodeContainer* BaseNodeContainer, const TArray<UInterchangeSourceData*>& SourceDatas);
 	/** The default implementation (call if the blueprint do not have any implementation) will call the virtual ExecuteImportPipeline */
-	bool ScriptedExecutePreImportPipeline_Implementation(UInterchangeBaseNodeContainer* BaseNodeContainer, const TArray<UInterchangeSourceData*>& SourceDatas)
+	void ScriptedExecutePreImportPipeline_Implementation(UInterchangeBaseNodeContainer* BaseNodeContainer, const TArray<UInterchangeSourceData*>& SourceDatas)
 	{
 		//By default we call the virtual import pipeline execution
-		return ExecutePreImportPipeline(BaseNodeContainer, SourceDatas);
+		ExecutePreImportPipeline(BaseNodeContainer, SourceDatas);
 	}
 
 	/**
@@ -49,12 +49,12 @@ public:
 	 * @note - the Interchange manager is calling this function not the virtual one that is call by the default implementation.
 	 */
 	UFUNCTION(BlueprintNativeEvent, BlueprintCallable, Category = "Interchange | Pipeline")
-	bool ScriptedExecutePostImportPipeline(const UInterchangeBaseNodeContainer* BaseNodeContainer, const FString& NodeKey, UObject* CreatedAsset, bool bIsAReimport);
+	void ScriptedExecutePostImportPipeline(const UInterchangeBaseNodeContainer* BaseNodeContainer, const FString& NodeKey, UObject* CreatedAsset, bool bIsAReimport);
 	/** The default implementation (call if the blueprint do not have any implementation) will call the virtual ExecuteImportPipeline */
-	bool ScriptedExecutePostImportPipeline_Implementation(const UInterchangeBaseNodeContainer* BaseNodeContainer, const FString& NodeKey, UObject* CreatedAsset, bool bIsAReimport)
+	void ScriptedExecutePostImportPipeline_Implementation(const UInterchangeBaseNodeContainer* BaseNodeContainer, const FString& NodeKey, UObject* CreatedAsset, bool bIsAReimport)
 	{
 		//By default we call the virtual import pipeline execution
-		return ExecutePostImportPipeline(BaseNodeContainer, NodeKey, CreatedAsset, bIsAReimport);
+		ExecutePostImportPipeline(BaseNodeContainer, NodeKey, CreatedAsset, bIsAReimport);
 	}
 
 	/**
@@ -62,12 +62,12 @@ public:
 	 * the Interchange manager is calling this function not the virtual one that is call by the default implementation.
 	 */
 	UFUNCTION(BlueprintCallable, BlueprintNativeEvent, Category = "Interchange | Pipeline")
-	bool ScriptedExecuteExportPipeline(UInterchangeBaseNodeContainer* BaseNodeContainer);
+	void ScriptedExecuteExportPipeline(UInterchangeBaseNodeContainer* BaseNodeContainer);
 
 	/** The default implementation (call if the blueprint do not have any implementation) will call the virtual ExecuteExportPipeline */
-	bool ScriptedExecuteExportPipeline_Implementation(UInterchangeBaseNodeContainer* BaseNodeContainer)
+	void ScriptedExecuteExportPipeline_Implementation(UInterchangeBaseNodeContainer* BaseNodeContainer)
 	{
-		return ExecuteExportPipeline(BaseNodeContainer);
+		ExecuteExportPipeline(BaseNodeContainer);
 	}
 
 	/**
@@ -131,9 +131,8 @@ protected:
 	 * The interchange manager is not calling this function directly. It is calling the blueprint native event in case this object is a blueprint derive object.
 	 * By default the scripted implementation is calling this virtual pipeline.
 	 */
-	virtual bool ExecutePreImportPipeline(UInterchangeBaseNodeContainer* BaseNodeContainer, const TArray<UInterchangeSourceData*>& SourceDatas)
+	virtual void ExecutePreImportPipeline(UInterchangeBaseNodeContainer* BaseNodeContainer, const TArray<UInterchangeSourceData*>& SourceDatas)
 	{
-		return false;
 	}
 
 	/**
@@ -141,9 +140,8 @@ protected:
 	 * The interchange manager is not calling this function directly. It is calling the blueprint native event in case this object is a blueprint derive object.
 	 * By default the scripted implementation is calling this virtual pipeline.
 	 */
-	virtual bool ExecutePostImportPipeline(const UInterchangeBaseNodeContainer* BaseNodeContainer, const FString& NodeKey, UObject* CreatedAsset, bool bIsAReimport)
+	virtual void ExecutePostImportPipeline(const UInterchangeBaseNodeContainer* BaseNodeContainer, const FString& NodeKey, UObject* CreatedAsset, bool bIsAReimport)
 	{
-		return false;
 	}
 
 	/**
@@ -158,9 +156,8 @@ protected:
 	}
 
 	/** This function can modify the BaseNodeContainer to create a pipeline that will set/validate the graph nodes hierarchy and options.*/
-	virtual bool ExecuteExportPipeline(UInterchangeBaseNodeContainer* BaseNodeContainer)
+	virtual void ExecuteExportPipeline(UInterchangeBaseNodeContainer* BaseNodeContainer)
 	{
-		return false;
 	}
 
 
diff --git a/Engine/Source/Runtime/Experimental/Interchange/Core/Public/Nodes/InterchangeBaseNode.h b/Engine/Source/Runtime/Experimental/Interchange/Core/Public/Nodes/InterchangeBaseNode.h
--- a/Engine/Source/Runtime/Experimental/Interchange/Core/Public/Nodes/InterchangeBaseNode.h
+++ b/Engine/Source/Runtime/Experimental/Interchange/Core/Public/Nodes/InterchangeBaseNode.h
@@ -3,6 +3,7 @@
 #pragma once
 
 #include "CoreMinimal.h"
+#include "Nodes/InterchangeBaseNodeUtilities.h"
 #include "Types/AttributeStorage.h"
 #include "UObject/Class.h"
 #include "UObject/Object.h"
@@ -221,616 +222,25 @@ namespace UE
 {
 	namespace Interchange
 	{
-
 		DECLARE_DELEGATE_RetVal_OneParam(bool, FApplyAttributeToAsset, UObject*);
 		DECLARE_DELEGATE_RetVal_OneParam(bool, FFillAttributeToAsset, UObject*);
 
-		class FNameAttributeArrayHelper
-		{
-		public:
-			~FNameAttributeArrayHelper()
-			{
-				Attributes = nullptr;
-				KeyCount = NAME_None;
-			}
-
-			void Initialize(const TSharedPtr<FAttributeStorage, ESPMode::ThreadSafe>& InAttributes, const FString& BaseKeyName)
-			{
-				Attributes = InAttributes;
-				check(Attributes.IsValid());
-				FString BaseTryName = BaseKeyName;
-				KeyCount = BaseTryName;
-			}
-
-			static const FString& IndexKey()
-			{
-				static FString IndexKeyString = TEXT("_NameIndex_");
-				return IndexKeyString;
-			}
-
-			int32 GetCount() const
-			{
-				//The Class must be initialise properly before we can use it
-				TSharedPtr<FAttributeStorage, ESPMode::ThreadSafe> AttributePtr = Attributes.Pin();
-				if (!ensure(AttributePtr.IsValid()))
-				{
-					return 0;
-				}
-				int32 NameCount = 0;
-				if (AttributePtr->ContainAttribute(GetKeyCount()))
-				{
-					FAttributeStorage::TAttributeHandle<int32> Handle = AttributePtr->GetAttributeHandle<int32>(GetKeyCount());
-					if (Handle.IsValid())
-					{
-						Handle.Get(NameCount);
-					}
-				}
-				return NameCount;
-			}
-
-			void GetName(const int32 Index, FString& OutName) const
-			{
-				OutName.Empty();
-				//The Class must be initialise properly before we can use it
-				TSharedPtr<FAttributeStorage, ESPMode::ThreadSafe> AttributePtr = Attributes.Pin();
-				if (!ensure(AttributePtr.IsValid()))
-				{
-					return;
-				}
-				int32 NameCount = 0;
-				if (!AttributePtr->ContainAttribute(GetKeyCount()))
-				{
-					return;
-				}
-
-				FAttributeStorage::TAttributeHandle<int32> Handle = AttributePtr->GetAttributeHandle<int32>(GetKeyCount());
-				if (!Handle.IsValid())
-				{
-					return;
-				}
-				Handle.Get(NameCount);
-				if (Index >= NameCount)
-				{
-					return;
-				}
-				FAttributeKey DepIndexKey = GetIndexKey(Index);
-				FAttributeStorage::TAttributeHandle<FString> HandleName = AttributePtr->GetAttributeHandle<FString>(DepIndexKey);
-				if (!HandleName.IsValid())
-				{
-					return;
-				}
-				HandleName.Get(OutName);
-			}
-
-			void GetNames(TArray<FString>& OutNames) const
-			{
-				//The Class must be initialise properly before we can use it
-				TSharedPtr<FAttributeStorage, ESPMode::ThreadSafe> AttributePtr = Attributes.Pin();
-				if (!ensure(AttributePtr.IsValid()))
-				{
-					OutNames.Empty();
-					return;
-				}
-				int32 NameCount = 0;
-				if (!AttributePtr->ContainAttribute(GetKeyCount()))
-				{
-					OutNames.Empty();
-					return;
-				}
-
-				//Reuse as much memory we can to avoid allocation
-				OutNames.Reset(NameCount);
-
-				FAttributeStorage::TAttributeHandle<int32> Handle = AttributePtr->GetAttributeHandle<int32>(GetKeyCount());
-				if (!Handle.IsValid())
-				{
-					return;
-				}
-				Handle.Get(NameCount);
-				for (int32 NameIndex = 0; NameIndex < NameCount; ++NameIndex)
-				{
-					FAttributeKey DepIndexKey = GetIndexKey(NameIndex);
-					FAttributeStorage::TAttributeHandle<FString> HandleName = AttributePtr->GetAttributeHandle<FString>(DepIndexKey);
-					if (!HandleName.IsValid())
-					{
-						continue;
-					}
-					FString& OutName = OutNames.AddDefaulted_GetRef();
-					HandleName.Get(OutName);
-				}
-			}
-
-			bool AddName(const FString& Name)
-			{
-				//The Class must be initialise properly before we can use it
-				TSharedPtr<FAttributeStorage, ESPMode::ThreadSafe> AttributePtr = Attributes.Pin();
-				if (!ensure(AttributePtr.IsValid()))
-				{
-					return false;
-				}
-
-				if (!AttributePtr->ContainAttribute(GetKeyCount()))
-				{
-					const int32 DependencyCount = 0;
-					EAttributeStorageResult Result = AttributePtr->RegisterAttribute<int32>(GetKeyCount(), DependencyCount);
-					if (!IsAttributeStorageResultSuccess(Result))
-					{
-						LogAttributeStorageErrors(Result, TEXT("FNameAttributeArrayHelper.AddName"), GetKeyCount());
-						return false;
-					}
-				}
-				FAttributeStorage::TAttributeHandle<int32> Handle = AttributePtr->GetAttributeHandle<int32>(GetKeyCount());
-				if (!ensure(Handle.IsValid()))
-				{
-					return false;
-				}
-				int32 NameIndex = 0;
-				Handle.Get(NameIndex);
-				FAttributeKey NameIndexKey = GetIndexKey(NameIndex);
-				//Increment the name counter
-				NameIndex++;
-				Handle.Set(NameIndex);
-
-				EAttributeStorageResult AddNameResult = AttributePtr->RegisterAttribute<FString>(NameIndexKey, Name);
-				if (!IsAttributeStorageResultSuccess(AddNameResult))
-				{
-					LogAttributeStorageErrors(AddNameResult, TEXT("FNameAttributeArrayHelper.AddName"), NameIndexKey);
-					return false;
-				}
-				return true;
-			}
-
-			bool RemoveName(const FString& NameToDelete)
-			{
-				//The Class must be initialise properly before we can use it
-				TSharedPtr<FAttributeStorage, ESPMode::ThreadSafe> AttributePtr = Attributes.Pin();
-				if (!ensure(AttributePtr.IsValid()))
-				{
-					return false;
-				}
-
-				int32 NameCount = 0;
-				if (!AttributePtr->ContainAttribute(GetKeyCount()))
-				{
-					return false;
-				}
-				FAttributeStorage::TAttributeHandle<int32> Handle = AttributePtr->GetAttributeHandle<int32>(GetKeyCount());
-				if (!Handle.IsValid())
-				{
-					return false;
-				}
-				Handle.Get(NameCount);
-				bool DecrementKey = false;
-				for (int32 NameIndex = 0; NameIndex < NameCount; ++NameIndex)
-				{
-					FAttributeKey DepIndexKey = GetIndexKey(NameIndex);
-					FAttributeStorage::TAttributeHandle<FString> HandleName = AttributePtr->GetAttributeHandle<FString>(DepIndexKey);
-					if (!HandleName.IsValid())
-					{
-						continue;
-					}
-					FString Name;
-					HandleName.Get(Name);
-					if (Name == NameToDelete)
-					{
-						//Remove this entry
-						AttributePtr->UnregisterAttribute(DepIndexKey);
-						Handle.Set(NameCount - 1);
-						//We have to rename the key for all the next item
-						DecrementKey = true;
-					}
-					else if (DecrementKey)
-					{
-						FAttributeKey NewDepIndexKey = GetIndexKey(NameIndex - 1);
-						EAttributeStorageResult UnregisterResult = AttributePtr->UnregisterAttribute(DepIndexKey);
-						if (IsAttributeStorageResultSuccess(UnregisterResult))
-						{
-							EAttributeStorageResult RegisterResult = AttributePtr->RegisterAttribute<FString>(NewDepIndexKey, Name);
-							if (!IsAttributeStorageResultSuccess(RegisterResult))
-							{
-								LogAttributeStorageErrors(RegisterResult, TEXT("FNameAttributeArrayHelper.RemoveName"), NewDepIndexKey);
-							}
-						}
-						else
-						{
-							LogAttributeStorageErrors(UnregisterResult, TEXT("FNameAttributeArrayHelper.RemoveName"), DepIndexKey);
-						}
-
-						//Avoid doing more code in the for since the HandleName is now invalid
-						continue;
-					}
-				}
-				return true;
-			}
-
-			bool RemoveAllNames()
-			{
-				//The Class must be initialise properly before we can use it
-				TSharedPtr<FAttributeStorage, ESPMode::ThreadSafe> AttributePtr = Attributes.Pin();
-				if (!ensure(AttributePtr.IsValid()))
-				{
-					return false;
-				}
-
-				int32 NameCount = 0;
-				if (!AttributePtr->ContainAttribute(GetKeyCount()))
-				{
-					return false;
-				}
-				FAttributeStorage::TAttributeHandle<int32> HandleCount = AttributePtr->GetAttributeHandle<int32>(GetKeyCount());
-				if (!HandleCount.IsValid())
-				{
-					return false;
-				}
-				HandleCount.Get(NameCount);
-				//Remove all attribute one by one
-				for (int32 NameIndex = 0; NameIndex < NameCount; ++NameIndex)
-				{
-					FAttributeKey DepIndexKey = GetIndexKey(NameIndex);
-					AttributePtr->UnregisterAttribute(DepIndexKey);
-				}
-				//Make sure Count is zero
-				NameCount = 0;
-				HandleCount.Set(NameCount);
-				return true;
-			}
-
-		private:
-			TWeakPtr<FAttributeStorage, ESPMode::ThreadSafe> Attributes = nullptr;
-
-			FAttributeKey KeyCount; //Assign in Initialize function, it will ensure if its the default value (NAME_None) when using the class
-			FAttributeKey GetKeyCount() const { ensure(!KeyCount.Key.IsEmpty()); return KeyCount; }
-
-			FAttributeKey GetIndexKey(int32 Index) const
-			{
-				FString DepIndexKeyString = GetKeyCount().ToString() + IndexKey() + FString::FromInt(Index);
-				return FAttributeKey(*DepIndexKeyString);
-			}
-		};
-
-		template<typename KeyType, typename ValueType>
-		class TMapAttributeHelper
-		{
-			static_assert(TAttributeTypeTraits<KeyType>::GetType() != EAttributeTypes::None, "The key type must be supported by the attribute storage");
-			static_assert(TAttributeTypeTraits<ValueType>::GetType() != EAttributeTypes::None, "The value type must be supported by the attribute storage");
-
-			template<typename T>
-			using TAttributeHandle = FAttributeStorage::TAttributeHandle<T>;
-
-		public:
-			void Initialize(const TSharedRef<FAttributeStorage, ESPMode::ThreadSafe>& InAttributes, const FString& BaseKeyName)
-			{
-				Attributes = InAttributes;
-				FString BaseTryName = BaseKeyName;
-				FAttributeKey KeyCountKey(MoveTemp(BaseTryName));
-				if (!InAttributes->ContainAttribute(KeyCountKey))
-				{ 
-					EAttributeStorageResult Result = InAttributes->RegisterAttribute<int32>(KeyCountKey, 0);
-					check (Result == EAttributeStorageResult::Operation_Success);
-					KeyCountHandle = InAttributes->GetAttributeHandle<int32>(KeyCountKey);
-				}
-				else
-				{
-					KeyCountHandle = InAttributes->GetAttributeHandle<int32>(KeyCountKey);
-
-					//Init map
-					int32 KeyCount;
-					KeyCountHandle.Get(KeyCount);
-
-					CachedKeysAndValues.Reserve(KeyCount);
-
-					for (int32 Index = 0; Index < KeyCount; Index++)
-					{
-						TAttributeHandle<KeyType> KeyAttribute = InAttributes->GetAttributeHandle<KeyType>(GetKeyAttribute(Index));
-						KeyType Key;
-						KeyAttribute.Get(Key);
-						TAttributeHandle<ValueType> ValueAttribute = InAttributes->GetAttributeHandle<ValueType>(GetValueAttribute(Key));
-						CachedKeysAndValues.Add(Key, TPair<TAttributeHandle<KeyType>, TAttributeHandle<ValueType>>(KeyAttribute, ValueAttribute));
-					}
-				}
-			}
-
-			TMapAttributeHelper() = default;
-			TMapAttributeHelper(const TMapAttributeHelper&) = default;
-			TMapAttributeHelper(TMapAttributeHelper&&) = default;
-			TMapAttributeHelper& operator=(const TMapAttributeHelper&) = default;
-			TMapAttributeHelper& operator=(TMapAttributeHelper&&) = default;
-
-			~TMapAttributeHelper()
-			{
-				Attributes.Reset();
-			}
-
-			void SetKeyValue(const KeyType& InKey, const ValueType& InValue)
-			{
-				const uint32 Hash = GetTypeHash(InKey);
-				SetKeyValueByHash(Hash, InKey, InValue);
-			}
-
-			bool GetValue(const KeyType& InKey, ValueType& OutValue) const
-			{
-				const uint32 Hash = GetTypeHash(InKey);
-				return GetValueByHash(Hash, InKey, OutValue);
-			}
-
-			bool RemoveKey(const KeyType& InKey)
-			{
-				const uint32 Hash = GetTypeHash(InKey);
-				return RemoveKeyByHash(Hash, InKey);
-			}
-
-			bool RemoveKeyAndGetValue(const KeyType& InKey, ValueType& OutValue)
-			{
-				const uint32 Hash = GetTypeHash(InKey);
-				return RemoveKeyAndGetValueByHash(Hash, InKey, OutValue);
-			}
-
-			void SetKeyValueByHash(uint32 Hash, const KeyType& InKey, const ValueType& InValue)
-			{
-				TSharedPtr<FAttributeStorage, ESPMode::ThreadSafe> AttributesPtr = Attributes.Pin();
-				if (!ensure(AttributesPtr.IsValid()))
-				{
-					return;
-				}
-
-				if (TPair<TAttributeHandle<KeyType>, TAttributeHandle<ValueType>>* Pair = CachedKeysAndValues.FindByHash(Hash, InKey))
-				{
-					Pair->Value.Set(InValue);
-				}
-				else
-				{
-					FAttributeKey IndexKey = GetKeyAttribute(CachedKeysAndValues.Num());
-					FAttributeKey AttributeKey = GetValueAttribute(InKey);
-					check(AttributesPtr->RegisterAttribute<KeyType>(IndexKey, InKey) == EAttributeStorageResult::Operation_Success);
-					check(AttributesPtr->RegisterAttribute<ValueType>(AttributeKey, InValue) == EAttributeStorageResult::Operation_Success);
-					CachedKeysAndValues.AddByHash(
-						Hash
-						, InKey
-						, TPair<TAttributeHandle<KeyType>, TAttributeHandle<ValueType>>(
-							AttributesPtr->GetAttributeHandle<KeyType>(IndexKey)
-							, AttributesPtr->GetAttributeHandle<ValueType>(AttributeKey)
-						));
-					KeyCountHandle.Set(CachedKeysAndValues.Num());
-				}
-			}
-
-		
-			bool GetValueByHash(uint32 Hash, const KeyType& InKey, ValueType& OutValue) const
-			{
-				TSharedPtr<FAttributeStorage, ESPMode::ThreadSafe> AttributesPtr = Attributes.Pin();
-				if (!ensure(AttributesPtr.IsValid()))
-				{
-					return false;
-				}
-
-				if (const TPair<TAttributeHandle<KeyType>, TAttributeHandle<ValueType>>* Pair = CachedKeysAndValues.FindByHash(Hash, InKey))
-				{
-					return Pair->Value.Get(OutValue) == EAttributeStorageResult::Operation_Success;
-				}
-			}
-
-	
-			bool RemoveKeyByHash(uint32 Hash, const KeyType& InKey)
-			{
-				TSharedPtr<FAttributeStorage, ESPMode::ThreadSafe> AttributesPtr = Attributes.Pin();
-				if (!ensure(AttributesPtr.IsValid()))
-				{
-					return false;
-				}
-
-				if (TPair<TAttributeHandle<KeyType>, TAttributeHandle<ValueType>>* Pair = CachedKeysAndValues.FindByHash(Hash, InKey))
-				{
-					return RemoveBySwap(Hash, InKey, *Pair);
-				}
-
-				return false;
-			}
-
-			bool RemoveKeyAndGetValueByHash(uint32 Hash, const KeyType& InKey, ValueType& OutValue)
-			{
-				TSharedPtr<FAttributeStorage, ESPMode::ThreadSafe> AttributesPtr = Attributes.Pin();
-				if (!ensure(AttributesPtr.IsValid()))
-				{
-					return false;
-				}
-
-				if (TPair<TAttributeHandle<KeyType>, TAttributeHandle<ValueType>>* Pair = CachedKeysAndValues.FindByHash(Hash, InKey))
-				{
-					if (Pair->Value.Get(OutValue) == EAttributeStorageResult::Operation_Success)
-					{
-						return RemoveBySwap(Hash, InKey, *Pair);
-					}
-				}
-
-				return false;
-			}
-
-			void Reserve(int32 Number)
-			{
-				CachedKeysAndValues.Reserve(Number);
-			}
-
-			void Empty(int32 NumOfExpectedElements = 0)
-			{
-				EmptyInternal(NumOfExpectedElements);
-			}
-
-			TMapAttributeHelper& operator=(const TMap<KeyType, ValueType>& InMap)
-			{
-				TSharedPtr<FAttributeStorage, ESPMode::ThreadSafe> AttributesPtr = Attributes.Pin();
-				if (!ensure(AttributesPtr.IsValid()))
-				{
-					return *this;
-				}
-
-				// Empty
-				EmptyInternal(InMap.Num());
-
-				KeyCountHandle.Set(InMap.Num());
-
-				for (const TPair<KeyType, ValueType>& Pair : InMap)
-				{
-					SetKeyValue(Pair.Key, Pair.Value);
-				}
-
-				return *this;
-			}
-
-			TMap<KeyType, ValueType> ToMap() const
-			{
-				TMap<KeyType, ValueType> Map;
-				Map.Reserve(CachedKeysAndValues.Num());
-				
-				for (const TPair<KeyType, TPair<TAttributeHandle<KeyType>, TAttributeHandle<ValueType>>>& Pair: CachedKeysAndValues)
-				{
-					KeyType Key;
-					Pair.Value.Key.Get(Key);
-					ValueType Value;
-					Pair.Value.Value.Get(Value);
-
-					Map.Add(MoveTemp(Key), MoveTemp(Value));
-				}
-
-				return Map;
-			}
-
-		private:
-			FAttributeKey GetKeyAttribute(int32 Index) const
-			{
-				static const FString KeyIndex = TEXT("_KeyIndex_");
-				FString IndexedKey =  KeyCountHandle.GetKey().ToString();
-				IndexedKey.Reserve(KeyIndex.Len() + 16 /*Max size for a int32*/);
-				IndexedKey.Append(KeyIndex);
-				IndexedKey.AppendInt(Index);
-				return FAttributeKey(MoveTemp(IndexedKey));
-			}
-
-			FAttributeKey GetValueAttribute(const KeyType& InKey) const
-			{
-				static const FString KeyChars = TEXT("_Key_");
-				FString ValueAttribute =  KeyCountHandle.GetKey().ToString();
-				FString Key = TTypeToString<KeyType>::ToString(InKey);
-				ValueAttribute.Reserve(KeyChars.Len() + Key.Len());
-				ValueAttribute.Append(KeyChars);
-				ValueAttribute.Append(Key);
-				return FAttributeKey(MoveTemp(ValueAttribute));
-			}
-
-			TAttributeHandle<KeyType> GetLastKeyAttributeHandle() const
-			{
-				TSharedPtr<FAttributeStorage, ESPMode::ThreadSafe> AttributesPtr = Attributes.Pin();
-				if (!ensure(AttributesPtr.IsValid()))
-				{
-					return {};
-				}
-				return AttributesPtr->GetAttributeHandle<KeyType>(GetKeyAttribute(CachedKeysAndValues.Num() - 1));
-			}
-
-			bool RemoveBySwap(uint32 Hash, const KeyType& InKey, TPair<TAttributeHandle<FString>, TAttributeHandle<ValueType>>& CachedPair)
-			{
-				TSharedPtr<FAttributeStorage, ESPMode::ThreadSafe> AttributesPtr = Attributes.Pin();
-				if (!ensure(AttributesPtr.IsValid()))
-				{
-					return false;
-				}
-
-				TAttributeHandle<KeyType> LastKeyIndex = GetLastKeyAttributeHandle();
-				KeyType LastKey;
-				LastKeyIndex.Get(LastKey);
-				CachedKeysAndValues[LastKey].Key = CachedPair.Key;
-				CachedPair->Key.Set(MoveTemp(LastKey));
-
-				AttributesPtr->UnregisterAttribute(LastKeyIndex.GetKey());
-				AttributesPtr->UnregisterAttribute(CachedPair.Value.GetKey());
-				CachedKeysAndValues.RemoveByHash(Hash, InKey);
-				KeyCountHandle.Set(CachedKeysAndValues.Num());
-			}
-
-			void EmptyInternal(int32 NumOfExpectedElements)
-			{
-				TSharedPtr<FAttributeStorage, ESPMode::ThreadSafe> AttributesPtr = Attributes.Pin();
-				if (!ensure(AttributesPtr.IsValid()))
-				{
-					return;
-				}
-
-				for (const TPair<KeyType, TPair<TAttributeHandle<KeyType>, TAttributeHandle<ValueType>>>& Pair : CachedKeysAndValues)
-				{
-					const FAttributeKey& KeyAttribute = Pair.Value.Key.GetKey();
-					AttributesPtr->UnregisterAttribute(KeyAttribute);
-
-					const FAttributeKey& ValueAttribute = Pair.Value.Value.GetKey();
-					AttributesPtr->UnregisterAttribute(ValueAttribute);
-				}
-
-				CachedKeysAndValues.Empty(NumOfExpectedElements);
-			}
-
-			TMap<KeyType, TPair<TAttributeHandle<KeyType>, TAttributeHandle<ValueType>>> CachedKeysAndValues;
-			TAttributeHandle<int32> KeyCountHandle; //Assign in Initialize function, it will ensure if its the default value (NAME_None) when using the class
-			TWeakPtr<FAttributeStorage, ESPMode::ThreadSafe> Attributes = nullptr;
-		};
-
 		/**
 		 * Helper struct use to declare static const data we use in the UInterchangeBaseNode
 		 * Node that derive from UInterchangeBaseNode can also add a struct that derive from this one to add there static data
 		 * @note: The static data are mainly for holding Attribute keys. All attributes that are always available for a node should be in this class or a derived class.
 		 */
-		struct FBaseNodeStaticData
+		struct INTERCHANGECORE_API FBaseNodeStaticData
 		{
-			static const FAttributeKey& UniqueIDKey()
-			{
-				static FAttributeKey AttributeKey(TEXT("__UNQ_ID_"));
-				return AttributeKey;
-			}
-
-			static const FAttributeKey& DisplayLabelKey()
-			{
-				static FAttributeKey AttributeKey(TEXT("__DSPL_LBL_"));
-				return AttributeKey;
-			}
-
-			static const FAttributeKey& ParentIDKey()
-			{
-				static FAttributeKey AttributeKey(TEXT("__PARENT_UID_"));
-				return AttributeKey;
-			}
-
-			static const FAttributeKey& IsEnabledKey()
-			{
-				static FAttributeKey AttributeKey(TEXT("__IS_NBLD_"));
-				return AttributeKey;
-			}
-
-			static const FString& TargetAssetIDsKey()
-			{
-				static FString Key(TEXT("__TARGET_ASSET_IDS_"));
-				return Key;
-			}
-
-			static const FString& GetFactoryDependenciesBaseKey()
-			{
-				static FString BaseNodeFactoryDependencies_BaseKey = TEXT("__BaseNodeFactoryDependencies__");
-				return BaseNodeFactoryDependencies_BaseKey;
-			}
-
-			static const FAttributeKey& ClassTypeAttributeKey()
-			{
-				static FAttributeKey AttributeKey(TEXT("__ClassTypeAttribute__"));
-				return AttributeKey;
-			}
-
-			static const FAttributeKey& AssetNameKey()
-			{
-				static FAttributeKey AttributeKey(TEXT("__Asset_Name_Key__"));
-				return AttributeKey;
-			}
-
-			static const FAttributeKey& NodeContainerTypeKey()
-			{
-				static FAttributeKey AttributeKey(TEXT("__Node_Container_Type_Key__"));
-				return AttributeKey;
-			}
+			static const FAttributeKey& UniqueIDKey();
+			static const FAttributeKey& DisplayLabelKey();
+			static const FAttributeKey& ParentIDKey();
+			static const FAttributeKey& IsEnabledKey();
+			static const FString& TargetAssetIDsKey();
+			static const FString& GetFactoryDependenciesBaseKey();
+			static const FAttributeKey& ClassTypeAttributeKey();
+			static const FAttributeKey& AssetNameKey();
+			static const FAttributeKey& NodeContainerTypeKey();
 		};
 
 	} //ns Interchange
@@ -839,10 +249,10 @@ namespace UE
 UENUM(BlueprintType)
 enum class EInterchangeNodeContainerType : uint8
 {
-	NodeContainerType_None,
-	NodeContainerType_TranslatedScene,
-	NodeContainerType_TranslatedAsset,
-	NodeContainerType_FactoryData
+	None,
+	TranslatedScene,
+	TranslatedAsset,
+	FactoryData
 };
 
 /**
@@ -858,7 +268,7 @@ class INTERCHANGECORE_API UInterchangeBaseNode : public UObject
 public:
 	UInterchangeBaseNode();
 
-	virtual ~UInterchangeBaseNode() = default;
+//	virtual ~UInterchangeBaseNode() = default;
 
 	/**
 	 * Initialize the base data of the node
@@ -875,89 +285,24 @@ public:
 	virtual FString GetTypeName() const;
 
 	/**
-	 * Icon name are simply create by adding "InterchangeIcon_" in front of the specialized type. If there is no special type the function will return NAME_None which will use the default icon.
+	 * Icon name, use to retrieve the brush when we display the node in any UI
 	 */
-	virtual FName GetIconName() const
-	{
-		return NAME_None;
-	}
+	virtual FName GetIconName() const;
 
-	virtual FString GetKeyDisplayName(const UE::Interchange::FAttributeKey& NodeAttributeKey) const
-	{
-		FString KeyDisplayName = NodeAttributeKey.ToString();
-		if (NodeAttributeKey == UE::Interchange::FBaseNodeStaticData::ParentIDKey())
-		{
-			KeyDisplayName = TEXT("Parent Unique ID");
-		}
-		else if (NodeAttributeKey == UE::Interchange::FBaseNodeStaticData::DisplayLabelKey())
-		{
-			KeyDisplayName = TEXT("Name");
-		}
-		else if (NodeAttributeKey == UE::Interchange::FBaseNodeStaticData::IsEnabledKey())
-		{
-			KeyDisplayName = TEXT("Enabled");
-		}
-		else if (NodeAttributeKey == UE::Interchange::FBaseNodeStaticData::UniqueIDKey())
-		{
-			KeyDisplayName = TEXT("Unique ID");
-		}
-		else if (NodeAttributeKey == UE::Interchange::FBaseNodeStaticData::ClassTypeAttributeKey())
-		{
-			KeyDisplayName = TEXT("Node Class Type");
-		}
-		else if (NodeAttributeKey == UE::Interchange::FBaseNodeStaticData::AssetNameKey())
-		{
-			KeyDisplayName = TEXT("Imported Asset Name");
-		}
-		else if (NodeAttributeKey.Key.Equals(UE::Interchange::FBaseNodeStaticData::GetFactoryDependenciesBaseKey()))
-		{
-			KeyDisplayName = TEXT("Factory Dependencies Count");
-		}
-		else if (NodeAttributeKey.Key.StartsWith(UE::Interchange::FBaseNodeStaticData::GetFactoryDependenciesBaseKey()))
-		{
-			KeyDisplayName = TEXT("Factory Dependencies Index ");
-			const FString IndexKey = UE::Interchange::FNameAttributeArrayHelper::IndexKey();
-			int32 IndexPosition = NodeAttributeKey.Key.Find(IndexKey) + IndexKey.Len();
-			if (IndexPosition < NodeAttributeKey.Key.Len())
-			{
-				KeyDisplayName += NodeAttributeKey.Key.RightChop(IndexPosition);
-			}
-		}
-		else if (NodeAttributeKey.Key.Equals(UE::Interchange::FBaseNodeStaticData::TargetAssetIDsKey()))
-		{
-			KeyDisplayName = TEXT("Target Asset Count");
-		}
-		else if (NodeAttributeKey.Key.StartsWith(UE::Interchange::FBaseNodeStaticData::TargetAssetIDsKey()))
-		{
-			KeyDisplayName = TEXT("Target Asset Index ");
-			const FString IndexKey = UE::Interchange::FNameAttributeArrayHelper::IndexKey();
-			int32 IndexPosition = NodeAttributeKey.Key.Find(IndexKey) + IndexKey.Len();
-			if (IndexPosition < NodeAttributeKey.Key.Len())
-			{
-				KeyDisplayName += NodeAttributeKey.Key.RightChop(IndexPosition);
-			}
-		}
-		else if (NodeAttributeKey == UE::Interchange::FBaseNodeStaticData::NodeContainerTypeKey())
-		{
-			KeyDisplayName = TEXT("Node Container Type");
-		}
-		return KeyDisplayName;
-	}
+	/**
+	 * UI that inspect node attribute call this to give a readable name to attribute key
+	 */
+	virtual FString GetKeyDisplayName(const UE::Interchange::FAttributeKey& NodeAttributeKey) const;
 
-	virtual bool ShouldHideAttribute(const UE::Interchange::FAttributeKey& NodeAttributeKey) const
-	{
-		return false;
-	}
+	/**
+	 * UI that inspect node attribute call this to display or hide an attribute
+	 */
+	virtual bool ShouldHideAttribute(const UE::Interchange::FAttributeKey& NodeAttributeKey) const;
 
-	virtual FString GetAttributeCategory(const UE::Interchange::FAttributeKey& NodeAttributeKey) const
-	{
-		FString CategoryName = TEXT("Attributes");
-		if (NodeAttributeKey.Key.StartsWith(UE::Interchange::FBaseNodeStaticData::GetFactoryDependenciesBaseKey()))
-		{
-			CategoryName = TEXT("FactoryDependencies");
-		}
-		return CategoryName;
-	}
+	/**
+	 * UI that inspect node attribute call this to display the attribute under the returned category
+	 */
+	virtual FString GetAttributeCategory(const UE::Interchange::FAttributeKey& NodeAttributeKey) const;
 
 	/**
 	 * Add an attribute to the node
@@ -990,10 +335,7 @@ public:
 	 *
 	 * @param NodeAttributeKey - The key of the attribute
 	 */
-	virtual UE::Interchange::EAttributeTypes GetAttributeType(const UE::Interchange::FAttributeKey& NodeAttributeKey) const
-	{
-		return Attributes->GetAttributeType(NodeAttributeKey);
-	}
+	virtual UE::Interchange::EAttributeTypes GetAttributeType(const UE::Interchange::FAttributeKey& NodeAttributeKey) const;
 
 	/**
 	 * This function return an  attribute handle for the specified Key.
@@ -1007,110 +349,79 @@ public:
 		return Attributes->GetAttributeHandle<T>(NodeAttributeKey);
 	}
 
-	void GetAttributeKeys(TArray<UE::Interchange::FAttributeKey>& AttributeKeys)
-	{
-		Attributes->GetAttributeKeys(AttributeKeys);
-	}
+	/**
+	 * Use this function to query all the node attribute keys
+	 */
+	void GetAttributeKeys(TArray<UE::Interchange::FAttributeKey>& AttributeKeys);
 
 	/**
 	 * Remove any attribute from this node. Return false if we cannot remove it. If the attribute do not exist it will return true.
 	 */
 	UFUNCTION(BlueprintCallable, Category = "Interchange | Node")
-	bool RemoveAttribute(const FString& NodeAttributeKey)
-	{
-		Attributes->UnregisterAttribute(UE::Interchange::FAttributeKey(NodeAttributeKey));
-		return !HasAttribute(UE::Interchange::FAttributeKey(NodeAttributeKey));
-	}
+	bool RemoveAttribute(const FString& NodeAttributeKey);
 
 	/**
 	 * Add a boolean attribute to this node. Return false if the attribute do not exist or if we cannot add it
 	 */
 	UFUNCTION(BlueprintCallable, Category = "Interchange | Node")
-	bool AddBooleanAttribute(const FString& NodeAttributeKey, const bool& Value)
-	{
-		INTERCHANGE_BASE_NODE_ADD_ATTRIBUTE(bool);
-	}
+	bool AddBooleanAttribute(const FString& NodeAttributeKey, const bool& Value);
 	
 	/**
 	 * Get a boolean attribute from this node. Return false if the attribute do not exist
 	 */
 	UFUNCTION(BlueprintCallable, Category = "Interchange | Node")
-	bool GetBooleanAttribute(const FString& NodeAttributeKey, bool& OutValue)
-	{
-		INTERCHANGE_BASE_NODE_GET_ATTRIBUTE(bool);
-	}
+	bool GetBooleanAttribute(const FString& NodeAttributeKey, bool& OutValue);
 
 	/**
 	 * Add a int32 attribute to this node. Return false if the attribute do not exist or if we cannot add it
 	 */
 	UFUNCTION(BlueprintCallable, Category = "Interchange | Node")
-	bool AddInt32Attribute(const FString& NodeAttributeKey, const int32& Value)
-	{
-		INTERCHANGE_BASE_NODE_ADD_ATTRIBUTE(int32);
-	}
+	bool AddInt32Attribute(const FString& NodeAttributeKey, const int32& Value);
 	
 	/**
 	 * Get a int32 attribute from this node. Return false if the attribute do not exist
 	 */
 	UFUNCTION(BlueprintCallable, Category = "Interchange | Node")
-	bool GetInt32Attribute(const FString& NodeAttributeKey, int32& OutValue)
-	{
-		INTERCHANGE_BASE_NODE_GET_ATTRIBUTE(int32);
-	}
+	bool GetInt32Attribute(const FString& NodeAttributeKey, int32& OutValue);
 
 	/**
 	 * Add a float attribute to this node. Return false if the attribute do not exist or if we cannot add it
 	 */
 	UFUNCTION(BlueprintCallable, Category = "Interchange | Node")
-	bool AddFloatAttribute(const FString& NodeAttributeKey, const float& Value)
-	{
-		INTERCHANGE_BASE_NODE_ADD_ATTRIBUTE(float);
-	}
+	bool AddFloatAttribute(const FString& NodeAttributeKey, const float& Value);
 	
 	/**
 	 * Get a float attribute from this node. Return false if the attribute do not exist
 	 */
 	UFUNCTION(BlueprintCallable, Category = "Interchange | Node")
-	bool GetFloatAttribute(const FString& NodeAttributeKey, float& OutValue)
-	{
-		INTERCHANGE_BASE_NODE_GET_ATTRIBUTE(float);
-	}
+	bool GetFloatAttribute(const FString& NodeAttributeKey, float& OutValue);
 
 	/**
 	 * Add a string attribute to this node. Return false if the attribute do not exist or if we cannot add it
 	 */
 	UFUNCTION(BlueprintCallable, Category = "Interchange | Node")
-	bool AddStringAttribute(const FString& NodeAttributeKey, const FString& Value)
-	{
-		INTERCHANGE_BASE_NODE_ADD_ATTRIBUTE(FString);
-	}
+	bool AddStringAttribute(const FString& NodeAttributeKey, const FString& Value);
 	
 	/**
 	 * Get a string attribute from this node. Return false if the attribute do not exist
 	 */
 	UFUNCTION(BlueprintCallable, Category = "Interchange | Node")
-	bool GetStringAttribute(const FString& NodeAttributeKey, FString& OutValue)
-	{
-		INTERCHANGE_BASE_NODE_GET_ATTRIBUTE(FString);
-	}
+	bool GetStringAttribute(const FString& NodeAttributeKey, FString& OutValue);
 
 	/**
 	 * Return the unique id pass in the constructor.
-	 *
 	 */
 	UFUNCTION(BlueprintCallable, Category = "Interchange | Node")
 	FString GetUniqueID() const;
 
 	/**
 	 * Return the display label.
-	 *
 	 */
 	UFUNCTION(BlueprintCallable, Category = "Interchange | Node")
 	FString GetDisplayLabel() const;
 
 	/**
 	 * Change the display label.
-	 *
 	 */
 	UFUNCTION(BlueprintCallable, Category = "Interchange | Node")
 	bool SetDisplayLabel(const FString& DisplayName);
@@ -1129,21 +440,18 @@ public:
 
 	/**
 	 * This function allow to retrieve the number of factory dependencies for this object.
-	 *
 	 */
 	UFUNCTION(BlueprintCallable, Category = "Interchange | Node")
 	int32 GetFactoryDependenciesCount() const;
 
 	/**
 	 * This function allow to retrieve the dependency for this object.
-	 * 
 	 */
 	UFUNCTION(BlueprintCallable, Category = "Interchange | Node")
 	void GetFactoryDependencies(TArray<FString>& OutDependencies ) const;
 
 	/**
 	 * This function allow to retrieve one dependency for this object.
-	 *
 	 */
 	UFUNCTION(BlueprintCallable, Category = "Interchange | Node")
 	void GetFactoryDependency(const int32 Index, FString& OutDependency) const;
@@ -1239,28 +547,21 @@ public:
 	 * Each Attribute that was set and have a delegate set for the specified UObject->UClass will
 	 * get the delegate execute so it apply the attribute to the UObject property.
 	 * See the macros IMPLEMENT_NODE_ATTRIBUTE_SETTER at the top of the file to know how delegates are setup for property.
-	 *
 	 */
 	void ApplyAllCustomAttributeToObject(UObject* Object) const;
 
 	void FillAllCustomAttributeFromObject(UObject* Object) const;
 
+	/**
+	 * Serialize the node, by default only the attribute storage is serialize for a node
+	 */
 	virtual void Serialize(FArchive& Ar) override;
 
-	static void CompareNodeStorage(UInterchangeBaseNode* NodeA, const UInterchangeBaseNode* NodeB, TArray<UE::Interchange::FAttributeKey>& RemovedAttributes, TArray<UE::Interchange::FAttributeKey>& AddedAttributes, TArray<UE::Interchange::FAttributeKey>& ModifiedAttributes)
-	{
-		UE::Interchange::FAttributeStorage::CompareStorage(*(NodeA->Attributes), *(NodeB->Attributes), RemovedAttributes, AddedAttributes, ModifiedAttributes);
-	}
+	static void CompareNodeStorage(UInterchangeBaseNode* NodeA, const UInterchangeBaseNode* NodeB, TArray<UE::Interchange::FAttributeKey>& RemovedAttributes, TArray<UE::Interchange::FAttributeKey>& AddedAttributes, TArray<UE::Interchange::FAttributeKey>& ModifiedAttributes);
 	
-	static void CopyStorageAttributes(const UInterchangeBaseNode* SourceNode, UInterchangeBaseNode* DestinationNode, TArray<UE::Interchange::FAttributeKey>& AttributeKeys)
-	{
-		UE::Interchange::FAttributeStorage::CopyStorageAttributes(*(SourceNode->Attributes), *(DestinationNode->Attributes), AttributeKeys);
-	}
+	static void CopyStorageAttributes(const UInterchangeBaseNode* SourceNode, UInterchangeBaseNode* DestinationNode, TArray<UE::Interchange::FAttributeKey>& AttributeKeys);
 
-	static void CopyStorage(const UInterchangeBaseNode* SourceNode, UInterchangeBaseNode* DestinationNode)
-	{
-		*(DestinationNode->Attributes) = *(SourceNode->Attributes);
-	}
+	static void CopyStorage(const UInterchangeBaseNode* SourceNode, UInterchangeBaseNode* DestinationNode);
 	
 	UPROPERTY()
 	mutable FSoftObjectPath ReferenceObject;
@@ -1280,10 +581,10 @@ protected:
 	 * Example: Mesh factory node will have dependencies on material factory node
 	 *          Material factory node will have dependencies on texture factory node
 	 */
-	UE::Interchange::FNameAttributeArrayHelper FactoryDependencies;
+	UE::Interchange::TArrayAttributeHelper<FString> FactoryDependencies;
 
 	/**
 	 * This tracks the IDs of asset nodes which are the target of factories
 	 */
-	mutable UE::Interchange::FNameAttributeArrayHelper TargetNodes;
+	mutable UE::Interchange::TArrayAttributeHelper<FString> TargetNodes;
 };
diff --git a/Engine/Source/Runtime/Experimental/Interchange/Core/Public/Nodes/InterchangeBaseNodeUtilities.h b/Engine/Source/Runtime/Experimental/Interchange/Core/Public/Nodes/InterchangeBaseNodeUtilities.h
new file mode 100644
--- /dev/null
+++ b/Engine/Source/Runtime/Experimental/Interchange/Core/Public/Nodes/InterchangeBaseNodeUtilities.h
@@ -0,0 +1,573 @@
+// Copyright Epic Games, Inc. All Rights Reserved.
+
+#pragma once
+
+#include "CoreMinimal.h"
+#include "Types/AttributeStorage.h"
+#include "UObject/Class.h"
+#include "UObject/Object.h"
+#include "UObject/ObjectMacros.h"
+
+//Interchange namespace
+namespace UE
+{
+namespace Interchange
+{
+
+/**
+ * Helper class to manage an array of items inside an attribute storage
+ */
+template<typename ItemType>
+class TArrayAttributeHelper
+{
+	static_assert(TAttributeTypeTraits<ItemType>::GetType() != EAttributeTypes::None, "The value type must be supported by the attribute storage");
+	
+public:
+	~TArrayAttributeHelper()
+	{
+		Attributes = nullptr;
+		KeyCount = NAME_None;
+	}
+
+	void Initialize(const TSharedPtr<FAttributeStorage, ESPMode::ThreadSafe>& InAttributes, const FString& BaseKeyName)
+	{
+		check(Attributes.IsValid());
+		Attributes = InAttributes;
+		KeyCount = BaseKeyName;
+	}
+
+	static const FString& IndexKey()
+	{
+		static FString IndexKeyString = TEXT("_NameIndex_");
+		return IndexKeyString;
+	}
+
+	int32 GetCount() const
+	{
+		//The Class must be initialise properly before we can use it
+		TSharedPtr<FAttributeStorage, ESPMode::ThreadSafe> AttributePtr = Attributes.Pin();
+		if (!ensure(AttributePtr.IsValid()))
+		{
+			return 0;
+		}
+		int32 ItemCount = 0;
+		if (AttributePtr->ContainAttribute(GetKeyCount()))
+		{
+			FAttributeStorage::TAttributeHandle<int32> Handle = AttributePtr->GetAttributeHandle<int32>(GetKeyCount());
+			if (Handle.IsValid())
+			{
+				Handle.Get(ItemCount);
+			}
+		}
+		return ItemCount;
+	}
+
+	void GetItem(const int32 Index, ItemType& OutItem) const
+	{
+		//The Class must be initialise properly before we can use it
+		TSharedPtr<FAttributeStorage, ESPMode::ThreadSafe> AttributePtr = Attributes.Pin();
+		if (!ensure(AttributePtr.IsValid()))
+		{
+			return;
+		}
+		int32 ItemCount = 0;
+		if (!AttributePtr->ContainAttribute(GetKeyCount()))
+		{
+			return;
+		}
+
+		FAttributeStorage::TAttributeHandle<int32> Handle = AttributePtr->GetAttributeHandle<int32>(GetKeyCount());
+		if (!Handle.IsValid())
+		{
+			return;
+		}
+		Handle.Get(ItemCount);
+		if (Index >= ItemCount)
+		{
+			return;
+		}
+		FAttributeKey DepIndexKey = GetIndexKey(Index);
+		FAttributeStorage::TAttributeHandle<ItemType> HandleItem = AttributePtr->GetAttributeHandle<ItemType>(DepIndexKey);
+		if (!HandleItem.IsValid())
+		{
+			return;
+		}
+		HandleItem.Get(OutItem);
+	}
+
+	void GetItems(TArray<ItemType>& OutItems) const
+	{
+		//The Class must be initialise properly before we can use it
+		TSharedPtr<FAttributeStorage, ESPMode::ThreadSafe> AttributePtr = Attributes.Pin();
+		if (!ensure(AttributePtr.IsValid()))
+		{
+			OutItems.Empty();
+			return;
+		}
+		int32 ItemCount = 0;
+		if (!AttributePtr->ContainAttribute(GetKeyCount()))
+		{
+			OutItems.Empty();
+			return;
+		}
+		
+		FAttributeStorage::TAttributeHandle<int32> Handle = AttributePtr->GetAttributeHandle<int32>(GetKeyCount());
+		if (!Handle.IsValid())
+		{
+			return;
+		}
+		Handle.Get(ItemCount);
+
+		//Reuse as much memory we can to avoid allocation
+		OutItems.Reset(ItemCount);
+		for (int32 NameIndex = 0; NameIndex < ItemCount; ++NameIndex)
+		{
+			FAttributeKey DepIndexKey = GetIndexKey(NameIndex);
+			FAttributeStorage::TAttributeHandle<ItemType> HandleItem = AttributePtr->GetAttributeHandle<ItemType>(DepIndexKey);
+			if (!HandleItem.IsValid())
+			{
+				continue;
+			}
+			ItemType& OutItem = OutItems.AddDefaulted_GetRef();
+			HandleItem.Get(OutItem);
+		}
+	}
+
+	bool AddItem(const ItemType& Item)
+	{
+		//The Class must be initialise properly before we can use it
+		TSharedPtr<FAttributeStorage, ESPMode::ThreadSafe> AttributePtr = Attributes.Pin();
+		if (!ensure(AttributePtr.IsValid()))
+		{
+			return false;
+		}
+
+		if (!AttributePtr->ContainAttribute(GetKeyCount()))
+		{
+			const int32 DependencyCount = 0;
+			EAttributeStorageResult Result = AttributePtr->RegisterAttribute<int32>(GetKeyCount(), DependencyCount);
+			if (!IsAttributeStorageResultSuccess(Result))
+			{
+				LogAttributeStorageErrors(Result, TEXT("TArrayAttributeHelper.AddName"), GetKeyCount());
+				return false;
+			}
+		}
+		FAttributeStorage::TAttributeHandle<int32> Handle = AttributePtr->GetAttributeHandle<int32>(GetKeyCount());
+		if (!ensure(Handle.IsValid()))
+		{
+			return false;
+		}
+		int32 ItemIndex = 0;
+		Handle.Get(ItemIndex);
+		FAttributeKey ItemIndexKey = GetIndexKey(ItemIndex);
+		EAttributeStorageResult AddItemResult = AttributePtr->RegisterAttribute<ItemType>(ItemIndexKey, Item);
+		if (!IsAttributeStorageResultSuccess(AddItemResult))
+		{
+			LogAttributeStorageErrors(AddItemResult, TEXT("TArrayAttributeHelper.AddName"), ItemIndexKey);
+			return false;
+		}
+
+		//Success, increment the item counter
+		ItemIndex++;
+		Handle.Set(ItemIndex);
+		return true;
+	}
+
+	bool RemoveItem(const ItemType& ItemToDelete)
+	{
+		//The Class must be initialise properly before we can use it
+		TSharedPtr<FAttributeStorage, ESPMode::ThreadSafe> AttributePtr = Attributes.Pin();
+		if (!ensure(AttributePtr.IsValid()))
+		{
+			return false;
+		}
+
+		int32 ItemCount = 0;
+		if (!AttributePtr->ContainAttribute(GetKeyCount()))
+		{
+			return false;
+		}
+		FAttributeStorage::TAttributeHandle<int32> Handle = AttributePtr->GetAttributeHandle<int32>(GetKeyCount());
+		if (!Handle.IsValid())
+		{
+			return false;
+		}
+		Handle.Get(ItemCount);
+		int32 DecrementKey = 0;
+		for (int32 ItemIndex = 0; ItemIndex < ItemCount; ++ItemIndex)
+		{
+			FAttributeKey DepIndexKey = GetIndexKey(ItemIndex);
+			FAttributeStorage::TAttributeHandle<ItemType> HandleItem = AttributePtr->GetAttributeHandle<ItemType>(DepIndexKey);
+			if (!HandleItem.IsValid())
+			{
+				continue;
+			}
+			ItemType Item;
+			HandleItem.Get(Item);
+			if (Item == ItemToDelete)
+			{
+				//Remove this entry
+				AttributePtr->UnregisterAttribute(DepIndexKey);
+				Handle.Set(ItemCount - 1);
+				//We have to rename the key for all the next item
+				DecrementKey++;
+			}
+			else if (DecrementKey > 0)
+			{
+				FAttributeKey NewDepIndexKey = GetIndexKey(ItemIndex - DecrementKey);
+				EAttributeStorageResult UnregisterResult = AttributePtr->UnregisterAttribute(DepIndexKey);
+				if (IsAttributeStorageResultSuccess(UnregisterResult))
+				{
+					EAttributeStorageResult RegisterResult = AttributePtr->RegisterAttribute<ItemType>(NewDepIndexKey, Item);
+					if (!IsAttributeStorageResultSuccess(RegisterResult))
+					{
+						LogAttributeStorageErrors(RegisterResult, TEXT("TArrayAttributeHelper.RemoveItem"), NewDepIndexKey);
+					}
+				}
+				else
+				{
+					LogAttributeStorageErrors(UnregisterResult, TEXT("TArrayAttributeHelper.RemoveItem"), DepIndexKey);
+				}
+
+				//Avoid doing more code in the for since the HandleItem is now invalid
+				continue;
+			}
+		}
+		return true;
+	}
+
+	bool RemoveAllItems()
+	{
+		//The Class must be initialise properly before we can use it
+		TSharedPtr<FAttributeStorage, ESPMode::ThreadSafe> AttributePtr = Attributes.Pin();
+		if (!ensure(AttributePtr.IsValid()))
+		{
+			return false;
+		}
+
+		int32 ItemCount = 0;
+		if (!AttributePtr->ContainAttribute(GetKeyCount()))
+		{
+			return false;
+		}
+		FAttributeStorage::TAttributeHandle<int32> HandleCount = AttributePtr->GetAttributeHandle<int32>(GetKeyCount());
+		if (!HandleCount.IsValid())
+		{
+			return false;
+		}
+		HandleCount.Get(ItemCount);
+		//Remove all attribute one by one
+		for (int32 NameIndex = 0; NameIndex < ItemCount; ++NameIndex)
+		{
+			FAttributeKey DepIndexKey = GetIndexKey(NameIndex);
+			AttributePtr->UnregisterAttribute(DepIndexKey);
+		}
+		//Make sure Count is zero
+		ItemCount = 0;
+		HandleCount.Set(ItemCount);
+		return true;
+	}
+
+private:
+	TWeakPtr<FAttributeStorage, ESPMode::ThreadSafe> Attributes = nullptr;
+	FAttributeKey KeyCount;
+	FAttributeKey GetKeyCount() const
+	{
+		ensure(!KeyCount.Key.IsEmpty()); return KeyCount;
+	}
+
+	FAttributeKey GetIndexKey(int32 Index) const
+	{
+		FString DepIndexKeyString = GetKeyCount().ToString() + IndexKey() + FString::FromInt(Index);
+		return FAttributeKey(*DepIndexKeyString);
+	}
+};
+
+/**
+ * Helper class to manage a TMap where the key and value type can be any type support by the FAttributeStorage
+ */
+template<typename KeyType, typename ValueType>
+class TMapAttributeHelper
+{
+	static_assert(TAttributeTypeTraits<KeyType>::GetType() != EAttributeTypes::None, "The key type must be supported by the attribute storage");
+	static_assert(TAttributeTypeTraits<ValueType>::GetType() != EAttributeTypes::None, "The value type must be supported by the attribute storage");
+
+	template<typename T>
+	using TAttributeHandle = FAttributeStorage::TAttributeHandle<T>;
+
+public:
+	void Initialize(const TSharedRef<FAttributeStorage, ESPMode::ThreadSafe>& InAttributes, const FString& BaseKeyName)
+	{
+		Attributes = InAttributes;
+		FString BaseTryName = BaseKeyName;
+		FAttributeKey KeyCountKey(MoveTemp(BaseTryName));
+		if (!InAttributes->ContainAttribute(KeyCountKey))
+		{
+			EAttributeStorageResult Result = InAttributes->RegisterAttribute<int32>(KeyCountKey, 0);
+			check(Result == EAttributeStorageResult::Operation_Success);
+			KeyCountHandle = InAttributes->GetAttributeHandle<int32>(KeyCountKey);
+		}
+		else
+		{
+			KeyCountHandle = InAttributes->GetAttributeHandle<int32>(KeyCountKey);
+
+			//Init map
+			int32 KeyCount;
+			KeyCountHandle.Get(KeyCount);
+
+			CachedKeysAndValues.Reserve(KeyCount);
+
+			for (int32 Index = 0; Index < KeyCount; Index++)
+			{
+				TAttributeHandle<KeyType> KeyAttribute = InAttributes->GetAttributeHandle<KeyType>(GetKeyAttribute(Index));
+				KeyType Key;
+				KeyAttribute.Get(Key);
+				TAttributeHandle<ValueType> ValueAttribute = InAttributes->GetAttributeHandle<ValueType>(GetValueAttribute(Key));
+				CachedKeysAndValues.Add(Key, TPair<TAttributeHandle<KeyType>, TAttributeHandle<ValueType>>(KeyAttribute, ValueAttribute));
+			}
+		}
+	}
+
+	TMapAttributeHelper() = default;
+	TMapAttributeHelper(const TMapAttributeHelper&) = default;
+	TMapAttributeHelper(TMapAttributeHelper&&) = default;
+	TMapAttributeHelper& operator=(const TMapAttributeHelper&) = default;
+	TMapAttributeHelper& operator=(TMapAttributeHelper&&) = default;
+
+	~TMapAttributeHelper()
+	{
+		Attributes.Reset();
+	}
+
+	void SetKeyValue(const KeyType& InKey, const ValueType& InValue)
+	{
+		const uint32 Hash = GetTypeHash(InKey);
+		SetKeyValueByHash(Hash, InKey, InValue);
+	}
+
+	bool GetValue(const KeyType& InKey, ValueType& OutValue) const
+	{
+		const uint32 Hash = GetTypeHash(InKey);
+		return GetValueByHash(Hash, InKey, OutValue);
+	}
+
+	bool RemoveKey(const KeyType& InKey)
+	{
+		const uint32 Hash = GetTypeHash(InKey);
+		return RemoveKeyByHash(Hash, InKey);
+	}
+
+	bool RemoveKeyAndGetValue(const KeyType& InKey, ValueType& OutValue)
+	{
+		const uint32 Hash = GetTypeHash(InKey);
+		return RemoveKeyAndGetValueByHash(Hash, InKey, OutValue);
+	}
+
+	void SetKeyValueByHash(uint32 Hash, const KeyType& InKey, const ValueType& InValue)
+	{
+		TSharedPtr<FAttributeStorage, ESPMode::ThreadSafe> AttributesPtr = Attributes.Pin();
+		if (!ensure(AttributesPtr.IsValid()))
+		{
+			return;
+		}
+
+		if (TPair<TAttributeHandle<KeyType>, TAttributeHandle<ValueType>>* Pair = CachedKeysAndValues.FindByHash(Hash, InKey))
+		{
+			Pair->Value.Set(InValue);
+		}
+		else
+		{
+			FAttributeKey IndexKey = GetKeyAttribute(CachedKeysAndValues.Num());
+			FAttributeKey AttributeKey = GetValueAttribute(InKey);
+			check(AttributesPtr->RegisterAttribute<KeyType>(IndexKey, InKey) == EAttributeStorageResult::Operation_Success);
+			check(AttributesPtr->RegisterAttribute<ValueType>(AttributeKey, InValue) == EAttributeStorageResult::Operation_Success);
+			CachedKeysAndValues.AddByHash(
+				Hash
+				, InKey
+				, TPair<TAttributeHandle<KeyType>, TAttributeHandle<ValueType>>(
+					AttributesPtr->GetAttributeHandle<KeyType>(IndexKey)
+					, AttributesPtr->GetAttributeHandle<ValueType>(AttributeKey)
+					));
+			KeyCountHandle.Set(CachedKeysAndValues.Num());
+		}
+	}
+
+
+	bool GetValueByHash(uint32 Hash, const KeyType& InKey, ValueType& OutValue) const
+	{
+		TSharedPtr<FAttributeStorage, ESPMode::ThreadSafe> AttributesPtr = Attributes.Pin();
+		if (!ensure(AttributesPtr.IsValid()))
+		{
+			return false;
+		}
+
+		if (const TPair<TAttributeHandle<KeyType>, TAttributeHandle<ValueType>>* Pair = CachedKeysAndValues.FindByHash(Hash, InKey))
+		{
+			return Pair->Value.Get(OutValue) == EAttributeStorageResult::Operation_Success;
+		}
+	}
+
+
+	bool RemoveKeyByHash(uint32 Hash, const KeyType& InKey)
+	{
+		TSharedPtr<FAttributeStorage, ESPMode::ThreadSafe> AttributesPtr = Attributes.Pin();
+		if (!ensure(AttributesPtr.IsValid()))
+		{
+			return false;
+		}
+
+		if (TPair<TAttributeHandle<KeyType>, TAttributeHandle<ValueType>>* Pair = CachedKeysAndValues.FindByHash(Hash, InKey))
+		{
+			return RemoveBySwap(Hash, InKey, *Pair);
+		}
+
+		return false;
+	}
+
+	bool RemoveKeyAndGetValueByHash(uint32 Hash, const KeyType& InKey, ValueType& OutValue)
+	{
+		TSharedPtr<FAttributeStorage, ESPMode::ThreadSafe> AttributesPtr = Attributes.Pin();
+		if (!ensure(AttributesPtr.IsValid()))
+		{
+			return false;
+		}
+
+		if (TPair<TAttributeHandle<KeyType>, TAttributeHandle<ValueType>>* Pair = CachedKeysAndValues.FindByHash(Hash, InKey))
+		{
+			if (Pair->Value.Get(OutValue) == EAttributeStorageResult::Operation_Success)
+			{
+				return RemoveBySwap(Hash, InKey, *Pair);
+			}
+		}
+
+		return false;
+	}
+
+	void Reserve(int32 Number)
+	{
+		CachedKeysAndValues.Reserve(Number);
+	}
+
+	void Empty(int32 NumOfExpectedElements = 0)
+	{
+		EmptyInternal(NumOfExpectedElements);
+	}
+
+	TMapAttributeHelper& operator=(const TMap<KeyType, ValueType>& InMap)
+	{
+		TSharedPtr<FAttributeStorage, ESPMode::ThreadSafe> AttributesPtr = Attributes.Pin();
+		if (!ensure(AttributesPtr.IsValid()))
+		{
+			return *this;
+		}
+
+		// Empty
+		EmptyInternal(InMap.Num());
+
+		KeyCountHandle.Set(InMap.Num());
+
+		for (const TPair<KeyType, ValueType>& Pair : InMap)
+		{
+			SetKeyValue(Pair.Key, Pair.Value);
+		}
+
+		return *this;
+	}
+
+	TMap<KeyType, ValueType> ToMap() const
+	{
+		TMap<KeyType, ValueType> Map;
+		Map.Reserve(CachedKeysAndValues.Num());
+
+		for (const TPair<KeyType, TPair<TAttributeHandle<KeyType>, TAttributeHandle<ValueType>>>& Pair : CachedKeysAndValues)
+		{
+			KeyType Key;
+			Pair.Value.Key.Get(Key);
+			ValueType Value;
+			Pair.Value.Value.Get(Value);
+
+			Map.Add(MoveTemp(Key), MoveTemp(Value));
+		}
+
+		return Map;
+	}
+
+private:
+	FAttributeKey GetKeyAttribute(int32 Index) const
+	{
+		static const FString KeyIndex = TEXT("_KeyIndex_");
+		FString IndexedKey = KeyCountHandle.GetKey().ToString();
+		IndexedKey.Reserve(KeyIndex.Len() + 16 /*Max size for a int32*/);
+		IndexedKey.Append(KeyIndex);
+		IndexedKey.AppendInt(Index);
+		return FAttributeKey(MoveTemp(IndexedKey));
+	}
+
+	FAttributeKey GetValueAttribute(const KeyType& InKey) const
+	{
+		static const FString KeyChars = TEXT("_Key_");
+		FString ValueAttribute = KeyCountHandle.GetKey().ToString();
+		FString Key = TTypeToString<KeyType>::ToString(InKey);
+		ValueAttribute.Reserve(KeyChars.Len() + Key.Len());
+		ValueAttribute.Append(KeyChars);
+		ValueAttribute.Append(Key);
+		return FAttributeKey(MoveTemp(ValueAttribute));
+	}
+
+	TAttributeHandle<KeyType> GetLastKeyAttributeHandle() const
+	{
+		TSharedPtr<FAttributeStorage, ESPMode::ThreadSafe> AttributesPtr = Attributes.Pin();
+		if (!ensure(AttributesPtr.IsValid()))
+		{
+			return {};
+		}
+		return AttributesPtr->GetAttributeHandle<KeyType>(GetKeyAttribute(CachedKeysAndValues.Num() - 1));
+	}
+
+	bool RemoveBySwap(uint32 Hash, const KeyType& InKey, TPair<TAttributeHandle<FString>, TAttributeHandle<ValueType>>& CachedPair)
+	{
+		TSharedPtr<FAttributeStorage, ESPMode::ThreadSafe> AttributesPtr = Attributes.Pin();
+		if (!ensure(AttributesPtr.IsValid()))
+		{
+			return false;
+		}
+
+		TAttributeHandle<KeyType> LastKeyIndex = GetLastKeyAttributeHandle();
+		KeyType LastKey;
+		LastKeyIndex.Get(LastKey);
+		CachedKeysAndValues[LastKey].Key = CachedPair.Key;
+		CachedPair->Key.Set(MoveTemp(LastKey));
+
+		AttributesPtr->UnregisterAttribute(LastKeyIndex.GetKey());
+		AttributesPtr->UnregisterAttribute(CachedPair.Value.GetKey());
+		CachedKeysAndValues.RemoveByHash(Hash, InKey);
+		KeyCountHandle.Set(CachedKeysAndValues.Num());
+	}
+
+	void EmptyInternal(int32 NumOfExpectedElements)
+	{
+		TSharedPtr<FAttributeStorage, ESPMode::ThreadSafe> AttributesPtr = Attributes.Pin();
+		if (!ensure(AttributesPtr.IsValid()))
+		{
+			return;
+		}
+
+		for (const TPair<KeyType, TPair<TAttributeHandle<KeyType>, TAttributeHandle<ValueType>>>& Pair : CachedKeysAndValues)
+		{
+			const FAttributeKey& KeyAttribute = Pair.Value.Key.GetKey();
+			AttributesPtr->UnregisterAttribute(KeyAttribute);
+
+			const FAttributeKey& ValueAttribute = Pair.Value.Value.GetKey();
+			AttributesPtr->UnregisterAttribute(ValueAttribute);
+		}
+
+		CachedKeysAndValues.Empty(NumOfExpectedElements);
+	}
+
+	TMap<KeyType, TPair<TAttributeHandle<KeyType>, TAttributeHandle<ValueType>>> CachedKeysAndValues;
+	TAttributeHandle<int32> KeyCountHandle; //Assign in Initialize function, it will ensure if its the default value (NAME_None) when using the class
+	TWeakPtr<FAttributeStorage, ESPMode::ThreadSafe> Attributes = nullptr;
+};
+
+} //ns Interchange
+} //ns UE
