LWC: Enable ISPC for LWC double types. Submitted on behalf of jeff.rous[REDACTED_HANDLE]

diff --git a/Engine/Source/Runtime/Core/Public/Math/WideVector.isph b/Engine/Source/Runtime/Core/Public/Math/WideVector.isph
--- a/Engine/Source/Runtime/Core/Public/Math/WideVector.isph
+++ b/Engine/Source/Runtime/Core/Public/Math/WideVector.isph
@@ -3,24 +3,66 @@
 #ifndef __WIDEVECTOR_ISPH__
 #define __WIDEVECTOR_ISPH__
 
+#include "Math/Vector.isph"
+
+#if UE_LARGE_WORLD_COORDINATES_DISABLED
 struct WideFVector4
 {
 	float V[programCount];
 };
 
+struct WideFVector4d
+{
+	double V[programCount];
+};
+#else
+struct WideFVector4f
+{
+	float V[programCount];
+};
+
+struct WideFVector4
+{
+	double V[programCount];
+};
+#endif
+
+// Define types
+#if UE_LARGE_WORLD_COORDINATES_DISABLED
+typedef WideFVector4 WideFVector4f;
+#else
+typedef WideFVector4 WideFVector4d;
+#endif
+
+#if UE_LARGE_WORLD_COORDINATES_DISABLED
+#if TARGET_WIDTH == 4
+static const uniform struct WideFVector4f WFloatZero = { {0.0f, 0.0f, 0.0f, 0.0f} };
+static const uniform struct WideFVector4f WFloat0001 = { {0.0f, 0.0f, 0.0f, 1.0f} };
+static const uniform struct WideFVector4f WSmallLengthThreshold = { {1.e-8f, 1.e-8f, 1.e-8f, 1.e-8f} };
+#elif TARGET_WIDTH == 8
+static const uniform struct WideFVector4f WFloatZero = { {0.0f, 0.0f, 0.0f, 0.0f, 0.0f, 0.0f, 0.0f, 0.0f} };
+static const uniform struct WideFVector4f WFloat0001 = { {0.0f, 0.0f, 0.0f, 1.0f, 0.0f, 0.0f, 0.0f, 1.0f} };
+static const uniform struct WideFVector4f WSmallLengthThreshold = { {1.e-8f, 1.e-8f, 1.e-8f, 1.e-8f, 1.e-8f, 1.e-8f, 1.e-8f, 1.e-8f} };
+#elif TARGET_WIDTH == 16
+static const uniform struct WideFVector4f WFloatZero = { {0.0f, 0.0f, 0.0f, 0.0f, 0.0f, 0.0f, 0.0f, 0.0f, 0.0f, 0.0f, 0.0f, 0.0f, 0.0f, 0.0f, 0.0f, 0.0f} };
+static const uniform struct WideFVector4f WFloat0001 = { {0.0f, 0.0f, 0.0f, 1.0f, 0.0f, 0.0f, 0.0f, 1.0f, 0.0f, 0.0f, 0.0f, 1.0f, 0.0f, 0.0f, 0.0f, 1.0f} };
+static const uniform struct WideFVector4f WSmallLengthThreshold = { {1.e-8f, 1.e-8f, 1.e-8f, 1.e-8f, 1.e-8f, 1.e-8f, 1.e-8f, 1.e-8f, 1.e-8f, 1.e-8f, 1.e-8f, 1.e-8f, 1.e-8f, 1.e-8f, 1.e-8f, 1.e-8f} };
+#endif
+#else
 #if TARGET_WIDTH == 4
-static const uniform struct WideFVector4 WFloatZero = { {0.0f, 0.0f, 0.0f, 0.0f} };
-static const uniform struct WideFVector4 WFloat0001 = { {0.0f, 0.0f, 0.0f, 1.0f} };
-static const uniform struct WideFVector4 WSmallLengthThreshold = { {1.e-8f, 1.e-8f, 1.e-8f, 1.e-8f} };
+static const uniform struct WideFVector4d WFloatZero = { {0.0d, 0.0d, 0.0d, 0.0d} };
+static const uniform struct WideFVector4d WFloat0001 = { {0.0d, 0.0d, 0.0d, 1.0d} };
+static const uniform struct WideFVector4d WSmallLengthThreshold = { {1.d-8, 1.d-8, 1.d-8, 1.d-8} };
 #elif TARGET_WIDTH == 8
-static const uniform struct WideFVector4 WFloatZero = { {0.0f, 0.0f, 0.0f, 0.0f, 0.0f, 0.0f, 0.0f, 0.0f} };
-static const uniform struct WideFVector4 WFloat0001 = { {0.0f, 0.0f, 0.0f, 1.0f, 0.0f, 0.0f, 0.0f, 1.0f} };
-static const uniform struct WideFVector4 WSmallLengthThreshold = { {1.e-8f, 1.e-8f, 1.e-8f, 1.e-8f, 1.e-8f, 1.e-8f, 1.e-8f, 1.e-8f} };
+static const uniform struct WideFVector4d WFloatZero = { {0.0d, 0.0d, 0.0d, 0.0d, 0.0d, 0.0d, 0.0d, 0.0d} };
+static const uniform struct WideFVector4d WFloat0001 = { {0.0d, 0.0d, 0.0d, 1.0d, 0.0d, 0.0d, 0.0d, 1.0d} };
+static const uniform struct WideFVector4d WSmallLengthThreshold = { {1.d-8, 1.d-8, 1.d-8, 1.d-8, 1.d-8, 1.d-8, 1.d-8, 1.d-8} };
 #elif TARGET_WIDTH == 16
-static const uniform struct WideFVector4 WFloatZero = { {0.0f, 0.0f, 0.0f, 0.0f, 0.0f, 0.0f, 0.0f, 0.0f, 0.0f, 0.0f, 0.0f, 0.0f, 0.0f, 0.0f, 0.0f, 0.0f} };
-static const uniform struct WideFVector4 WFloat0001 = { {0.0f, 0.0f, 0.0f, 1.0f, 0.0f, 0.0f, 0.0f, 1.0f, 0.0f, 0.0f, 0.0f, 1.0f, 0.0f, 0.0f, 0.0f, 1.0f} };
-static const uniform struct WideFVector4 WSmallLengthThreshold = { {1.e-8f, 1.e-8f, 1.e-8f, 1.e-8f, 1.e-8f, 1.e-8f, 1.e-8f, 1.e-8f, 1.e-8f, 1.e-8f, 1.e-8f, 1.e-8f, 1.e-8f, 1.e-8f, 1.e-8f, 1.e-8f} };
+static const uniform struct WideFVector4d WFloatZero = { {0.0d, 0.0d, 0.0d, 0.0d, 0.0d, 0.0d, 0.0d, 0.0d, 0.0d, 0.0d, 0.0d, 0.0d, 0.0d, 0.0d, 0.0d, 0.0d} };
+static const uniform struct WideFVector4d WFloat0001 = { {0.0d, 0.0d, 0.0d, 1.0d, 0.0d, 0.0d, 0.0d, 1.0d, 0.0d, 0.0d, 0.0d, 1.0d, 0.0d, 0.0d, 0.0d, 1.0d} };
+static const uniform struct WideFVector4d WSmallLengthThreshold = { {1.d-8, 1.d-8, 1.d-8, 1.d-8, 1.d-8, 1.d-8, 1.d-8, 1.d-8, 1.d-8, 1.d-8, 1.d-8, 1.d-8, 1.d-8, 1.d-8, 1.d-8, 1.d-8} };
 #endif
+#endif // UE_LARGE_WORLD_COORDINATES_DISABLED
 
 unmasked inline uniform WideFVector4 operator+(const uniform WideFVector4 &A, const uniform WideFVector4 &B)
 {
@@ -49,7 +91,7 @@ unmasked inline uniform WideFVector4 operator*(const uniform WideFVector4 &A, co
 	return Result;
 }
 
-unmasked inline uniform WideFVector4 operator*(const uniform WideFVector4 &A, const uniform float F)
+unmasked inline uniform WideFVector4 operator*(const uniform WideFVector4 &A, const uniform FReal F)
 {
 	uniform WideFVector4 Result;
 
@@ -68,8 +110,8 @@ unmasked inline uniform WideFVector4 VectorSwizzle(const uniform WideFVector4 &V
 	const varying int vPerm = { X, Y, Z, W, X+4, Y+4, Z+4, W+4, X+8, Y+8, Z+8, W+8, X+12, Y+12, Z+12, W+12 };
 #endif
 
-	const float V = Vec.V[programIndex];
-	const float R = shuffle(V, vPerm);
+	const FReal V = Vec.V[programIndex];
+	const FReal R = shuffle(V, vPerm);
 	uniform WideFVector4 Result;
 	Result.V[programIndex] = R;
 
@@ -86,8 +128,8 @@ unmasked inline uniform WideFVector4 VectorReplicate(const uniform WideFVector4
 	const varying int vPerm = { R, R, R, R, R+4, R+4, R+4, R+4, R+8, R+8, R+8, R+8, R+12, R+12, R+12, R+12 };
 #endif
 
-	const float V = Vec.V[programIndex];
-	const float S = shuffle(V, vPerm);
+	const FReal V = Vec.V[programIndex];
+	const FReal S = shuffle(V, vPerm);
 	uniform WideFVector4 Result;
 	Result.V[programIndex] = S;
 
@@ -98,20 +140,29 @@ unmasked inline uniform WideFVector4 VectorCompareGE(const uniform WideFVector4
 {
 	uniform WideFVector4 Result;
 
-	Result.V[programIndex] = select(A.V[programIndex] >= B.V[programIndex], floatbits(0xFFFFFFFF), floatbits(0));
+	Result.V[programIndex] = select(A.V[programIndex] >= B.V[programIndex], REAL_BITS(NAN), REAL_BITS(0));
 
 	return Result;
 }
 
-unmasked inline uniform WideFVector4 VectorReciprocalSqrtAccurate(const uniform WideFVector4& V)
+unmasked inline uniform WideFVector4f VectorReciprocalSqrtAccurate(const uniform WideFVector4f& V)
 {
-	uniform WideFVector4 Result;
+	uniform WideFVector4f Result;
 
 	Result.V[programIndex] = rsqrt(V.V[programIndex]);
 
 	return Result;
 }
 
+unmasked inline uniform WideFVector4d VectorReciprocalSqrtAccurate(const uniform WideFVector4d& V)
+{
+	uniform WideFVector4d Result;
+
+	Result.V[programIndex] = 1.0d / sqrt(V.V[programIndex]);
+
+	return Result;
+}
+
 unmasked inline uniform WideFVector4 VectorDot4( const uniform WideFVector4& Vec1, const uniform WideFVector4& Vec2 )
 {
 	uniform WideFVector4 Temp1, Temp2;
@@ -124,36 +175,33 @@ unmasked inline uniform WideFVector4 VectorDot4( const uniform WideFVector4& Vec
 
 unmasked inline uniform WideFVector4 VectorBitwiseAnd(const uniform WideFVector4 &Vec1, const uniform WideFVector4 &Vec2)
 {
-	varying unsigned int S0, S1;
-	varying float Result;
+	varying FReal S0, S1, Result;
 	*((uniform WideFVector4 *uniform)&S0) = *((uniform WideFVector4 *uniform)&Vec1);
 	*((uniform WideFVector4 *uniform)&S1) = *((uniform WideFVector4 *uniform)&Vec2);
 
-	Result = floatbits(S0 & S1);
+	Result = REAL_BITS(intbits(S0) & intbits(S1));
 
 	return *((uniform WideFVector4 *uniform)&Result);
 }
 
 unmasked inline uniform WideFVector4 VectorBitwiseOr(const uniform WideFVector4 &Vec1, const uniform WideFVector4 &Vec2)
 {
-	varying unsigned int S0, S1;
-	varying float Result;
+	varying FReal S0, S1, Result;
 	*((uniform WideFVector4 *uniform)&S0) = *((uniform WideFVector4 *uniform)&Vec1);
 	*((uniform WideFVector4 *uniform)&S1) = *((uniform WideFVector4 *uniform)&Vec2);
 
-	Result = floatbits(S0 | S1);
+	Result = REAL_BITS(intbits(S0) | intbits(S1));
 
 	return *((uniform WideFVector4 *uniform)&Result);
 }
 
 unmasked inline uniform WideFVector4 VectorBitwiseXor(const uniform WideFVector4 &Vec1, const uniform WideFVector4 &Vec2)
 {
-	varying unsigned int S0, S1;
-	varying float Result;
+	varying FReal S0, S1, Result;
 	*((uniform WideFVector4 *uniform)&S0) = *((uniform WideFVector4 *uniform)&Vec1);
 	*((uniform WideFVector4 *uniform)&S1) = *((uniform WideFVector4 *uniform)&Vec2);
 
-	Result = floatbits(S0 ^ S1);
+	Result = REAL_BITS(intbits(S0) ^ intbits(S1));
 
 	return *((uniform WideFVector4 *uniform)&Result);
 }
diff --git a/Engine/Source/Runtime/Engine/Classes/Components/SkeletalMeshComponent.h b/Engine/Source/Runtime/Engine/Classes/Components/SkeletalMeshComponent.h
--- a/Engine/Source/Runtime/Engine/Classes/Components/SkeletalMeshComponent.h
+++ b/Engine/Source/Runtime/Engine/Classes/Components/SkeletalMeshComponent.h
@@ -2576,3 +2576,12 @@ struct FLinkedInstancesAdapter
 		}
 	}
 };
+
+// Support ISPC enable/disable in non-shipping builds
+#if !INTEL_ISPC
+const bool bAnim_SkeletalMesh_ISPC_Enabled = false;
+#elif UE_BUILD_SHIPPING
+const bool bAnim_SkeletalMesh_ISPC_Enabled = true;
+#else
+extern bool bAnim_SkeletalMesh_ISPC_Enabled;
+#endif
diff --git a/Engine/Source/Runtime/Engine/Private/Animation/AnimEncoding_ConstantKeyLerp.cpp b/Engine/Source/Runtime/Engine/Private/Animation/AnimEncoding_ConstantKeyLerp.cpp
--- a/Engine/Source/Runtime/Engine/Private/Animation/AnimEncoding_ConstantKeyLerp.cpp
+++ b/Engine/Source/Runtime/Engine/Private/Animation/AnimEncoding_ConstantKeyLerp.cpp
@@ -14,6 +14,11 @@ static_assert(sizeof(ispc::FTransform) == sizeof(FTransform), "sizeof(ispc::FTra
 static_assert(sizeof(ispc::BoneTrackPair) == sizeof(BoneTrackPair), "sizeof(ispc::BoneTrackPair) != sizeof(BoneTrackPair)");
 #endif
 
+#if INTEL_ISPC && !UE_BUILD_SHIPPING
+bool bAnim_ConstantKeyLerp_ISPC_Enabled = true;
+FAutoConsoleVariableRef CVarAnimConstantKeyLerpISPCEnabled(TEXT("a.ConstantKeyLerp.ISPC"), bAnim_ConstantKeyLerp_ISPC_Enabled, TEXT("Whether to use ISPC optimizations in constant key anim encoding"));
+#endif
+
 /**
  * Handles the ByteSwap of compressed rotation data on import
  *
@@ -263,7 +268,7 @@ inline void AEFConstantKeyLerp<FORMAT>::GetPoseRotations(
 		return;
 	}
 
-	if (INTEL_ISPC)
+	if (bAnim_ConstantKeyLerp_ISPC_Enabled)
 	{
 #if INTEL_ISPC
 		const FUECompressedAnimData& AnimData = static_cast<const FUECompressedAnimData&>(DecompContext.CompressedAnimData);
@@ -315,7 +320,7 @@ inline void AEFConstantKeyLerp<FORMAT>::GetPoseTranslations(
 		return;
 	}
 
-	if (INTEL_ISPC)
+	if (bAnim_ConstantKeyLerp_ISPC_Enabled)
 	{
 #if INTEL_ISPC
 		const FUECompressedAnimData& AnimData = static_cast<const FUECompressedAnimData&>(DecompContext.CompressedAnimData);
@@ -383,7 +388,7 @@ inline void AEFConstantKeyLerp<FORMAT>::GetPoseScales(
 		return;
 	}
 
-	if (INTEL_ISPC)
+	if (bAnim_ConstantKeyLerp_ISPC_Enabled)
 	{
 #if INTEL_ISPC
 		const FUECompressedAnimData& AnimData = static_cast<const FUECompressedAnimData&>(DecompContext.CompressedAnimData);
diff --git a/Engine/Source/Runtime/Engine/Private/Animation/AnimEncoding_ConstantKeyLerp.ispc b/Engine/Source/Runtime/Engine/Private/Animation/AnimEncoding_ConstantKeyLerp.ispc
--- a/Engine/Source/Runtime/Engine/Private/Animation/AnimEncoding_ConstantKeyLerp.ispc
+++ b/Engine/Source/Runtime/Engine/Private/Animation/AnimEncoding_ConstantKeyLerp.ispc
@@ -36,7 +36,7 @@ export void GetConstantKeyLerpPoseRotations(
 		int NumRotKeys = TrackData[3];
 		const unsigned int8* RotStream = CompressedByteStream + RotKeysOffset;
 
-		FVector4 R0;
+		FVector4f R0;
 
 		if (NumRotKeys == 1)
 		{
@@ -57,17 +57,17 @@ export void GetConstantKeyLerpPoseRotations(
 			{
 				// unpack and lerp between the two nearest keys
 				const unsigned int8* KeyData1= RotStream + RotationStreamOffset +(Index1*CompressedRotationStrideNum);
-				FVector4 R1;
+				FVector4f R1;
 
 				DecompressRotation(R1, RotStream, KeyData1, FORMAT);
 
 				// Fast linear quaternion interpolation.
-				const FVector4 BlendedQuat = QuatFastLerp(R0, R1, Alpha);
+				const FVector4f BlendedQuat = QuatFastLerp(R0, R1, Alpha);
 				R0 = VectorNormalizeQuaternion(BlendedQuat);
 			}
 		}
 
-		uniform FVector4 AOSScratch[programCount];
+		uniform FVector4f AOSScratch[programCount];
 		soa_to_aos4(R0.V[0], R0.V[1], R0.V[2], R0.V[3], (uniform float* uniform)&AOSScratch);
 		foreach_active(i)
 		{
@@ -75,7 +75,7 @@ export void GetConstantKeyLerpPoseRotations(
 			{
 				uniform int uAtomIndex = extract(AtomIndex, i);
 				uniform FVector4 *uniform BoneAtomRotation = (uniform FVector4 *uniform)&Atoms[uAtomIndex].Rotation;
-				*BoneAtomRotation = AOSScratch[i];
+				*BoneAtomRotation = ConvertVector4fTo4Native(AOSScratch[i]);
 			}
 		}
 	}
@@ -117,27 +117,27 @@ export void GetConstantKeyLerpPoseTranslations(
 		const int32 TransStreamOffset = ((FORMAT == ACF_IntervalFixed32NoW) && NumTransKeys > 1) ? (sizeof(uniform float)*6) : 0; // offset past Min and Range data
 
 		const unsigned int8* KeyData0 = TransStream + TransStreamOffset + Index0*CompressedTranslationStrideNum;
-		FVector P0;
+		FVector3f P0;
 		DecompressTranslation(P0, TransStream, KeyData0, FORMAT);
 
 		if (Index0 != Index1)
 		{
 			const unsigned int8* KeyData1 = TransStream + TransStreamOffset + Index1*CompressedTranslationStrideNum;
-			FVector P1;
+			FVector3f P1;
 
 			DecompressTranslation(P1, TransStream, KeyData1, FORMAT);
 			P0 = VectorLerp( P0, P1, Alpha );
 		}
 
-		uniform FVector4 AOSScratch[programCount];
-		soa_to_aos4(P0.V[0], P0.V[1], P0.V[2], Zero, (uniform float* uniform)&AOSScratch);
+		uniform FVector4f AOSScratch[programCount];
+		soa_to_aos4(P0.V[0], P0.V[1], P0.V[2], 0.0f, (uniform float* uniform)&AOSScratch);
 		foreach_active(i)
 		{
 			unmasked
 			{
 				uniform int uAtomIndex = extract(AtomIndex, i);
 				uniform FVector4 *uniform BoneAtomTranslation = (uniform FVector4 *uniform)&Atoms[uAtomIndex].Translation;
-				*BoneAtomTranslation = AOSScratch[i];
+				*BoneAtomTranslation = ConvertVector4fTo4Native(AOSScratch[i]);
 			}
 		}
 	}
@@ -180,27 +180,27 @@ export void GetConstantKeyLerpPoseScales(
 		const int32 ScaleStreamOffset = ((FORMAT == ACF_IntervalFixed32NoW) && NumScaleKeys > 1) ? (sizeof(uniform float)*6) : 0; // offset past Min and Range data
 
 		const unsigned int8* KeyData0 = ScaleStream + ScaleStreamOffset + Index0*CompressedScaleStrideNum;
-		FVector P0;
+		FVector3f P0;
 		DecompressScale(P0, ScaleStream, KeyData0, FORMAT);
 
 		if (Index0 != Index1)
 		{
 			const unsigned int8* KeyData1 = ScaleStream + ScaleStreamOffset + Index1*CompressedScaleStrideNum;
-			FVector P1;
+			FVector3f P1;
 
 			DecompressScale(P1, ScaleStream, KeyData1, FORMAT);
 			P0 = VectorLerp( P0, P1, Alpha );
 		}
 
-		uniform FVector4 AOSScratch[programCount];
-		soa_to_aos4(P0.V[0], P0.V[1], P0.V[2], Zero, (uniform float* uniform)&AOSScratch);
+		uniform FVector4f AOSScratch[programCount];
+		soa_to_aos4(P0.V[0], P0.V[1], P0.V[2], 0.0f, (uniform float* uniform)&AOSScratch);
 		foreach_active(i)
 		{
 			unmasked
 			{
 				uniform int uAtomIndex = extract(AtomIndex, i);
 				uniform FVector4 *uniform BoneAtomScale = (uniform FVector4 *uniform)&Atoms[uAtomIndex].Scale3D;
-				*BoneAtomScale = AOSScratch[i];
+				*BoneAtomScale = ConvertVector4fTo4Native(AOSScratch[i]);
 			}
 		}
 	}
diff --git a/Engine/Source/Runtime/Engine/Private/Animation/AnimEncoding_PerTrackCompression.cpp b/Engine/Source/Runtime/Engine/Private/Animation/AnimEncoding_PerTrackCompression.cpp
--- a/Engine/Source/Runtime/Engine/Private/Animation/AnimEncoding_PerTrackCompression.cpp
+++ b/Engine/Source/Runtime/Engine/Private/Animation/AnimEncoding_PerTrackCompression.cpp
@@ -16,6 +16,11 @@ static_assert(sizeof(ispc::FTransform) == sizeof(FTransform), "sizeof(ispc::FTra
 static_assert(sizeof(ispc::BoneTrackPair) == sizeof(BoneTrackPair), "sizeof(ispc::BoneTrackPair) != sizeof(BoneTrackPair)");
 #endif
 
+#if INTEL_ISPC && !UE_BUILD_SHIPPING
+bool bAnim_PerTrackCompression_ISPC_Enabled = true;
+FAutoConsoleVariableRef CVarAnimPerTrackCompressionISPCEnabled(TEXT("a.PerTrackCompression.ISPC"), bAnim_PerTrackCompression_ISPC_Enabled, TEXT("Whether to use ISPC optimizations in per track anim encoding"));
+#endif
+
 // This define controls whether scalar or vector code is used to decompress keys.  Note that not all key decompression code
 // is vectorized yet, so some (seldom used) formats will actually get slower (due to extra LHS stalls) when enabled.
 // The code also relies on a flexible permute instruction being available (e.g., PPC vperm)
@@ -735,7 +740,7 @@ void AEFPerTrackCompressionCodec::GetPoseRotations(
 		return;
 	}
 
-	if (INTEL_ISPC)
+	if (bAnim_PerTrackCompression_ISPC_Enabled)
 	{
 #if INTEL_ISPC
 		const FUECompressedAnimData& AnimData = static_cast<const FUECompressedAnimData&>(DecompContext.CompressedAnimData);
@@ -783,7 +788,7 @@ void AEFPerTrackCompressionCodec::GetPoseTranslations(
 	{
 		return;
 	}
-	if (INTEL_ISPC)
+	if (bAnim_PerTrackCompression_ISPC_Enabled)
 	{
 #if INTEL_ISPC
 		const FUECompressedAnimData& AnimData = static_cast<const FUECompressedAnimData&>(DecompContext.CompressedAnimData);
@@ -833,7 +838,7 @@ void AEFPerTrackCompressionCodec::GetPoseScales(
 		return;
 	}
 
-	if (INTEL_ISPC)
+	if (bAnim_PerTrackCompression_ISPC_Enabled)
 	{
 #if INTEL_ISPC
 		const FUECompressedAnimData& AnimData = static_cast<const FUECompressedAnimData&>(DecompContext.CompressedAnimData);
diff --git a/Engine/Source/Runtime/Engine/Private/Animation/AnimEncoding_PerTrackCompression.ispc b/Engine/Source/Runtime/Engine/Private/Animation/AnimEncoding_PerTrackCompression.ispc
--- a/Engine/Source/Runtime/Engine/Private/Animation/AnimEncoding_PerTrackCompression.ispc
+++ b/Engine/Source/Runtime/Engine/Private/Animation/AnimEncoding_PerTrackCompression.ispc
@@ -113,7 +113,7 @@ export void GetPerTrackCompressionPoseRotations(
 		#pragma ignore warning(perf)
 		const int RotKeysOffset = TrackOffsetData[1];
 
-		FVector4 R0;
+		FVector4f R0;
 
 		if (RotKeysOffset != INDEX_NONE)
 		{
@@ -159,7 +159,7 @@ export void GetPerTrackCompressionPoseRotations(
 			{
 				const unsigned int8* KeyData1 = TrackData + FixedBytes + (Index1 * BytesPerKey);
 
-				FVector4 R1;
+				FVector4f R1;
 				DecompressRotation(KeyFormat, FormatFlags, R1, TrackData, KeyData1);
 
 				// Fast linear quaternion interpolation.
@@ -174,7 +174,7 @@ export void GetPerTrackCompressionPoseRotations(
 			R0 = Float0001;
 		}
 
-		uniform FVector4 AOSScratch[programCount];
+		uniform FVector4f AOSScratch[programCount];
 		soa_to_aos4(R0.V[0], R0.V[1], R0.V[2], R0.V[3], (uniform float* uniform)&AOSScratch);
 		foreach_active(i)
 		{
@@ -182,7 +182,7 @@ export void GetPerTrackCompressionPoseRotations(
 			{
 				uniform int uAtomIndex = extract(AtomIndex, i);
 				uniform FVector4 *uniform BoneAtomRotation = (uniform FVector4 *uniform)&Atoms[uAtomIndex].Rotation;
-				*BoneAtomRotation = AOSScratch[i];
+				*BoneAtomRotation = ConvertVector4fTo4Native(AOSScratch[i]);
 			}
 		}
 	}
@@ -211,7 +211,7 @@ export void GetPerTrackCompressionPoseTranslations(
 		#pragma ignore warning(perf)
 		const int PosKeysOffset = TrackOffsetData[0];
 
-		FVector R0;
+		FVector3f R0;
 
 		if (PosKeysOffset != INDEX_NONE)
 		{
@@ -257,7 +257,7 @@ export void GetPerTrackCompressionPoseTranslations(
 			{
 				const unsigned int8* KeyData1 = TrackData + FixedBytes + (Index1 * BytesPerKey);
 
-				FVector R1;
+				FVector3f R1;
 				DecompressTranslation(KeyFormat, FormatFlags, R1, TrackData, KeyData1);
 
 				R0 = VectorLerp(R0, R1, Alpha);
@@ -266,18 +266,18 @@ export void GetPerTrackCompressionPoseTranslations(
 		else
 		{
 			// Identity track
-			R0 = ZeroVector;
+			R0 = FloatZeroVector;
 		}
 
-		uniform FVector4 AOSScratch[programCount];
-		soa_to_aos4(R0.V[0], R0.V[1], R0.V[2], Zero, (uniform float* uniform)&AOSScratch);
+		uniform FVector4f AOSScratch[programCount];
+		soa_to_aos4(R0.V[0], R0.V[1], R0.V[2], 0.0f, (uniform float* uniform)&AOSScratch);
 		foreach_active(i)
 		{
 			unmasked
 			{
 				uniform int uAtomIndex = extract(AtomIndex, i);
 				uniform FVector4 *uniform BoneAtomTranslation = (uniform FVector4 *uniform)&Atoms[uAtomIndex].Translation;
-				*BoneAtomTranslation = AOSScratch[i];
+				*BoneAtomTranslation = ConvertVector4fTo4Native(AOSScratch[i]);
 			}
 		}
 	}
@@ -307,7 +307,7 @@ export void GetPerTrackCompressionPoseScales(
 		#pragma ignore warning(perf)
 		const int32 ScaleKeysOffset = ScaleData[0];
 
-		FVector R0;
+		FVector3f R0;
 
 		if (ScaleKeysOffset != INDEX_NONE)
 		{
@@ -353,7 +353,7 @@ export void GetPerTrackCompressionPoseScales(
 			{
 				const unsigned int8* KeyData1 = TrackData + FixedBytes + (Index1 * BytesPerKey);
 
-				FVector R1;
+				FVector3f R1;
 				DecompressScale(KeyFormat, FormatFlags, R1, TrackData, KeyData1);
 
 				R0 = VectorLerp(R0, R1, Alpha);
@@ -362,18 +362,18 @@ export void GetPerTrackCompressionPoseScales(
 		else
 		{
 			// Identity track
-			R0 = ZeroVector;
+			R0 = FloatZeroVector;
 		}
 
-		uniform FVector4 AOSScratch[programCount];
-		soa_to_aos4(R0.V[0], R0.V[1], R0.V[2], Zero, (uniform float* uniform)&AOSScratch);
+		uniform FVector4f AOSScratch[programCount];
+		soa_to_aos4(R0.V[0], R0.V[1], R0.V[2], 0.0f, (uniform float* uniform)&AOSScratch);
 		foreach_active(i)
 		{
 			unmasked
 			{
 				uniform int uAtomIndex = extract(AtomIndex, i);
 				uniform FVector4 *uniform BoneAtomScale = (uniform FVector4 *uniform)&Atoms[uAtomIndex].Scale3D;
-				*BoneAtomScale = AOSScratch[i];
+				*BoneAtomScale = ConvertVector4fTo4Native(AOSScratch[i]);
 			}
 		}
 	}
diff --git a/Engine/Source/Runtime/Engine/Private/Animation/AnimEncoding_VariableKeyLerp.cpp b/Engine/Source/Runtime/Engine/Private/Animation/AnimEncoding_VariableKeyLerp.cpp
--- a/Engine/Source/Runtime/Engine/Private/Animation/AnimEncoding_VariableKeyLerp.cpp
+++ b/Engine/Source/Runtime/Engine/Private/Animation/AnimEncoding_VariableKeyLerp.cpp
@@ -14,6 +14,11 @@ static_assert(sizeof(ispc::FTransform) == sizeof(FTransform), "sizeof(ispc::FTra
 static_assert(sizeof(ispc::BoneTrackPair) == sizeof(BoneTrackPair), "sizeof(ispc::BoneTrackPair) != sizeof(BoneTrackPair)");
 #endif
 
+#if INTEL_ISPC && !UE_BUILD_SHIPPING
+bool bAnim_VariableKeyLerp_ISPC_Enabled = true;
+FAutoConsoleVariableRef CVarAnimVariableKeyLerpISPCEnabled(TEXT("a.VariableKeyLerp.ISPC"), bAnim_VariableKeyLerp_ISPC_Enabled, TEXT("Whether to use ISPC optimizations in variable key anim encoding"));
+#endif
+
 /**
  * Handles the ByteSwap of compressed rotation data on import
  *
@@ -228,7 +233,7 @@ void AEFVariableKeyLerp<FORMAT>::GetPoseRotations(
 		return;
 	}
 
-	if (INTEL_ISPC)
+	if (bAnim_VariableKeyLerp_ISPC_Enabled)
 	{
 #if INTEL_ISPC
 		const FUECompressedAnimData& AnimData = static_cast<const FUECompressedAnimData&>(DecompContext.CompressedAnimData);
@@ -280,7 +285,7 @@ void AEFVariableKeyLerp<FORMAT>::GetPoseTranslations(
 		return;
 	}
 
-	if (INTEL_ISPC)
+	if (bAnim_VariableKeyLerp_ISPC_Enabled)
 	{
 #if INTEL_ISPC
 		const FUECompressedAnimData& AnimData = static_cast<const FUECompressedAnimData&>(DecompContext.CompressedAnimData);
@@ -332,7 +337,7 @@ void AEFVariableKeyLerp<FORMAT>::GetPoseScales(
 		return;
 	}
 
-	if (INTEL_ISPC)
+	if (bAnim_VariableKeyLerp_ISPC_Enabled)
 	{
 #if INTEL_ISPC
 		const FUECompressedAnimData& AnimData = static_cast<const FUECompressedAnimData&>(DecompContext.CompressedAnimData);
diff --git a/Engine/Source/Runtime/Engine/Private/Animation/AnimEncoding_VariableKeyLerp.ispc b/Engine/Source/Runtime/Engine/Private/Animation/AnimEncoding_VariableKeyLerp.ispc
--- a/Engine/Source/Runtime/Engine/Private/Animation/AnimEncoding_VariableKeyLerp.ispc
+++ b/Engine/Source/Runtime/Engine/Private/Animation/AnimEncoding_VariableKeyLerp.ispc
@@ -36,7 +36,7 @@ export void GetVariableKeyLerpPoseRotations(
 		int NumRotKeys = TrackData[3];
 		const unsigned int8* RotStream = CompressedByteStream + RotKeysOffset;
 
-		FVector4 R0;
+		FVector4f R0;
 
 		if (NumRotKeys == 1)
 		{
@@ -60,17 +60,17 @@ export void GetVariableKeyLerpPoseRotations(
 			{
 				// unpack and lerp between the two nearest keys
 				const unsigned int8* KeyData1= RotStream + RotationStreamOffset +(Index1*CompressedRotationStrideNum);
-				FVector4 R1;
+				FVector4f R1;
 
 				DecompressRotation(R1, RotStream, KeyData1, FORMAT);
 
 				// Fast linear quaternion interpolation.
-				const FVector4 BlendedQuat = QuatFastLerp(R0, R1, Alpha);
+				const FVector4f BlendedQuat = QuatFastLerp(R0, R1, Alpha);
 				R0 = VectorNormalizeQuaternion(BlendedQuat);
 			}
 		}
 
-		uniform FVector4 AOSScratch[programCount];
+		uniform FVector4f AOSScratch[programCount];
 		soa_to_aos4(R0.V[0], R0.V[1], R0.V[2], R0.V[3], (uniform float* uniform)&AOSScratch);
 		foreach_active(i)
 		{
@@ -78,7 +78,7 @@ export void GetVariableKeyLerpPoseRotations(
 			{
 				uniform int uAtomIndex = extract(AtomIndex, i);
 				uniform FVector4 *uniform BoneAtomRotation = (uniform FVector4 *uniform)&Atoms[uAtomIndex].Rotation;
-				*BoneAtomRotation = AOSScratch[i];
+				*BoneAtomRotation = ConvertVector4fTo4Native(AOSScratch[i]);
 			}
 		}
 	}
@@ -122,27 +122,27 @@ export void GetVariableKeyLerpPoseTranslations(
 		const int32 TransStreamOffset = ((FORMAT == ACF_IntervalFixed32NoW) && NumTransKeys > 1) ? (sizeof(uniform float)*6) : 0; // offset past Min and Range data
 
 		const unsigned int8* KeyData0 = TransStream + TransStreamOffset + Index0*CompressedTranslationStrideNum;
-		FVector P0;
+		FVector3f P0;
 		DecompressTranslation(P0, TransStream, KeyData0, FORMAT);
 
 		if (Index0 != Index1)
 		{
 			const unsigned int8* KeyData1 = TransStream + TransStreamOffset + Index1*CompressedTranslationStrideNum;
-			FVector P1;
+			FVector3f P1;
 
 			DecompressTranslation(P1, TransStream, KeyData1, FORMAT);
 			P0 = VectorLerp( P0, P1, Alpha );
 		}
 
-		uniform FVector4 AOSScratch[programCount];
-		soa_to_aos4(P0.V[0], P0.V[1], P0.V[2], Zero, (uniform float* uniform)&AOSScratch);
+		uniform FVector4f AOSScratch[programCount];
+		soa_to_aos4(P0.V[0], P0.V[1], P0.V[2], 0.0f, (uniform float* uniform)&AOSScratch);
 		foreach_active(i)
 		{
 			unmasked
 			{
 				uniform int uAtomIndex = extract(AtomIndex, i);
 				uniform FVector4 *uniform BoneAtomTranslation = (uniform FVector4 *uniform)&Atoms[uAtomIndex].Translation;
-				*BoneAtomTranslation = AOSScratch[i];
+				*BoneAtomTranslation = ConvertVector4fTo4Native(AOSScratch[i]);
 			}
 		}
 	}
@@ -187,27 +187,27 @@ export void GetVariableKeyLerpPoseScales(
 		const int32 ScaleStreamOffset = ((FORMAT == ACF_IntervalFixed32NoW) && NumScaleKeys > 1) ? (sizeof(uniform float)*6) : 0; // offset past Min and Range data
 
 		const unsigned int8* KeyData0 = ScaleStream + ScaleStreamOffset + Index0*CompressedScaleStrideNum;
-		FVector P0;
+		FVector3f P0;
 		DecompressScale(P0, ScaleStream, KeyData0, FORMAT);
 
 		if (Index0 != Index1)
 		{
 			const unsigned int8* KeyData1 = ScaleStream + ScaleStreamOffset + Index1*CompressedScaleStrideNum;
-			FVector P1;
+			FVector3f P1;
 
 			DecompressScale(P1, ScaleStream, KeyData1, FORMAT);
 			P0 = VectorLerp( P0, P1, Alpha );
 		}
 
-		uniform FVector4 AOSScratch[programCount];
-		soa_to_aos4(P0.V[0], P0.V[1], P0.V[2], Zero, (uniform float* uniform)&AOSScratch);
+		uniform FVector4f AOSScratch[programCount];
+		soa_to_aos4(P0.V[0], P0.V[1], P0.V[2], 0.0f, (uniform float* uniform)&AOSScratch);
 		foreach_active(i)
 		{
 			unmasked
 			{
 				uniform int uAtomIndex = extract(AtomIndex, i);
 				uniform FVector4 *uniform BoneAtomScale = (uniform FVector4 *uniform)&Atoms[uAtomIndex].Scale3D;
-				*BoneAtomScale = AOSScratch[i];
+				*BoneAtomScale = ConvertVector4fTo4Native(AOSScratch[i]);
 			}
 		}
 	}
diff --git a/Engine/Source/Runtime/Engine/Private/Animation/AnimationRuntime.cpp b/Engine/Source/Runtime/Engine/Private/Animation/AnimationRuntime.cpp
--- a/Engine/Source/Runtime/Engine/Private/Animation/AnimationRuntime.cpp
+++ b/Engine/Source/Runtime/Engine/Private/Animation/AnimationRuntime.cpp
@@ -32,10 +32,15 @@ DECLARE_CYCLE_STAT(TEXT("BlendPosesPerBoneFilter"), STAT_BlendPosesPerBoneFilter
 #if INTEL_ISPC
 static_assert(sizeof(ispc::FTransform) == sizeof(FTransform), "sizeof(ispc::FTransform) != sizeof(FTransform)");
 static_assert(sizeof(ispc::FVector) == sizeof(FVector), "sizeof(ispc::FVector) != sizeof(FVector)");
-static_assert(sizeof(ispc::FVector4) == sizeof(FVector4), "sizeof(ispc::FVector4) != sizeof(FVector4)");
+static_assert(sizeof(ispc::FVector4) == sizeof(FQuat), "sizeof(ispc::FVector4) != sizeof(FQuat)");
 static_assert(sizeof(ispc::FPerBoneBlendWeight) == sizeof(FPerBoneBlendWeight), "sizeof(ispc::FPerBoneBlendWeight) != sizeof(FPerBoneBlendWeight)");
 #endif
 
+#if INTEL_ISPC && !UE_BUILD_SHIPPING
+bool bAnim_Runtime_ISPC_Enabled = true;
+FAutoConsoleVariableRef CVarAnimRuntimeISPCEnabled(TEXT("a.Runtime.ISPC"), bAnim_Runtime_ISPC_Enabled, TEXT("Whether to use ISPC optimizations in animation runtime"));
+#endif
+
 void FAnimationRuntime::NormalizeRotations(const FBoneContainer& RequiredBones, /*inout*/ FTransformArrayA2& Atoms)
 {
 	check( Atoms.Num() == RequiredBones.GetNumBones() );
@@ -88,7 +93,7 @@ FORCEINLINE void BlendPose(const FCompactPose& SourcePose, FCompactPose& ResultP
 template <>
 FORCEINLINE void BlendPose<ETransformBlendMode::Overwrite>(const FCompactPose& SourcePose, FCompactPose& ResultPose, const float BlendWeight)
 {
-	if (INTEL_ISPC)
+	if (bAnim_Runtime_ISPC_Enabled)
 	{
 #if INTEL_ISPC
 		ispc::BlendTransformOverwrite(
@@ -110,7 +115,7 @@ FORCEINLINE void BlendPose<ETransformBlendMode::Overwrite>(const FCompactPose& S
 template <>
 FORCEINLINE void BlendPose<ETransformBlendMode::Accumulate>(const FCompactPose& SourcePose, FCompactPose& ResultPose, const float BlendWeight)
 {
-	if (INTEL_ISPC)
+	if (bAnim_Runtime_ISPC_Enabled)
 	{
 #if INTEL_ISPC
 		ispc::BlendTransformAccumulate(
@@ -159,7 +164,7 @@ FORCEINLINE void BlendCurves(const TArrayView<const FBlendedCurve> SourceCurves,
 {
 	if (SourceCurves.Num() > 0)
 	{
-		if (INTEL_ISPC)
+		if (bAnim_Runtime_ISPC_Enabled)
 		{
 #if INTEL_ISPC
 			OutCurve.InitFrom(SourceCurves[0]);
@@ -842,7 +847,7 @@ void FAnimationRuntime::LerpBoneTransforms(TArray<FTransform>& A, const TArray<F
 		return;
 	}
 
-	if (INTEL_ISPC)
+	if (bAnim_Runtime_ISPC_Enabled)
 	{
 #if INTEL_ISPC
 		ispc::LerpBoneTransforms(
@@ -954,7 +959,7 @@ void FAnimationRuntime::ConvertPoseToMeshRotation(FCompactPose& LocalPose)
 
 	// Convert all rotations to mesh space
 	// only the root bone doesn't have a parent. So skip it to save a branch in the iteration.
-	if (INTEL_ISPC)
+	if (bAnim_Runtime_ISPC_Enabled)
 	{
 #if INTEL_ISPC
 		ispc::ConvertPoseToMeshRotation(
@@ -981,7 +986,7 @@ void FAnimationRuntime::ConvertMeshRotationPoseToLocalSpace(FCompactPose& Pose)
 
 	// Convert all rotations to mesh space
 	// only the root bone doesn't have a parent. So skip it to save a branch in the iteration.
-	if (INTEL_ISPC)
+	if (bAnim_Runtime_ISPC_Enabled)
 	{
 #if INTEL_ISPC
 		ispc::ConvertMeshRotationPoseToLocalSpace(
@@ -1067,7 +1072,7 @@ void FAnimationRuntime::AccumulateLocalSpaceAdditivePoseInternal(FCompactPose& B
 		const ScalarRegister VBlendWeight(Weight);
 		if (FAnimWeight::IsFullWeight(Weight))
 		{
-			if (INTEL_ISPC)
+			if (bAnim_Runtime_ISPC_Enabled)
 			{
 #if INTEL_ISPC
 				ispc::AccumulateWithAdditiveScale(
@@ -1754,7 +1759,7 @@ FAnimationPoseData& OutAnimationPoseData, TArray<FPerBoneBlendWeight>& BoneBlend
 	// blend poses with both mesh space rotation and scaling (we assume uniform scale)
 	if (bMeshSpaceRotationBlend && bMeshSpaceScaleBlend)
 	{
-		if (INTEL_ISPC)
+		if (bAnim_Runtime_ISPC_Enabled)
 		{
 #if INTEL_ISPC
 			ispc::BlendPosesPerBoneFilterScaleRotation(
diff --git a/Engine/Source/Runtime/Engine/Private/Animation/AnimationRuntime.ispc b/Engine/Source/Runtime/Engine/Private/Animation/AnimationRuntime.ispc
--- a/Engine/Source/Runtime/Engine/Private/Animation/AnimationRuntime.ispc
+++ b/Engine/Source/Runtime/Engine/Private/Animation/AnimationRuntime.ispc
@@ -190,8 +190,8 @@ export void BlendPosesPerBoneFilterScaleRotation(uniform FTransform OutPose[],
 		}
 		else
 		{
-			SourceRotation = VectorQuaternionMultiply2(Float0001, BaseAtom.Rotation);
-			TargetRotation = VectorQuaternionMultiply2(Float0001, TargetAtom.Rotation);
+			SourceRotation = VectorQuaternionMultiply2(Vector0001, BaseAtom.Rotation);
+			TargetRotation = VectorQuaternionMultiply2(Vector0001, TargetAtom.Rotation);
 			SourceScale = OneVector * SetVector(BaseAtom.Scale3D);
 			TargetScale = OneVector * SetVector(TargetAtom.Scale3D);
 		}
@@ -238,7 +238,7 @@ export void BlendPosesPerBoneFilterScaleRotation(uniform FTransform OutPose[],
 
 			const uniform FVector ParentScaleInv = TransformGetSafeScaleReciprocal(BlendScales[ParentIndex], SMALL_NUMBER);
 			const uniform FVector LocalBlendScale = ParentScaleInv * BlendScale;
-			BlendAtom.Scale3D = SetVector4(LocalBlendScale.V[0], LocalBlendScale.V[1], LocalBlendScale.V[2], 0.0f);
+			BlendAtom.Scale3D = SetVector4(LocalBlendScale.V[0], LocalBlendScale.V[1], LocalBlendScale.V[2], ZERO);
 		}
 
 		BlendRotations[BoneIndex] = BlendRotation;
diff --git a/Engine/Source/Runtime/Engine/Private/Animation/BonePose.cpp b/Engine/Source/Runtime/Engine/Private/Animation/BonePose.cpp
--- a/Engine/Source/Runtime/Engine/Private/Animation/BonePose.cpp
+++ b/Engine/Source/Runtime/Engine/Private/Animation/BonePose.cpp
@@ -10,10 +10,15 @@
 static_assert(sizeof(ispc::FTransform) == sizeof(FTransform), "sizeof(ispc::FTransform) != sizeof(FTransform)");
 #endif
 
+#if INTEL_ISPC && !UE_BUILD_SHIPPING
+bool bAnim_BonePose_ISPC_Enabled = true;
+FAutoConsoleVariableRef CVarAnimBonePoseISPCEnabled(TEXT("a.BonePose.ISPC"), bAnim_BonePose_ISPC_Enabled, TEXT("Whether to use ISPC optimizations in bone pose calculations"));
+#endif
+
 // Normalizes all rotations in this pose
 void FCompactPose::NormalizeRotations()
 {
-	if (INTEL_ISPC)
+	if (bAnim_BonePose_ISPC_Enabled)
 	{
 #if INTEL_ISPC
 		ispc::NormalizeRotations((ispc::FTransform*)this->Bones.GetData(), this->Bones.Num());
@@ -31,7 +36,7 @@ void FCompactPose::NormalizeRotations()
 // Sets every bone transform to Identity
 void FCompactPose::ResetToAdditiveIdentity()
 {
-	if (INTEL_ISPC)
+	if (bAnim_BonePose_ISPC_Enabled)
 	{
 #if INTEL_ISPC
 		ispc::ResetToAdditiveIdentity((ispc::FTransform*)this->Bones.GetData(), this->Bones.Num());
@@ -49,7 +54,7 @@ void FCompactPose::ResetToAdditiveIdentity()
 // Normalizes all rotations in this pose
 void FCompactHeapPose::NormalizeRotations()
 {
-	if (INTEL_ISPC)
+	if (bAnim_BonePose_ISPC_Enabled)
 	{
 #if INTEL_ISPC
 		ispc::NormalizeRotations((ispc::FTransform*)this->Bones.GetData(), this->Bones.Num());
@@ -67,7 +72,7 @@ void FCompactHeapPose::NormalizeRotations()
 // Sets every bone transform to Identity
 void FCompactHeapPose::ResetToAdditiveIdentity()
 {
-	if (INTEL_ISPC)
+	if (bAnim_BonePose_ISPC_Enabled)
 	{
 #if INTEL_ISPC
 		ispc::ResetToAdditiveIdentity((ispc::FTransform*)this->Bones.GetData(), this->Bones.Num());
diff --git a/Engine/Source/Runtime/Engine/Private/Animation/BonePose.ispc b/Engine/Source/Runtime/Engine/Private/Animation/BonePose.ispc
--- a/Engine/Source/Runtime/Engine/Private/Animation/BonePose.ispc
+++ b/Engine/Source/Runtime/Engine/Private/Animation/BonePose.ispc
@@ -45,8 +45,8 @@ export void ResetToAdditiveIdentity(uniform FTransform Bones[],
 {
 	for(uniform int BoneIdx = 0; BoneIdx < NumBones; BoneIdx++)
 	{
-		Bones[BoneIdx].Rotation = Float0001;
-		Bones[BoneIdx].Translation = FloatZero;
-		Bones[BoneIdx].Scale3D = FloatZero;
+		Bones[BoneIdx].Rotation = Vector0001;
+		Bones[BoneIdx].Translation = VectorZero;
+		Bones[BoneIdx].Scale3D = VectorZero;
 	}
 }
\ No newline at end of file
diff --git a/Engine/Source/Runtime/Engine/Private/Collision/KAggregateGeom.cpp b/Engine/Source/Runtime/Engine/Private/Collision/KAggregateGeom.cpp
--- a/Engine/Source/Runtime/Engine/Private/Collision/KAggregateGeom.cpp
+++ b/Engine/Source/Runtime/Engine/Private/Collision/KAggregateGeom.cpp
@@ -30,6 +30,14 @@
 #if INTEL_ISPC
 static_assert(sizeof(ispc::FBox) == sizeof(FBox), "sizeof(ispc::FBox) != sizeof(FBox)");
 static_assert(sizeof(ispc::FRotator) == sizeof(FRotator), "sizeof(ispc::FRotator) != sizeof(FRotator)");
+static_assert(sizeof(ispc::FVector) == sizeof(FVector), "sizeof(ispc::FVector) != sizeof(FVector)");
+#endif
+
+#if INTEL_ISPC && !UE_BUILD_SHIPPING
+bool bPhysics_AggregateGeom_ISPC_Enabled = true;
+FAutoConsoleVariableRef CVarPhysicsAggregateGeomISPCEnabled(TEXT("p.AggregateGeom.ISPC"), bPhysics_AggregateGeom_ISPC_Enabled, TEXT("Whether to use ISPC optimizations in physics aggregate geometry calculations"));
+#else
+const bool bPhysics_AggregateGeom_ISPC_Enabled = INTEL_ISPC;
 #endif
 
 ///////////////////////////////////////
@@ -316,7 +324,7 @@ EAggCollisionShape::Type FKBoxElem::StaticShapeType = EAggCollisionShape::Box;
 
 FBox FKBoxElem::CalcAABB(const FTransform& BoneTM, float Scale) const
 {
-	if (INTEL_ISPC)
+	if (bPhysics_AggregateGeom_ISPC_Enabled)
 	{
 #if INTEL_ISPC
 		FBox LocalBox(ForceInit);
@@ -355,7 +363,7 @@ EAggCollisionShape::Type FKSphylElem::StaticShapeType = EAggCollisionShape::Sphy
 
 FBox FKSphylElem::CalcAABB(const FTransform& BoneTM, float Scale) const
 {
-	if (INTEL_ISPC)
+	if (bPhysics_AggregateGeom_ISPC_Enabled)
 	{
 #if INTEL_ISPC
 		FBox Result(ForceInit);
diff --git a/Engine/Source/Runtime/Engine/Private/Collision/KAggregateGeom.ispc b/Engine/Source/Runtime/Engine/Private/Collision/KAggregateGeom.ispc
--- a/Engine/Source/Runtime/Engine/Private/Collision/KAggregateGeom.ispc
+++ b/Engine/Source/Runtime/Engine/Private/Collision/KAggregateGeom.ispc
@@ -10,9 +10,9 @@
 
 struct FRotator
 {
-	float Pitch;
-	float Yaw;
-	float Roll;
+	FReal Pitch;
+	FReal Yaw;
+	FReal Roll;
 };
 
 struct FBox
@@ -24,8 +24,8 @@ struct FBox
 
 static inline uniform FVector4 RotatorToQuaternion(const uniform FRotator &InRotation)
 {
-	const uniform FVector4 Angles = SetVector4(InRotation.Pitch, InRotation.Yaw, InRotation.Roll, 0.0f);
-	const uniform FVector4 AnglesNoWinding = VectorMod(Angles, Float360);
+	const uniform FVector4 Angles = SetVector4(InRotation.Pitch, InRotation.Yaw, InRotation.Roll, ZERO);
+	const uniform FVector4 AnglesNoWinding = VectorMod(Angles, Vector360);
 	const uniform FVector4 HalfAngles = VectorMultiply(AnglesNoWinding, DEG_TO_RAD_HALF);
 
 	uniform FVector4 SinAngles, CosAngles;
@@ -44,10 +44,8 @@ static inline uniform FVector4 RotatorToQuaternion(const uniform FRotator &InRot
 	const uniform FVector4 CP_CP_SP_SP = VectorShuffle(CosAngles, SinAngles, 0, 0, 0, 0);
 	const uniform FVector4 CY_SY_CY_SY = VectorShuffle(SY_SY_CY_CY_Temp, SY_SY_CY_CY_Temp, 2, 0, 2, 0);
 
-	static const uniform unsigned int Neg = (SIGN_BIT);
-	static const uniform unsigned int Pos = (0);
-	const uniform FVector4 SignBitsLeft  = SetVector4(floatbits(Pos), floatbits(Neg), floatbits(Pos), floatbits(Pos)); // {{ Pos, Neg, Pos, Pos }};
-	const uniform FVector4 SignBitsRight = SetVector4(floatbits(Neg), floatbits(Neg), floatbits(Neg), floatbits(Pos)); // {{ Neg, Neg, Neg, Pos }};
+	const uniform FVector4 SignBitsLeft  = SetVector4(ZERO, REAL_BITS(SIGN_BIT), ZERO, ZERO); // {{ Pos, Neg, Pos, Pos }};
+	const uniform FVector4 SignBitsRight = SetVector4(REAL_BITS(SIGN_BIT), REAL_BITS(SIGN_BIT), REAL_BITS(SIGN_BIT), ZERO); // {{ Neg, Neg, Neg, Pos }};
 	const uniform FVector4 LeftTerm  = VectorBitwiseXor(SignBitsLeft , VectorMultiply(CR, VectorMultiply(SP_SP_CP_CP, SY_CY_SY_CY)));
 	const uniform FVector4 RightTerm = VectorBitwiseXor(SignBitsRight, VectorMultiply(SR, VectorMultiply(CP_CP_SP_SP, CY_SY_CY_SY)));
 
@@ -59,9 +57,8 @@ static inline void BoxTransformByMatrix(uniform FVector4 &BoxMin, uniform FVecto
 	const uniform FVector4 VecMin = BoxMin;
 	const uniform FVector4 VecMax = BoxMax;
 
-	static const uniform FVector4 Half = { {0.5f, 0.5f, 0.5f, 0.0f} };
-	const uniform FVector4 Origin = VectorMultiply(VectorAdd(VecMax, VecMin), Half);
-	const uniform FVector4 Extent = VectorMultiply(VectorSubtract(VecMax, VecMin), Half);
+	const uniform FVector4 Origin = VectorMultiply(VectorAdd(VecMax, VecMin), VectorOneHalf);
+	const uniform FVector4 Extent = VectorMultiply(VectorSubtract(VecMax, VecMin), VectorOneHalf);
 
 	const uniform FVector4 m0 = *((uniform FVector4 *uniform)&M.M[0]);
 	const uniform FVector4 m1 = *((uniform FVector4 *uniform)&M.M[4]);
@@ -98,16 +95,16 @@ export void BoxCalcAABB(uniform FBox &LocalBox,
 {
 	uniform FTransform ElemTM;
 	ElemTM.Rotation = RotatorToQuaternion(Rotation);
-	ElemTM.Translation = SetVector4(Center.V[0], Center.V[1], Center.V[2], 0.0f);
-	ElemTM.Scale3D = Float1110;
+	ElemTM.Translation = SetVector4(Center.V[0], Center.V[1], Center.V[2], ZERO);
+	ElemTM.Scale3D = Vector1110;
 
 	ScaleTranslation(ElemTM, Scale);
 	TransformMultiply(ElemTM, ElemTM, BoneTM);
 
 	const uniform float HalfScale = 0.5f * Scale;
-	const uniform FVector4 XYZ_0 = SetVector4(X, Y, Z, 0.0f);
-	uniform FVector4 BoxMax = XYZ_0 * HalfScale;
-	uniform FVector4 BoxMin = BoxMax * -1.0f;
+	const uniform FVector4f XYZ_0 = SetVector4(X, Y, Z, 0.0f);
+	uniform FVector4 BoxMax = ConvertVector4fTo4Native(XYZ_0 * HalfScale);
+	uniform FVector4 BoxMin = BoxMax * -ONE;
 
 	BoxTransformByTransform(BoxMin, BoxMax, ElemTM);
 
@@ -126,8 +123,8 @@ export void SPhylCalcAABB(uniform FBox &Result,
 {
 	uniform FTransform ElemTM;
 	ElemTM.Rotation = RotatorToQuaternion(Rotation);
-	ElemTM.Translation = SetVector4(Center.V[0], Center.V[1], Center.V[2], 0.0f);
-	ElemTM.Scale3D = Float1110;
+	ElemTM.Translation = SetVector4(Center.V[0], Center.V[1], Center.V[2], ZERO);
+	ElemTM.Scale3D = Vector1110;
 
 	ScaleTranslation(ElemTM, Scale);
 	TransformMultiply(ElemTM, ElemTM, BoneTM);
diff --git a/Engine/Source/Runtime/Engine/Private/Components/SkeletalMeshComponent.cpp b/Engine/Source/Runtime/Engine/Private/Components/SkeletalMeshComponent.cpp
--- a/Engine/Source/Runtime/Engine/Private/Components/SkeletalMeshComponent.cpp
+++ b/Engine/Source/Runtime/Engine/Private/Components/SkeletalMeshComponent.cpp
@@ -64,6 +64,10 @@ static TAutoConsoleVariable<float> CVarStallParallelAnimation(
 	0.0f,
 	TEXT("Sleep for the given time in each parallel animation task. Time is given in ms. This is a debug option used for critical path analysis and forcing a change in the critical path."));
 
+#if INTEL_ISPC && !UE_BUILD_SHIPPING
+bool bAnim_SkeletalMesh_ISPC_Enabled = true;
+FAutoConsoleVariableRef CVarAnimSkeletalMeshISPCEnabled(TEXT("a.SkeletalMesh.ISPC"), bAnim_SkeletalMesh_ISPC_Enabled, TEXT("Whether to use ISPC optimizations in animation skeletal mesh components"));
+#endif
 
 static TAutoConsoleVariable<int32> CVarCacheLocalSpaceBounds(
 	TEXT("a.CacheLocalSpaceBounds"),
@@ -1571,7 +1575,7 @@ void USkeletalMeshComponent::FillComponentSpaceTransforms(const USkeletalMesh* I
 #endif
 	}
 
-	if (INTEL_ISPC)
+	if (bAnim_SkeletalMesh_ISPC_Enabled)
 	{
 #if INTEL_ISPC
 		ispc::FillComponentSpaceTransforms(
diff --git a/Engine/Source/Runtime/Engine/Private/GPUSkinVertexFactory.cpp b/Engine/Source/Runtime/Engine/Private/GPUSkinVertexFactory.cpp
--- a/Engine/Source/Runtime/Engine/Private/GPUSkinVertexFactory.cpp
+++ b/Engine/Source/Runtime/Engine/Private/GPUSkinVertexFactory.cpp
@@ -103,8 +103,17 @@ static TAutoConsoleVariable<int32> CVarVelocityTest(
 	ECVF_Cheat | ECVF_RenderThreadSafe);
 #endif // if !(UE_BUILD_SHIPPING || UE_BUILD_TEST)
 
+#if INTEL_ISPC && !UE_BUILD_SHIPPING
+bool bGPUSkin_CopyBones_ISPC_Enabled = true;
+FAutoConsoleVariableRef CVarGPUSkinCopyBonesISPCEnabled(TEXT("r.GPUSkin.CopyBones.ISPC"), bGPUSkin_CopyBones_ISPC_Enabled, TEXT("Whether to use ISPC optimizations when copying bones for GPU skinning"));
+#endif
+
+#if INTEL_ISPC && UE_LARGE_WORLD_COORDINATES_DISABLED
+namespace ispc { typedef FMatrix FMatrix44f; } // Need this because ISPC doesn't expose typedefs in header and optimized function uses mixed types
+#endif
+
 #if INTEL_ISPC
-static_assert(sizeof(ispc::FMatrix) == sizeof(FMatrix), "sizeof(ispc::FMatrix) != sizeof(FMatrix)");
+static_assert(sizeof(ispc::FMatrix44f) == sizeof(FMatrix44f), "sizeof(ispc::FMatrix44f) != sizeof(FMatrix44f)");
 static_assert(sizeof(ispc::FMatrix3x4) == sizeof(FMatrix3x4), "sizeof(ispc::FMatrix3x4) != sizeof(FMatrix3x4)");
 #endif
 
@@ -324,12 +333,12 @@ bool FGPUBaseSkinVertexFactory::FShaderDataType::UpdateBoneData(FRHICommandListI
 		//  sizeof(FMatrix) == 64
 		// PLATFORM_CACHE_LINE_SIZE (128) / 64 = 2
 
-		if (INTEL_ISPC)
+		if (bGPUSkin_CopyBones_ISPC_Enabled)
 		{
 #if INTEL_ISPC
 			ispc::UpdateBoneData_CopyBones(
 				(ispc::FMatrix3x4*)&ChunkMatrices[0],
-				(ispc::FMatrix*)&ReferenceToLocalMatrices[0],
+				(ispc::FMatrix44f*)&ReferenceToLocalMatrices[0],
 				BoneMap.GetData(),
 				NumBones);
 #endif
diff --git a/Engine/Source/Runtime/Engine/Private/GPUSkinVertexFactory.ispc b/Engine/Source/Runtime/Engine/Private/GPUSkinVertexFactory.ispc
--- a/Engine/Source/Runtime/Engine/Private/GPUSkinVertexFactory.ispc
+++ b/Engine/Source/Runtime/Engine/Private/GPUSkinVertexFactory.ispc
@@ -3,12 +3,12 @@
 #include "Math/Matrix.isph"
 
 export void UpdateBoneData_CopyBones(uniform FMatrix3x4 ChunkMatrices[],
-										const uniform FMatrix ReferenceToLocalMatrices[],
+										const uniform FMatrix44f ReferenceToLocalMatrices[],
 										const uniform unsigned int16 BoneMap[],
 										const uniform int NumBones)
 {
 	uniform FMatrix3x4 *uniform ChunkPtr = (uniform FMatrix3x4 * uniform) ChunkMatrices;
-	uniform FMatrix *uniform RefToLocalPtr = (uniform FMatrix * uniform) ReferenceToLocalMatrices;
+	uniform FMatrix44f *uniform RefToLocalPtr = (uniform FMatrix44f * uniform) ReferenceToLocalMatrices;
 
 	for(uniform int BoneIdx = 0; BoneIdx < NumBones; BoneIdx++)
 	{
@@ -52,7 +52,6 @@ export void UpdateBoneData_CopyBones(uniform FMatrix3x4 ChunkMatrices[],
 		}
 
 #elif TARGET_WIDTH == 16
-		// perf seems whacky - need to verify
 		static const varying int vPerm = {0, 4, 8, 12, 1, 5, 9, 13,
 										2, 6, 10, 14, 3, 7, 11, 15};
 
@@ -64,7 +63,7 @@ export void UpdateBoneData_CopyBones(uniform FMatrix3x4 ChunkMatrices[],
 		}
 
 #else
-	    #error "Unknown ISPC_TARGET_* value"
+		#error "Unknown ISPC_TARGET_* value"
 #endif
 	}
 }
diff --git a/Engine/Source/Runtime/Engine/Public/AnimEncoding_ConstantKeyLerp.h b/Engine/Source/Runtime/Engine/Public/AnimEncoding_ConstantKeyLerp.h
--- a/Engine/Source/Runtime/Engine/Public/AnimEncoding_ConstantKeyLerp.h
+++ b/Engine/Source/Runtime/Engine/Public/AnimEncoding_ConstantKeyLerp.h
@@ -321,4 +321,13 @@ FORCEINLINE void AEFConstantKeyLerp<FORMAT>::GetBoneAtomScale(FTransform& OutAto
 		DecompressScale<FORMAT>( P0, ScaleStream, KeyData);
 		OutAtom.SetScale3D(P0);
 	}
-}
\ No newline at end of file
+}
+
+// Support ISPC enable/disable in non-shipping builds
+#if !INTEL_ISPC
+const bool bAnim_ConstantKeyLerp_ISPC_Enabled = false;
+#elif UE_BUILD_SHIPPING
+const bool bAnim_ConstantKeyLerp_ISPC_Enabled = true;
+#else
+extern bool bAnim_ConstantKeyLerp_ISPC_Enabled;
+#endif
diff --git a/Engine/Source/Runtime/Engine/Public/AnimEncoding_PerTrackCompression.h b/Engine/Source/Runtime/Engine/Public/AnimEncoding_PerTrackCompression.h
--- a/Engine/Source/Runtime/Engine/Public/AnimEncoding_PerTrackCompression.h
+++ b/Engine/Source/Runtime/Engine/Public/AnimEncoding_PerTrackCompression.h
@@ -133,3 +133,12 @@ protected:
 
 	friend class UAnimCompress_PerTrackCompression;
 };
+
+// Support ISPC enable/disable in non-shipping builds
+#if !INTEL_ISPC
+const bool bAnim_PerTrackCompression_ISPC_Enabled = false;
+#elif UE_BUILD_SHIPPING
+const bool bAnim_PerTrackCompression_ISPC_Enabled = true;
+#else
+extern bool bAnim_PerTrackCompression_ISPC_Enabled;
+#endif
diff --git a/Engine/Source/Runtime/Engine/Public/AnimEncoding_VariableKeyLerp.h b/Engine/Source/Runtime/Engine/Public/AnimEncoding_VariableKeyLerp.h
--- a/Engine/Source/Runtime/Engine/Public/AnimEncoding_VariableKeyLerp.h
+++ b/Engine/Source/Runtime/Engine/Public/AnimEncoding_VariableKeyLerp.h
@@ -328,4 +328,13 @@ FORCEINLINE_DEBUGGABLE void AEFVariableKeyLerp<FORMAT>::GetBoneAtomScale(FTransf
 		DecompressScale<FORMAT>( P0, ScaleStream, KeyData);
 		OutAtom.SetScale3D( P0 );
 	}
-}
\ No newline at end of file
+}
+
+// Support ISPC enable/disable in non-shipping builds
+#if !INTEL_ISPC
+const bool bAnim_VariableKeyLerp_ISPC_Enabled = false;
+#elif UE_BUILD_SHIPPING
+const bool bAnim_VariableKeyLerp_ISPC_Enabled = true;
+#else
+extern bool bAnim_VariableKeyLerp_ISPC_Enabled;
+#endif
diff --git a/Engine/Source/Runtime/Engine/Public/Animation/AnimationCompression.isph b/Engine/Source/Runtime/Engine/Public/Animation/AnimationCompression.isph
--- a/Engine/Source/Runtime/Engine/Public/Animation/AnimationCompression.isph
+++ b/Engine/Source/Runtime/Engine/Public/Animation/AnimationCompression.isph
@@ -414,7 +414,7 @@ inline float TimeToIndex(
 	return Alpha;
 }
 
-inline FVector4 Float96NoWToQuat(const float* Values)
+inline FVector4f Float96NoWToQuat(const float* Values)
 {
 	#pragma ignore warning(perf)
 	const float X = Values[0];
@@ -430,7 +430,7 @@ inline FVector4 Float96NoWToQuat(const float* Values)
 	return SetVector4(X, Y, Z, WSquared > 0.f ? sqrt(WSquared) : 0.f);
 }
 
-inline FVector4 Fixed32NoWToQuat(const unsigned int* PackedValue)
+inline FVector4f Fixed32NoWToQuat(const unsigned int* PackedValue)
 {
 	#pragma ignore warning(perf)
 	const unsigned int Packed = *PackedValue;
@@ -447,7 +447,7 @@ inline FVector4 Fixed32NoWToQuat(const unsigned int* PackedValue)
 	return SetVector4(X, Y, Z, WSquared > 0.f ? sqrt(WSquared) : 0.f);
 }
 
-inline FVector4 Fixed48NoWToQuat(const unsigned int16* Values)
+inline FVector4f Fixed48NoWToQuat(const unsigned int16* Values)
 {
 	#pragma ignore warning(perf)
 	const unsigned int16 X = Values[0];
@@ -466,9 +466,9 @@ inline FVector4 Fixed48NoWToQuat(const unsigned int16* Values)
 	return SetVector4(FX, FY, FZ, WSquared > 0.f ? sqrt( WSquared ) : 0.f);
 }
 
-inline FVector4 IntervalFixed32NoWToQuat(
-	const FVector& Mins, 
-	const FVector& Ranges, 
+inline FVector4f IntervalFixed32NoWToQuat(
+	const FVector3f& Mins, 
+	const FVector3f& Ranges, 
 	const unsigned int* PackedValue)
 {
 	#pragma ignore warning(perf)
@@ -515,7 +515,7 @@ inline float DecodeFloat(const uniform unsigned int NumExponentBits, const unifo
 	return floatbits((Sign << SignShift) | (Exponent << NumMantissaBits) | (Mantissa));
 }
 
-inline FVector4 Float32NoWToQuat(const unsigned int* PackedValue)
+inline FVector4f Float32NoWToQuat(const unsigned int* PackedValue)
 {
 	#pragma ignore warning(perf)
 	const unsigned int Packed = *PackedValue;
@@ -532,7 +532,7 @@ inline FVector4 Float32NoWToQuat(const unsigned int* PackedValue)
 	return SetVector4(X, Y, Z, WSquared > 0.f ? sqrt(WSquared) : 0.f);
 }
 
-inline FVector IntervalFixed32NoWToVector(const FVector& Mins, const FVector& Ranges, const unsigned int* PackedValue)
+inline FVector3f IntervalFixed32NoWToVector(const FVector3f& Mins, const FVector3f& Ranges, const unsigned int* PackedValue)
 {
 	#pragma ignore warning(perf)
 	const unsigned int Packed = *PackedValue;
@@ -548,7 +548,7 @@ inline FVector IntervalFixed32NoWToVector(const FVector& Mins, const FVector& Ra
 	return SetVector(X, Y, Z);
 }
 
-inline FVector Fixed48NoWToVector(const unsigned int16* Values)
+inline FVector3f Fixed48NoWToVector(const unsigned int16* Values)
 {
 	#pragma ignore warning(perf)
 	const unsigned int16 X = Values[0];
@@ -575,12 +575,13 @@ inline float DecompressFixed16(const unsigned int16 Value)
 	return ((int)Value - QuantOffset) * InvQuantFactor;
 }
 
-inline void DecompressRotation(FVector4& Out, const unsigned int8* TopOfStream, const unsigned int8* KeyData, const uniform int FORMAT)
+inline void DecompressRotation(FVector4f& Out, const unsigned int8* TopOfStream, const unsigned int8* KeyData, const uniform int FORMAT)
 {
 	if ( FORMAT == ACF_None )
 	{
 		#pragma ignore warning(perf)
-		Out = *((const FVector4 *)KeyData);
+		const FVector4f Keys = *((const FVector4f *)KeyData);
+		Out = SetVector4(Keys.V[0], Keys.V[1], Keys.V[2], Keys.V[3]);
 	}
 	else if ( FORMAT == ACF_Float96NoW )
 	{
@@ -596,13 +597,13 @@ inline void DecompressRotation(FVector4& Out, const unsigned int8* TopOfStream,
 	}
 	else if ( FORMAT == ACF_IntervalFixed32NoW )
 	{
-		const FVector* VTopOfStream = (const FVector*)TopOfStream;
+		const FVector3f* VTopOfStream = (const FVector3f*)TopOfStream;
 
 		#pragma ignore warning(perf)
-		const FVector Mins = *(VTopOfStream);
+		const FVector3f Mins = *(VTopOfStream);
 
 		#pragma ignore warning(perf)
-		const FVector Ranges = *(VTopOfStream+1);
+		const FVector3f Ranges = *(VTopOfStream+1);
 
 		Out = IntervalFixed32NoWToQuat(Mins, Ranges, (const unsigned int*)KeyData);
 	}
@@ -620,7 +621,7 @@ inline void DecompressRotation(FVector4& Out, const unsigned int8* TopOfStream,
 	}
 }
 
-inline void DecompressRotation(const int Format, const int FormatFlags, FVector4& Out, const unsigned int8* TopOfStream, const unsigned int8* KeyData)
+inline void DecompressRotation(const int Format, const int FormatFlags, FVector4f& Out, const unsigned int8* TopOfStream, const unsigned int8* KeyData)
 {
 	if (Format == ACF_Fixed48NoW)
 	{
@@ -658,8 +659,8 @@ inline void DecompressRotation(const int Format, const int FormatFlags, FVector4
 	{
 		const float* SourceBounds = (float*)TopOfStream;
 
-		FVector Mins = {{0.0f, 0.0f, 0.0f}};
-		FVector Ranges = {{0.0f, 0.0f, 0.0f}};
+		FVector3f Mins = FloatZeroVector;
+		FVector3f Ranges = FloatZeroVector;
 
 		if (FormatFlags & 1)
 		{
@@ -706,27 +707,28 @@ inline void DecompressRotation(const int Format, const int FormatFlags, FVector4
 	}
 }
 
-inline void DecompressTranslation(FVector& Out, const unsigned int8* TopOfStream, const unsigned int8* KeyData, const uniform int FORMAT)
+inline void DecompressTranslation(FVector3f& Out, const unsigned int8* TopOfStream, const unsigned int8* KeyData, const uniform int FORMAT)
 {
 	if ( (FORMAT == ACF_None) || (FORMAT == ACF_Float96NoW) )
 	{
 		#pragma ignore warning(perf)
-		Out = *((const FVector *)KeyData);
+		const FVector3f Keys = *((const FVector3f *)KeyData);
+		Out = SetVector(Keys.V[0], Keys.V[1], Keys.V[2]);
 	}
 	else if ( FORMAT == ACF_IntervalFixed32NoW )
 	{
-		const FVector* VTopOfStream = (const FVector*)TopOfStream;
+		const FVector3f* VTopOfStream = (const FVector3f*)TopOfStream;
 
 		#pragma ignore warning(perf)
-		const FVector Mins = *(VTopOfStream);
+		const FVector3f Mins = *(VTopOfStream);
 
 		#pragma ignore warning(perf)
-		const FVector Ranges = *(VTopOfStream+1);
+		const FVector3f Ranges = *(VTopOfStream+1);
 		Out = IntervalFixed32NoWToVector(Mins, Ranges, (const unsigned int*)KeyData);
 	}
 	else if ( FORMAT == ACF_Identity )
 	{
-		Out = ZeroVector;
+		Out = FloatZeroVector;
 	}
 	else if ( FORMAT == ACF_Fixed48NoW )
 	{
@@ -734,11 +736,11 @@ inline void DecompressTranslation(FVector& Out, const unsigned int8* TopOfStream
 	}
 	else
 	{
-		Out = ZeroVector;
+		Out = FloatZeroVector;
 	}
 }
 
-inline void DecompressTranslation(const int Format, const int FormatFlags, FVector& Out, const unsigned int8* TopOfStream, const unsigned int8* KeyData)
+inline void DecompressTranslation(const int Format, const int FormatFlags, FVector3f& Out, const unsigned int8* TopOfStream, const unsigned int8* KeyData)
 {
 	if( Format == ACF_Float96NoW )
 	{
@@ -746,7 +748,8 @@ inline void DecompressTranslation(const int Format, const int FormatFlags, FVect
 		if( (FormatFlags & 7) == 0 )
 		{
 			#pragma ignore warning(perf)
-			Out = *((const FVector *)KeyData);
+			const FVector3f Keys = *((const FVector3f *)KeyData);
+			Out = SetVector(Keys.V[0], Keys.V[1], Keys.V[2]);
 		}
 		// Stored per components
 		else
@@ -754,21 +757,23 @@ inline void DecompressTranslation(const int Format, const int FormatFlags, FVect
 			const float* TypedKeyData = (const float*)KeyData;
 			
 			#pragma ignore warning(perf)
-			Out.V[0] = (FormatFlags & 1) ? (*TypedKeyData++) : 0.0f;
+			const float X = (FormatFlags & 1) ? (*TypedKeyData++) : 0.0f;
 
 			#pragma ignore warning(perf)
-			Out.V[1] = (FormatFlags & 2) ? (*TypedKeyData++) : 0.0f;
+			const float Y = (FormatFlags & 2) ? (*TypedKeyData++) : 0.0f;
 
 			#pragma ignore warning(perf)
-			Out.V[2] = (FormatFlags & 4) ? (*TypedKeyData++) : 0.0f;
+			const float Z = (FormatFlags & 4) ? (*TypedKeyData++) : 0.0f;
+
+			Out = SetVector(X, Y, Z);
 		}
 	}
 	else if (Format == ACF_IntervalFixed32NoW)
 	{
 		const float* SourceBounds = (float*)TopOfStream;
 
-		FVector Mins = {{0.0f, 0.0f, 0.0f}};
-		FVector Ranges = {{0.0f, 0.0f, 0.0f}};
+		FVector3f Mins = FloatZeroVector;
+		FVector3f Ranges = FloatZeroVector;
 
 		if (FormatFlags & 1)
 		{
@@ -802,45 +807,48 @@ inline void DecompressTranslation(const int Format, const int FormatFlags, FVect
 		const unsigned int16* TypedKeyData = (const unsigned int16*)KeyData;
 
 		#pragma ignore warning(perf)
-		Out.V[0] = (FormatFlags & 1) ? DecompressFixed16(*TypedKeyData++) : 0.0f;
+		const float X = (FormatFlags & 1) ? DecompressFixed16(*TypedKeyData++) : 0.0f;
 
 		#pragma ignore warning(perf)
-		Out.V[1] = (FormatFlags & 2) ? DecompressFixed16(*TypedKeyData++) : 0.0f;
+		const float Y = (FormatFlags & 2) ? DecompressFixed16(*TypedKeyData++) : 0.0f;
 
 		#pragma ignore warning(perf)
-		Out.V[2] = (FormatFlags & 4) ? DecompressFixed16(*TypedKeyData++) : 0.0f;
+		const float Z = (FormatFlags & 4) ? DecompressFixed16(*TypedKeyData++) : 0.0f;
+
+		Out = SetVector(X, Y, Z);
 	}
 	else if ( Format == ACF_Identity )
 	{
-		Out = ZeroVector;
+		Out = FloatZeroVector;
 	}
 	else
 	{
-		Out = ZeroVector;
+		Out = FloatZeroVector;
 	}
 }
 
-inline void DecompressScale(FVector& Out, const unsigned int8* TopOfStream, const unsigned int8* KeyData, const uniform int FORMAT)
+inline void DecompressScale(FVector3f& Out, const unsigned int8* TopOfStream, const unsigned int8* KeyData, const uniform int FORMAT)
 {
 	if ( (FORMAT == ACF_None) || (FORMAT == ACF_Float96NoW) )
 	{
 		#pragma ignore warning(perf)
-		Out = *((const FVector *)KeyData);
+		const FVector3f Keys = *((const FVector3f *)KeyData);
+		Out = SetVector(Keys.V[0], Keys.V[1], Keys.V[2]);
 	}
 	else if ( FORMAT == ACF_IntervalFixed32NoW )
 	{
-		const FVector* VTopOfStream = (const FVector*)TopOfStream;
+		const FVector3f* VTopOfStream = (const FVector3f*)TopOfStream;
 
 		#pragma ignore warning(perf)
-		const FVector Mins = *(VTopOfStream);
+		const FVector3f Mins = *(VTopOfStream);
 
 		#pragma ignore warning(perf)
-		const FVector Ranges = *(VTopOfStream+1);
+		const FVector3f Ranges = *(VTopOfStream+1);
 		Out = IntervalFixed32NoWToVector(Mins, Ranges, (const unsigned int*)KeyData);
 	}
 	else if ( FORMAT == ACF_Identity )
 	{
-		Out = ZeroVector;
+		Out = FloatZeroVector;
 	}
 	else if ( FORMAT == ACF_Fixed48NoW )
 	{
@@ -848,11 +856,11 @@ inline void DecompressScale(FVector& Out, const unsigned int8* TopOfStream, cons
 	}
 	else
 	{
-		Out = ZeroVector;
+		Out = FloatZeroVector;
 	}
 }
 
-inline void DecompressScale(const int Format, const int FormatFlags, FVector& Out, const unsigned int8* TopOfStream, const unsigned int8* KeyData)
+inline void DecompressScale(const int Format, const int FormatFlags, FVector3f& Out, const unsigned int8* TopOfStream, const unsigned int8* KeyData)
 {
 	if( Format == ACF_Float96NoW )
 	{
@@ -860,7 +868,8 @@ inline void DecompressScale(const int Format, const int FormatFlags, FVector& Ou
 		if( (FormatFlags & 7) == 0 )
 		{
 			#pragma ignore warning(perf)
-			Out = *((FVector*)KeyData);
+			const FVector3f Keys = *((const FVector3f *)KeyData);
+			Out = SetVector(Keys.V[0], Keys.V[1], Keys.V[2]);
 		}
 		// Stored per components
 		else
@@ -868,21 +877,23 @@ inline void DecompressScale(const int Format, const int FormatFlags, FVector& Ou
 			const float* TypedKeyData = (const float*)KeyData;
 
 			#pragma ignore warning(perf)
-			Out.V[0] = (FormatFlags & 1) ? (*TypedKeyData++) : 0.0f;
+			const float X = (FormatFlags & 1) ? (*TypedKeyData++) : 0.0f;
 
 			#pragma ignore warning(perf)
-			Out.V[1] = (FormatFlags & 2) ? (*TypedKeyData++) : 0.0f;
+			const float Y = (FormatFlags & 2) ? (*TypedKeyData++) : 0.0f;
 
 			#pragma ignore warning(perf)
-			Out.V[2] = (FormatFlags & 4) ? (*TypedKeyData++) : 0.0f;
+			const float Z = (FormatFlags & 4) ? (*TypedKeyData++) : 0.0f;
+
+			Out = SetVector(X, Y, Z);
 		}
 	}
 	else if (Format == ACF_IntervalFixed32NoW)
 	{
 		const float* SourceBounds = (float*)TopOfStream;
 
-		FVector Mins = {{0.0f, 0.0f, 0.0f}};
-		FVector Ranges = {{0.0f, 0.0f, 0.0f}};
+		FVector3f Mins = FloatZeroVector;
+		FVector3f Ranges = FloatZeroVector;
 
 		if (FormatFlags & 1)
 		{
@@ -916,21 +927,23 @@ inline void DecompressScale(const int Format, const int FormatFlags, FVector& Ou
 		const unsigned int16* TypedKeyData = (const unsigned int16*)KeyData;
 		
 		#pragma ignore warning(perf)
-		Out.V[0] = (FormatFlags & 1) ? DecompressFixed16(*TypedKeyData++) : 0.0f;
+		const float X = (FormatFlags & 1) ? DecompressFixed16(*TypedKeyData++) : 0.0f;
 
 		#pragma ignore warning(perf)
-		Out.V[1] = (FormatFlags & 2) ? DecompressFixed16(*TypedKeyData++) : 0.0f;
+		const float Y = (FormatFlags & 2) ? DecompressFixed16(*TypedKeyData++) : 0.0f;
 
 		#pragma ignore warning(perf)
-		Out.V[2] = (FormatFlags & 4) ? DecompressFixed16(*TypedKeyData++) : 0.0f;
+		const float Z = (FormatFlags & 4) ? DecompressFixed16(*TypedKeyData++) : 0.0f;
+
+		Out = SetVector(X, Y, Z);
 	}
 	else if ( Format == ACF_Identity )
 	{
-		Out = ZeroVector;
+		Out = FloatZeroVector;
 	}
 	else
 	{
-		Out = ZeroVector;
+		Out = FloatZeroVector;
 	}
 }
 
diff --git a/Engine/Source/Runtime/Engine/Public/AnimationRuntime.h b/Engine/Source/Runtime/Engine/Public/AnimationRuntime.h
--- a/Engine/Source/Runtime/Engine/Public/AnimationRuntime.h
+++ b/Engine/Source/Runtime/Engine/Public/AnimationRuntime.h
@@ -627,3 +627,12 @@ public:
 };
 
 ENUM_CLASS_FLAGS(FAnimationRuntime::EBlendPosesPerBoneFilterFlags);
+
+// Support ISPC enable/disable in non-shipping builds
+#if !INTEL_ISPC
+const bool bAnim_Runtime_ISPC_Enabled = false;
+#elif UE_BUILD_SHIPPING
+const bool bAnim_Runtime_ISPC_Enabled = true;
+#else
+extern bool bAnim_Runtime_ISPC_Enabled;
+#endif
diff --git a/Engine/Source/Runtime/Engine/Public/BonePose.h b/Engine/Source/Runtime/Engine/Public/BonePose.h
--- a/Engine/Source/Runtime/Engine/Public/BonePose.h
+++ b/Engine/Source/Runtime/Engine/Public/BonePose.h
@@ -958,3 +958,12 @@ extern ENGINE_API void BuildPoseFromRawData(
 	const TArray<FTransform>& RetargetTransforms,
 	const TMap<int32, const FTransformCurve*>* AdditiveBoneTransformCurves = nullptr
 	);
+
+// Support ISPC enable/disable in non-shipping builds
+#if !INTEL_ISPC
+const bool bAnim_BonePose_ISPC_Enabled = false;
+#elif UE_BUILD_SHIPPING
+const bool bAnim_BonePose_ISPC_Enabled = true;
+#else
+extern bool bAnim_BonePose_ISPC_Enabled;
+#endif
diff --git a/Engine/Source/Runtime/Engine/Public/GPUSkinVertexFactory.h b/Engine/Source/Runtime/Engine/Public/GPUSkinVertexFactory.h
--- a/Engine/Source/Runtime/Engine/Public/GPUSkinVertexFactory.h
+++ b/Engine/Source/Runtime/Engine/Public/GPUSkinVertexFactory.h
@@ -989,3 +989,12 @@ public:
 	static void ModifyCompilationEnvironment(const FVertexFactoryShaderPermutationParameters& Parameters, FShaderCompilerEnvironment& OutEnvironment);
 	static bool ShouldCompilePermutation(const FVertexFactoryShaderPermutationParameters& Parameters);
 };
+
+// Support ISPC enable/disable in non-shipping builds
+#if !INTEL_ISPC
+const bool bGPUSkin_CopyBones_ISPC_Enabled = false;
+#elif UE_BUILD_SHIPPING
+const bool bGPUSkin_CopyBones_ISPC_Enabled = true;
+#else
+extern bool bGPUSkin_CopyBones_ISPC_Enabled;
+#endif
diff --git a/Engine/Source/Runtime/Experimental/Chaos/Private/Chaos/AABB.cpp b/Engine/Source/Runtime/Experimental/Chaos/Private/Chaos/AABB.cpp
--- a/Engine/Source/Runtime/Experimental/Chaos/Private/Chaos/AABB.cpp
+++ b/Engine/Source/Runtime/Experimental/Chaos/Private/Chaos/AABB.cpp
@@ -5,6 +5,14 @@
 #include "Chaos/Capsule.h"
 #if INTEL_ISPC
 #include "AABB.ispc.generated.h"
+
+static_assert(sizeof(ispc::FTransform) == sizeof(FTransform), "sizeof(ispc::FTransform) != sizeof(FTransform)");
+static_assert(sizeof(ispc::FVector) == sizeof(Chaos::TVector<Chaos::FReal, 3>), "sizeof(ispc::FVector) != sizeof(Chaos::TVector<Chaos::FReal, 3>)");
+#endif
+
+#if INTEL_ISPC && !UE_BUILD_SHIPPING
+bool bChaos_AABBTransform_ISPC_Enabled = true;
+FAutoConsoleVariableRef CVarChaosAABBTransformISPCEnabled(TEXT("p.Chaos.AABBTransform.ISPC"), bChaos_AABBTransform_ISPC_Enabled, TEXT("Whether to use ISPC optimizations when computing AABB transforms"));
 #endif
 
 namespace Chaos
@@ -286,26 +294,44 @@ inline TAABB<T, 3> TransformedAABBHelper(const TAABB<T, 3>& AABB, const FMatrix4
 	return TAABB<T, 3>(Min, Max);
 }
 
-template<typename T>
-inline TAABB<T, 3> TransformedAABBHelperISPC(const TAABB<T, 3>& AABB, const FTransform& SpaceTransform)
+inline TAABB<FReal, 3> TransformedAABBHelperISPC(const TAABB<FReal, 3>& AABB, const FTransform& SpaceTransform)
 {
 	check(bRealTypeCompatibleWithISPC);
 #if INTEL_ISPC
-	TVec3<FReal> NewMin, NewMax;
+	TVector<Chaos::FReal, 3> NewMin, NewMax;
 	ispc::TransformedAABB((const ispc::FTransform&)SpaceTransform, (const ispc::FVector&)AABB.Min(), (const ispc::FVector&)AABB.Max(), (ispc::FVector&)NewMin, (ispc::FVector&)NewMax);
 
 	TAABB<FReal, 3> NewAABB(NewMin, NewMax);
 	return NewAABB;
 #else
 	check(false);
-	return TAABB<T, 3>::EmptyAABB();
+	return TAABB<FReal, 3>::EmptyAABB();
 #endif
 }
 
+#if !UE_LARGE_WORLD_COORDINATES_DISABLED
+inline TAABB<Chaos::FRealSingle, 3> TransformedAABBHelperISPC(const TAABB<Chaos::FRealSingle, 3>& AABB, const FTransform& SpaceTransform)
+{
+	check(bRealTypeCompatibleWithISPC);
+#if INTEL_ISPC
+	static_assert(sizeof(ispc::FVector3f) == sizeof(Chaos::TVector<Chaos::FRealSingle, 3>), "sizeof(ispc::FVector3f) != sizeof(Chaos::TVector<Chaos::FRealSingle, 3>)");
+
+	TVector<Chaos::FRealSingle, 3> NewMin, NewMax;
+	ispc::TransformedAABBMixed((const ispc::FTransform&)SpaceTransform, (const ispc::FVector3f&)AABB.Min(), (const ispc::FVector3f&)AABB.Max(), (ispc::FVector3f&)NewMin, (ispc::FVector3f&)NewMax);
+
+	TAABB<Chaos::FRealSingle, 3> NewAABB(NewMin, NewMax);
+	return NewAABB;
+#else
+	check(false);
+	return TAABB<Chaos::FRealSingle, 3>::EmptyAABB();
+#endif
+}
+#endif
+
 template<typename T, int d>
 TAABB<T, d> TAABB<T, d>::TransformedAABB(const Chaos::TRigidTransform<FReal, 3>& SpaceTransform) const
 {
-	if (bRealTypeCompatibleWithISPC && INTEL_ISPC)
+	if (bRealTypeCompatibleWithISPC && INTEL_ISPC && bChaos_AABBTransform_ISPC_Enabled)
 	{
 		return TransformedAABBHelperISPC(*this, SpaceTransform);
 	}
@@ -331,7 +357,7 @@ TAABB<T, d> TAABB<T, d>::TransformedAABB(const Chaos::PMatrix<FReal, 4, 4>& Spac
 template<typename T, int d>
 TAABB<T, d> TAABB<T, d>::TransformedAABB(const FTransform& SpaceTransform) const
 {
-	if (bRealTypeCompatibleWithISPC && INTEL_ISPC)
+	if (bRealTypeCompatibleWithISPC && INTEL_ISPC && bChaos_AABBTransform_ISPC_Enabled)
 	{
 		return TransformedAABBHelperISPC(*this, SpaceTransform);
 	}
diff --git a/Engine/Source/Runtime/Experimental/Chaos/Private/Chaos/AABB.ispc b/Engine/Source/Runtime/Experimental/Chaos/Private/Chaos/AABB.ispc
--- a/Engine/Source/Runtime/Experimental/Chaos/Private/Chaos/AABB.ispc
+++ b/Engine/Source/Runtime/Experimental/Chaos/Private/Chaos/AABB.ispc
@@ -21,13 +21,15 @@ inline static void GrowToInclude(uniform FVector8 &MinMax, const uniform FVector
 	MinMax = SetVector8(Min.V[0], Min.V[1], Min.V[2], Min.V[3], Max.V[4], Max.V[5], Max.V[6], Max.V[7]);
 }
 
+// Three permutations need supporting - FTransform<float>/AABB<float>, FTransform<double>/AABB<double>, FTransform<double>/AABB<float>.
+// TransformAABB handles the first two based on LWC enablement. Mixed handles the third case when LWC is enabled.
 export void TransformedAABB(const uniform FTransform &SpaceTransform, const uniform FVector &MMin, const uniform FVector &MMax, uniform FVector &NewMin, uniform FVector &NewMax)
 {
-	const uniform FVector4 Min = SetVector4(MMin.V[0], MMin.V[1], MMin.V[2], 0.0f);
-	const uniform FVector4 Max = SetVector4(MMax.V[0], MMax.V[1], MMax.V[2], 0.0f);
+	const uniform FVector4 Min = SetVector4(MMin.V[0], MMin.V[1], MMin.V[2], ZERO);
+	const uniform FVector4 Max = SetVector4(MMax.V[0], MMax.V[1], MMax.V[2], ZERO);
 	const uniform FVector4 CurrentExtents = Max - Min;
 
-#if TARGET_WIDTH == 4
+#if UE_LARGE_WORLD_COORDINATES_DISABLED == 0 || TARGET_WIDTH == 4
 	uniform FVector4 MinToNewSpace = TransformPosition(SpaceTransform, Min);
 	uniform FVector4 MaxToNewSpace = MinToNewSpace;
 	GrowToInclude(MinToNewSpace, MaxToNewSpace, TransformPosition(SpaceTransform, Max));
@@ -60,3 +62,31 @@ export void TransformedAABB(const uniform FTransform &SpaceTransform, const unif
 	NewMax = SetVector(MinMaxToNewSpace.V[4], MinMaxToNewSpace.V[5], MinMaxToNewSpace.V[6]);
 #endif
 }
+
+#if UE_LARGE_WORLD_COORDINATES_DISABLED == 0
+export void TransformedAABBMixed(const uniform FTransform &SpaceTransform, const uniform FVector3f &MMin, const uniform FVector3f &MMax, uniform FVector3f &NewMin, uniform FVector3f &NewMax)
+{
+	const uniform FVector Min3 = ConvertVector3fTo3Native(MMin);
+	const uniform FVector Max3 = ConvertVector3fTo3Native(MMax);
+	const uniform FVector4 Min = SetVector4(Min3.V[0], Min3.V[1], Min3.V[2], ZERO);
+	const uniform FVector4 Max = SetVector4(Max3.V[0], Max3.V[1], Max3.V[2], ZERO);
+	const uniform FVector4 CurrentExtents = Max - Min;
+
+	uniform FVector4 MinToNewSpace = TransformPosition(SpaceTransform, Min);
+	uniform FVector4 MaxToNewSpace = MinToNewSpace;
+	GrowToInclude(MinToNewSpace, MaxToNewSpace, TransformPosition(SpaceTransform, Max));
+
+	GrowToInclude(MinToNewSpace, MaxToNewSpace, TransformPosition(SpaceTransform, Min + ForwardVector4 * CurrentExtents));
+	GrowToInclude(MinToNewSpace, MaxToNewSpace, TransformPosition(SpaceTransform, Max - ForwardVector4 * CurrentExtents));
+	GrowToInclude(MinToNewSpace, MaxToNewSpace, TransformPosition(SpaceTransform, Min + RightVector4 * CurrentExtents));
+	GrowToInclude(MinToNewSpace, MaxToNewSpace, TransformPosition(SpaceTransform, Max - RightVector4 * CurrentExtents));
+	GrowToInclude(MinToNewSpace, MaxToNewSpace, TransformPosition(SpaceTransform, Min + UpVector4 * CurrentExtents));
+	GrowToInclude(MinToNewSpace, MaxToNewSpace, TransformPosition(SpaceTransform, Max - UpVector4 * CurrentExtents));
+
+	const uniform FVector4f MinToNewSpace4f = ConvertVector4NativeTo4f(MinToNewSpace);
+	const uniform FVector4f MaxToNewSpace4f = ConvertVector4NativeTo4f(MaxToNewSpace);
+
+	NewMin = SetVector(MinToNewSpace4f);
+	NewMax = SetVector(MaxToNewSpace4f);
+}
+#endif
diff --git a/Engine/Source/Runtime/Experimental/Chaos/Private/Chaos/CollisionResolutionUtil.cpp b/Engine/Source/Runtime/Experimental/Chaos/Private/Chaos/CollisionResolutionUtil.cpp
--- a/Engine/Source/Runtime/Experimental/Chaos/Private/Chaos/CollisionResolutionUtil.cpp
+++ b/Engine/Source/Runtime/Experimental/Chaos/Private/Chaos/CollisionResolutionUtil.cpp
@@ -14,6 +14,10 @@
 
 #if INTEL_ISPC
 #include "PBDCollisionConstraints.ispc.generated.h"
+
+static_assert(sizeof(ispc::FVector) == sizeof(Chaos::FVec3), "sizeof(ispc::FVector) != sizeof(Chaos::FVec3)");
+static_assert(sizeof(ispc::FIntVector) == sizeof(Chaos::TVector<int32, 3>), "sizeof(ispc::FIntVector) != sizeof(Chaos::TVector<int32, 3>");
+static_assert(sizeof(ispc::FTransform) == sizeof(Chaos::FRigidTransform3), "sizeof(ispc::FTransform) != sizeof(Chaos::FRigidTransform3)");
 #endif
 
 //PRAGMA_DISABLE_OPTIMIZATION
diff --git a/Engine/Source/Runtime/Experimental/Chaos/Private/Chaos/Evolution/PBDMinEvolution.ispc b/Engine/Source/Runtime/Experimental/Chaos/Private/Chaos/Evolution/PBDMinEvolution.ispc
--- a/Engine/Source/Runtime/Experimental/Chaos/Private/Chaos/Evolution/PBDMinEvolution.ispc
+++ b/Engine/Source/Runtime/Experimental/Chaos/Private/Chaos/Evolution/PBDMinEvolution.ispc
@@ -39,7 +39,7 @@ struct FPBDRigidArrays
 	FVector* PreW;
 	FVector* CenterOfMass;
 	FVector4* RotationOfMass;
-	float* InvM;
+	FReal* InvM;
 	FMatrix33* InvI;
 	FVector* F;
 	FVector* T;
@@ -47,8 +47,8 @@ struct FPBDRigidArrays
 	FVector* AngularImpulse;
 	int8* Disabled;
 	int8* GravityEnabled;
-	float* LinearEtherDrag;
-	float* AngularEtherDrag;
+	FReal* LinearEtherDrag;
+	FReal* AngularEtherDrag;
 	int8* HasBounds;
 	FAABB* LocalBounds;
 	FAABB* WorldBounds;
@@ -78,31 +78,31 @@ export uniform int32 SizeofFSimulationSpace()
 struct FSimulationSpaceSettings
 {
 	// Global multipler on the effects of simulation space movement
-	float MasterAlpha;
+	FReal MasterAlpha;
 
 	// How much of the simulation frame's linear acceleration to pass onto the particles
-	float LinearAccelerationAlpha;
+	FReal LinearAccelerationAlpha;
 
 	// How much of the coriolis force to apply to the particles
-	float CoriolisAlpha;
+	FReal CoriolisAlpha;
 
 	// How much of the centrifugal force to apply to the particles
-	float CentrifugalAlpha;
+	FReal CentrifugalAlpha;
 
 	// How much of the euler force to apply to the particles
-	float EulerAlpha;
+	FReal EulerAlpha;
 
 	// How much of the simulation frame's angular acceleration to pass onto the particles
-	float AngularAccelerationAlpha;
+	FReal AngularAccelerationAlpha;
 
 	// How much of the simulation frame's linear velocity to pass onto the particles (linear ether drag)
-	float LinearVelocityAlpha;
+	FReal LinearVelocityAlpha;
 
 	// How much of the simulation frame's angular velocity to pass onto the particles (angular ether drag)
-	float AngularVelocityAlpha;
+	FReal AngularVelocityAlpha;
 
 	// An additional linear drag (on top of the EtherDrag that can be set on the physics asset)
-	float ExternalLinearEtherDrag;
+	FVector ExternalLinearEtherDrag;
 };
 
 export uniform int32 SizeofFSimulationSpaceSettings()
@@ -123,10 +123,10 @@ static inline uniform FMatrix33 ComputeWorldSpaceInertia(const uniform FVector4
 	return MultiplyAB(QM, MultiplyABt(I, QM));
 }
 
-static inline uniform FVector4 IntegrateRotationWithAngularVelocity(const uniform FVector4& Q0, const uniform FVector& W, const uniform float Dt)
+static inline uniform FVector4 IntegrateRotationWithAngularVelocity(const uniform FVector4& Q0, const uniform FVector& W, const uniform FReal Dt)
 {
-	uniform FVector4 R1 = Q0 + VectorQuaternionMultiply2(SetVector4(W.V[0], W.V[1], W.V[2], 0.f), Q0) * (Dt * 0.5f);
-	return VectorNormalizeSafe(R1, SetVector4(0, 0, 0, 0));
+	uniform FVector4 R1 = Q0 + VectorQuaternionMultiply2(SetVector4(W.V[0], W.V[1], W.V[2], ZERO), Q0) * (Dt * ONE_HALF);
+	return VectorNormalizeSafe(R1, VectorZero);
 }
 
 static inline uniform int8 EObjectStateType_Dynamic()
@@ -183,20 +183,20 @@ static inline uniform FVector InverseTransformVectorNoScale(const uniform FTrans
 }
 
 export void MinEvolutionIntegrate(
-	uniform float Dt, 
+	uniform FReal Dt, 
 	uniform FPBDRigidArrays& Rigids, 
 	uniform FSimulationSpace& SimulationSpace, 
 	uniform FSimulationSpaceSettings& SimulationSpaceSettings, 
 	const uniform FVector& Gravity, 
-	const uniform float BoundsExtension,
-	const uniform float CullDistance)
+	const uniform FReal BoundsExtension,
+	const uniform FReal CullDistance)
 {
 	// Simulation space velocity and acceleration
 	uniform FVector SpaceV = ZeroVector;	// Velocity
 	uniform FVector SpaceW = ZeroVector;	// Angular Velocity
 	uniform FVector SpaceA = ZeroVector;	// Acceleration
 	uniform FVector SpaceB = ZeroVector;	// Angular Acceleration
-	if (SimulationSpaceSettings.MasterAlpha > 0.0f)
+	if (SimulationSpaceSettings.MasterAlpha > ZERO)
 	{
 		SpaceV = InverseTransformVectorNoScale(SimulationSpace.Transform, SimulationSpace.LinearVelocity);
 		SpaceW = InverseTransformVectorNoScale(SimulationSpace.Transform, SimulationSpace.AngularVelocity);
@@ -233,9 +233,9 @@ export void MinEvolutionIntegrate(
 
 			// Moving and accelerating simulation frame
 			// https://en.wikipedia.org/wiki/Rotating_reference_frame
-			if (SimulationSpaceSettings.MasterAlpha > 0.0f)
+			if (SimulationSpaceSettings.MasterAlpha > ZERO)
 			{
-				const uniform FVector CoriolisAcc = SimulationSpaceSettings.CoriolisAlpha * 2.0f * VectorCross(SpaceW, Rigids.V[ParticleIndex]);
+				const uniform FVector CoriolisAcc = SimulationSpaceSettings.CoriolisAlpha * TWO * VectorCross(SpaceW, Rigids.V[ParticleIndex]);
 				const uniform FVector CentrifugalAcc = SimulationSpaceSettings.CentrifugalAlpha * VectorCross(SpaceW, VectorCross(SpaceW, XCoM));
 				const uniform FVector EulerAcc = SimulationSpaceSettings.EulerAlpha * VectorCross(SpaceB, XCoM);
 				const uniform FVector LinearAcc = SimulationSpaceSettings.LinearAccelerationAlpha * SpaceA;
@@ -248,8 +248,8 @@ export void MinEvolutionIntegrate(
 			}
 
 			// New velocity
-			const uniform float LinearDrag = min(1.0f, Rigids.LinearEtherDrag[ParticleIndex] * Dt);
-			const uniform float AngularDrag = min(1.0f, Rigids.AngularEtherDrag[ParticleIndex] * Dt);
+			const uniform FReal LinearDrag = min(ONE, Rigids.LinearEtherDrag[ParticleIndex] * Dt);
+			const uniform FReal AngularDrag = min(ONE, Rigids.AngularEtherDrag[ParticleIndex] * Dt);
 			const uniform FVector VCoM = VectorLerp(Rigids.V[ParticleIndex] + DV, TargetV, LinearDrag);
 			const uniform FVector WCoM = VectorLerp(Rigids.W[ParticleIndex] + DW, TargetW, AngularDrag);
 
@@ -273,8 +273,8 @@ export void MinEvolutionIntegrate(
 			{
 				uniform FTransform SpaceTransform;
 				SpaceTransform.Rotation = Rigids.Q[ParticleIndex];
-				SpaceTransform.Translation = SetVector4(Rigids.P[ParticleIndex], 1.0f);
-				SpaceTransform.Scale3D = SetVector4(OneVector, 1.0f);
+				SpaceTransform.Translation = SetVector4(Rigids.P[ParticleIndex], ONE);
+				SpaceTransform.Scale3D = SetVector4(OneVector, ONE);
 			
 				uniform FVector WorldSpaceBoundsMin;
 				uniform FVector WorldSpaceBoundsMax;
diff --git a/Engine/Source/Runtime/Experimental/Chaos/Private/Chaos/Joint/PBDJointSolverGaussSeidel.cpp b/Engine/Source/Runtime/Experimental/Chaos/Private/Chaos/Joint/PBDJointSolverGaussSeidel.cpp
--- a/Engine/Source/Runtime/Experimental/Chaos/Private/Chaos/Joint/PBDJointSolverGaussSeidel.cpp
+++ b/Engine/Source/Runtime/Experimental/Chaos/Private/Chaos/Joint/PBDJointSolverGaussSeidel.cpp
@@ -9,6 +9,11 @@
 #include "ChaosStats.h"
 #if INTEL_ISPC
 #include "PBDJointSolverGaussSeidel.ispc.generated.h"
+
+static_assert(sizeof(ispc::FVector) == sizeof(Chaos::FVec3), "sizeof(ispc::FVector) != sizeof(Chaos::FVec3)");
+static_assert(sizeof(ispc::FTransform) == sizeof(Chaos::FRigidTransform3), "sizeof(ispc::FTransform) != sizeof(Chaos::FRigidTransform3)");
+static_assert(sizeof(ispc::FVector4) == sizeof(Chaos::FRotation3), "sizeof(ispc::FVector4) != sizeof(Chaos::FRotation3)");
+static_assert(sizeof(ispc::FMatrix) == sizeof(Chaos::FMatrix33), "sizeof(ispc::FMatrix) != sizeof(Chaos::FMatrix33)");
 #endif
 
 //PRAGMA_DISABLE_OPTIMIZATION
@@ -28,7 +33,7 @@ namespace Chaos
 		if (bChaos_Joint_ISPC_Enabled)
 		{
 #if INTEL_ISPC
-			check(sizeof(FPBDJointSolver) == ispc::SizeofFPBDJointSolver());
+			check(sizeof(FPBDJointSolver) == ispc::SizeofFPBDJointSolver()); 
 #endif
 		}
 	}
@@ -1201,7 +1206,7 @@ namespace Chaos
 		if (bRealTypeCompatibleWithISPC && bChaos_Joint_ISPC_Enabled)
 		{
 #if INTEL_ISPC
-			float ReturnedLambda = Lambda;
+			FReal ReturnedLambda = Lambda;
 			ispc::ApplyPositionConstraintSoft((ispc::FPBDJointSolver*)this, Dt, JointStiffness, JointDamping, bAccelerationMode, (ispc::FVector&)Axis, Delta, TargetVel, ReturnedLambda);
 			Lambda = ReturnedLambda;
 #endif
@@ -1341,7 +1346,7 @@ namespace Chaos
 		if (bRealTypeCompatibleWithISPC && bChaos_Joint_ISPC_Enabled)
 		{
 #if INTEL_ISPC
-			float ReturnedLambda = Lambda;
+			FReal ReturnedLambda = Lambda;
 			ispc::ApplyRotationConstraintSoftKD((ispc::FPBDJointSolver*)this, KIndex, DIndex, Dt, JointStiffness, JointDamping, bAccelerationMode, (ispc::FVector&) Axis, Angle, AngVelTarget, ReturnedLambda);
 			Lambda = ReturnedLambda;
 #endif
@@ -1397,7 +1402,7 @@ namespace Chaos
 		if (bRealTypeCompatibleWithISPC && bChaos_Joint_ISPC_Enabled)
 		{
 #if INTEL_ISPC
-			float ReturnedLambda = Lambda;
+			FReal ReturnedLambda = Lambda;
 			ispc::ApplyRotationConstraintSoftDD((ispc::FPBDJointSolver*)this, Dt, JointStiffness, JointDamping, bAccelerationMode, (ispc::FVector&) Axis, Angle, AngVelTarget, ReturnedLambda);
 			Lambda = ReturnedLambda;
 #endif
diff --git a/Engine/Source/Runtime/Experimental/Chaos/Private/Chaos/Joint/PBDJointSolverGaussSeidel.ispc b/Engine/Source/Runtime/Experimental/Chaos/Private/Chaos/Joint/PBDJointSolverGaussSeidel.ispc
--- a/Engine/Source/Runtime/Experimental/Chaos/Private/Chaos/Joint/PBDJointSolverGaussSeidel.ispc
+++ b/Engine/Source/Runtime/Experimental/Chaos/Private/Chaos/Joint/PBDJointSolverGaussSeidel.ispc
@@ -4,23 +4,17 @@
 
 #include "Math/Vector.isph"
 #include "Math/Quat.isph"
+#include "Math/Transform.isph"
 #include "Chaos/Matrix33.isph"
 
-struct FRigidTransform
-{
-	FVector4 Rotation;
-	FVector Translation;
-	FVector Scale3D;
-};
-
 static const uniform int32 MaxConstrainedBodies = 2;
 
 struct FPBDJointSolver
 {
 	// Local-space constraint settings
-	FRigidTransform XLs[MaxConstrainedBodies];	// Local-space joint connector transforms
+	FTransform XLs[MaxConstrainedBodies];		// Local-space joint connector transforms
 	FVector InvILs[MaxConstrainedBodies];		// Local-space inverse inertias
-	float InvMs[MaxConstrainedBodies];			// Inverse masses
+	FReal InvMs[MaxConstrainedBodies];			// Inverse masses
 
 	// World-space constraint state
 	FVector Xs[MaxConstrainedBodies];			// World-space joint connector positions
@@ -47,27 +41,27 @@ struct FPBDJointSolver
 	FVector AngularHardLambda;
 
 	// XPBD Accumulators (net impulse for each soft constraint/drive)
-	float LinearSoftLambda;
-	float TwistSoftLambda;
-	float SwingSoftLambda;
+	FReal LinearSoftLambda;
+	FReal TwistSoftLambda;
+	FReal SwingSoftLambda;
 	FVector LinearDriveLambdas;
 	FVector RotationDriveLambdas;
 	
 	// Solver stiffness - increased over iterations for stability
 	// @todo(chaos): remove Stiffness from SolverSettings (because it is not a solver constant)
-	float SolverStiffness;
+	FReal SolverStiffness;
 
 	// Constraint padding which can act something like a velocity constraint (for restitution)
 	FVector LinearConstraintPadding;
 	FVector AngularConstraintPadding;
 
 	// Tolerances below which we stop solving
-	float PositionTolerance;					// Distance error below which we consider a constraint or drive solved
-	float AngleTolerance;						// Angle error below which we consider a constraint or drive solved
+	FReal PositionTolerance;					// Distance error below which we consider a constraint or drive solved
+	FReal AngleTolerance;						// Angle error below which we consider a constraint or drive solved
 
 	FVector LastPs[MaxConstrainedBodies];		// Positions at the beginning of the iteration
 	FVector4 LastQs[MaxConstrainedBodies];		// Rotations at the beginning of the iteration
-	int NumActiveConstraints;					// The number of active constraints and drives in the last iteration (-1 initial value)
+	int32 NumActiveConstraints;					// The number of active constraints and drives in the last iteration (-1 initial value)
 	bool bIsActive;								// Whether any constraints actually moved any bodies in last iteration
 };
 
@@ -83,7 +77,7 @@ static inline uniform FMatrix33 ComputeWorldSpaceInertia(const uniform FVector4
 	return MultiplyAB(QM, MultiplyABt(L, QM));
 }
 
-static inline uniform FMatrix33 ComputeJointFactorMatrix(const uniform FVector& V, const uniform FMatrix33& M, const uniform float Im)
+static inline uniform FMatrix33 ComputeJointFactorMatrix(const uniform FVector& V, const uniform FMatrix33& M, const uniform FReal Im)
 {
 	// Rigid objects rotational contribution to the impulse.
 	// Vx*M*VxT+Im
@@ -114,19 +108,19 @@ static inline uniform FMatrix33 ComputeJointFactorMatrix(const uniform FVector&
 	return SetMatrix33(R1.V[0], R1.V[1], R1.V[2], R2.V[0], R2.V[1], R2.V[2]);
 }
 
-static inline uniform FVector CalculateAngularVelocity(const uniform FVector4 &R0, const uniform FVector4 &R1, const uniform float Dt)
+static inline uniform FVector CalculateAngularVelocity(const uniform FVector4 &R0, const uniform FVector4 &R1, const uniform FReal Dt)
 {
 	const uniform FVector4 R1Shortest = VectorQuaternionEnforceShortestArcWith(R1, R0);
 
 	// W = 2 * dQ/dT * Qinv
-	const uniform FVector4 DRDt = (R1Shortest - R0) * (1.0f / Dt);
+	const uniform FVector4 DRDt = (R1Shortest - R0) * (ONE / Dt);
 	const uniform FVector4 RInv = QuatInverse(R0);
-	const uniform FVector4 W = VectorQuaternionMultiply2(DRDt, RInv) * 2.0f;
+	const uniform FVector4 W = VectorQuaternionMultiply2(DRDt, RInv) * TWO;
 
 	return SetVector(W);
 }
 
-static inline uniform FVector CalculateVelocity(const uniform FVector& P0, const uniform FVector& P1, const uniform float Dt)
+static inline uniform FVector CalculateVelocity(const uniform FVector& P0, const uniform FVector& P1, const uniform FReal Dt)
 {
 	return (P1 - P0) / Dt;
 }
@@ -156,7 +150,7 @@ static inline void ApplyRotationDelta(
 	const uniform FVector4 Qs = M->Qs[Index];
 	const uniform FVector4 DR = SetVector4(InDR, 0);
 
-	const uniform FVector4 DQ = VectorQuaternionMultiply2(DR, Qs) * 0.5f;
+	const uniform FVector4 DQ = VectorQuaternionMultiply2(DR, Qs) * ONE_HALF;
 
 	const uniform FVector4 QsNorm = VectorNormalizeQuaternion(Qs + DQ);
 
@@ -165,7 +159,7 @@ static inline void ApplyRotationDelta(
 	const uniform FVector4 Qs1Shortest = VectorQuaternionEnforceShortestArcWith(M->Qs[1], M->Qs[0]);
 	M->Qs[1] = Qs1Shortest;
 
-	M->Xs[Index] = M->Ps[Index] + VectorQuaternionRotateVector(M->Qs[Index], M->XLs[Index].Translation);
+	M->Xs[Index] = M->Ps[Index] + SetVector(VectorQuaternionRotateVector(M->Qs[Index], M->XLs[Index].Translation));
 	M->Rs[Index] = VectorQuaternionMultiply2(M->Qs[Index], M->XLs[Index].Rotation);
 
 	const uniform FVector4 Rs1Shortest = VectorQuaternionEnforceShortestArcWith(M->Rs[1], M->Rs[0]);
@@ -179,26 +173,26 @@ static inline void ApplyRotationDelta(
 	const uniform FVector& InDR0,
 	const uniform FVector& InDR1)
 {
-	if (M->InvMs[0] > 0.0f)
+	if (M->InvMs[0] > 0)
 	{
 		const uniform FVector4 Qs0 = M->Qs[0];
 		const uniform FVector4 DR0 = SetVector4(InDR0, 0);
-		const uniform FVector4 DQ0 = VectorQuaternionMultiply2(DR0, Qs0) * 0.5f;
+		const uniform FVector4 DQ0 = VectorQuaternionMultiply2(DR0, Qs0) * ONE_HALF;
 		const uniform FVector4 Q0 = VectorNormalizeQuaternion(Qs0 + DQ0);
 		M->Qs[0] = Q0;
-		M->Xs[0] = M->Ps[0] + VectorQuaternionRotateVector(Q0, M->XLs[0].Translation);
+		M->Xs[0] = M->Ps[0] + SetVector(VectorQuaternionRotateVector(Q0, M->XLs[0].Translation));
 		M->Rs[0] = VectorQuaternionMultiply2(Q0, M->XLs[0].Rotation);
 		M->InvIs[0] = ComputeWorldSpaceInertia(Q0, M->InvILs[0]);
 	}
 
-	if (M->InvMs[1] > 0.0f)
+	if (M->InvMs[1] > 0)
 	{
 		const uniform FVector4 Qs1 = M->Qs[1];
 		const uniform FVector4 DR1 = SetVector4(InDR1, 0);
-		const uniform FVector4 DQ1 = VectorQuaternionMultiply2(DR1, Qs1) * 0.5f;
+		const uniform FVector4 DQ1 = VectorQuaternionMultiply2(DR1, Qs1) * ONE_HALF;
 		const uniform FVector4 Q1 = VectorNormalizeQuaternion(Qs1 + DQ1);
 		M->Qs[1] = Q1;
-		M->Xs[1] = M->Ps[1] + VectorQuaternionRotateVector(Q1, M->XLs[1].Translation);
+		M->Xs[1] = M->Ps[1] + SetVector(VectorQuaternionRotateVector(Q1, M->XLs[1].Translation));
 		M->Rs[1] = VectorQuaternionMultiply2(Q1, M->XLs[1].Rotation);
 		M->InvIs[1] = ComputeWorldSpaceInertia(Q1, M->InvILs[1]);
 	}
@@ -251,37 +245,37 @@ static inline void ApplyVelocityDelta(
 
 export void ApplyPositionConstraintSoft(
 	uniform FPBDJointSolver * uniform M,
-	const uniform float Dt,
-	const uniform float JointStiffness,
-	const uniform float JointDamping,
+	const uniform FReal Dt,
+	const uniform FReal JointStiffness,
+	const uniform FReal JointDamping,
 	const uniform bool bAccelerationMode,
 	const uniform FVector& Axis,
-	const uniform float Delta,
-	const uniform float TargetVel,
-	uniform float& Lambda)
+	const uniform FReal Delta,
+	const uniform FReal TargetVel,
+	uniform FReal& Lambda)
 {
 	// Joint-space inverse mass
 	const uniform FVector AngularAxis0 = VectorCross(M->Xs[0] - M->Ps[0], Axis);
 	const uniform FVector AngularAxis1 = VectorCross(M->Xs[1] - M->Ps[1], Axis);
 	const uniform FVector IA0 = Multiply(M->InvIs[0], AngularAxis0);
 	const uniform FVector IA1 = Multiply(M->InvIs[1], AngularAxis1);
-	const uniform float II0 = VectorDot(AngularAxis0, IA0);
-	const uniform float II1 = VectorDot(AngularAxis1, IA1);
-	const uniform float II = (M->InvMs[0] + II0 + M->InvMs[1] + II1);
+	const uniform FReal II0 = VectorDot(AngularAxis0, IA0);
+	const uniform FReal II1 = VectorDot(AngularAxis1, IA1);
+	const uniform FReal II = (M->InvMs[0] + II0 + M->InvMs[1] + II1);
 
-	uniform float VelDt = 0;
+	uniform FReal VelDt = 0;
 	if (JointDamping > KINDA_SMALL_NUMBER)
 	{
-		const uniform FVector V0Dt = CalculateVelocity(M->InitXs[0], M->Xs[0], 1.0f);
-		const uniform FVector V1Dt = CalculateVelocity(M->InitXs[1], M->Xs[1], 1.0f);
+		const uniform FVector V0Dt = CalculateVelocity(M->InitXs[0], M->Xs[0], ONE);
+		const uniform FVector V1Dt = CalculateVelocity(M->InitXs[1], M->Xs[1], ONE);
 		VelDt = TargetVel * Dt + VectorDot(V0Dt - V1Dt, Axis);
 	}
 
-	const uniform float SpringMassScale = (bAccelerationMode) ? 1.0f / (M->InvMs[0] + M->InvMs[1]) : 1.0f;
-	const uniform float S = SpringMassScale * JointStiffness * Dt * Dt;
-	const uniform float D = SpringMassScale * JointDamping * Dt;
-	const uniform float Multiplier = 1.0f / ((S + D) * II + 1.0f);
-	const uniform float DLambda = M->SolverStiffness * Multiplier * (S * Delta - D * VelDt - Lambda);
+	const uniform FReal SpringMassScale = (bAccelerationMode) ? ONE / (M->InvMs[0] + M->InvMs[1]) : ONE;
+	const uniform FReal S = SpringMassScale * JointStiffness * Dt * Dt;
+	const uniform FReal D = SpringMassScale * JointDamping * Dt;
+	const uniform FReal Multiplier = ONE / ((S + D) * II + ONE);
+	const uniform FReal DLambda = M->SolverStiffness * Multiplier * (S * Delta - D * VelDt - Lambda);
 	
 	const uniform FVector DP0 = (M->InvMs[0] * DLambda) * Axis;
 	const uniform FVector DP1 = (-M->InvMs[1] * DLambda) * Axis;
@@ -300,36 +294,36 @@ export void ApplyRotationConstraintSoftKD(
 	uniform FPBDJointSolver * uniform M,
 	const uniform int KIndex,
 	const uniform int DIndex,
-	const uniform float Dt,
-	const uniform float JointStiffness,
-	const uniform float JointDamping,
+	const uniform FReal Dt,
+	const uniform FReal JointStiffness,
+	const uniform FReal JointDamping,
 	const uniform bool bAccelerationMode,
 	const uniform FVector& Axis,
-	const uniform float Angle,
-	const uniform float AngVelTarget,
-	uniform float& Lambda)
+	const uniform FReal Angle,
+	const uniform FReal AngVelTarget,
+	uniform FReal& Lambda)
 {
 	// World-space inverse mass
 	const uniform FVector IA1 = Multiply(M->InvIs[DIndex], Axis);
 
 	// Joint-space inverse mass
-	uniform float II1 = VectorDot(Axis, IA1);
-	const uniform float II = II1;
+	uniform FReal II1 = VectorDot(Axis, IA1);
+	const uniform FReal II = II1;
 
 	// Damping angular velocity
-	uniform float AngVelDt = 0;
+	uniform FReal AngVelDt = 0;
 	if (JointDamping > KINDA_SMALL_NUMBER)
 	{
-		const uniform FVector W0Dt = CalculateAngularVelocity(M->InitRs[KIndex], M->Rs[KIndex], 1.0f);
-		const uniform FVector W1Dt = CalculateAngularVelocity(M->InitRs[DIndex], M->Rs[DIndex], 1.0f);
+		const uniform FVector W0Dt = CalculateAngularVelocity(M->InitRs[KIndex], M->Rs[KIndex], ONE);
+		const uniform FVector W1Dt = CalculateAngularVelocity(M->InitRs[DIndex], M->Rs[DIndex], ONE);
 		AngVelDt = AngVelTarget * Dt + VectorDot(Axis, W0Dt - W1Dt);
 	}
 
-	const uniform float SpringMassScale = (bAccelerationMode) ? 1.0f / II : 1.0f;
-	const uniform float S = SpringMassScale * JointStiffness * Dt * Dt;
-	const uniform float D = SpringMassScale * JointDamping * Dt;
-	const uniform float Multiplier = 1.0f / ((S + D) * II + 1.0f);
-	const uniform float DLambda = M->SolverStiffness * Multiplier * (S * Angle - D * AngVelDt - Lambda);
+	const uniform FReal SpringMassScale = (bAccelerationMode) ? ONE / II : ONE;
+	const uniform FReal S = SpringMassScale * JointStiffness * Dt * Dt;
+	const uniform FReal D = SpringMassScale * JointDamping * Dt;
+	const uniform FReal Multiplier = ONE / ((S + D) * II + ONE);
+	const uniform FReal DLambda = M->SolverStiffness * Multiplier * (S * Angle - D * AngVelDt - Lambda);
 
 	//const FVec3 DR1 = IA1 * -DLambda;
 	const uniform FVector DR1 = Axis * -(DLambda * II1);
@@ -342,38 +336,38 @@ export void ApplyRotationConstraintSoftKD(
 
 export void ApplyRotationConstraintSoftDD(
 	uniform FPBDJointSolver * uniform M,
-	const uniform float Dt,
-	const uniform float JointStiffness,
-	const uniform float JointDamping,
+	const uniform FReal Dt,
+	const uniform FReal JointStiffness,
+	const uniform FReal JointDamping,
 	const uniform bool bAccelerationMode,
 	const uniform FVector& Axis,
-	const uniform float Angle,
-	const uniform float AngVelTarget,
-	uniform float& Lambda)
+	const uniform FReal Angle,
+	const uniform FReal AngVelTarget,
+	uniform FReal& Lambda)
 {
 	// World-space inverse mass
 	const uniform FVector IA0 = Multiply(M->InvIs[0], Axis);
 	const uniform FVector IA1 = Multiply(M->InvIs[1], Axis);
 
 	// Joint-space inverse mass
-	uniform float II0 = VectorDot(Axis, IA0);
-	uniform float II1 = VectorDot(Axis, IA1);
-	const uniform float II = (II0 + II1);
+	uniform FReal II0 = VectorDot(Axis, IA0);
+	uniform FReal II1 = VectorDot(Axis, IA1);
+	const uniform FReal II = (II0 + II1);
 
 	// Damping angular velocity
-	uniform float AngVelDt = 0;
+	uniform FReal AngVelDt = 0;
 	if (JointDamping > KINDA_SMALL_NUMBER)
 	{
-		const uniform FVector W0Dt = CalculateAngularVelocity(M->InitRs[0], M->Rs[0], 1.0f);
-		const uniform FVector W1Dt = CalculateAngularVelocity(M->InitRs[1], M->Rs[1], 1.0f);
+		const uniform FVector W0Dt = CalculateAngularVelocity(M->InitRs[0], M->Rs[0], ONE);
+		const uniform FVector W1Dt = CalculateAngularVelocity(M->InitRs[1], M->Rs[1], ONE);
 		AngVelDt = AngVelTarget * Dt + VectorDot(Axis, W0Dt - W1Dt);
 	}
 
-	const uniform float SpringMassScale = (bAccelerationMode) ? 1.0f / II : 1.0f;
-	const uniform float S = SpringMassScale * JointStiffness * Dt * Dt;
-	const uniform float D = SpringMassScale * JointDamping * Dt;
-	const uniform float Multiplier = 1.0f / ((S + D) * II + 1.0f);
-	const uniform float DLambda = M->SolverStiffness * Multiplier * (S * Angle - D * AngVelDt - Lambda);
+	const uniform FReal SpringMassScale = (bAccelerationMode) ? ONE / II : ONE;
+	const uniform FReal S = SpringMassScale * JointStiffness * Dt * Dt;
+	const uniform FReal D = SpringMassScale * JointDamping * Dt;
+	const uniform FReal Multiplier = ONE / ((S + D) * II + ONE);
+	const uniform FReal DLambda = M->SolverStiffness * Multiplier * (S * Angle - D * AngVelDt - Lambda);
 
 	//const FVec3 DR0 = IA0 * DLambda;
 	//const FVec3 DR1 = IA1 * -DLambda;
@@ -392,9 +386,8 @@ export void ApplyPointPositionConstraintKD(
 	const uniform int KIndex,
 	const uniform int DIndex,
 	const uniform FVector &CX,
-	const uniform float InStiffness)
+	const uniform FReal InStiffness)
 {
-
 	// Calculate constraint correction
 	const uniform FMatrix33 M1 = ComputeJointFactorMatrix(M->Xs[DIndex] - M->Ps[DIndex], M->InvIs[DIndex], M->InvMs[DIndex]);
 	const uniform FMatrix33 MI = MatrixInverse(M1);
@@ -412,11 +405,11 @@ export void ApplyPointPositionConstraintKD(
 export void ApplyPointPositionConstraintDD(
 	uniform FPBDJointSolver * uniform M,
 	const uniform FVector &CX,
-	const uniform float InStiffness)
+	const uniform FReal InStiffness)
 {
 	// Calculate constraint correction
-	uniform FMatrix33 M0 = ComputeJointFactorMatrix(M->Xs[0] - M->Ps[0], M->InvIs[0], M->InvMs[0]);
-	uniform FMatrix33 M1 = ComputeJointFactorMatrix(M->Xs[1] - M->Ps[1], M->InvIs[1], M->InvMs[1]);
+	const uniform FMatrix33 M0 = ComputeJointFactorMatrix(M->Xs[0] - M->Ps[0], M->InvIs[0], M->InvMs[0]);
+	const uniform FMatrix33 M1 = ComputeJointFactorMatrix(M->Xs[1] - M->Ps[1], M->InvIs[1], M->InvMs[1]);
 	const uniform FMatrix33 MI = MatrixInverse(AddAB(M0, M1));
 	const uniform FVector DX = InStiffness * Multiply(MI, CX);
 
