LWC: Enable ISPC for LWC double types. Submitted on behalf of jeff.rous[REDACTED_HANDLE]

diff --git a/Engine/Source/Runtime/Experimental/Chaos/Private/Chaos/PBDAxialSpringConstraints.ispc b/Engine/Source/Runtime/Experimental/Chaos/Private/Chaos/PBDAxialSpringConstraints.ispc
--- a/Engine/Source/Runtime/Experimental/Chaos/Private/Chaos/PBDAxialSpringConstraints.ispc
+++ b/Engine/Source/Runtime/Experimental/Chaos/Private/Chaos/PBDAxialSpringConstraints.ispc
@@ -7,10 +7,10 @@
 export void ApplyAxialSpringConstraints(uniform FVector P[],
 										const uniform FIntVector AllConstraints[],
 										const uniform int32 ActiveConstraints[],
-										const uniform float InvM[],
-										const uniform float Barys[],
-										const uniform float Dists[],
-										const uniform float Stiffness,
+										const uniform FReal InvM[],
+										const uniform FReal Barys[],
+										const uniform FReal Dists[],
+										const uniform FReal Stiffness,
 										const uniform int32 NumConstraints)
 {
 	varying int32 ActiveConstraint;
@@ -34,9 +34,9 @@ export void ApplyAxialSpringConstraints(uniform FVector P[],
 #endif
 
 		#pragma ignore warning(perf)
-		const varying float Bary = Barys[ActiveConstraint];
+		const varying FReal Bary = Barys[ActiveConstraint];
 		#pragma ignore warning(perf)
-		const varying float Dist = Dists[ActiveConstraint];
+		const varying FReal Dist = Dists[ActiveConstraint];
 
 		const varying int32 i1 = Constraint.V[0];
 		const varying int32 i2 = Constraint.V[1];
@@ -59,21 +59,21 @@ export void ApplyAxialSpringConstraints(uniform FVector P[],
 		const FVector P0 = (P2 - P3) * Bary + P3;
 		const FVector Difference = P1 - P0;
 
-		const float Distance = VectorSize(Difference);
-		const float Multiplier = 2 / (max(Bary, 1 - Bary) + 1);
+		const FReal Distance = VectorSize(Difference);
+		const FReal Multiplier = 2 / (max(Bary, 1 - Bary) + 1);
 
 		#pragma ignore warning(perf)
-		const varying float M1 = InvM[i1];
+		const varying FReal M1 = InvM[i1];
 		#pragma ignore warning(perf)
-		const varying float M2 = InvM[i2];
+		const varying FReal M2 = InvM[i2];
 		#pragma ignore warning(perf)
-		const varying float M3 = InvM[i3];
+		const varying FReal M3 = InvM[i3];
 
 		const FVector Direction = Difference / Distance;
 		const FVector InnerDelta = (Distance - Dist) * Direction;
 
-		const float PInvMass = M3 * (1 - Bary) + M2 * Bary;
-		const float CombinedInvMass = PInvMass + M1;
+		const FReal PInvMass = M3 * (1 - Bary) + M2 * Bary;
+		const FReal CombinedInvMass = PInvMass + M1;
 		const FVector ComputedDelta = Stiffness * InnerDelta / CombinedInvMass;
 
 		const FVector Delta = VectorSelect(((M1 == 0 && PInvMass == 0) || Distance <= SMALL_NUMBER), ZeroVector, ComputedDelta);
@@ -96,11 +96,11 @@ export void ApplyAxialSpringConstraints(uniform FVector P[],
 export void ApplyAxialSpringConstraintsWithWeightMaps(uniform FVector P[],
 										const uniform FIntVector AllConstraints[],
 										const uniform int32 ActiveConstraints[],
-										const uniform float InvM[],
-										const uniform float Barys[],
-										const uniform float Dists[],
+										const uniform FReal InvM[],
+										const uniform FReal Barys[],
+										const uniform FReal Dists[],
 										const uniform uint8 StiffnessIndices[],
-										const uniform float StiffnessTable[],
+										const uniform FReal StiffnessTable[],
 										const uniform int32 NumConstraints)
 {
 	varying int32 ActiveConstraint;
@@ -124,9 +124,9 @@ export void ApplyAxialSpringConstraintsWithWeightMaps(uniform FVector P[],
 #endif
 
 		#pragma ignore warning(perf)
-		const varying float Bary = Barys[ActiveConstraint];
+		const varying FReal Bary = Barys[ActiveConstraint];
 		#pragma ignore warning(perf)
-		const varying float Dist = Dists[ActiveConstraint];
+		const varying FReal Dist = Dists[ActiveConstraint];
 		#pragma ignore warning(perf)
 		const varying uint8 StiffnessIndex = StiffnessIndices[ActiveConstraint];
 
@@ -151,24 +151,24 @@ export void ApplyAxialSpringConstraintsWithWeightMaps(uniform FVector P[],
 		const FVector P0 = (P2 - P3) * Bary + P3;
 		const FVector Difference = P1 - P0;
 
-		const float Distance = VectorSize(Difference);
-		const float Multiplier = 2 / (max(Bary, 1 - Bary) + 1);
+		const FReal Distance = VectorSize(Difference);
+		const FReal Multiplier = 2 / (max(Bary, 1 - Bary) + 1);
 
 		#pragma ignore warning(perf)
-		const varying float M1 = InvM[i1];
+		const varying FReal M1 = InvM[i1];
 		#pragma ignore warning(perf)
-		const varying float M2 = InvM[i2];
+		const varying FReal M2 = InvM[i2];
 		#pragma ignore warning(perf)
-		const varying float M3 = InvM[i3];
+		const varying FReal M3 = InvM[i3];
 
 		#pragma ignore warning(perf)
-		const varying float Stiffness = StiffnessTable[StiffnessIndex];
+		const varying FReal Stiffness = StiffnessTable[StiffnessIndex];
 
 		const FVector Direction = Difference / Distance;
 		const FVector InnerDelta = (Distance - Dist) * Direction;
 
-		const float PInvMass = M3 * (1 - Bary) + M2 * Bary;
-		const float CombinedInvMass = PInvMass + M1;
+		const FReal PInvMass = M3 * (1 - Bary) + M2 * Bary;
+		const FReal CombinedInvMass = PInvMass + M1;
 		const FVector ComputedDelta = Stiffness * InnerDelta / CombinedInvMass;
 
 		const FVector Delta = VectorSelect(((M1 == 0 && PInvMass == 0) || Distance <= SMALL_NUMBER), ZeroVector, ComputedDelta);
diff --git a/Engine/Source/Runtime/Experimental/Chaos/Private/Chaos/PBDCollisionConstraints.ispc b/Engine/Source/Runtime/Experimental/Chaos/Private/Chaos/PBDCollisionConstraints.ispc
--- a/Engine/Source/Runtime/Experimental/Chaos/Private/Chaos/PBDCollisionConstraints.ispc
+++ b/Engine/Source/Runtime/Experimental/Chaos/Private/Chaos/PBDCollisionConstraints.ispc
@@ -2,105 +2,38 @@
 
 #define EXPLICIT_VECTOR4 1
 
-#include "Chaos/Matrix33.isph"
-#include "Math/Quat.isph"
 #include "Math/Vector.isph"
-
-#define SMALL_NUMBER		(1.e-8f)
-#define KINDA_SMALL_NUMBER	(1.e-4f)
-
-struct FTransform
-{
-	FVector4 Rotation;
-	FVector4 Translation;
-	FVector4 Scale3D;
-};
+#include "Math/Quat.isph"
+#include "Math/Transform.isph"
+#include "Chaos/Matrix33.isph"
 
 struct TArrayND
 {
 	uniform FIntVector MCounts;
-	uniform float * uniform MArray;
+	uniform FReal * uniform MArray;
 };
 
-inline static FVector TransformPositionNoScale(const uniform FTransform &T, const FVector& V)
-{
-	FVector4 InputVectorW0 = SetVector4(V.V[0], V.V[1], V.V[2], 0.0f);
-
-	//Transform using QST is following
-	//QST(P) = Q.Rotate(S*P) + T where Q = quaternion, S = 1.0f, T = translation
-
-	//RotatedVec = Q.Rotate(V.X, V.Y, V.Z, 0.f)
-	const FVector4 RotatedVec = VectorQuaternionRotateVector(T.Rotation, InputVectorW0);
-	const FVector4 TranslatedVec = RotatedVec + T.Translation;
-
-	return SetVector(TranslatedVec.V[0], TranslatedVec.V[1], TranslatedVec.V[2]);
-}
-
-inline static float SquaredSum(const FVector &Vec)
-{
-	float SquareSum = 0.0f;
-	for(uniform int i = 0; i < 3; i++)
-	{
-		SquareSum += (Vec.V[i] * Vec.V[i]);
-	}
-	return SquareSum;
-}
-
-inline static float VectorSize2(const FVector &Normal)
-{
-	return rcp(rsqrt(SquaredSum(Normal)));
-}
-
-inline static float BoxPhi(const uniform FVector &BoxMin, const uniform FVector &BoxMax, const FVector& LocalPoint)
+inline static FReal BoxPhi(const uniform FVector &BoxMin, const uniform FVector &BoxMax, const FVector& LocalPoint)
 {
 	const FVector MaxDists = LocalPoint - BoxMax;
 	const FVector MinDists = BoxMin - LocalPoint;
-	float Phi;
+	FReal Phi;
 	if (IsVectorLessEqual(LocalPoint, BoxMax) && IsVectorGreaterEqual(LocalPoint, BoxMin))
 	{
 		const FVector LocalMax = VectorMax(MinDists, MaxDists);
-		if (LocalMax.V[0] > LocalMax.V[1])
-		{
-			if (LocalMax.V[0] > LocalMax.V[2])
-			{
-				Phi = LocalMax.V[0];
-			}
-			else
-			{
-				Phi = LocalMax.V[2];
-			}
-		}
-		else
-		{
-			if (LocalMax.V[1] > LocalMax.V[2])
-			{
-				Phi = LocalMax.V[1];
-			}
-			else
-			{
-				Phi = LocalMax.V[2];
-			}
-		}
+		Phi = select(LocalMax.V[0] > LocalMax.V[1], 
+				select(LocalMax.V[0] > LocalMax.V[2], LocalMax.V[0], LocalMax.V[2]), 
+				select(LocalMax.V[1] > LocalMax.V[2], LocalMax.V[1], LocalMax.V[2]));
 	}
 	else
 	{
 		FVector Normal;
-		for(uniform int i = 0; i < 3; i++)
-		{
-			if (MaxDists.V[i] > 0)
-			{
-				Normal.V[i] = MaxDists.V[i];
-			}
-			else if (MinDists.V[i] > 0)
-			{
-				Normal.V[i] = -MinDists.V[i];
-			}
-			else
-			{
-				Normal.V[i] = 0;
-			}
-		}
-		Phi = VectorSize2(Normal);
+
+		Normal.V[0] = select(MaxDists.V[0] > ZERO, MaxDists.V[0], select(MinDists.V[0] > ZERO, -MinDists.V[0], ZERO));
+		Normal.V[1] = select(MaxDists.V[1] > ZERO, MaxDists.V[1], select(MinDists.V[1] > ZERO, -MinDists.V[1], ZERO));
+		Normal.V[2] = select(MaxDists.V[2] > ZERO, MaxDists.V[2], select(MinDists.V[2] > ZERO, -MinDists.V[2], ZERO));
+
+		Phi = VectorSize(Normal);
 	}
 
 	return Phi;
@@ -111,22 +44,22 @@ export void SampleBoxNormalAverage(const uniform FVector &BoxMin,
 									const uniform FTransform &SampleToObjectTransform, 
 									const uniform FVector SampleParticles[],
 									const uniform int PotentialParticles[],
-									const uniform float Thickness, 
-									uniform float &TotalThickness, 
+									const uniform FReal Thickness, 
+									uniform FReal &TotalThickness, 
 									uniform FVector &ConstraintLocation,
 									const uniform unsigned int NumParticles)
 {
-	FVector ConstraintLocationSum = SetVector(0.0f, 0.0f, 0.0f);
-	float ThicknessSum = 0.0f;
+	FVector ConstraintLocationSum = ZeroVector;
+	FReal ThicknessSum = 0;
 
 	foreach(i = 0 ... NumParticles)
 	{
 		FVector Particle = VectorGather(&SampleParticles[PotentialParticles[i]]);
 
 		FVector LocalPoint = TransformPositionNoScale(SampleToObjectTransform, Particle);
-		float LocalPhi = BoxPhi(BoxMin, BoxMax, LocalPoint);
+		FReal LocalPhi = BoxPhi(BoxMin, BoxMax, LocalPoint);
 
-		float LocalThickness = LocalPhi - Thickness;
+		FReal LocalThickness = LocalPhi - Thickness;
 		if(LocalThickness < -KINDA_SMALL_NUMBER)
 		{
 			ConstraintLocationSum = ConstraintLocationSum + (LocalPoint * LocalThickness);
@@ -134,7 +67,7 @@ export void SampleBoxNormalAverage(const uniform FVector &BoxMin,
 		}
 	}
 
-	ConstraintLocation = ConstraintLocation + SetVector(reduce_add(ConstraintLocationSum.V[0]), reduce_add(ConstraintLocationSum.V[1]), reduce_add(ConstraintLocationSum.V[2]));
+	ConstraintLocation = ConstraintLocation + VectorReduceAdd(ConstraintLocationSum);
 	TotalThickness = TotalThickness + reduce_add(ThicknessSum);
 }
 
@@ -144,7 +77,7 @@ export uniform bool SampleBoxNoNormal(const uniform FVector &BoxMin,
 									const uniform FVector SampleParticles[],
 									const uniform int PotentialParticles[],
 									uniform int &DeepestParticle,
-									uniform float &Phi,
+									uniform FReal &Phi,
 									const uniform unsigned int NumParticles)
 {
 	uniform bool bPhiFound = false;
@@ -156,7 +89,7 @@ export uniform bool SampleBoxNoNormal(const uniform FVector &BoxMin,
 		FVector Particle = VectorGather(&SampleParticles[ParticleNum]);
 
 		FVector LocalPoint = TransformPositionNoScale(SampleToObjectTransform, Particle);
-		float LocalPhi = BoxPhi(BoxMin, BoxMax, LocalPoint);
+		FReal LocalPhi = BoxPhi(BoxMin, BoxMax, LocalPoint);
 
 		if(LocalPhi < Phi)
 		{
@@ -180,22 +113,22 @@ export void SampleBoxNormalAverageAll(const uniform FVector &BoxMin,
 									const uniform FVector &BoxMax, 
 									const uniform FTransform &SampleToObjectTransform, 
 									const uniform FVector SampleParticles[],
-									const uniform float Thickness, 
-									uniform float &TotalThickness, 
+									const uniform FReal Thickness, 
+									uniform FReal &TotalThickness, 
 									uniform FVector &ConstraintLocation,
 									const uniform unsigned int NumParticles)
 {
-	FVector ConstraintLocationSum = SetVector(0.0f, 0.0f, 0.0f);
-	float ThicknessSum = 0.0f;
+	FVector ConstraintLocationSum = ZeroVector;
+	FReal ThicknessSum = 0;
 
 	foreach(i = 0 ... NumParticles)
 	{
 		FVector Particle = VectorLoad(&SampleParticles[extract(i, 0)]);
 
 		FVector LocalPoint = TransformPositionNoScale(SampleToObjectTransform, Particle);
-		float LocalPhi = BoxPhi(BoxMin, BoxMax, LocalPoint);
+		FReal LocalPhi = BoxPhi(BoxMin, BoxMax, LocalPoint);
 
-		float LocalThickness = LocalPhi - Thickness;
+		FReal LocalThickness = LocalPhi - Thickness;
 		if(LocalThickness < -KINDA_SMALL_NUMBER)
 		{
 			ConstraintLocationSum = ConstraintLocationSum + (LocalPoint * LocalThickness);
@@ -203,7 +136,7 @@ export void SampleBoxNormalAverageAll(const uniform FVector &BoxMin,
 		}
 	}
 
-	ConstraintLocation = ConstraintLocation + SetVector(reduce_add(ConstraintLocationSum.V[0]), reduce_add(ConstraintLocationSum.V[1]), reduce_add(ConstraintLocationSum.V[2]));
+	ConstraintLocation = ConstraintLocation + VectorReduceAdd(ConstraintLocationSum);
 	TotalThickness = TotalThickness + reduce_add(ThicknessSum);
 }
 
@@ -212,7 +145,7 @@ export uniform bool SampleBoxNoNormalAll(const uniform FVector &BoxMin,
 										const uniform FTransform &SampleToObjectTransform,
 										const uniform FVector SampleParticles[],
 										uniform int &DeepestParticle,
-										uniform float &Phi,
+										uniform FReal &Phi,
 										const uniform unsigned int NumParticles)
 {
 	uniform bool bPhiFound = false;
@@ -222,7 +155,7 @@ export uniform bool SampleBoxNoNormalAll(const uniform FVector &BoxMin,
 		FVector Particle = VectorLoad(&SampleParticles[extract(i, 0)]);
 
 		FVector LocalPoint = TransformPositionNoScale(SampleToObjectTransform, Particle);
-		float LocalPhi = BoxPhi(BoxMin, BoxMax, LocalPoint);
+		FReal LocalPhi = BoxPhi(BoxMin, BoxMax, LocalPoint);
 
 		if(LocalPhi < Phi)
 		{
@@ -244,25 +177,15 @@ export uniform bool SampleBoxNoNormalAll(const uniform FVector &BoxMin,
 inline static FVector ClampMinusHalf(const uniform FVector &Min, const uniform FVector &Max, const FVector& LocalPoint)
 {
 	FVector Result;
-	for (uniform int i = 0; i < 3; ++i)
-	{
-		if (LocalPoint.V[i] > Max.V[i])
-		{
-			Result.V[i] = Max.V[i];
-		}
-		else if (LocalPoint.V[i] < Min.V[i])
-		{
-			Result.V[i] = Min.V[i];
-		}
-		else
-		{
-			Result.V[i] = LocalPoint.V[i];
-		}
-	}
+
+	Result.V[0] = select(LocalPoint.V[0] > Max.V[0], Max.V[0], select(LocalPoint.V[0] < Min.V[0], Min.V[0], LocalPoint.V[0]));
+	Result.V[1] = select(LocalPoint.V[1] > Max.V[1], Max.V[1], select(LocalPoint.V[1] < Min.V[1], Min.V[1], LocalPoint.V[1]));
+	Result.V[2] = select(LocalPoint.V[2] > Max.V[2], Max.V[2], select(LocalPoint.V[2] < Min.V[2], Min.V[2], LocalPoint.V[2]));
+
 	return Result;
 }
 
-inline static float GetCell(const uniform TArrayND* uniform LevelSetPhiArray, const FIntVector &Vec)
+inline static FReal GetCell(const uniform TArrayND* uniform LevelSetPhiArray, const FIntVector &Vec)
 {
 	const int Offset = (Vec.V[0] * LevelSetPhiArray->MCounts.V[1] + Vec.V[1]) * LevelSetPhiArray->MCounts.V[2] + Vec.V[2];
 
@@ -270,7 +193,7 @@ inline static float GetCell(const uniform TArrayND* uniform LevelSetPhiArray, co
 	return LevelSetPhiArray->MArray[Offset];
 }
 
-inline static float LinearlyInterpolate1D(const float& Prev, const float& Next, const float Alpha)
+inline static FReal LinearlyInterpolate1D(const FReal& Prev, const FReal& Next, const FReal Alpha)
 {
 	return Prev + (Next - Prev) * Alpha;
 }
@@ -280,7 +203,7 @@ inline static FVector Location(const uniform FVector &MinCorner, const uniform F
 	return Dx * X + MinCorner + (Dx / 2);
 }
 
-inline static float LinearlyInterpolate(const uniform FVector &MinCorner,
+inline static FReal LinearlyInterpolate(const uniform FVector &MinCorner,
 							const uniform FVector &MaxCorner,
 							const uniform FVector &Dx,
 							const uniform FIntVector &Cells,
@@ -312,12 +235,12 @@ inline static float LinearlyInterpolate(const uniform FVector &MinCorner,
 		}
 	}
 
-	const float interpx1 = LinearlyInterpolate1D(GetCell(LevelSetPhiArray, CellPrev), GetCell(LevelSetPhiArray, CellPrev + SetIntVector(1, 0, 0)), Alpha.V[0]);
-	const float interpx2 = LinearlyInterpolate1D(GetCell(LevelSetPhiArray, CellPrev + SetIntVector(0, 1, 0)), GetCell(LevelSetPhiArray, CellPrev + SetIntVector(1, 1, 0)), Alpha.V[0]);
-	const float interpx3 = LinearlyInterpolate1D(GetCell(LevelSetPhiArray, CellPrev + SetIntVector(0, 0, 1)), GetCell(LevelSetPhiArray, CellPrev + SetIntVector(1, 0, 1)), Alpha.V[0]);
-	const float interpx4 = LinearlyInterpolate1D(GetCell(LevelSetPhiArray, CellPrev + SetIntVector(0, 1, 1)), GetCell(LevelSetPhiArray, CellPrev + SetIntVector(1, 1, 1)), Alpha.V[0]);
-	const float interpy1 = LinearlyInterpolate1D(interpx1, interpx2, Alpha.V[1]);
-	const float interpy2 = LinearlyInterpolate1D(interpx3, interpx4, Alpha.V[1]);
+	const FReal interpx1 = LinearlyInterpolate1D(GetCell(LevelSetPhiArray, CellPrev), GetCell(LevelSetPhiArray, CellPrev + SetIntVector(1, 0, 0)), Alpha.V[0]);
+	const FReal interpx2 = LinearlyInterpolate1D(GetCell(LevelSetPhiArray, CellPrev + SetIntVector(0, 1, 0)), GetCell(LevelSetPhiArray, CellPrev + SetIntVector(1, 1, 0)), Alpha.V[0]);
+	const FReal interpx3 = LinearlyInterpolate1D(GetCell(LevelSetPhiArray, CellPrev + SetIntVector(0, 0, 1)), GetCell(LevelSetPhiArray, CellPrev + SetIntVector(1, 0, 1)), Alpha.V[0]);
+	const FReal interpx4 = LinearlyInterpolate1D(GetCell(LevelSetPhiArray, CellPrev + SetIntVector(0, 1, 1)), GetCell(LevelSetPhiArray, CellPrev + SetIntVector(1, 1, 1)), Alpha.V[0]);
+	const FReal interpy1 = LinearlyInterpolate1D(interpx1, interpx2, Alpha.V[1]);
+	const FReal interpy2 = LinearlyInterpolate1D(interpx3, interpx4, Alpha.V[1]);
 	return LinearlyInterpolate1D(interpy1, interpy2, Alpha.V[2]);
 }
 
@@ -329,16 +252,16 @@ export void SampleLevelSetNormalAverage(const uniform FVector &GridMinCorner,
 										const uniform FTransform &SampleToObjectTransform, 
 										const uniform FVector SampleParticles[],
 										const uniform int PotentialParticles[],
-										const uniform float Thickness, 
-										uniform float &TotalThickness, 
+										const uniform FReal Thickness, 
+										uniform FReal &TotalThickness, 
 										uniform FVector &ConstraintLocation,
 										const uniform unsigned int NumParticles)
 {
-	FVector ConstraintLocationSum = SetVector(0.0f, 0.0f, 0.0f);
-	float ThicknessSum = 0.0f;
+	FVector ConstraintLocationSum = ZeroVector;
+	FReal ThicknessSum = 0;
 
-	uniform FVector GridClampMax = GridMaxCorner - GridDx * 0.5f;
-	uniform FVector GridClampMin = GridMinCorner + GridDx * 0.5f;
+	uniform FVector GridClampMax = GridMaxCorner - GridDx * ONE_HALF;
+	uniform FVector GridClampMin = GridMinCorner + GridDx * ONE_HALF;
 
 	foreach(i = 0 ... NumParticles)
 	{
@@ -346,15 +269,15 @@ export void SampleLevelSetNormalAverage(const uniform FVector &GridMinCorner,
 
 		FVector LocalPoint = TransformPositionNoScale(SampleToObjectTransform, Particle);
 		FVector Location = ClampMinusHalf(GridClampMin, GridClampMax, LocalPoint);
-		float SizeSquared = SquaredSum(Location - LocalPoint);
-		float LocalPhi = LinearlyInterpolate(GridMinCorner, GridMaxCorner, GridDx, GridCounts, LevelSetPhiArray, Location);
+		FReal SizeSquared = VectorSizeSquared(Location - LocalPoint);
+		FReal LocalPhi = LinearlyInterpolate(GridMinCorner, GridMaxCorner, GridDx, GridCounts, LevelSetPhiArray, Location);
 
 		if(SizeSquared)
 		{
 			LocalPhi = rcp(rsqrt(SizeSquared)) + LocalPhi;
 		}
 
-		float LocalThickness = LocalPhi - Thickness;
+		FReal LocalThickness = LocalPhi - Thickness;
 		if(LocalThickness < -KINDA_SMALL_NUMBER)
 		{
 			ConstraintLocationSum = ConstraintLocationSum + (LocalPoint * LocalThickness);
@@ -362,7 +285,7 @@ export void SampleLevelSetNormalAverage(const uniform FVector &GridMinCorner,
 		}
 	}
 
-	ConstraintLocation = ConstraintLocation + SetVector(reduce_add(ConstraintLocationSum.V[0]), reduce_add(ConstraintLocationSum.V[1]), reduce_add(ConstraintLocationSum.V[2]));
+	ConstraintLocation = ConstraintLocation + VectorReduceAdd(ConstraintLocationSum);
 	TotalThickness = TotalThickness + reduce_add(ThicknessSum);
 }
 
@@ -375,11 +298,11 @@ export uniform bool SampleLevelSetNoNormal(const uniform FVector &GridMinCorner,
 										const uniform FVector SampleParticles[],
 										const uniform int PotentialParticles[],
 										uniform int &DeepestParticle,
-										uniform float &Phi,
+										uniform FReal &Phi,
 										const uniform unsigned int NumParticles)
 {
-	uniform FVector GridClampMax = GridMaxCorner - GridDx * 0.5f;
-	uniform FVector GridClampMin = GridMinCorner + GridDx * 0.5f;
+	uniform FVector GridClampMax = GridMaxCorner - GridDx * ONE_HALF;
+	uniform FVector GridClampMin = GridMinCorner + GridDx * ONE_HALF;
 
 	uniform bool bPhiFound = false;
 
@@ -391,8 +314,8 @@ export uniform bool SampleLevelSetNoNormal(const uniform FVector &GridMinCorner,
 
 		FVector LocalPoint = TransformPositionNoScale(SampleToObjectTransform, Particle);
 		FVector Location = ClampMinusHalf(GridClampMin, GridClampMax, LocalPoint);
-		float SizeSquared = SquaredSum(Location - LocalPoint);
-		float LocalPhi = LinearlyInterpolate(GridMinCorner, GridMaxCorner, GridDx, GridCounts, LevelSetPhiArray, Location);
+		FReal SizeSquared = VectorSizeSquared(Location - LocalPoint);
+		FReal LocalPhi = LinearlyInterpolate(GridMinCorner, GridMaxCorner, GridDx, GridCounts, LevelSetPhiArray, Location);
 
 		if(SizeSquared)
 		{
@@ -423,16 +346,16 @@ export void SampleLevelSetNormalAverageAll(const uniform FVector &GridMinCorner,
 										const uniform TArrayND* uniform LevelSetPhiArray,
 										const uniform FTransform &SampleToObjectTransform, 
 										const uniform FVector SampleParticles[],
-										const uniform float Thickness, 
-										uniform float &TotalThickness, 
+										const uniform FReal Thickness, 
+										uniform FReal &TotalThickness, 
 										uniform FVector &ConstraintLocation,
 										const uniform unsigned int NumParticles)
 {
-	FVector ConstraintLocationSum = SetVector(0.0f, 0.0f, 0.0f);
-	float ThicknessSum = 0.0f;
+	FVector ConstraintLocationSum = ZeroVector;
+	FReal ThicknessSum = 0;
 
-	uniform FVector GridClampMax = GridMaxCorner - GridDx * 0.5f;
-	uniform FVector GridClampMin = GridMinCorner + GridDx * 0.5f;
+	uniform FVector GridClampMax = GridMaxCorner - GridDx * ONE_HALF;
+	uniform FVector GridClampMin = GridMinCorner + GridDx * ONE_HALF;
 
 	foreach(i = 0 ... NumParticles)
 	{
@@ -440,15 +363,15 @@ export void SampleLevelSetNormalAverageAll(const uniform FVector &GridMinCorner,
 
 		FVector LocalPoint = TransformPositionNoScale(SampleToObjectTransform, Particle);
 		FVector Location = ClampMinusHalf(GridClampMin, GridClampMax, LocalPoint);
-		float SizeSquared = SquaredSum(Location - LocalPoint);
-		float LocalPhi = LinearlyInterpolate(GridMinCorner, GridMaxCorner, GridDx, GridCounts, LevelSetPhiArray, Location);
+		FReal SizeSquared = VectorSizeSquared(Location - LocalPoint);
+		FReal LocalPhi = LinearlyInterpolate(GridMinCorner, GridMaxCorner, GridDx, GridCounts, LevelSetPhiArray, Location);
 
 		if(SizeSquared)
 		{
 			LocalPhi = rcp(rsqrt(SizeSquared)) + LocalPhi;
 		}
 
-		float LocalThickness = LocalPhi - Thickness;
+		FReal LocalThickness = LocalPhi - Thickness;
 		if(LocalThickness < -KINDA_SMALL_NUMBER)
 		{
 			ConstraintLocationSum = ConstraintLocationSum + (LocalPoint * LocalThickness);
@@ -456,7 +379,7 @@ export void SampleLevelSetNormalAverageAll(const uniform FVector &GridMinCorner,
 		}
 	}
 
-	ConstraintLocation = ConstraintLocation + SetVector(reduce_add(ConstraintLocationSum.V[0]), reduce_add(ConstraintLocationSum.V[1]), reduce_add(ConstraintLocationSum.V[2]));
+	ConstraintLocation = ConstraintLocation + VectorReduceAdd(ConstraintLocationSum);
 	TotalThickness = TotalThickness + reduce_add(ThicknessSum);
 }
 
@@ -468,11 +391,11 @@ export uniform bool SampleLevelSetNoNormalAll(const uniform FVector &GridMinCorn
 											const uniform FTransform &SampleToObjectTransform, 
 											const uniform FVector SampleParticles[],
 											uniform int &DeepestParticle,
-											uniform float &Phi,
+											uniform FReal &Phi,
 											const uniform unsigned int NumParticles)
 {
-	uniform FVector GridClampMax = GridMaxCorner - GridDx * 0.5f;
-	uniform FVector GridClampMin = GridMinCorner + GridDx * 0.5f;
+	uniform FVector GridClampMax = GridMaxCorner - GridDx * ONE_HALF;
+	uniform FVector GridClampMin = GridMinCorner + GridDx * ONE_HALF;
 
 	uniform bool bPhiFound = false;
 
@@ -482,8 +405,8 @@ export uniform bool SampleLevelSetNoNormalAll(const uniform FVector &GridMinCorn
 
 		FVector LocalPoint = TransformPositionNoScale(SampleToObjectTransform, Particle);
 		FVector Location = ClampMinusHalf(GridClampMin, GridClampMax, LocalPoint);
-		float SizeSquared = SquaredSum(Location - LocalPoint);
-		float LocalPhi = LinearlyInterpolate(GridMinCorner, GridMaxCorner, GridDx, GridCounts, LevelSetPhiArray, Location);
+		FReal SizeSquared = VectorSizeSquared(Location - LocalPoint);
+		FReal LocalPhi = LinearlyInterpolate(GridMinCorner, GridMaxCorner, GridDx, GridCounts, LevelSetPhiArray, Location);
 
 		if(SizeSquared)
 		{
@@ -507,31 +430,26 @@ export uniform bool SampleLevelSetNoNormalAll(const uniform FVector &GridMinCorn
 	return bPhiFound;
 }
 
-inline static float DotProduct(const FVector &V1, const uniform FVector &V2)
-{
-	return V1.V[0] * V2.V[0] + V1.V[1] * V2.V[1] + V1.V[2] * V2.V[2];
-}
-
 export void SamplePlaneNormalAverageAll(const uniform FVector &ObjectNormal,
 									const uniform FVector &ObjectMX,
 									const uniform FTransform &SampleToObjectTransform, 
 									const uniform FVector SampleParticles[],
-									const uniform float Thickness, 
-									uniform float &TotalThickness, 
+									const uniform FReal Thickness, 
+									uniform FReal &TotalThickness, 
 									uniform FVector &ConstraintLocation,
 									const uniform unsigned int NumParticles)
 {
-	FVector ConstraintLocationSum = SetVector(0.0f, 0.0f, 0.0f);
-	float ThicknessSum = 0.0f;
+	FVector ConstraintLocationSum = ZeroVector;
+	FReal ThicknessSum = 0;
 
 	foreach(i = 0 ... NumParticles)
 	{
 		FVector Particle = VectorLoad(&SampleParticles[extract(i, 0)]);
 
 		FVector LocalPoint = TransformPositionNoScale(SampleToObjectTransform, Particle);
-		float LocalPhi = DotProduct(LocalPoint - ObjectMX, ObjectNormal);
+		FReal LocalPhi = VectorDot(LocalPoint - ObjectMX, ObjectNormal);
 
-		float LocalThickness = LocalPhi - Thickness;
+		FReal LocalThickness = LocalPhi - Thickness;
 		if(LocalThickness < -KINDA_SMALL_NUMBER)
 		{
 			ConstraintLocationSum = ConstraintLocationSum + (LocalPoint * LocalThickness);
@@ -539,7 +457,7 @@ export void SamplePlaneNormalAverageAll(const uniform FVector &ObjectNormal,
 		}
 	}
 
-	ConstraintLocation = ConstraintLocation + SetVector(reduce_add(ConstraintLocationSum.V[0]), reduce_add(ConstraintLocationSum.V[1]), reduce_add(ConstraintLocationSum.V[2]));
+	ConstraintLocation = ConstraintLocation + VectorReduceAdd(ConstraintLocationSum);
 	TotalThickness = TotalThickness + reduce_add(ThicknessSum);
 }
 
@@ -548,7 +466,7 @@ export uniform bool SamplePlaneNoNormalAll(const uniform FVector &ObjectNormal,
 										const uniform FTransform &SampleToObjectTransform, 
 										const uniform FVector SampleParticles[],
 										uniform int &DeepestParticle,
-										uniform float &Phi,
+										uniform FReal &Phi,
 										const uniform unsigned int NumParticles)
 {
 	uniform bool bPhiFound = false;
@@ -558,7 +476,7 @@ export uniform bool SamplePlaneNoNormalAll(const uniform FVector &ObjectNormal,
 		FVector Particle = VectorLoad(&SampleParticles[extract(i, 0)]);
 
 		FVector LocalPoint = TransformPositionNoScale(SampleToObjectTransform, Particle);
-		float LocalPhi = DotProduct(LocalPoint - ObjectMX, ObjectNormal);
+		FReal LocalPhi = VectorDot(LocalPoint - ObjectMX, ObjectNormal);
 
 		if(LocalPhi < Phi)
 		{
@@ -577,26 +495,26 @@ export uniform bool SamplePlaneNoNormalAll(const uniform FVector &ObjectNormal,
 	return bPhiFound;
 }
 
-export void SampleSphereNormalAverageAll(const uniform float &ObjectRadius,
+export void SampleSphereNormalAverageAll(const uniform FReal &ObjectRadius,
 									const uniform FVector &ObjectMX,
 									const uniform FTransform &SampleToObjectTransform, 
 									const uniform FVector SampleParticles[],
-									const uniform float Thickness, 
-									uniform float &TotalThickness, 
+									const uniform FReal Thickness, 
+									uniform FReal &TotalThickness, 
 									uniform FVector &ConstraintLocation,
 									const uniform unsigned int NumParticles)
 {
-	FVector ConstraintLocationSum = SetVector(0.0f, 0.0f, 0.0f);
-	float ThicknessSum = 0.0f;
+	FVector ConstraintLocationSum = ZeroVector;
+	FReal ThicknessSum = 0;
 
 	foreach(i = 0 ... NumParticles)
 	{
 		FVector Particle = VectorLoad(&SampleParticles[extract(i, 0)]);
 
 		FVector LocalPoint = TransformPositionNoScale(SampleToObjectTransform, Particle);
-		float LocalPhi = ObjectRadius - VectorSize2(LocalPoint - ObjectMX);
+		FReal LocalPhi = ObjectRadius - VectorSize(LocalPoint - ObjectMX);
 
-		float LocalThickness = LocalPhi - Thickness;
+		FReal LocalThickness = LocalPhi - Thickness;
 		if(LocalThickness < -KINDA_SMALL_NUMBER)
 		{
 			ConstraintLocationSum = ConstraintLocationSum + (LocalPoint * LocalThickness);
@@ -604,16 +522,16 @@ export void SampleSphereNormalAverageAll(const uniform float &ObjectRadius,
 		}
 	}
 
-	ConstraintLocation = ConstraintLocation + SetVector(reduce_add(ConstraintLocationSum.V[0]), reduce_add(ConstraintLocationSum.V[1]), reduce_add(ConstraintLocationSum.V[2]));
+	ConstraintLocation = ConstraintLocation + VectorReduceAdd(ConstraintLocationSum);
 	TotalThickness = TotalThickness + reduce_add(ThicknessSum);
 }
 
-export uniform bool SampleSphereNoNormalAll(const uniform float &ObjectRadius,
+export uniform bool SampleSphereNoNormalAll(const uniform FReal &ObjectRadius,
 										const uniform FVector &ObjectMX,
 										const uniform FTransform &SampleToObjectTransform, 
 										const uniform FVector SampleParticles[],
 										uniform int &DeepestParticle,
-										uniform float &Phi,
+										uniform FReal &Phi,
 										const uniform unsigned int NumParticles)
 {
 	uniform bool bPhiFound = false;
@@ -623,7 +541,7 @@ export uniform bool SampleSphereNoNormalAll(const uniform float &ObjectRadius,
 		FVector Particle = VectorLoad(&SampleParticles[extract(i, 0)]);
 
 		FVector LocalPoint = TransformPositionNoScale(SampleToObjectTransform, Particle);
-		float LocalPhi = ObjectRadius - VectorSize2(LocalPoint - ObjectMX);
+		FReal LocalPhi = ObjectRadius - VectorSize(LocalPoint - ObjectMX);
 
 		if(LocalPhi < Phi)
 		{
@@ -689,7 +607,7 @@ static inline uniform FMatrix33 ComputeWorldSpaceInertia(const uniform FVector4
 }
 
 // @todo(ccaulfield): dedupe
-static inline uniform FMatrix33 ComputeJointFactorMatrix(const uniform FVector& V, const uniform FMatrix33& M, const uniform float Im)
+static inline uniform FMatrix33 ComputeJointFactorMatrix(const uniform FVector& V, const uniform FMatrix33& M, const uniform FReal Im)
 {
 	// Rigid objects rotational contribution to the impulse.
 	// Vx*M*VxT+Im
diff --git a/Engine/Source/Runtime/Experimental/Chaos/Private/Chaos/PBDEvolution.ispc b/Engine/Source/Runtime/Experimental/Chaos/Private/Chaos/PBDEvolution.ispc
--- a/Engine/Source/Runtime/Experimental/Chaos/Private/Chaos/PBDEvolution.ispc
+++ b/Engine/Source/Runtime/Experimental/Chaos/Private/Chaos/PBDEvolution.ispc
@@ -7,26 +7,26 @@
 export void PostIterationUpdates(uniform FVector ParticlesV[],
 								uniform FVector ParticlesX[],
 								const uniform FVector ParticlesP[],
-								const uniform float Dt,
+								const uniform FReal Dt,
 								const uniform int32 Offset,
 								const uniform int32 Range)
 {
 	const uniform int RangeBase = (Range - Offset) & ~(programCount-1);
-	const varying float vDt = Dt;
+	const varying FReal vDt = Dt;
 
 	for(uniform int32 Index = Offset; Index < Offset + RangeBase; Index+=programCount)
 	{
-		uniform float *uniform VPtr = (uniform float *uniform)&ParticlesV[Index];
-		uniform float *uniform XPtr = (uniform float *uniform)&ParticlesX[Index];
-		const uniform float *uniform PPtr = (uniform float *uniform)&ParticlesP[Index];
+		uniform FReal *uniform VPtr = (uniform FReal *uniform)&ParticlesV[Index];
+		uniform FReal *uniform XPtr = (uniform FReal *uniform)&ParticlesX[Index];
+		const uniform FReal *uniform PPtr = (uniform FReal *uniform)&ParticlesP[Index];
 
-		const varying float P0 = PPtr[programIndex];
-		const varying float P1 = PPtr[programCount + programIndex];
-		const varying float P2 = PPtr[2*programCount + programIndex];
+		const varying FReal P0 = PPtr[programIndex];
+		const varying FReal P1 = PPtr[programCount + programIndex];
+		const varying FReal P2 = PPtr[2*programCount + programIndex];
 
-		const varying float X0 = XPtr[programIndex];
-		const varying float X1 = XPtr[programCount + programIndex];
-		const varying float X2 = XPtr[2*programCount + programIndex];
+		const varying FReal X0 = XPtr[programIndex];
+		const varying FReal X1 = XPtr[programCount + programIndex];
+		const varying FReal X2 = XPtr[2*programCount + programIndex];
 
 		VPtr[programIndex] = (P0 - X0) / vDt;
 		VPtr[programCount + programIndex] = (P1 - X1) / vDt;
diff --git a/Engine/Source/Runtime/Experimental/Chaos/Private/Chaos/PBDLongRangeConstraints.cpp b/Engine/Source/Runtime/Experimental/Chaos/Private/Chaos/PBDLongRangeConstraints.cpp
--- a/Engine/Source/Runtime/Experimental/Chaos/Private/Chaos/PBDLongRangeConstraints.cpp
+++ b/Engine/Source/Runtime/Experimental/Chaos/Private/Chaos/PBDLongRangeConstraints.cpp
@@ -6,6 +6,9 @@
 #endif
 
 #if INTEL_ISPC && !UE_BUILD_SHIPPING
+static_assert(sizeof(ispc::FVector) == sizeof(Chaos::FVec3), "sizeof(ispc::FVector) != sizeof(Chaos::FVec3)");
+static_assert(sizeof(ispc::FTether) == sizeof(Chaos::FPBDLongRangeConstraintsBase::FTether), "sizeof(ispc::FTether) != sizeof(Chaos::FPBDLongRangeConstraintsBase::FTether)");
+
 bool bChaos_LongRange_ISPC_Enabled = true;
 FAutoConsoleVariableRef CVarChaosLongRangeISPCEnabled(TEXT("p.Chaos.LongRange.ISPC"), bChaos_LongRange_ISPC_Enabled, TEXT("Whether to use ISPC optimizations in long range constraints"));
 #endif
diff --git a/Engine/Source/Runtime/Experimental/Chaos/Private/Chaos/PBDLongRangeConstraints.ispc b/Engine/Source/Runtime/Experimental/Chaos/Private/Chaos/PBDLongRangeConstraints.ispc
--- a/Engine/Source/Runtime/Experimental/Chaos/Private/Chaos/PBDLongRangeConstraints.ispc
+++ b/Engine/Source/Runtime/Experimental/Chaos/Private/Chaos/PBDLongRangeConstraints.ispc
@@ -23,8 +23,8 @@ static inline varying FTether LoadTetherHelper(const uniform FTether *uniform Sr
 
 export void ApplyLongRangeConstraints(uniform FVector P[],
 										const uniform FTether Tethers[],
-										const uniform float Stiffness,
-										const uniform float Scale,
+										const uniform FReal Stiffness,
+										const uniform FReal Scale,
 										const uniform int32 NumConstraints,
 										const uniform int32 ParticleOffset)
 {
@@ -60,13 +60,13 @@ export void ApplyLongRangeConstraints(uniform FVector P[],
 #endif
 
 		FVector Direction = P1 - P2;
-		const float Length = VectorSize(Direction);
+		const FReal Length = VectorSize(Direction);
 		Direction = VectorSelect(Length < KINDA_SMALL_NUMBER, ForwardVector, Direction / Length);
 
-		const float Offset = Length - RefLength * Scale;
+		const FReal Offset = Length - RefLength * Scale;
 		const FVector ComputedDelta = Stiffness * Offset * Direction;
 
-		const FVector Delta = VectorSelect(Offset < 0.f, ZeroVector, ComputedDelta);
+		const FVector Delta = VectorSelect(Offset < ZERO, ZeroVector, ComputedDelta);
 
 		VectorScatter(&P[i2], P2 + Delta);
 	}
@@ -75,8 +75,8 @@ export void ApplyLongRangeConstraints(uniform FVector P[],
 export void ApplyLongRangeConstraintsStiffnessWeightmap(uniform FVector P[],
 														const uniform FTether Tethers[],
 														const uniform uint8 StiffnessIndices[],
-														const uniform float StiffnessTable[],
-														const uniform float Scale,
+														const uniform FReal StiffnessTable[],
+														const uniform FReal Scale,
 														const uniform int32 NumConstraints,
 														const uniform int32 ParticleOffset)
 {
@@ -115,16 +115,16 @@ export void ApplyLongRangeConstraintsStiffnessWeightmap(uniform FVector P[],
 #endif
 
 		#pragma ignore warning(perf)
-		const varying float Stiffness = StiffnessTable[StiffnessIndex];
+		const varying FReal Stiffness = StiffnessTable[StiffnessIndex];
 
 		FVector Direction = P1 - P2;
-		const float Length = VectorSize(Direction);
+		const FReal Length = VectorSize(Direction);
 		Direction = VectorSelect(Length < KINDA_SMALL_NUMBER, ForwardVector, Direction / Length);
 
-		const float Offset = Length - RefLength * Scale;
+		const FReal Offset = Length - RefLength * Scale;
 		const FVector ComputedDelta = Stiffness * Offset * Direction;
 
-		const FVector Delta = VectorSelect(Offset < 0.f, ZeroVector, ComputedDelta);
+		const FVector Delta = VectorSelect(Offset < ZERO, ZeroVector, ComputedDelta);
 
 		VectorScatter(&P[i2], P2 + Delta);
 	}
@@ -132,9 +132,9 @@ export void ApplyLongRangeConstraintsStiffnessWeightmap(uniform FVector P[],
 
 export void ApplyLongRangeConstraintsScaleWeightmap(uniform FVector P[],
 													const uniform FTether Tethers[],
-													const uniform float Stiffness,
+													const uniform FReal Stiffness,
 													const uniform uint8 ScaleIndices[],
-													const uniform float ScaleTable[],
+													const uniform FReal ScaleTable[],
 													const uniform int32 NumConstraints,
 													const uniform int32 ParticleOffset)
 {
@@ -173,16 +173,16 @@ export void ApplyLongRangeConstraintsScaleWeightmap(uniform FVector P[],
 #endif
 
 		#pragma ignore warning(perf)
-		const varying float Scale = ScaleTable[ScaleIndex];
+		const varying FReal Scale = ScaleTable[ScaleIndex];
 
 		FVector Direction = P1 - P2;
-		const float Length = VectorSize(Direction);
+		const FReal Length = VectorSize(Direction);
 		Direction = VectorSelect(Length < KINDA_SMALL_NUMBER, ForwardVector, Direction / Length);
 
-		const float Offset = Length - RefLength * Scale;
+		const FReal Offset = Length - RefLength * Scale;
 		const FVector ComputedDelta = Stiffness * Offset * Direction;
 
-		const FVector Delta = VectorSelect(Offset < 0.f, ZeroVector, ComputedDelta);
+		const FVector Delta = VectorSelect(Offset < ZERO, ZeroVector, ComputedDelta);
 
 		VectorScatter(&P[i2], P2 + Delta);
 	}
@@ -192,9 +192,9 @@ export void ApplyLongRangeConstraintsScaleWeightmap(uniform FVector P[],
 export void ApplyLongRangeConstraintsStiffnessScaleWeightmaps(uniform FVector P[],
 																const uniform FTether Tethers[],
 																const uniform uint8 StiffnessIndices[],
-																const uniform float StiffnessTable[],
+																const uniform FReal StiffnessTable[],
 																const uniform uint8 ScaleIndices[],
-																const uniform float ScaleTable[],
+																const uniform FReal ScaleTable[],
 																const uniform int32 NumConstraints,
 																const uniform int32 ParticleOffset)
 {
@@ -235,18 +235,18 @@ export void ApplyLongRangeConstraintsStiffnessScaleWeightmaps(uniform FVector P[
 #endif
 
 		#pragma ignore warning(perf)
-		const varying float Stiffness = StiffnessTable[StiffnessIndex];
+		const varying FReal Stiffness = StiffnessTable[StiffnessIndex];
 		#pragma ignore warning(perf)
-		const varying float Scale = ScaleTable[ScaleIndex];
+		const varying FReal Scale = ScaleTable[ScaleIndex];
 
 		FVector Direction = P1 - P2;
-		const float Length = VectorSize(Direction);
+		const FReal Length = VectorSize(Direction);
 		Direction = VectorSelect(Length < KINDA_SMALL_NUMBER, ForwardVector, Direction / Length);
 
-		const float Offset = Length - RefLength * Scale;
+		const FReal Offset = Length - RefLength * Scale;
 		const FVector ComputedDelta = Stiffness * Offset * Direction;
 
-		const FVector Delta = VectorSelect(Offset < 0.f, ZeroVector, ComputedDelta);
+		const FVector Delta = VectorSelect(Offset < ZERO, ZeroVector, ComputedDelta);
 
 		VectorScatter(&P[i2], P2 + Delta);
 	}
diff --git a/Engine/Source/Runtime/Experimental/Chaos/Private/Chaos/PBDSphericalConstraint.ispc b/Engine/Source/Runtime/Experimental/Chaos/Private/Chaos/PBDSphericalConstraint.ispc
--- a/Engine/Source/Runtime/Experimental/Chaos/Private/Chaos/PBDSphericalConstraint.ispc
+++ b/Engine/Source/Runtime/Experimental/Chaos/Private/Chaos/PBDSphericalConstraint.ispc
@@ -6,18 +6,18 @@
 
 export void ApplySphericalConstraints(uniform FVector P[],
 									const uniform FVector AnimationPositions[],
-									const uniform float InvM[],
+									const uniform FReal InvM[],
 									const uniform float SphereRadii[],
-									const uniform float SphereRadiiMultiplier,
+									const uniform FReal SphereRadiiMultiplier,
 									const uniform int32 ParticleOffset,
 									const uniform int32 ParticleCount)
 {
-	static const float DeadZoneSquareRadius = SMALL_NUMBER; // We will not push the particle away in the dead zone
+	static const FReal DeadZoneSquareRadius = SMALL_NUMBER; // We will not push the particle away in the dead zone
 
 	foreach(Index = 0 ... ParticleCount)
 	{
 		const int32 ParticleIndex = ParticleOffset + Index;
-		const float IM = InvM[ParticleIndex];
+		const FReal IM = InvM[ParticleIndex];
 
 		if(all(IM == 0))
 		{
@@ -30,10 +30,10 @@ export void ApplySphericalConstraints(uniform FVector P[],
 		const FVector Position = VectorLoad(&P[Offset]);
 
 		const FVector CenterToParticle = Position - Center;
-		const float DistanceSquared = VectorSizeSquared(CenterToParticle);
-		const float Distance = sqrt(DistanceSquared);
+		const FReal DistanceSquared = VectorSizeSquared(CenterToParticle);
+		const FReal Distance = sqrt(DistanceSquared);
 
-		const float Radius = SphereRadii[Index] * SphereRadiiMultiplier;
+		const FReal Radius = SphereRadii[Index] * SphereRadiiMultiplier;
 		const FVector FinalPosition = ((Radius / Distance) * CenterToParticle) + Center;
 
 		if (IM != 0 && DistanceSquared > ((Radius * Radius) + DeadZoneSquareRadius))
@@ -46,19 +46,19 @@ export void ApplySphericalConstraints(uniform FVector P[],
 export void ApplyLegacySphericalBackstopConstraints(uniform FVector P[],
 													const uniform FVector AnimationPositions[],
 													const uniform FVector AnimationNormals[],
-													const uniform float InvM[],
+													const uniform FReal InvM[],
 													const uniform float SphereOffsetDistances[],
 													const uniform float SphereRadii[],
-													const uniform float SphereRadiiMultiplier,
+													const uniform FReal SphereRadiiMultiplier,
 													const uniform int32 ParticleOffset,
 													const uniform int32 ParticleCount)
 {
-	static const float DeadZoneSquareRadius = SMALL_NUMBER; // We will not push the particle away in the dead zone
+	static const FReal DeadZoneSquareRadius = SMALL_NUMBER; // We will not push the particle away in the dead zone
 
 	foreach(Index = 0 ... ParticleCount)
 	{
 		const int32 ParticleIndex = ParticleOffset + Index;
-		const float IM = InvM[ParticleIndex];
+		const FReal IM = InvM[ParticleIndex];
 
 		if(all(IM == 0))
 		{
@@ -69,11 +69,11 @@ export void ApplyLegacySphericalBackstopConstraints(uniform FVector P[],
 		const FVector AnimationNormal = VectorLoad(&AnimationNormals[extract(ParticleIndex, 0)]);
 
 		const float SphereOffsetDistance = SphereOffsetDistances[Index];
-		const float Radius = SphereRadii[Index] * SphereRadiiMultiplier;
+		const FReal Radius = SphereRadii[Index] * SphereRadiiMultiplier;
 
 		const FVector Center = AnimationPosition - SphereOffsetDistance * AnimationNormal;  // Legacy version already includes the radius within the distance
 		const FVector CenterToParticle = VectorLoad(&P[extract(ParticleIndex, 0)]) - Center;
-		const float DistanceSquared = VectorSizeSquared(CenterToParticle);
+		const FReal DistanceSquared = VectorSizeSquared(CenterToParticle);
 
 		const FVector DistSquaredLessRadiusSquared = ((Radius / sqrt(DistanceSquared)) * CenterToParticle) + Center;
 		const FVector DistSquaredLessDeadZoneRadius = AnimationPosition - (SphereOffsetDistance - Radius) * AnimationNormal;  // Legacy version already includes the radius to the distance
@@ -91,19 +91,19 @@ export void ApplyLegacySphericalBackstopConstraints(uniform FVector P[],
 export void ApplySphericalBackstopConstraints(uniform FVector P[],
 												const uniform FVector AnimationPositions[],
 												const uniform FVector AnimationNormals[],
-												const uniform float InvM[],
+												const uniform FReal InvM[],
 												const uniform float SphereOffsetDistances[],
 												const uniform float SphereRadii[],
-												const uniform float SphereRadiiMultiplier,
+												const uniform FReal SphereRadiiMultiplier,
 												const uniform int32 ParticleOffset,
 												const uniform int32 ParticleCount)
 {
-	static const float DeadZoneSquareRadius = SMALL_NUMBER; // We will not push the particle away in the dead zone
+	static const FReal DeadZoneSquareRadius = SMALL_NUMBER; // We will not push the particle away in the dead zone
 
 	foreach(Index = 0 ... ParticleCount)
 	{
 		const int32 ParticleIndex = ParticleOffset + Index;
-		const float IM = InvM[ParticleIndex];
+		const FReal IM = InvM[ParticleIndex];
 
 		if(all(IM == 0))
 		{
@@ -114,11 +114,11 @@ export void ApplySphericalBackstopConstraints(uniform FVector P[],
 		const FVector AnimationNormal = VectorLoad(&AnimationNormals[extract(ParticleIndex, 0)]);
 
 		const float SphereOffsetDistance = SphereOffsetDistances[Index];
-		const float Radius = SphereRadii[Index] * SphereRadiiMultiplier;
+		const FReal Radius = SphereRadii[Index] * SphereRadiiMultiplier;
 
 		const FVector Center = AnimationPosition - (Radius + SphereOffsetDistance) * AnimationNormal;  // Non legacy version adds radius to the distance
 		const FVector CenterToParticle = VectorLoad(&P[extract(ParticleIndex, 0)]) - Center;
-		const float DistanceSquared = VectorSizeSquared(CenterToParticle);
+		const FReal DistanceSquared = VectorSizeSquared(CenterToParticle);
 
 		const FVector DistSquaredLessRadiusSquared = ((Radius / sqrt(DistanceSquared)) * CenterToParticle) + Center;
 		const FVector DistSquaredLessDeadZoneRadius = AnimationPosition - SphereOffsetDistance * AnimationNormal;  // Non legacy version adds radius to the distance
diff --git a/Engine/Source/Runtime/Experimental/Chaos/Private/Chaos/PBDSpringConstraints.cpp b/Engine/Source/Runtime/Experimental/Chaos/Private/Chaos/PBDSpringConstraints.cpp
--- a/Engine/Source/Runtime/Experimental/Chaos/Private/Chaos/PBDSpringConstraints.cpp
+++ b/Engine/Source/Runtime/Experimental/Chaos/Private/Chaos/PBDSpringConstraints.cpp
@@ -15,6 +15,9 @@
 DECLARE_CYCLE_STAT(TEXT("Chaos PBD Spring Constraint"), STAT_PBD_Spring, STATGROUP_Chaos);
 
 #if INTEL_ISPC && !UE_BUILD_SHIPPING
+static_assert(sizeof(ispc::FVector) == sizeof(Chaos::FVec3), "sizeof(ispc::FVector) != sizeof(Chaos::FVec3)");
+static_assert(sizeof(ispc::FIntVector2) == sizeof(Chaos::TVec2<int32>), "sizeof(ispc::FIntVector2) != sizeof(Chaos::TVec2<int32>)");
+
 bool bChaos_Spring_ISPC_Enabled = true;
 FAutoConsoleVariableRef CVarChaosSpringISPCEnabled(TEXT("p.Chaos.Spring.ISPC"), bChaos_Spring_ISPC_Enabled, TEXT("Whether to use ISPC optimizations in Spring constraints"));
 #endif
diff --git a/Engine/Source/Runtime/Experimental/Chaos/Private/Chaos/PBDSpringConstraints.ispc b/Engine/Source/Runtime/Experimental/Chaos/Private/Chaos/PBDSpringConstraints.ispc
--- a/Engine/Source/Runtime/Experimental/Chaos/Private/Chaos/PBDSpringConstraints.ispc
+++ b/Engine/Source/Runtime/Experimental/Chaos/Private/Chaos/PBDSpringConstraints.ispc
@@ -4,19 +4,20 @@
 
 #include "Math/Vector.isph"
 
-static inline float SafeNormalize(FVector &Direction)
+static inline FReal SafeNormalize(FVector &Direction)
 {
-	const float Size = sqrt(VectorSizeSquared(Direction));
-	Direction = VectorSelect((Size < KINDA_SMALL_NUMBER), ForwardVector, Direction / Size);
-	return (Size < KINDA_SMALL_NUMBER) ? 0.f : Size;
+	const FReal SizeSquared = VectorSizeSquared(Direction);
+	const FReal Size = sqrt(SizeSquared);
+	Direction = VectorSelect((SizeSquared < KINDA_SMALL_NUMBER), ForwardVector, Direction / Size);
+	return (SizeSquared < KINDA_SMALL_NUMBER) ? ZERO : Size;
 }
 
 export void ApplySpringConstraints(uniform FVector P[],
 									const uniform FIntVector2 AllConstraints[],
 									const uniform int32 ActiveConstraints[],
-									const uniform float InvM[],
-									const uniform float Dists[],
-									const uniform float Stiffness,
+									const uniform FReal InvM[],
+									const uniform FReal Dists[],
+									const uniform FReal Stiffness,
 									const uniform int32 NumConstraints)
 {
 	varying int32 ActiveConstraint;
@@ -46,10 +47,10 @@ export void ApplySpringConstraints(uniform FVector P[],
 		const varying FVector P2 = VectorGather(&P[i2]);
 
 		FVector Direction = P1 - P2;
-		const float Distance = SafeNormalize(Direction);
+		const FReal Distance = SafeNormalize(Direction);
 
 		#pragma ignore warning(perf)
-		const varying float Dist = Dists[ActiveConstraint];
+		const varying FReal Dist = Dists[ActiveConstraint];
 
 #if HW_GATHER_SUPPORTED == 1
 		if(i + programCount < NumConstraints)
@@ -63,20 +64,20 @@ export void ApplySpringConstraints(uniform FVector P[],
 		const FVector InnerDelta = (Distance - Dist) * Direction;
 
 		#pragma ignore warning(perf)
-		const varying float InvM1 = InvM[i1];
+		const varying FReal InvM1 = InvM[i1];
 		#pragma ignore warning(perf)
-		const varying float InvM2 = InvM[i2];
+		const varying FReal InvM2 = InvM[i2];
 
-		const float CombinedInvMass = InvM2 + InvM1;
+		const FReal CombinedInvMass = InvM2 + InvM1;
 		const FVector ComputedDelta = Stiffness * InnerDelta / CombinedInvMass;
 
-		const FVector Delta = VectorSelect((InvM2 == 0 && InvM1 == 0), ZeroVector, ComputedDelta);
+		const FVector Delta = VectorSelect((InvM2 == ZERO && InvM1 == ZERO), ZeroVector, ComputedDelta);
 
-		if (InvM1 > 0)
+		if (InvM1 > ZERO)
 		{
 			VectorScatter(&P[i1], P1 - (InvM1 * Delta));
 		}
-		if (InvM2 > 0)
+		if (InvM2 > ZERO)
 		{
 			VectorScatter(&P[i2], P2 + (InvM2 * Delta));
 		}
@@ -86,10 +87,10 @@ export void ApplySpringConstraints(uniform FVector P[],
 export void ApplySpringConstraintsWithWeightMaps(uniform FVector P[],
 									const uniform FIntVector2 AllConstraints[],
 									const uniform int32 ActiveConstraints[],
-									const uniform float InvM[],
-									const uniform float Dists[],
+									const uniform FReal InvM[],
+									const uniform FReal Dists[],
 									const uniform uint8 StiffnessIndices[],
-									const uniform float StiffnessTable[],
+									const uniform FReal StiffnessTable[],
 									const uniform int32 NumConstraints)
 {
 	varying int32 ActiveConstraint;
@@ -119,10 +120,10 @@ export void ApplySpringConstraintsWithWeightMaps(uniform FVector P[],
 		const varying FVector P2 = VectorGather(&P[i2]);
 
 		FVector Direction = P1 - P2;
-		const float Distance = SafeNormalize(Direction);
+		const FReal Distance = SafeNormalize(Direction);
 
 		#pragma ignore warning(perf)
-		const varying float Dist = Dists[ActiveConstraint];
+		const varying FReal Dist = Dists[ActiveConstraint];
 
 		#pragma ignore warning(perf)
 		const varying uint8 StiffnessIndex = StiffnessIndices[ActiveConstraint];
@@ -139,23 +140,23 @@ export void ApplySpringConstraintsWithWeightMaps(uniform FVector P[],
 		const FVector InnerDelta = (Distance - Dist) * Direction;
 
 		#pragma ignore warning(perf)
-		const varying float InvM1 = InvM[i1];
+		const varying FReal InvM1 = InvM[i1];
 		#pragma ignore warning(perf)
-		const varying float InvM2 = InvM[i2];
+		const varying FReal InvM2 = InvM[i2];
 
 		#pragma ignore warning(perf)
-		const varying float Stiffness = StiffnessTable[StiffnessIndex];
+		const varying FReal Stiffness = StiffnessTable[StiffnessIndex];
 
-		const float CombinedInvMass = InvM2 + InvM1;
+		const FReal CombinedInvMass = InvM2 + InvM1;
 		const FVector ComputedDelta = Stiffness * InnerDelta / CombinedInvMass;
 
-		const FVector Delta = VectorSelect((InvM2 == 0 && InvM1 == 0), ZeroVector, ComputedDelta);
+		const FVector Delta = VectorSelect((InvM2 == ZERO && InvM1 == ZERO), ZeroVector, ComputedDelta);
 
-		if (InvM1 > 0)
+		if (InvM1 > ZERO)
 		{
 			VectorScatter(&P[i1], P1 - (InvM1 * Delta));
 		}
-		if (InvM2 > 0)
+		if (InvM2 > ZERO)
 		{
 			VectorScatter(&P[i2], P2 + (InvM2 * Delta));
 		}
diff --git a/Engine/Source/Runtime/Experimental/Chaos/Private/Chaos/PerParticleDampVelocity.cpp b/Engine/Source/Runtime/Experimental/Chaos/Private/Chaos/PerParticleDampVelocity.cpp
--- a/Engine/Source/Runtime/Experimental/Chaos/Private/Chaos/PerParticleDampVelocity.cpp
+++ b/Engine/Source/Runtime/Experimental/Chaos/Private/Chaos/PerParticleDampVelocity.cpp
@@ -7,6 +7,8 @@
 #endif
 
 #if INTEL_ISPC && !UE_BUILD_SHIPPING
+static_assert(sizeof(ispc::FVector) == sizeof(Chaos::FVec3), "sizeof(ispc::FVector) != sizeof(Chaos::FVec3)");
+
 bool bChaos_DampVelocity_ISPC_Enabled = true;
 FAutoConsoleVariableRef CVarChaosDampVelocityISPCEnabled(TEXT("p.Chaos.DampVelocity.ISPC"), bChaos_DampVelocity_ISPC_Enabled, TEXT("Whether to use ISPC optimizations in per particle damp velocity calculation"));
 #endif
diff --git a/Engine/Source/Runtime/Experimental/Chaos/Private/Chaos/PerParticleDampVelocity.ispc b/Engine/Source/Runtime/Experimental/Chaos/Private/Chaos/PerParticleDampVelocity.ispc
--- a/Engine/Source/Runtime/Experimental/Chaos/Private/Chaos/PerParticleDampVelocity.ispc
+++ b/Engine/Source/Runtime/Experimental/Chaos/Private/Chaos/PerParticleDampVelocity.ispc
@@ -11,23 +11,28 @@ static inline uniform bool IsFinite(uniform float A)
 	return (intbits(A) & 0x7F800000U) != 0x7F800000U;
 }
 
+static inline uniform bool IsFinite(uniform double A)
+{
+	return (intbits(A) & 0x7FF0000000000000ULL) != 0x7FF0000000000000ULL;
+}
+
 export void UpdatePositionBasedState(uniform FVector& MXcm,
 									uniform FVector& MVcm,
 									uniform FVector& MOmega,
 									const uniform FVector X[],
 									const uniform FVector V[],
-									const uniform float M[],
-									const uniform float InvM[],
+									const uniform FReal M[],
+									const uniform FReal InvM[],
 									const uniform int Offset,
 									const uniform int32 Range)
 {
 	varying FVector Xcm = ZeroVector;
 	varying FVector Vcm = ZeroVector;
-	varying float Mcm = 0.0f;
+	varying FReal Mcm = 0.;
 
 	foreach(Index = Offset ... Range)
 	{
-		const varying float IM = InvM[Index];
+		const varying FReal IM = InvM[Index];
 
 		if(all(!IM))
 		{
@@ -41,7 +46,7 @@ export void UpdatePositionBasedState(uniform FVector& MXcm,
 
 		const varying FVector MX = VectorLoad(&X[extract(Index, 0)]);
 		const varying FVector MV = VectorLoad(&V[extract(Index, 0)]);
-		const varying float MM = M[Index];
+		const varying FReal MM = M[Index];
 
 		Xcm = Xcm + (MX * MM);
 		Vcm = Vcm + (MV * MM);
@@ -50,9 +55,9 @@ export void UpdatePositionBasedState(uniform FVector& MXcm,
 
 	MXcm = VectorReduceAdd(Xcm);
 	MVcm = VectorReduceAdd(Vcm);
-	uniform float MMcm = reduce_add(Mcm);
+	uniform FReal MMcm = reduce_add(Mcm);
 
-	if (MMcm != 0.0f)
+	if (MMcm != ZERO)
 	{
 		MXcm = MXcm / MMcm;
 		MVcm = MVcm / MMcm;
@@ -63,7 +68,7 @@ export void UpdatePositionBasedState(uniform FVector& MXcm,
 
 	foreach(Index = Offset ... Range)
 	{
-		const varying float IM = InvM[Index];
+		const varying FReal IM = InvM[Index];
 
 		if(all(!IM))
 		{
@@ -77,7 +82,7 @@ export void UpdatePositionBasedState(uniform FVector& MXcm,
 
 		const varying FVector MX = VectorLoad(&X[extract(Index, 0)]);
 		const varying FVector MV = VectorLoad(&V[extract(Index, 0)]);
-		const varying float MM = M[Index];
+		const varying FReal MM = M[Index];
 
 		const varying FVector VX = MX - MXcm;
 		VL = VL + VectorCross(VX, MM * MV);
@@ -89,7 +94,7 @@ export void UpdatePositionBasedState(uniform FVector& MXcm,
 	const uniform FVector L = VectorReduceAdd(VL);
 	const uniform FMatrix33 I = MatrixReduceAdd(VI);
 
-	const uniform float Det = MatrixDeterminant(I);
+	const uniform FReal Det = MatrixDeterminant(I);
 
 	MOmega = VectorSelect(Det < SMALL_NUMBER || !IsFinite(Det), ZeroVector, MatrixInverseTransformVector(I, L)); // Calls FMatrix::InverseFast(), which tests against SMALL_NUMBER
 }
diff --git a/Engine/Source/Runtime/Experimental/Chaos/Private/Chaos/PerParticlePBDCollisionConstraint.cpp b/Engine/Source/Runtime/Experimental/Chaos/Private/Chaos/PerParticlePBDCollisionConstraint.cpp
--- a/Engine/Source/Runtime/Experimental/Chaos/Private/Chaos/PerParticlePBDCollisionConstraint.cpp
+++ b/Engine/Source/Runtime/Experimental/Chaos/Private/Chaos/PerParticlePBDCollisionConstraint.cpp
@@ -8,6 +8,9 @@
 #endif
 
 #if INTEL_ISPC && !UE_BUILD_SHIPPING
+static_assert(sizeof(ispc::FVector) == sizeof(Chaos::FVec3), "sizeof(ispc::FVector) != sizeof(Chaos::FVec3)");
+static_assert(sizeof(ispc::FVector4) == sizeof(Chaos::FRotation3), "sizeof(ispc::FVector4) != sizeof(Chaos::FRotation3)");
+
 bool bChaos_PerParticleCollision_ISPC_Enabled = true;
 FAutoConsoleVariableRef CVarChaosPerParticleCollisionISPCEnabled(TEXT("p.Chaos.PerParticleCollision.ISPC"), bChaos_PerParticleCollision_ISPC_Enabled, TEXT("Whether to use ISPC optimizations in per particle collisions"));
 #endif
diff --git a/Engine/Source/Runtime/Experimental/Chaos/Private/Chaos/PerParticlePBDCollisionConstraint.ispc b/Engine/Source/Runtime/Experimental/Chaos/Private/Chaos/PerParticlePBDCollisionConstraint.ispc
--- a/Engine/Source/Runtime/Experimental/Chaos/Private/Chaos/PerParticlePBDCollisionConstraint.ispc
+++ b/Engine/Source/Runtime/Experimental/Chaos/Private/Chaos/PerParticlePBDCollisionConstraint.ispc
@@ -118,7 +118,7 @@ struct Segment
 {
 	FVector Point;
 	FVector Axis;
-	float Length;
+	FReal Length;
 };
 
 struct PlaneConcrete
@@ -153,45 +153,45 @@ struct ConvexShape
 	FVector LocalBoundingBoxMin;
 	FVector LocalBoundingBoxMax;
 	FConvexStructureData StructureData;
-	float Volume;
+	FReal Volume;
 	FVector CenterOfMass;
 };
 
-static inline float SafeNormalize(FVector &Direction)
+static inline FReal SafeNormalize(FVector &Direction)
 {
-	const float SizeSqr = VectorSizeSquared(Direction);
-	const float Size = sqrt(SizeSqr);
+	const FReal SizeSqr = VectorSizeSquared(Direction);
+	const FReal Size = sqrt(SizeSqr);
 	Direction = VectorSelect((SizeSqr < KINDA_SMALL_NUMBER), ForwardVector, Direction / Size);
-	return (Size < KINDA_SMALL_NUMBER) ? 0.f : Size;
+	return (Size < KINDA_SMALL_NUMBER) ? ZERO : Size;
 }
 
 static inline bool VectorNormalize(FVector &V)
 {
-	const float SquareSum = VectorSizeSquared(V);
+	const FReal SquareSum = VectorSizeSquared(V);
 	if(SquareSum > SMALL_NUMBER)
 	{
-		const float Scale = InvSqrt(SquareSum);
+		const FReal Scale = InvSqrt(SquareSum);
 		V = V * Scale;
 		return true;
 	}
 	return false;
 }
 
-static inline void PhiWithNormalSphere(const uniform uint8 *uniform TypedObjectPtr, const FVector &PhiWithNormalInput, FVector &Normal, float &Phi, const uniform float Radius)
+static inline void PhiWithNormalSphere(const uniform uint8 *uniform TypedObjectPtr, const FVector &PhiWithNormalInput, FVector &Normal, FReal &Phi, const uniform FReal Radius)
 {
 	const uniform FVector Center = *((const uniform FVector *uniform)&TypedObjectPtr[0]);
 	Normal = PhiWithNormalInput - Center;
 	Phi = SafeNormalize(Normal) - Radius;
 }
 
-static inline void PhiWithNormalPlane(const uniform uint8 *uniform TypedObjectPtr, const FVector &PhiWithNormalInput, FVector &Normal, float &Phi)
+static inline void PhiWithNormalPlane(const uniform uint8 *uniform TypedObjectPtr, const FVector &PhiWithNormalInput, FVector &Normal, FReal &Phi)
 {
 	const uniform Plane P = *((const uniform Plane *uniform)&TypedObjectPtr[0]);
 	Normal = P.MPlaneConcrete.MNormal;
 	Phi = VectorDot(PhiWithNormalInput - P.MPlaneConcrete.MX, P.MPlaneConcrete.MNormal);
 }
 
-static inline void PhiWithNormalPlane(const uniform Plane *varying InPlane, const FVector &PhiWithNormalInput, FVector &Normal, float &Phi)
+static inline void PhiWithNormalPlane(const uniform Plane *varying InPlane, const FVector &PhiWithNormalInput, FVector &Normal, FReal &Phi)
 {
 	const varying FVector MNormal = VectorGather(&InPlane->MPlaneConcrete.MNormal);
 	const varying FVector MX = VectorGather(&InPlane->MPlaneConcrete.MX);
@@ -200,54 +200,54 @@ static inline void PhiWithNormalPlane(const uniform Plane *varying InPlane, cons
 	Phi = VectorDot(PhiWithNormalInput - MX, MNormal);
 }
 
-static inline float SignedDistancePlane(const uniform Plane *uniform InPlane, const FVector &x)
+static inline FReal SignedDistancePlane(const uniform Plane *uniform InPlane, const FVector &x)
 {
 	const uniform Plane P = *InPlane;
 	return VectorDot(x - P.MPlaneConcrete.MX, P.MPlaneConcrete.MNormal);
 }
 
-static inline void PhiWithNormalCapsule(const uniform uint8 *uniform TypedObjectPtr, const FVector &PhiWithNormalInput, FVector &Normal, float &Phi, const uniform float Radius)
+static inline void PhiWithNormalCapsule(const uniform uint8 *uniform TypedObjectPtr, const FVector &PhiWithNormalInput, FVector &Normal, FReal &Phi, const uniform FReal Radius)
 {
 	const uniform Segment MSegment = *((const uniform Segment *uniform)&TypedObjectPtr[0]);
 
-	const float Dot = clamp(VectorDot(PhiWithNormalInput - MSegment.Point, MSegment.Axis), 0.0f, MSegment.Length);
+	const FReal Dot = clamp(VectorDot(PhiWithNormalInput - MSegment.Point, MSegment.Axis), (uniform FReal)(0), MSegment.Length);
 	const FVector ProjectedPoint = SetVector(Dot) * MSegment.Axis + MSegment.Point;
 	Normal = PhiWithNormalInput - ProjectedPoint;
 	Phi = SafeNormalize(Normal) - Radius;
 }
 
-static inline void PhiWithNormalTaperedCapsule(const uniform uint8* uniform TypedObjectPtr, const FVector& PhiWithNormalInput, FVector& Normal, float& Phi)
+static inline void PhiWithNormalTaperedCapsule(const uniform uint8* uniform TypedObjectPtr, const FVector& PhiWithNormalInput, FVector& Normal, FReal& Phi)
 {
 	const uniform FVector Origin = *((const uniform FVector *uniform)& TypedObjectPtr[0]);
 	const uniform FVector Axis   = *((const uniform FVector *uniform)& TypedObjectPtr[sizeof(uniform FVector)]);
-	const uniform float Height   = *((const uniform float *uniform)& TypedObjectPtr[2*sizeof(uniform FVector)]);
-	const uniform float Radius1  = *((const uniform float *uniform)& TypedObjectPtr[2*sizeof(uniform FVector)+ sizeof(uniform float)]);
-	const uniform float Radius2  = *((const uniform float *uniform)& TypedObjectPtr[2*sizeof(uniform FVector)+ 2*sizeof(uniform float)]);
+	const uniform FReal Height   = *((const uniform FReal *uniform)& TypedObjectPtr[2*sizeof(uniform FVector)]);
+	const uniform FReal Radius1  = *((const uniform FReal *uniform)& TypedObjectPtr[2*sizeof(uniform FVector)+ sizeof(uniform FReal)]);
+	const uniform FReal Radius2  = *((const uniform FReal *uniform)& TypedObjectPtr[2*sizeof(uniform FVector)+ 2*sizeof(uniform FReal)]);
 
-	const float DistanceAlongAxis = clamp(VectorDot(PhiWithNormalInput - Origin, Axis), 0.0f, Height);
+	const FReal DistanceAlongAxis = clamp(VectorDot(PhiWithNormalInput - Origin, Axis), (uniform FReal)(0), Height);
 	const FVector ProjectedPoint = SetVector(DistanceAlongAxis) * Axis + Origin;
-	const float Alpha = DistanceAlongAxis / Height;
-	const float Radius = Radius1 * (1.0f - Alpha) + Radius2 * Alpha;
+	const FReal Alpha = DistanceAlongAxis / Height;
+	const FReal Radius = Radius1 * ((uniform FReal)(1.0) - Alpha) + Radius2 * Alpha;
 	Normal = PhiWithNormalInput - ProjectedPoint;
 	Phi = SafeNormalize(Normal) - Radius;
 }
 
-static inline void PhiWithNormalTaperedCylinder(const uniform uint8 *uniform TypedObjectPtr, const FVector &PhiWithNormalInput, FVector &Normal, float &Phi, const uniform int SizeofFImplicitObject)
+static inline void PhiWithNormalTaperedCylinder(const uniform uint8 *uniform TypedObjectPtr, const FVector &PhiWithNormalInput, FVector &Normal, FReal &Phi, const uniform int SizeofFImplicitObject)
 {
 	const uniform Plane MPlane1 = *((const uniform Plane *uniform)&TypedObjectPtr[SizeofFImplicitObject]);
 	const uniform Plane MPlane2 = *((const uniform Plane *uniform)&TypedObjectPtr[2*SizeofFImplicitObject+sizeof(uniform Plane)]);
-	const uniform float MHeight = *((const uniform float *uniform)&TypedObjectPtr[2*SizeofFImplicitObject+2*sizeof(uniform Plane)]);
-	const uniform float MRadius1 = *((const uniform float *uniform)&TypedObjectPtr[2*SizeofFImplicitObject + 2*sizeof(uniform Plane) + sizeof(uniform float)]);
-	const uniform float MRadius2 = *((const uniform float *uniform)&TypedObjectPtr[2*SizeofFImplicitObject + 2*sizeof(uniform Plane) + 2*sizeof(uniform float)]);
+	const uniform FReal MHeight = *((const uniform FReal *uniform)&TypedObjectPtr[2*SizeofFImplicitObject+2*sizeof(uniform Plane)]);
+	const uniform FReal MRadius1 = *((const uniform FReal *uniform)&TypedObjectPtr[2*SizeofFImplicitObject + 2*sizeof(uniform Plane) + sizeof(uniform FReal)]);
+	const uniform FReal MRadius2 = *((const uniform FReal *uniform)&TypedObjectPtr[2*SizeofFImplicitObject + 2*sizeof(uniform Plane) + 2*sizeof(uniform FReal)]);
 
-	const float Distance1 = VectorDot(PhiWithNormalInput - MPlane1.MPlaneConcrete.MX, MPlane1.MPlaneConcrete.MNormal);
+	const FReal Distance1 = VectorDot(PhiWithNormalInput - MPlane1.MPlaneConcrete.MX, MPlane1.MPlaneConcrete.MNormal);
 	// Used to be Distance2 = MPlane2.PhiWithNormal(x, Normal2); but that would trigger 
-	const float Distance2 = MHeight - Distance1;          // the ensure on Distance2 being slightly larger than MHeight in some border cases
+	const FReal Distance2 = MHeight - Distance1;          // the ensure on Distance2 being slightly larger than MHeight in some border cases
 	const FVector SideVector = (PhiWithNormalInput - (MPlane1.MPlaneConcrete.MNormal * Distance1 + MPlane1.MPlaneConcrete.MX));
-	const float Alpha = Distance1 / MHeight;
-	const float Radius = MRadius1 * (1.0f - Alpha) + MRadius2 * Alpha;
-	const float SideDistance = VectorSize(SideVector) - Radius;
-	const float TopDistance = Distance1 < Distance2 ? Distance1 : Distance2;
+	const FReal Alpha = Distance1 / MHeight;
+	const FReal Radius = MRadius1 * ((uniform FReal)(1.0) - Alpha) + MRadius2 * Alpha;
+	const FReal SideDistance = VectorSize(SideVector) - Radius;
+	const FReal TopDistance = Distance1 < Distance2 ? Distance1 : Distance2;
 
 	if (Distance1 < SMALL_NUMBER)
 	{
@@ -261,7 +261,7 @@ static inline void PhiWithNormalTaperedCylinder(const uniform uint8 *uniform Typ
 		}
 		else
 		{
-			Normal = MPlane1.MPlaneConcrete.MNormal * -1.0f;
+			Normal = MPlane1.MPlaneConcrete.MNormal * (uniform FReal)(-1.0);
 			Phi = -Distance1;
 		}
 	}
@@ -277,13 +277,13 @@ static inline void PhiWithNormalTaperedCylinder(const uniform uint8 *uniform Typ
 		}
 		else
 		{
-			Normal = MPlane2.MPlaneConcrete.MNormal * -1.0f;
+			Normal = MPlane2.MPlaneConcrete.MNormal * (uniform FReal)(-1.0);
 			Phi = -Distance2;
 		}
 	}
-	else if (SideDistance < 0.0f && TopDistance < -SideDistance)
+	else if (SideDistance < (uniform FReal)(0.0) && TopDistance < -SideDistance)
 	{
-		Normal = Distance1 < Distance2 ? MPlane1.MPlaneConcrete.MNormal * -1.0f : MPlane2.MPlaneConcrete.MNormal * -1.0f;
+		Normal = Distance1 < Distance2 ? MPlane1.MPlaneConcrete.MNormal * (uniform FReal)(-1.0) : MPlane2.MPlaneConcrete.MNormal * (uniform FReal)(-1.0);
 		Phi = -TopDistance;
 	}
 	else
@@ -293,7 +293,7 @@ static inline void PhiWithNormalTaperedCylinder(const uniform uint8 *uniform Typ
 	}
 }
 
-static inline void PhiWithNormalUnion(const uniform uint8 *uniform TypedPtr, const FVector &PhiWithNormalInput, FVector &Normal, float &Phi, const uniform int SizeofFImplicitObject, const uniform int OffsetOfMargin)
+static inline void PhiWithNormalUnion(const uniform uint8 *uniform TypedPtr, const FVector &PhiWithNormalInput, FVector &Normal, FReal &Phi, const uniform int SizeofFImplicitObject, const uniform int OffsetOfMargin)
 {
 	const uniform UniquePtr *uniform MObjects = (const uniform UniquePtr *uniform)TypedPtr;
 	Phi = FLT_MAX;
@@ -303,7 +303,7 @@ static inline void PhiWithNormalUnion(const uniform uint8 *uniform TypedPtr, con
 	uniform uint8 *uniform TypedObjectPtr = Object + SizeofFImplicitObject;
 
 	FVector NextNormal;
-	float NextPhi;
+	FReal NextPhi;
 
 	PhiWithNormalTaperedCylinder(TypedObjectPtr, PhiWithNormalInput, NextNormal, NextPhi, SizeofFImplicitObject);
 
@@ -320,7 +320,7 @@ static inline void PhiWithNormalUnion(const uniform uint8 *uniform TypedPtr, con
 	}
 
 	Object = MObjects->Data[1].Ptr;
-	uniform float Margin = *((const uniform float *uniform)(Object + OffsetOfMargin));
+	uniform FReal Margin = *((const uniform FReal *uniform)(Object + OffsetOfMargin));
 	TypedObjectPtr = Object + SizeofFImplicitObject;
 
 	PhiWithNormalSphere(TypedObjectPtr, PhiWithNormalInput, NextNormal, NextPhi, Margin);
@@ -338,7 +338,7 @@ static inline void PhiWithNormalUnion(const uniform uint8 *uniform TypedPtr, con
 	}
 
 	Object = MObjects->Data[2].Ptr;
-	Margin = *((const uniform float *uniform)(Object + OffsetOfMargin));
+	Margin = *((const uniform FReal *uniform)(Object + OffsetOfMargin));
 	TypedObjectPtr = Object + SizeofFImplicitObject;
 
 	PhiWithNormalSphere(TypedObjectPtr, PhiWithNormalInput, NextNormal, NextPhi, Margin);
@@ -365,11 +365,11 @@ static inline FVector FindClosestPointOnLineSegment(const FVector& P0, const FVe
 {
 	const FVector P10 = P1 - P0;
 	const FVector PP0 = P - P0;
-	const float Proj = VectorDot(P10, PP0);
-	const float Denom2 = VectorSizeSquared(P10);
-	const float NormalProj = Proj / Denom2;
+	const FReal Proj = VectorDot(P10, PP0);
+	const FReal Denom2 = VectorSizeSquared(P10);
+	const FReal NormalProj = Proj / Denom2;
 	const FVector P2 = P0 + NormalProj * P10;
-	return VectorSelect(Proj < 0.0f || Denom2 < KINDA_SMALL_NUMBER, P0, VectorSelect(NormalProj > 1.0f, P1, P2));
+	return VectorSelect(Proj < (uniform FReal)(0.0) || Denom2 < KINDA_SMALL_NUMBER, P0, VectorSelect(NormalProj > (uniform FReal)(1.0), P1, P2));
 }
 
 static inline FVector FindClosestPointOnTriangle(const FVector& ClosestPointOnPlane, const FVector& P0, const FVector& P1, const FVector& P2, const FVector& P)
@@ -378,14 +378,14 @@ static inline FVector FindClosestPointOnTriangle(const FVector& ClosestPointOnPl
 	const FVector P10 = P1 - P0;
 	const FVector P20 = P2 - P0;
 	const FVector PP0 = P - P0;
-	const float Size10 = VectorSizeSquared(P10);
-	const float Size20 = VectorSizeSquared(P20);
-	const float ProjSides = VectorDot(P10, P20);
-	const float ProjP1 = VectorDot(PP0, P10);
-	const float ProjP2 = VectorDot(PP0, P20);
-	const float Denom = Size10 * Size20 - ProjSides * ProjSides;
-	const float BaryX = (Size20 * ProjP1 - ProjSides * ProjP2) / Denom;
-	const float BaryY = (Size10 * ProjP2 - ProjSides * ProjP1) / Denom;
+	const FReal Size10 = VectorSizeSquared(P10);
+	const FReal Size20 = VectorSizeSquared(P20);
+	const FReal ProjSides = VectorDot(P10, P20);
+	const FReal ProjP1 = VectorDot(PP0, P10);
+	const FReal ProjP2 = VectorDot(PP0, P20);
+	const FReal Denom = Size10 * Size20 - ProjSides * ProjSides;
+	const FReal BaryX = (Size20 * ProjP1 - ProjSides * ProjP2) / Denom;
+	const FReal BaryY = (Size10 * ProjP2 - ProjSides * ProjP1) / Denom;
 
 	const bool bClosestPoint = (BaryX >= -KINDA_SMALL_NUMBER && BaryX <= 1 + KINDA_SMALL_NUMBER && BaryY >= -KINDA_SMALL_NUMBER && BaryY <= 1 + KINDA_SMALL_NUMBER && (BaryX + BaryY) <= (1 + KINDA_SMALL_NUMBER));
 
@@ -393,9 +393,9 @@ static inline FVector FindClosestPointOnTriangle(const FVector& ClosestPointOnPl
 	const FVector P20Closest = FindClosestPointOnLineSegment(P0, P2, P);
 	const FVector P21Closest = FindClosestPointOnLineSegment(P1, P2, P);
 
-	const float P10Dist2 = VectorSizeSquared(P - P10Closest);
-	const float P20Dist2 = VectorSizeSquared(P - P20Closest);
-	const float P21Dist2 = VectorSizeSquared(P - P21Closest);
+	const FReal P10Dist2 = VectorSizeSquared(P - P10Closest);
+	const FReal P20Dist2 = VectorSizeSquared(P - P20Closest);
+	const FReal P21Dist2 = VectorSizeSquared(P - P21Closest);
 
 	return VectorSelect(bClosestPoint, ClosestPointOnPlane, VectorSelect(P10Dist2 < P20Dist2, VectorSelect(P10Dist2 < P21Dist2, P10Closest, P21Closest), VectorSelect(P20Dist2 < P21Dist2, P20Closest, P21Closest)));
 }
@@ -474,7 +474,7 @@ static inline int32 GetPlaneVertex(const uniform FConvexStructureData *uniform S
 	}
 }
 
-static inline void PhiWithNormalConvex(const uniform uint8 *uniform TypedObjectPtr, const FVector &PhiWithNormalInput, FVector &Normal, float &Phi, const uniform int SizeofFImplicitObject)
+static inline void PhiWithNormalConvex(const uniform uint8 *uniform TypedObjectPtr, const FVector &PhiWithNormalInput, FVector &Normal, FReal &Phi, const uniform int SizeofFImplicitObject)
 {
 	const uniform ConvexShape *uniform Object = (const uniform ConvexShape *uniform)TypedObjectPtr;
 	const uniform Plane *uniform Planes = (const uniform Plane *uniform)&Object->Planes.Data[0];
@@ -485,12 +485,12 @@ static inline void PhiWithNormalConvex(const uniform uint8 *uniform TypedObjectP
 		return;
 	}
 
-	float MaxPhi = -FLT_MAX;
+	FReal MaxPhi = -FLT_MAX;
 	int32 MaxPlane = 0;
 
 	for (uniform int32 Idx = 0; Idx < NumPlanes; ++Idx)
 	{
-		const float InnerPhi = SignedDistancePlane(&Planes[Idx], PhiWithNormalInput);
+		const FReal InnerPhi = SignedDistancePlane(&Planes[Idx], PhiWithNormalInput);
 		if (InnerPhi > MaxPhi)
 		{
 			MaxPhi = InnerPhi;
@@ -513,7 +513,7 @@ static inline void PhiWithNormalConvex(const uniform uint8 *uniform TypedObjectP
 		return;
 	}
 
-	float ClosestDistance = FLT_MAX;
+	FReal ClosestDistance = FLT_MAX;
 	FVector ClosestPoint = ZeroVector;
 	for (uniform int32 Index = 0; Index < PlaneVerticesNum - 2; Index++)
 	{
@@ -528,7 +528,7 @@ static inline void PhiWithNormalConvex(const uniform uint8 *uniform TypedObjectP
 
 		const FVector Point = FindClosestPointOnTriangle(XOnPlane, A, B, C, PhiWithNormalInput);
 
-		const float Distance = VectorSize(Point - XOnPlane);
+		const FReal Distance = VectorSize(Point - XOnPlane);
 		if (Distance < ClosestDistance)
 		{
 			ClosestDistance = Distance;
@@ -544,25 +544,25 @@ static inline void PhiWithNormalConvex(const uniform uint8 *uniform TypedObjectP
 	}
 }
 
-extern "C" unmasked void GetPhiWithNormal(const uniform uint8 *uniform CollisionParticles, const uniform float *uniform V, uniform float *uniform Normal, uniform float *uniform Phi, const uniform int i, const uniform int ProgramCount, const uniform int Mask);
+extern "C" unmasked void GetPhiWithNormal(const uniform uint8 *uniform CollisionParticles, const uniform FReal *uniform V, uniform FReal *uniform Normal, uniform FReal *uniform Phi, const uniform int i, const uniform int ProgramCount, const uniform int Mask);
 
 export void ApplyPerParticleCollision(uniform FVector ParticlesP[],
 										const uniform FVector ParticlesX[],
-										const uniform float ParticlesInvM[],
+										const uniform FReal ParticlesInvM[],
 										const uniform FVector CollisionV[],
 										const uniform FVector CollisionX[],
 										const uniform FVector CollisionW[],
 										const uniform FVector4 CollisionR[],
 										const uniform uint32 DynamicGroupIds[],
 										const uniform uint32 KinematicGroupIds[],
-										const uniform float PerGroupFriction[],
-										const uniform float PerGroupThickness[],
+										const uniform FReal PerGroupFriction[],
+										const uniform FReal PerGroupThickness[],
 										const uniform uint8 *uniform CollisionParticles,
 										const uniform uint8 *uniform Geometry,
 										const uniform int SizeofFImplicitObject,
 										const uniform int OffsetofGeometryType,
 										const uniform int OffsetOfMargin,
-										const uniform float Dt,
+										const uniform FReal Dt,
 										const uniform int InnerOffset,
 										const uniform int InnerRange,
 										const uniform int OuterOffset,
@@ -596,7 +596,7 @@ export void ApplyPerParticleCollision(uniform FVector ParticlesP[],
 
 	foreach(Index = OuterOffset ... BailAt)
 	{
-		const float IM = ParticlesInvM[Index];
+		const FReal IM = ParticlesInvM[Index];
 
 		if(all(IM == 0))
 		{
@@ -609,10 +609,10 @@ export void ApplyPerParticleCollision(uniform FVector ParticlesP[],
 
 		for(uniform int i = InnerOffset; i < InnerRange; i++) // SequentialFor
 		{
-			const uniform FTransform Frame = SetTransform(CollisionR[i], SetVector4(CollisionX[i], 0.0f), Float1110);
+			const uniform FTransform Frame = SetTransform(CollisionR[i], SetVector4(CollisionX[i], (uniform FReal)(0.0)), Vector1110);
 			const FVector PhiWithNormalInput = InverseTransformPosition(Frame, P);
 			FVector Normal;
-			float Phi;
+			FReal Phi;
 
 			const uniform SerializedPtr *uniform GeometryPtr = (const uniform SerializedPtr *uniform)Geometry;
 			const uniform uint8 *uniform Ptr = GeometryPtr[i].Ptr;
@@ -623,13 +623,13 @@ export void ApplyPerParticleCollision(uniform FVector ParticlesP[],
 			{
 			case Sphere:
 				{
-					const uniform float Margin = *((const uniform float *uniform)(Ptr + OffsetOfMargin));
+					const uniform FReal Margin = *((const uniform FReal *uniform)(Ptr + OffsetOfMargin));
 					PhiWithNormalSphere(TypedPtr, PhiWithNormalInput, Normal, Phi, Margin);
 					break;
 				}
 			case Capsule:
 				{
-					const uniform float Margin = *((const uniform float *uniform)(Ptr + OffsetOfMargin));
+					const uniform FReal Margin = *((const uniform FReal *uniform)(Ptr + OffsetOfMargin));
 					PhiWithNormalCapsule(TypedPtr, PhiWithNormalInput, Normal, Phi, Margin);
 					break;
 				}
@@ -656,13 +656,13 @@ export void ApplyPerParticleCollision(uniform FVector ParticlesP[],
 						IMMask = lanemask();
 					}
 
-					GetPhiWithNormal(CollisionParticles, (const uniform float *uniform)&PhiWithNormalInput, (uniform float *uniform)&Normal, (uniform float *uniform)&Phi, i, programCount, IMMask);
+					GetPhiWithNormal(CollisionParticles, (const uniform FReal *uniform)&PhiWithNormalInput, (uniform FReal *uniform)&Normal, (uniform FReal *uniform)&Phi, i, programCount, IMMask);
 					break;
 				}
 			}
 
 			#pragma ignore warning(perf)
-			const float Thickness = PerGroupThickness[DynamicGroupId];
+			const FReal Thickness = PerGroupThickness[DynamicGroupId];
 
 			if(Phi < Thickness)
 			{
@@ -670,19 +670,19 @@ export void ApplyPerParticleCollision(uniform FVector ParticlesP[],
 				P = P + ((-Phi + Thickness) * NormalWorld);
 
 				#pragma ignore warning(perf)
-				const float CoefficientOfFriction = PerGroupFriction[DynamicGroupId];
+				const FReal CoefficientOfFriction = PerGroupFriction[DynamicGroupId];
 
-				if(CoefficientOfFriction > 0.0f)
+				if(CoefficientOfFriction > (uniform FReal)(0.0))
 				{
-					const uniform float MaximumFrictionCorrectionPerStep = 1.0f; // 1cm absolute maximum correction the friction force can provide
-					const float Penetration = (-Phi + Thickness); // This is related to the Normal impulse
+					const uniform FReal MaximumFrictionCorrectionPerStep = (uniform FReal)(1.0); // 1cm absolute maximum correction the friction force can provide
+					const FReal Penetration = (-Phi + Thickness); // This is related to the Normal impulse
 
 					FVector VectorToPoint = P - CollisionX[i];
 					const FVector X = VectorLoad(&ParticlesX[extract(Index, 0)]);
 					const FVector RelativeDisplacement = (P - X) - (CollisionV[i] + VectorCross(CollisionW[i], VectorToPoint)) * Dt; // This corresponds to the tangential velocity multiplied by dt (friction will drive this to zero if it is high enough)
 					const FVector RelativeDisplacementTangent = RelativeDisplacement - VectorDot(RelativeDisplacement, NormalWorld) * NormalWorld; // Project displacement into the tangential plane
-					const float RelativeDisplacementTangentLength = VectorSize(RelativeDisplacementTangent);
-					float PositionCorrection = clamp(Penetration * CoefficientOfFriction, 0.0f, RelativeDisplacementTangentLength);
+					const FReal RelativeDisplacementTangentLength = VectorSize(RelativeDisplacementTangent);
+					FReal PositionCorrection = clamp(Penetration * CoefficientOfFriction, (uniform FReal)(0.0), RelativeDisplacementTangentLength);
 					if (PositionCorrection > MaximumFrictionCorrectionPerStep)
 					{
 						PositionCorrection = MaximumFrictionCorrectionPerStep;
@@ -704,21 +704,21 @@ export void ApplyPerParticleCollision(uniform FVector ParticlesP[],
 
 export void ApplyPerParticleCollisionFastFriction(uniform FVector ParticlesP[],
 													const uniform FVector ParticlesX[],
-													const uniform float ParticlesInvM[],
+													const uniform FReal ParticlesInvM[],
 													const uniform FVector CollisionV[],
 													const uniform FVector CollisionX[],
 													const uniform FVector CollisionW[],
 													const uniform FVector4 CollisionR[],
 													const uniform uint32 DynamicGroupId,
 													const uniform uint32 KinematicGroupIds[],
-													const uniform float PerGroupFriction,
-													const uniform float PerGroupThickness,
+													const uniform FReal PerGroupFriction,
+													const uniform FReal PerGroupThickness,
 													const uniform uint8 *uniform CollisionParticles,
 													const uniform uint8 *uniform Geometry,
 													const uniform int SizeofFImplicitObject,
 													const uniform int OffsetofGeometryType,
 													const uniform int OffsetOfMargin,
-													const uniform float Dt,
+													const uniform FReal Dt,
 													const uniform int InnerOffset,
 													const uniform int InnerRange,
 													const uniform int OuterOffset,
@@ -750,7 +750,7 @@ export void ApplyPerParticleCollisionFastFriction(uniform FVector ParticlesP[],
 
 	foreach(Index = OuterOffset ... BailAt)
 	{
-		const float IM = ParticlesInvM[Index];
+		const FReal IM = ParticlesInvM[Index];
 
 		if(all(IM == 0))
 		{
@@ -761,10 +761,10 @@ export void ApplyPerParticleCollisionFastFriction(uniform FVector ParticlesP[],
 
 		for(uniform int i = InnerOffset; i < InnerRange; i++) // SequentialFor
 		{
-			const uniform FTransform Frame = SetTransform(CollisionR[i], SetVector4(CollisionX[i], 0.0f), Float1110);
+			const uniform FTransform Frame = SetTransform(CollisionR[i], SetVector4(CollisionX[i], (uniform FReal)(0.0)), Vector1110);
 			const FVector PhiWithNormalInput = InverseTransformPosition(Frame, P);
 			FVector Normal;
-			float Phi;
+			FReal Phi;
 
 			const uniform SerializedPtr *uniform GeometryPtr = (const uniform SerializedPtr *uniform)Geometry;
 			const uniform uint8 *uniform Ptr = GeometryPtr[i].Ptr;
@@ -775,13 +775,13 @@ export void ApplyPerParticleCollisionFastFriction(uniform FVector ParticlesP[],
 			{
 			case Sphere:
 				{
-					const uniform float Margin = *((const uniform float *uniform)(Ptr + OffsetOfMargin));
+					const uniform FReal Margin = *((const uniform FReal *uniform)(Ptr + OffsetOfMargin));
 					PhiWithNormalSphere(TypedPtr, PhiWithNormalInput, Normal, Phi, Margin);
 					break;
 				}
 			case Capsule:
 				{
-					const uniform float Margin = *((const uniform float *uniform)(Ptr + OffsetOfMargin));
+					const uniform FReal Margin = *((const uniform FReal *uniform)(Ptr + OffsetOfMargin));
 					PhiWithNormalCapsule(TypedPtr, PhiWithNormalInput, Normal, Phi, Margin);
 					break;
 				}
@@ -808,13 +808,13 @@ export void ApplyPerParticleCollisionFastFriction(uniform FVector ParticlesP[],
 						IMMask = lanemask();
 					}
 
-					GetPhiWithNormal(CollisionParticles, (const uniform float *uniform)&PhiWithNormalInput, (uniform float *uniform)&Normal, (uniform float *uniform)&Phi, i, programCount, IMMask);
+					GetPhiWithNormal(CollisionParticles, (const uniform FReal *uniform)&PhiWithNormalInput, (uniform FReal *uniform)&Normal, (uniform FReal *uniform)&Phi, i, programCount, IMMask);
 					break;
 				}
 			}
 
-			const float Penetration = PerGroupThickness - Phi; // This is related to the Normal impulse
-			if(Penetration > 0.0f)
+			const FReal Penetration = PerGroupThickness - Phi; // This is related to the Normal impulse
+			if(Penetration > (uniform FReal)(0.0))
 			{
 				const FVector NormalWorld = TransformVector(Frame, Normal);
 				P = P + (Penetration * NormalWorld);
@@ -823,10 +823,10 @@ export void ApplyPerParticleCollisionFastFriction(uniform FVector ParticlesP[],
 				const FVector X = VectorLoad(&ParticlesX[extract(Index, 0)]);
 				const FVector RelativeDisplacement = (P - X) - (CollisionV[i] + VectorCross(CollisionW[i], VectorToPoint)) * Dt; // This corresponds to the tangential velocity multiplied by dt (friction will drive this to zero if it is high enough)
 				const FVector RelativeDisplacementTangent = RelativeDisplacement - VectorDot(RelativeDisplacement, NormalWorld) * NormalWorld; // Project displacement into the tangential plane
-				const float RelativeDisplacementTangentLength = VectorSize(RelativeDisplacementTangent);
+				const FReal RelativeDisplacementTangentLength = VectorSize(RelativeDisplacementTangent);
 
-				const float PositionCorrection = min(Penetration * PerGroupFriction, RelativeDisplacementTangentLength);
-				const float CorrectionRatio = select(RelativeDisplacementTangentLength < SMALL_NUMBER, 0.0f, PositionCorrection / RelativeDisplacementTangentLength);
+				const FReal PositionCorrection = min(Penetration * PerGroupFriction, RelativeDisplacementTangentLength);
+				const FReal CorrectionRatio = select(RelativeDisplacementTangentLength < SMALL_NUMBER, (uniform FReal)(0.0), PositionCorrection / RelativeDisplacementTangentLength);
 				P = P - (CorrectionRatio * RelativeDisplacementTangent);
 			}
 		}
@@ -840,20 +840,20 @@ export void ApplyPerParticleCollisionFastFriction(uniform FVector ParticlesP[],
 
 export void ApplyPerParticleCollisionNoFriction(uniform FVector ParticlesP[],
 												const uniform FVector ParticlesX[],
-												const uniform float ParticlesInvM[],
+												const uniform FReal ParticlesInvM[],
 												const uniform FVector CollisionV[],
 												const uniform FVector CollisionX[],
 												const uniform FVector CollisionW[],
 												const uniform FVector4 CollisionR[],
 												const uniform uint32 DynamicGroupId,
 												const uniform uint32 KinematicGroupIds[],
-												const uniform float PerGroupThickness,
+												const uniform FReal PerGroupThickness,
 												const uniform uint8 *uniform CollisionParticles,
 												const uniform uint8 *uniform Geometry,
 												const uniform int SizeofFImplicitObject,
 												const uniform int OffsetofGeometryType,
 												const uniform int OffsetOfMargin,
-												const uniform float Dt,
+												const uniform FReal Dt,
 												const uniform int InnerOffset,
 												const uniform int InnerRange,
 												const uniform int OuterOffset,
@@ -885,7 +885,7 @@ export void ApplyPerParticleCollisionNoFriction(uniform FVector ParticlesP[],
 
 	foreach(Index = OuterOffset ... BailAt)
 	{
-		const float IM = ParticlesInvM[Index];
+		const FReal IM = ParticlesInvM[Index];
 
 		if(all(IM == 0))
 		{
@@ -896,10 +896,10 @@ export void ApplyPerParticleCollisionNoFriction(uniform FVector ParticlesP[],
 
 		for(uniform int i = InnerOffset; i < InnerRange; i++) // SequentialFor
 		{
-			const uniform FTransform Frame = SetTransform(CollisionR[i], SetVector4(CollisionX[i], 0.0f), Float1110);
+			const uniform FTransform Frame = SetTransform(CollisionR[i], SetVector4(CollisionX[i], (uniform FReal)(0.0)), Vector1110);
 			const FVector PhiWithNormalInput = InverseTransformPosition(Frame, P);
 			FVector Normal;
-			float Phi;
+			FReal Phi;
 
 			const uniform SerializedPtr *uniform GeometryPtr = (const uniform SerializedPtr *uniform)Geometry;
 			const uniform uint8 *uniform Ptr = GeometryPtr[i].Ptr;
@@ -910,13 +910,13 @@ export void ApplyPerParticleCollisionNoFriction(uniform FVector ParticlesP[],
 			{
 			case Sphere:
 				{
-					const uniform float Margin = *((const uniform float *uniform)(Ptr + OffsetOfMargin));
+					const uniform FReal Margin = *((const uniform FReal *uniform)(Ptr + OffsetOfMargin));
 					PhiWithNormalSphere(TypedPtr, PhiWithNormalInput, Normal, Phi, Margin);
 					break;
 				}
 			case Capsule:
 				{
-					const uniform float Margin = *((const uniform float *uniform)(Ptr + OffsetOfMargin));
+					const uniform FReal Margin = *((const uniform FReal *uniform)(Ptr + OffsetOfMargin));
 					PhiWithNormalCapsule(TypedPtr, PhiWithNormalInput, Normal, Phi, Margin);
 					break;
 				}
@@ -943,13 +943,13 @@ export void ApplyPerParticleCollisionNoFriction(uniform FVector ParticlesP[],
 						IMMask = lanemask();
 					}
 
-					GetPhiWithNormal(CollisionParticles, (const uniform float *uniform)&PhiWithNormalInput, (uniform float *uniform)&Normal, (uniform float *uniform)&Phi, i, programCount, IMMask);
+					GetPhiWithNormal(CollisionParticles, (const uniform FReal *uniform)&PhiWithNormalInput, (uniform FReal *uniform)&Normal, (uniform FReal *uniform)&Phi, i, programCount, IMMask);
 					break;
 				}
 			}
 
-			const float Penetration = PerGroupThickness - Phi; // This is related to the Normal impulse
-			if(Penetration > 0.0f)
+			const FReal Penetration = PerGroupThickness - Phi; // This is related to the Normal impulse
+			if(Penetration > (uniform FReal)(0.0))
 			{
 				const FVector NormalWorld = TransformVector(Frame, Normal);
 				P = P + (Penetration * NormalWorld);
diff --git a/Engine/Source/Runtime/Experimental/Chaos/Private/Chaos/TriangleMesh.cpp b/Engine/Source/Runtime/Experimental/Chaos/Private/Chaos/TriangleMesh.cpp
--- a/Engine/Source/Runtime/Experimental/Chaos/Private/Chaos/TriangleMesh.cpp
+++ b/Engine/Source/Runtime/Experimental/Chaos/Private/Chaos/TriangleMesh.cpp
@@ -18,6 +18,9 @@
 #endif
 
 #if INTEL_ISPC && !UE_BUILD_SHIPPING
+static_assert(sizeof(ispc::FVector) == sizeof(Chaos::FVec3), "sizeof(ispc::FVector) != sizeof(Chaos::FVec3)");
+static_assert(sizeof(ispc::TArrayInt) == sizeof(TArray<int32>), "sizeof(ispc::TArrayInt) != sizeof(TArray<int32>)");
+
 bool bChaos_TriangleMesh_ISPC_Enabled = true;
 FAutoConsoleVariableRef CVarChaosTriangleMeshISPCEnabled(TEXT("p.Chaos.TriangleMesh.ISPC"), bChaos_TriangleMesh_ISPC_Enabled, TEXT("Whether to use ISPC optimizations in triangle mesh calculations"));
 #endif
@@ -256,7 +259,6 @@ void FTriangleMesh::GetFaceNormals(TArray<FVec3>& Normals, const TConstArrayView
 			Normals.SetNumUninitialized(MElements.Num());
 
 #if INTEL_ISPC
-			static_assert(std::is_same<FReal, float>::value == true, "ISPC only supports float template type");
 			ispc::GetFaceNormals(
 				(ispc::FVector*)Normals.GetData(),
 				(ispc::FVector*)Points.GetData(),
@@ -303,7 +305,6 @@ void FTriangleMesh::GetPointNormals(TArrayView<FVec3> PointNormals, const TConst
 	if (bRealTypeCompatibleWithISPC && bChaos_TriangleMesh_ISPC_Enabled)
 	{
 #if INTEL_ISPC
-		static_assert(std::is_same<FReal, float>::value == true, "ISPC only supports float template type");
 		ispc::GetPointNormals(
 			(ispc::FVector*)PointNormals.GetData(),
 			(const ispc::FVector*)FaceNormals.GetData(),
diff --git a/Engine/Source/Runtime/Experimental/Chaos/Private/Chaos/VelocityField.cpp b/Engine/Source/Runtime/Experimental/Chaos/Private/Chaos/VelocityField.cpp
--- a/Engine/Source/Runtime/Experimental/Chaos/Private/Chaos/VelocityField.cpp
+++ b/Engine/Source/Runtime/Experimental/Chaos/Private/Chaos/VelocityField.cpp
@@ -4,6 +4,10 @@
 #include "HAL/IConsoleManager.h"
 #if INTEL_ISPC
 #include "VelocityField.ispc.generated.h"
+
+static_assert(sizeof(ispc::FVector) == sizeof(Chaos::FVec3), "sizeof(ispc::FVector) != sizeof(Chaos::FVec3)");
+static_assert(sizeof(ispc::FIntVector) == sizeof(Chaos::TVec3<int32>), "sizeof(ispc::FIntVector) != sizeof(Chaos::TVec3<int32>)");
+static_assert(sizeof(ispc::FVector2) == sizeof(Chaos::FVec2), "sizeof(ispc::FVector2) != sizeof(Chaos::FVec2)");
 #endif
 
 #if INTEL_ISPC && !UE_BUILD_SHIPPING
diff --git a/Engine/Source/Runtime/Experimental/Chaos/Private/Chaos/VelocityField.ispc b/Engine/Source/Runtime/Experimental/Chaos/Private/Chaos/VelocityField.ispc
--- a/Engine/Source/Runtime/Experimental/Chaos/Private/Chaos/VelocityField.ispc
+++ b/Engine/Source/Runtime/Experimental/Chaos/Private/Chaos/VelocityField.ispc
@@ -4,11 +4,13 @@
 
 #include "Math/Vector.isph"
 
-static inline float SafeNormalize(FVector &Direction)
+static const FReal OneThird = 1.0/3.0;
+
+static inline FReal SafeNormalize(FVector &Direction)
 {
-	const float Size = sqrt(VectorSizeSquared(Direction));
+	const FReal Size = sqrt(VectorSizeSquared(Direction));
 	Direction = VectorSelect((Size < KINDA_SMALL_NUMBER), ForwardVector, Direction / Size);
-	return (Size < KINDA_SMALL_NUMBER) ? 0.f : Size;
+	return (Size < KINDA_SMALL_NUMBER) ? 0 : Size;
 }
 
 export void UpdateField(uniform FVector Forces[],
@@ -16,9 +18,9 @@ export void UpdateField(uniform FVector Forces[],
 								const uniform FVector V[],
 								const uniform FVector X[],
 								const uniform FVector& Velocity,
-								const uniform float QuarterRho,
-								const uniform float Cd,
-								const uniform float Cl,
+								const uniform FReal QuarterRho,
+								const uniform FReal Cd,
+								const uniform FReal Cl,
 								const uniform int32 NumElements)
 {
 	varying FIntVector Element;
@@ -43,19 +45,19 @@ export void UpdateField(uniform FVector Forces[],
 
 		// Calculate the normal and the area of the surface exposed to the flow
 		FVector N = VectorCross(X1 - X0, X2 - X0);
-		const float DoubleArea = SafeNormalize(N);
+		const FReal DoubleArea = SafeNormalize(N);
 
 		const FVector V0 = VectorGather(&V[Element.V[0]]);
 		const FVector V1 = VectorGather(&V[Element.V[1]]);
 		const FVector V2 = VectorGather(&V[Element.V[2]]);
 
 		// Calculate the direction and the relative velocity of the triangle to the flow
-		const FVector SurfaceVelocity = (1.0f/3.0f) * (V0 + V1 + V2);
+		const FVector SurfaceVelocity = OneThird * (V0 + V1 + V2);
 		const FVector RelVelocity = Velocity - SurfaceVelocity;
 
 		// Set the aerodynamic forces
-		const float VDotN = VectorDot(RelVelocity, N);
-		const float VSquare = VectorDot(RelVelocity, RelVelocity);
+		const FReal VDotN = VectorDot(RelVelocity, N);
+		const FReal VSquare = VectorDot(RelVelocity, RelVelocity);
 
 #if HW_GATHER_SUPPORTED == 1
 		if(ElementIndex + programCount < NumElements)
@@ -67,7 +69,7 @@ export void UpdateField(uniform FVector Forces[],
 		const FVector EvenForce = (Cd - Cl) * VDotN * RelVelocity + Cl * VSquare * N;
 		const FVector OddForce = (Cl - Cd) * VDotN * RelVelocity - Cl * VSquare * N;
 
-		const FVector Force = QuarterRho * DoubleArea * VectorSelect(VDotN >= 0.0f, EvenForce, OddForce);
+		const FVector Force = QuarterRho * DoubleArea * VectorSelect(VDotN >= 0, EvenForce, OddForce);
 
 		VectorStore(&Forces[Offset], Force);
 
@@ -81,11 +83,11 @@ export void UpdateFieldWithWeightMaps(uniform FVector Forces[],
 								const uniform FVector X[],
 								const uniform FVector2 Multipliers[],
 								const uniform FVector& Velocity,
-								const uniform float QuarterRho,
-								const uniform float DragBase,
-								const uniform float DragRange,
-								const uniform float LiftBase,
-								const uniform float LiftRange,
+								const uniform FReal QuarterRho,
+								const uniform FReal DragBase,
+								const uniform FReal DragRange,
+								const uniform FReal LiftBase,
+								const uniform FReal LiftRange,
 								const uniform int32 NumElements)
 {
 	varying FIntVector Element;
@@ -111,22 +113,22 @@ export void UpdateFieldWithWeightMaps(uniform FVector Forces[],
 
 		// Calculate the normal and the area of the surface exposed to the flow
 		FVector N = VectorCross(X1 - X0, X2 - X0);
-		const float DoubleArea = SafeNormalize(N);
+		const FReal DoubleArea = SafeNormalize(N);
 
 		const FVector V0 = VectorGather(&V[Element.V[0]]);
 		const FVector V1 = VectorGather(&V[Element.V[1]]);
 		const FVector V2 = VectorGather(&V[Element.V[2]]);
 
 		// Calculate the direction and the relative velocity of the triangle to the flow
-		const FVector SurfaceVelocity = (1.0f/3.0f) * (V0 + V1 + V2);
+		const FVector SurfaceVelocity = OneThird * (V0 + V1 + V2);
 		const FVector RelVelocity = Velocity - SurfaceVelocity;
 
 		// Set the aerodynamic forces
-		const float VDotN = VectorDot(RelVelocity, N);
-		const float VSquare = VectorDot(RelVelocity, RelVelocity);
+		const FReal VDotN = VectorDot(RelVelocity, N);
+		const FReal VSquare = VectorDot(RelVelocity, RelVelocity);
 
-		const float Cd = DragBase + DragRange * Multiplier.V[0];
-		const float Cl = LiftBase + LiftRange * Multiplier.V[1];
+		const FReal Cd = DragBase + DragRange * Multiplier.V[0];
+		const FReal Cl = LiftBase + LiftRange * Multiplier.V[1];
 
 #if HW_GATHER_SUPPORTED == 1
 		if(ElementIndex + programCount < NumElements)
@@ -138,7 +140,7 @@ export void UpdateFieldWithWeightMaps(uniform FVector Forces[],
 		const FVector EvenForce = (Cd - Cl) * VDotN * RelVelocity + Cl * VSquare * N;
 		const FVector OddForce = (Cl - Cd) * VDotN * RelVelocity - Cl * VSquare * N;
 
-		const FVector Force = QuarterRho * DoubleArea * VectorSelect(VDotN >= 0.0f, EvenForce, OddForce);
+		const FVector Force = QuarterRho * DoubleArea * VectorSelect(VDotN >= 0, EvenForce, OddForce);
 
 		VectorStore(&Forces[Offset], Force);
 
diff --git a/Engine/Source/Runtime/Experimental/Chaos/Public/Chaos/AABB.h b/Engine/Source/Runtime/Experimental/Chaos/Public/Chaos/AABB.h
--- a/Engine/Source/Runtime/Experimental/Chaos/Public/Chaos/AABB.h
+++ b/Engine/Source/Runtime/Experimental/Chaos/Public/Chaos/AABB.h
@@ -652,4 +652,13 @@ namespace Chaos
 	};
 
 	using FAABB3 = TAABB<FReal, 3>;
-}
\ No newline at end of file
+}
+
+// Support ISPC enable/disable in non-shipping builds
+#if !INTEL_ISPC
+const bool bChaos_AABBTransform_ISPC_Enabled = false;
+#elif UE_BUILD_SHIPPING
+const bool bChaos_AABBTransform_ISPC_Enabled = true;
+#else
+extern bool bChaos_AABBTransform_ISPC_Enabled;
+#endif
\ No newline at end of file
