LWC: Enable ISPC for LWC double types. Submitted on behalf of jeff.rous[REDACTED_HANDLE]

diff --git a/Engine/Source/Runtime/Experimental/ChaosCore/Public/Chaos/Matrix33.isph b/Engine/Source/Runtime/Experimental/ChaosCore/Public/Chaos/Matrix33.isph
--- a/Engine/Source/Runtime/Experimental/ChaosCore/Public/Chaos/Matrix33.isph
+++ b/Engine/Source/Runtime/Experimental/ChaosCore/Public/Chaos/Matrix33.isph
@@ -11,42 +11,42 @@ typedef FMatrix FMatrix33;
 static const uniform struct FMatrix33 Matrix33Zero = {{ 0, 0, 0, 0, 0, 0, 0, 0, 0, 0, 0, 0, 0, 0, 0, 1}};
 
 inline uniform FMatrix33 SetMatrix33(
-	const uniform float x00, const uniform float x10, const uniform float x20,
-	const uniform float x01, const uniform float x11, const uniform float x21,
-	const uniform float x02, const uniform float x12, const uniform float x22)
+	const uniform FReal x00, const uniform FReal x10, const uniform FReal x20,
+	const uniform FReal x01, const uniform FReal x11, const uniform FReal x21,
+	const uniform FReal x02, const uniform FReal x12, const uniform FReal x22)
 {
-	return SetMatrix(SetVector4(x00, x01, x02, 0), SetVector4(x10, x11, x12, 0), SetVector4(x20, x21, x22, 0), Float0001);
+	return SetMatrix(SetVector4(x00, x01, x02, ZERO), SetVector4(x10, x11, x12, ZERO), SetVector4(x20, x21, x22, ZERO), Vector0001);
 }
 
 inline FMatrix33 SetMatrix33(
-	const float x00, const float x10, const float x20,
-	const float x01, const float x11, const float x21,
-	const float x02, const float x12, const float x22)
+	const FReal x00, const FReal x10, const FReal x20,
+	const FReal x01, const FReal x11, const FReal x21,
+	const FReal x02, const FReal x12, const FReal x22)
 {
-	return SetMatrix(SetVector4(x00, x01, x02, 0.0f), SetVector4(x10, x11, x12, 0.0f), SetVector4(x20, x21, x22, 0.0f), Float0001);
+	return SetMatrix(SetVector4(x00, x01, x02, ZERO), SetVector4(x10, x11, x12, ZERO), SetVector4(x20, x21, x22, ZERO), Vector0001);
 }
 
 inline uniform FMatrix33 SetMatrix33(
-	const uniform float x00, const uniform float x10, const uniform float x20,
-	const uniform float x11, const uniform float x21, const uniform float x22)
+	const uniform FReal x00, const uniform FReal x10, const uniform FReal x20,
+	const uniform FReal x11, const uniform FReal x21, const uniform FReal x22)
 {
-	return SetMatrix(SetVector4(x00, x10, x20, 0), SetVector4(x10, x11, x21, 0), SetVector4(x20, x21, x22, 0), Float0001);
+	return SetMatrix(SetVector4(x00, x10, x20, ZERO), SetVector4(x10, x11, x21, ZERO), SetVector4(x20, x21, x22, ZERO), Vector0001);
 }
 
 inline uniform FMatrix33 SetMatrix33(
-	const uniform float x00, const uniform float x11, const uniform float x22)
+	const uniform FReal x00, const uniform FReal x11, const uniform FReal x22)
 {
-	return SetMatrix(SetVector4(x00, 0, 0, 0), SetVector4(0, x11, 0, 0), SetVector4(0, 0, x22, 0), Float0001);
+	return SetMatrix(SetVector4(x00, ZERO, ZERO, ZERO), SetVector4(ZERO, x11, ZERO, ZERO), SetVector4(ZERO, ZERO, x22, ZERO), Vector0001);
 }
 
 inline uniform FMatrix33 SetMatrix33(const uniform FVector4& x0, const uniform FVector4& x1, const uniform FVector4& x2)
 {
-	return SetMatrix(x0, x1, x2, Float0001);
+	return SetMatrix(x0, x1, x2, Vector0001);
 }
 
 inline FMatrix33 SetMatrix33(const FVector4& x0, const FVector4& x1, const FVector4& x2)
 {
-	return SetMatrix(x0, x1, x2, Float0001);
+	return SetMatrix(x0, x1, x2, Vector0001);
 }
 
 inline uniform FMatrix33 MultiplyAB(const uniform FMatrix33& L, const uniform FMatrix33& R)
@@ -139,7 +139,7 @@ inline varying FVector Multiply(const uniform FMatrix33& L, const varying FVecto
 	return SetVector(Result.V[0], Result.V[1], Result.V[2]);
 }
 
-inline uniform FMatrix33 Multiply(const uniform FMatrix33& L, const uniform float R)
+inline uniform FMatrix33 Multiply(const uniform FMatrix33& L, const uniform FReal R)
 {
 	return SetMatrix33(
 		L.M[0] * R, L.M[1] * R, L.M[2] * R,
@@ -148,7 +148,7 @@ inline uniform FMatrix33 Multiply(const uniform FMatrix33& L, const uniform floa
 	);
 }
 
-inline FMatrix33 Multiply(const FMatrix33& L, const float R)
+inline FMatrix33 Multiply(const FMatrix33& L, const FReal R)
 {
 	return SetMatrix33(
 		L.M[0] * R, L.M[1] * R, L.M[2] * R,
@@ -206,7 +206,7 @@ inline uniform FMatrix33 MakeQuatRotationTranslationMatrix(const uniform FVector
 	const uniform FVector4 x = xyz2 * Q.V[0]; // xx, xy, xz
 	const uniform FVector4 y = xyz2 * Q.V[1]; // yx, yy, yz
 	const uniform FVector4 w = xyz2 * Q.V[3]; // wx, wy, wz
-	const uniform float zz = Q.V[2] * xyz2.V[2];
+	const uniform FReal zz = Q.V[2] * xyz2.V[2];
 
 	const uniform FVector4 yx = SetVector4(y.V[2], x.V[2], x.V[1], 0); // yz, xz, xy
 	const uniform FVector4 Adds = yx + w; // M[6], M[8], M[1]
@@ -214,7 +214,7 @@ inline uniform FMatrix33 MakeQuatRotationTranslationMatrix(const uniform FVector
 
 	const uniform FVector4 OneMinus0 = SetVector4(y.V[1], x.V[0], x.V[0], 0); // yy, xx, xx
 	const uniform FVector4 Oneminus1 = SetVector4(zz, zz, y.V[1], 0); // zz, zz, yy
-	const uniform FVector4 OneMinuses = FloatOne - (OneMinus0 + Oneminus1); // M[0], M[5], M[10]
+	const uniform FVector4 OneMinuses = VectorOne - (OneMinus0 + Oneminus1); // M[0], M[5], M[10]
 
 	return SetMatrix33(
 		SetVector4(OneMinuses.V[0], Adds.V[2], Minuses.V[1], 0),
diff --git a/Engine/Source/Runtime/Experimental/ChaosCore/Public/Chaos/Real.h b/Engine/Source/Runtime/Experimental/ChaosCore/Public/Chaos/Real.h
--- a/Engine/Source/Runtime/Experimental/ChaosCore/Public/Chaos/Real.h
+++ b/Engine/Source/Runtime/Experimental/ChaosCore/Public/Chaos/Real.h
@@ -27,8 +27,8 @@ namespace Chaos
 #endif
 
 	/**
-	* ISPC optimization supports float, this allows classes that uses ISPC to branch to the right implementation 
+	* ISPC optimization supports float and double, this allows classes that uses ISPC to branch to the right implementation 
 	* without having to check the actual underlying type of FReal
 	*/
-	constexpr bool bRealTypeCompatibleWithISPC = (TAreTypesEqual<FReal, float>::Value == true);
+	constexpr bool bRealTypeCompatibleWithISPC = (TAreTypesEqual<FReal, float>::Value == true) || (TAreTypesEqual<FReal, double>::Value == true);
 }
diff --git a/Engine/Source/Runtime/Experimental/GeometryCollectionEngine/Private/GeometryCollection/GeometryCollectionComponent.cpp b/Engine/Source/Runtime/Experimental/GeometryCollectionEngine/Private/GeometryCollection/GeometryCollectionComponent.cpp
--- a/Engine/Source/Runtime/Experimental/GeometryCollectionEngine/Private/GeometryCollection/GeometryCollectionComponent.cpp
+++ b/Engine/Source/Runtime/Experimental/GeometryCollectionEngine/Private/GeometryCollection/GeometryCollectionComponent.cpp
@@ -61,6 +61,13 @@ MSVC_PRAGMA(warning(disable : ALL_CODE_ANALYSIS_WARNINGS))
 MSVC_PRAGMA(warning(pop))
 #endif    // USING_CODE_ANALYSIS
 
+static_assert(sizeof(ispc::FMatrix) == sizeof(FMatrix), "sizeof(ispc::FMatrix) != sizeof(FMatrix)");
+static_assert(sizeof(ispc::FBox) == sizeof(FBox), "sizeof(ispc::FBox) != sizeof(FBox)");
+#endif
+
+#if INTEL_ISPC && !UE_BUILD_SHIPPING
+bool bChaos_BoxCalcBounds_ISPC_Enabled = true;
+FAutoConsoleVariableRef CVarChaosBoxCalcBoundsISPCEnabled(TEXT("p.Chaos.BoxCalcBounds.ISPC"), bChaos_BoxCalcBounds_ISPC_Enabled, TEXT("Whether to use ISPC optimizations in calculating box bounds in geometry collections"));
 #endif
 
 DEFINE_LOG_CATEGORY_STATIC(UGCC_LOG, Error, All);
@@ -477,26 +484,31 @@ FBoxSphereBounds UGeometryCollectionComponent::CalcBounds(const FTransform& Loca
 		}
 		else
 		{
+			if (bChaos_BoxCalcBounds_ISPC_Enabled)
+			{
 #if INTEL_ISPC
-			ispc::BoxCalcBounds(
-				(int32 *)&TransformToGeometryIndex[0],
-				(int32 *)&TransformIndices[0],
-				(ispc::FMatrix *)&GlobalMatrices[0],
-				(ispc::FBox *)&BoundingBoxes[0],
-				(ispc::FMatrix &)LocalToWorldWithScale,
-				(ispc::FBox &)BoundingBox,
-				NumBoxes);
-#else
-			for (int32 BoxIdx = 0; BoxIdx < NumBoxes; ++BoxIdx)
+				ispc::BoxCalcBounds(
+					(int32*)&TransformToGeometryIndex[0],
+					(int32*)&TransformIndices[0],
+					(ispc::FMatrix*)&GlobalMatrices[0],
+					(ispc::FBox*)&BoundingBoxes[0],
+					(ispc::FMatrix&)LocalToWorldWithScale,
+					(ispc::FBox&)BoundingBox,
+					NumBoxes);
+#endif
+			}
+			else
 			{
-				const int32 TransformIndex = TransformIndices[BoxIdx];
-
-				if(RestCollection->GetGeometryCollection()->IsGeometry(TransformIndex))
+				for (int32 BoxIdx = 0; BoxIdx < NumBoxes; ++BoxIdx)
 				{
-					BoundingBox += BoundingBoxes[BoxIdx].TransformBy(GlobalMatrices[TransformIndex] * LocalToWorldWithScale);
+					const int32 TransformIndex = TransformIndices[BoxIdx];
+
+					if (RestCollection->GetGeometryCollection()->IsGeometry(TransformIndex))
+					{
+						BoundingBox += BoundingBoxes[BoxIdx].TransformBy(GlobalMatrices[TransformIndex] * LocalToWorldWithScale);
+					}
 				}
 			}
-#endif
 		}
 
 		return FBoxSphereBounds(BoundingBox);
diff --git a/Engine/Source/Runtime/Experimental/GeometryCollectionEngine/Private/GeometryCollection/GeometryCollectionComponent.ispc b/Engine/Source/Runtime/Experimental/GeometryCollectionEngine/Private/GeometryCollection/GeometryCollectionComponent.ispc
--- a/Engine/Source/Runtime/Experimental/GeometryCollectionEngine/Private/GeometryCollection/GeometryCollectionComponent.ispc
+++ b/Engine/Source/Runtime/Experimental/GeometryCollectionEngine/Private/GeometryCollection/GeometryCollectionComponent.ispc
@@ -15,7 +15,7 @@ struct FBox
 inline static FBox BoxDefaultInit()
 {
 	FBox Box;
-	Box.Min = Box.Max = SetVector(0.0f, 0.0f, 0.0f);
+	Box.Min = Box.Max = ZeroVector;
 	Box.IsValid = 0;
 	return Box;
 }
@@ -69,9 +69,8 @@ inline static FBox BoxTransformByMatrix(const FBox &Box, const FMatrix &M)
 	const FVector4 m2 = SetVector4(M.M[8], M.M[9], M.M[10], M.M[11]);
 	const FVector4 m3 = SetVector4(M.M[12], M.M[13], M.M[14], M.M[15]);
 
-	const uniform FVector Half = SetVector(0.5f, 0.5f, 0.5f);
-	const FVector Origin = (Box.Max + Box.Min) * Half;
-	const FVector Extent = (Box.Max - Box.Min) * Half;
+	const FVector Origin = (Box.Max + Box.Min) * HalfVector;
+	const FVector Extent = (Box.Max - Box.Min) * HalfVector;
 
 	FVector4 NewOrigin = SetVector4(Origin.V[0], Origin.V[0], Origin.V[0], Origin.V[0]) * m0;
 	NewOrigin = VectorMultiplyAdd(SetVector4(Origin.V[1], Origin.V[1], Origin.V[1], Origin.V[1]), m1, NewOrigin);
diff --git a/Engine/Source/Runtime/Experimental/GeometryCollectionEngine/Private/GeometryCollection/GeometryCollectionSceneProxy.cpp b/Engine/Source/Runtime/Experimental/GeometryCollectionEngine/Private/GeometryCollection/GeometryCollectionSceneProxy.cpp
--- a/Engine/Source/Runtime/Experimental/GeometryCollectionEngine/Private/GeometryCollection/GeometryCollectionSceneProxy.cpp
+++ b/Engine/Source/Runtime/Experimental/GeometryCollectionEngine/Private/GeometryCollection/GeometryCollectionSceneProxy.cpp
@@ -36,6 +36,13 @@
 #if USING_CODE_ANALYSIS
     MSVC_PRAGMA( warning( pop ) )
 #endif    // USING_CODE_ANALYSIS
+
+#if UE_LARGE_WORLD_COORDINATES_DISABLED
+	namespace ispc { typedef FVector FVector3f; typedef FMatrix FMatrix44f; } // Need this because ISPC doesn't expose typedefs in header and optimized function uses mixed types
+#endif
+
+static_assert(sizeof(ispc::FMatrix44f) == sizeof(FMatrix44f), "sizeof(ispc::FMatrix44f) != sizeof(FMatrix44f)");
+static_assert(sizeof(ispc::FVector3f) == sizeof(FVector3f), "sizeof(ispc::FVector3f) != sizeof(FVector3f)");
 #endif
 
 #include "ComponentReregisterContext.h"
@@ -81,6 +88,11 @@ FAutoConsoleVariableRef CVarRayTracingGeometryCollectionProxyMeshes(
 	ECVF_RenderThreadSafe
 );
 
+#if INTEL_ISPC && !UE_BUILD_SHIPPING
+bool bGeometryCollection_SetDynamicData_ISPC_Enabled = true;
+FAutoConsoleVariableRef CVarGeometryCollectionSetDynamicDataISPCEnabled(TEXT("r.GeometryCollectionSetDynamicData.ISPC"), bGeometryCollection_SetDynamicData_ISPC_Enabled, TEXT("Whether to use ISPC optimizations to set dynamic data in geometry collections"));
+#endif
+
 DEFINE_LOG_CATEGORY_STATIC(FGeometryCollectionSceneProxyLogging, Log, All);
 
 FGeometryCollectionDynamicDataPool GDynamicDataPool;
@@ -729,22 +741,28 @@ void FGeometryCollectionSceneProxy::SetDynamicData_RenderThread(FGeometryCollect
 				if (ThisBatchSize > 0)
 				{
 					const FMatrix44f* RESTRICT BoneTransformsPtr = DynamicData->IsDynamic ? DynamicData->Transforms.GetData() : ConstantData->RestTransforms.GetData();
-	#if INTEL_ISPC
-					uint8* VertexBufferOffset = (uint8*)VertexBufferData + (IndexOffset * VertexBuffer.GetStride());
-					ispc::SetDynamicData_RenderThread(
-						(ispc::FVector*)VertexBufferOffset, 
-						ThisBatchSize, 
-						VertexBuffer.GetStride(), 
-						&ConstantData->BoneMap[IndexOffset], 
-						(ispc::FMatrix*)BoneTransformsPtr,
-						(ispc::FVector*)&ConstantData->Vertices[IndexOffset]);
-	#else
-					for (int32 i = IndexOffset; i < IndexOffset + ThisBatchSize; i++)
+
+					if (bGeometryCollection_SetDynamicData_ISPC_Enabled)
+					{
+#if INTEL_ISPC
+						uint8* VertexBufferOffset = (uint8*)VertexBufferData + (IndexOffset * VertexBuffer.GetStride());
+						ispc::SetDynamicData_RenderThread(
+							(ispc::FVector3f*)VertexBufferOffset,
+							ThisBatchSize,
+							VertexBuffer.GetStride(),
+							&ConstantData->BoneMap[IndexOffset],
+							(ispc::FMatrix44f*)BoneTransformsPtr,
+							(ispc::FVector3f*)&ConstantData->Vertices[IndexOffset]);
+#endif
+					}
+					else
 					{
-						FVector3f Transformed = BoneTransformsPtr[ConstantData->BoneMap[i]].TransformPosition(ConstantData->Vertices[i]);
-						FMemory::Memcpy((uint8*)VertexBufferData + (i * VertexBuffer.GetStride()), &Transformed, sizeof(FVector3f));
+						for (int32 i = IndexOffset; i < IndexOffset + ThisBatchSize; i++)
+						{
+							FVector3f Transformed = BoneTransformsPtr[ConstantData->BoneMap[i]].TransformPosition(ConstantData->Vertices[i]);
+							FMemory::Memcpy((uint8*)VertexBufferData + (i * VertexBuffer.GetStride()), &Transformed, sizeof(FVector3f));
+						}
 					}
-	#endif
 				}
 			});
 
diff --git a/Engine/Source/Runtime/Experimental/GeometryCollectionEngine/Private/GeometryCollection/GeometryCollectionSceneProxy.h b/Engine/Source/Runtime/Experimental/GeometryCollectionEngine/Private/GeometryCollection/GeometryCollectionSceneProxy.h
--- a/Engine/Source/Runtime/Experimental/GeometryCollectionEngine/Private/GeometryCollection/GeometryCollectionSceneProxy.h
+++ b/Engine/Source/Runtime/Experimental/GeometryCollectionEngine/Private/GeometryCollection/GeometryCollectionSceneProxy.h
@@ -448,3 +448,12 @@ protected:
 	uint32 bCurrentlyInMotion : 1;
 	uint32 bRequiresGPUSceneUpdate : 1;
 };
+
+// Support ISPC enable/disable in non-shipping builds
+#if !INTEL_ISPC
+const bool bGeometryCollection_SetDynamicData_ISPC_Enabled = false;
+#elif UE_BUILD_SHIPPING
+const bool bGeometryCollection_SetDynamicData_ISPC_Enabled = true;
+#else
+extern bool bGeometryCollection_SetDynamicData_ISPC_Enabled;
+#endif
diff --git a/Engine/Source/Runtime/Experimental/GeometryCollectionEngine/Private/GeometryCollection/GeometryCollectionSceneProxy.ispc b/Engine/Source/Runtime/Experimental/GeometryCollectionEngine/Private/GeometryCollection/GeometryCollectionSceneProxy.ispc
--- a/Engine/Source/Runtime/Experimental/GeometryCollectionEngine/Private/GeometryCollection/GeometryCollectionSceneProxy.ispc
+++ b/Engine/Source/Runtime/Experimental/GeometryCollectionEngine/Private/GeometryCollection/GeometryCollectionSceneProxy.ispc
@@ -3,12 +3,12 @@
 #include "Math/Vector.isph"
 #include "Math/Matrix.isph"
 
-export void SetDynamicData_RenderThread(uniform FVector PositionBuffer[], 
+export void SetDynamicData_RenderThread(uniform FVector3f PositionBuffer[], 
 										const uniform unsigned int NumVertices, 
 										const uniform unsigned int Stride, 
 										const uniform int32 BoneMap[], 
-										const uniform FMatrix Transforms[], 
-										const uniform FVector Vertices[])
+										const uniform FMatrix44f Transforms[], 
+										const uniform FVector3f Vertices[])
 {
 	uniform unsigned int Chunk = 0;
 
@@ -17,10 +17,10 @@ export void SetDynamicData_RenderThread(uniform FVector PositionBuffer[],
 		uniform float * uniform pVertices = (uniform float * uniform)&Vertices[Chunk];
 		uniform float * uniform pPosition = (uniform float * uniform)&PositionBuffer[Chunk];
 
-		FVector P;
+		FVector3f P;
 		aos_to_soa3(pVertices, &P.V[0], &P.V[1], &P.V[2]);
 
-		FVector Out;
+		FVector3f Out;
 		const int32 Bone = BoneMap[i];
 		foreach_unique(Index in Bone)
 		{
diff --git a/Engine/Source/Runtime/Experimental/GeometryCollectionEngine/Public/GeometryCollection/GeometryCollectionComponent.h b/Engine/Source/Runtime/Experimental/GeometryCollectionEngine/Public/GeometryCollection/GeometryCollectionComponent.h
--- a/Engine/Source/Runtime/Experimental/GeometryCollectionEngine/Public/GeometryCollection/GeometryCollectionComponent.h
+++ b/Engine/Source/Runtime/Experimental/GeometryCollectionEngine/Public/GeometryCollection/GeometryCollectionComponent.h
@@ -857,3 +857,12 @@ private:
 	bool bIsMoving;
 
 };
+
+// Support ISPC enable/disable in non-shipping builds
+#if !INTEL_ISPC
+const bool bChaos_BoxCalcBounds_ISPC_Enabled = false;
+#elif UE_BUILD_SHIPPING
+const bool bChaos_BoxCalcBounds_ISPC_Enabled = true;
+#else
+extern bool bChaos_BoxCalcBounds_ISPC_Enabled;
+#endif
