Fixed crash on InterchangeBaseNode CDO creation.
The check wasn't checking the correct variable.

diff --git a/Engine/Source/Runtime/Experimental/Interchange/Core/Public/Nodes/InterchangeBaseNodeUtilities.h b/Engine/Source/Runtime/Experimental/Interchange/Core/Public/Nodes/InterchangeBaseNodeUtilities.h
--- a/Engine/Source/Runtime/Experimental/Interchange/Core/Public/Nodes/InterchangeBaseNodeUtilities.h
+++ b/Engine/Source/Runtime/Experimental/Interchange/Core/Public/Nodes/InterchangeBaseNodeUtilities.h
@@ -31,7 +31,7 @@ public:
 
 	void Initialize(const TSharedPtr<FAttributeStorage, ESPMode::ThreadSafe>& InAttributes, const FString& BaseKeyName)
 	{
-		check(Attributes.IsValid());
+		check(InAttributes.IsValid());
 		Attributes = InAttributes;
 		KeyCount = BaseKeyName;
 	}
