Chaos : Joint Constraint
- Reset Joint constraint based on internal plasticity values when the settings are updated from the proxy.

diff --git a/Engine/Source/Runtime/Experimental/Chaos/Private/Chaos/PBDJointConstraints.cpp b/Engine/Source/Runtime/Experimental/Chaos/Private/Chaos/PBDJointConstraints.cpp
--- a/Engine/Source/Runtime/Experimental/Chaos/Private/Chaos/PBDJointConstraints.cpp
+++ b/Engine/Source/Runtime/Experimental/Chaos/Private/Chaos/PBDJointConstraints.cpp
@@ -36,12 +36,17 @@ namespace Chaos
 
 	
 	FPBDJointConstraintHandle::FPBDJointConstraintHandle()
+		: bLinearPlasticityInitialized(false)
+		, bAngularPlasticityInitialized(false)
 	{
 	}
 
 	
 	FPBDJointConstraintHandle::FPBDJointConstraintHandle(FConstraintContainer* InConstraintContainer, int32 InConstraintIndex)
 		: TIndexedContainerConstraintHandle<FPBDJointConstraints>(InConstraintContainer, InConstraintIndex)
+		, bLinearPlasticityInitialized(false)
+		, bAngularPlasticityInitialized(false)
+
 	{
 	}
 
@@ -129,9 +134,33 @@ namespace Chaos
 		return ConcreteContainer()->GetConstraintSettings(ConstraintIndex);
 	}
 
-	void FPBDJointConstraintHandle::SetSettings(const FPBDJointSettings& Settings)
+	void FPBDJointConstraintHandle::SetSettings(const FPBDJointSettings& InSettings)
 	{
-		ConcreteContainer()->SetConstraintSettings(ConstraintIndex, Settings);
+		// buffer the previous targets so plasticiy can reuse them
+		FVec3 LinearTarget = GetSettings().LinearDrivePositionTarget;
+		FRotation3 AngularTarget = GetSettings().AngularDrivePositionTarget;
+		if (!bLinearPlasticityInitialized && !FMath::IsNearlyEqual(InSettings.LinearPlasticityLimit, FLT_MAX))
+		{
+			bLinearPlasticityInitialized = true;
+		}
+		if (!bAngularPlasticityInitialized && !FMath::IsNearlyEqual(InSettings.AngularPlasticityLimit, FLT_MAX))
+		{
+			bAngularPlasticityInitialized = true;
+		}
+
+
+		ConcreteContainer()->SetConstraintSettings(ConstraintIndex, InSettings);
+
+
+		// transfer the previous targets when controlled by plasticiy
+		if (bLinearPlasticityInitialized)
+		{
+			ConcreteContainer()->SetLinearDrivePositionTarget(ConstraintIndex,LinearTarget);
+		}
+		if (bAngularPlasticityInitialized)
+		{
+			ConcreteContainer()->SetAngularDrivePositionTarget(ConstraintIndex,AngularTarget);
+		}
 	}
 
 	TVector<FGeometryParticleHandle*, 2> FPBDJointConstraintHandle::GetConstrainedParticles() const 
@@ -654,6 +683,15 @@ namespace Chaos
 		ConstraintSettings[ConstraintIndex].Sanitize();
 	}
 
+	void FPBDJointConstraints::SetLinearDrivePositionTarget(int32 ConstraintIndex, FVec3 InLinearDrivePositionTarget)
+	{
+		ConstraintSettings[ConstraintIndex].LinearDrivePositionTarget = InLinearDrivePositionTarget;
+	}
+
+	void FPBDJointConstraints::SetAngularDrivePositionTarget(int32 ConstraintIndex, FRotation3 InAngularDrivePositionTarget)
+	{
+		ConstraintSettings[ConstraintIndex].AngularDrivePositionTarget = InAngularDrivePositionTarget;
+	}
 
 	int32 FPBDJointConstraints::GetConstraintIsland(int32 ConstraintIndex) const
 	{
diff --git a/Engine/Source/Runtime/Experimental/Chaos/Public/Chaos/PBDJointConstraints.h b/Engine/Source/Runtime/Experimental/Chaos/Public/Chaos/PBDJointConstraints.h
--- a/Engine/Source/Runtime/Experimental/Chaos/Public/Chaos/PBDJointConstraints.h
+++ b/Engine/Source/Runtime/Experimental/Chaos/Public/Chaos/PBDJointConstraints.h
@@ -73,6 +73,9 @@ namespace Chaos
 	protected:
 		using Base::ConstraintIndex;
 		using Base::ConcreteContainer;
+	private:
+		bool bLinearPlasticityInitialized;
+		bool bAngularPlasticityInitialized;
 	};
 
 	class CHAOS_API FPBDJointState
@@ -229,6 +232,8 @@ namespace Chaos
 		const FPBDJointSettings& GetConstraintSettings(int32 ConstraintIndex) const;
 
 		void SetConstraintSettings(int32 ConstraintIndex, const FPBDJointSettings& InConstraintSettings);
+		void SetLinearDrivePositionTarget(int32 ConstraintIndex, FVec3 InLinearDrivePositionTarget);
+		void SetAngularDrivePositionTarget(int32 ConstraintIndex, FRotation3 InAngularDrivePositionTarget);
 
 		int32 GetConstraintIsland(int32 ConstraintIndex) const;
 		int32 GetConstraintLevel(int32 ConstraintIndex) const;
