Updated heightfields to generate correct bounds and triangle queries when provided a negative scale - fixed queries, contacts and manifolds against negatively scaled heightfields

diff --git a/Engine/Source/Runtime/Experimental/Chaos/Private/Chaos/HeightField.cpp b/Engine/Source/Runtime/Experimental/Chaos/Private/Chaos/HeightField.cpp
--- a/Engine/Source/Runtime/Experimental/Chaos/Private/Chaos/HeightField.cpp
+++ b/Engine/Source/Runtime/Experimental/Chaos/Private/Chaos/HeightField.cpp
@@ -1295,12 +1295,12 @@ namespace Chaos
 		OutIntersections.Reset();
 
 		const FBounds2D FlatBounds = GetFlatBounds();
-		const FVec2 Scale2D(GeomData.Scale[0], GeomData.Scale[1]);
+		FVec2 Scale2D(GeomData.Scale[0], GeomData.Scale[1]);
+
+		InFlatBounds = FBounds2D::FromPoints(FlatBounds.Clamp(InFlatBounds.Min) / Scale2D, FlatBounds.Clamp(InFlatBounds.Max) / Scale2D);
 
-		InFlatBounds.Min = FlatBounds.Clamp(InFlatBounds.Min);
-		InFlatBounds.Max = FlatBounds.Clamp(InFlatBounds.Max);
-		TVec2<int32> MinCell = FlatGrid.Cell(InFlatBounds.Min / Scale2D);
-		TVec2<int32> MaxCell = FlatGrid.Cell(InFlatBounds.Max / Scale2D);
+		TVec2<int32> MinCell = FlatGrid.Cell(InFlatBounds.Min);
+		TVec2<int32> MaxCell = FlatGrid.Cell(InFlatBounds.Max);
 
 		// We want to capture the first cell (delta == 0) as well
 		const int32 NumX = MaxCell[0] - MinCell[0] + 1;
diff --git a/Engine/Source/Runtime/Experimental/Chaos/Public/Chaos/AABB.h b/Engine/Source/Runtime/Experimental/Chaos/Public/Chaos/AABB.h
--- a/Engine/Source/Runtime/Experimental/Chaos/Public/Chaos/AABB.h
+++ b/Engine/Source/Runtime/Experimental/Chaos/Public/Chaos/AABB.h
@@ -545,6 +545,22 @@ namespace Chaos
 
 		FORCEINLINE constexpr bool IsConvex() const { return true; }
 
+		/**
+		 * Given a set of points, wrap an AABB around them
+		 * @param P0 The first of the points to wrap
+		 * @param InPoints Parameter pack of all subsequent points
+		 */
+		template<typename... Points>
+		static TAABB<T, d> FromPoints(const TVector<T, d>& P0, const Points&... InPoints)
+		{
+			static_assert(sizeof...(InPoints) > 0);
+			static_assert(std::is_same_v<std::common_type_t<Points...>, TVector<T, d>>);
+
+			TAABB<T, d> Result(P0, P0);
+			(Result.GrowToInclude(InPoints), ...);
+			return Result;
+		}
+
 	private:
 		TVector<T, d> MMin, MMax;
 	};
diff --git a/Engine/Source/Runtime/Experimental/Chaos/Public/Chaos/HeightField.h b/Engine/Source/Runtime/Experimental/Chaos/Public/Chaos/HeightField.h
--- a/Engine/Source/Runtime/Experimental/Chaos/Public/Chaos/HeightField.h
+++ b/Engine/Source/Runtime/Experimental/Chaos/Public/Chaos/HeightField.h
@@ -134,7 +134,9 @@ namespace Chaos
 
 		virtual const FAABB3 BoundingBox() const
 		{
-			CachedBounds = FAABB3(LocalBounds.Min() * GeomData.Scale, LocalBounds.Max() * GeomData.Scale);
+			// Generate a correct bound including scales (may be negative)
+			CachedBounds = FAABB3::FromPoints(LocalBounds.Min() * GeomData.Scale, LocalBounds.Max() * GeomData.Scale);
+
 			return CachedBounds;
 		}
 
@@ -305,7 +307,8 @@ namespace Chaos
 				OutPts[2] *= Scale;
 				OutPts[3] *= Scale;
 
-				OutBounds.Scale(FVec3(Scale));
+				const FVec3 Scale3(Scale);
+				OutBounds = FAABB3::FromPoints(OutBounds.Min() * Scale3, OutBounds.Max() * Scale3);
 			}
 
 			FORCEINLINE FReal GetMinHeight() const
@@ -415,6 +418,27 @@ namespace Chaos
 				, Max(0)
 			{}
 
+			FBounds2D(const FVec2& InMin, const FVec2& InMax)
+				: Min(InMin)
+				, Max(InMax)
+			{}
+
+			/**
+			 * Given a set of points, wrap a 2D bounds around them
+			 * @param P0 The first of the points to wrap
+			 * @param InPoints Parameter pack of all subsequent points
+			 */
+			template<typename... Points>
+			static FBounds2D FromPoints(const FVec2& P0, const Points&... InPoints)
+			{
+				static_assert(sizeof...(InPoints) > 0);
+				static_assert(std::is_same_v<std::common_type_t<Points...>, FVec2>);
+
+				FBounds2D Result(P0, P0);
+				(Result.GrowToInclude(InPoints), ...);
+				return Result;
+			}
+
 			explicit FBounds2D(const FAABB3& In3DBounds)
 			{
 				Set(In3DBounds);
@@ -436,6 +460,12 @@ namespace Chaos
 				return InPoint[0] >= Min[0] && InPoint[0] <= Max[0] && InPoint[1] >= Min[1] && InPoint[1] <= Max[1];
 			}
 
+			void GrowToInclude(const FVec2& InPoint)
+			{
+				Min = { FMath::Min(Min.X, InPoint.X), FMath::Min(Min.Y, InPoint.Y) };
+				Max = { FMath::Max(Max.X, InPoint.X), FMath::Max(Max.Y, InPoint.Y) };
+			}
+
 			FVec2 Clamp(const FVec2& InToClamp, FReal InNudge = SMALL_NUMBER) const
 			{
 				const FVec2 NudgeVec(InNudge, InNudge);
diff --git a/Engine/Source/Runtime/Landscape/Private/LandscapeCollision.cpp b/Engine/Source/Runtime/Landscape/Private/LandscapeCollision.cpp
--- a/Engine/Source/Runtime/Landscape/Private/LandscapeCollision.cpp
+++ b/Engine/Source/Runtime/Landscape/Private/LandscapeCollision.cpp
@@ -438,7 +438,7 @@ void ULandscapeHeightfieldCollisionComponent::OnCreatePhysicsState()
 				// First add complex geometry
 				TUniquePtr<Chaos::FPerShapeData> NewShape = Chaos::FPerShapeData::CreatePerShapeData(ShapeArray.Num());
 
-				HeightfieldRef->Heightfield->SetScale(FinalScale);
+				HeightfieldRef->Heightfield->SetScale(FinalScale * LandscapeComponentTransform.GetScale3D().GetSignVector());
 				TUniquePtr<Chaos::TImplicitObjectTransformed<Chaos::FReal, 3>> ChaosHeightFieldFromCooked = MakeUnique<Chaos::TImplicitObjectTransformed<Chaos::FReal, 3>>(MakeSerializable(HeightfieldRef->Heightfield), Chaos::FRigidTransform3(FTransform::Identity));
 
 				// Setup filtering
