Fixed uninitialized value causing a NaN to be created in an unused raycast parameter - fallout from engine merge

diff --git a/Engine/Source/Runtime/Experimental/Chaos/Public/Chaos/AABBTree.h b/Engine/Source/Runtime/Experimental/Chaos/Public/Chaos/AABBTree.h
--- a/Engine/Source/Runtime/Experimental/Chaos/Public/Chaos/AABBTree.h
+++ b/Engine/Source/Runtime/Experimental/Chaos/Public/Chaos/AABBTree.h
@@ -271,7 +271,7 @@ struct TAABBTreeLeafArray : public TBoundsWrapperHelper<TPayloadType, T, bComput
 	template <typename TSQVisitor, typename TQueryFastData>
 	FORCEINLINE_DEBUGGABLE bool RaycastFast(const TVec3<T>& Start, TQueryFastData& QueryFastData, TSQVisitor& Visitor, const TVec3<T>& Dir, const TVec3<T> InvDir, const bool bParallel[3]) const
 	{
-		return RaycastSweepImp</*bSweep=*/false>(Start, QueryFastData, TVec3<T>(), Visitor, Dir, InvDir, bParallel);
+		return RaycastSweepImp</*bSweep=*/false>(Start, QueryFastData, TVec3<T>((T)0), Visitor, Dir, InvDir, bParallel);
 	}
 
 	template <typename TSQVisitor, typename TQueryFastData>
