Remove useless type in attribute storage
#rnx

diff --git a/Engine/Source/Runtime/Experimental/Interchange/Core/Private/Tests/StorageTest.cpp b/Engine/Source/Runtime/Experimental/Interchange/Core/Private/Tests/StorageTest.cpp
--- a/Engine/Source/Runtime/Experimental/Interchange/Core/Private/Tests/StorageTest.cpp
+++ b/Engine/Source/Runtime/Experimental/Interchange/Core/Private/Tests/StorageTest.cpp
@@ -12,6 +12,28 @@
 
 IMPLEMENT_SIMPLE_AUTOMATION_TEST(FAttributeStorageTest, "System.Runtime.Interchange.AttributeStorage", EAutomationTestFlags::EditorContext | EAutomationTestFlags::EngineFilter)
 
+
+#define ADD_ATTRIBUTE_STORAGE(ATTRIBUTE_TYPE, ATTRIBUTE_VALUE, COMPARE_CAST)																			\
+{																																						\
+	const ATTRIBUTE_TYPE RefValue = ATTRIBUTE_VALUE;																									\
+	FAttributeKey Key = CreateUniqueKey();																												\
+	if (!IsAttributeStorageResultSuccess(TestStorage.RegisterAttribute(Key, RefValue)))																	\
+	{																																					\
+		AddError(FString(TEXT("`AttributeStorage` must handle adding ATTRIBUTE_TYPE attribute")));														\
+	}																																					\
+	ATTRIBUTE_TYPE StoredValue;																															\
+	TestStorage.GetAttributeHandle<ATTRIBUTE_TYPE>(Key).Get(StoredValue);																				\
+	TestEqual(TEXT("`AttributeStorage` must handle add and retrieve" #ATTRIBUTE_TYPE "attribute"), COMPARE_CAST(StoredValue), COMPARE_CAST(RefValue));		\
+}
+
+#define ADD_ATTRIBUTE_STORAGE_NOCOMPARE(ATTRIBUTE_TYPE, ATTRIBUTE_VALUE)																				\
+	const ATTRIBUTE_TYPE RefValue = ATTRIBUTE_VALUE;																									\
+	FAttributeKey Key = CreateUniqueKey();																												\
+	if (!IsAttributeStorageResultSuccess(TestStorage.RegisterAttribute(Key, RefValue)))																	\
+	{																																					\
+		AddError(FString(TEXT("`AttributeStorage` must handle adding ATTRIBUTE_TYPE attribute")));														\
+	}
+
 bool FAttributeStorageTest::RunTest(const FString& Parameters)
 {
 	using namespace UE::Interchange;
@@ -37,6 +59,107 @@ bool FAttributeStorageTest::RunTest(const FString& Parameters)
 	int32 TestStoredSeed = 0;
 	const int32 RandHelperUInt8 = 255;
 
+	ADD_ATTRIBUTE_STORAGE(bool, false, );
+	ADD_ATTRIBUTE_STORAGE(bool, true, static_cast<int32>);
+	ADD_ATTRIBUTE_STORAGE(FBox, FBox(FVector(-1.0), FVector(1.0)), );
+	ADD_ATTRIBUTE_STORAGE(FBoxSphereBounds, FBoxSphereBounds(FVector(0.0), FVector(2.0), 2.2), );
+	ADD_ATTRIBUTE_STORAGE(FColor, FColor(0xdf,0xcf, 0x00), );
+	ADD_ATTRIBUTE_STORAGE(FDateTime, FDateTime(2022,01,25,10, 24), );
+	ADD_ATTRIBUTE_STORAGE(double, 0.1234567890123456789, );
+	ADD_ATTRIBUTE_STORAGE(float, 0.1234567f, );
+	ADD_ATTRIBUTE_STORAGE(FGuid, FGuid(0xA4F0E4CD, 0xF97C4375, 0x9C25C211, 0xFD17B8BF), );
+	ADD_ATTRIBUTE_STORAGE(int8, -0x7f, );
+	ADD_ATTRIBUTE_STORAGE(int8, 0x7f, );
+	ADD_ATTRIBUTE_STORAGE(int16, -0x75cc, );
+	ADD_ATTRIBUTE_STORAGE(int16, 0x7fcc, );
+	ADD_ATTRIBUTE_STORAGE(int32, -0x75ccabcd, );
+	ADD_ATTRIBUTE_STORAGE(int32, 0x7fccabcd, );
+	ADD_ATTRIBUTE_STORAGE(int64, -0x75ccabcd12345678, );
+	ADD_ATTRIBUTE_STORAGE(int64, 0x7fccabcd12345678, );
+	ADD_ATTRIBUTE_STORAGE(FIntRect, FIntRect(-1,-2, 1, 2), );
+	ADD_ATTRIBUTE_STORAGE(FLinearColor, FLinearColor(0.59, 0.49, 0), );
+	ADD_ATTRIBUTE_STORAGE(FMatrix, FMatrix(FTransform(FRotator(25.0, 2.0, 3.14159), FVector(150.0, -203.0, 4500.7), FVector(1.0, 1.0, 1.0)).ToMatrixWithScale()), );
+	ADD_ATTRIBUTE_STORAGE(FName, FName(TEXT("Testing FName storage")), );
+	ADD_ATTRIBUTE_STORAGE(FPlane, FPlane(FVector(1.0, 1.0, 0), FVector(0, 0, 1.0)), );
+	ADD_ATTRIBUTE_STORAGE(FQuat, FQuat(1.0, 1.0, 0, 1.0), );
+
+	{
+		ADD_ATTRIBUTE_STORAGE_NOCOMPARE(FRandomStream, FRandomStream(34))
+		FRandomStream StoredValue;
+		TestStorage.GetAttributeHandle<FRandomStream>(Key).Get(StoredValue);
+		TestEqual(TEXT("`AttributeStorage` must handle add and retrieve FRandomStream attribute"), StoredValue.ToString(), RefValue.ToString());
+	}
+	
+	ADD_ATTRIBUTE_STORAGE(FRotator, FRotator(25.0, 2.0, 3.14159), );
+	ADD_ATTRIBUTE_STORAGE(FString, FString(TEXT("Testing FString storage")), );
+	ADD_ATTRIBUTE_STORAGE(FTimespan, FTimespan(10, 44, 38), );
+
+	{
+		ADD_ATTRIBUTE_STORAGE_NOCOMPARE(FTransform, FTransform(FRotator(25.0, 2.0, 3.14159), FVector(150.0, -203.0, 4500.7), FVector(1.0, 1.0, 1.0)))
+		FTransform StoredValue;
+		TestStorage.GetAttributeHandle<FTransform>(Key).Get(StoredValue);
+		if (!StoredValue.Equals(RefValue))
+		{
+			AddError(FString(TEXT("`AttributeStorage` must handle adding FTransform attribute")));
+		}
+	}
+
+	ADD_ATTRIBUTE_STORAGE(FTwoVectors, FTwoVectors(FVector(150.0, -203.0, 4500.7), FVector(1.0, 1.0, 1.0)), );
+	ADD_ATTRIBUTE_STORAGE(uint8, 0x85, );
+	ADD_ATTRIBUTE_STORAGE(uint8, 0x7f, );
+	ADD_ATTRIBUTE_STORAGE(uint16, 0x85cc, );
+	ADD_ATTRIBUTE_STORAGE(uint16, 0x7fcc, );
+	ADD_ATTRIBUTE_STORAGE(uint32, 0x85ccabcd, );
+	ADD_ATTRIBUTE_STORAGE(uint32, 0x7fccabcd, );
+	ADD_ATTRIBUTE_STORAGE(uint64, 0x85ccabcd12345678, );
+	ADD_ATTRIBUTE_STORAGE(uint64, 0x7fccabcd12345678, );
+	ADD_ATTRIBUTE_STORAGE(FVector, FVector(150.0, -203.0, 4500.7), );
+	ADD_ATTRIBUTE_STORAGE(FVector2D, FVector2D(150.0, -203.0), );
+	ADD_ATTRIBUTE_STORAGE(FVector4, FVector4(150.0, -203.0, 4500.7, 2.0), );
+	ADD_ATTRIBUTE_STORAGE(FIntPoint, FIntPoint(1, 4), );
+	ADD_ATTRIBUTE_STORAGE(FIntVector, FIntVector(1, 4, 3), );
+
+	{
+		ADD_ATTRIBUTE_STORAGE_NOCOMPARE(FVector2DHalf, FVector2DHalf(150.0, -203.0))
+		FVector2DHalf StoredValue;
+		TestStorage.GetAttributeHandle<FVector2DHalf>(Key).Get(StoredValue);
+		if (TestEqual(TEXT("`AttributeStorage` must handle add and retrieve FVector2DHalf attribute"), StoredValue.X, RefValue.X))
+		{
+			TestEqual(TEXT("`AttributeStorage` must handle add and retrieve FVector2DHalf attribute"), StoredValue.Y, RefValue.Y);
+		}
+	}
+
+	ADD_ATTRIBUTE_STORAGE(FFloat16, FFloat16(0.1234567f), );
+
+	{
+		ADD_ATTRIBUTE_STORAGE_NOCOMPARE(FOrientedBox, FOrientedBox())
+		FOrientedBox StoredValue;
+		TestStorage.GetAttributeHandle<FOrientedBox>(Key).Get(StoredValue);
+		if (!StoredValue.AxisX.Equals(RefValue.AxisX) ||
+			!StoredValue.AxisY.Equals(RefValue.AxisY) ||
+			!StoredValue.AxisZ.Equals(RefValue.AxisZ) ||
+			!FMath::IsNearlyEqual(StoredValue.ExtentX, RefValue.ExtentX) ||
+			!FMath::IsNearlyEqual(StoredValue.ExtentY, RefValue.ExtentY) ||
+			!FMath::IsNearlyEqual(StoredValue.ExtentZ, RefValue.ExtentZ) ||
+			!StoredValue.Center.Equals(RefValue.Center))
+		{
+			AddError(FString(TEXT("`AttributeStorage` must handle adding FOrientedBox attribute")));
+		}
+	}
+
+	{
+		ADD_ATTRIBUTE_STORAGE_NOCOMPARE(FSphere, FSphere(FVector(2.0), 2.2))
+		FSphere StoredValue;
+		TestStorage.GetAttributeHandle<FSphere>(Key).Get(StoredValue);
+		if (!StoredValue.Equals(RefValue))
+		{
+			AddError(FString(TEXT("`AttributeStorage` must handle adding FSphere attribute")));
+		}
+	}
+	ADD_ATTRIBUTE_STORAGE(FFrameNumber, FFrameNumber(30), );
+	ADD_ATTRIBUTE_STORAGE(FFrameRate, FFrameRate(1,60), );
+	ADD_ATTRIBUTE_STORAGE(FFrameTime, FFrameTime(29,0.5f), );
+	ADD_ATTRIBUTE_STORAGE(FSoftObjectPath, FSoftObjectPath(UClass::StaticClass()), );
 	//Test Adding/Reading uint8 attribute with 0 value
 	{
 		const uint8 RefValue = 0;
diff --git a/Engine/Source/Runtime/Experimental/Interchange/Core/Private/Types/AttributeStorage.cpp b/Engine/Source/Runtime/Experimental/Interchange/Core/Private/Types/AttributeStorage.cpp
--- a/Engine/Source/Runtime/Experimental/Interchange/Core/Private/Types/AttributeStorage.cpp
+++ b/Engine/Source/Runtime/Experimental/Interchange/Core/Private/Types/AttributeStorage.cpp
@@ -53,13 +53,9 @@ namespace UE
 				case EAttributeTypes::Vector4: { AttributeTypeString = TEXT("Vector4"); } break;
 				case EAttributeTypes::IntPoint: { AttributeTypeString = TEXT("IntPoint"); } break;
 				case EAttributeTypes::IntVector: { AttributeTypeString = TEXT("IntVector"); } break;
-				case EAttributeTypes::NetworkGUID: { AttributeTypeString = TEXT("NetworkGUID"); } break;
 				case EAttributeTypes::Vector2DHalf: { AttributeTypeString = TEXT("Vector2DHalf"); } break;
 				case EAttributeTypes::Float16: { AttributeTypeString = TEXT("Float16"); } break;
-				case EAttributeTypes::CapsuleShape: { AttributeTypeString = TEXT("CapsuleShape"); } break;
-				case EAttributeTypes::DualQuat: { AttributeTypeString = TEXT("DualQuat"); } break;
 				case EAttributeTypes::OrientedBox: { AttributeTypeString = TEXT("OrientedBox"); } break;
-				case EAttributeTypes::Ray: { AttributeTypeString = TEXT("Ray"); } break;
 				case EAttributeTypes::Sphere: { AttributeTypeString = TEXT("Sphere"); } break;
 				case EAttributeTypes::FrameNumber: { AttributeTypeString = TEXT("FrameNumber"); } break;
 				case EAttributeTypes::FrameRate: { AttributeTypeString = TEXT("FrameRate"); } break;
diff --git a/Engine/Source/Runtime/Experimental/Interchange/Core/Public/Types/AttributeStorage.h b/Engine/Source/Runtime/Experimental/Interchange/Core/Public/Types/AttributeStorage.h
--- a/Engine/Source/Runtime/Experimental/Interchange/Core/Public/Types/AttributeStorage.h
+++ b/Engine/Source/Runtime/Experimental/Interchange/Core/Public/Types/AttributeStorage.h
@@ -5,10 +5,8 @@
 #include "Containers/SortedMap.h"
 #include "Containers/UnrealString.h"
 #include "CoreMinimal.h"
-#include "Math/DualQuat.h"
 #include "Math/RandomStream.h"
 #include "Misc/FrameRate.h"
-#include "Misc/NetworkGuid.h"
 #include "Misc/ScopeLock.h"
 #include "UObject/SoftObjectPath.h"
 
@@ -182,18 +180,14 @@ namespace UE
 			Vector4				= 34,
 			IntPoint			= 35,
 			IntVector			= 36,
-			NetworkGUID			= 37,
-			Vector2DHalf		= 38,
-			Float16				= 39,
-			CapsuleShape		= 40,
-			DualQuat			= 41,
-			OrientedBox			= 42,
-			Ray					= 43,
-			Sphere				= 44,
-			FrameNumber			= 45,
-			FrameRate			= 46,
-			FrameTime			= 47,
-			SoftObjectPath		= 48,
+			Vector2DHalf		= 37,
+			Float16				= 38,
+			OrientedBox			= 39,
+			Sphere				= 40,
+			FrameNumber			= 41,
+			FrameRate			= 42,
+			FrameTime			= 43,
+			SoftObjectPath		= 44,
 		};
 
 		/**
@@ -1119,21 +1113,6 @@ namespace UE
 			static CONSTEXPR EAttributeTypes GetType() { return EAttributeTypes::Int64; }
 		};
 
-
-		/** Implements variant type traits for the built-in FIntPoint type. */
-		template<> struct TAttributeTypeTraits<FIntPoint>
-		{
-			static CONSTEXPR EAttributeTypes GetType() { return EAttributeTypes::IntPoint; }
-		};
-
-
-		/** Implements variant type traits for the built-in FIntVector type. */
-		template<> struct TAttributeTypeTraits<FIntVector>
-		{
-			static CONSTEXPR EAttributeTypes GetType() { return EAttributeTypes::IntVector; }
-		};
-
-
 		/** Implements variant type traits for the built-in FIntRect type. */
 		template<> struct TAttributeTypeTraits<FIntRect>
 		{
@@ -1154,6 +1133,11 @@ namespace UE
 			static CONSTEXPR EAttributeTypes GetType() { return EAttributeTypes::Matrix; }
 		};
 
+		/** Implements variant type traits for the built-in FName type. */
+		template<> struct TAttributeTypeTraits<FName>
+		{
+			static CONSTEXPR EAttributeTypes GetType() { return EAttributeTypes::Name; }
+		};
 
 		/** Implements variant type traits for the built-in FPlane type. */
 		template<> struct TAttributeTypeTraits<FPlane>
@@ -1168,14 +1152,6 @@ namespace UE
 			static CONSTEXPR EAttributeTypes GetType() { return EAttributeTypes::Quat; }
 		};
 
-
-		/** Implements variant type traits for the built-in FName type. */
-		template<> struct TAttributeTypeTraits<FName>
-		{
-			static CONSTEXPR EAttributeTypes GetType() { return EAttributeTypes::Name; }
-		};
-
-
 		/** Implements variant type traits for the built-in FRandomStream type. */
 		template<> struct TAttributeTypeTraits<FRandomStream>
 		{
@@ -1259,18 +1235,22 @@ namespace UE
 			static CONSTEXPR EAttributeTypes GetType() { return EAttributeTypes::Vector2d; }
 		};
 
-
 		/** Implements variant type traits for the built-in FVector4 type. */
 		template<> struct TAttributeTypeTraits<FVector4>
 		{
 			static CONSTEXPR EAttributeTypes GetType() { return EAttributeTypes::Vector4; }
 		};
 
+		/** Implements variant type traits for the built-in Vector2DHalf type. */
+		template<> struct TAttributeTypeTraits<FIntPoint>
+		{
+			static CONSTEXPR EAttributeTypes GetType() { return EAttributeTypes::IntPoint; }
+		};
 
-		/** Implements variant type traits for the built-in NetworkGUID type. */
-		template<> struct TAttributeTypeTraits<FNetworkGUID>
+		/** Implements variant type traits for the built-in Vector2DHalf type. */
+		template<> struct TAttributeTypeTraits<FIntVector>
 		{
-			static CONSTEXPR EAttributeTypes GetType() { return EAttributeTypes::NetworkGUID; }
+			static CONSTEXPR EAttributeTypes GetType() { return EAttributeTypes::IntVector; }
 		};
 
 		/** Implements variant type traits for the built-in Vector2DHalf type. */
@@ -1285,30 +1265,12 @@ namespace UE
 			static CONSTEXPR EAttributeTypes GetType() { return EAttributeTypes::Float16; }
 		};
 
-		/** Implements variant type traits for the built-in CapsuleShape type. */
-		template<> struct TAttributeTypeTraits<FCapsuleShape>
-		{
-			static CONSTEXPR EAttributeTypes GetType() { return EAttributeTypes::CapsuleShape; }
-		};
-
-		/** Implements variant type traits for the built-in DualQuat type. */
-		template<> struct TAttributeTypeTraits<FDualQuat>
-		{
-			static CONSTEXPR EAttributeTypes GetType() { return EAttributeTypes::DualQuat; }
-		};
-
 		/** Implements variant type traits for the built-in OrientedBox type. */
 		template<> struct TAttributeTypeTraits<FOrientedBox>
 		{
 			static CONSTEXPR EAttributeTypes GetType() { return EAttributeTypes::OrientedBox; }
 		};
 
-		/** Implements variant type traits for the built-in Ray type. */
-		template<> struct TAttributeTypeTraits<FRay>
-		{
-			static CONSTEXPR EAttributeTypes GetType() { return EAttributeTypes::Ray; }
-		};
-
 		/** Implements variant type traits for the built-in Sphere type. */
 		template<> struct TAttributeTypeTraits<FSphere>
 		{
