Chaos : Make sure 4.27-Chaos assets can be loaded in UE5 ( in relation to back integration CL 18747785 )

diff --git a/Engine/Source/Runtime/Core/Public/UObject/PhysicsObjectVersion.h b/Engine/Source/Runtime/Core/Public/UObject/PhysicsObjectVersion.h
--- a/Engine/Source/Runtime/Core/Public/UObject/PhysicsObjectVersion.h
+++ b/Engine/Source/Runtime/Core/Public/UObject/PhysicsObjectVersion.h
@@ -54,6 +54,18 @@ struct CORE_API FPhysicsObjectVersion
 		// Chaos FGeometryCollectionObject user defined collision shapes support
 		GeometryCollectionUserDefinedCollisionShapes,
 
+		// Chaos Remove scale from TKinematicTarget object
+		ChaosKinematicTargetRemoveScale,
+
+		// Chaos Added support for per-object collision constraint flag.
+		AddCollisionConstraintFlag,
+
+		// Expose particle Disabled flag to the game thread
+		AddDisabledFlag,
+
+		// Added max linear and angular speed to Chaos bodies
+		AddChaosMaxLinearAngularSpeed,
+
 		// -----<new versions can be added above this line>-------------------------------------------------
 		VersionPlusOne,
 		LatestVersion = VersionPlusOne - 1
diff --git a/Engine/Source/Runtime/Experimental/Chaos/Public/Chaos/KinematicTargets.h b/Engine/Source/Runtime/Experimental/Chaos/Public/Chaos/KinematicTargets.h
--- a/Engine/Source/Runtime/Experimental/Chaos/Public/Chaos/KinematicTargets.h
+++ b/Engine/Source/Runtime/Experimental/Chaos/Public/Chaos/KinematicTargets.h
@@ -80,8 +80,12 @@ namespace Chaos
 		friend FChaosArchive& operator<<(FChaosArchive& Ar, TKinematicTarget<T, d>& KinematicTarget)
 		{
 			Ar.UsingCustomVersion(FFortniteReleaseBranchCustomObjectVersion::GUID);
+			Ar.UsingCustomVersion(FPhysicsObjectVersion::GUID);
 
-			if (Ar.CustomVer(FFortniteReleaseBranchCustomObjectVersion::GUID) >= FFortniteReleaseBranchCustomObjectVersion::ChaosKinematicTargetRemoveScale)
+			const bool bRemovedScaleFN = (Ar.CustomVer(FFortniteReleaseBranchCustomObjectVersion::GUID) >= FFortniteReleaseBranchCustomObjectVersion::ChaosKinematicTargetRemoveScale);
+			const bool bRemovedScaleUE4 = (Ar.CustomVer(FPhysicsObjectVersion::GUID) >= FPhysicsObjectVersion::ChaosKinematicTargetRemoveScale);
+
+			if (bRemovedScaleFN || bRemovedScaleUE4)
 			{
 				Ar << KinematicTarget.Position << KinematicTarget.Rotation << KinematicTarget.Mode;
 			}
diff --git a/Engine/Source/Runtime/Experimental/Chaos/Public/Chaos/ParticleDirtyFlags.h b/Engine/Source/Runtime/Experimental/Chaos/Public/Chaos/ParticleDirtyFlags.h
--- a/Engine/Source/Runtime/Experimental/Chaos/Public/Chaos/ParticleDirtyFlags.h
+++ b/Engine/Source/Runtime/Experimental/Chaos/Public/Chaos/ParticleDirtyFlags.h
@@ -301,6 +301,7 @@ public:
 		Ar.UsingCustomVersion(FUE5ReleaseStreamObjectVersion::GUID);
 		Ar.UsingCustomVersion(FUE5MainStreamObjectVersion::GUID);
 		Ar.UsingCustomVersion(FExternalPhysicsCustomObjectVersion::GUID);
+		Ar.UsingCustomVersion(FPhysicsObjectVersion::GUID);
 
 		Ar << MLinearEtherDrag;
 		Ar << MAngularEtherDrag;
@@ -321,17 +322,23 @@ public:
 			Ar << bCCDEnabled;
 		}
 
-		if (Ar.CustomVer(FUE5MainStreamObjectVersion::GUID) >= FUE5MainStreamObjectVersion::AddCollisionConstraintFlag)
+		const bool bAddCollisionConstraintFlagUE4 = (Ar.CustomVer(FPhysicsObjectVersion::GUID) >= FPhysicsObjectVersion::AddCollisionConstraintFlag);
+		const bool bAddCollisionConstraintFlagUE5 = (Ar.CustomVer(FUE5MainStreamObjectVersion::GUID) >= FUE5MainStreamObjectVersion::AddCollisionConstraintFlag);		
+		if (bAddCollisionConstraintFlagUE4 || bAddCollisionConstraintFlagUE5)
 		{
 			Ar << MCollisionConstraintFlag;
 		}
 
-		if (Ar.CustomVer(FUE5MainStreamObjectVersion::GUID) >= FUE5MainStreamObjectVersion::AddDisabledFlag)
+		const bool bAddDisableFlagUE4 = (Ar.CustomVer(FPhysicsObjectVersion::GUID) >= FPhysicsObjectVersion::AddDisabledFlag);
+		const bool bAddDisableFlagUE5 = (Ar.CustomVer(FUE5MainStreamObjectVersion::GUID) >= FUE5MainStreamObjectVersion::AddDisabledFlag);
+		if (bAddDisableFlagUE4 || bAddDisableFlagUE5)
 		{
 			Ar << bDisabled;
 		}
 		
-		if(Ar.CustomVer(FUE5ReleaseStreamObjectVersion::GUID) >= FUE5ReleaseStreamObjectVersion::AddChaosMaxLinearAngularSpeed)
+		const bool bAddChaosMaxLinearAngularSpeedUE4 = (Ar.CustomVer(FPhysicsObjectVersion::GUID) >= FPhysicsObjectVersion::AddChaosMaxLinearAngularSpeed);
+		const bool bAddChaosMaxLinearAngularSpeedUE5 = (Ar.CustomVer(FUE5ReleaseStreamObjectVersion::GUID) >= FUE5ReleaseStreamObjectVersion::AddChaosMaxLinearAngularSpeed);
+		if (bAddChaosMaxLinearAngularSpeedUE4 || bAddChaosMaxLinearAngularSpeedUE5)
 		{
 			Ar << MMaxLinearSpeedSq;
 			Ar << MMaxAngularSpeedSq;
