Chaos - Support for zero margins (still disabled in settings)
- Lower tolerances for GJK and EPA
- Various zero-margin related optimizations

diff --git a/Engine/Source/Programs/HeadlessChaos/Private/HeadlessChaosTestContactModification.cpp b/Engine/Source/Programs/HeadlessChaos/Private/HeadlessChaosTestContactModification.cpp
--- a/Engine/Source/Programs/HeadlessChaos/Private/HeadlessChaosTestContactModification.cpp
+++ b/Engine/Source/Programs/HeadlessChaos/Private/HeadlessChaosTestContactModification.cpp
@@ -188,9 +188,7 @@ namespace ChaosTest
 					int32 NumContacts = PairModifier.GetNumContacts();
 					for (int32 PointIdx = 0; PointIdx < NumContacts; ++PointIdx)
 					{
-						// Tell solver contacts are SeparationPadding units closer than they are, so object will float SeparationPadding units over floor.
-						FReal Separation = PairModifier.GetSeparation(PointIdx);
-						PairModifier.ModifySeparation(Separation - SeparationPadding, PointIdx);
+						PairModifier.ModifyTargetSeparation(SeparationPadding, PointIdx);
 					}
 				}
 			}
@@ -390,21 +388,10 @@ namespace ChaosTest
 						int32 DynamicIdx = (UniqueIndices[0] == Idx0) ? 0 : 1;
 						const FVec3 CoM = FParticleUtilities::GetCoMWorldPosition(FConstGenericParticleHandle(Particles[DynamicIdx]));
 
-						if (DynamicIdx == 0)
-						{
-							// Move point0 under center of mass and move second point under CoM but keep the same distance between bodies
-							FVec3 PointUnderCoM0(CoM.X, CoM.Y, WorldPos0.Z);
-							FVec3 Delta = PointUnderCoM0 - WorldPos0;
-							FVec3 PointUnderCoM1(WorldPos1 + Delta);
-							PairModifier.ModifyWorldContactLocations(PointUnderCoM0, PointUnderCoM1, PointIdx);
-						}
-						else
-						{
-							FVec3 PointUnderCoM1(CoM.X, CoM.Y, WorldPos1.Z);
-							FVec3 Delta = PointUnderCoM1 - WorldPos1;
-							FVec3 PointUnderCoM0(WorldPos0 + Delta);
-							PairModifier.ModifyWorldContactLocations(PointUnderCoM0, PointUnderCoM1, PointIdx);
-						}
+						// Move point0 under center of mass and move second point under CoM but keep the same distance between bodies
+						FVec3 PointUnderCoM0(CoM.X, CoM.Y, WorldPos0.Z);
+						FVec3 PointUnderCoM1(CoM.X, CoM.Y, WorldPos1.Z);
+						PairModifier.ModifyWorldContactLocations(PointUnderCoM0, PointUnderCoM1, PointIdx);
 					}
 				}
 			}
diff --git a/Engine/Source/Programs/HeadlessChaos/Private/HeadlessChaosTestEPA.cpp b/Engine/Source/Programs/HeadlessChaos/Private/HeadlessChaosTestEPA.cpp
--- a/Engine/Source/Programs/HeadlessChaos/Private/HeadlessChaosTestEPA.cpp
+++ b/Engine/Source/Programs/HeadlessChaos/Private/HeadlessChaosTestEPA.cpp
@@ -17,7 +17,7 @@ namespace ChaosTest
 {
 	using namespace Chaos;
 
-	void ValidFace(const FVec3* Verts, const TArray<TEPAEntry<FReal>>& TetFaces, int32 Idx)
+	void ValidFace(const FVec3* Verts, const TEPAWorkingArray<TEPAEntry<FReal>>& TetFaces, int32 Idx)
 	{
 		const TEPAEntry<FReal>& Entry = TetFaces[Idx];
 
@@ -95,7 +95,7 @@ namespace ChaosTest
 		{
 			TArray<FVec3> VertsA = { {-1,-1,1}, {-1,-1,-1}, {-1,1,-1}, {1,1,-1} };
 			TArray<FVec3> VertsB = { FVec3(0), FVec3(0), FVec3(0), FVec3(0) };
-			TArray<TEPAEntry<FReal>> TetFaces;
+			TEPAWorkingArray<TEPAEntry<FReal>> TetFaces;
 			FVec3 TouchingNormal;
 			EXPECT_TRUE(InitializeEPA(VertsA,VertsB,ErrorSupport,ErrorSupport,TetFaces,TouchingNormal));
 
@@ -109,7 +109,7 @@ namespace ChaosTest
 		{
 			TArray<FVec3> VertsA = { {-1,-1,-1}, {-1,-1,1}, {-1,1,-1}, {1,1,-1} };
 			TArray<FVec3> VertsB = { FVec3(0), FVec3(0), FVec3(0), FVec3(0) };
-			TArray<TEPAEntry<FReal>> TetFaces;
+			TEPAWorkingArray<TEPAEntry<FReal>> TetFaces;
 			FVec3 TouchingNormal;
 
 			EXPECT_TRUE(InitializeEPA(VertsA,VertsB,ErrorSupport,ErrorSupport,TetFaces,TouchingNormal));
@@ -179,7 +179,7 @@ namespace ChaosTest
 				TArray<FVec3> VertsA = { AllVerts[0], AllVerts[1], AllVerts[2] };
 				TArray<FVec3> VertsB = { FVec3(0), FVec3(0), FVec3(0) };
 
-				TArray<TEPAEntry<FReal>> TetFaces;
+				TEPAWorkingArray<TEPAEntry<FReal>> TetFaces;
 				FVec3 TouchingNormal;
 				EXPECT_TRUE(InitializeEPA(VertsA,VertsB,ASupport,EmptySupport,TetFaces,TouchingNormal));
 				EXPECT_VECTOR_NEAR(VertsA[3], AllVerts[3], 1e-4);
@@ -207,7 +207,7 @@ namespace ChaosTest
 				TArray<FVec3> VertsA = { AllVerts[1], AllVerts[0], AllVerts[2] };
 				TArray<FVec3> VertsB = { FVec3(0), FVec3(0), FVec3(0) };
 
-				TArray<TEPAEntry<FReal>> TetFaces;
+				TEPAWorkingArray<TEPAEntry<FReal>> TetFaces;
 				FVec3 TouchingNormal;
 				EXPECT_TRUE(InitializeEPA(VertsA,VertsB,ASupport,EmptySupport,TetFaces,TouchingNormal));
 				EXPECT_VECTOR_NEAR(VertsA[3], AllVerts[3], 1e-4);
@@ -235,7 +235,7 @@ namespace ChaosTest
 				TArray<FVec3> VertsA = { AllVerts[1], AllVerts[0], AllVerts[2] };
 				TArray<FVec3> VertsB = { FVec3(0), FVec3(0), FVec3(0) };
 
-				TArray<TEPAEntry<FReal>> TetFaces;
+				TEPAWorkingArray<TEPAEntry<FReal>> TetFaces;
 				FVec3 TouchingNormal;
 				EXPECT_FALSE(InitializeEPA(VertsA,VertsB,ASupportNoX,EmptySupport,TetFaces,TouchingNormal));
 				EXPECT_EQ(TouchingNormal.Z,0);
@@ -278,7 +278,7 @@ namespace ChaosTest
 				TArray<FVec3> VertsA = { AllVerts[0], AllVerts[2] };
 				TArray<FVec3> VertsB = { FVec3(0), FVec3(0) };
 
-				TArray<TEPAEntry<FReal>> TetFaces;
+				TEPAWorkingArray<TEPAEntry<FReal>> TetFaces;
 				FVec3 TouchingNormal;
 				EXPECT_TRUE(InitializeEPA(VertsA,VertsB,ASupport,EmptySupport,TetFaces,TouchingNormal));
 				EXPECT_VECTOR_NEAR(VertsA[2], AllVerts[1], 1e-4);
@@ -299,7 +299,7 @@ namespace ChaosTest
 				TArray<FVec3> VertsA = { AllVerts[2], AllVerts[0] };
 				TArray<FVec3> VertsB = { FVec3(0), FVec3(0) };
 
-				TArray<TEPAEntry<FReal>> TetFaces;
+				TEPAWorkingArray<TEPAEntry<FReal>> TetFaces;
 				FVec3 TouchingNormal;
 				EXPECT_TRUE(InitializeEPA(VertsA,VertsB,ASupport,EmptySupport,TetFaces,TouchingNormal));
 				EXPECT_VECTOR_NEAR(VertsA[2], AllVerts[1], 1e-4);
@@ -337,7 +337,7 @@ namespace ChaosTest
 				TArray<FVec3> VertsA = { AllVerts[2], AllVerts[0] };
 				TArray<FVec3> VertsB = { FVec3(0), FVec3(0) };
 
-				TArray<TEPAEntry<FReal>> TetFaces;
+				TEPAWorkingArray<TEPAEntry<FReal>> TetFaces;
 				FVec3 TouchingNormal;
 				EXPECT_FALSE(InitializeEPA(VertsA,VertsB,ASupportNoX,EmptySupport, TetFaces, TouchingNormal));
 				EXPECT_EQ(TouchingNormal.X,0);
@@ -366,7 +366,7 @@ namespace ChaosTest
 				TArray<FVec3> VertsA = { AllVerts[2], AllVerts[0] };
 				TArray<FVec3> VertsB = { FVec3(0), FVec3(0) };
 
-				TArray<TEPAEntry<FReal>> TetFaces;
+				TEPAWorkingArray<TEPAEntry<FReal>> TetFaces;
 				FVec3 TouchingNormal;
 				EXPECT_FALSE(InitializeEPA(VertsA,VertsB,ASupportNoXOrZ,EmptySupport,TetFaces,TouchingNormal));
 				EXPECT_EQ(TouchingNormal.X,0);
diff --git a/Engine/Source/Runtime/Experimental/Chaos/Private/Chaos/Collision/PBDCollisionConstraint.cpp b/Engine/Source/Runtime/Experimental/Chaos/Private/Chaos/Collision/PBDCollisionConstraint.cpp
--- a/Engine/Source/Runtime/Experimental/Chaos/Private/Chaos/Collision/PBDCollisionConstraint.cpp
+++ b/Engine/Source/Runtime/Experimental/Chaos/Private/Chaos/Collision/PBDCollisionConstraint.cpp
@@ -193,8 +193,8 @@ namespace Chaos
 		, LastShapeWorldRotationDelta()
 		, SolverBodies{ nullptr, nullptr }
 		, GJKWarmStartData()
-		, ManifoldPoints()
 		, SavedManifoldPoints()
+		, ManifoldPoints()
 	{
 	}
 
@@ -228,8 +228,8 @@ namespace Chaos
 		, LastShapeWorldRotationDelta()
 		, SolverBodies{ nullptr, nullptr }
 		, GJKWarmStartData()
-		, ManifoldPoints()
 		, SavedManifoldPoints()
+		, ManifoldPoints()
 	{
 	}
 
@@ -263,34 +263,52 @@ namespace Chaos
 	void FPBDCollisionConstraint::InitMarginsAndTolerances(const EImplicitObjectType ImplicitType0, const EImplicitObjectType ImplicitType1, const FReal Margin0, const FReal Margin1)
 	{
 		// Set up the margins and tolerances to be used during the narrow phase.
-		// One shape in a collision will always have a margin. Only triangles have zero margin and we don't 
-		// collide two triangles. If we have a triangle, it is always the second shape.
+		// When convex margins are enabled, at least one shape in a collision will always have a margin.
+		// If convex margins are disabled, only quadratic shapes have a margin (their radius).
 		// The collision tolerance is used for knowing whether a new contact matches an existing one.
-		// If we have two non-quadratic shapes, we use the smallest margin on both shapes (unless one shape has zero margin).
-		// If we have a quadratic shape versus a non-quadratic, we don't need a margin on the non-quadratic.
-		// For non-quadratics the collision tolerance is the smallest non-zero margin. 
+		//
+		// Margins: (Assuming convex margins are enabled...)
+		// If we have two polygonal shapes, we use the smallest of the two margins (unless one shape has zero margin, e.g. triangle).
+		// If we have a quadratic shape versus a polygonal shape, we use a zero margin on the polygoinal shape.
+		// Note: If we have a triangle, it is always the second shape (currently we do not support triangle-triangle collision)
+		//
+		// CollisionTolerance:
+		// For polygonal shapes the collision tolerance is proportional to the size of the smaller object. 
 		// For quadratic shapes we want a collision tolerance much smaller than the radius.
+		//
 		const bool bIsQuadratic0 = ((ImplicitType0 == ImplicitObjectType::Sphere) || (ImplicitType0 == ImplicitObjectType::Capsule));
 		const bool bIsQuadratic1 = ((ImplicitType1 == ImplicitObjectType::Sphere) || (ImplicitType1 == ImplicitObjectType::Capsule));
+		
+		// @todo(chaos): should probably be tunable. Used to use the same settings as the margin scale (for convex), but we want to support zero
+		// margins, but still have a non-zero CollisionTolerance (it is used for matching contact points for friction and manifold reuse)
+		const FReal ToleranceScale = 0.1f;
 		const FReal QuadraticToleranceScale = 0.05f;
+		
 		if (!bIsQuadratic0 && !bIsQuadratic1)
 		{
+			const FReal MaxSize0 = ((Implicit[0] != nullptr) && Implicit[0]->HasBoundingBox()) ? Implicit[0]->BoundingBox().Extents().GetAbsMax() : FReal(0);
+			const FReal MaxSize1 = ((Implicit[1] != nullptr) && Implicit[1]->HasBoundingBox()) ? Implicit[1]->BoundingBox().Extents().GetAbsMax() : FReal(0);
+			const FReal MaxSize = FMath::Min(MaxSize0, MaxSize1);
+			CollisionTolerance = ToleranceScale * MaxSize;
+
 			// If one shape has a zero margin, enforce a minimum margin to avoid the EPA degenerate case. E.g., Box-Triangle
 			// If both shapes have a margin, use the smaller margin on both shapes. E.g., Box-Box
 			// We should never see both shapes with zero margin, but if we did we'd end up with a zero margin
-			if ((Margin0 == FReal(0)) || (Margin1 == FReal(0)))
+			const FReal MinMargin = Chaos_Collision_ConvexZeroMargin;
+			if (Margin0 == FReal(0))
 			{
-				const FReal MinMargin = Chaos_Collision_ConvexZeroMargin;
-				const FReal NonZeroMargin = FMath::Max(Margin0, Margin1);
-				CollisionMargins[0] = FMath::Min(MinMargin, NonZeroMargin);
-				CollisionMargins[1] = CollisionMargins[0];
-				CollisionTolerance = NonZeroMargin;
+				CollisionMargins[0] = 0;
+				CollisionMargins[1] = FMath::Max(MinMargin, Margin1);
+			}
+			else if (Margin1 == FReal(0))
+			{
+				CollisionMargins[0] = FMath::Max(MinMargin, Margin0);
+				CollisionMargins[1] = 0;
 			}
 			else
 			{
 				CollisionMargins[0] = FMath::Min(Margin0, Margin1);
 				CollisionMargins[1] = CollisionMargins[0];
-				CollisionTolerance = CollisionMargins[0];
 			}
 		}
 		else if (bIsQuadratic0 && bIsQuadratic1)
@@ -544,7 +562,7 @@ namespace Chaos
 
 	void FPBDCollisionConstraint::ResetManifold()
 	{
-		SavedManifoldPoints.Reset();
+		ResetSavedManifoldPoints();
 		ResetActiveManifoldContacts();
 	}
 
@@ -930,6 +948,9 @@ namespace Chaos
 				{
 					BestScore = Score;
 					MatchIndex = SavedManifoldPointIndex;
+
+					// Just take the first match we find
+					break;
 				}
 			}
 
diff --git a/Engine/Source/Runtime/Experimental/Chaos/Private/Chaos/CollisionOneShotManifolds.cpp b/Engine/Source/Runtime/Experimental/Chaos/Private/Chaos/CollisionOneShotManifolds.cpp
--- a/Engine/Source/Runtime/Experimental/Chaos/Private/Chaos/CollisionOneShotManifolds.cpp
+++ b/Engine/Source/Runtime/Experimental/Chaos/Private/Chaos/CollisionOneShotManifolds.cpp
@@ -45,6 +45,12 @@ namespace Chaos
 	FAutoConsoleVariableRef CVarChaos_Collision_EnableManifoldReplace(TEXT("p.Chaos.Collision.EnableManifoldGJKReplace"), bChaos_Collision_EnableManifoldGJKReplace, TEXT(""));
 	FAutoConsoleVariableRef CVarChaos_Collision_EnableManifoldInject(TEXT("p.Chaos.Collision.EnableManifoldGJKInject"), bChaos_Collision_EnableManifoldGJKInject, TEXT(""));
 
+	// See GJKContactPointMargin for comments on why these matter
+	FRealSingle Chaos_Collision_GJKEpsilon = 1.e-6f;
+	FRealSingle Chaos_Collision_EPAEpsilon = 1.e-6f;
+	FAutoConsoleVariableRef CVarChaos_Collision_GJKEpsilon(TEXT("p.Chaos.Collision.GJKEpsilon"), Chaos_Collision_GJKEpsilon, TEXT(""));
+	FAutoConsoleVariableRef CVarChaos_Collision_EPAEpsilon(TEXT("p.Chaos.Collision.EPAEpsilon"), Chaos_Collision_EPAEpsilon, TEXT(""));
+
 	
 	namespace Collisions
 	{
@@ -421,7 +427,7 @@ namespace Chaos
 		// Use GJK to find the closest points (or shallowest penetrating points) on two convex shapes usingthe specified margin
 		// @todo(chaos): dedupe from GJKContactPoint in CollisionResolution.cpp
 		template <typename GeometryA, typename GeometryB>
-		FContactPoint GJKContactPointMargin(const GeometryA& A, const GeometryB& B, const FRigidTransform3& ATM, const FRigidTransform3& BTM, FReal MarginA, FReal MarginB, FGJKSimplexData& InOutGjkWarmStartData, FReal& OutMaxMarginDelta, int32& VertexIndexA, int32& VertexIndexB)
+		FContactPoint GJKContactPointMargin(const GeometryA& A, const GeometryB& B, const FRigidTransform3& BToATM, FReal MarginA, FReal MarginB, FGJKSimplexData& InOutGjkWarmStartData, FReal& OutMaxMarginDelta, int32& VertexIndexA, int32& VertexIndexB)
 		{
 			SCOPE_CYCLE_COUNTER_MANIFOLD_GJK();
 
@@ -430,14 +436,24 @@ namespace Chaos
 			FReal Penetration;
 			FVec3 ClosestA, ClosestB, NormalA, NormalB;
 
-			// Slightly increased epsilon to reduce error in normal for almost touching objects.
-			const FReal Epsilon = 3.e-3f;
+			// GJK and EPA tolerances.
+			// The GJK tolerance controls the separating distance at which GJK hands over to EPA. We need to be
+			// able to normalize vectors of size epsilon with low error to avoid bad normals. 
+			// This is a problem for floats with Epsilon < 1e-3f. For doubles it can be nmuch smaller although this
+			// increases how many iterations we need when quadratic shapes are involved.
+			// The EPA tolerance is used to determine whether a point projects inside a face on the simplezx.
+			// EPA is also sensitive to the GJK tolerance, because it assumes the simplex from GJK contains the origin
+			// but if it does not and the origin also does not project onto any simplex planes, we can end up with
+			// a bad case in EPA where it rejects the correct face and returns some arbitrary other face. We need an
+			// EPA Tolerance of about 1e-6 to avoid this in our test cases. It may need to be smaller.
+			//const FReal Epsilon = 3.e-3f;	// original value for float implementation
+			const FReal GJKEpsilon = Chaos_Collision_GJKEpsilon;
+			const FReal EPAEpsilon = Chaos_Collision_EPAEpsilon;
 
 			const TGJKCoreShape<GeometryA> AWithMargin(A, MarginA);
 			const TGJKCoreShape<GeometryB> BWithMargin(B, MarginB);
-			const FRigidTransform3 BToATM = BTM.GetRelativeTransformNoScale(ATM);
 
-			if (GJKPenetrationWarmStartable(AWithMargin, BWithMargin, BToATM, Penetration, ClosestA, ClosestB, NormalA, NormalB, VertexIndexA, VertexIndexB, InOutGjkWarmStartData, OutMaxMarginDelta, Epsilon))
+			if (GJKPenetrationWarmStartable(AWithMargin, BWithMargin, BToATM, Penetration, ClosestA, ClosestB, NormalA, NormalB, VertexIndexA, VertexIndexB, InOutGjkWarmStartData, OutMaxMarginDelta, GJKEpsilon, EPAEpsilon))
 			{
 				Contact.ShapeContactPoints[0] = ClosestA;
 				Contact.ShapeContactPoints[1] = ClosestB;
@@ -665,12 +681,13 @@ namespace Chaos
 			// Get the adjusted margins for each convex
 			const FReal Margin1 = Constraint.GetCollisionMargin0();
 			const FReal Margin2 = Constraint.GetCollisionMargin1();
+			const FRigidTransform3 Convex2ToConvex1Transform = Convex2Transform.GetRelativeTransformNoScale(Convex1Transform);
 
 			// Find the deepest penetration. This is used to determine the planes and points to use for the manifold
 			// MaxMarginDelta is an upper bound on the distance from the contact on the rounded core shape to the actual shape surface. 
 			FReal MaxMarginDelta = FReal(0);
 			int32 VertexIndexA = INDEX_NONE, VertexIndexB = INDEX_NONE;
-			FContactPoint GJKContactPoint = GJKContactPointMargin(Convex1, Convex2, Convex1Transform, Convex2Transform, Margin1, Margin2, Constraint.GetGJKWarmStartData(), MaxMarginDelta, VertexIndexA, VertexIndexB);
+			FContactPoint GJKContactPoint = GJKContactPointMargin(Convex1, Convex2, Convex2ToConvex1Transform, Margin1, Margin2, Constraint.GetGJKWarmStartData(), MaxMarginDelta, VertexIndexA, VertexIndexB);
 			PHYSICS_CSV_CUSTOM_EXPENSIVE(PhysicsCounters, NumManifoldsGJKCalled, 1, ECsvCustomStatOp::Accumulate);
 
 			const bool bCanUpdateManifold = bChaos_Collision_EnableManifoldGJKReplace;
@@ -693,19 +710,20 @@ namespace Chaos
 
 			PHYSICS_CSV_CUSTOM_EXPENSIVE(PhysicsCounters, NumManifoldsCreated, 1, ECsvCustomStatOp::Accumulate);
 
-			const FRigidTransform3 Convex2ToConvex1Transform = Convex2Transform.GetRelativeTransformNoScale(Convex1Transform);
 
 			// @todo(chaos): get the vertex index from GJK and use to to get the plane
 			const FVec3 SeparationDirectionLocalConvex1 = Convex2ToConvex1Transform.TransformVectorNoScale(GJKContactPoint.ShapeContactNormal);
 			const int32 MostOpposingPlaneIndexConvex1 = SelectContactPlane(Convex1, GJKContactPoint.ShapeContactPoints[0], SeparationDirectionLocalConvex1, Margin1, VertexIndexA);
-			const TPlaneConcrete<FReal, 3> BestPlaneConvex1 = Convex1.GetPlane(MostOpposingPlaneIndexConvex1);
-			const FReal BestPlaneDotNormalConvex1 = !bConvex1IsCapsule ? FMath::Abs(FVec3::DotProduct(-SeparationDirectionLocalConvex1, BestPlaneConvex1.Normal())) : -FLT_MAX;
+			FVec3 BestPlanePosition1, BestPlaneNormal1;
+			Convex1.GetPlaneNX(MostOpposingPlaneIndexConvex1, BestPlaneNormal1, BestPlanePosition1);
+			const FReal BestPlaneDotNormalConvex1 = !bConvex1IsCapsule ? FMath::Abs(FVec3::DotProduct(-SeparationDirectionLocalConvex1, BestPlaneNormal1)) : -FLT_MAX;
 
 			// Now for Convex2
 			const FVec3 SeparationDirectionLocalConvex2 = GJKContactPoint.ShapeContactNormal;
 			const int32 MostOpposingPlaneIndexConvex2 = SelectContactPlane(Convex2, GJKContactPoint.ShapeContactPoints[1], -SeparationDirectionLocalConvex2, Margin2, VertexIndexB);
-			const TPlaneConcrete<FReal, 3> BestPlaneConvex2 = Convex2.GetPlane(MostOpposingPlaneIndexConvex2);
-			const FReal BestPlaneDotNormalConvex2 = FMath::Abs(FVec3::DotProduct(SeparationDirectionLocalConvex2, BestPlaneConvex2.Normal()));
+			FVec3 BestPlanePosition2, BestPlaneNormal2;
+			Convex2.GetPlaneNX(MostOpposingPlaneIndexConvex2, BestPlaneNormal2, BestPlanePosition2);
+			const FReal BestPlaneDotNormalConvex2 = FMath::Abs(FVec3::DotProduct(SeparationDirectionLocalConvex2, BestPlaneNormal2));
 
 			const FReal SmallBiasToPreventFeatureFlipping = 0.002f; // This improves frame coherence by penalizing convex 1 in favour of convex 2
 			bool ReferenceFaceConvex1 = true; // Is the reference face on convex1 or convex2?
@@ -718,7 +736,7 @@ namespace Chaos
 			const FReal PlaneContactNormalEpsilon = Chaos_Collision_Manifold_PlaneContactNormalEpsilon;
 			const bool bIsPlaneContact = FMath::IsNearlyEqual(BestPlaneDotNormalConvex1, (FReal)1., PlaneContactNormalEpsilon) || FMath::IsNearlyEqual(BestPlaneDotNormalConvex2, (FReal)1., PlaneContactNormalEpsilon);
 
-			// For edge-edge contacts, we find the edges involved and project the contact onto the edges
+			// For edge-edge contacts, we find the edges involved and project the contact onto the edges (if necessary - see below)
 			if (!bIsPlaneContact)
 			{
 				SCOPE_CYCLE_COUNTER_MANIFOLD_ADDEDGEEDGE();
@@ -728,24 +746,34 @@ namespace Chaos
 					return;
 				}
 
-				// @todo(chaos): this does not work well when the edges are parallel. We should always have points with zero
-				// position delta perpendicular to the normal, but that is not the case for parallel edges
-				FVec3 ShapeEdgePos1 = Convex1.GetClosestEdgePosition(MostOpposingPlaneIndexConvex1, GJKContactPoint.ShapeContactPoints[0]);
-				FVec3 ShapeEdgePos2 = Convex2.GetClosestEdgePosition(MostOpposingPlaneIndexConvex2, GJKContactPoint.ShapeContactPoints[1]);
-				if (bConvex1IsCapsule)
+				// @todo(chaos): remove this if we ditch convex margins
+				// We only need to find the best edges when we have convex margins enabled because then the GJK result
+				// is a contact on the rounded margin-reduced shape and not the actual shape. If we are using zero
+				// margins, then the GJK result is already on the real shape surface. Likewise, if we have a
+				// quadratic-versus-polygonal contact, the GJK results are already on the real surfaces.
+				const FReal PolygonalMargin1 = Constraint.IsQuadratic0() ? FReal(0) : Margin1;
+				const FReal PolygonalMargin2 = Constraint.IsQuadratic1() ? FReal(0) : Margin2;
+				if ((PolygonalMargin1 > 0) || (PolygonalMargin2 > 0))
 				{
-					ShapeEdgePos1 -= Margin1 * SeparationDirectionLocalConvex1;
-				}
+					// @todo(chaos): this does not work well when the edges are parallel. We should always have points with zero
+					// position delta perpendicular to the normal, but that is not the case for parallel edges
+					FVec3 ShapeEdgePos1 = Convex1.GetClosestEdgePosition(MostOpposingPlaneIndexConvex1, GJKContactPoint.ShapeContactPoints[0]);
+					FVec3 ShapeEdgePos2 = Convex2.GetClosestEdgePosition(MostOpposingPlaneIndexConvex2, GJKContactPoint.ShapeContactPoints[1]);
+					if (bConvex1IsCapsule)
+					{
+						ShapeEdgePos1 -= Margin1 * SeparationDirectionLocalConvex1;
+					}
 
-				const FVec3 EdgePos1In2 = Convex2ToConvex1Transform.InverseTransformPositionNoScale(ShapeEdgePos1);
-				const FVec3 EdgePos2In2 = ShapeEdgePos2;
-				const FReal EdgePhi = FVec3::DotProduct(EdgePos1In2 - EdgePos2In2, GJKContactPoint.ShapeContactNormal);
-				const FVec3 EdgePosIn2 = FReal(0.5) * (EdgePos1In2 + EdgePos2In2);
+					const FVec3 EdgePos1In2 = Convex2ToConvex1Transform.InverseTransformPositionNoScale(ShapeEdgePos1);
+					const FVec3 EdgePos2In2 = ShapeEdgePos2;
+					const FReal EdgePhi = FVec3::DotProduct(EdgePos1In2 - EdgePos2In2, GJKContactPoint.ShapeContactNormal);
+					const FVec3 EdgePosIn2 = FReal(0.5) * (EdgePos1In2 + EdgePos2In2);
 
-				GJKContactPoint.ShapeContactPoints[0] = Convex2ToConvex1Transform.TransformPositionNoScale(EdgePosIn2 + FReal(0.5) * EdgePhi * GJKContactPoint.ShapeContactNormal);
-				GJKContactPoint.ShapeContactPoints[1] = EdgePosIn2 - FReal(0.5) * EdgePhi * GJKContactPoint.ShapeContactNormal;
-				GJKContactPoint.Phi = EdgePhi;
-				// Normal unchanged from GJK result
+					GJKContactPoint.ShapeContactPoints[0] = Convex2ToConvex1Transform.TransformPositionNoScale(EdgePosIn2 + FReal(0.5) * EdgePhi * GJKContactPoint.ShapeContactNormal);
+					GJKContactPoint.ShapeContactPoints[1] = EdgePosIn2 - FReal(0.5) * EdgePhi * GJKContactPoint.ShapeContactNormal;
+					GJKContactPoint.Phi = EdgePhi;
+					// Normal unchanged from GJK result
+				}
 
 				Constraint.AddOneshotManifoldContact(GJKContactPoint);
 				return;
@@ -753,8 +781,8 @@ namespace Chaos
 
 			// For vertex-plane contacts, we use a convex face as the manifold plane
 			const FVec3 RefSeparationDirection = ReferenceFaceConvex1 ? SeparationDirectionLocalConvex1 : SeparationDirectionLocalConvex2;
-			const FVec3 RefPlaneNormal = ReferenceFaceConvex1 ? BestPlaneConvex1.Normal() : BestPlaneConvex2.Normal();
-			const FVec3 RefPlanePosition = ReferenceFaceConvex1 ? BestPlaneConvex1.X() : BestPlaneConvex2.X();
+			const FVec3 RefPlaneNormal = ReferenceFaceConvex1 ? BestPlaneNormal1 : BestPlaneNormal2;
+			const FVec3 RefPlanePosition = ReferenceFaceConvex1 ? BestPlanePosition1 : BestPlanePosition2;
 
 			// @todo(chaos): fix use of hard-coded max array size
 			// We will use a double buffer as an optimization
diff --git a/Engine/Source/Runtime/Experimental/Chaos/Private/Chaos/ContactModification.cpp b/Engine/Source/Runtime/Experimental/Chaos/Private/Chaos/ContactModification.cpp
--- a/Engine/Source/Runtime/Experimental/Chaos/Private/Chaos/ContactModification.cpp
+++ b/Engine/Source/Runtime/Experimental/Chaos/Private/Chaos/ContactModification.cpp
@@ -73,21 +73,16 @@ namespace Chaos
 		return FVec3::DotProduct(Normal, (WorldPos0 - WorldPos1));
 	}
 
-	void FContactPairModifier::ModifySeparation(FReal Separation, int32 ContactPointIdx)
-	{	
-		FVec3 WorldPos0, WorldPos1;
-		GetWorldContactLocations(ContactPointIdx, WorldPos0, WorldPos1);
-
-		const FVec3 Normal = GetWorldNormal(ContactPointIdx);
-		const FReal CurrentSeparation = FVec3::DotProduct(Normal, (WorldPos0 - WorldPos1));
-		const FReal DeltaSeparation = Separation - CurrentSeparation;
-		
-		// Adjust contact locations to match desired separation
-		WorldPos0 += 0.5f * DeltaSeparation * Normal;
-		WorldPos1 -= 0.5f * DeltaSeparation * Normal;
-		ModifyWorldContactLocations(WorldPos0, WorldPos1, ContactPointIdx);
+	FReal FContactPairModifier::GetTargetSeparation(int32 ContactPointIdx) const
+	{
+		FManifoldPoint& ManifoldPoint = Constraint->GetManifoldPoint(ContactPointIdx);
+		return ManifoldPoint.TargetPhi;
+	}
 
-		Modifier->MarkConstraintForManifoldUpdate(*Constraint);
+	void FContactPairModifier::ModifyTargetSeparation(FReal TargetSeparation, int32 ContactPointIdx)
+	{	
+		FManifoldPoint& ManifoldPoint = Constraint->GetManifoldPoint(ContactPointIdx);
+		ManifoldPoint.TargetPhi = TargetSeparation;
 	}
 
 	FVec3 FContactPairModifier::GetWorldNormal(int32 ContactPointIdx) const
@@ -142,8 +137,10 @@ namespace Chaos
 
 		ManifoldPoint.ContactPoint.ShapeContactPoints[0] = Constraint->GetShapeWorldTransform0().InverseTransformPositionNoScale(Location0);
 		ManifoldPoint.ContactPoint.ShapeContactPoints[1] = Constraint->GetShapeWorldTransform1().InverseTransformPositionNoScale(Location1);
-		// @todo(chaos): We may also need to do something with the SavedManifoldPoints on the Constraint otherwise static friction
-		// may still try to push the points back to their original relative positions
+
+		// Clear the friction data since it will now have the wrong contact positions
+		// @todo(chaos): maybe try to do something better here
+		Constraint->ResetSavedManifoldPoints();
 
 		Modifier->MarkConstraintForManifoldUpdate(*Constraint);
 	}
diff --git a/Engine/Source/Runtime/Experimental/Chaos/Private/Chaos/Convex.cpp b/Engine/Source/Runtime/Experimental/Chaos/Private/Chaos/Convex.cpp
--- a/Engine/Source/Runtime/Experimental/Chaos/Private/Chaos/Convex.cpp
+++ b/Engine/Source/Runtime/Experimental/Chaos/Private/Chaos/Convex.cpp
@@ -59,7 +59,8 @@ namespace Chaos
 		FReal MostOpposingDot = TNumericLimits<FReal>::Max();
 		for (int32 Idx = 0; Idx < Planes.Num(); ++Idx)
 		{
-			const FPlaneType Plane = FPlaneType::MakeScaledUnsafe(Planes[Idx], Scale);
+			// @todo(chaos): pass in InvScale
+			const FPlaneType Plane = FPlaneType::MakeScaledUnsafe(Planes[Idx], Scale, FVec3(1) / Scale);
 			const FReal Distance = Plane.SignedDistance(Position);
 			if (FMath::Abs(Distance) < SearchDist)
 			{
diff --git a/Engine/Source/Runtime/Experimental/Chaos/Public/Chaos/Box.h b/Engine/Source/Runtime/Experimental/Chaos/Public/Chaos/Box.h
--- a/Engine/Source/Runtime/Experimental/Chaos/Public/Chaos/Box.h
+++ b/Engine/Source/Runtime/Experimental/Chaos/Public/Chaos/Box.h
@@ -349,6 +349,12 @@ namespace Chaos
 			return TPlaneConcrete<FReal>(PlaneX, PlaneN);
 		}
 
+		void GetPlaneNX(const int32 FaceIndex, FVec3& OutN, FVec3& OutX) const
+		{
+			OutN = SNormals[FaceIndex];
+			OutX = AABB.Center() + 0.5f * (SNormals[FaceIndex] * AABB.Extents());
+		}
+
 		// Get the vertex at the specified index (e.g., indices from GetPlaneVertexs)
 		const FVec3 GetVertex(int32 VertexIndex) const
 		{
diff --git a/Engine/Source/Runtime/Experimental/Chaos/Public/Chaos/Capsule.h b/Engine/Source/Runtime/Experimental/Chaos/Public/Chaos/Capsule.h
--- a/Engine/Source/Runtime/Experimental/Chaos/Public/Chaos/Capsule.h
+++ b/Engine/Source/Runtime/Experimental/Chaos/Public/Chaos/Capsule.h
@@ -529,6 +529,12 @@ namespace Chaos
 			return TPlaneConcrete<FReal, 3>(FVec3(0), FVec3(0));
 		}
 
+		void GetPlaneNX(const int32 FaceIndex, FVec3& OutN, FVec3& OutX) const
+		{
+			OutN = FVec3(0);
+			OutX = FVec3(0);
+		}
+
 		// Get an array of all the plane indices that belong to a vertex (up to MaxVertexPlanes).
 		// Returns the number of planes found.
 		int32 FindVertexPlanes(int32 VertexIndex, int32* OutVertexPlanes, int32 MaxVertexPlanes) const
diff --git a/Engine/Source/Runtime/Experimental/Chaos/Public/Chaos/Collision/PBDCollisionConstraint.h b/Engine/Source/Runtime/Experimental/Chaos/Public/Chaos/Collision/PBDCollisionConstraint.h
--- a/Engine/Source/Runtime/Experimental/Chaos/Public/Chaos/Collision/PBDCollisionConstraint.h
+++ b/Engine/Source/Runtime/Experimental/Chaos/Public/Chaos/Collision/PBDCollisionConstraint.h
@@ -47,6 +47,7 @@ namespace Chaos
 			, InitialShapeContactPoints{ FVec3(0), FVec3(0) }
 			, NetPushOut(0)
 			, NetImpulse(0)
+			, TargetPhi(0)
 			, Flags()
 		{}
 
@@ -56,14 +57,16 @@ namespace Chaos
 			, InitialShapeContactPoints{ FVec3(0), FVec3(0) }
 			, NetPushOut(0)
 			, NetImpulse(0)
+			, TargetPhi(0)
 			, Flags()
 		{}
 
 		FContactPoint ContactPoint;			// Contact point results of low-level collision detection
-		FVec3 ShapeAnchorPoints[2];
+		FVec3 ShapeAnchorPoints[2];			// The contact points from a prior frame, if availble. Used by friction.
 		FVec3 InitialShapeContactPoints[2];	// ShapeContactPoints when the constraint was first initialized. Used to track reusablility
 		FVec3 NetPushOut;					// Total pushout applied at this contact point
 		FVec3 NetImpulse;					// Total impulse applied by this contact point
+		FReal TargetPhi;					// Usually 0, but can be used to add padding or penetration (e.g., via a collision modifer)
 
 		union FFlags
 		{
@@ -305,7 +308,6 @@ namespace Chaos
 
 		const FReal GetCollisionMargin0() const { return CollisionMargins[0]; }
 		const FReal GetCollisionMargin1() const { return CollisionMargins[1]; }
-		const FReal GetCollisionTolerance() const { return CollisionTolerance; }
 
 		const bool IsQuadratic0() const { return Flags.bIsQuadratic0; }
 		const bool IsQuadratic1() const { return Flags.bIsQuadratic1; }
@@ -386,6 +388,7 @@ namespace Chaos
 		const FManifoldPoint& GetManifoldPoint(const int32 PointIndex) const { return ManifoldPoints[PointIndex]; }
 		const FManifoldPoint* GetClosestManifoldPoint() const { return (ClosestManifoldPointIndex != INDEX_NONE) ? &ManifoldPoints[ClosestManifoldPointIndex] : nullptr; }
 		const FSavedManifoldPoint* FindSavedManifoldPoint(const FManifoldPoint& ManifoldPoint) const;
+		void ResetSavedManifoldPoints() { SavedManifoldPoints.Reset(); }
 
 		void AddIncrementalManifoldContact(const FContactPoint& ContactPoint);
 
@@ -470,7 +473,7 @@ namespace Chaos
 		*/
 		void ResetSolverResults()
 		{
-			SavedManifoldPoints.Reset();
+			ResetSavedManifoldPoints();
 			AccumulatedImpulse = FVec3(0);
 		}
 
@@ -562,6 +565,7 @@ namespace Chaos
 			FManifoldPoint& ManifoldPoint = ManifoldPoints[ManifoldPointIndex];
 			ManifoldPoint.InitialShapeContactPoints[0] = ManifoldPoint.ContactPoint.ShapeContactPoints[0];
 			ManifoldPoint.InitialShapeContactPoints[1] = ManifoldPoint.ContactPoint.ShapeContactPoints[1];
+			ManifoldPoint.TargetPhi = FReal(0);
 			ManifoldPoint.Flags.Reset();
 			ManifoldPoint.NetPushOut = FVec3(0);
 			ManifoldPoint.NetImpulse = FVec3(0);
@@ -651,8 +655,8 @@ namespace Chaos
 		// Simplex data from the last call to GJK, used to warm-start GJK
 		FGJKSimplexData GJKWarmStartData;
 
-		TCArray<FManifoldPoint, MaxManifoldPoints> ManifoldPoints;
 		TCArray<FSavedManifoldPoint, MaxManifoldPoints> SavedManifoldPoints;
+		TCArray<FManifoldPoint, MaxManifoldPoints> ManifoldPoints;
 	};
 
 
diff --git a/Engine/Source/Runtime/Experimental/Chaos/Public/Chaos/ContactModification.h b/Engine/Source/Runtime/Experimental/Chaos/Public/Chaos/ContactModification.h
--- a/Engine/Source/Runtime/Experimental/Chaos/Public/Chaos/ContactModification.h
+++ b/Engine/Source/Runtime/Experimental/Chaos/Public/Chaos/ContactModification.h
@@ -61,10 +61,14 @@ namespace Chaos
 		FReal GetSeparation(int32 ContactPointIdx) const;
 
 		/**
-		* Modifying separation will adjust separation of contact locations in direction of contact normal.
-		* When modifying separation with normal or locations, order of these operations should be considered.
+		* @return desired separation of contact locations in direction of normal.
 		*/
-		void ModifySeparation(FReal Separation, int32 ContactPointIdx);
+		FReal GetTargetSeparation(int32 ContactPointIdx) const;
+
+		/**
+		* Set the desired separation at the contact point (negative for penetration).
+		*/
+		void ModifyTargetSeparation(FReal TargetSeparation, int32 ContactPointIdx);
 
 		/**
 		* @return Normal of contact in world space.
diff --git a/Engine/Source/Runtime/Experimental/Chaos/Public/Chaos/Convex.h b/Engine/Source/Runtime/Experimental/Chaos/Public/Chaos/Convex.h
--- a/Engine/Source/Runtime/Experimental/Chaos/Public/Chaos/Convex.h
+++ b/Engine/Source/Runtime/Experimental/Chaos/Public/Chaos/Convex.h
@@ -418,6 +418,12 @@ namespace Chaos
 			return TPlaneConcrete<FReal, 3>::MakeFrom(Planes[FaceIndex]);
 		}
 
+		void GetPlaneNX(const int32 FaceIndex, FVec3& OutN, FVec3& OutX) const
+		{
+			OutN = FVec3(Planes[FaceIndex].Normal());
+			OutX = FVec3(Planes[FaceIndex].X());
+		}
+
 		const FPlaneType& GetPlaneRaw(int32 FaceIndex) const
 		{
 			return Planes[FaceIndex];
@@ -456,15 +462,15 @@ namespace Chaos
 		}
 
 	private:
-		int32 GetSupportVertex(const FVec3& Direction) const
+		FORCEINLINE_DEBUGGABLE int32 GetSupportVertex(const FVec3Type& Direction) const
 		{
-			FReal MaxDot = TNumericLimits<FReal>::Lowest();
+			FRealType MaxDot = TNumericLimits<FRealType>::Lowest();
 			int32 MaxVIdx = INDEX_NONE;
 			const int32 NumVertices = Vertices.Num();
 
 			for (int32 Idx = 0; Idx < NumVertices; ++Idx)
 			{
-				const FReal Dot = FVec3::DotProduct(Vertices[Idx], Direction);
+				const FRealType Dot = FVec3Type::DotProduct(Vertices[Idx], Direction);
 				if (Dot > MaxDot)
 				{
 					MaxDot = Dot;
@@ -563,11 +569,6 @@ namespace Chaos
 
 		FVec3 GetMarginAdjustedVertexScaled(int32 VertexIndex, FReal InMargin, const FVec3& Scale, FReal* OutSupportDelta) const
 		{
-			if (InMargin == 0.0f)
-			{
-				return GetVertex(VertexIndex) * Scale;
-			}
-
 			// Get any 3 planes that contribute to this vertex
 			int32 PlaneIndex0 = INDEX_NONE;
 			int32 PlaneIndex1 = INDEX_NONE;
@@ -626,7 +627,7 @@ namespace Chaos
 		// Return support point on the core shape (the convex shape with all planes moved inwards by margin).
 		FVec3 SupportCore(const FVec3& Direction, const FReal InMargin, FReal* OutSupportDelta, int32& VertexIndex) const
 		{
-			const int32 SupportVertexIndex = GetSupportVertex(Direction);
+			const int32 SupportVertexIndex = GetSupportVertex(FVec3Type(Direction));
 			VertexIndex = SupportVertexIndex;
 			if (SupportVertexIndex != INDEX_NONE)
 			{
@@ -640,15 +641,22 @@ namespace Chaos
 		FVec3 SupportCoreScaled(const FVec3& Direction, FReal InMargin, const FVec3& Scale, FReal* OutSupportDelta, int32& VertexIndex) const
 		{
 			// Find the supporting vertex index
-			const FVec3 DirectionScaled = Scale * Direction;	// does not need to be normalized
+			const FVec3Type DirectionScaled = FVec3Type(Scale * Direction);	// does not need to be normalized
 			const int32 SupportVertexIndex = GetSupportVertex(DirectionScaled);
 			VertexIndex = SupportVertexIndex;
 			// Adjust the vertex position based on margin
 			if (SupportVertexIndex != INDEX_NONE)
 			{
-				// Note: Shapes wrapped in a non-uniform scale should not have their own margin and we assume that here
-				// @chaos(todo): apply an upper limit to the margin to prevent a non-convex or null shape (also see comments in GetMarginAdjustedVertex)
-				return GetMarginAdjustedVertexScaled(SupportVertexIndex, InMargin, Scale, OutSupportDelta);
+				if (InMargin == 0.0f)
+				{
+					return FVec3(GetVertex(VertexIndex)) * Scale;
+				}
+				else
+				{
+					// Note: Shapes wrapped in a non-uniform scale should not have their own margin and we assume that here
+					// @chaos(todo): apply an upper limit to the margin to prevent a non-convex or null shape (also see comments in GetMarginAdjustedVertex)
+					return FVec3(GetMarginAdjustedVertexScaled(SupportVertexIndex, InMargin, Scale, OutSupportDelta));
+				}
 			}
 			return FVec3(0);
 		}
@@ -657,7 +665,7 @@ namespace Chaos
 		// @todo(chaos): do we need to support thickness?
 		FORCEINLINE FVec3 Support(const FVec3& Direction, const FReal Thickness, int32& VertexIndex) const
 		{
-			const int32 MaxVIdx = GetSupportVertex(Direction);
+			const int32 MaxVIdx = GetSupportVertex(FVec3Type(Direction));
 			VertexIndex = MaxVIdx;
 			if (MaxVIdx != INDEX_NONE)
 			{
diff --git a/Engine/Source/Runtime/Experimental/Chaos/Public/Chaos/EPA.h b/Engine/Source/Runtime/Experimental/Chaos/Public/Chaos/EPA.h
--- a/Engine/Source/Runtime/Experimental/Chaos/Public/Chaos/EPA.h
+++ b/Engine/Source/Runtime/Experimental/Chaos/Public/Chaos/EPA.h
@@ -9,6 +9,14 @@
 namespace Chaos
 {
 
+constexpr int32 ArraySizeEPA = 16;
+
+// Array type used in EPA to avoid heap allocation for small convex shapes
+// @todo(chaos): The inline size was picked to avoid allocations in box-box collision - it might need adjusting after more general purpose testing
+// @todo(chaos): We might also consider different inline sizes based on array use-case (e.g., Entries array versus Border array)
+template<typename T>
+using TEPAWorkingArray = TArray<T, TInlineAllocator<32>>;
+
 template <typename T>
 FORCEINLINE const TVec3<T> MinkowskiVert(const TVec3<T>* VertsABuffer, const TVec3<T>* VertsBBuffer, const int32 Idx)
 {
@@ -26,26 +34,27 @@ struct TEPAEntry
 	TVector<int32,3> AdjEdges;	//Adjacent edges (idx in adjacent face)
 	bool bObsolete;	//indicates that an entry can be skipped (became part of bigger polytope)
 
-	bool operator>(const TEPAEntry<T>& Other) const
+	FORCEINLINE bool operator>(const TEPAEntry<T>& Other) const
 	{
 		return Distance > Other.Distance;
 	}
 
-	bool Initialize(const TVec3<T>* VerticesA, const TVec3<T>* VerticesB, int32 InIdx0, int32 InIdx1, int32 InIdx2, const TVector<int32,3>& InAdjFaces, const TVector<int32,3>& InAdjEdges)
+	FORCEINLINE_DEBUGGABLE bool Initialize(const TVec3<T>* VerticesA, const TVec3<T>* VerticesB, int32 InIdx0, int32 InIdx1, int32 InIdx2, const TVector<int32,3>& InAdjFaces, const TVector<int32,3>& InAdjEdges)
 	{
-		const TVec3<T>& V0 = MinkowskiVert(VerticesA, VerticesB, InIdx0);
-		const TVec3<T>& V1 = MinkowskiVert(VerticesA, VerticesB, InIdx1);
-		const TVec3<T>& V2 = MinkowskiVert(VerticesA, VerticesB, InIdx2);
+		const TVec3<T> V0 = MinkowskiVert(VerticesA, VerticesB, InIdx0);
+		const TVec3<T> V1 = MinkowskiVert(VerticesA, VerticesB, InIdx1);
+		const TVec3<T> V2 = MinkowskiVert(VerticesA, VerticesB, InIdx2);
 
 		const TVec3<T> V0V1 = V1 - V0;
 		const TVec3<T> V0V2 = V2 - V0;
 		const TVec3<T> Norm = TVec3<T>::CrossProduct(V0V1, V0V2);
-		PlaneNormal = Norm.GetSafeNormal();
+		const FReal NormLenSq = Norm.SizeSquared();
 		constexpr T Eps = 1e-4f;
-		if (PlaneNormal.SizeSquared() < Eps)
+		if (NormLenSq < Eps)
 		{
 			return false;
 		}
+		PlaneNormal = Norm * FMath::InvSqrt(NormLenSq);
 		
 		IdxBuffer[0] = InIdx0;
 		IdxBuffer[1] = InIdx1;
@@ -90,7 +99,7 @@ struct TEPAEntry
 		Distance = -Distance;
 	}
 
-	T DistanceToPlane(const TVec3<T>& X) const
+	FORCEINLINE_DEBUGGABLE T DistanceToPlane(const TVec3<T>& X) const
 	{
 		return TVec3<T>::DotProduct(PlaneNormal, X) - Distance;
 	}
@@ -126,7 +135,7 @@ struct TEPAEntry
 };
 
 template <typename T>
-bool InitializeEPA(TArray<TVec3<T>>& VertsA, TArray<TVec3<T>>& VertsB, TFunctionRef<TVector<T, 3>(const TVec3<T>& V)> SupportA, TFunctionRef<TVector<T, 3>(const TVec3<T>& V)> SupportB, TArray<TEPAEntry<T>>& OutEntries, TVec3<T>& OutTouchNormal)
+bool InitializeEPA(TArray<TVec3<T>>& VertsA, TArray<TVec3<T>>& VertsB, TFunctionRef<TVector<T, 3>(const TVec3<T>& V)> SupportA, TFunctionRef<TVector<T, 3>(const TVec3<T>& V)> SupportB, TEPAWorkingArray<TEPAEntry<T>>& OutEntries, TVec3<T>& OutTouchNormal)
 {
 	const int32 NumVerts = VertsA.Num();
 	check(VertsB.Num() == NumVerts);
@@ -198,9 +207,9 @@ bool InitializeEPA(TArray<TVec3<T>>& VertsA, TArray<TVec3<T>>& VertsB, TFunction
 				AddFartherPoint(Orthog);
 				AddFartherPoint(Orthog2);
 
-				bValid = OutEntries[0].Initialize(VertsA.GetData(), VertsB.GetData(), 1, 2, 3, { 3, 1, 2 }, { 1,1, 1 });
-				bValid &= OutEntries[1].Initialize(VertsA.GetData(), VertsB.GetData(), 0, 3, 2, { 2,0,3 }, { 2, 1, 0 });
-				bValid &= OutEntries[2].Initialize(VertsA.GetData(), VertsB.GetData(), 0, 1, 3, { 3,0, 1 }, { 2,2,0 });
+				bValid = OutEntries[0].Initialize(VertsA.GetData(), VertsB.GetData(), 1, 2, 3, { 3,1,2 }, { 1,1,1 });
+				bValid &= OutEntries[1].Initialize(VertsA.GetData(), VertsB.GetData(), 0, 3, 2, { 2,0,3 }, { 2,1,0 });
+				bValid &= OutEntries[2].Initialize(VertsA.GetData(), VertsB.GetData(), 0, 1, 3, { 3,0,1 }, { 2,2,0 });
 				bValid &= OutEntries[3].Initialize(VertsA.GetData(), VertsB.GetData(), 0, 2, 1, { 1,0,2 }, { 2,0,0 });
 
 				if(!bValid)
@@ -227,9 +236,9 @@ bool InitializeEPA(TArray<TVec3<T>>& VertsA, TArray<TVec3<T>>& VertsB, TFunction
 
 				AddFartherPoint(Base.PlaneNormal);
 
-				bValid = OutEntries[0].Initialize(VertsA.GetData(), VertsB.GetData(), 1, 2, 3, { 3, 1, 2 }, { 1,1, 1 });
-				bValid &= OutEntries[1].Initialize(VertsA.GetData(), VertsB.GetData(), 0, 3, 2, { 2,0,3 }, { 2, 1, 0 });
-				bValid &= OutEntries[2].Initialize(VertsA.GetData(), VertsB.GetData(), 0, 1, 3, { 3,0, 1 }, { 2,2,0 });
+				bValid = OutEntries[0].Initialize(VertsA.GetData(), VertsB.GetData(), 1, 2, 3, { 3,1,2 }, { 1,1,1 });
+				bValid &= OutEntries[1].Initialize(VertsA.GetData(), VertsB.GetData(), 0, 3, 2, { 2,0,3 }, { 2,1,0 });
+				bValid &= OutEntries[2].Initialize(VertsA.GetData(), VertsB.GetData(), 0, 1, 3, { 3,0,1 }, { 2,2,0 });
 
 				if(!bValid)
 				{
@@ -241,9 +250,9 @@ bool InitializeEPA(TArray<TVec3<T>>& VertsA, TArray<TVec3<T>>& VertsB, TFunction
 		}
 		case 4:
 		{
-			bValid = OutEntries[0].Initialize(VertsA.GetData(), VertsB.GetData(), 1, 2, 3, { 3, 1, 2 }, { 1,1, 1 });
-			bValid &= OutEntries[1].Initialize(VertsA.GetData(), VertsB.GetData(), 0, 3, 2, { 2,0,3 }, { 2, 1, 0 });
-			bValid &= OutEntries[2].Initialize(VertsA.GetData(), VertsB.GetData(), 0, 1, 3, { 3,0, 1 }, { 2,2,0 });
+			bValid = OutEntries[0].Initialize(VertsA.GetData(), VertsB.GetData(), 1, 2, 3, { 3,1,2 }, { 1,1,1 });
+			bValid &= OutEntries[1].Initialize(VertsA.GetData(), VertsB.GetData(), 0, 3, 2, { 2,0,3 }, { 2,1,0 });
+			bValid &= OutEntries[2].Initialize(VertsA.GetData(), VertsB.GetData(), 0, 1, 3, { 3,0,1 }, { 2,2,0 });
 			bValid &= OutEntries[3].Initialize(VertsA.GetData(), VertsB.GetData(), 0, 2, 1, { 1,0,2 }, { 2,0,0 });
 			
 			if(!bValid)
@@ -290,7 +299,7 @@ bool InitializeEPA(TArray<TVec3<T>>& VertsA, TArray<TVec3<T>>& VertsB, TFunction
 }
 
 template <typename T, typename SupportALambda, typename SupportBLambda >
-bool InitializeEPA(TArray<TVec3<T>>& VertsA, TArray<TVec3<T>>& VertsB, const SupportALambda& SupportA, const SupportBLambda& SupportB, TArray<TEPAEntry<T>>& OutEntries, TVec3<T>& OutTouchNormal)
+bool InitializeEPA(TArray<TVec3<T>>& VertsA, TArray<TVec3<T>>& VertsB, const SupportALambda& SupportA, const SupportBLambda& SupportB, TEPAWorkingArray<TEPAEntry<T>>& OutEntries, TVec3<T>& OutTouchNormal)
 {
 	TFunctionRef<TVector<T, 3>(const TVec3<T>& V)> SupportARef(SupportA);
 	TFunctionRef<TVector<T, 3>(const TVec3<T>& V)> SupportBRef(SupportB);
@@ -305,9 +314,8 @@ struct FEPAFloodEntry
 };
 
 template <typename T>
-void EPAComputeVisibilityBorder(TArray<TEPAEntry<T>>& Entries, int32 EntryIdx, const TVec3<T>& W, TArray<FEPAFloodEntry>& OutBorderEdges)
+void EPAComputeVisibilityBorder(TEPAWorkingArray<TEPAEntry<T>>& Entries, int32 EntryIdx, const TVec3<T>& W, TEPAWorkingArray<FEPAFloodEntry>& OutBorderEdges, TEPAWorkingArray<FEPAFloodEntry> &ToVisitStack)
 {
-	TArray<FEPAFloodEntry> ToVisitStack;
 	{
 		TEPAEntry<T>& Entry = Entries[EntryIdx];
 		for (int i = 0; i < 3; ++i)
@@ -413,7 +421,7 @@ inline const bool IsEPASuccess(EEPAResult EPAResult)
 // See e.g., "Collision Detection in Interactive 3D Environments" (Gino van den Bergen, 2004)
 // or "Real-time Collision Detection with Implicit Objects" (Leif Olvang, 2010)
 template <typename T, typename TSupportA, typename TSupportB>
-EEPAResult EPA(TArray<TVec3<T>>& VertsABuffer, TArray<TVec3<T>>& VertsBBuffer, const TSupportA& SupportA, const TSupportB& SupportB, T& OutPenetration, TVec3<T>& OutDir, TVec3<T>& WitnessA, TVec3<T>& WitnessB)
+EEPAResult EPA(TArray<TVec3<T>>& VertsABuffer, TArray<TVec3<T>>& VertsBBuffer, const TSupportA& SupportA, const TSupportB& SupportB, T& OutPenetration, TVec3<T>& OutDir, TVec3<T>& WitnessA, TVec3<T>& WitnessB, const FReal Eps = 1.e-2f)
 {
 	struct FEPAEntryWrapper
 	{
@@ -426,14 +434,10 @@ EEPAResult EPA(TArray<TVec3<T>>& VertsABuffer, TArray<TVec3<T>>& VertsBBuffer, c
 		}
 	};
 
-	T UpperBound = TNumericLimits<T>::Max();
-	T LowerBound = TNumericLimits<T>::Lowest();
-
-	// @todo(chaos): Should we pass in Eps? Does it need to match anything used in GJK?
-	constexpr T Eps = 1e-2f;
 	constexpr T OriginInsideEps = 0.0f;
 
-	TArray<TEPAEntry<T>> Entries;
+	TEPAWorkingArray<TEPAEntry<T>> Entries;
+
 	if(!InitializeEPA(VertsABuffer,VertsBBuffer,SupportA,SupportB, Entries, OutDir))
 	{
 		//either degenerate or a touching hit. Either way return penetration 0
@@ -451,29 +455,39 @@ EEPAResult EPA(TArray<TVec3<T>>& VertsABuffer, TArray<TVec3<T>>& VertsBBuffer, c
 	}
 #endif
 	
-	std::priority_queue<FEPAEntryWrapper, std::vector<FEPAEntryWrapper>, std::greater<FEPAEntryWrapper>> Queue;
+	TEPAWorkingArray<int32> Queue;
 	for(int32 Idx = 0; Idx < Entries.Num(); ++Idx)
 	{
 		//ensure(Entries[Idx].Distance > -Eps);
 		// Entries[Idx].Distance <= 0.0f is true if the origin is a bit out of the polytope (we need to support this case for robustness)
 		if(Entries[Idx].Distance <= 0.0f || Entries[Idx].IsOriginProjectedInside(VertsABuffer.GetData(), VertsBBuffer.GetData(), OriginInsideEps))
 		{
-			Queue.push(FEPAEntryWrapper {&Entries, Idx});
+			Queue.Add(Idx);
 		}
 	}
-	
+
+	TEPAWorkingArray<FEPAFloodEntry> VisibilityBorder;
+	TEPAWorkingArray<FEPAFloodEntry> VisibilityBorderToVisitStack;
+
 	//TEPAEntry<T> BestEntry;
 	//BestEntry.Distance = 0;
-	TEPAEntry<T> LastEntry = Queue.size() > 0 ? Entries[Queue.top().Idx] : Entries[0];
-
-	TArray<FEPAFloodEntry> VisibilityBorder;
+	TEPAEntry<T> LastEntry = Queue.Num() > 0 ? Entries[Queue.Last()] : Entries[0];
+	T UpperBound = TNumericLimits<T>::Max();
+	T LowerBound = TNumericLimits<T>::Lowest();
+	bool bQueueDirty = true;
 	int32 Iteration = 0;
 	int32 constexpr MaxIterations = 128;
 	EEPAResult ResultStatus = EEPAResult::MaxIterations;
-	while (Queue.size() && (Iteration++ < MaxIterations))
+	while (Queue.Num() && (Iteration++ < MaxIterations))
 	{
-		int32 EntryIdx = Queue.top().Idx;
-		Queue.pop();
+		if (bQueueDirty)
+		{
+			// Avoiding UE's Sort here because it has a call to FMath::Loge. The std version calls insertion sort when possible
+			std::sort(Queue.GetData(), Queue.GetData() + Queue.Num(), [&Entries](const int32 L, const int32 R) { return Entries[L] > Entries[R]; });
+			bQueueDirty = false;
+		}
+
+		int32 EntryIdx = Queue.Pop(false);
 		TEPAEntry<T>& Entry = Entries[EntryIdx];
 		//bool bBadFace = Entry.IsOriginProjectedInside(VertsABuffer.GetData(), VertsBBuffer.GetData());
 		{
@@ -533,7 +547,8 @@ EEPAResult EPA(TArray<TVec3<T>>& VertsABuffer, TArray<TVec3<T>>& VertsBBuffer, c
 
 		Entry.bObsolete = true;
 		VisibilityBorder.Reset();
-		EPAComputeVisibilityBorder(Entries, EntryIdx, W, VisibilityBorder);
+		VisibilityBorderToVisitStack.Reset();
+		EPAComputeVisibilityBorder(Entries, EntryIdx, W, VisibilityBorder, VisibilityBorderToVisitStack);
 		const int32 NumBorderEdges = VisibilityBorder.Num();
 		const int32 FirstIdxInBatch = Entries.Num();
 		int32 NewIdx = FirstIdxInBatch;
@@ -575,7 +590,8 @@ EEPAResult EPA(TArray<TVec3<T>>& VertsABuffer, TArray<TVec3<T>>& VertsBBuffer, c
 					// NewEntry.Distance <= 0.0f is if the origin is a bit out of the polytope
 					if (NewEntry.Distance <= 0.0f || NewEntry.IsOriginProjectedInside(VertsABuffer.GetData(), VertsBBuffer.GetData(), OriginInsideEps))
 					{
-						Queue.push(FEPAEntryWrapper{ &Entries, NewIdx });
+						Queue.Add(NewIdx);
+						bQueueDirty = true;
 					}
 				}
 
diff --git a/Engine/Source/Runtime/Experimental/Chaos/Public/Chaos/GJK.h b/Engine/Source/Runtime/Experimental/Chaos/Public/Chaos/GJK.h
--- a/Engine/Source/Runtime/Experimental/Chaos/Public/Chaos/GJK.h
+++ b/Engine/Source/Runtime/Experimental/Chaos/Public/Chaos/GJK.h
@@ -277,7 +277,7 @@ namespace Chaos
 	 *
 	*/
 	template <typename T, typename TGeometryA, typename TGeometryB>
-	bool GJKPenetrationWarmStartable(const TGeometryA& A, const TGeometryB& B, const TRigidTransform<T, 3>& BToATM, T& OutPenetration, TVec3<T>& OutClosestA, TVec3<T>& OutClosestB, TVec3<T>& OutNormalA, TVec3<T>& OutNormalB, int32& OutVertexA, int32& OutVertexB, TGJKSimplexData<T>& InOutSimplexData, T& OutMaxSupportDelta, const T Epsilon = T(1.e-3))
+	bool GJKPenetrationWarmStartable(const TGeometryA& A, const TGeometryB& B, const TRigidTransform<T, 3>& BToATM, T& OutPenetration, TVec3<T>& OutClosestA, TVec3<T>& OutClosestB, TVec3<T>& OutNormalA, TVec3<T>& OutNormalB, int32& OutVertexA, int32& OutVertexB, TGJKSimplexData<T>& InOutSimplexData, T& OutMaxSupportDelta, const T Epsilon = T(1.e-3), const T EPAEpsilon = T(1.e-2))
 	{
 		T SupportDeltaA = 0;
 		T SupportDeltaB = 0;
@@ -352,7 +352,7 @@ namespace Chaos
 
 			// If we are not too close, update the separating vector and keep iterating
 			// If we are too close we drop through to EPA, but if EPA determines the shapes are actually separated 
-			// after all, we will end up using the normal from the previous loop
+			// after all, we will wend up using the normal from the previous loop
 			if (!bIsContact)
 			{
 				V /= NewDistance;
@@ -387,7 +387,7 @@ namespace Chaos
 
 			T Penetration;
 			TVec3<T> MTD, ClosestA, ClosestBInA;
-			const EEPAResult EPAResult = EPA(VertsA, VertsB, SupportAFunc, SupportBInAFunc, Penetration, MTD, ClosestA, ClosestBInA);
+			const EEPAResult EPAResult = EPA(VertsA, VertsB, SupportAFunc, SupportBInAFunc, Penetration, MTD, ClosestA, ClosestBInA, EPAEpsilon);
 
 			switch (EPAResult)
 			{
@@ -408,6 +408,7 @@ namespace Chaos
 				// and penetration but we keep the position generated by GJK.
 				Normal = MTD;
 				Distance = -Penetration;
+				//UE_LOG(LogChaos, Warning, TEXT("EPA Touching Case"));
 				break;
 			case EEPAResult::Degenerate:
 				// We hit a degenerate simplex condition and could not reach a solution so use whatever near touching point GJK came up with.
diff --git a/Engine/Source/Runtime/Experimental/Chaos/Public/Chaos/ImplicitObjectScaled.h b/Engine/Source/Runtime/Experimental/Chaos/Public/Chaos/ImplicitObjectScaled.h
--- a/Engine/Source/Runtime/Experimental/Chaos/Public/Chaos/ImplicitObjectScaled.h
+++ b/Engine/Source/Runtime/Experimental/Chaos/Public/Chaos/ImplicitObjectScaled.h
@@ -300,6 +300,11 @@ public:
 		return MObject->GetPlane(FaceIndex);
 	}
 
+	void GetPlaneNX(const int32 FaceIndex, FVec3& OutN, FVec3& OutX) const
+	{
+		MObject->GetPlaneNX(FaceIndex, OutN, OutX);
+	}
+
 	// Get the vertex at the specified index (e.g., indices from GetPlaneVertexs)
 	const FVec3 GetVertex(int32 VertexIndex) const
 	{
@@ -750,7 +755,14 @@ public:
 	const TPlaneConcrete<FReal, 3> GetPlane(int32 FaceIndex) const
 	{
 		const TPlaneConcrete<FReal, 3> InnerPlane = MObject->GetPlane(FaceIndex);
-		return TPlaneConcrete<FReal, 3>::MakeScaledUnsafe(InnerPlane, MScale);	// "Unsafe" means scale has no zeros
+		return TPlaneConcrete<FReal, 3>::MakeScaledUnsafe(InnerPlane, MScale, MInvScale);	// "Unsafe" means scale must have no zeros
+	}
+
+	void GetPlaneNX(const int32 FaceIndex, FVec3& OutN, FVec3& OutX) const
+	{
+		FVec3 InnerN, InnerX;
+		MObject->GetPlaneNX(FaceIndex, InnerN, InnerX);
+		TPlaneConcrete<FReal>::MakeScaledUnsafe(InnerN, InnerX, MScale, MInvScale, OutN, OutX);	// "Unsafe" means scale must have no zeros
 	}
 
 	// Get the vertex at the specified index (e.g., indices from GetPlaneVertex)
diff --git a/Engine/Source/Runtime/Experimental/Chaos/Public/Chaos/Plane.h b/Engine/Source/Runtime/Experimental/Chaos/Public/Chaos/Plane.h
--- a/Engine/Source/Runtime/Experimental/Chaos/Public/Chaos/Plane.h
+++ b/Engine/Source/Runtime/Experimental/Chaos/Public/Chaos/Plane.h
@@ -70,10 +70,10 @@ public:
 
 	// Scale the plane and assume that none of the scale components are zero
 	template <typename U>
-	static TPlaneConcrete<T> MakeScaledUnsafe(const TPlaneConcrete<U>& Plane, const TVec3<T>& Scale)
+	static FORCEINLINE TPlaneConcrete<T> MakeScaledUnsafe(const TPlaneConcrete<U>& Plane, const TVec3<T>& Scale, const TVec3<T>& InvScale)
 	{
 		const TVec3<T> ScaledX = TVec3<T>(Plane.X()) * Scale;
-		TVec3<T> ScaledN = TVec3<T>(Plane.Normal()) / Scale;
+		TVec3<T> ScaledN = TVec3<T>(Plane.Normal()) * InvScale;
 
 		// We don't handle zero scales, but we could still end up with a small normal
 		const T ScaleN2 = ScaledN.SizeSquared();
@@ -90,7 +90,29 @@ public:
 	}
 
 	template <typename U>
-	static TPlaneConcrete<T> MakeFrom(const TPlaneConcrete<U>& Plane)
+	static FORCEINLINE void MakeScaledUnsafe(const TVec3<U>& PlaneN, const TVec3<U>& PlaneX, const TVec3<T>& Scale, const TVec3<T>& InvScale, TVec3<T>& OutN, TVec3<T>& OutX)
+	{
+		const TVec3<T> ScaledX = TVec3<T>(PlaneX * Scale);
+		TVec3<T> ScaledN = TVec3<T>(PlaneN * InvScale);
+
+		// We don't handle zero scales, but we could still end up with a small normal
+		const T ScaleN2 = ScaledN.SizeSquared();
+		if (ScaleN2 > SMALL_NUMBER)
+		{
+			ScaledN = ScaledN * FMath::InvSqrt(ScaleN2);
+		}
+		else
+		{
+			ScaledN = TVec3<T>(PlaneN);
+		}
+
+		OutN = ScaledN;
+		OutX = ScaledX;
+	}
+
+
+	template <typename U>
+	static FORCEINLINE TPlaneConcrete<T> MakeFrom(const TPlaneConcrete<U>& Plane)
 	{
 		return TPlaneConcrete<T>(TVec3<T>(Plane.X()), TVec3<T>(Plane.Normal()));
 	}
diff --git a/Engine/Source/Runtime/Experimental/Chaos/Public/Chaos/Triangle.h b/Engine/Source/Runtime/Experimental/Chaos/Public/Chaos/Triangle.h
--- a/Engine/Source/Runtime/Experimental/Chaos/Public/Chaos/Triangle.h
+++ b/Engine/Source/Runtime/Experimental/Chaos/Public/Chaos/Triangle.h
@@ -48,6 +48,12 @@ namespace Chaos
 			return TPlaneConcrete < FReal, 3> (A, GetNormal());
 		}
 
+		FORCEINLINE void GetPlaneNX(const int32 FaceIndex, FVec3& OutN, FVec3& OutX) const
+		{
+			OutN = GetNormal();
+			OutX = A;
+		}
+
 		// Get the nearest point on an edge
 		// Used for manifold generation
 		FVec3 GetClosestEdgePosition(int32 PlaneIndexHint, const FVec3& Position) const
