Chaos - fix island group count when using "onethread" or when physics parallelism is disabled
- fix constraint sets for island groups when not using colours

diff --git a/Engine/Source/Runtime/Experimental/Chaos/Private/Chaos/Island/IslandManager.cpp b/Engine/Source/Runtime/Experimental/Chaos/Private/Chaos/Island/IslandManager.cpp
--- a/Engine/Source/Runtime/Experimental/Chaos/Private/Chaos/Island/IslandManager.cpp
+++ b/Engine/Source/Runtime/Experimental/Chaos/Private/Chaos/Island/IslandManager.cpp
@@ -10,13 +10,14 @@
 #include "ChaosLog.h"
 
 #include "Framework/Threading.h"
+#include "Misc/App.h"
 
 /** Cvar to enable/disable the island sleeping */
 bool ChaosSolverSleepEnabled = true;
 FAutoConsoleVariableRef CVarChaosSolverSleepEnabled(TEXT("p.Chaos.Solver.SleepEnabled"), ChaosSolverSleepEnabled, TEXT(""));
 
 /** Cvar to control the number of island groups used in the solver. The total number will be NumThreads * IslandGroupsMultiplier */
-int32 GChaosSolverIslandGroupsMultiplier = 1;
+Chaos::FRealSingle GChaosSolverIslandGroupsMultiplier = 1;
 FAutoConsoleVariableRef CVarSolverIslandGroupsMultiplier(TEXT("p.Chaos.Solver.IslandGroupsMultiplier"), GChaosSolverIslandGroupsMultiplier, TEXT("Total number of island groups in the solver will be NumThreads * IslandGroupsMultiplier.[def:1]"));
 
 /** Cvar to override the sleep counter threshold if necessary */
@@ -33,6 +34,7 @@ FAutoConsoleVariableRef CVarChaosSolverCollisionDefaultAngularSleepThreshold(TEX
 
 namespace Chaos
 {
+	extern int32 GSingleThreadedPhysics;
 	
 /** Check if a particle is dynamic or sleeping */
 FORCEINLINE bool IsDynamicParticle(const FGeometryParticleHandle* ParticleHandle)
@@ -297,8 +299,12 @@ void FPBDIslandManager::ResetIndices()
 
 void FPBDIslandManager::InitializeGroups()
 {
-	const int32 NumThreads = FTaskGraphInterface::Get().GetNumWorkerThreads();
-	IslandGroups.SetNum(NumThreads * GChaosSolverIslandGroupsMultiplier, false);
+	// @todo(chaos): is the number of worker threads a good indicator of how many threads we get in the solver loop? (Currently uses ParallelFor)
+	// Check for use of the "-onethread" command line arg, and physics threading disabled (GetNumWorkerThreads() is not affected by these)
+	const int32 NumWorkerThreads = (FApp::ShouldUseThreadingForPerformance() && !GSingleThreadedPhysics) ? FTaskGraphInterface::Get().GetNumWorkerThreads() : 0;
+	const int32 MaxIslandGroups = FMath::Max(1, FMath::CeilToInt(FReal(NumWorkerThreads) * GChaosSolverIslandGroupsMultiplier));
+
+	IslandGroups.SetNum(MaxIslandGroups, false);
 	
 	for(int32 GroupIndex = 0, NumGroups = IslandGroups.Num(); GroupIndex < NumGroups; ++GroupIndex)
 	{
diff --git a/Engine/Source/Runtime/Experimental/Chaos/Private/Chaos/PBDConstraintRule.cpp b/Engine/Source/Runtime/Experimental/Chaos/Private/Chaos/PBDConstraintRule.cpp
--- a/Engine/Source/Runtime/Experimental/Chaos/Private/Chaos/PBDConstraintRule.cpp
+++ b/Engine/Source/Runtime/Experimental/Chaos/Private/Chaos/PBDConstraintRule.cpp
@@ -264,6 +264,7 @@ namespace Chaos
 				for (auto& IslandSolver : IslandGroup->GetIslands() )
 				{
 					ConstraintSetBegin = ConstraintSetEnd;
+
 					if(!IslandSolver->IsSleeping() && !IslandSolver->IsUsingCache())
 					{
 						if(!bResizeDone)
@@ -280,7 +281,10 @@ namespace Chaos
 							const int32 MaxColor = IsSortingUsingColors() ? FMath::Max(1,ConstraintGraph->GetIslandGraph()->GraphIslands[GraphIndex].MaxColors+1) : 1;
 							const int32 MaxLevel = IsSortingUsingLevels() ? FMath::Max(1,ConstraintGraph->GetIslandGraph()->GraphIslands[GraphIndex].MaxLevels+1) : 1;
 							
-							IslandConstraintSets.Reserve(IslandConstraintSets.Num()+MaxLevel * MaxColor);	// Pessimistic array size - we could store the actual required size in coloring algorithm
+							if (IsSortingUsingColors())
+							{
+								IslandConstraintSets.Reserve(IslandConstraintSets.Num()+MaxLevel * MaxColor);	// Pessimistic array size - we could store the actual required size in coloring algorithm
+							}
 
 							for (int32 Level = 0; Level < MaxLevel; ++Level)
 							{
@@ -323,12 +327,13 @@ namespace Chaos
 							}
 						}
 					}
-					// If we aren't coloring, we have a single group of all constraints (they have been created in level order above)
-					if (!IsSortingUsingColors())
-					{
-						const TPair<int32, int32> ConstrainSet(ConstraintSetBegin, ConstraintSetEnd);
-						IslandConstraintSets.Add(ConstrainSet);
-					}
+				}
+
+				// If we aren't coloring, we have a single group of all constraints in the IslandGroup (they have been created in level order above)
+				if (!IsSortingUsingColors())
+				{
+					const TPair<int32, int32> ConstrainSet(0, ConstraintSetEnd);
+					IslandConstraintSets.Add(ConstrainSet);
 				}
 			}
 		}
