Chaos
- support for Target Separation in collision solver
- VectorRegister use for CollisionConstraint Gather(disbaled)

diff --git a/Engine/Source/Runtime/Experimental/Chaos/Private/Chaos/Collision/PBDCollisionSolver.cpp b/Engine/Source/Runtime/Experimental/Chaos/Private/Chaos/Collision/PBDCollisionSolver.cpp
--- a/Engine/Source/Runtime/Experimental/Chaos/Private/Chaos/Collision/PBDCollisionSolver.cpp
+++ b/Engine/Source/Runtime/Experimental/Chaos/Private/Chaos/Collision/PBDCollisionSolver.cpp
@@ -23,10 +23,13 @@ namespace Chaos
 	{
 		extern int32 Chaos_Collision_UseShockPropagation;
 
+		bool bChaos_PBDCollisionSolver_VectorRegister = false;
+		FAutoConsoleVariableRef CVarChaos_PBDCollisionSolver_VectorRegister(TEXT("p.Chaos.PBDCollisionSolver.VectorRegister"), bChaos_PBDCollisionSolver_VectorRegister, TEXT(""));
+
 		bool bChaos_PBDCollisionSolver_Position_SolveEnabled = true;
 		int32 Chaos_PBDCollisionSolver_Position_ShockPropagationIterations = 3;
 		float Chaos_PBDCollisionSolver_Position_MinInvMassScale = 0.3f;
-		int32 Chaos_PBDCollisionSolver_Position_ZeroFrictionIterations = 4;
+		int32 Chaos_PBDCollisionSolver_Position_FrictionIterations = 4;
 		float Chaos_PBDCollisionSolver_Position_StaticFrictionStiffness = 0.5f;
 		float Chaos_PBDCollisionSolver_Position_PositionSolverTolerance = 0.001f;		// cms
 		float Chaos_PBDCollisionSolver_Position_RotationSolverTolerance = 0.001f;		// rads
@@ -34,7 +37,7 @@ namespace Chaos
 		FAutoConsoleVariableRef CVarChaos_PBDCollisionSolver_Position_SolveEnabled(TEXT("p.Chaos.PBDCollisionSolver.Position.SolveEnabled"), bChaos_PBDCollisionSolver_Position_SolveEnabled, TEXT(""));
 		FAutoConsoleVariableRef CVarChaos_PBDCollisionSolver_Position_UseShockPropagation(TEXT("p.Chaos.PBDCollisionSolver.Position.ShockPropagationIterations"), Chaos_PBDCollisionSolver_Position_ShockPropagationIterations, TEXT(""));
 		FAutoConsoleVariableRef CVarChaos_PBDCollisionSolver_Position_MinInvMassScale(TEXT("p.Chaos.PBDCollisionSolver.Position.MinInvMassScale"), Chaos_PBDCollisionSolver_Position_MinInvMassScale, TEXT(""));
-		FAutoConsoleVariableRef CVarChaos_PBDCollisionSolver_Position_ZeroFrictionIterations(TEXT("p.Chaos.PBDCollisionSolver.Position.ZeroFrictionIterations"), Chaos_PBDCollisionSolver_Position_ZeroFrictionIterations, TEXT(""));
+		FAutoConsoleVariableRef CVarChaos_PBDCollisionSolver_Position_FrictionIterations(TEXT("p.Chaos.PBDCollisionSolver.Position.FrictionIterations"), Chaos_PBDCollisionSolver_Position_FrictionIterations, TEXT(""));
 		FAutoConsoleVariableRef CVarChaos_PBDCollisionSolver_Position_StaticFrictionStiffness(TEXT("p.Chaos.PBDCollisionSolver.Position.StaticFriction.Stiffness"), Chaos_PBDCollisionSolver_Position_StaticFrictionStiffness, TEXT(""));
 		FAutoConsoleVariableRef CVarChaos_PBDCollisionSolver_Position_PositionSolverTolerance(TEXT("p.Chaos.PBDCollisionSolver.Position.PositionTolerance"), Chaos_PBDCollisionSolver_Position_PositionSolverTolerance, TEXT(""));
 		FAutoConsoleVariableRef CVarChaos_PBDCollisionSolver_Position_RotationSolverTolerance(TEXT("p.Chaos.PBDCollisionSolver.Position.RotationTolerance"), Chaos_PBDCollisionSolver_Position_RotationSolverTolerance, TEXT(""));
@@ -55,120 +58,6 @@ namespace Chaos
 	using namespace CVars;
 
 
-	void FPBDCollisionSolverManifoldPoint::InitContact(
-		const FSolverReal Dt,
-		const FConstraintSolverBody& Body0,
-		const FConstraintSolverBody& Body1,
-		const FSolverVec3& InRelativeContactPosition0,
-		const FSolverVec3& InRelativeContactPosition1,
-		const FSolverVec3& InWorldContactNormal,
-		const FSolverVec3& InWorldContactTangentU,
-		const FSolverVec3& InWorldContactTangentV,
-		const FSolverReal InWorldContactDeltaNormal,
-		const FSolverReal InWorldContactDeltaTangentU,
-		const FSolverReal InWorldContactDeltaTangentV)
-	{
-		RelativeContactPosition0 = InRelativeContactPosition0;
-		RelativeContactPosition1 = InRelativeContactPosition1;
-		WorldContactNormal = InWorldContactNormal;
-		WorldContactTangentU = InWorldContactTangentU;
-		WorldContactTangentV = InWorldContactTangentV;
-		WorldContactDeltaNormal = InWorldContactDeltaNormal;
-		WorldContactDeltaTangentU = InWorldContactDeltaTangentU;
-		WorldContactDeltaTangentV = InWorldContactDeltaTangentV;
-
-		NetPushOutNormal = FSolverReal(0);
-		NetPushOutTangentU = FSolverReal(0);
-		NetPushOutTangentV = FSolverReal(0);
-		NetImpulseNormal = FSolverReal(0);
-		NetImpulseTangentU = FSolverReal(0);
-		NetImpulseTangentV = FSolverReal(0);
-
-		UpdateMass(Body0, Body1);
-	}
-
-	void FPBDCollisionSolverManifoldPoint::InitMaterial(
-		const FConstraintSolverBody& Body0,
-		const FConstraintSolverBody& Body1,
-		const FSolverReal InRestitution,
-		const FSolverReal InRestitutionVelocityThreshold)
-	{
-		StaticFrictionRatio = FSolverReal(0);
-		WorldContactVelocityTargetNormal = FSolverReal(0);
-
-		if (InRestitution > FSolverReal(0))
-		{
-			const FSolverVec3 ContactVelocity = CalculateContactVelocity(Body0, Body1);
-			const FSolverReal ContactVelocityNormal = FSolverVec3::DotProduct(ContactVelocity, WorldContactNormal);
-			if (ContactVelocityNormal < -InRestitutionVelocityThreshold)
-			{
-				WorldContactVelocityTargetNormal = -InRestitution * ContactVelocityNormal;
-			}
-		}
-
-	}
-
-	inline void FPBDCollisionSolverManifoldPoint::UpdateMass(const FConstraintSolverBody& Body0, const FConstraintSolverBody& Body1)
-	{
-		FSolverReal ContactMassInvNormal = FSolverReal(0);
-		FSolverReal ContactMassInvTangentU = FSolverReal(0);
-		FSolverReal ContactMassInvTangentV = FSolverReal(0);
-
-		// These are not used if not initialized below so no need to clear
-		//WorldContactNormalAngular0 = FSolverVec3(0);
-		//WorldContactTangentUAngular0 = FSolverVec3(0);
-		//WorldContactTangentVAngular0 = FSolverVec3(0);
-		//WorldContactNormalAngular1 = FSolverVec3(0);
-		//WorldContactTangentUAngular1 = FSolverVec3(0);
-		//WorldContactTangentVAngular1 = FSolverVec3(0);
-
-		if (Body0.IsDynamic())
-		{
-			const FSolverVec3 R0xN = FSolverVec3::CrossProduct(RelativeContactPosition0, WorldContactNormal);
-			const FSolverVec3 R0xU = FSolverVec3::CrossProduct(RelativeContactPosition0, WorldContactTangentU);
-			const FSolverVec3 R0xV = FSolverVec3::CrossProduct(RelativeContactPosition0, WorldContactTangentV);
-
-			const FSolverMatrix33 InvI0 = Body0.InvI();
-
-			WorldContactNormalAngular0 = InvI0 * R0xN;
-			WorldContactTangentUAngular0 = InvI0 * R0xU;
-			WorldContactTangentVAngular0 = InvI0 * R0xV;
-
-			ContactMassInvNormal += FSolverVec3::DotProduct(R0xN, WorldContactNormalAngular0) + Body0.InvM();
-			ContactMassInvTangentU += FSolverVec3::DotProduct(R0xU, WorldContactTangentUAngular0) + Body0.InvM();
-			ContactMassInvTangentV += FSolverVec3::DotProduct(R0xV, WorldContactTangentVAngular0) + Body0.InvM();
-		}
-		if (Body1.IsDynamic())
-		{
-			const FSolverVec3 R1xN = FSolverVec3::CrossProduct(RelativeContactPosition1, WorldContactNormal);
-			const FSolverVec3 R1xU = FSolverVec3::CrossProduct(RelativeContactPosition1, WorldContactTangentU);
-			const FSolverVec3 R1xV = FSolverVec3::CrossProduct(RelativeContactPosition1, WorldContactTangentV);
-
-			const FSolverMatrix33 InvI1 = Body1.InvI();
-
-			WorldContactNormalAngular1 = InvI1 * R1xN;
-			WorldContactTangentUAngular1 = InvI1 * R1xU;
-			WorldContactTangentVAngular1 = InvI1 * R1xV;
-
-			ContactMassInvNormal += FSolverVec3::DotProduct(R1xN, WorldContactNormalAngular1) + Body1.InvM();
-			ContactMassInvTangentU += FSolverVec3::DotProduct(R1xU, WorldContactTangentUAngular1) + Body1.InvM();
-			ContactMassInvTangentV += FSolverVec3::DotProduct(R1xV, WorldContactTangentVAngular1) + Body1.InvM();
-		}
-
-		ContactMassNormal = (ContactMassInvNormal > FSolverReal(SMALL_NUMBER)) ? FSolverReal(1) / ContactMassInvNormal : FSolverReal(0);
-		ContactMassTangentU = (ContactMassInvTangentU > FSolverReal(SMALL_NUMBER)) ? FSolverReal(1) / ContactMassInvTangentU : FSolverReal(0);
-		ContactMassTangentV = (ContactMassInvTangentV > FSolverReal(SMALL_NUMBER)) ? FSolverReal(1) / ContactMassInvTangentV : FSolverReal(0);
-	}
-
-	FSolverVec3 FPBDCollisionSolverManifoldPoint::CalculateContactVelocity(const FConstraintSolverBody& Body0, const FConstraintSolverBody& Body1) const
-	{
-		const FSolverVec3 ContactVelocity0 = Body0.V() + FVec3::CrossProduct(Body0.W(), RelativeContactPosition0);
-		const FSolverVec3 ContactVelocity1 = Body1.V() + FVec3::CrossProduct(Body1.W(), RelativeContactPosition1);
-		return ContactVelocity0 - ContactVelocity1;
-	}
-
-
-
 	//////////////////////////////////////////////////////////////////////////////////////////////////
 	//////////////////////////////////////////////////////////////////////////////////////////////////
 	//////////////////////////////////////////////////////////////////////////////////////////////////
@@ -230,50 +119,12 @@ namespace Chaos
 			{
 				for (int32 PointIndex = 0; PointIndex < State.NumManifoldPoints; ++PointIndex)
 				{
-					State.ManifoldPoints[PointIndex].UpdateMass(Body0, Body1);
+					State.ManifoldPoints[PointIndex].UpdateMassNormal(Body0, Body1);
 				}
 			}
 		}
 	}
 
-	void FPBDCollisionSolver::InitContact(
-		const int32 ManifoldPoiontIndex,
-		const FSolverReal Dt,
-		const FSolverVec3& InRelativeContactPosition0,
-		const FSolverVec3& InRelativeContactPosition1,
-		const FSolverVec3& InWorldContactNormal,
-		const FSolverVec3& InWorldContactTangentU,
-		const FSolverVec3& InWorldContactTangentV,
-		const FSolverReal InWorldContactDeltaNormal,
-		const FSolverReal InWorldContactDeltaTangentU,
-		const FSolverReal InWorldContactDeltaTangentV)
-	{
-		State.ManifoldPoints[ManifoldPoiontIndex].InitContact(
-			FSolverReal(Dt),
-			State.SolverBodies[0], 
-			State.SolverBodies[1], 
-			InRelativeContactPosition0,
-			InRelativeContactPosition1,
-			InWorldContactNormal,
-			InWorldContactTangentU,
-			InWorldContactTangentV,
-			InWorldContactDeltaNormal,
-			InWorldContactDeltaTangentU,
-			InWorldContactDeltaTangentV);
-	}
-
-	void FPBDCollisionSolver::InitMaterial(
-		const int32 ManifoldPoiontIndex,
-		const FSolverReal InRestitution,
-		const FSolverReal InRestitutionVelocityThreshold)
-	{
-		State.ManifoldPoints[ManifoldPoiontIndex].InitMaterial(
-			State.SolverBodies[0], 
-			State.SolverBodies[1], 
-			InRestitution, 
-			InRestitutionVelocityThreshold);
-	}
-
 	void FPBDCollisionSolver::SolveVelocityAverage(const FSolverReal Dt)
 	{
 		FConstraintSolverBody& Body0 = SolverBody0();
diff --git a/Engine/Source/Runtime/Experimental/Chaos/Private/Chaos/Collision/PBDCollisionSolver.h b/Engine/Source/Runtime/Experimental/Chaos/Private/Chaos/Collision/PBDCollisionSolver.h
--- a/Engine/Source/Runtime/Experimental/Chaos/Private/Chaos/Collision/PBDCollisionSolver.h
+++ b/Engine/Source/Runtime/Experimental/Chaos/Private/Chaos/Collision/PBDCollisionSolver.h
@@ -55,6 +55,19 @@ namespace Chaos
 			const FSolverReal InWorldContactDeltaTangentU,
 			const FSolverReal InWorldContactDeltaTangentV);
 
+		void InitContact(
+			const FSolverReal Dt,
+			const FConstraintSolverBody& Body0,
+			const FConstraintSolverBody& Body1,
+			const SolverVectorRegister& InRelativeContactPosition0,
+			const SolverVectorRegister& InRelativeContactPosition1,
+			const SolverVectorRegister& InWorldContactNormal,
+			const SolverVectorRegister& InWorldContactTangentU,
+			const SolverVectorRegister& InWorldContactTangentV,
+			const FSolverReal InWorldContactDeltaNormal,
+			const FSolverReal InWorldContactDeltaTangentU,
+			const FSolverReal InWorldContactDeltaTangentV);
+
 		/**
 		 * @brief Initialize the material related properties of the contact
 		*/
@@ -69,6 +82,21 @@ namespace Chaos
 		*/
 		void UpdateMass(const FConstraintSolverBody& Body0, const FConstraintSolverBody& Body1);
 
+		/**
+		 * @brief Update the contact mass for the normal correction
+		 * This is used by shock propagation.
+		*/
+		void UpdateMassNormal(const FConstraintSolverBody& Body0, const FConstraintSolverBody& Body1);
+
+		void UpdateMass(
+			const FConstraintSolverBody& Body0,
+			const FConstraintSolverBody& Body1,
+			const SolverVectorRegister& InRelativeContactPosition0,
+			const SolverVectorRegister& InRelativeContactPosition1,
+			const SolverVectorRegister& InWorldContactNormal,
+			const SolverVectorRegister& InWorldContactTangentU,
+			const SolverVectorRegister& InWorldContactTangentV);
+
 		/**
 		 * @brief Calculate the position error at the current transforms
 		 * @param MaxPushOut a limit on the position error for this iteration to prevent initial-penetration explosion (a common PBD problem)
@@ -92,6 +120,9 @@ namespace Chaos
 		*/
 		FSolverVec3 CalculateContactVelocity(const FConstraintSolverBody& Body0, const FConstraintSolverBody& Body1) const;
 
+		/**
+		 * @brief Whether we need to solve velocity for this manifold point (only if we were penetrating or applied a pushout)
+		*/
 		bool ShouldSolveVelocity() const;
 
 		// World-space body-relative contact points
@@ -203,9 +234,11 @@ namespace Chaos
 			return State.ManifoldPoints[ManifoldPointIndex];
 		}
 
-		void InitContact(
+		void SetManifoldPoint(
 			const int32 ManifoldPoiontIndex,
 			const FSolverReal Dt,
+			const FSolverReal InRestitution,
+			const FSolverReal InRestitutionVelocityThreshold,
 			const FSolverVec3& InRelativeContactPosition0,
 			const FSolverVec3& InRelativeContactPosition1,
 			const FSolverVec3& InWorldContactNormal,
@@ -215,10 +248,19 @@ namespace Chaos
 			const FSolverReal InWorldContactDeltaTangentU,
 			const FSolverReal InWorldContactDeltaTangentV);
 
-		void InitMaterial(
+		void SetManifoldPoint(
 			const int32 ManifoldPoiontIndex,
+			const FSolverReal Dt,
 			const FSolverReal InRestitution,
-			const FSolverReal InRestitutionVelocityThreshold);
+			const FSolverReal InRestitutionVelocityThreshold,
+			const SolverVectorRegister& InRelativeContactPosition0,
+			const SolverVectorRegister& InRelativeContactPosition1,
+			const SolverVectorRegister& InWorldContactNormal,
+			const SolverVectorRegister& InWorldContactTangentU,
+			const SolverVectorRegister& InWorldContactTangentV,
+			const FSolverReal InWorldContactDeltaNormal,
+			const FSolverReal InWorldContactDeltaTangentU,
+			const FSolverReal InWorldContactDeltaTangentV);
 
 		/**
 		 * @brief Get the first (decaorated) solver body
@@ -610,6 +652,234 @@ namespace Chaos
 	//////////////////////////////////////////////////////////////////////////////////////////////////
 	//////////////////////////////////////////////////////////////////////////////////////////////////
 
+	FORCEINLINE_DEBUGGABLE void FPBDCollisionSolverManifoldPoint::InitContact(
+		const FSolverReal Dt,
+		const FConstraintSolverBody& Body0,
+		const FConstraintSolverBody& Body1,
+		const FSolverVec3& InRelativeContactPosition0,
+		const FSolverVec3& InRelativeContactPosition1,
+		const FSolverVec3& InWorldContactNormal,
+		const FSolverVec3& InWorldContactTangentU,
+		const FSolverVec3& InWorldContactTangentV,
+		const FSolverReal InWorldContactDeltaNormal,
+		const FSolverReal InWorldContactDeltaTangentU,
+		const FSolverReal InWorldContactDeltaTangentV)
+	{
+		RelativeContactPosition0 = InRelativeContactPosition0;
+		RelativeContactPosition1 = InRelativeContactPosition1;
+		WorldContactNormal = InWorldContactNormal;
+		WorldContactTangentU = InWorldContactTangentU;
+		WorldContactTangentV = InWorldContactTangentV;
+		WorldContactDeltaNormal = InWorldContactDeltaNormal;
+		WorldContactDeltaTangentU = InWorldContactDeltaTangentU;
+		WorldContactDeltaTangentV = InWorldContactDeltaTangentV;
+
+		NetPushOutNormal = FSolverReal(0);
+		NetPushOutTangentU = FSolverReal(0);
+		NetPushOutTangentV = FSolverReal(0);
+		NetImpulseNormal = FSolverReal(0);
+		NetImpulseTangentU = FSolverReal(0);
+		NetImpulseTangentV = FSolverReal(0);
+
+		StaticFrictionRatio = FSolverReal(0);
+		WorldContactVelocityTargetNormal = FSolverReal(0);
+
+		UpdateMass(Body0, Body1);
+	}
+
+
+	FORCEINLINE_DEBUGGABLE void FPBDCollisionSolverManifoldPoint::InitContact(
+		const FSolverReal Dt,
+		const FConstraintSolverBody& Body0,
+		const FConstraintSolverBody& Body1,
+		const SolverVectorRegister& InRelativeContactPosition0,
+		const SolverVectorRegister& InRelativeContactPosition1,
+		const SolverVectorRegister& InWorldContactNormal,
+		const SolverVectorRegister& InWorldContactTangentU,
+		const SolverVectorRegister& InWorldContactTangentV,
+		const FSolverReal InWorldContactDeltaNormal,
+		const FSolverReal InWorldContactDeltaTangentU,
+		const FSolverReal InWorldContactDeltaTangentV)
+	{
+		VectorStoreFloat3(InRelativeContactPosition0, &RelativeContactPosition0);
+		VectorStoreFloat3(InRelativeContactPosition1, &RelativeContactPosition1);
+		VectorStoreFloat3(InWorldContactNormal, &WorldContactNormal);
+		VectorStoreFloat3(InWorldContactTangentU, &WorldContactTangentU);
+		VectorStoreFloat3(InWorldContactTangentV, &WorldContactTangentV);
+		WorldContactDeltaNormal = InWorldContactDeltaNormal;
+		WorldContactDeltaTangentU = InWorldContactDeltaTangentU;
+		WorldContactDeltaTangentV = InWorldContactDeltaTangentV;
+
+		NetPushOutNormal = FSolverReal(0);
+		NetPushOutTangentU = FSolverReal(0);
+		NetPushOutTangentV = FSolverReal(0);
+		NetImpulseNormal = FSolverReal(0);
+		NetImpulseTangentU = FSolverReal(0);
+		NetImpulseTangentV = FSolverReal(0);
+
+		StaticFrictionRatio = FSolverReal(0);
+		WorldContactVelocityTargetNormal = FSolverReal(0);
+
+		UpdateMass(
+			Body0, 
+			Body1,
+			InRelativeContactPosition0,
+			InRelativeContactPosition1,
+			InWorldContactNormal,
+			InWorldContactTangentU,
+			InWorldContactTangentV);
+	}
+
+	FORCEINLINE_DEBUGGABLE void FPBDCollisionSolverManifoldPoint::InitMaterial(
+		const FConstraintSolverBody& Body0,
+		const FConstraintSolverBody& Body1,
+		const FSolverReal InRestitution,
+		const FSolverReal InRestitutionVelocityThreshold)
+	{
+		if (InRestitution > FSolverReal(0))
+		{
+			const FSolverVec3 ContactVelocity = CalculateContactVelocity(Body0, Body1);
+			const FSolverReal ContactVelocityNormal = FSolverVec3::DotProduct(ContactVelocity, WorldContactNormal);
+			if (ContactVelocityNormal < -InRestitutionVelocityThreshold)
+			{
+				WorldContactVelocityTargetNormal = -InRestitution * ContactVelocityNormal;
+			}
+		}
+	}
+
+	FORCEINLINE_DEBUGGABLE void FPBDCollisionSolverManifoldPoint::UpdateMass(const FConstraintSolverBody& Body0, const FConstraintSolverBody& Body1)
+	{
+		FSolverReal ContactMassInvNormal = FSolverReal(0);
+		FSolverReal ContactMassInvTangentU = FSolverReal(0);
+		FSolverReal ContactMassInvTangentV = FSolverReal(0);
+
+		// These are not used if not initialized below so no need to clear
+		//WorldContactNormalAngular0 = FSolverVec3(0);
+		//WorldContactTangentUAngular0 = FSolverVec3(0);
+		//WorldContactTangentVAngular0 = FSolverVec3(0);
+		//WorldContactNormalAngular1 = FSolverVec3(0);
+		//WorldContactTangentUAngular1 = FSolverVec3(0);
+		//WorldContactTangentVAngular1 = FSolverVec3(0);
+
+		if (Body0.IsDynamic())
+		{
+			const FSolverVec3 R0xN = FSolverVec3::CrossProduct(RelativeContactPosition0, WorldContactNormal);
+			const FSolverVec3 R0xU = FSolverVec3::CrossProduct(RelativeContactPosition0, WorldContactTangentU);
+			const FSolverVec3 R0xV = FSolverVec3::CrossProduct(RelativeContactPosition0, WorldContactTangentV);
+
+			const FSolverMatrix33 InvI0 = Body0.InvI();
+
+			WorldContactNormalAngular0 = InvI0 * R0xN;
+			WorldContactTangentUAngular0 = InvI0 * R0xU;
+			WorldContactTangentVAngular0 = InvI0 * R0xV;
+
+			ContactMassInvNormal += FSolverVec3::DotProduct(R0xN, WorldContactNormalAngular0) + Body0.InvM();
+			ContactMassInvTangentU += FSolverVec3::DotProduct(R0xU, WorldContactTangentUAngular0) + Body0.InvM();
+			ContactMassInvTangentV += FSolverVec3::DotProduct(R0xV, WorldContactTangentVAngular0) + Body0.InvM();
+		}
+		if (Body1.IsDynamic())
+		{
+			const FSolverVec3 R1xN = FSolverVec3::CrossProduct(RelativeContactPosition1, WorldContactNormal);
+			const FSolverVec3 R1xU = FSolverVec3::CrossProduct(RelativeContactPosition1, WorldContactTangentU);
+			const FSolverVec3 R1xV = FSolverVec3::CrossProduct(RelativeContactPosition1, WorldContactTangentV);
+
+			const FSolverMatrix33 InvI1 = Body1.InvI();
+
+			WorldContactNormalAngular1 = InvI1 * R1xN;
+			WorldContactTangentUAngular1 = InvI1 * R1xU;
+			WorldContactTangentVAngular1 = InvI1 * R1xV;
+
+			ContactMassInvNormal += FSolverVec3::DotProduct(R1xN, WorldContactNormalAngular1) + Body1.InvM();
+			ContactMassInvTangentU += FSolverVec3::DotProduct(R1xU, WorldContactTangentUAngular1) + Body1.InvM();
+			ContactMassInvTangentV += FSolverVec3::DotProduct(R1xV, WorldContactTangentVAngular1) + Body1.InvM();
+		}
+
+		ContactMassNormal = (ContactMassInvNormal > FSolverReal(SMALL_NUMBER)) ? FSolverReal(1) / ContactMassInvNormal : FSolverReal(0);
+		ContactMassTangentU = (ContactMassInvTangentU > FSolverReal(SMALL_NUMBER)) ? FSolverReal(1) / ContactMassInvTangentU : FSolverReal(0);
+		ContactMassTangentV = (ContactMassInvTangentV > FSolverReal(SMALL_NUMBER)) ? FSolverReal(1) / ContactMassInvTangentV : FSolverReal(0);
+	}
+
+	FORCEINLINE_DEBUGGABLE void FPBDCollisionSolverManifoldPoint::UpdateMass(
+		const FConstraintSolverBody& Body0, 
+		const FConstraintSolverBody& Body1,
+		const SolverVectorRegister& InRelativeContactPosition0,
+		const SolverVectorRegister& InRelativeContactPosition1,
+		const SolverVectorRegister& InWorldContactNormal,
+		const SolverVectorRegister& InWorldContactTangentU,
+		const SolverVectorRegister& InWorldContactTangentV)
+	{
+		FSolverReal ContactMassInvNormal = FSolverReal(0);
+		FSolverReal ContactMassInvTangentU = FSolverReal(0);
+		FSolverReal ContactMassInvTangentV = FSolverReal(0);
+
+		if (Body0.IsDynamic())
+		{
+			const SolverVectorRegister R0xN = VectorCross(InRelativeContactPosition0, InWorldContactNormal);
+			const SolverVectorRegister R0xU = VectorCross(InRelativeContactPosition0, InWorldContactTangentU);
+			const SolverVectorRegister R0xV = VectorCross(InRelativeContactPosition0, InWorldContactTangentV);
+
+			const SolverVectorRegister InvMScale0 = VectorSetFloat1(Body0.InvMScale());
+			const FSolverMatrix33& InvI0 = Body0.SolverBody().InvI();
+
+			const SolverVectorRegister NormalAngular0 = VectorMultiply(VectorTransformVector(R0xN, &InvI0), InvMScale0);
+			const SolverVectorRegister TangentUAngular0 = VectorMultiply(VectorTransformVector(R0xU, &InvI0), InvMScale0);
+			const SolverVectorRegister TangentVAngular0 = VectorMultiply(VectorTransformVector(R0xV, &InvI0), InvMScale0);
+
+			ContactMassInvNormal += VectorGetComponent(VectorDot3(R0xN, NormalAngular0), 0) + Body0.InvM();
+			ContactMassInvTangentU += VectorGetComponent(VectorDot3(R0xU, TangentUAngular0), 0) + Body0.InvM();
+			ContactMassInvTangentV += VectorGetComponent(VectorDot3(R0xV, TangentVAngular0), 0) + Body0.InvM();
+
+			VectorStoreFloat3(NormalAngular0, &WorldContactNormalAngular0);
+			VectorStoreFloat3(TangentUAngular0, &WorldContactTangentUAngular0);
+			VectorStoreFloat3(TangentVAngular0, &WorldContactTangentVAngular0);
+		}
+		if (Body1.IsDynamic())
+		{
+			const SolverVectorRegister R1xN = VectorCross(InRelativeContactPosition1, InWorldContactNormal);
+			const SolverVectorRegister R1xU = VectorCross(InRelativeContactPosition1, InWorldContactTangentU);
+			const SolverVectorRegister R1xV = VectorCross(InRelativeContactPosition1, InWorldContactTangentV);
+
+			const SolverVectorRegister InvMScale1 = VectorSetFloat1(Body1.InvMScale());
+			const FSolverMatrix33& InvI1 = Body1.SolverBody().InvI();
+
+			const SolverVectorRegister NormalAngular1 = VectorMultiply(VectorTransformVector(R1xN, &InvI1), InvMScale1);
+			const SolverVectorRegister TangentUAngular1 = VectorMultiply(VectorTransformVector(R1xU, &InvI1), InvMScale1);
+			const SolverVectorRegister TangentVAngular1 = VectorMultiply(VectorTransformVector(R1xV, &InvI1), InvMScale1);
+
+			ContactMassInvNormal += VectorGetComponent(VectorDot3(R1xN, NormalAngular1), 0) + Body1.InvM();
+			ContactMassInvTangentU += VectorGetComponent(VectorDot3(R1xU, TangentUAngular1), 0) + Body1.InvM();
+			ContactMassInvTangentV += VectorGetComponent(VectorDot3(R1xV, TangentVAngular1), 0) + Body1.InvM();
+
+			VectorStoreFloat3(NormalAngular1, &WorldContactNormalAngular1);
+			VectorStoreFloat3(TangentUAngular1, &WorldContactTangentUAngular1);
+			VectorStoreFloat3(TangentVAngular1, &WorldContactTangentVAngular1);
+		}
+
+		ContactMassNormal = (ContactMassInvNormal > FSolverReal(SMALL_NUMBER)) ? FSolverReal(1) / ContactMassInvNormal : FSolverReal(0);
+		ContactMassTangentU = (ContactMassInvTangentU > FSolverReal(SMALL_NUMBER)) ? FSolverReal(1) / ContactMassInvTangentU : FSolverReal(0);
+		ContactMassTangentV = (ContactMassInvTangentV > FSolverReal(SMALL_NUMBER)) ? FSolverReal(1) / ContactMassInvTangentV : FSolverReal(0);
+	}
+
+	FORCEINLINE_DEBUGGABLE void FPBDCollisionSolverManifoldPoint::UpdateMassNormal(const FConstraintSolverBody& Body0, const FConstraintSolverBody& Body1)
+	{
+		FSolverReal ContactMassInvNormal = FSolverReal(0);
+		if (Body0.IsDynamic())
+		{
+			const FSolverVec3 R0xN = FSolverVec3::CrossProduct(RelativeContactPosition0, WorldContactNormal);
+			const FSolverMatrix33 InvI0 = Body0.InvI();
+			WorldContactNormalAngular0 = InvI0 * R0xN;
+			ContactMassInvNormal += FSolverVec3::DotProduct(R0xN, WorldContactNormalAngular0) + Body0.InvM();
+		}
+		if (Body1.IsDynamic())
+		{
+			const FSolverVec3 R1xN = FSolverVec3::CrossProduct(RelativeContactPosition1, WorldContactNormal);
+			const FSolverMatrix33 InvI1 = Body1.InvI();
+			WorldContactNormalAngular1 = InvI1 * R1xN;
+			ContactMassInvNormal += FSolverVec3::DotProduct(R1xN, WorldContactNormalAngular1) + Body1.InvM();
+		}
+		ContactMassNormal = (ContactMassInvNormal > FSolverReal(SMALL_NUMBER)) ? FSolverReal(1) / ContactMassInvNormal : FSolverReal(0);
+	}
+
 	FORCEINLINE_DEBUGGABLE void FPBDCollisionSolverManifoldPoint::CalculateContactPositionErrorNormal(const FConstraintSolverBody& Body0, const FConstraintSolverBody& Body1, const FSolverReal MaxPushOut, FSolverReal& OutContactDeltaNormal) const
 	{
 		// Linear version: calculate the contact delta assuming linear motion after applying a positional impulse at the contact point. There will be an error that depends on the size of the rotation.
@@ -682,6 +952,13 @@ namespace Chaos
 		OutContactVelocityDeltaNormal = (ContactVelocityNormal - WorldContactVelocityTargetNormal);
 	}
 
+	FORCEINLINE_DEBUGGABLE FSolverVec3 FPBDCollisionSolverManifoldPoint::CalculateContactVelocity(const FConstraintSolverBody& Body0, const FConstraintSolverBody& Body1) const
+	{
+		const FSolverVec3 ContactVelocity0 = Body0.V() + FVec3::CrossProduct(Body0.W(), RelativeContactPosition0);
+		const FSolverVec3 ContactVelocity1 = Body1.V() + FVec3::CrossProduct(Body1.W(), RelativeContactPosition1);
+		return ContactVelocity0 - ContactVelocity1;
+	}
+
 	FORCEINLINE_DEBUGGABLE bool FPBDCollisionSolverManifoldPoint::ShouldSolveVelocity() const
 	{
 		// We ensure positive separating velocity for close contacts even if they didn't receive a pushout
@@ -693,6 +970,73 @@ namespace Chaos
 	//////////////////////////////////////////////////////////////////////////////////////////////////
 	//////////////////////////////////////////////////////////////////////////////////////////////////
 
+	FORCEINLINE_DEBUGGABLE void FPBDCollisionSolver::SetManifoldPoint(
+		const int32 ManifoldPoiontIndex,
+		const FSolverReal Dt,
+		const FSolverReal InRestitution,
+		const FSolverReal InRestitutionVelocityThreshold,
+		const FSolverVec3& InRelativeContactPosition0,
+		const FSolverVec3& InRelativeContactPosition1,
+		const FSolverVec3& InWorldContactNormal,
+		const FSolverVec3& InWorldContactTangentU,
+		const FSolverVec3& InWorldContactTangentV,
+		const FSolverReal InWorldContactDeltaNormal,
+		const FSolverReal InWorldContactDeltaTangentU,
+		const FSolverReal InWorldContactDeltaTangentV)
+	{
+		State.ManifoldPoints[ManifoldPoiontIndex].InitContact(
+			FSolverReal(Dt),
+			State.SolverBodies[0],
+			State.SolverBodies[1],
+			InRelativeContactPosition0,
+			InRelativeContactPosition1,
+			InWorldContactNormal,
+			InWorldContactTangentU,
+			InWorldContactTangentV,
+			InWorldContactDeltaNormal,
+			InWorldContactDeltaTangentU,
+			InWorldContactDeltaTangentV);
+
+		State.ManifoldPoints[ManifoldPoiontIndex].InitMaterial(
+			State.SolverBodies[0],
+			State.SolverBodies[1],
+			InRestitution,
+			InRestitutionVelocityThreshold);
+	}
+
+	FORCEINLINE_DEBUGGABLE void FPBDCollisionSolver::SetManifoldPoint(
+		const int32 ManifoldPoiontIndex,
+		const FSolverReal Dt,
+		const FSolverReal InRestitution,
+		const FSolverReal InRestitutionVelocityThreshold,
+		const SolverVectorRegister& InRelativeContactPosition0,
+		const SolverVectorRegister& InRelativeContactPosition1,
+		const SolverVectorRegister& InWorldContactNormal,
+		const SolverVectorRegister& InWorldContactTangentU,
+		const SolverVectorRegister& InWorldContactTangentV,
+		const FSolverReal InWorldContactDeltaNormal,
+		const FSolverReal InWorldContactDeltaTangentU,
+		const FSolverReal InWorldContactDeltaTangentV)
+	{
+		State.ManifoldPoints[ManifoldPoiontIndex].InitContact(
+			FSolverReal(Dt),
+			State.SolverBodies[0],
+			State.SolverBodies[1],
+			InRelativeContactPosition0,
+			InRelativeContactPosition1,
+			InWorldContactNormal,
+			InWorldContactTangentU,
+			InWorldContactTangentV,
+			InWorldContactDeltaNormal,
+			InWorldContactDeltaTangentU,
+			InWorldContactDeltaTangentV);
+
+		State.ManifoldPoints[ManifoldPoiontIndex].InitMaterial(
+			State.SolverBodies[0],
+			State.SolverBodies[1],
+			InRestitution,
+			InRestitutionVelocityThreshold);
+	}
 
 	FORCEINLINE_DEBUGGABLE bool FPBDCollisionSolver::SolvePositionWithFriction(const FSolverReal Dt, const FSolverReal MaxPushOut)
 	{
diff --git a/Engine/Source/Runtime/Experimental/Chaos/Private/Chaos/Collision/SolverCollisionContainer.cpp b/Engine/Source/Runtime/Experimental/Chaos/Private/Chaos/Collision/SolverCollisionContainer.cpp
--- a/Engine/Source/Runtime/Experimental/Chaos/Private/Chaos/Collision/SolverCollisionContainer.cpp
+++ b/Engine/Source/Runtime/Experimental/Chaos/Private/Chaos/Collision/SolverCollisionContainer.cpp
@@ -19,7 +19,9 @@ namespace Chaos
 {
 	namespace CVars
 	{
-		extern int32 Chaos_PBDCollisionSolver_Position_ZeroFrictionIterations;
+		extern bool bChaos_PBDCollisionSolver_VectorRegister;
+
+		extern int32 Chaos_PBDCollisionSolver_Position_FrictionIterations;
 		extern int32 Chaos_PBDCollisionSolver_Position_ShockPropagationIterations;
 
 		extern int32 Chaos_PBDCollisionSolver_Velocity_ShockPropagationIterations;
@@ -95,7 +97,14 @@ namespace Chaos
 
 			Solver.SetSolverBodies(Body0, Body1);
 
-			GatherManifoldPoints(Dt, Body0, Body1);
+			if (!bChaos_PBDCollisionSolver_VectorRegister)
+			{
+				GatherManifoldPoints(Dt, Body0, Body1);
+			}
+			else
+			{
+				GatherManifoldPointsVectorRegister(Dt, Body0, Body1);
+			}
 
 			// We should try to remove this - the Constraint should not need to know about solver objects
 			Constraint->SetSolverBodies(Body0, Body1);
@@ -130,11 +139,12 @@ namespace Chaos
 				const FSolverVec3 WorldContactNormal = FSolverVec3(ShapeWorldTransform1.TransformVectorNoScale(ManifoldPoint.ContactPoint.ShapeContactNormal));
 				const FSolverVec3 RelativeContactPosition0 = FSolverVec3(WorldContactPoint - Body0->P());
 				const FSolverVec3 RelativeContactPosition1 = FSolverVec3(WorldContactPoint - Body1->P());
+				const FSolverReal TargetPhi = FSolverReal(ManifoldPoint.TargetPhi);
 
 				// If we have contact data from a previous tick, use it to calculate the lateral position delta we need
 				// to apply to move the contacts back to their original relative locations (i.e., to enforce static friction)
 				// @todo(chaos): we should not be writing back to the constraint here - find a better way to update the friction anchor. See FPBDCollisionConstraint::SetSolverResults
-				FSolverVec3 WorldFrictionDelta = FVec3(0);
+				FSolverVec3 WorldFrictionDelta = FSolverVec3(0);
 				const FSavedManifoldPoint* SavedManifoldPoint = Constraint->FindSavedManifoldPoint(ManifoldPoint);
 				if (SavedManifoldPoint != nullptr)
 				{
@@ -167,15 +177,17 @@ namespace Chaos
 				const FSolverVec3 WorldContactTangentV = FSolverVec3::CrossProduct(WorldContactNormal, WorldContactTangentU);
 
 				// The contact point error we are trying to correct in this solver
-				const FSolverVec3 WorldContactDelta = FSolverVec3(WorldContactPoint0 - WorldContactPoint1);
+				const FSolverVec3 WorldContactDelta = FSolverVec3(WorldContactPoint0 - WorldContactPoint1) - TargetPhi;
 				const FSolverReal WorldContactDeltaNormal = FSolverVec3::DotProduct(WorldContactDelta, WorldContactNormal);
 				const FSolverReal WorldContactDeltaTangentU = FSolverVec3::DotProduct(WorldContactDelta + WorldFrictionDelta, WorldContactTangentU);
 				const FSolverReal WorldContactDeltaTangentV = FSolverVec3::DotProduct(WorldContactDelta + WorldFrictionDelta, WorldContactTangentV);
 
-				// Initialize the structural data in the contact (relative contact points, contact mass etc)
-				Solver.InitContact(
+				// Copy all the properties into the solver
+				Solver.SetManifoldPoint(
 					ManifoldPointIndex,
 					Dt,
+					Restitution,
+					RestitutionVelocityThreshold,
 					RelativeContactPosition0,
 					RelativeContactPosition1,
 					WorldContactNormal,
@@ -184,15 +196,123 @@ namespace Chaos
 					WorldContactDeltaNormal,
 					WorldContactDeltaTangentU,
 					WorldContactDeltaTangentV);
+			}
+		}
+
+		void GatherManifoldPointsVectorRegister(
+			const FReal InDt,
+			FSolverBody* Body0,
+			FSolverBody* Body1)
+		{
+			// Attempt to make this LWC independent. Actually this isn't easy because we explicitly use MakeVectorRegisterFloatFromDouble
+
+			// We handle incremental manifolds by just collecting any new contacts
+			const int32 BeginPointIndex = Solver.NumManifoldPoints();
+			const int32 EndPointIndex = Solver.SetNumManifoldPoints(Constraint->GetManifoldPoints().Num());
+			TArrayView<FManifoldPoint> ManifoldPoints = Constraint->GetManifoldPoints();
+
+			const FSolverReal Dt = FSolverReal(InDt);
+			const FSolverReal RestitutionVelocityThreshold = FSolverReal(Constraint->GetRestitutionThreshold()) * Dt;
+			const FSolverReal Restitution = FSolverReal(Constraint->GetRestitution());
+
+			// World space positions - must be LWC compatible
+			const VectorRegister ShapeWorldPos0 = Constraint->GetShapeWorldTransform0().GetTranslationRegister();
+			const VectorRegister ShapeWorldRot0 = Constraint->GetShapeWorldTransform0().GetRotationRegister();
+			const VectorRegister ShapeWorldPos1 = Constraint->GetShapeWorldTransform1().GetTranslationRegister();
+			const VectorRegister ShapeWorldRot1 = Constraint->GetShapeWorldTransform1().GetRotationRegister();
+			const VectorRegister BodyPos0 = VectorLoadFloat3(&Body0->P());
+			const VectorRegister BodyPos1 = VectorLoadFloat3(&Body1->P());
 
-				// Initialize the material properties (restitution and friction related)
-				Solver.InitMaterial(
+			for (int32 ManifoldPointIndex = BeginPointIndex; ManifoldPointIndex < EndPointIndex; ++ManifoldPointIndex)
+			{
+				FManifoldPoint& ManifoldPoint = ManifoldPoints[ManifoldPointIndex];
+				const FSavedManifoldPoint* SavedManifoldPoint = Constraint->FindSavedManifoldPoint(ManifoldPoint);
+
+				const VectorRegister ShapeContactPos0 = VectorLoadFloat3(&ManifoldPoint.ContactPoint.ShapeContactPoints[0]);
+				const VectorRegister ShapeContactPos1 = VectorLoadFloat3(&ManifoldPoint.ContactPoint.ShapeContactPoints[1]);
+				const VectorRegister ShapeContactNormal1 = VectorLoadFloat3(&ManifoldPoint.ContactPoint.ShapeContactNormal);
+
+				const VectorRegister WorldContactPoint0 = VectorAdd(ShapeWorldPos0, VectorQuaternionRotateVector(ShapeWorldRot0, ShapeContactPos0));
+				const VectorRegister WorldContactPoint1 = VectorAdd(ShapeWorldPos1, VectorQuaternionRotateVector(ShapeWorldRot1, ShapeContactPos1));
+				const VectorRegister WorldContactPoint = VectorMultiply(VectorSetFloat1(FSolverReal(0.5)), VectorAdd(WorldContactPoint0, WorldContactPoint1));
+
+				// NOTE: low precision for relative coordinates
+				const SolverVectorRegister WorldContactNormal = MakeVectorRegisterFloatFromDouble(VectorQuaternionRotateVector(ShapeWorldRot1, ShapeContactNormal1));
+				const SolverVectorRegister RelativeContactPosition0 = MakeVectorRegisterFloatFromDouble(VectorSubtract(WorldContactPoint, BodyPos0));
+				const SolverVectorRegister RelativeContactPosition1 = MakeVectorRegisterFloatFromDouble(VectorSubtract(WorldContactPoint, BodyPos1));
+				const FSolverReal TargetPhi = FSolverReal(ManifoldPoint.TargetPhi);
+
+				// If we have contact data from a previous tick, use it to calculate the lateral position delta we need
+				// to apply to move the contacts back to their original relative locations (i.e., to enforce static friction)
+				// @todo(chaos): we should not be writing back to the constraint here - find a better way to update the friction anchor. See FPBDCollisionConstraint::SetSolverResults
+				// NOTE: low precision for relative coordinates
+				SolverVectorRegister WorldFrictionDelta = VectorZeroFloat();
+				if (SavedManifoldPoint != nullptr)
+				{
+					const VectorRegister SavedShapeContactPos0 = VectorLoadFloat3(&SavedManifoldPoint->ShapeContactPoints[0]);
+					const VectorRegister SavedShapeContactPos1 = VectorLoadFloat3(&SavedManifoldPoint->ShapeContactPoints[1]);
+					const VectorRegister FrictionDelta0 = VectorSubtract(SavedShapeContactPos0, ShapeContactPos0);
+					const VectorRegister FrictionDelta1 = VectorSubtract(SavedShapeContactPos1, ShapeContactPos1);
+					WorldFrictionDelta = MakeVectorRegisterFloatFromDouble(VectorSubtract(VectorQuaternionRotateVector(ShapeWorldRot0, FrictionDelta0), VectorQuaternionRotateVector(ShapeWorldRot1, FrictionDelta1)));
+
+					ManifoldPoint.ShapeAnchorPoints[0] = SavedManifoldPoint->ShapeContactPoints[0];
+					ManifoldPoint.ShapeAnchorPoints[1] = SavedManifoldPoint->ShapeContactPoints[1];
+				}
+				else
+				{
+					const VectorRegister BodyV0 = VectorLoadFloat3(&Body0->V());
+					const VectorRegister BodyW0 = VectorLoadFloat3(&Body0->W());
+					const VectorRegister BodyV1 = VectorLoadFloat3(&Body1->V());
+					const VectorRegister BodyW1 = VectorLoadFloat3(&Body1->W());
+					const VectorRegister ContactVel0 = VectorAdd(BodyV0, VectorCross(BodyW0, RelativeContactPosition0));
+					const VectorRegister ContactVel1 = VectorAdd(BodyV1, VectorCross(BodyW1, RelativeContactPosition1));
+					const SolverVectorRegister ContactVel = MakeVectorRegisterFloatFromDouble(VectorSubtract(ContactVel0, ContactVel1));
+					WorldFrictionDelta = VectorMultiply(ContactVel, VectorSetFloat1(Dt));
+
+					ManifoldPoint.ShapeAnchorPoints[0] = ManifoldPoint.ContactPoint.ShapeContactPoints[0];
+					ManifoldPoint.ShapeAnchorPoints[1] = ManifoldPoint.ContactPoint.ShapeContactPoints[1];
+				}
+
+				// World-space contact tangents. We are treating the normal as the constraint-space Z axis
+				// and the Tangent U and V as the constraint-space X and Y axes respectively
+				SolverVectorRegister WorldContactTangentU = VectorCross(MakeVectorRegisterFloatConstant(0, 1, 0, 0), WorldContactNormal);
+				SolverVectorRegister WorldContactTangentULenSq = VectorDot3(WorldContactTangentU, WorldContactTangentU);
+				SolverVectorRegister WorldContactTangentUCompareMask = VectorCompareGT(WorldContactTangentULenSq, VectorSetFloat1(KINDA_SMALL_NUMBER));
+				if (VectorMaskBits(WorldContactTangentUCompareMask))
+				{
+					WorldContactTangentU = VectorMultiply(WorldContactTangentU, VectorReciprocalSqrt(WorldContactTangentULenSq));
+				}
+				else
+				{
+					WorldContactTangentU = VectorCross(MakeVectorRegisterFloatConstant(1, 0, 0, 0), WorldContactNormal);
+					WorldContactTangentU = VectorNormalize(WorldContactTangentU);
+				}
+				const SolverVectorRegister WorldContactTangentV = VectorCross(WorldContactNormal, WorldContactTangentU);
+
+				// The contact point error we are trying to correct in this solver
+				const SolverVectorRegister WorldContactDelta = MakeVectorRegisterFloatFromDouble(VectorSubtract(WorldContactPoint0, WorldContactPoint1));
+				const FSolverReal WorldContactDeltaNormal = VectorGetComponent(VectorDot3(WorldContactDelta, WorldContactNormal), 0) - TargetPhi;
+				const FSolverReal WorldContactDeltaTangentU = VectorGetComponent(VectorDot3(VectorAdd(WorldContactDelta, WorldFrictionDelta), WorldContactTangentU), 0);
+				const FSolverReal WorldContactDeltaTangentV = VectorGetComponent(VectorDot3(VectorAdd(WorldContactDelta, WorldFrictionDelta), WorldContactTangentV), 0);
+
+				// Copy all the properties into the solver
+				Solver.SetManifoldPoint(
 					ManifoldPointIndex,
+					Dt,
 					Restitution,
-					RestitutionVelocityThreshold);
+					RestitutionVelocityThreshold,
+					RelativeContactPosition0,
+					RelativeContactPosition1,
+					WorldContactNormal,
+					WorldContactTangentU,
+					WorldContactTangentV,
+					WorldContactDeltaNormal,
+					WorldContactDeltaTangentU,
+					WorldContactDeltaTangentV);
 			}
 		}
 
+
 		/**
 		 * @brief Send all solver results to the constraint
 		*/
@@ -345,8 +465,10 @@ namespace Chaos
 
 		UpdatePositionShockPropagation(Dt, It, NumIts, BeginIndex, EndIndex);
 
-		const int32 NumZeroFrictionIterations = Chaos_PBDCollisionSolver_Position_ZeroFrictionIterations;
-		const bool bApplyStaticFriction = (It >= NumZeroFrictionIterations);
+		// Only apply friction for the last few (tunable) iterations
+		const bool bApplyStaticFriction = (It >= (NumIts - Chaos_PBDCollisionSolver_Position_FrictionIterations));
+
+		// Adjust max pushout to attempt to make it iteration count independent
 		const FReal MaxPushOut = (MaxPushOutVelocity > 0) ? (MaxPushOutVelocity * Dt) / FReal(NumIts) : 0;
 
 		// Apply the position correction
diff --git a/Engine/Source/Runtime/Experimental/Chaos/Public/Chaos/Evolution/SolverBody.h b/Engine/Source/Runtime/Experimental/Chaos/Public/Chaos/Evolution/SolverBody.h
--- a/Engine/Source/Runtime/Experimental/Chaos/Public/Chaos/Evolution/SolverBody.h
+++ b/Engine/Source/Runtime/Experimental/Chaos/Public/Chaos/Evolution/SolverBody.h
@@ -23,8 +23,10 @@ namespace Chaos
 	// issues and/or memory bloat.
 #if CHAOS_CONSTRAINTSOLVER_LOWPRECISION
 	using FSolverReal = FRealSingle;
+	using SolverVectorRegister = VectorRegister4Float;
 #else
 	using FSolverReal = FReal;
+	using SolverVectorRegister = VectorRegister4;
 #endif
 	using FSolverVec3 = TVec3<FSolverReal>;
 	using FSolverMatrix33 = TMatrix33<FSolverReal>;
diff --git a/Engine/Source/Runtime/Experimental/Chaos/Public/Chaos/Transform.h b/Engine/Source/Runtime/Experimental/Chaos/Public/Chaos/Transform.h
--- a/Engine/Source/Runtime/Experimental/Chaos/Public/Chaos/Transform.h
+++ b/Engine/Source/Runtime/Experimental/Chaos/Public/Chaos/Transform.h
@@ -170,6 +170,13 @@ namespace Chaos
 			return BaseTransform::ToMatrixNoScale();
 		}
 
+		// For low-level VectorRegister programming
+		const TPersistentVectorRegisterType<FReal>& GetTranslationRegister() const { return Translation; }
+		const TPersistentVectorRegisterType<FReal>& GetRotationRegister() const { return Rotation; }
+		void SetTranslationRegister(TransformVectorRegister InTranslation) { Translation = InTranslation; }
+		void SetRotationRegister(TransformVectorRegister InRotation) { Rotation = InRotation; }
+
+
 		CHAOS_API Chaos::PMatrix<Chaos::FReal, 4, 4> operator*(const Chaos::PMatrix<Chaos::FReal, 4, 4>& Matrix) const;
 		
 		inline TRigidTransform<FReal, 3> operator*(const TRigidTransform<FReal, 3>& Other) const
