Chaos - fix TargetPhi
- fix headless chaos

diff --git a/Engine/Source/Runtime/Experimental/Chaos/Private/Chaos/Collision/SolverCollisionContainer.cpp b/Engine/Source/Runtime/Experimental/Chaos/Private/Chaos/Collision/SolverCollisionContainer.cpp
--- a/Engine/Source/Runtime/Experimental/Chaos/Private/Chaos/Collision/SolverCollisionContainer.cpp
+++ b/Engine/Source/Runtime/Experimental/Chaos/Private/Chaos/Collision/SolverCollisionContainer.cpp
@@ -177,8 +177,8 @@ namespace Chaos
 				const FSolverVec3 WorldContactTangentV = FSolverVec3::CrossProduct(WorldContactNormal, WorldContactTangentU);
 
 				// The contact point error we are trying to correct in this solver
-				const FSolverVec3 WorldContactDelta = FSolverVec3(WorldContactPoint0 - WorldContactPoint1) - TargetPhi;
-				const FSolverReal WorldContactDeltaNormal = FSolverVec3::DotProduct(WorldContactDelta, WorldContactNormal);
+				const FSolverVec3 WorldContactDelta = FSolverVec3(WorldContactPoint0 - WorldContactPoint1);
+				const FSolverReal WorldContactDeltaNormal = FSolverVec3::DotProduct(WorldContactDelta, WorldContactNormal) - TargetPhi;
 				const FSolverReal WorldContactDeltaTangentU = FSolverVec3::DotProduct(WorldContactDelta + WorldFrictionDelta, WorldContactTangentU);
 				const FSolverReal WorldContactDeltaTangentV = FSolverVec3::DotProduct(WorldContactDelta + WorldFrictionDelta, WorldContactTangentV);
 
