Geometry Collection : Convex Generation Defaults
-Added an interface for storing the convex generation attributes to the FGeometryCollection.

diff --git a/Engine/Source/Runtime/Experimental/Chaos/Private/GeometryCollection/GeometryCollection.cpp b/Engine/Source/Runtime/Experimental/Chaos/Private/GeometryCollection/GeometryCollection.cpp
--- a/Engine/Source/Runtime/Experimental/Chaos/Private/GeometryCollection/GeometryCollection.cpp
+++ b/Engine/Source/Runtime/Experimental/Chaos/Private/GeometryCollection/GeometryCollection.cpp
@@ -33,6 +33,7 @@ const FName FGeometryCollection::StatusFlagsAttribute("StatusFlags");
 
 FGeometryCollection::FGeometryCollection()
 	: FTransformCollection()
+	, FGeometryCollectionConvexPropertiesInterface(this)
 {
 	Construct();
 }
@@ -79,6 +80,7 @@ void FGeometryCollection::Construct()
 	// Material Group
 	AddExternalAttribute<FGeometryCollectionSection>("Sections", FGeometryCollection::MaterialGroup, Sections, FacesDependency);
 
+	FGeometryCollectionConvexPropertiesInterface::InitializeInterface();
 }
 
 
@@ -884,6 +886,11 @@ void FGeometryCollection::UpdateBoundingBox()
 
 void FGeometryCollection::Serialize(Chaos::FChaosArchive& Ar)
 {
+	if (Ar.IsCooking())
+	{
+		FGeometryCollectionConvexPropertiesInterface::CleanInterfaceForCook();
+	}
+
 	Super::Serialize(Ar);
 
 	if (Ar.IsLoading())
diff --git a/Engine/Source/Runtime/Experimental/Chaos/Private/GeometryCollection/GeometryCollectionConvexPropertiesInterface.cpp b/Engine/Source/Runtime/Experimental/Chaos/Private/GeometryCollection/GeometryCollectionConvexPropertiesInterface.cpp
new file mode 100644
--- /dev/null
+++ b/Engine/Source/Runtime/Experimental/Chaos/Private/GeometryCollection/GeometryCollectionConvexPropertiesInterface.cpp
@@ -0,0 +1,131 @@
+// Copyright Epic Games, Inc. All Rights Reserved.
+
+#include "GeometryCollection/GeometryCollectionConvexPropertiesInterface.h"
+#include "GeometryCollection/GeometryCollection.h"
+
+const FName FGeometryCollectionConvexPropertiesInterface::ConvexPropertiesGroup = "ConvexProperties";
+const FName FGeometryCollectionConvexPropertiesInterface::ConvexIndexAttribute = "TransformGroupIndex";
+const FName FGeometryCollectionConvexPropertiesInterface::ConvexFractionRemoveAttribute = "FractionRemove";
+const FName FGeometryCollectionConvexPropertiesInterface::ConvexSimplificationThresholdAttribute = "SimplificationThreshold";
+
+FGeometryCollectionConvexPropertiesInterface::FGeometryCollectionConvexPropertiesInterface(FGeometryCollection* InGeometryCollection)
+	: FManagedArrayInterface(InGeometryCollection)
+{}
+
+void 
+FGeometryCollectionConvexPropertiesInterface::InitializeInterface()
+{
+	if (!ManagedCollection->HasGroup(ConvexPropertiesGroup))
+	{
+		ManagedCollection->AddGroup(ConvexPropertiesGroup);
+	}
+
+	if (!ManagedCollection->HasAttribute(ConvexIndexAttribute, ConvexPropertiesGroup))
+	{
+		FManagedArrayCollection::FConstructionParameters AttributeDependency(FTransformCollection::TransformGroup);
+		ManagedCollection->AddAttribute<int32>(ConvexIndexAttribute, ConvexPropertiesGroup, AttributeDependency);
+	}
+
+	if (!ManagedCollection->HasAttribute(ConvexFractionRemoveAttribute, ConvexPropertiesGroup))
+	{
+		ManagedCollection->AddAttribute<float>(ConvexFractionRemoveAttribute, ConvexPropertiesGroup);
+	}
+
+	if (!ManagedCollection->HasAttribute(ConvexSimplificationThresholdAttribute, ConvexPropertiesGroup))
+	{
+		ManagedCollection->AddAttribute<float>(ConvexSimplificationThresholdAttribute, ConvexPropertiesGroup);
+	}
+
+	SetDefaultProperty();
+}
+
+void FGeometryCollectionConvexPropertiesInterface::SetDefaultProperty()
+{
+	FConvexCreationProperties DefaultProperty;
+	TManagedArray<int32>& Index = ManagedCollection->GetAttribute<int32>(ConvexIndexAttribute, ConvexPropertiesGroup);
+	if (Index.Num() == 0 || !Index.Contains(INDEX_NONE))
+	{
+		TManagedArray<float>& FractionRemove = ManagedCollection->GetAttribute<float>(ConvexFractionRemoveAttribute, ConvexPropertiesGroup);
+		TManagedArray<float>& SimplificationThreshold = ManagedCollection->GetAttribute<float>(ConvexSimplificationThresholdAttribute, ConvexPropertiesGroup);
+
+		int32 DefaultAttributeIndex = ManagedCollection->AddElements(1, ConvexPropertiesGroup);
+		Index[DefaultAttributeIndex] = INDEX_NONE;
+		SetConvexProperties(DefaultProperty);
+	}
+}
+
+FGeometryCollectionConvexPropertiesInterface::FConvexCreationProperties 
+FGeometryCollectionConvexPropertiesInterface::GetDefaultProperty() const
+{
+	FConvexCreationProperties DefaultProperty;
+	const TManagedArray<int32>& Index = ManagedCollection->GetAttribute<int32>(ConvexIndexAttribute, ConvexPropertiesGroup);
+	if (Index.Contains(INDEX_NONE))
+	{
+		const TManagedArray<float>& FractionRemove = ManagedCollection->GetAttribute<float>(ConvexFractionRemoveAttribute, ConvexPropertiesGroup);
+		const TManagedArray<float>& SimplificationThreshold = ManagedCollection->GetAttribute<float>(ConvexSimplificationThresholdAttribute, ConvexPropertiesGroup);
+
+		int32 DefaultIndex = Index.Find(INDEX_NONE);
+		if (0 <= DefaultIndex && DefaultIndex < Index.Num())
+		{
+			DefaultProperty.FractionRemove = FractionRemove[DefaultIndex];
+			DefaultProperty.SimplificationThreshold = SimplificationThreshold[DefaultIndex];
+		}
+	}
+	return DefaultProperty;
+}
+
+void 
+FGeometryCollectionConvexPropertiesInterface::CleanInterfaceForCook()
+{
+	RemoveInterfaceAttributes();
+}
+
+void 
+FGeometryCollectionConvexPropertiesInterface::RemoveInterfaceAttributes()
+{
+	ManagedCollection->RemoveGroup(ConvexPropertiesGroup);
+
+}
+
+FGeometryCollectionConvexPropertiesInterface::FConvexCreationProperties
+FGeometryCollectionConvexPropertiesInterface::GetConvexProperties(int TransformGroupIndex) const
+{
+	FConvexCreationProperties ConvexProperty = GetDefaultProperty();
+	const TManagedArray<int32>& Index = ManagedCollection->GetAttribute<int32>(ConvexIndexAttribute, ConvexPropertiesGroup);
+	int32 PropIndex = Index.Find(TransformGroupIndex);
+	if (0 <= PropIndex && PropIndex < Index.Num())
+	{
+		const TManagedArray<float>& FractionRemove = ManagedCollection->GetAttribute<float>(ConvexFractionRemoveAttribute, ConvexPropertiesGroup);
+		const TManagedArray<float>& SimplificationThreshold = ManagedCollection->GetAttribute<float>(ConvexSimplificationThresholdAttribute, ConvexPropertiesGroup);
+
+		ConvexProperty.FractionRemove = FractionRemove[PropIndex];
+		ConvexProperty.SimplificationThreshold = SimplificationThreshold[PropIndex];
+	}
+	return ConvexProperty;
+}
+
+void
+FGeometryCollectionConvexPropertiesInterface::SetConvexProperties(const FConvexCreationProperties& InConvexAttributes, int TransformGroupIndex)
+{
+	FConvexCreationProperties Property;
+	TManagedArray<int32>& Index = ManagedCollection->GetAttribute<int32>(ConvexIndexAttribute, ConvexPropertiesGroup);
+	TManagedArray<float>& FractionRemove = ManagedCollection->GetAttribute<float>(ConvexFractionRemoveAttribute, ConvexPropertiesGroup);
+	TManagedArray<float>& SimplificationThreshold = ManagedCollection->GetAttribute<float>(ConvexSimplificationThresholdAttribute, ConvexPropertiesGroup);
+
+	int32 AttributeIndex = INDEX_NONE;
+	if (!Index.Contains(TransformGroupIndex))
+	{
+		AttributeIndex = ManagedCollection->AddElements(1, ConvexPropertiesGroup);
+	}
+	else
+	{
+		AttributeIndex = Index.Find(TransformGroupIndex);
+	}
+
+	if (0 <= AttributeIndex && AttributeIndex < Index.Num())
+	{
+		Index[AttributeIndex] = TransformGroupIndex;
+		FractionRemove[AttributeIndex] = InConvexAttributes.FractionRemove;
+		SimplificationThreshold[AttributeIndex] = InConvexAttributes.SimplificationThreshold;
+	}
+}
diff --git a/Engine/Source/Runtime/Experimental/Chaos/Public/GeometryCollection/GeometryCollection.h b/Engine/Source/Runtime/Experimental/Chaos/Public/GeometryCollection/GeometryCollection.h
--- a/Engine/Source/Runtime/Experimental/Chaos/Public/GeometryCollection/GeometryCollection.h
+++ b/Engine/Source/Runtime/Experimental/Chaos/Public/GeometryCollection/GeometryCollection.h
@@ -6,10 +6,13 @@
 #include "TransformCollection.h"
 #include "Misc/Crc.h"
 
+#include "GeometryCollection/GeometryCollectionConvexPropertiesInterface.h"
+
 namespace Chaos
 {
 	class FChaosArchive;
 }
+class FGeometryCollectionConvexPropertiesInterface;
 
 namespace GeometryCollectionUV
 {
@@ -25,7 +28,8 @@ namespace GeometryCollectionUV
 *    see : https://wiki.it.epicgames.net/display/~Brice.Criswell/Geometry+Collections
 *
 */
-class CHAOS_API FGeometryCollection : public FTransformCollection
+class CHAOS_API FGeometryCollection : public FTransformCollection, 
+	public FGeometryCollectionConvexPropertiesInterface
 {
 
 public:
diff --git a/Engine/Source/Runtime/Experimental/Chaos/Public/GeometryCollection/GeometryCollectionConvexPropertiesInterface.h b/Engine/Source/Runtime/Experimental/Chaos/Public/GeometryCollection/GeometryCollectionConvexPropertiesInterface.h
new file mode 100644
--- /dev/null
+++ b/Engine/Source/Runtime/Experimental/Chaos/Public/GeometryCollection/GeometryCollectionConvexPropertiesInterface.h
@@ -0,0 +1,46 @@
+// Copyright Epic Games, Inc. All Rights Reserved.
+
+#pragma once
+
+#include "CoreMinimal.h"
+#include "GeometryCollection/ManagedArrayCollection.h"
+
+class FGeometryCollection;
+
+class CHAOS_API FGeometryCollectionConvexPropertiesInterface : public FManagedArrayInterface
+{
+public :
+	typedef FManagedArrayInterface Super;
+	using FManagedArrayInterface::ManagedCollection;
+
+	// Convex Properties Group Name
+	static const FName ConvexPropertiesGroup; 
+	// Attribute
+	static const FName ConvexIndexAttribute;
+	// Attribute
+	static const FName ConvexFractionRemoveAttribute;
+	// Attribute
+	static const FName ConvexSimplificationThresholdAttribute;
+
+
+	struct FConvexCreationProperties {
+		float FractionRemove = 0.5f;
+		float SimplificationThreshold = 10.0f;
+	};
+
+	FGeometryCollectionConvexPropertiesInterface(FGeometryCollection* InGeometryCollection);
+
+	void InitializeInterface() override;
+
+	void CleanInterfaceForCook() override;
+	
+	void RemoveInterfaceAttributes() override;
+
+	FConvexCreationProperties GetConvexProperties(int TransformGroupIndex = INDEX_NONE) const;
+	void SetConvexProperties(const FConvexCreationProperties&, int TransformGroupIndex = INDEX_NONE);
+private:
+
+	void SetDefaultProperty();
+	FConvexCreationProperties GetDefaultProperty() const;
+};
+
diff --git a/Engine/Source/Runtime/Experimental/Chaos/Public/GeometryCollection/ManagedArrayCollection.h b/Engine/Source/Runtime/Experimental/Chaos/Public/GeometryCollection/ManagedArrayCollection.h
--- a/Engine/Source/Runtime/Experimental/Chaos/Public/GeometryCollection/ManagedArrayCollection.h
+++ b/Engine/Source/Runtime/Experimental/Chaos/Public/GeometryCollection/ManagedArrayCollection.h
@@ -480,3 +480,22 @@ protected:
 
 };
 
+
+/*
+*  FManagedArrayInterface
+*/
+class CHAOS_API FManagedArrayInterface
+{
+public:
+	FManagedArrayInterface() : ManagedCollection(nullptr) {}
+	FManagedArrayInterface(FManagedArrayCollection* InManagedArray)
+		: ManagedCollection(InManagedArray) {}
+
+	virtual void InitializeInterface() = 0;
+	virtual void CleanInterfaceForCook() = 0; 
+	virtual void RemoveInterfaceAttributes() = 0;
+
+protected:
+	FManagedArrayCollection* ManagedCollection;
+
+};
