fix crash in trianglemesh.cpp pointordering if there are no points to sort at the end

#rnx

diff --git a/Engine/Source/Runtime/Experimental/Chaos/Private/Chaos/TriangleMesh.cpp b/Engine/Source/Runtime/Experimental/Chaos/Private/Chaos/TriangleMesh.cpp
--- a/Engine/Source/Runtime/Experimental/Chaos/Private/Chaos/TriangleMesh.cpp
+++ b/Engine/Source/Runtime/Experimental/Chaos/Private/Chaos/TriangleMesh.cpp
@@ -1286,8 +1286,11 @@ TArray<int32> FTriangleMesh::GetVertexImportanceOrdering(
 	}
 
 	// Sort remaining non-coincident points by curvature
-	DescendingPredicate<FReal> DescendingCurvaturePred(PointCurvatures);
-	Sort(&PointOrder[MovedPtsOffset], NumPoints - MovedPtsOffset - NumCoincident, DescendingCurvaturePred);
+	if (MovedPtsOffset < NumPoints)
+	{
+		DescendingPredicate<FReal> DescendingCurvaturePred(PointCurvatures);
+		Sort(&PointOrder[MovedPtsOffset], NumPoints - MovedPtsOffset - NumCoincident, DescendingCurvaturePred);
+	}
 
 	return PointOrder;
 }
diff --git a/Engine/Source/Runtime/GeometryCore/Private/Spatial/PriorityOrderPoints.cpp b/Engine/Source/Runtime/GeometryCore/Private/Spatial/PriorityOrderPoints.cpp
--- a/Engine/Source/Runtime/GeometryCore/Private/Spatial/PriorityOrderPoints.cpp
+++ b/Engine/Source/Runtime/GeometryCore/Private/Spatial/PriorityOrderPoints.cpp
@@ -306,7 +306,7 @@ void OrderPoints(TArray<int32>& PointOrder, TArrayView<const TVector<RealType>>
 	}
 
 	// Sort remaining non-coincident points by first metric
-	if (NumWeights > 0 && MovedPtsOffset != EarlyStop)
+	if (NumWeights > 0 && MovedPtsOffset != EarlyStop && MovedPtsOffset < NumPoints)
 	{
 		DescendingPredicate<float> DescendingCurvaturePred(ImportanceWeights[0]);
 		Sort(&PointOrder[MovedPtsOffset], NumPoints - MovedPtsOffset, DescendingCurvaturePred);
