add basic manual manipulation of convex hulls to the fracture mode's convex tool
+ add automatic skipping of poorly-fitting convex hulls on clusters

diff --git a/Engine/Plugins/Experimental/ChaosEditor/Source/FractureEditor/Private/FractureToolConvex.cpp b/Engine/Plugins/Experimental/ChaosEditor/Source/FractureEditor/Private/FractureToolConvex.cpp
--- a/Engine/Plugins/Experimental/ChaosEditor/Source/FractureEditor/Private/FractureToolConvex.cpp
+++ b/Engine/Plugins/Experimental/ChaosEditor/Source/FractureEditor/Private/FractureToolConvex.cpp
@@ -12,6 +12,132 @@
 
 #define LOCTEXT_NAMESPACE "FractureToolConvex"
 
+
+void UFractureConvexSettings::DeleteFromSelected()
+{
+	UFractureToolConvex* ConvexTool = Cast<UFractureToolConvex>(OwnerTool.Get());
+	ConvexTool->DeleteConvexFromSelected();
+}
+
+void UFractureConvexSettings::PromoteChildren()
+{
+	UFractureToolConvex* ConvexTool = Cast<UFractureToolConvex>(OwnerTool.Get());
+	ConvexTool->PromoteChildren();
+}
+
+void UFractureConvexSettings::ClearCustomConvex()
+{
+	UFractureToolConvex* ConvexTool = Cast<UFractureToolConvex>(OwnerTool.Get());
+	ConvexTool->ClearCustomConvex();
+}
+
+void UFractureToolConvex::DeleteConvexFromSelected()
+{
+	TArray<FFractureToolContext> FractureContexts = GetFractureToolContexts();
+
+	for (const FFractureToolContext& FractureContext : FractureContexts)
+	{
+		FGeometryCollection& Collection = *FractureContext.GetGeometryCollection();
+
+		if (!Collection.HasAttribute("ConvexHull", "Convex") ||
+			!Collection.HasAttribute("TransformToConvexIndices", FGeometryCollection::TransformGroup))
+		{
+			continue;
+		}
+
+		TManagedArray<int32>& HasCustomConvex = *FGeometryCollectionConvexUtility::GetCustomConvexFlags(&Collection, true);
+
+		TArray<int32> TransformsToClear;
+		for (int32 TransformIdx : FractureContext.GetSelection())
+		{
+			if (Collection.SimulationType[TransformIdx] == FGeometryCollection::ESimulationTypes::FST_Clustered)
+			{
+				TransformsToClear.Add(TransformIdx);
+				HasCustomConvex[TransformIdx] = 1;
+			}
+		}
+		TransformsToClear.Sort();
+		FGeometryCollectionConvexUtility::RemoveConvexHulls(&Collection, TransformsToClear);
+	}
+
+	FractureContextChanged();
+}
+
+void UFractureToolConvex::PromoteChildren()
+{
+	TArray<FFractureToolContext> FractureContexts = GetFractureToolContexts();
+
+	for (FFractureToolContext& FractureContext : FractureContexts)
+	{
+		UFractureActionTool::FModifyContextScope ModifyScope(this, &FractureContext);
+
+		FGeometryCollection& Collection = *FractureContext.GetGeometryCollection();
+
+		if (!Collection.HasAttribute("ConvexHull", "Convex") ||
+			!Collection.HasAttribute("TransformToConvexIndices", FGeometryCollection::TransformGroup))
+		{
+			continue;
+		}
+
+		TArray<int32> SelectedTransforms = FractureContext.GetSelection();
+		FGeometryCollectionConvexUtility::CopyChildConvexes(&Collection, SelectedTransforms, &Collection, SelectedTransforms, false);
+	}
+
+	FractureContextChanged();
+}
+
+void UFractureToolConvex::ClearCustomConvex()
+{
+	TArray<FFractureToolContext> FractureContexts = GetFractureToolContexts();
+
+	bool bAnyChanged = false;
+	for (FFractureToolContext& FractureContext : FractureContexts)
+	{
+		bool bHasChanged = false;
+
+		TManagedArray<int32>* HasCustomConvex = FGeometryCollectionConvexUtility::GetCustomConvexFlags(FractureContext.GetGeometryCollection().Get(), false);
+		if (!HasCustomConvex)
+		{
+			continue;
+		}
+
+		UFractureActionTool::FModifyContextScope ModifyScope(this, &FractureContext);
+
+		FGeometryCollection& Collection = *FractureContext.GetGeometryCollection();
+		for (int32 TransformIdx : FractureContext.GetSelection())
+		{
+			if ((*HasCustomConvex)[TransformIdx])
+			{
+				bAnyChanged = bHasChanged = true;
+				(*HasCustomConvex)[TransformIdx] = false;
+			}
+		}
+
+		if (bHasChanged)
+		{
+			bool bAllFalse = true;
+			for (int32 TransformIdx = 0; bAllFalse && TransformIdx < HasCustomConvex->Num(); TransformIdx++)
+			{
+				if ((*HasCustomConvex)[TransformIdx])
+				{
+					bAllFalse = false;
+				}
+			}
+			if (bAllFalse)
+			{
+				Collection.RemoveAttribute("HasCustomConvex", FTransformCollection::TransformGroup);
+			}
+
+			AutoComputeConvex(FractureContext);
+		}
+	}
+
+	if (bAnyChanged)
+	{
+		FractureContextChanged();
+	}
+}
+
 UFractureToolConvex::UFractureToolConvex(const FObjectInitializer& ObjInit)
 	: Super(ObjInit)
 {
@@ -73,6 +199,8 @@ void UFractureToolConvex::FractureContextChanged()
 			continue;
 		}
 
+		TManagedArray<int32>* HasCustomConvex = FGeometryCollectionConvexUtility::GetCustomConvexFlags(&Collection, false);
+
 		int32 CollectionIdx = VisualizedCollections.Add(FractureContext.GetGeometryCollectionComponent());
 
 		FTransform OuterTransform = FractureContext.GetTransform();
@@ -80,6 +208,7 @@ void UFractureToolConvex::FractureContextChanged()
 		{
 			FTransform InnerTransform = GeometryCollectionAlgo::GlobalMatrix(Collection.Transform, Collection.Parent, TransformIdx);
 			FTransform CombinedTransform = InnerTransform * OuterTransform;
+			bool bIsCustom = HasCustomConvex ? bool((*HasCustomConvex)[TransformIdx]) : false;
 
 			TManagedArray<TSet<int32>>& TransformToConvexIndices = Collection.GetAttribute<TSet<int32>>("TransformToConvexIndices", FTransformCollection::TransformGroup);
 			TManagedArray<TUniquePtr<Chaos::FConvex>>& ConvexHull = Collection.GetAttribute<TUniquePtr<Chaos::FConvex>>("ConvexHull", "Convex");
@@ -102,7 +231,7 @@ void UFractureToolConvex::FractureContextChanged()
 					{
 						int32 V0 = HullPtsStart + HullData.GetPlaneVertex(PlaneIdx, PlaneVertexIdx);
 						int32 V1 = HullPtsStart + HullData.GetPlaneVertex(PlaneIdx, (PlaneVertexIdx + 1) % NumPlaneVerts);
-						HullEdges.Add(TPair<int32, int32>(V0, V1));
+						HullEdges.Add({ V0, V1, bIsCustom });
 					}
 				}
 			}
@@ -114,10 +243,10 @@ void UFractureToolConvex::Render(const FSceneView* View, FViewport* Viewport, FP
 {
 	EnumerateVisualizationMapping(EdgesMappings, HullEdges.Num(), [&](int32 Idx, FVector ExplodedVector)
 	{
-		TPair<int32, int32> Edge = HullEdges[Idx];
-		FVector P1 = HullPoints[Edge.Key] + ExplodedVector;
-		FVector P2 = HullPoints[Edge.Value] + ExplodedVector;
-		PDI->DrawLine(P1, P2, FLinearColor::Green, SDPG_Foreground, 0.0f, 0.001f);
+		const FEdgeVisInfo& Edge = HullEdges[Idx];
+		FVector P1 = HullPoints[Edge.A] + ExplodedVector;
+		FVector P2 = HullPoints[Edge.B] + ExplodedVector;
+		PDI->DrawLine(P1, P2, Edge.bIsCustom ? FLinearColor::Red : FLinearColor::Green, SDPG_Foreground, 0.0f, 0.001f);
 	});
 }
 
@@ -145,16 +274,21 @@ TArray<FFractureToolContext> UFractureToolConvex::GetFractureToolContexts() cons
 	return Contexts;
 }
 
-
-int32 UFractureToolConvex::ExecuteFracture(const FFractureToolContext& FractureContext)
+void UFractureToolConvex::AutoComputeConvex(const FFractureToolContext& FractureContext)
 {
 	if (FractureContext.GetGeometryCollection().IsValid())
 	{
 		FGeometryCollection& Collection = *FractureContext.GetGeometryCollection();
 		FGeometryCollectionProximityUtility ProximityUtility(&Collection);
 		ProximityUtility.UpdateProximity();
-		FGeometryCollectionConvexUtility::CreateNonOverlappingConvexHullData(&Collection, ConvexSettings->FractionAllowRemove, ConvexSettings->SimplificationDistanceThreshold);
+		FGeometryCollectionConvexUtility::CreateNonOverlappingConvexHullData(&Collection, ConvexSettings->FractionAllowRemove, ConvexSettings->SimplificationDistanceThreshold, ConvexSettings->CanExceedFraction);
 	}
+}
+
+
+int32 UFractureToolConvex::ExecuteFracture(const FFractureToolContext& FractureContext)
+{
+	AutoComputeConvex(FractureContext);
 
 	return INDEX_NONE;
 }
diff --git a/Engine/Plugins/Experimental/ChaosEditor/Source/FractureEditor/Private/FractureToolConvex.h b/Engine/Plugins/Experimental/ChaosEditor/Source/FractureEditor/Private/FractureToolConvex.h
--- a/Engine/Plugins/Experimental/ChaosEditor/Source/FractureEditor/Private/FractureToolConvex.h
+++ b/Engine/Plugins/Experimental/ChaosEditor/Source/FractureEditor/Private/FractureToolConvex.h
@@ -21,12 +21,28 @@ public:
 	{}
 
 	/** Fraction of the convex hulls for a transform that we can remove before instead using the hulls of the children */
-	UPROPERTY(EditAnywhere, Category = MapSettings, meta = (ClampMin = ".01", ClampMax = "1"))
+	UPROPERTY(EditAnywhere, Category = Automatic, meta = (DisplayName = "Can Remove Fraction", ClampMin = ".01", ClampMax = "1"))
 	double FractionAllowRemove = .5;
 
+	/** Fraction (of geometry volume) by which a cluster's convex hull volume can exceed the actual geometry volume before instead using the hulls of the children.  0 means the convex volume cannot exceed the geometry volume; 1 means the convex volume is allowed to be 100% larger (2x) the geometry volume. */
+	UPROPERTY(EditAnywhere, Category = Automatic, meta = (ClampMin = "0"))
+	double CanExceedFraction = .5;
+
 	/** We simplify the convex shape to keep points spaced at least this far apart (except to keep the hull from collapsing to zero volume) */
-	UPROPERTY(EditAnywhere, Category = MapSettings, meta = (ClampMin = "0"))
+	UPROPERTY(EditAnywhere, Category = Automatic, meta = (ClampMin = "0"))
 	double SimplificationDistanceThreshold = 10.0;
+
+	/** Delete convex hulls from selected clusters.  Does not affect hulls on leaves. */
+	UFUNCTION(CallInEditor, Category = Custom, meta = (DisplayName = "Delete From Selected"))
+	void DeleteFromSelected();
+
+	/** Promote (and save) child convex hulls on to the selected bone(s) */
+	UFUNCTION(CallInEditor, Category = Custom, meta = (DisplayName = "Promote Children"))
+	void PromoteChildren();
+
+	/** Clear any manual adjustments to convex hulls on the selected bones */
+	UFUNCTION(CallInEditor, Category = Custom, meta = (DisplayName = "Clear Custom Convex"))
+	void ClearCustomConvex();
 };
 
 
@@ -62,6 +78,11 @@ public:
 
 	virtual TArray<FFractureToolContext> GetFractureToolContexts() const override;
 
+	void DeleteConvexFromSelected();
+	void PromoteChildren();
+	void ClearCustomConvex();
+	void AutoComputeConvex(const FFractureToolContext& FractureContext);
+
 
 protected:
 	UPROPERTY(EditAnywhere, Category = Slicing)
@@ -75,8 +96,14 @@ protected:
 		EdgesMappings.Empty();
 	}
 
+	struct FEdgeVisInfo
+	{
+		int32 A, B;
+		bool bIsCustom; // allow for different coloring of hulls that have been manually set vs auto-generated ones
+	};
+
 	TArray<FVector> HullPoints;
-	TArray<TPair<int32, int32>> HullEdges;
+	TArray<FEdgeVisInfo> HullEdges;
 	FVisualizationMappings EdgesMappings;
 
 };
diff --git a/Engine/Plugins/Experimental/ChaosEditor/Source/FractureEditor/Public/FractureTool.cpp b/Engine/Plugins/Experimental/ChaosEditor/Source/FractureEditor/Public/FractureTool.cpp
--- a/Engine/Plugins/Experimental/ChaosEditor/Source/FractureEditor/Public/FractureTool.cpp
+++ b/Engine/Plugins/Experimental/ChaosEditor/Source/FractureEditor/Public/FractureTool.cpp
@@ -7,6 +7,8 @@
 #include "LevelEditor.h"
 #include "LevelEditorViewport.h"
 #include "SLevelViewport.h"
+#include "EditorModeManager.h"
+#include "FractureEditorMode.h"
 
 #include "FractureToolContext.h"
 #include "FractureSelectionTools.h"
@@ -47,6 +49,30 @@ const TSharedPtr<FUICommandInfo>& UFractureActionTool::GetUICommandInfo() const
 	return UICommandInfo;
 }
 
+UFractureActionTool::FModifyContextScope::FModifyContextScope(UFractureActionTool* ActionTool, FFractureToolContext* FractureContext) : ActionTool(ActionTool), FractureContext(FractureContext)
+{
+	check(FractureContext);
+	check(ActionTool);
+	FractureContext->GetFracturedGeometryCollection()->Modify();
+	FractureContext->GetGeometryCollectionComponent()->Modify();
+}
+
+UFractureActionTool::FModifyContextScope::~FModifyContextScope()
+{
+	UFractureEditorMode* FractureMode = Cast<UFractureEditorMode>(GLevelEditorModeTools().GetActiveScriptableMode(UFractureEditorMode::EM_FractureEditorModeId));
+	if (!FractureMode)
+	{
+		return;
+	}
+	TSharedPtr<FModeToolkit> ModeToolkit = FractureMode->GetToolkit().Pin();
+	if (ModeToolkit.IsValid())
+	{
+		FFractureEditorModeToolkit* Toolkit = (FFractureEditorModeToolkit*)ModeToolkit.Get();
+		ActionTool->Refresh(*FractureContext, Toolkit);
+	}
+}
+
+
 bool UFractureActionTool::CanExecute() const
 {
 	return IsGeometryCollectionSelected();
@@ -261,9 +287,10 @@ void UFractureModalTool::NotifyOfPropertyChangeByTool(UFractureToolSettings* Pro
 
 void UFractureModalTool::Execute(TWeakPtr<FFractureEditorModeToolkit> InToolkit)
 {
-	if (InToolkit.IsValid())
+	TSharedPtr<FFractureEditorModeToolkit> ModeToolkit = InToolkit.Pin();
+	if (ModeToolkit.IsValid())
 	{
-		FFractureEditorModeToolkit* Toolkit = InToolkit.Pin().Get();
+		FFractureEditorModeToolkit* Toolkit = ModeToolkit.Get();
 
 		TArray<FFractureToolContext> FractureContexts = GetFractureToolContexts();
 
diff --git a/Engine/Plugins/Experimental/ChaosEditor/Source/FractureEditor/Public/FractureTool.h b/Engine/Plugins/Experimental/ChaosEditor/Source/FractureEditor/Public/FractureTool.h
--- a/Engine/Plugins/Experimental/ChaosEditor/Source/FractureEditor/Public/FractureTool.h
+++ b/Engine/Plugins/Experimental/ChaosEditor/Source/FractureEditor/Public/FractureTool.h
@@ -37,6 +37,7 @@ public:
 	TObjectPtr<UFractureModalTool> OwnerTool;
 };
 
+class FFractureToolContext;
 
 /** Tools derived from this class should require parameter inputs from the user, only the bone selection. */
 UCLASS(Abstract)
@@ -63,6 +64,17 @@ public:
 	virtual void RegisterUICommand(FFractureEditorCommands* BindingContext) {}
 
 	virtual TArray<FFractureToolContext> GetFractureToolContexts() const;
+
+	// Scope sets up modification and updates UI after
+	struct FModifyContextScope
+	{
+		UFractureActionTool* ActionTool;
+		FFractureToolContext* FractureContext;
+
+		FModifyContextScope(UFractureActionTool* ActionTool, FFractureToolContext* FractureContext);
+
+		~FModifyContextScope();
+	};
 	
 protected:
 	static bool IsStaticMeshSelected();
diff --git a/Engine/Source/Runtime/Experimental/Chaos/Private/GeometryCollection/GeometryCollectionConvexUtility.cpp b/Engine/Source/Runtime/Experimental/Chaos/Private/GeometryCollection/GeometryCollectionConvexUtility.cpp
--- a/Engine/Source/Runtime/Experimental/Chaos/Private/GeometryCollection/GeometryCollectionConvexUtility.cpp
+++ b/Engine/Source/Runtime/Experimental/Chaos/Private/GeometryCollection/GeometryCollectionConvexUtility.cpp
@@ -9,6 +9,23 @@
 #include "Templates/Sorting.h"
 #include "Spatial/PointHashGrid3.h"
 
+TOptional<FGeometryCollectionConvexUtility::FGeometryCollectionConvexData> FGeometryCollectionConvexUtility::GetConvexHullDataIfPresent(FGeometryCollection* GeometryCollection)
+{
+	check(GeometryCollection);
+
+	if (!GeometryCollection->HasAttribute("TransformToConvexIndices", FTransformCollection::TransformGroup) ||
+		!GeometryCollection->HasAttribute("ConvexHull", "Convex"))
+	{
+		return TOptional<FGeometryCollectionConvexUtility::FGeometryCollectionConvexData>();
+	}
+
+	FGeometryCollectionConvexUtility::FGeometryCollectionConvexData ConvexData{
+		GeometryCollection->GetAttribute<TSet<int32>>("TransformToConvexIndices", FTransformCollection::TransformGroup),
+		GeometryCollection->GetAttribute<TUniquePtr<Chaos::FConvex>>("ConvexHull", "Convex")
+	};
+	return TOptional<FGeometryCollectionConvexUtility::FGeometryCollectionConvexData>(ConvexData);
+}
+
 
 FGeometryCollectionConvexUtility::FGeometryCollectionConvexData FGeometryCollectionConvexUtility::GetValidConvexHullData(FGeometryCollection* GeometryCollection)
 {
@@ -277,14 +294,17 @@ bool SplitHull(const Chaos::FConvex& HullIn, FChaosPlane Plane, bool KeepSide, T
 void CreateNonoverlappingConvexHulls(
 	TArray<TUniquePtr<Chaos::FConvex>>& Convexes,
 	TArray<TSet<int32>>& TransformToConvexIndices,
+	TFunctionRef<bool(int32)> HasCustomConvexFn,
 	const TManagedArray<int32>& SimulationType,
 	int32 LeafType,
 	int32 SkipType,
 	const TManagedArray<int32>& Parents,
 	const TManagedArray<TSet<int32>>* GeoProximity,
 	const TManagedArray<int32>& GeometryToTransformIndex,
+	const TManagedArray<float>* Volume,
 	double FracAllowRemove,
-	double SimplificationDistanceThreshold
+	double SimplificationDistanceThreshold,
+	double CanExceedFraction
 )
 {
 	int32 NumBones = TransformToConvexIndices.Num();
@@ -335,6 +355,21 @@ void CreateNonoverlappingConvexHulls(
 		return GJKIntersection(*Convexes[ConvexA], *Convexes[ConvexB], IdentityTransform);
 	};
 
+	auto IsBoneColliding = [&TransformToConvexIndices, &Convexes, &IsColliding](int32 BoneA, int32 BoneB)
+	{
+		for (int32 ConvexA : TransformToConvexIndices[BoneA])
+		{
+			for (int32 ConvexB : TransformToConvexIndices[BoneB])
+			{
+				if (IsColliding(ConvexA, ConvexB))
+				{
+					return true;
+				}
+			}
+		}
+		return false;
+	};
+
 	auto GetConvexSpan = [](const Chaos::FConvex& Convex, const Chaos::FVec3& Center, const Chaos::FVec3& Normal) -> Chaos::FVec2
 	{
 		int32 NumVertices = Convex.NumVertices();
@@ -521,13 +556,14 @@ void CreateNonoverlappingConvexHulls(
 		Depths[HullIdx] = Depth;
 		MaxDepth = FMath::Max(Depth, MaxDepth);
 
-		if (TransformToConvexIndices[HullIdx].Num() > 0)
+		if (!HasCustomConvexFn(HullIdx) && TransformToConvexIndices[HullIdx].Num() > 0)
 		{
 			const TSet<int32>& Neighbors = LeafProximity[HullIdx];
 			for (int32 NbrIdx : Neighbors)
 			{
-				if (NbrIdx < HullIdx && TransformToConvexIndices[NbrIdx].Num() > 0)
+				if (!HasCustomConvexFn(NbrIdx) && NbrIdx < HullIdx && TransformToConvexIndices[NbrIdx].Num() > 0)
 				{
+					// TODO: consider a one-sided cut if one of the bones has custom convexes
 					FixCollisionWithCut(OnlyConvex(HullIdx), OnlyConvex(NbrIdx));
 				}
 			}
@@ -552,7 +588,25 @@ void CreateNonoverlappingConvexHulls(
 		}
 	};
 
-	auto FixLeafCollisions = [&AddLeaves, &FixCollisionWithCut, &OnlyConvex](int32 BoneA, int32 BoneB)
+	// Fill OutChildren with the shallowest descendents of Bone that have convex(es) (i.e., the direct children if they all have convexes, otherwise descend further to grandchildren searching for convex hulls)
+	auto AddDescendentsWithHulls = [&Children, &TransformToConvexIndices, &SimulationType, LeafType](int32 Bone, TArray<int32>& OutChildren)
+	{
+		TArray<int32> ToExpand = Children[Bone];
+		while (ToExpand.Num() > 0)
+		{
+			int32 ToProcess = ToExpand.Pop(false);
+			if (TransformToConvexIndices[ToProcess].Num() > 0)
+			{
+				OutChildren.Add(ToProcess);
+			}
+			else if (Children[ToProcess].Num() > 0)
+			{
+				ToExpand.Append(Children[ToProcess]);
+			}
+		}
+	};
+
+	auto FixLeafCollisions = [&AddLeaves, &FixCollisionWithCut, &OnlyConvex, &HasCustomConvexFn](int32 BoneA, int32 BoneB)
 	{
 		TArray<int32> LeavesA, LeavesB;
 		AddLeaves(BoneA, LeavesA);
@@ -560,8 +614,20 @@ void CreateNonoverlappingConvexHulls(
 		bool bAnyCollides = false;
 		for (int32 LeafA : LeavesA)
 		{
+			if (HasCustomConvexFn(LeafA))
+			{
+				continue;
+			}
 			for (int32 LeafB : LeavesB)
 			{
+				if (HasCustomConvexFn(LeafB))
+				{
+					continue;
+				}
+				// TODO: Instead of skipping all custom convex, if only one of the leaves is custom, we 
+				// could iterate through the custom leaf convexes and clip the non-custom convex.
+				// But need to decide a fallback if this would delete the whole non-custom convex.
+				// For now, if you have custom convex leaves we just assume you they're set up to not explode the sim.
 				bool bCollides = FixCollisionWithCut(OnlyConvex(LeafA), OnlyConvex(LeafB));
 				bAnyCollides |= bCollides;
 			}
@@ -575,6 +641,8 @@ void CreateNonoverlappingConvexHulls(
 	TArray<TSet<int32>> ClusterProximity;
 	ClusterProximity.SetNum(NumBones);
 	int ProcessDepth = MaxDepth;
+	TArray<Chaos::FConvex::FVec3Type> JoinedHullPts;
+	TArray<int32> ChildrenWithHulls;
 	while (--ProcessDepth >= 0)
 	{
 		TSet<int32> ToProcess;
@@ -582,10 +650,12 @@ void CreateNonoverlappingConvexHulls(
 		{
 			if (Depths[Bone] == ProcessDepth)
 			{
-				if (TransformToConvexIndices[Bone].Num() == 0)
+				if (TransformToConvexIndices[Bone].Num() == 0 && !HasCustomConvexFn(Bone))
 				{
-					TArray<Chaos::FConvex::FVec3Type> JoinedHullPts;
-					for (int32 Child : Children[Bone])
+					JoinedHullPts.Reset();
+					ChildrenWithHulls.Reset();
+					AddDescendentsWithHulls(Bone, ChildrenWithHulls);
+					for (int32 Child : ChildrenWithHulls)
 					{
 						for (int32 ConvexIdx : TransformToConvexIndices[Child])
 						{
@@ -595,9 +665,22 @@ void CreateNonoverlappingConvexHulls(
 					if (JoinedHullPts.Num() > 0)
 					{
 						FilterHullPoints(JoinedHullPts, SimplificationDistanceThreshold);
-						int32 ConvexIdx = Convexes.Add(MakeUnique<Chaos::FConvex>(JoinedHullPts, KINDA_SMALL_NUMBER));
-						TransformToConvexIndices[Bone].Add(ConvexIdx);
-						ToProcess.Add(Bone);
+						TUniquePtr<Chaos::FConvex> Hull = MakeUnique<Chaos::FConvex>(JoinedHullPts, KINDA_SMALL_NUMBER);
+						bool bIsTooBig = false;
+						if (Volume)
+						{
+							Chaos::FReal HullVolume = Hull->GetVolume();
+							if (HullVolume > Chaos::FReal((*Volume)[Bone]) * (1.0 + CanExceedFraction))
+							{
+								bIsTooBig = true;
+							}
+						}
+						if (!bIsTooBig)
+						{
+							int32 ConvexIdx = Convexes.Add(MoveTemp(Hull));
+							TransformToConvexIndices[Bone].Add(ConvexIdx);
+							ToProcess.Add(Bone);
+						}
 					}
 				}
 				else
@@ -607,8 +690,6 @@ void CreateNonoverlappingConvexHulls(
 			}
 		}
 
-		ensure(ToProcess.Num() > 0);
-
 		// We don't have proximity at the cluster level, so reconstruct from n^2 collision tests for now
 		for (int32 BoneA : ToProcess)
 		{
@@ -616,9 +697,9 @@ void CreateNonoverlappingConvexHulls(
 			{
 				if (BoneB < BoneA)
 				{
-					if (IsColliding(OnlyConvex(BoneA), OnlyConvex(BoneB)))
+					if (IsBoneColliding(BoneA, BoneB))
 					{
-						bool bLeavesCollide = FixLeafCollisions(BoneA, BoneB);
+						FixLeafCollisions(BoneA, BoneB);
 						int32 Bones[2]{ BoneA, BoneB };
 						for (int32 BoneIdx = 0; BoneIdx < 2; BoneIdx++)
 						{
@@ -630,7 +711,7 @@ void CreateNonoverlappingConvexHulls(
 							while (TraverseBones.Num() > 0)
 							{
 								int32 ToProc = TraverseBones.Pop(false);
-								if (IsColliding(OnlyConvex(OtherBone), OnlyConvex(ToProc)))
+								if (IsBoneColliding(OtherBone, ToProc))
 								{
 									ClusterProximity[OtherBone].Add(ToProc);
 									ClusterProximity[ToProc].Add(OtherBone);
@@ -658,10 +739,16 @@ void CreateNonoverlappingConvexHulls(
 		TSet<int32> ToProcess;
 		for (int32 Bone = 0; Bone < NumBones; Bone++)
 		{
-			if (Depths[Bone] == ProcessDepth && Children[Bone].Num() > 0 && !SkipBone(Bone))
+			bool bCustom = HasCustomConvexFn(Bone);
+			if (!bCustom && Depths[Bone] == ProcessDepth && // if we need an automatic hull and are at the current depth
+				Children[Bone].Num() > 0 && !SkipBone(Bone) && // and we have children && aren't embedded geo
+				TransformToConvexIndices[Bone].Num() == 1 // and hull wasn't already ruled out by CanExceedFraction
+				)
 			{
-				TArray<Chaos::FConvex::FVec3Type> JoinedHullPts;
-				for (int32 Child : Children[Bone])
+				JoinedHullPts.Reset();
+				ChildrenWithHulls.Reset();
+				AddDescendentsWithHulls(Bone, ChildrenWithHulls);
+				for (int32 Child : ChildrenWithHulls)
 				{
 					for (int32 ConvexIdx : TransformToConvexIndices[Child])
 					{
@@ -678,7 +765,8 @@ void CreateNonoverlappingConvexHulls(
 		}
 	}
 
-	auto CutIfOk = [&Convexes, &OnlyConvex, &NonLeafVolumes, &FindCutPlane, &FracAllowRemove, &SimplificationDistanceThreshold](bool bOneSidedCut, int32 ConvexA, int32 ConvexB) -> bool
+	// if bOneSidedCut, then only ConvexA is cut; ConvexB is left unchanged
+	auto CutIfOk = [&Convexes, &NonLeafVolumes, &FindCutPlane, &FracAllowRemove, &SimplificationDistanceThreshold](bool bOneSidedCut, int32 ConvexA, int32 ConvexB) -> bool
 	{
 		Chaos::FReal Depth;
 		Chaos::FVec3 CloseA, CloseB, Normal;
@@ -762,7 +850,8 @@ void CreateNonoverlappingConvexHulls(
 
 		for (int32 Bone : ToProcess)
 		{
-			if (WasNotOk.Contains(Bone) || TransformToConvexIndices[Bone].Num() == 0)
+			bool bCustom = HasCustomConvexFn(Bone);
+			if (bCustom || TransformToConvexIndices[Bone].Num() == 0 || WasNotOk.Contains(Bone))
 			{
 				continue;
 			}
@@ -772,12 +861,13 @@ void CreateNonoverlappingConvexHulls(
 				{
 					continue;
 				}
+				bool bNbrCustom = HasCustomConvexFn(Nbr);
 				bool bAllOk = true;
 				for (int32 ConvexBone : TransformToConvexIndices[Bone])
 				{
 					for (int32 ConvexNbr : TransformToConvexIndices[Nbr])
 					{
-						bool bOneSidedCut = Depths[Bone] != Depths[Nbr] || Children[Nbr].Num() == 0;
+						bool bOneSidedCut = Depths[Bone] != Depths[Nbr] || Children[Nbr].Num() == 0 || bNbrCustom;
 						// if the neighbor is less deep and not a leaf, skip processing this to favor processing the neighbor instead
 						if (Depths[Bone] > Depths[Nbr] && Children[Nbr].Num() > 0)
 						{
@@ -791,9 +881,16 @@ void CreateNonoverlappingConvexHulls(
 						if (!bWasOk)
 						{
 							bAllOk = false;
-							auto ResetBoneHulls = [&WasNotOk, &TransformToConvexIndices](int32 ToReset)
+							auto ResetBoneHulls = [&WasNotOk, &TransformToConvexIndices, &Convexes](int32 ToReset)
 							{
 								WasNotOk.Add(ToReset);
+								for (int32 ConvexIdx : TransformToConvexIndices[ToReset])
+								{
+									// we just leave these hulls in as null, without any references here
+									// and come through and clear them later after everything is set up in the geometry collection
+									// (because that has the built-in machinery to update the TransformToConvexIndices indices accordingly)
+									Convexes[ConvexIdx].Reset();
+								}
 								TransformToConvexIndices[ToReset].Reset();
 							};
 							ResetBoneHulls(Bone);
@@ -819,8 +916,9 @@ void CreateNonoverlappingConvexHulls(
 
 /// Helper to get convex hulls from a geometry collection in the format required by CreateNonoverlappingConvexHulls
 void HullsFromGeometry(
-	const FGeometryCollection& Geometry,
+	FGeometryCollection& Geometry,
 	TArray<FTransform>& GlobalTransformArray,
+	TFunctionRef<bool(int32)> HasCustomConvexFn,
 	TArray<TUniquePtr<Chaos::FConvex>>& Convexes,
 	TArray<TSet<int32>>& TransformToConvexIndices,
 	const TManagedArray<int32>& SimulationType,
@@ -829,6 +927,8 @@ void HullsFromGeometry(
 )
 {
 	TArray<FVector> GlobalVertices;
+
+	TOptional<FGeometryCollectionConvexUtility::FGeometryCollectionConvexData> OrigConvexData = FGeometryCollectionConvexUtility::GetConvexHullDataIfPresent(&Geometry);
 	
 	GlobalVertices.SetNum(Geometry.Vertex.Num());
 	for (int32 Idx = 0; Idx < GlobalVertices.Num(); Idx++)
@@ -838,19 +938,32 @@ void HullsFromGeometry(
 
 	int32 NumBones = Geometry.TransformToGeometryIndex.Num();
 	TransformToConvexIndices.SetNum(NumBones);
+	TArray<Chaos::FConvex::FVec3Type> HullPts;
 	for (int32 Idx = 0; Idx < NumBones; Idx++)
 	{
-		if (SimulationType[Idx] != RigidType)
+		int32 GeomIdx = Geometry.TransformToGeometryIndex[Idx];
+		if (OrigConvexData.IsSet() && HasCustomConvexFn(Idx))
 		{
-			continue;
+			// custom convex hulls are kept, but transformed to global space + copied to the new data structure
+			const FTransform& Transform = GlobalTransformArray[Idx];
+			for (int32 OrigConvexIdx : OrigConvexData->TransformToConvexIndices[Idx])
+			{
+				HullPts.Reset();
+				for (const Chaos::FConvex::FVec3Type& P : OrigConvexData->ConvexHull[OrigConvexIdx]->GetVertices())
+				{
+					HullPts.Add(Transform.TransformPosition(P));
+				}
+				// Do not simplify hulls when we're just trying to transform them
+				int32 NewConvexIdx = Convexes.Add(MakeUnique<Chaos::FConvex>(HullPts, KINDA_SMALL_NUMBER));
+				TransformToConvexIndices[Idx].Add(NewConvexIdx);
+			}
 		}
-		int32 GeomIdx = Geometry.TransformToGeometryIndex[Idx];
-		if (GeomIdx != INDEX_NONE)
+		else if (SimulationType[Idx] == RigidType && GeomIdx != INDEX_NONE)
 		{
 			int32 VStart = Geometry.VertexStart[GeomIdx];
 			int32 VCount = Geometry.VertexCount[GeomIdx];
 			int32 VEnd = VStart + VCount;
-			TArray<Chaos::FConvex::FVec3Type> HullPts;
+			HullPts.Reset();
 			HullPts.Reserve(VCount);
 			for (int32 VIdx = VStart; VIdx < VEnd; VIdx++)
 			{
@@ -906,23 +1019,59 @@ void TransformHullsToLocal(
 	}
 }
 
+// copy hulls from InBone to OutBone
+// Store the results in a temporary array, to be added in bulk later
+bool CopyHulls(
+	const TArray<FTransform>& InGlobalTransformArray,
+	const TManagedArray<TUniquePtr<Chaos::FConvex>>& InConvexes,
+	const TManagedArray<TSet<int32>>& InTransformToConvexIndices,
+	int32 InBone,
+	const TArray<FTransform>& OutGlobalTransformArray,
+	TArray<TUniquePtr<Chaos::FConvex>>& OutConvexes,
+	TArray<TSet<int32>>& OutTransformToConvexIndices,
+	int32 OutBone
+)
+{
+	const FTransform& InTransform = InGlobalTransformArray[InBone];
+	const FTransform& OutTransform = InGlobalTransformArray[OutBone];
+	TArray<Chaos::FConvex::FVec3Type> HullPts;
+	for (int32 ConvexIdx : InTransformToConvexIndices[InBone])
+	{
+		HullPts.Reset();
+		for (const Chaos::FConvex::FVec3Type& P : InConvexes[ConvexIdx]->GetVertices())
+		{
+			HullPts.Add(OutTransform.InverseTransformPosition(InTransform.TransformPosition(P)));
+		}
+		// Do not simplify hulls when we're just trying to transform them
+		int32 OutIdx = OutConvexes.Add(MakeUnique<Chaos::FConvex>(HullPts, KINDA_SMALL_NUMBER));
+		OutTransformToConvexIndices[OutBone].Add(OutIdx);
+	}
 
+	return true;
 }
 
 
-FGeometryCollectionConvexUtility::FGeometryCollectionConvexData FGeometryCollectionConvexUtility::CreateNonOverlappingConvexHullData(FGeometryCollection* GeometryCollection, double FracAllowRemove, double SimplificationDistanceThreshold)
+}
+
+FGeometryCollectionConvexUtility::FGeometryCollectionConvexData FGeometryCollectionConvexUtility::CreateNonOverlappingConvexHullData(FGeometryCollection* GeometryCollection, double FracAllowRemove, double SimplificationDistanceThreshold, double CanExceedFraction)
 {
 	check(GeometryCollection);
 
+	TManagedArray<int32>* CustomConvexFlags = GetCustomConvexFlags(GeometryCollection, false);
+	auto ConvexFlagsAlwaysFalse = [](int32) -> bool { return false; };
+	auto ConvexFlagsFromArray = [&CustomConvexFlags](int32 TransformIdx) -> bool { return (bool)(*CustomConvexFlags)[TransformIdx]; };
+	TFunctionRef<bool(int32)> HasCustomConvexFn = (CustomConvexFlags != nullptr) ? (TFunctionRef<bool(int32)>)ConvexFlagsFromArray : (TFunctionRef<bool(int32)>)ConvexFlagsAlwaysFalse;
+
 	bool bHasProximity = GeometryCollection->HasAttribute("Proximity", FGeometryCollection::GeometryGroup);
 	const TManagedArray<TSet<int32>>* GCProximity = GeometryCollection->FindAttribute<TSet<int32>>("Proximity", FGeometryCollection::GeometryGroup);
+	const TManagedArray<float>* Volume = GeometryCollection->FindAttribute<float>("Volume", FGeometryCollection::TransformGroup);
 	TArray<TUniquePtr<Chaos::FConvex>> Convexes;
 	TArray<TSet<int32>> TransformToConvexIndexArr;
 	TArray<FTransform> GlobalTransformArray;
 	GeometryCollectionAlgo::GlobalMatrices(GeometryCollection->Transform, GeometryCollection->Parent, GlobalTransformArray);
-	HullsFromGeometry(*GeometryCollection, GlobalTransformArray, Convexes, TransformToConvexIndexArr, GeometryCollection->SimulationType, FGeometryCollection::ESimulationTypes::FST_Rigid, SimplificationDistanceThreshold);
+	HullsFromGeometry(*GeometryCollection, GlobalTransformArray, HasCustomConvexFn, Convexes, TransformToConvexIndexArr, GeometryCollection->SimulationType, FGeometryCollection::ESimulationTypes::FST_Rigid, SimplificationDistanceThreshold);
 
-	CreateNonoverlappingConvexHulls(Convexes, TransformToConvexIndexArr, GeometryCollection->SimulationType, FGeometryCollection::ESimulationTypes::FST_Rigid, FGeometryCollection::ESimulationTypes::FST_None, GeometryCollection->Parent, GCProximity, GeometryCollection->TransformIndex, FracAllowRemove, SimplificationDistanceThreshold);
+	CreateNonoverlappingConvexHulls(Convexes, TransformToConvexIndexArr, HasCustomConvexFn, GeometryCollection->SimulationType, FGeometryCollection::ESimulationTypes::FST_Rigid, FGeometryCollection::ESimulationTypes::FST_None, GeometryCollection->Parent, GCProximity, GeometryCollection->TransformIndex, Volume, FracAllowRemove, SimplificationDistanceThreshold, CanExceedFraction);
 
 	TransformHullsToLocal(GlobalTransformArray, Convexes, TransformToConvexIndexArr);
 
@@ -945,13 +1094,54 @@ FGeometryCollectionConvexUtility::FGeometryCollectionConvexData FGeometryCollect
 	TManagedArray<TSet<int32>>& TransformToConvexIndices = GeometryCollection->GetAttribute<TSet<int32>>("TransformToConvexIndices", FTransformCollection::TransformGroup);
 	TManagedArray<TUniquePtr<Chaos::FConvex>>& ConvexHull = GeometryCollection->GetAttribute<TUniquePtr<Chaos::FConvex>>("ConvexHull", "Convex");
 	TransformToConvexIndices = MoveTemp(TransformToConvexIndexArr);
+	GeometryCollection->EmptyGroup("Convex");
 	GeometryCollection->Resize(Convexes.Num(), "Convex");
 	ConvexHull = MoveTemp(Convexes);
 
+	// clear all null hulls
+	TArray<int32> EmptyConvex;
+	for (int32 ConvexIdx = 0; ConvexIdx < ConvexHull.Num(); ConvexIdx++)
+	{
+		if (!ConvexHull[ConvexIdx].IsValid())
+		{
+			EmptyConvex.Add(ConvexIdx);
+		}
+	}
+	GeometryCollection->RemoveElements("Convex", EmptyConvex);
+
+	checkSlow(FGeometryCollectionConvexUtility::ValidateConvexData(GeometryCollection));
+
 	return { TransformToConvexIndices, ConvexHull };
 }
 
 
+bool FGeometryCollectionConvexUtility::ValidateConvexData(const FGeometryCollection* GeometryCollection)
+{
+	if (!GeometryCollection->HasAttribute("ConvexHull", "Convex") || !GeometryCollection->HasAttribute("TransformToConvexIndices", FTransformCollection::TransformGroup))
+	{
+		return false;
+	}
+	const TManagedArray<TUniquePtr<Chaos::FConvex>>& ConvexHull = GeometryCollection->GetAttribute<TUniquePtr<Chaos::FConvex>>("ConvexHull", "Convex");
+	for (int32 ConvexIdx = 0; ConvexIdx < ConvexHull.Num(); ConvexIdx++)
+	{
+		if (!ConvexHull[ConvexIdx].IsValid())
+		{
+			return false;
+		}
+	}
+	const TManagedArray<TSet<int32>>& TransformToConvexIndices = GeometryCollection->GetAttribute<TSet<int32>>("TransformToConvexIndices", FTransformCollection::TransformGroup);
+	for (int32 TransformIdx = 0; TransformIdx < TransformToConvexIndices.Num(); TransformIdx++)
+	{
+		for (int32 ConvexIdx : TransformToConvexIndices[TransformIdx])
+		{
+			if (ConvexIdx < 0 || ConvexIdx >= ConvexHull.Num())
+			{
+				return false;
+			}
+		}
+	}
+	return true;
+}
 
 TUniquePtr<Chaos::FConvex> FGeometryCollectionConvexUtility::FindConvexHull(const FGeometryCollection* GeometryCollection, int32 GeometryIndex)
 {
@@ -1000,4 +1190,150 @@ void FGeometryCollectionConvexUtility::RemoveConvexHulls(FGeometryCollection* Ge
 
 void FGeometryCollectionConvexUtility::SetDefaults(FGeometryCollection* GeometryCollection, FName Group, uint32 StartSize, uint32 NumElements)
 {
-}
\ No newline at end of file
+}
+
+
+TManagedArray<int32>* FGeometryCollectionConvexUtility::GetCustomConvexFlags(FGeometryCollection* GeometryCollection, bool bAddIfMissing)
+{
+	if (!GeometryCollection->HasAttribute("HasCustomConvex", FTransformCollection::TransformGroup))
+	{
+		if (!bAddIfMissing)
+		{
+			return nullptr;
+		}
+		return &GeometryCollection->AddAttribute<int32>("HasCustomConvex", FTransformCollection::TransformGroup);
+	}
+	return &GeometryCollection->GetAttribute<int32>("HasCustomConvex", FTransformCollection::TransformGroup);
+}
+
+
+void FGeometryCollectionConvexUtility::CopyChildConvexes(FGeometryCollection* FromCollection, const TArrayView<const int32>& FromTransformIdx, FGeometryCollection* ToCollection, const TArrayView<const int32>& ToTransformIdx, bool bLeafOnly)
+{
+	TOptional<FGeometryCollectionConvexUtility::FGeometryCollectionConvexData> InConvex = FGeometryCollectionConvexUtility::GetConvexHullDataIfPresent(FromCollection);
+	TOptional<FGeometryCollectionConvexUtility::FGeometryCollectionConvexData> OutConvex = FGeometryCollectionConvexUtility::GetConvexHullDataIfPresent(ToCollection);
+	TManagedArray<int32>* OutCustomFlags = GetCustomConvexFlags(ToCollection, true);
+
+	if (!ensure(InConvex.IsSet()) || !ensure(OutConvex.IsSet())) // missing convex data on collection(s); cannot copy
+	{
+		return;
+	}
+
+	if (!ensure(FromTransformIdx.Num() == ToTransformIdx.Num())) // these arrays must be matched up
+	{
+		return;
+	}
+
+	int32 InNumBones = FromCollection->NumElements(FGeometryCollection::TransformGroup);
+	int32 LeafType = FGeometryCollection::ESimulationTypes::FST_Rigid;
+	TArray<TArray<int32>> Children;
+	Children.SetNum(InNumBones);
+	for (int InBone = 0; InBone < InNumBones; InBone++)
+	{
+		if (FromCollection->SimulationType[InBone] == FGeometryCollection::ESimulationTypes::FST_None)
+		{
+			continue; // skip embedded geo
+		}
+
+		int32 ParentBone = FromCollection->Parent[InBone];
+		if (ParentBone != INDEX_NONE)
+		{
+			if (FromCollection->SimulationType[ParentBone] != LeafType)
+			{
+				Children[ParentBone].Add(InBone);
+			}
+			else
+			{
+				// parent is leaf (should only happen for embedded geo, which we've already skipped above)
+				continue;
+			}
+		}
+	}
+
+	auto GetChildrenWithConvex = [&FromCollection, bLeafOnly, &InConvex, LeafType, &Children](int32 Bone, TArray<int32>& OutChildren)
+	{
+		TArray<int32> ToExpand = Children[Bone];
+		if (ToExpand.IsEmpty()) // if no children, fall back to copying the convex on the bone itself
+		{
+			ToExpand.Add(Bone);
+		}
+		while (ToExpand.Num() > 0)
+		{
+			int32 ToProcess = ToExpand.Pop();
+			bool bIsLeaf = FromCollection->SimulationType[ToProcess] == LeafType;
+			if (bIsLeaf || (!bLeafOnly && InConvex->TransformToConvexIndices[ToProcess].Num() > 0))
+			{
+				OutChildren.Add(ToProcess);
+			}
+			else if (Children[ToProcess].Num() > 0)
+			{
+				ToExpand.Append(Children[ToProcess]);
+			}
+		}
+	};
+
+	TArray<FTransform> InGlobalTransformArray, OutGlobalTransformArray;
+	GeometryCollectionAlgo::GlobalMatrices(FromCollection->Transform, FromCollection->Parent, InGlobalTransformArray);
+	if (FromCollection == ToCollection)
+	{
+		OutGlobalTransformArray = InGlobalTransformArray;
+	}
+	else
+	{
+		GeometryCollectionAlgo::GlobalMatrices(ToCollection->Transform, ToCollection->Parent, OutGlobalTransformArray);
+	}
+	
+	// build the new convex data in separate arrays (which support incremental add)
+	TArray<TUniquePtr<Chaos::FConvex>> ConvexToAdd;
+	TArray<TSet<int32>> OutTransformToConvexIndices;
+	OutTransformToConvexIndices.SetNum(OutConvex->TransformToConvexIndices.Num());
+	TSet<int32> ToRemove;
+
+	for (int32 i = 0; i < ToTransformIdx.Num(); i++)
+	{
+		int32 InBone = FromTransformIdx[i];
+		int32 OutBone = ToTransformIdx[i];
+		TArray<int32> BonesWithHulls;
+		GetChildrenWithConvex(InBone, BonesWithHulls);
+
+		ToRemove.Add(OutBone);
+		(*OutCustomFlags)[OutBone] = true;
+		OutTransformToConvexIndices[OutBone].Reset();
+		for (int32 InBoneWithHulls : BonesWithHulls)
+		{
+			CopyHulls(
+				InGlobalTransformArray, InConvex->ConvexHull, InConvex->TransformToConvexIndices, InBoneWithHulls,
+				OutGlobalTransformArray, ConvexToAdd, OutTransformToConvexIndices, OutBone);
+		}
+	}
+
+	int32 OrigNumConvex = ToCollection->NumElements("Convex");
+
+	TArray<int32> ToRemoveArr = ToRemove.Array();
+	ToRemoveArr.Sort();
+	if (ToRemoveArr.Num() > 0)
+	{
+		if (!ensure(ToRemoveArr.Last() < OrigNumConvex))
+		{
+			return; // this could only happen if the input was messed up in some way, e.g. maybe if ToTransformIdx array had duplicates?
+		}
+		RemoveConvexHulls(ToCollection, ToRemoveArr);
+	}
+
+	int32 NewNumConvex = ToCollection->NumElements("Convex");
+
+	for (int32 OutBone : ToTransformIdx)
+	{
+		OutConvex->TransformToConvexIndices[OutBone] = MoveTemp(OutTransformToConvexIndices[OutBone]);
+		for (int32& ConvexInd : OutConvex->TransformToConvexIndices[OutBone])
+		{
+			ConvexInd += NewNumConvex;
+		}
+	}
+	
+	ToCollection->Resize(NewNumConvex + ConvexToAdd.Num(), "Convex");
+	for (int32 i = 0; i < ConvexToAdd.Num(); i++)
+	{
+		OutConvex->ConvexHull[NewNumConvex + i] = MoveTemp(ConvexToAdd[i]);
+	}
+}
+
diff --git a/Engine/Source/Runtime/Experimental/Chaos/Public/GeometryCollection/GeometryCollectionConvexUtility.h b/Engine/Source/Runtime/Experimental/Chaos/Public/GeometryCollection/GeometryCollectionConvexUtility.h
--- a/Engine/Source/Runtime/Experimental/Chaos/Public/GeometryCollection/GeometryCollectionConvexUtility.h
+++ b/Engine/Source/Runtime/Experimental/Chaos/Public/GeometryCollection/GeometryCollectionConvexUtility.h
@@ -25,14 +25,18 @@ public:
 	/** Ensure that convex hull data exists for the Geometry Collection and construct it if not (or if some data is missing. */
 	static FGeometryCollectionConvexData GetValidConvexHullData(FGeometryCollection* GeometryCollection);
 
+	/** Get convex hull data for the Geometry Collection if it is present */
+	static TOptional<FGeometryCollectionConvexData> GetConvexHullDataIfPresent(FGeometryCollection* GeometryCollection);
+
 	/**
 	 Create non-overlapping convex hull data for all transforms in the geometry collection (except transforms where it would be better to just use the hulls of the children) 
 
 	 @param GeometryCollection					The collection to add convex hulls to
 	 @param FractionAllowRemove					The fraction of a convex body we can cut away to remove overlaps with neighbors, before we fall back to using the hulls of the children directly.  (Does not affect leaves of hierarchy)
 	 @param SimplificationDistanceThreshold		Approximate minimum distance between vertices, below which we remove vertices to generate a simpler convex shape.  If 0.0, no simplification will occur.
+	 @param CanExceedFraction					The fraction by which the convex body volume on a cluster can exceed the volume of the geometry under that cluster (a value of 1 == exceed by 100% == convex hull has 2x the volume of the geometry)
 	 */
-	static FGeometryCollectionConvexData CreateNonOverlappingConvexHullData(FGeometryCollection* GeometryCollection, double FractionAllowRemove = .3, double SimplificationDistanceThreshold = 0.0);
+	static FGeometryCollectionConvexData CreateNonOverlappingConvexHullData(FGeometryCollection* GeometryCollection, double FractionAllowRemove = .3, double SimplificationDistanceThreshold = 0.0, double CanExceedFraction = .5);
 
 	/** Returns the convex hull of the vertices contained in the specified geometry. */
 	static TUniquePtr<Chaos::FConvex> FindConvexHull(const FGeometryCollection* GeometryCollection, int32 GeometryIndex);
@@ -43,5 +47,24 @@ public:
 	/** Set default values for convex hull related managed arrays. */
 	static void SetDefaults(FGeometryCollection* GeometryCollection, FName Group, uint32 StartSize, uint32 NumElements);
 
+	/** Get the HasCustomConvex flags.  If they're missing, either add them (if bAddIfMissing) or return nullptr */
+	static TManagedArray<int32>* GetCustomConvexFlags(FGeometryCollection* GeometryCollection, bool bAddIfMissing = false);
+
+	/** @return true if the GeometryCollection has convex data with no null pointers and no invalid indices */
+	static bool ValidateConvexData(const FGeometryCollection* GeometryCollection);
+
+	/**
+	 * Copy convex hulls from *below* FromTransformIdx over to all live at ToTransformIdx.
+	 * The two geometry collections can be the same but do not need to be.
+	 * Note: This will also set HasCustomConvex flags for the ToTransformIdx.
+	 * 
+	 * @param FromCollection	The collection to copy from
+	 * @param FromTransformIdx	The transform indices whose child nodes should be copied from (or, if it's a leaf, the hull on the leaf will be copied instead)
+	 * @param ToCollection		The collection to copy to (can be the same as FromCollection)
+	 * @param ToTransformIdx	The transform indices whose convexes will be replaced with *copies* of the child convexes.  Must be same length as FromTransformIdx; can be the same array.
+	 * @param bLeafOnly			If true, we will only collect convexes from leaf bones, not from clusters.
+	 */
+	static void CopyChildConvexes(FGeometryCollection* FromCollection, const TArrayView<const int32>& FromTransformIdx, FGeometryCollection* ToCollection, const TArrayView<const int32>& ToTransformIdx, bool bLeafOnly);
+
 };
 
