Fixed collsiion constraints thread safety issue - reading the main thread material array from the physics thread.

diff --git a/Engine/Source/Runtime/Engine/Private/PhysicsEngine/ImmediatePhysics/ImmediatePhysicsChaos/ImmediatePhysicsSimulation_Chaos.cpp b/Engine/Source/Runtime/Engine/Private/PhysicsEngine/ImmediatePhysics/ImmediatePhysicsChaos/ImmediatePhysicsSimulation_Chaos.cpp
--- a/Engine/Source/Runtime/Engine/Private/PhysicsEngine/ImmediatePhysics/ImmediatePhysicsChaos/ImmediatePhysicsSimulation_Chaos.cpp
+++ b/Engine/Source/Runtime/Engine/Private/PhysicsEngine/ImmediatePhysics/ImmediatePhysicsChaos/ImmediatePhysicsSimulation_Chaos.cpp
@@ -262,7 +262,7 @@ namespace ImmediatePhysics_Chaos
 		FImplementation()
 			: Particles(UniqueIndices)
 			, Joints()
-			, Collisions(Particles, CollidedParticles, ParticleMaterials, PerParticleMaterials, 0, 0, ChaosImmediate_Collision_CullDistance)
+			, Collisions(Particles, CollidedParticles, ParticleMaterials, PerParticleMaterials, nullptr, 0, 0, ChaosImmediate_Collision_CullDistance)
 			, BroadPhase(&ActivePotentiallyCollidingPairs, nullptr, nullptr)
 			, NarrowPhase(ChaosImmediate_Collision_CullDistance, FReal(0), Collisions.GetConstraintAllocator())
 			, CollisionDetector(BroadPhase, NarrowPhase, Collisions)
diff --git a/Engine/Source/Runtime/Experimental/Chaos/Private/Chaos/PBDCollisionConstraints.cpp b/Engine/Source/Runtime/Experimental/Chaos/Private/Chaos/PBDCollisionConstraints.cpp
--- a/Engine/Source/Runtime/Experimental/Chaos/Private/Chaos/PBDCollisionConstraints.cpp
+++ b/Engine/Source/Runtime/Experimental/Chaos/Private/Chaos/PBDCollisionConstraints.cpp
@@ -98,6 +98,7 @@ namespace Chaos
 		TArrayCollectionArray<bool>& Collided,
 		const TArrayCollectionArray<TSerializablePtr<FChaosPhysicsMaterial>>& InPhysicsMaterials,
 		const TArrayCollectionArray<TUniquePtr<FChaosPhysicsMaterial>>& InPerParticlePhysicsMaterials,
+		const THandleArray<FChaosPhysicsMaterial>* const InSimMaterials,
 		const int32 InApplyPairIterations /*= 1*/,
 		const int32 InApplyPushOutPairIterations /*= 1*/,
 		const FReal InRestitutionThreshold /*= (FReal)0*/)
@@ -107,6 +108,7 @@ namespace Chaos
 		, MCollided(Collided)
 		, MPhysicsMaterials(InPhysicsMaterials)
 		, MPerParticlePhysicsMaterials(InPerParticlePhysicsMaterials)
+		, SimMaterials(InSimMaterials)
 		, MApplyPairIterations(InApplyPairIterations)
 		, MApplyPushOutPairIterations(InApplyPushOutPairIterations)
 		, RestitutionThreshold(InRestitutionThreshold)	// @todo(chaos): expose as property
@@ -142,7 +144,7 @@ namespace Chaos
 		return ConstraintAllocator.GetConstConstraints();
 	}
 
-	const FChaosPhysicsMaterial* GetPhysicsMaterial(const TGeometryParticleHandle<FReal, 3>* Particle, const FImplicitObject* Geom, const TArrayCollectionArray<TSerializablePtr<FChaosPhysicsMaterial>>& PhysicsMaterials, const TArrayCollectionArray<TUniquePtr<FChaosPhysicsMaterial>>& PerParticlePhysicsMaterials)
+	const FChaosPhysicsMaterial* GetPhysicsMaterial(const TGeometryParticleHandle<FReal, 3>* Particle, const FImplicitObject* Geom, const TArrayCollectionArray<TSerializablePtr<FChaosPhysicsMaterial>>& PhysicsMaterials, const TArrayCollectionArray<TUniquePtr<FChaosPhysicsMaterial>>& PerParticlePhysicsMaterials, const THandleArray<FChaosPhysicsMaterial>* const SimMaterials)
 	{
 		// Use the per-particle material if it exists
 		const FChaosPhysicsMaterial* UniquePhysicsMaterial = Particle->AuxilaryValue(PerParticlePhysicsMaterials).Get();
@@ -166,7 +168,14 @@ namespace Chaos
 			{
 				if (ShapeData->GetMaterials().Num() > 0)
 				{
-					return ShapeData->GetMaterials()[0].Get();
+					if(SimMaterials)
+					{
+						return SimMaterials->Get(ShapeData->GetMaterials()[0].InnerHandle);
+					}
+					else
+					{
+						UE_LOG(LogChaos, Warning, TEXT("Attempted to resolve a material for a constraint but we do not have a sim material container."));
+					}
 				}
 				else
 				{
@@ -183,8 +192,8 @@ namespace Chaos
 
 	void FPBDCollisionConstraints::UpdateConstraintMaterialProperties(FPBDCollisionConstraint& Constraint)
 	{
-		const FChaosPhysicsMaterial* PhysicsMaterial0 = GetPhysicsMaterial(Constraint.Particle[0], Constraint.Implicit[0], MPhysicsMaterials, MPerParticlePhysicsMaterials);
-		const FChaosPhysicsMaterial* PhysicsMaterial1 = GetPhysicsMaterial(Constraint.Particle[1], Constraint.Implicit[1], MPhysicsMaterials, MPerParticlePhysicsMaterials);
+		const FChaosPhysicsMaterial* PhysicsMaterial0 = GetPhysicsMaterial(Constraint.Particle[0], Constraint.Implicit[0], MPhysicsMaterials, MPerParticlePhysicsMaterials, SimMaterials);
+		const FChaosPhysicsMaterial* PhysicsMaterial1 = GetPhysicsMaterial(Constraint.Particle[1], Constraint.Implicit[1], MPhysicsMaterials, MPerParticlePhysicsMaterials, SimMaterials);
 
 		FCollisionContact& Contact = Constraint.Manifold;
 		if (PhysicsMaterial0 && PhysicsMaterial1)
diff --git a/Engine/Source/Runtime/Experimental/Chaos/Private/Chaos/PBDRigidsEvolutionGBF.cpp b/Engine/Source/Runtime/Experimental/Chaos/Private/Chaos/PBDRigidsEvolutionGBF.cpp
--- a/Engine/Source/Runtime/Experimental/Chaos/Private/Chaos/PBDRigidsEvolutionGBF.cpp
+++ b/Engine/Source/Runtime/Experimental/Chaos/Private/Chaos/PBDRigidsEvolutionGBF.cpp
@@ -654,7 +654,7 @@ FPBDRigidsEvolutionGBF::FPBDRigidsEvolutionGBF(FPBDRigidsSOAs& InParticles,THand
 	, Clustering(*this, Particles.GetClusteredParticles())
 	, JointConstraintRule(JointConstraints, ChaosSolverJointPriority)
 	, SuspensionConstraintRule(SuspensionConstraints, ChaosSolverSuspensionPriority)
-	, CollisionConstraints(InParticles, Collided, PhysicsMaterials, PerParticlePhysicsMaterials, DefaultNumCollisionPairIterations, DefaultNumCollisionPushOutPairIterations, DefaultRestitutionThreshold)
+	, CollisionConstraints(InParticles, Collided, PhysicsMaterials, PerParticlePhysicsMaterials, &SolverPhysicsMaterials, DefaultNumCollisionPairIterations, DefaultNumCollisionPushOutPairIterations, DefaultRestitutionThreshold)
 	, CollisionRule(CollisionConstraints, ChaosSolverCollisionPriority)
 	, BroadPhase(InParticles)
 	, NarrowPhase(DefaultCollisionCullDistance, BoundsThicknessVelocityMultiplier, CollisionConstraints.GetConstraintAllocator())
diff --git a/Engine/Source/Runtime/Experimental/Chaos/Public/Chaos/PBDCollisionConstraints.h b/Engine/Source/Runtime/Experimental/Chaos/Public/Chaos/PBDCollisionConstraints.h
--- a/Engine/Source/Runtime/Experimental/Chaos/Public/Chaos/PBDCollisionConstraints.h
+++ b/Engine/Source/Runtime/Experimental/Chaos/Public/Chaos/PBDCollisionConstraints.h
@@ -57,7 +57,8 @@ public:
 	FPBDCollisionConstraints(const FPBDRigidsSOAs& InParticles, 
 		TArrayCollectionArray<bool>& Collided, 
 		const TArrayCollectionArray<TSerializablePtr<FChaosPhysicsMaterial>>& PhysicsMaterials, 
-		const TArrayCollectionArray<TUniquePtr<FChaosPhysicsMaterial>>& PerParticlePhysicsMaterials, 
+		const TArrayCollectionArray<TUniquePtr<FChaosPhysicsMaterial>>& PerParticlePhysicsMaterials,
+		const THandleArray<FChaosPhysicsMaterial>* const SimMaterials,
 		const int32 ApplyPairIterations = 1, 
 		const int32 ApplyPushOutPairIterations = 1, 
 		const FReal RestitutionThreshold = 2000.0f);
@@ -311,6 +312,7 @@ private:
 	TArrayCollectionArray<bool>& MCollided;
 	const TArrayCollectionArray<TSerializablePtr<FChaosPhysicsMaterial>>& MPhysicsMaterials;
 	const TArrayCollectionArray<TUniquePtr<FChaosPhysicsMaterial>>& MPerParticlePhysicsMaterials;
+	const THandleArray<FChaosPhysicsMaterial>* const SimMaterials;
 	int32 MApplyPairIterations;
 	int32 MApplyPushOutPairIterations;
 	FReal RestitutionThreshold;
