Interchange: Fixed attribute storage crash when trying to get an empty string from the storage.

diff --git a/Engine/Source/Runtime/Experimental/Interchange/Core/Private/Tests/StorageTest.cpp b/Engine/Source/Runtime/Experimental/Interchange/Core/Private/Tests/StorageTest.cpp
--- a/Engine/Source/Runtime/Experimental/Interchange/Core/Private/Tests/StorageTest.cpp
+++ b/Engine/Source/Runtime/Experimental/Interchange/Core/Private/Tests/StorageTest.cpp
@@ -92,6 +92,8 @@ bool FAttributeStorageTest::RunTest(const FString& Parameters)
 	
 	ADD_ATTRIBUTE_STORAGE(FRotator, FRotator(25.0, 2.0, 3.14159), );
 	ADD_ATTRIBUTE_STORAGE(FString, FString(TEXT("Testing FString storage")), );
+	ADD_ATTRIBUTE_STORAGE(FString, FString(TEXT("")), );
+	ADD_ATTRIBUTE_STORAGE(FString, FString(TEXT("A")), );
 	ADD_ATTRIBUTE_STORAGE(FTimespan, FTimespan(10, 44, 38), );
 
 	{
diff --git a/Engine/Source/Runtime/Experimental/Interchange/Core/Private/Types/AttributeStorage.cpp b/Engine/Source/Runtime/Experimental/Interchange/Core/Private/Types/AttributeStorage.cpp
--- a/Engine/Source/Runtime/Experimental/Interchange/Core/Private/Types/AttributeStorage.cpp
+++ b/Engine/Source/Runtime/Experimental/Interchange/Core/Private/Types/AttributeStorage.cpp
@@ -409,7 +409,7 @@ namespace UE
 				return EAttributeStorageResult::Operation_Error_CannotFoundKey;
 			}
 
-			if (AttributeAllocationInfo->Size == 0)
+			if (AttributeAllocationInfo->Size <= sizeof(TCHAR)) // Account for the null-terminator
 			{
 				OutValue = TEXT("");
 				return EAttributeStorageResult::Operation_Success;
