Fixed: Editor crashes when opening TM-PhysEntrance in QAGame with mobile preview enabled

diff --git a/Engine/Source/Runtime/Experimental/GeometryCollectionEngine/Private/GeometryCollection/GeometryCollectionSceneProxy.cpp b/Engine/Source/Runtime/Experimental/GeometryCollectionEngine/Private/GeometryCollection/GeometryCollectionSceneProxy.cpp
--- a/Engine/Source/Runtime/Experimental/GeometryCollectionEngine/Private/GeometryCollection/GeometryCollectionSceneProxy.cpp
+++ b/Engine/Source/Runtime/Experimental/GeometryCollectionEngine/Private/GeometryCollection/GeometryCollectionSceneProxy.cpp
@@ -382,6 +382,13 @@ void FGeometryCollectionSceneProxy::InitResources()
 		Data.BoneTransformSRV = TransformBuffers[0].VertexBufferSRV;
 		Data.BonePrevTransformSRV = PrevTransformBuffers[0].VertexBufferSRV;
 	}
+	else
+	{
+		// make sure these are not null to pass UB validation
+		Data.BoneMapSRV = GNullColorVertexBuffer.VertexBufferSRV;
+		Data.BoneTransformSRV = GNullColorVertexBuffer.VertexBufferSRV;
+		Data.BonePrevTransformSRV = GNullColorVertexBuffer.VertexBufferSRV;
+	}
 
 	// 
 	// from InitOrUpdateResource(VertexFactory);
