Fix CIS static analysis
#rnx

diff --git a/Engine/Source/Runtime/Experimental/Interchange/Core/Private/Tests/StorageTest.cpp b/Engine/Source/Runtime/Experimental/Interchange/Core/Private/Tests/StorageTest.cpp
--- a/Engine/Source/Runtime/Experimental/Interchange/Core/Private/Tests/StorageTest.cpp
+++ b/Engine/Source/Runtime/Experimental/Interchange/Core/Private/Tests/StorageTest.cpp
@@ -150,7 +150,12 @@ bool FAttributeStorageTest::RunTest(const FString& Parameters)
 	}
 
 	{
-		ADD_ATTRIBUTE_STORAGE_NOCOMPARE(FSphere, FSphere(FVector(2.0), 2.2))
+		const FSphere RefValue = FSphere(FVector(2.0), 2.2);
+		FAttributeKey Key = CreateUniqueKey();
+		if (!IsAttributeStorageResultSuccess(TestStorage.RegisterAttribute(Key, RefValue)))
+		{
+			AddError(FString(TEXT("`AttributeStorage` must handle adding FSphere attribute")));
+		}
 		FSphere StoredValue;
 		TestStorage.GetAttributeHandle<FSphere>(Key).Get(StoredValue);
 		if (!StoredValue.Equals(RefValue))
@@ -162,6 +167,64 @@ bool FAttributeStorageTest::RunTest(const FString& Parameters)
 	ADD_ATTRIBUTE_STORAGE(FFrameRate, FFrameRate(1,60), );
 	ADD_ATTRIBUTE_STORAGE(FFrameTime, FFrameTime(29,0.5f), );
 	ADD_ATTRIBUTE_STORAGE(FSoftObjectPath, FSoftObjectPath(UClass::StaticClass()), );
+	ADD_ATTRIBUTE_STORAGE(FMatrix44f, FMatrix44f(FTransform3f(FRotator3f(25.0f, 2.0f, 3.14159f), FVector3f(150.0f, -203.0f, 4500.7f), FVector3f(1.0f, 1.0f, 1.0f)).ToMatrixWithScale()), );
+	ADD_ATTRIBUTE_STORAGE(FMatrix44d, FMatrix44d(FTransform3d(FRotator3d(25.0, 2.0, 3.14159), FVector3d(150.0, -203.0, 4500.7), FVector3d(1.0, 1.0, 1.0)).ToMatrixWithScale()), );
+	ADD_ATTRIBUTE_STORAGE(FPlane4f, FPlane4f(FVector3f(1.0f, 1.0f, 0.0f), FVector3f(0.0f, 0.0f, 1.0f)), );
+	ADD_ATTRIBUTE_STORAGE(FPlane4d, FPlane4d(FVector3d(1.0, 1.0, 0), FVector3d(0, 0, 1.0)), );
+	ADD_ATTRIBUTE_STORAGE(FQuat4f, FQuat4f(1.0f, 1.0f, 0.0f, 1.0f), );
+	ADD_ATTRIBUTE_STORAGE(FQuat4d, FQuat4d(1.0, 1.0, 0, 1.0), );
+	ADD_ATTRIBUTE_STORAGE(FRotator3f, FRotator3f(25.0f, 2.0f, 3.14159f), );
+	ADD_ATTRIBUTE_STORAGE(FRotator3d, FRotator3d(25.0, 2.0, 3.14159), );
+	{
+		ADD_ATTRIBUTE_STORAGE_NOCOMPARE(FTransform3f, FTransform3f(FRotator3f(25.0f, 2.0f, 3.14159f), FVector3f(150.0f, -203.0f, 4500.7f), FVector3f(1.0f, 1.0f, 1.0f)))
+		FTransform3f StoredValue;
+		TestStorage.GetAttributeHandle<FTransform3f>(Key).Get(StoredValue);
+		if (!StoredValue.Equals(RefValue))
+		{
+			AddError(FString(TEXT("`AttributeStorage` must handle adding FTransform3f attribute")));
+		}
+	}
+	{
+		ADD_ATTRIBUTE_STORAGE_NOCOMPARE(FTransform3d, FTransform3d(FRotator3d(25.0, 2.0, 3.14159), FVector3d(150.0, -203.0, 4500.7), FVector3d(1.0, 1.0, 1.0)))
+		FTransform3d StoredValue;
+		TestStorage.GetAttributeHandle<FTransform>(Key).Get(StoredValue);
+		if (!StoredValue.Equals(RefValue))
+		{
+			AddError(FString(TEXT("`AttributeStorage` must handle adding FTransform3d attribute")));
+		}
+	}
+	ADD_ATTRIBUTE_STORAGE(FVector3f, FVector3f(150.0f, -203.0f, 4500.7f), );
+	ADD_ATTRIBUTE_STORAGE(FVector3d, FVector3d(150.0, -203.0, 4500.7), );
+	ADD_ATTRIBUTE_STORAGE(FVector2f, FVector2f(150.0f, -203.0f), );
+	ADD_ATTRIBUTE_STORAGE(FVector4f, FVector4f(150.0f, -203.0f, 4500.7f, 2.0f), );
+	ADD_ATTRIBUTE_STORAGE(FVector4d, FVector4d(150.0, -203.0, 4500.7, 2.0), );
+	ADD_ATTRIBUTE_STORAGE(FBox2f, FBox2f(FVector2f(-1.0f), FVector2f(1.0f)), );
+	ADD_ATTRIBUTE_STORAGE(FBox2D, FBox2D(FVector2D(-1.0), FVector2D(1.0)), );
+	ADD_ATTRIBUTE_STORAGE(FBox3f, FBox3f(FVector3f(-1.0f), FVector3f(1.0f)), );
+	ADD_ATTRIBUTE_STORAGE(FBox3d, FBox3d(FVector3d(-1.0), FVector3d(1.0)), );
+	ADD_ATTRIBUTE_STORAGE(FBoxSphereBounds3f, FBoxSphereBounds3f(FVector3f(0.0f), FVector3f(2.0f), 2.2f), );
+	ADD_ATTRIBUTE_STORAGE(FBoxSphereBounds3d, FBoxSphereBounds3d(FVector3d(0.0), FVector3d(2.0), 2.2), );
+	{
+		ADD_ATTRIBUTE_STORAGE_NOCOMPARE(FSphere3f, FSphere3f(FVector3f(2.0f), 2.2f))
+		FSphere3f StoredValue;
+		uint32 bobafet = -45;
+		TestStorage.GetAttributeHandle<FSphere3f>(Key).Get(StoredValue);
+		if (!StoredValue.Equals(RefValue))
+		{
+			AddError(FString(TEXT("`AttributeStorage` must handle adding FSphere3f attribute")));
+		}
+	}
+	{
+		ADD_ATTRIBUTE_STORAGE_NOCOMPARE(FSphere3d, FSphere3d(FVector3d(2.0), 2.2))
+		FSphere3d StoredValue;
+		uint32 bobafet = -45;
+		TestStorage.GetAttributeHandle<FSphere3d>(Key).Get(StoredValue);
+		if (!StoredValue.Equals(RefValue))
+		{
+			AddError(FString(TEXT("`AttributeStorage` must handle adding FSphere attribute")));
+		}
+	}
+
 	//Test Adding/Reading uint8 attribute with 0 value
 	{
 		const uint8 RefValue = 0;
diff --git a/Engine/Source/Runtime/Experimental/Interchange/Core/Private/Types/AttributeStorage.cpp b/Engine/Source/Runtime/Experimental/Interchange/Core/Private/Types/AttributeStorage.cpp
--- a/Engine/Source/Runtime/Experimental/Interchange/Core/Private/Types/AttributeStorage.cpp
+++ b/Engine/Source/Runtime/Experimental/Interchange/Core/Private/Types/AttributeStorage.cpp
@@ -61,6 +61,29 @@ namespace UE
 				case EAttributeTypes::FrameRate: { AttributeTypeString = TEXT("FrameRate"); } break;
 				case EAttributeTypes::FrameTime: { AttributeTypeString = TEXT("FrameTime"); } break;
 				case EAttributeTypes::SoftObjectPath: { AttributeTypeString = TEXT("SoftObjectPath"); } break;
+				case EAttributeTypes::Matrix44f: { AttributeTypeString = TEXT("Matrix44f"); } break;
+				case EAttributeTypes::Matrix44d: { AttributeTypeString = TEXT("Matrix44d"); } break;
+				case EAttributeTypes::Plane4f: { AttributeTypeString = TEXT("Plane4f"); } break;
+				case EAttributeTypes::Plane4d: { AttributeTypeString = TEXT("Plane4d"); } break;
+				case EAttributeTypes::Quat4f: { AttributeTypeString = TEXT("Quat4f"); } break;
+				case EAttributeTypes::Quat4d: { AttributeTypeString = TEXT("Quat4d"); } break;
+				case EAttributeTypes::Rotator3f: { AttributeTypeString = TEXT("Rotator3f"); } break;
+				case EAttributeTypes::Rotator3d: { AttributeTypeString = TEXT("Rotator3d"); } break;
+				case EAttributeTypes::Transform3f: { AttributeTypeString = TEXT("Transform3f"); } break;
+				case EAttributeTypes::Transform3d: { AttributeTypeString = TEXT("Transform3d"); } break;
+				case EAttributeTypes::Vector3f: { AttributeTypeString = TEXT("Vector3f"); } break;
+				case EAttributeTypes::Vector3d: { AttributeTypeString = TEXT("Vector3d"); } break;
+				case EAttributeTypes::Vector2f: { AttributeTypeString = TEXT("Vector2f"); } break;
+				case EAttributeTypes::Vector4f: { AttributeTypeString = TEXT("Vector4f"); } break;
+				case EAttributeTypes::Vector4d: { AttributeTypeString = TEXT("Vector4d"); } break;
+				case EAttributeTypes::Box2f: { AttributeTypeString = TEXT("Box2f"); } break;
+				case EAttributeTypes::Box2D: { AttributeTypeString = TEXT("Box2D"); } break;
+				case EAttributeTypes::Box3f: { AttributeTypeString = TEXT("Box3f"); } break;
+				case EAttributeTypes::Box3d: { AttributeTypeString = TEXT("Box3d"); } break;
+				case EAttributeTypes::BoxSphereBounds3f: { AttributeTypeString = TEXT("BoxSphereBounds3f"); } break;
+				case EAttributeTypes::BoxSphereBounds3d: { AttributeTypeString = TEXT("BoxSphereBounds3d"); } break;
+				case EAttributeTypes::Sphere3f: { AttributeTypeString = TEXT("Sphere3f"); } break;
+				case EAttributeTypes::Sphere3d: { AttributeTypeString = TEXT("Sphere3d"); } break;
 				default:
 				{
 					//Ensure if we ask an unknown type
diff --git a/Engine/Source/Runtime/Experimental/Interchange/Core/Public/Types/AttributeStorage.h b/Engine/Source/Runtime/Experimental/Interchange/Core/Public/Types/AttributeStorage.h
--- a/Engine/Source/Runtime/Experimental/Interchange/Core/Public/Types/AttributeStorage.h
+++ b/Engine/Source/Runtime/Experimental/Interchange/Core/Public/Types/AttributeStorage.h
@@ -188,6 +188,29 @@ namespace UE
 			FrameRate			= 42,
 			FrameTime			= 43,
 			SoftObjectPath		= 44,
+			Matrix44f			= 45,
+			Matrix44d			= 46,
+			Plane4f				= 47,
+			Plane4d				= 48,
+			Quat4f				= 49,
+			Quat4d				= 50,
+			Rotator3f			= 51,
+			Rotator3d			= 52,
+			Transform3f			= 53,
+			Transform3d			= 54,
+			Vector3f			= 55,
+			Vector3d			= 56,
+			Vector2f			= 57,
+			Vector4f			= 58,
+			Vector4d			= 59,
+			Box2f				= 60,
+			Box2D				= 61,
+			Box3f				= 62,
+			Box3d				= 63,
+			BoxSphereBounds3f	= 64,
+			BoxSphereBounds3d	= 65,
+			Sphere3f			= 66,
+			Sphere3d			= 67
 		};
 
 		/**
@@ -1017,18 +1040,10 @@ namespace UE
 		};
 
 
-		/** Implements variant type traits for the built-in FBox type. */
-		template<> struct TAttributeTypeTraits<FBox>
-		{
-			static CONSTEXPR EAttributeTypes GetType() { return EAttributeTypes::Box; }
-		};
+		/** See FBox3f/FBox3d variant type traits for the built-in FBox type. */
 
 
-		/** Implements variant type traits for the built-in FBoxSphereBounds type. */
-		template<> struct TAttributeTypeTraits<FBoxSphereBounds>
-		{
-			static CONSTEXPR EAttributeTypes GetType() { return EAttributeTypes::BoxSphereBounds; }
-		};
+		/** See FBoxSphereBounds3f/FBoxSphereBounds3d variant type traits for the built-in FBoxSphereBounds type. */
 
 
 		/** Implements variant type traits for byte arrays. */
@@ -1127,11 +1142,7 @@ namespace UE
 		};
 
 
-		/** Implements variant type traits for the built-in FMatrix type. */
-		template<> struct TAttributeTypeTraits<FMatrix>
-		{
-			static CONSTEXPR EAttributeTypes GetType() { return EAttributeTypes::Matrix; }
-		};
+		/** See FMatrix44f/FMatrix44d variant type traits for the built-in FMatrix type. */
 
 		/** Implements variant type traits for the built-in FName type. */
 		template<> struct TAttributeTypeTraits<FName>
@@ -1139,18 +1150,10 @@ namespace UE
 			static CONSTEXPR EAttributeTypes GetType() { return EAttributeTypes::Name; }
 		};
 
-		/** Implements variant type traits for the built-in FPlane type. */
-		template<> struct TAttributeTypeTraits<FPlane>
-		{
-			static CONSTEXPR EAttributeTypes GetType() { return EAttributeTypes::Plane; }
-		};
+		/** See FPlane4f/FPlane4d variant type traits for the built-in FPlane type. */
 
 
-		/** Implements variant type traits for the built-in FQuat type. */
-		template<> struct TAttributeTypeTraits<FQuat>
-		{
-			static CONSTEXPR EAttributeTypes GetType() { return EAttributeTypes::Quat; }
-		};
+		/** See FQuat4f/FQuat4d variant type traits for the built-in FQuat type. */
 
 		/** Implements variant type traits for the built-in FRandomStream type. */
 		template<> struct TAttributeTypeTraits<FRandomStream>
@@ -1159,11 +1162,7 @@ namespace UE
 		};
 
 
-		/** Implements variant type traits for the built-in FRotator type. */
-		template<> struct TAttributeTypeTraits<FRotator>
-		{
-			static CONSTEXPR EAttributeTypes GetType() { return EAttributeTypes::Rotator; }
-		};
+		/** See FRotator3f/FRotator3d variant type traits for the built-in FRotator type. */
 
 
 		/** Implements variant type traits for the built-in FString type. */
@@ -1180,11 +1179,7 @@ namespace UE
 		};
 
 
-		/** Implements variant type traits for the built-in FTransform type. */
-		template<> struct TAttributeTypeTraits<FTransform>
-		{
-			static CONSTEXPR EAttributeTypes GetType() { return EAttributeTypes::Transform; }
-		};
+		/** See FTransform3f/FTransform3d variant type traits for the built-in FTransform type. */
 
 
 		/** Implements variant type traits for the built-in FTwoVectors type. */
@@ -1222,11 +1217,7 @@ namespace UE
 		};
 
 
-		/** Implements variant type traits for the built-in FVector type. */
-		template<> struct TAttributeTypeTraits<FVector>
-		{
-			static CONSTEXPR EAttributeTypes GetType() { return EAttributeTypes::Vector; }
-		};
+		/** See FVector3f/FVector3d variant type traits for the built-in FVector type. */
 
 
 		/** Implements variant type traits for the built-in FVector2D type. */
@@ -1235,11 +1226,7 @@ namespace UE
 			static CONSTEXPR EAttributeTypes GetType() { return EAttributeTypes::Vector2d; }
 		};
 
-		/** Implements variant type traits for the built-in FVector4 type. */
-		template<> struct TAttributeTypeTraits<FVector4>
-		{
-			static CONSTEXPR EAttributeTypes GetType() { return EAttributeTypes::Vector4; }
-		};
+		/** See FVector4f/FVector4d variant type traits for the built-in FVector4 type. */
 
 		/** Implements variant type traits for the built-in Vector2DHalf type. */
 		template<> struct TAttributeTypeTraits<FIntPoint>
@@ -1271,11 +1258,7 @@ namespace UE
 			static CONSTEXPR EAttributeTypes GetType() { return EAttributeTypes::OrientedBox; }
 		};
 
-		/** Implements variant type traits for the built-in Sphere type. */
-		template<> struct TAttributeTypeTraits<FSphere>
-		{
-			static CONSTEXPR EAttributeTypes GetType() { return EAttributeTypes::Sphere; }
-		};
+		/** See FSphere3f/FSphere3d variant type traits for the built-in Sphere type. */
 
 		/** Implements variant type traits for the built-in FrameNumber type. */
 		template<> struct TAttributeTypeTraits<FFrameNumber>
@@ -1301,5 +1284,142 @@ namespace UE
 			static CONSTEXPR EAttributeTypes GetType() { return EAttributeTypes::SoftObjectPath; }
 		};
 
+		/** Implements variant type traits for the built-in Matrix44f type. */
+		template<> struct TAttributeTypeTraits<FMatrix44f>
+		{
+			static CONSTEXPR EAttributeTypes GetType() { return EAttributeTypes::Matrix44f; }
+		};
+
+		/** Implements variant type traits for the built-in Matrix44f type. */
+		template<> struct TAttributeTypeTraits<FMatrix44d>
+		{
+			static CONSTEXPR EAttributeTypes GetType() { return EAttributeTypes::Matrix44d; }
+		};
+
+		/** Implements variant type traits for the built-in Plane4f type. */
+		template<> struct TAttributeTypeTraits<FPlane4f>
+		{
+			static CONSTEXPR EAttributeTypes GetType() { return EAttributeTypes::Plane4f; }
+		};
+
+		/** Implements variant type traits for the built-in Plane4d type. */
+		template<> struct TAttributeTypeTraits<FPlane4d>
+		{
+			static CONSTEXPR EAttributeTypes GetType() { return EAttributeTypes::Plane4d; }
+		};
+
+		/** Implements variant type traits for the built-in Quat4f type. */
+		template<> struct TAttributeTypeTraits<FQuat4f>
+		{
+			static CONSTEXPR EAttributeTypes GetType() { return EAttributeTypes::Quat4f; }
+		};
+
+		/** Implements variant type traits for the built-in Quat4d type. */
+		template<> struct TAttributeTypeTraits<FQuat4d>
+		{
+			static CONSTEXPR EAttributeTypes GetType() { return EAttributeTypes::Quat4d; }
+		};
+
+		/** Implements variant type traits for the built-in Rotator3f type. */
+		template<> struct TAttributeTypeTraits<FRotator3f>
+		{
+			static CONSTEXPR EAttributeTypes GetType() { return EAttributeTypes::Rotator3f; }
+		};
+
+		/** Implements variant type traits for the built-in Rotator3d type. */
+		template<> struct TAttributeTypeTraits<FRotator3d>
+		{
+			static CONSTEXPR EAttributeTypes GetType() { return EAttributeTypes::Rotator3d; }
+		};
+
+		/** Implements variant type traits for the built-in Rotator3f type. */
+		template<> struct TAttributeTypeTraits<FTransform3f>
+		{
+			static CONSTEXPR EAttributeTypes GetType() { return EAttributeTypes::Transform3f; }
+		};
+
+		/** Implements variant type traits for the built-in Transform3d type. */
+		template<> struct TAttributeTypeTraits<FTransform3d>
+		{
+			static CONSTEXPR EAttributeTypes GetType() { return EAttributeTypes::Transform3d; }
+		};
+
+		/** Implements variant type traits for the built-in Vector3f type. */
+		template<> struct TAttributeTypeTraits<FVector3f>
+		{
+			static CONSTEXPR EAttributeTypes GetType() { return EAttributeTypes::Vector3f; }
+		};
+
+		/** Implements variant type traits for the built-in Vector3d type. */
+		template<> struct TAttributeTypeTraits<FVector3d>
+		{
+			static CONSTEXPR EAttributeTypes GetType() { return EAttributeTypes::Vector3d; }
+		};
+
+		/** Implements variant type traits for the built-in Vector2f type. */
+		template<> struct TAttributeTypeTraits<FVector2f>
+		{
+			static CONSTEXPR EAttributeTypes GetType() { return EAttributeTypes::Vector2f; }
+		};
+
+		/** Implements variant type traits for the built-in Vector4f type. */
+		template<> struct TAttributeTypeTraits<FVector4f>
+		{
+			static CONSTEXPR EAttributeTypes GetType() { return EAttributeTypes::Vector4f; }
+		};
+
+		/** Implements variant type traits for the built-in Vector4d type. */
+		template<> struct TAttributeTypeTraits<FVector4d>
+		{
+			static CONSTEXPR EAttributeTypes GetType() { return EAttributeTypes::Vector4d; }
+		};
+
+		/** Implements variant type traits for the built-in Box2f type. */
+		template<> struct TAttributeTypeTraits<FBox2f>
+		{
+			static CONSTEXPR EAttributeTypes GetType() { return EAttributeTypes::Box2f; }
+		};
+
+		/** Implements variant type traits for the built-in Box2D type. */
+		template<> struct TAttributeTypeTraits<FBox2D>
+		{
+			static CONSTEXPR EAttributeTypes GetType() { return EAttributeTypes::Box2D; }
+		};
+
+		/** Implements variant type traits for the built-in Box3f type. */
+		template<> struct TAttributeTypeTraits<FBox3f>
+		{
+			static CONSTEXPR EAttributeTypes GetType() { return EAttributeTypes::Box3f; }
+		};
+
+		/** Implements variant type traits for the built-in Box3d type. */
+		template<> struct TAttributeTypeTraits<FBox3d>
+		{
+			static CONSTEXPR EAttributeTypes GetType() { return EAttributeTypes::Box3d; }
+		};
+
+		/** Implements variant type traits for the built-in BoxSphereBounds3f type. */
+		template<> struct TAttributeTypeTraits<FBoxSphereBounds3f>
+		{
+			static CONSTEXPR EAttributeTypes GetType() { return EAttributeTypes::BoxSphereBounds3f; }
+		};
+
+		/** Implements variant type traits for the built-in BoxSphereBounds3d type. */
+		template<> struct TAttributeTypeTraits<FBoxSphereBounds3d>
+		{
+			static CONSTEXPR EAttributeTypes GetType() { return EAttributeTypes::BoxSphereBounds3d; }
+		};
+
+		/** Implements variant type traits for the built-in Sphere3f type. */
+		template<> struct TAttributeTypeTraits<FSphere3f>
+		{
+			static CONSTEXPR EAttributeTypes GetType() { return EAttributeTypes::Sphere3f; }
+		};
+
+		/** Implements variant type traits for the built-in Sphere3d type. */
+		template<> struct TAttributeTypeTraits<FSphere3d>
+		{
+			static CONSTEXPR EAttributeTypes GetType() { return EAttributeTypes::Sphere3d; }
+		};
 	} //ns interchange
 }
