Chaos - fix manifold point re-use on first tick

#okfirgithub public

diff --git a/Engine/Source/Runtime/Experimental/Chaos/Private/Chaos/Collision/ParticlePairMidPhase.cpp b/Engine/Source/Runtime/Experimental/Chaos/Private/Chaos/Collision/ParticlePairMidPhase.cpp
--- a/Engine/Source/Runtime/Experimental/Chaos/Private/Chaos/Collision/ParticlePairMidPhase.cpp
+++ b/Engine/Source/Runtime/Experimental/Chaos/Private/Chaos/Collision/ParticlePairMidPhase.cpp
@@ -286,10 +286,11 @@ namespace Chaos
 			}
 
 			bool bWasManifoldRestored = false;
-			if (Flags.bEnableManifoldUpdate)
+			if (Flags.bEnableManifoldUpdate && bWasUpdatedLastTick)
 			{
 				// Update the existing manifold. We can re-use as-is if none of the points have moved much and the bodies have not moved much
 				// NOTE: this can succeed in "restoring" even if we have no manifold points
+				// NOTE: this uses the previous world-space shape transforms, so we can only do this if we were updated last tick
 				bWasManifoldRestored = Constraint->UpdateAndTryRestoreManifold();
 			}
 			else
