minor fixes / cleanup for geometrycollectionconvexutility

#rnx

diff --git a/Engine/Source/Runtime/Experimental/Chaos/Private/GeometryCollection/GeometryCollectionConvexUtility.cpp b/Engine/Source/Runtime/Experimental/Chaos/Private/GeometryCollection/GeometryCollectionConvexUtility.cpp
--- a/Engine/Source/Runtime/Experimental/Chaos/Private/GeometryCollection/GeometryCollectionConvexUtility.cpp
+++ b/Engine/Source/Runtime/Experimental/Chaos/Private/GeometryCollection/GeometryCollectionConvexUtility.cpp
@@ -26,6 +26,10 @@ TOptional<FGeometryCollectionConvexUtility::FGeometryCollectionConvexData> FGeom
 	return TOptional<FGeometryCollectionConvexUtility::FGeometryCollectionConvexData>(ConvexData);
 }
 
+bool FGeometryCollectionConvexUtility::HasConvexHullData(FGeometryCollection* GeometryCollection)
+{
+	return GeometryCollection->HasAttribute("TransformToConvexIndices", FTransformCollection::TransformGroup) && GeometryCollection->HasAttribute("ConvexHull", "Convex");
+}
 
 FGeometryCollectionConvexUtility::FGeometryCollectionConvexData FGeometryCollectionConvexUtility::GetValidConvexHullData(FGeometryCollection* GeometryCollection)
 {
@@ -1207,13 +1211,14 @@ TManagedArray<int32>* FGeometryCollectionConvexUtility::GetCustomConvexFlags(FGe
 }
 
 
-void FGeometryCollectionConvexUtility::CopyChildConvexes(FGeometryCollection* FromCollection, const TArrayView<const int32>& FromTransformIdx, FGeometryCollection* ToCollection, const TArrayView<const int32>& ToTransformIdx, bool bLeafOnly)
+void FGeometryCollectionConvexUtility::CopyChildConvexes(const FGeometryCollection* FromCollection, const TArrayView<const int32>& FromTransformIdx, FGeometryCollection* ToCollection, const TArrayView<const int32>& ToTransformIdx, bool bLeafOnly)
 {
-	TOptional<FGeometryCollectionConvexUtility::FGeometryCollectionConvexData> InConvex = FGeometryCollectionConvexUtility::GetConvexHullDataIfPresent(FromCollection);
+	const TManagedArray<TSet<int32>>* InTransformToConvexIndices = FromCollection->FindAttribute<TSet<int32>>("TransformToConvexIndices", FTransformCollection::TransformGroup);
+	const TManagedArray<TUniquePtr<Chaos::FConvex>>* InConvexHull = FromCollection->FindAttribute<TUniquePtr<Chaos::FConvex>>("ConvexHull", "Convex");
 	TOptional<FGeometryCollectionConvexUtility::FGeometryCollectionConvexData> OutConvex = FGeometryCollectionConvexUtility::GetConvexHullDataIfPresent(ToCollection);
 	TManagedArray<int32>* OutCustomFlags = GetCustomConvexFlags(ToCollection, true);
 
-	if (!ensure(InConvex.IsSet()) || !ensure(OutConvex.IsSet())) // missing convex data on collection(s); cannot copy
+	if (!ensure(InTransformToConvexIndices != nullptr && InConvexHull != nullptr) || !ensure(OutConvex.IsSet())) // missing convex data on collection(s); cannot copy
 	{
 		return;
 	}
@@ -1249,7 +1254,7 @@ void FGeometryCollectionConvexUtility::CopyChildConvexes(FGeometryCollection* Fr
 		}
 	}
 
-	auto GetChildrenWithConvex = [&FromCollection, bLeafOnly, &InConvex, LeafType, &Children](int32 Bone, TArray<int32>& OutChildren)
+	auto GetChildrenWithConvex = [&FromCollection, bLeafOnly, &InTransformToConvexIndices, LeafType, &Children](int32 Bone, TArray<int32>& OutChildren)
 	{
 		TArray<int32> ToExpand = Children[Bone];
 		if (ToExpand.IsEmpty()) // if no children, fall back to copying the convex on the bone itself
@@ -1260,7 +1265,7 @@ void FGeometryCollectionConvexUtility::CopyChildConvexes(FGeometryCollection* Fr
 		{
 			int32 ToProcess = ToExpand.Pop();
 			bool bIsLeaf = FromCollection->SimulationType[ToProcess] == LeafType;
-			if (bIsLeaf || (!bLeafOnly && InConvex->TransformToConvexIndices[ToProcess].Num() > 0))
+			if (bIsLeaf || (!bLeafOnly && (*InTransformToConvexIndices)[ToProcess].Num() > 0))
 			{
 				OutChildren.Add(ToProcess);
 			}
@@ -1301,21 +1306,15 @@ void FGeometryCollectionConvexUtility::CopyChildConvexes(FGeometryCollection* Fr
 		for (int32 InBoneWithHulls : BonesWithHulls)
 		{
 			CopyHulls(
-				InGlobalTransformArray, InConvex->ConvexHull, InConvex->TransformToConvexIndices, InBoneWithHulls,
+				InGlobalTransformArray, *InConvexHull, *InTransformToConvexIndices, InBoneWithHulls,
 				OutGlobalTransformArray, ConvexToAdd, OutTransformToConvexIndices, OutBone);
 		}
 	}
 
-	int32 OrigNumConvex = ToCollection->NumElements("Convex");
-
 	TArray<int32> ToRemoveArr = ToRemove.Array();
 	ToRemoveArr.Sort();
 	if (ToRemoveArr.Num() > 0)
 	{
-		if (!ensure(ToRemoveArr.Last() < OrigNumConvex))
-		{
-			return; // this could only happen if the input was messed up in some way, e.g. maybe if ToTransformIdx array had duplicates?
-		}
 		RemoveConvexHulls(ToCollection, ToRemoveArr);
 	}
 
diff --git a/Engine/Source/Runtime/Experimental/Chaos/Public/GeometryCollection/GeometryCollectionConvexUtility.h b/Engine/Source/Runtime/Experimental/Chaos/Public/GeometryCollection/GeometryCollectionConvexUtility.h
--- a/Engine/Source/Runtime/Experimental/Chaos/Public/GeometryCollection/GeometryCollectionConvexUtility.h
+++ b/Engine/Source/Runtime/Experimental/Chaos/Public/GeometryCollection/GeometryCollectionConvexUtility.h
@@ -28,6 +28,9 @@ public:
 	/** Get convex hull data for the Geometry Collection if it is present */
 	static TOptional<FGeometryCollectionConvexData> GetConvexHullDataIfPresent(FGeometryCollection* GeometryCollection);
 
+	/** @return true if convex hull data is present */
+	static bool HasConvexHullData(FGeometryCollection* GeometryCollection);
+
 	/**
 	 Create non-overlapping convex hull data for all transforms in the geometry collection (except transforms where it would be better to just use the hulls of the children) 
 
@@ -64,7 +67,7 @@ public:
 	 * @param ToTransformIdx	The transform indices whose convexes will be replaced with *copies* of the child convexes.  Must be same length as FromTransformIdx; can be the same array.
 	 * @param bLeafOnly			If true, we will only collect convexes from leaf bones, not from clusters.
 	 */
-	static void CopyChildConvexes(FGeometryCollection* FromCollection, const TArrayView<const int32>& FromTransformIdx, FGeometryCollection* ToCollection, const TArrayView<const int32>& ToTransformIdx, bool bLeafOnly);
+	static void CopyChildConvexes(const FGeometryCollection* FromCollection, const TArrayView<const int32>& FromTransformIdx, FGeometryCollection* ToCollection, const TArrayView<const int32>& ToTransformIdx, bool bLeafOnly);
 
 };
 
