Fix issues CCD issues against  heightfield/trimesh. Non-convex no longer contribute to CCD velocity threshold as bounds do not represent thickness.

diff --git a/Engine/Source/Runtime/Experimental/Chaos/Private/Chaos/CCDUtilities.cpp b/Engine/Source/Runtime/Experimental/Chaos/Private/Chaos/CCDUtilities.cpp
--- a/Engine/Source/Runtime/Experimental/Chaos/Private/Chaos/CCDUtilities.cpp
+++ b/Engine/Source/Runtime/Experimental/Chaos/Private/Chaos/CCDUtilities.cpp
@@ -32,8 +32,14 @@ namespace Chaos
 	{
 		if (Implicit)
 		{
-			const FReal MinExtent = Implicit->BoundingBox().Extents().Min();
-			return MinExtent * CVars::CCDEnableThresholdBoundsScale;
+			// Trimesh/Heightfield are thin, cannot use bounds. We do not want them to contribute to CCD threshold.
+			if (Implicit->IsConvex())
+			{
+				const FReal MinExtent = Implicit->BoundingBox().Extents().Min();
+				return MinExtent * CVars::CCDEnableThresholdBoundsScale;
+			}
+
+			return 0;
 		}
 		return TNumericLimits<FReal>::Max();
 	}
