Chaos
- add smooth edge collision flag
- remove unused collision flag
- clean up collision flag setting API

diff --git a/Engine/Source/Runtime/Engine/Classes/PhysicsEngine/BodyInstance.h b/Engine/Source/Runtime/Engine/Classes/PhysicsEngine/BodyInstance.h
--- a/Engine/Source/Runtime/Engine/Classes/PhysicsEngine/BodyInstance.h
+++ b/Engine/Source/Runtime/Engine/Classes/PhysicsEngine/BodyInstance.h
@@ -339,6 +339,13 @@ public:
 	/**	Enable contact modification. Assumes custom contact modification has been provided (see FPhysXContactModifyCallback) */
 	uint8 bContactModification : 1;
 
+	/**
+	 * Remove unnecessary edge collisions to allow smooth sliding over surfaces composed of multiple actors/components.
+	 * This is fairly expensive and should only be enabled on hero objects. 
+	 */
+	UPROPERTY(EditAnywhere, AdvancedDisplay, BlueprintReadOnly, Category = Collision)
+	uint8 bSmoothEdgeCollisions : 1;
+
 	/////////
 	// SIM SETTINGS
 
@@ -801,6 +808,8 @@ public:
 	void SetEnableGravity(bool bGravityEnabled);
 	/** Enables/disables contact modification */
 	void SetContactModification(bool bNewContactModification);
+	/** Enables/disabled smoothed edge collisions */
+	void SetSmoothEdgeCollisionsEnabled(bool bNewSmoothEdgeCollisions);
 
 	/** Enable/disable Continuous Collidion Detection feature */
 	void SetUseCCD(bool bInUseCCD);
diff --git a/Engine/Source/Runtime/Engine/Private/PhysicsEngine/BodyInstance.cpp b/Engine/Source/Runtime/Engine/Private/PhysicsEngine/BodyInstance.cpp
--- a/Engine/Source/Runtime/Engine/Private/PhysicsEngine/BodyInstance.cpp
+++ b/Engine/Source/Runtime/Engine/Private/PhysicsEngine/BodyInstance.cpp
@@ -1217,6 +1217,7 @@ void FInitBodiesHelperBase::CreateActor_AssumesLocked(FBodyInstance* Instance, c
 		FPhysicsInterface::SetIsKinematic_AssumesLocked(Instance->ActorHandle, !Instance->ShouldInstanceSimulatingPhysics());
 #if WITH_CHAOS
 		FPhysicsInterface::SetMaxLinearVelocity_AssumesLocked(Instance->ActorHandle, TNumericLimits<float>::Max());
+		FPhysicsInterface::SetSmoothEdgeCollisionsEnabled(Instance->ActorHandle, Instance->bSmoothEdgeCollisions);
 #endif
 
 		// Set sleep even notification
@@ -3676,6 +3677,25 @@ void FBodyInstance::SetContactModification(bool bNewContactModification)
 	}
 }
 
+void FBodyInstance::SetSmoothEdgeCollisionsEnabled(bool bNewSmoothEdgeCollisions)
+{
+	if (bNewSmoothEdgeCollisions != bSmoothEdgeCollisions)
+	{
+		bSmoothEdgeCollisions = bNewSmoothEdgeCollisions;
+#if WITH_CHAOS
+		FPhysicsCommand::ExecuteWrite(ActorHandle, 
+			[bNewSmoothEdgeCollisions](const FPhysicsActorHandle& Actor)
+			{
+				if (FPhysicsInterface::IsRigidBody(Actor))
+				{
+					FPhysicsInterface::SetSmoothEdgeCollisionsEnabled(Actor, bNewSmoothEdgeCollisions);
+				}
+			});
+#endif
+	}
+}
+
+
 void FBodyInstance::SetUseCCD(bool bInUseCCD)
 {
 	if (bUseCCD != bInUseCCD)
diff --git a/Engine/Source/Runtime/Experimental/Chaos/Private/Chaos/ParticleHandle.cpp b/Engine/Source/Runtime/Experimental/Chaos/Private/Chaos/ParticleHandle.cpp
--- a/Engine/Source/Runtime/Experimental/Chaos/Private/Chaos/ParticleHandle.cpp
+++ b/Engine/Source/Runtime/Experimental/Chaos/Private/Chaos/ParticleHandle.cpp
@@ -314,7 +314,7 @@ namespace Chaos
 		SetCCDEnabled(DynamicMisc.CCDEnabled());
 		SetDisabled(DynamicMisc.Disabled());
 		SetOneWayInteraction(DynamicMisc.OneWayInteraction());
-		AddCollisionConstraintFlag((Chaos::ECollisionConstraintFlags)DynamicMisc.CollisionConstraintFlag());
+		SetCollisionConstraintFlags(DynamicMisc.CollisionConstraintFlags());
 
 		Evolution.SetParticleObjectState(this, DynamicMisc.ObjectState());
 		Evolution.SetParticleSleepType(this, DynamicMisc.SleepType());
diff --git a/Engine/Source/Runtime/Experimental/Chaos/Private/PhysicsProxy/JointConstraintProxy.cpp b/Engine/Source/Runtime/Experimental/Chaos/Private/PhysicsProxy/JointConstraintProxy.cpp
--- a/Engine/Source/Runtime/Experimental/Chaos/Private/PhysicsProxy/JointConstraintProxy.cpp
+++ b/Engine/Source/Runtime/Experimental/Chaos/Private/PhysicsProxy/JointConstraintProxy.cpp
@@ -187,18 +187,6 @@ void FJointConstraintPhysicsProxy::PushStateOnPhysicsThread(FPBDRigidsSolver* In
 						}
 					}
 				}
-
-				// Contact transfer enables buffered collision constraints on the child of the joint
-				if (Handle0)
-				{
-					if (Chaos::TPBDRigidParticleHandle<FReal, 3>*RigidHandle0 = Handle0->CastToRigidParticle())
-					{
-						if (!FMath::IsNearlyZero(CurrentConstraintSettings.ContactTransferScale))
-						{
-							RigidHandle0->AddCollisionConstraintFlag(Chaos::ECollisionConstraintFlags::CCF_DoBufferCollisions);
-						}
-					}
-				}
 			}
 
 			Constraint_PT->SetSettings(JointSettingsBuffer);
diff --git a/Engine/Source/Runtime/Experimental/Chaos/Public/Chaos/Collision/CollisionConstraintFlags.h b/Engine/Source/Runtime/Experimental/Chaos/Public/Chaos/Collision/CollisionConstraintFlags.h
--- a/Engine/Source/Runtime/Experimental/Chaos/Public/Chaos/Collision/CollisionConstraintFlags.h
+++ b/Engine/Source/Runtime/Experimental/Chaos/Public/Chaos/Collision/CollisionConstraintFlags.h
@@ -14,11 +14,14 @@
 
 namespace Chaos
 {
+	/**
+	 * @brief Flags for user-control over per-particle collision behaviour
+	*/
 	enum class ECollisionConstraintFlags : uint32
 	{
 		CCF_None                       = 0x0,
 		CCF_BroadPhaseIgnoreCollisions = 0x1,
-		CCF_DoBufferCollisions         = 0x2,
+		CCF_SmoothEdgeCollisions       = 0x2,
 		CCF_DummyFlag
 	};
 
diff --git a/Engine/Source/Runtime/Experimental/Chaos/Public/Chaos/ParticleDirtyFlags.h b/Engine/Source/Runtime/Experimental/Chaos/Public/Chaos/ParticleDirtyFlags.h
--- a/Engine/Source/Runtime/Experimental/Chaos/Public/Chaos/ParticleDirtyFlags.h
+++ b/Engine/Source/Runtime/Experimental/Chaos/Public/Chaos/ParticleDirtyFlags.h
@@ -9,6 +9,7 @@
 #include "Chaos/PhysicalMaterials.h"
 #include "Chaos/GeometryParticlesfwd.h"
 #include "Chaos/CollisionFilterData.h"
+#include "Chaos/Collision/CollisionConstraintFlags.h"
 #include "Chaos/KinematicTargets.h"
 #include "UObject/ExternalPhysicsCustomObjectVersion.h"
 #include "UObject/ExternalPhysicsMaterialCustomObjectVersion.h"
@@ -357,7 +358,7 @@ public:
 		SetCollisionGroup(Other.CollisionGroup());
 		SetSleepType(Other.SleepType());
 		SetOneWayInteraction(Other.OneWayInteraction());
-		SetCollisionConstraintFlag(Other.CollisionConstraintFlag());
+		SetCollisionConstraintFlags(Other.CollisionConstraintFlags());
 		SetCCDEnabled(Other.CCDEnabled());
 		SetDisabled(Other.Disabled());
 	}
@@ -374,7 +375,7 @@ public:
 			&& CollisionGroup() == Other.CollisionGroup()
 			&& SleepType() == Other.SleepType()
 			&& OneWayInteraction() == Other.OneWayInteraction() 
-			&& CollisionConstraintFlag() == Other.CollisionConstraintFlag()
+			&& CollisionConstraintFlags() == Other.CollisionConstraintFlags()
 			&& CCDEnabled() == Other.CCDEnabled()
 			&& Disabled() == Other.Disabled();
 	}
@@ -414,8 +415,11 @@ public:
 	ESleepType SleepType() const { return MSleepType; }
 	void SetSleepType(ESleepType Type) { MSleepType = Type; }
 
-	uint32 CollisionConstraintFlag() const { return MCollisionConstraintFlag; }
-	void SetCollisionConstraintFlag(uint32 InCollisionConstraintFlag) { MCollisionConstraintFlag = InCollisionConstraintFlag; }
+	uint32 CollisionConstraintFlags() const { return MCollisionConstraintFlag; }
+	void SetCollisionConstraintFlags(uint32 InCollisionConstraintFlag) { MCollisionConstraintFlag = InCollisionConstraintFlag; }
+	void AddCollisionConstraintFlag(const ECollisionConstraintFlags Flag) { MCollisionConstraintFlag |= uint32(Flag); }
+	void RemoveCollisionConstraintFlag(const ECollisionConstraintFlags Flag) { MCollisionConstraintFlag &= ~uint32(Flag); }
+	
 	bool OneWayInteraction() const { return MOneWayInteraction; }
 	void SetOneWayInteraction(bool InOneWayInteraction) { MOneWayInteraction = InOneWayInteraction; }
 
diff --git a/Engine/Source/Runtime/Experimental/Chaos/Public/Chaos/ParticleHandle.h b/Engine/Source/Runtime/Experimental/Chaos/Public/Chaos/ParticleHandle.h
--- a/Engine/Source/Runtime/Experimental/Chaos/Public/Chaos/ParticleHandle.h
+++ b/Engine/Source/Runtime/Experimental/Chaos/Public/Chaos/ParticleHandle.h
@@ -832,7 +832,7 @@ protected:
 		: TKinematicGeometryParticleHandleImp<T, d, bPersistent>(TSerializablePtr<TKinematicGeometryParticles<T, d>>(Particles), InIdx, InGlobalIdx, Params)
 	{
 		PBDRigidParticleDefaultConstruct<T, d>(*this, Params);
-		ClearCollisionConstraintFlag();
+		SetCollisionConstraintFlags(0);
 		SetDisabled(Params.bDisabled);
 		SetPreV(this->V());
 		SetPreW(this->W());
@@ -879,8 +879,8 @@ public:
 	bool HasCollisionConstraintFlag(const ECollisionConstraintFlags Flag) const { return  PBDRigidParticles->HasCollisionConstraintFlag(Flag, ParticleIdx); }
 	void AddCollisionConstraintFlag(const ECollisionConstraintFlags Flag) { PBDRigidParticles->AddCollisionConstraintFlag(Flag, ParticleIdx); }
 	void RemoveCollisionConstraintFlag(const ECollisionConstraintFlags Flag) { PBDRigidParticles->RemoveCollisionConstraintFlag(Flag, ParticleIdx); }
-	void ClearCollisionConstraintFlag() { PBDRigidParticles->ClearCollisionConstraintFlag(ParticleIdx); }
-	uint32 CollisionConstraintFlag() const { return PBDRigidParticles->CollisionConstraintFlag(ParticleIdx); }
+	void SetCollisionConstraintFlags(const uint32 Flags) { PBDRigidParticles->SetCollisionConstraintFlags(ParticleIdx, Flags); }
+	uint32 CollisionConstraintFlags() const { return PBDRigidParticles->CollisionConstraintFlags(ParticleIdx); }
 
 	bool Disabled() const { return PBDRigidParticles->Disabled(ParticleIdx); }
 	bool& Disabled() { return PBDRigidParticles->DisabledRef(ParticleIdx); }
@@ -2441,10 +2441,22 @@ public:
 		MMiscData.Modify(true, MDirtyFlags, Proxy, [InOneWayInteraction](auto& Data) { Data.SetOneWayInteraction(InOneWayInteraction); });
 	}
 
-	uint32 CollisionConstraintFlag() const { return MMiscData.Read().CollisionConstraintFlag(); }
-	void SetCollisionConstraintFlag(const uint32 InCollisionConstraintFlag)
-	{
-		MMiscData.Modify(true, MDirtyFlags, Proxy, [InCollisionConstraintFlag](auto& Data) { Data.SetCollisionConstraintFlag(InCollisionConstraintFlag); });
+	// Enable a single flag
+	void AddCollisionConstraintFlag(const ECollisionConstraintFlags Flag)
+	{ 
+		MMiscData.Modify(true, MDirtyFlags, Proxy, [Flag](auto& Data) { Data.AddCollisionConstraintFlag(Flag); });
+	}
+	// Disable a single flag
+	void RemoveCollisionConstraintFlag(const ECollisionConstraintFlags Flag)
+	{ 
+		MMiscData.Modify(true, MDirtyFlags, Proxy, [Flag](auto& Data) { Data.RemoveCollisionConstraintFlag(Flag); });
+	}
+	// A mask of all the active flags
+	uint32 CollisionConstraintFlags() const { return MMiscData.Read().CollisionConstraintFlags(); }
+	// Replace all flags
+	void SetCollisionConstraintFlags(const uint32 Flags)
+	{ 
+		MMiscData.Modify(true, MDirtyFlags, Proxy, [Flags](auto& Data) { Data.SetCollisionConstraintFlags(Flags); });
 	}
 
 	bool CCDEnabled() const { return MMiscData.Read().CCDEnabled(); }
diff --git a/Engine/Source/Runtime/Experimental/Chaos/Public/Chaos/RigidParticles.h b/Engine/Source/Runtime/Experimental/Chaos/Public/Chaos/RigidParticles.h
--- a/Engine/Source/Runtime/Experimental/Chaos/Public/Chaos/RigidParticles.h
+++ b/Engine/Source/Runtime/Experimental/Chaos/Public/Chaos/RigidParticles.h
@@ -219,8 +219,8 @@ public:
 	FORCEINLINE bool HasCollisionConstraintFlag(const ECollisionConstraintFlags Flag, const int32 Index) const { return (MCollisionConstraintFlags[Index] & (uint32)Flag) != 0; }
 	FORCEINLINE void AddCollisionConstraintFlag(const ECollisionConstraintFlags Flag, const int32 Index) { MCollisionConstraintFlags[Index] |= (uint32)Flag; }
 	FORCEINLINE void RemoveCollisionConstraintFlag(const ECollisionConstraintFlags Flag, const int32 Index) { MCollisionConstraintFlags[Index] &= ~(uint32)Flag; }
-	FORCEINLINE void ClearCollisionConstraintFlag(const int32 Index) { MCollisionConstraintFlags[Index] = (uint32)ECollisionConstraintFlags::CCF_None; }
-	FORCEINLINE uint32 CollisionConstraintFlag(const int32 Index) const { return MCollisionConstraintFlags[Index]; }
+	FORCEINLINE void SetCollisionConstraintFlags(const int32 Index, const uint32 Flags) { MCollisionConstraintFlags[Index] = Flags; }
+	FORCEINLINE uint32 CollisionConstraintFlags(const int32 Index) const { return MCollisionConstraintFlags[Index]; }
 
 	FORCEINLINE const bool Disabled(const int32 Index) const { return MDisabled[Index]; }
 
diff --git a/Engine/Source/Runtime/Experimental/Chaos/Public/PhysicsProxy/SingleParticlePhysicsProxy.h b/Engine/Source/Runtime/Experimental/Chaos/Public/PhysicsProxy/SingleParticlePhysicsProxy.h
--- a/Engine/Source/Runtime/Experimental/Chaos/Public/PhysicsProxy/SingleParticlePhysicsProxy.h
+++ b/Engine/Source/Runtime/Experimental/Chaos/Public/PhysicsProxy/SingleParticlePhysicsProxy.h
@@ -1157,6 +1157,22 @@ public:
 		}
 	}
 
+	void SetSmoothEdgeCollisionsEnabled(bool bEnabled)
+	{
+		VerifyContext();
+		if (TPBDRigidParticle<FReal, 3>*Rigid = GetParticle_LowLevel()->CastToRigidParticle())
+		{
+			if (bEnabled)
+			{
+				Rigid->AddCollisionConstraintFlag(ECollisionConstraintFlags::CCF_SmoothEdgeCollisions);
+			}
+			else
+			{
+				Rigid->RemoveCollisionConstraintFlag(ECollisionConstraintFlags::CCF_SmoothEdgeCollisions);
+			}
+		}
+	}
+
 	void SetCCDEnabled(bool bEnabled)
 	{
 		VerifyContext();
diff --git a/Engine/Source/Runtime/PhysicsCore/Private/ChaosEngineInterface.cpp b/Engine/Source/Runtime/PhysicsCore/Private/ChaosEngineInterface.cpp
--- a/Engine/Source/Runtime/PhysicsCore/Private/ChaosEngineInterface.cpp
+++ b/Engine/Source/Runtime/PhysicsCore/Private/ChaosEngineInterface.cpp
@@ -282,6 +282,11 @@ void FChaosEngineInterface::DetachShape(const FPhysicsActorHandle& InActor,FPhys
 	}
 }
 
+void FChaosEngineInterface::SetSmoothEdgeCollisionsEnabled(const FPhysicsActorHandle& InActor, const bool bSmoothEdgeCollisionsEnabled)
+{
+	InActor->GetGameThreadAPI().SetSmoothEdgeCollisionsEnabled(bSmoothEdgeCollisionsEnabled);
+}
+
 void FChaosEngineInterface::AddDisabledCollisionsFor_AssumesLocked(const TMap<FPhysicsActorHandle, TArray< FPhysicsActorHandle > >& InMap)
 {
 	if (bEnableChaosCollisionManager)
@@ -306,12 +311,12 @@ void FChaosEngineInterface::AddDisabledCollisionsFor_AssumesLocked(const TMap<FP
 
 			if (Chaos::FPBDRigidParticle* Rigid0 = ActorReference->GetParticle_LowLevel()->CastToRigidParticle())
 			{
-				Rigid0->SetCollisionConstraintFlag((uint32)Chaos::ECollisionConstraintFlags::CCF_BroadPhaseIgnoreCollisions);
+				Rigid0->AddCollisionConstraintFlag(Chaos::ECollisionConstraintFlags::CCF_BroadPhaseIgnoreCollisions);
 				for (auto Handle1 : Elem.Value)
 				{
 					if (Chaos::FPBDRigidParticle* Rigid1 = Handle1->GetParticle_LowLevel()->CastToRigidParticle())
 					{
-						Rigid1->SetCollisionConstraintFlag((uint32)Chaos::ECollisionConstraintFlags::CCF_BroadPhaseIgnoreCollisions);
+						Rigid1->AddCollisionConstraintFlag(Chaos::ECollisionConstraintFlags::CCF_BroadPhaseIgnoreCollisions);
 						DisabledCollisions.Add(Handle1->GetGameThreadAPI().UniqueIdx());
 					}
 				}
diff --git a/Engine/Source/Runtime/PhysicsCore/Public/Chaos/ChaosEngineInterface.h b/Engine/Source/Runtime/PhysicsCore/Public/Chaos/ChaosEngineInterface.h
--- a/Engine/Source/Runtime/PhysicsCore/Public/Chaos/ChaosEngineInterface.h
+++ b/Engine/Source/Runtime/PhysicsCore/Public/Chaos/ChaosEngineInterface.h
@@ -376,6 +376,8 @@ public:
 	static void AttachShape(const FPhysicsActorHandle& InActor,const FPhysicsShapeHandle& InNewShape);
 	static void DetachShape(const FPhysicsActorHandle& InActor,FPhysicsShapeHandle& InShape,bool bWakeTouching = true);
 
+	static void SetSmoothEdgeCollisionsEnabled(const FPhysicsActorHandle& InActor, const bool bSmoothEdgeCollisionsEnabled);
+
 	static void AddDisabledCollisionsFor_AssumesLocked(const TMap<FPhysicsActorHandle, TArray< FPhysicsActorHandle > >& InMap);
 	static void RemoveDisabledCollisionsFor_AssumesLocked(TArray< FPhysicsActorHandle > & InPhysicsActors);
 
