Backout 18780632, did not merge properly

diff --git a/Engine/Source/Runtime/Experimental/Chaos/Private/Chaos/CCDUtilities.cpp b/Engine/Source/Runtime/Experimental/Chaos/Private/Chaos/CCDUtilities.cpp
--- a/Engine/Source/Runtime/Experimental/Chaos/Private/Chaos/CCDUtilities.cpp
+++ b/Engine/Source/Runtime/Experimental/Chaos/Private/Chaos/CCDUtilities.cpp
@@ -32,14 +32,8 @@ namespace Chaos
 	{
 		if (Implicit)
 		{
-			// Trimesh/Heightfield are thin, cannot use bounds. We do not want them to contribute to CCD threshold.
-			if (Implicit->IsConvex())
-			{
-				const FReal MinExtent = Implicit->BoundingBox().Extents().Min();
-				return MinExtent * CVars::CCDEnableThresholdBoundsScale;
-			}
-
-			return 0;
+			const FReal MinExtent = Implicit->BoundingBox().Extents().Min();
+			return MinExtent * CVars::CCDEnableThresholdBoundsScale;
 		}
 		return TNumericLimits<FReal>::Max();
 	}
@@ -48,13 +42,13 @@ namespace Chaos
 	{
 		for (int32 i = 0; i < 2; i++)
 		{
-			const TPBDRigidParticleHandle<FReal, 3>* Rigid = Constraint->GetParticle(i)->CastToRigidParticle();
+			const TPBDRigidParticleHandle<FReal, 3>* Rigid = Constraint->Particle[i]->CastToRigidParticle();
 			if (Rigid && Rigid->ObjectState() == EObjectStateType::Kinematic)
 			{
 				// The same computation is carried out in UseCCDImpl function when constructing constraints. But we don't have access to FCCDConstraint at that point. This part could potentially be optimized away. 
 				const FVec3 D = Displacements[i];
 				const FReal DSizeSquared = D.SizeSquared();
-				const FReal CCDThreshold = GetParticleCCDThreshold(Constraint->GetImplicit(i));
+				const FReal CCDThreshold = GetParticleCCDThreshold(Constraint->Manifold.Implicit[i]);
 				if (DSizeSquared > CCDThreshold * CCDThreshold)
 				{
 					return i;
@@ -109,7 +103,7 @@ namespace Chaos
 			FVec3 Displacements[2] = {FVec3(0.f), FVec3(0.f)};
 			for (int i = 0; i < 2; i++)
 			{
-				TPBDRigidParticleHandle<FReal, 3>* RigidParticle = Constraint->GetParticle(i)->CastToRigidParticle();
+				TPBDRigidParticleHandle<FReal, 3>* RigidParticle = Constraint->Particle[i]->CastToRigidParticle();
 				FCCDParticle* CCDParticle = nullptr;
 				const bool IsParticleDynamic = RigidParticle && RigidParticle->ObjectState() == EObjectStateType::Dynamic;
 				if (IsParticleDynamic)
@@ -138,31 +132,27 @@ namespace Chaos
 			}
 
 			// Compute the relative displacement. If relative displacement is smaller than 0.5 * (Extents0.Min() + Extents1.Min()), it is impossible for an particle to tunnel through another particle, even though the absolute velocities of the particles might be large.
-			const FReal CCDThreshold0 = GetParticleCCDThreshold(Constraint->GetImplicit0());
-			const FReal CCDThreshold1 = GetParticleCCDThreshold(Constraint->GetImplicit1());
+			const FReal CCDThreshold0 = GetParticleCCDThreshold(Constraint->Manifold.Implicit[0]);
+			const FReal CCDThreshold1 = GetParticleCCDThreshold(Constraint->Manifold.Implicit[1]);
 			const FReal CCDConstraintThreshold = CCDThreshold0 + CCDThreshold1;
 			if ((Displacements[1] - Displacements[0]).SizeSquared() > CCDConstraintThreshold * CCDConstraintThreshold)
 			{
 				bNeedCCDSolve = true;
 			}
 
-			// make sure we ignore pairs that don't include any dynamics
-			if (CCDParticlePair[0] != nullptr || CCDParticlePair[1] != nullptr)
+			CCDConstraints.Add(FCCDConstraint(Constraint, CCDParticlePair, Displacements));
+			for (int32 i = 0; i < 2; i++)
 			{
-				CCDConstraints.Add(FCCDConstraint(Constraint, CCDParticlePair, Displacements));
-				for (int32 i = 0; i < 2; i++)
+				if (CCDParticlePair[i] != nullptr)
 				{
-					if (CCDParticlePair[i] != nullptr)
-					{
-						CCDParticlePair[i]->AddConstraint(&CCDConstraints.Last());
-					}
+					CCDParticlePair[i]->AddConstraint(&CCDConstraints.Last());
 				}
+			}
 
-				if (IsDynamic[0] && IsDynamic[1])
-				{
-					CCDParticlePair[0]->AddOverlappingDynamicParticle(CCDParticlePair[1]);
-					CCDParticlePair[1]->AddOverlappingDynamicParticle(CCDParticlePair[0]);
-				}
+			if (IsDynamic[0] && IsDynamic[1])
+			{
+				CCDParticlePair[0]->AddOverlappingDynamicParticle(CCDParticlePair[1]);
+				CCDParticlePair[1]->AddOverlappingDynamicParticle(CCDParticlePair[0]);
 			}
 		}
 		return bNeedCCDSolve;
@@ -227,13 +217,7 @@ namespace Chaos
 			}
 			if (Island == INDEX_NONE)
 			{
-				// non-dynamic pairs are already ignored in Init() so if Particle 0 is null the second one should not be 
-				ensure(CCDConstraint.Particle[1] != nullptr);
-
-				if (CCDConstraint.Particle[1])
-				{
-					Island = CCDConstraint.Particle[1]->Island;
-				}	
+				Island = CCDConstraint.Particle[1]->Island;
 			}
 			CCDConstraint.Island = Island;
 			IslandConstraintNum[Island]++;
@@ -334,8 +318,8 @@ namespace Chaos
 					{
 						if (CCDConstraint->FastMovingKinematicIndex != INDEX_NONE)
 						{
-							const FConstGenericParticleHandle Particle1 = FGenericParticleHandle(CCDConstraint->SweptConstraint->GetParticle(CCDConstraint->FastMovingKinematicIndex));
-							const FVec3 Normal = CCDConstraint->SweptConstraint->CalculateWorldContactNormal();
+							const FConstGenericParticleHandle Particle1 = FGenericParticleHandle(CCDConstraint->SweptConstraint->Particle[CCDConstraint->FastMovingKinematicIndex]);
+							const FVec3 Normal = CCDConstraint->SweptConstraint->GetNormal();
 							const FVec3 Offset = FVec3::DotProduct(Particle1->V() * ((1.f - IslandTOI) * Dt), Normal) * Normal;
 							ClipParticleP(CCDConstraintParticles[0], Offset);
 						}
@@ -411,7 +395,7 @@ namespace Chaos
 								}
 								else
 								{
-									FGenericParticleHandle AffectedParticle = FGenericParticleHandle(AttachedCCDConstraint->SweptConstraint->GetParticle(j));
+									FGenericParticleHandle AffectedParticle = FGenericParticleHandle(AttachedCCDConstraint->SweptConstraint->Particle[j]);
 									const bool IsKinematic = AffectedParticle->ObjectState() == EObjectStateType::Kinematic;
 									if (IsKinematic)
 									{
@@ -454,19 +438,6 @@ namespace Chaos
 				}
 			}
 		}
-
-		// We need to update the world-space contact points at the final locations
-		for (int32 i = ConstraintStart; i < ConstraintEnd; i++)
-		{
-			FPBDCollisionConstraint* Constraint = SortedCCDConstraints[i]->SweptConstraint;
-			FRigidTransform3 ShapeWorldTransform0 = Constraint->GetShapeWorldTransform0();
-			FRigidTransform3 ShapeWorldTransform1 = Constraint->GetShapeWorldTransform1();
-			ShapeWorldTransform0.SetTranslation(FConstGenericParticleHandle(Constraint->GetParticle0())->P());
-			ShapeWorldTransform1.SetTranslation(FConstGenericParticleHandle(Constraint->GetParticle1())->P());
-
-			Constraint->SetShapeWorldTransforms(ShapeWorldTransform0, ShapeWorldTransform1);
-			Constraint->UpdateManifoldContacts();
-		}
 	}
 
 
@@ -525,14 +496,13 @@ namespace Chaos
 	void FCCDManager::ApplyImpulse(FCCDConstraint *CCDConstraint)
 	{
 		FPBDCollisionConstraint *Constraint = CCDConstraint->SweptConstraint;
-		TPBDRigidParticleHandle<FReal, 3> *Rigid0 = Constraint->GetParticle0()->CastToRigidParticle();
-		TPBDRigidParticleHandle<FReal, 3> *Rigid1 = Constraint->GetParticle1()->CastToRigidParticle();
+		TPBDRigidParticleHandle<FReal, 3> *Rigid0 = Constraint->Particle[0]->CastToRigidParticle();
+		TPBDRigidParticleHandle<FReal, 3> *Rigid1 = Constraint->Particle[1]->CastToRigidParticle();
 		check(Rigid0 != nullptr || Rigid1 != nullptr);
-		const FReal Restitution = Constraint->GetRestitution();
-		const FRigidTransform3& ShapeWorldTransform1 = Constraint->GetShapeWorldTransform1();
+		const FReal Restitution = Constraint->Manifold.Restitution;
 		for(const FManifoldPoint &ManifoldPoint : Constraint->GetManifoldPoints())
 		{
-			const FVec3 Normal = ShapeWorldTransform1.TransformVectorNoScale(ManifoldPoint.ContactPoint.ShapeContactNormal);
+			const FVec3 Normal = ManifoldPoint.ContactPoint.Normal;
 			const FVec3 V0 = Rigid0 != nullptr ? Rigid0->V() : FVec3(0.f);
 			const FVec3 V1 = Rigid1 != nullptr ? Rigid1->V() : FVec3(0.f);
 			const FReal NormalV = FVec3::DotProduct(V0 - V1, Normal);
@@ -564,7 +534,7 @@ namespace Chaos
 			FRigidTransform3 RigidTransforms[2];
 			for (int32 i = 0; i < 2; i++)
 			{
-				FGenericParticleHandle Particle = FGenericParticleHandle(SweptConstraint->GetParticle(i));
+				FGenericParticleHandle Particle = FGenericParticleHandle(SweptConstraint->Particle[i]);
 				const bool IsStatic = Particle->ObjectState() == EObjectStateType::Static;
 				if (IsStatic)
 				{
