Fix CIS static analysis

diff --git a/Engine/Source/Runtime/Experimental/Interchange/Core/Private/Tests/StorageTest.cpp b/Engine/Source/Runtime/Experimental/Interchange/Core/Private/Tests/StorageTest.cpp
--- a/Engine/Source/Runtime/Experimental/Interchange/Core/Private/Tests/StorageTest.cpp
+++ b/Engine/Source/Runtime/Experimental/Interchange/Core/Private/Tests/StorageTest.cpp
@@ -61,8 +61,6 @@ bool FAttributeStorageTest::RunTest(const FString& Parameters)
 
 	ADD_ATTRIBUTE_STORAGE(bool, false, );
 	ADD_ATTRIBUTE_STORAGE(bool, true, static_cast<int32>);
-	ADD_ATTRIBUTE_STORAGE(FBox, FBox(FVector(-1.0), FVector(1.0)), );
-	ADD_ATTRIBUTE_STORAGE(FBoxSphereBounds, FBoxSphereBounds(FVector(0.0), FVector(2.0), 2.2), );
 	ADD_ATTRIBUTE_STORAGE(FColor, FColor(0xdf,0xcf, 0x00), );
 	ADD_ATTRIBUTE_STORAGE(FDateTime, FDateTime(2022,01,25,10, 24), );
 	ADD_ATTRIBUTE_STORAGE(double, 0.1234567890123456789, );
@@ -77,11 +75,8 @@ bool FAttributeStorageTest::RunTest(const FString& Parameters)
 	ADD_ATTRIBUTE_STORAGE(int64, -0x75ccabcd12345678, );
 	ADD_ATTRIBUTE_STORAGE(int64, 0x7fccabcd12345678, );
 	ADD_ATTRIBUTE_STORAGE(FIntRect, FIntRect(-1,-2, 1, 2), );
-	ADD_ATTRIBUTE_STORAGE(FLinearColor, FLinearColor(0.59, 0.49, 0), );
-	ADD_ATTRIBUTE_STORAGE(FMatrix, FMatrix(FTransform(FRotator(25.0, 2.0, 3.14159), FVector(150.0, -203.0, 4500.7), FVector(1.0, 1.0, 1.0)).ToMatrixWithScale()), );
+	ADD_ATTRIBUTE_STORAGE(FLinearColor, FLinearColor(0.59f, 0.49f, 0.0f), );
 	ADD_ATTRIBUTE_STORAGE(FName, FName(TEXT("Testing FName storage")), );
-	ADD_ATTRIBUTE_STORAGE(FPlane, FPlane(FVector(1.0, 1.0, 0), FVector(0, 0, 1.0)), );
-	ADD_ATTRIBUTE_STORAGE(FQuat, FQuat(1.0, 1.0, 0, 1.0), );
 
 	{
 		ADD_ATTRIBUTE_STORAGE_NOCOMPARE(FRandomStream, FRandomStream(34))
@@ -90,22 +85,10 @@ bool FAttributeStorageTest::RunTest(const FString& Parameters)
 		TestEqual(TEXT("`AttributeStorage` must handle add and retrieve FRandomStream attribute"), StoredValue.ToString(), RefValue.ToString());
 	}
 	
-	ADD_ATTRIBUTE_STORAGE(FRotator, FRotator(25.0, 2.0, 3.14159), );
 	ADD_ATTRIBUTE_STORAGE(FString, FString(TEXT("Testing FString storage")), );
 	ADD_ATTRIBUTE_STORAGE(FString, FString(TEXT("")), );
 	ADD_ATTRIBUTE_STORAGE(FString, FString(TEXT("A")), );
 	ADD_ATTRIBUTE_STORAGE(FTimespan, FTimespan(10, 44, 38), );
-
-	{
-		ADD_ATTRIBUTE_STORAGE_NOCOMPARE(FTransform, FTransform(FRotator(25.0, 2.0, 3.14159), FVector(150.0, -203.0, 4500.7), FVector(1.0, 1.0, 1.0)))
-		FTransform StoredValue;
-		TestStorage.GetAttributeHandle<FTransform>(Key).Get(StoredValue);
-		if (!StoredValue.Equals(RefValue))
-		{
-			AddError(FString(TEXT("`AttributeStorage` must handle adding FTransform attribute")));
-		}
-	}
-
 	ADD_ATTRIBUTE_STORAGE(FTwoVectors, FTwoVectors(FVector(150.0, -203.0, 4500.7), FVector(1.0, 1.0, 1.0)), );
 	ADD_ATTRIBUTE_STORAGE(uint8, 0x85, );
 	ADD_ATTRIBUTE_STORAGE(uint8, 0x7f, );
@@ -115,9 +98,6 @@ bool FAttributeStorageTest::RunTest(const FString& Parameters)
 	ADD_ATTRIBUTE_STORAGE(uint32, 0x7fccabcd, );
 	ADD_ATTRIBUTE_STORAGE(uint64, 0x85ccabcd12345678, );
 	ADD_ATTRIBUTE_STORAGE(uint64, 0x7fccabcd12345678, );
-	ADD_ATTRIBUTE_STORAGE(FVector, FVector(150.0, -203.0, 4500.7), );
-	ADD_ATTRIBUTE_STORAGE(FVector2D, FVector2D(150.0, -203.0), );
-	ADD_ATTRIBUTE_STORAGE(FVector4, FVector4(150.0, -203.0, 4500.7, 2.0), );
 	ADD_ATTRIBUTE_STORAGE(FIntPoint, FIntPoint(1, 4), );
 	ADD_ATTRIBUTE_STORAGE(FIntVector, FIntVector(1, 4, 3), );
 
@@ -149,20 +129,6 @@ bool FAttributeStorageTest::RunTest(const FString& Parameters)
 		}
 	}
 
-	{
-		const FSphere RefValue = FSphere(FVector(2.0), 2.2);
-		FAttributeKey Key = CreateUniqueKey();
-		if (!IsAttributeStorageResultSuccess(TestStorage.RegisterAttribute(Key, RefValue)))
-		{
-			AddError(FString(TEXT("`AttributeStorage` must handle adding FSphere attribute")));
-		}
-		FSphere StoredValue;
-		TestStorage.GetAttributeHandle<FSphere>(Key).Get(StoredValue);
-		if (!StoredValue.Equals(RefValue))
-		{
-			AddError(FString(TEXT("`AttributeStorage` must handle adding FSphere attribute")));
-		}
-	}
 	ADD_ATTRIBUTE_STORAGE(FFrameNumber, FFrameNumber(30), );
 	ADD_ATTRIBUTE_STORAGE(FFrameRate, FFrameRate(1,60), );
 	ADD_ATTRIBUTE_STORAGE(FFrameTime, FFrameTime(29,0.5f), );
@@ -205,7 +171,14 @@ bool FAttributeStorageTest::RunTest(const FString& Parameters)
 	ADD_ATTRIBUTE_STORAGE(FBoxSphereBounds3f, FBoxSphereBounds3f(FVector3f(0.0f), FVector3f(2.0f), 2.2f), );
 	ADD_ATTRIBUTE_STORAGE(FBoxSphereBounds3d, FBoxSphereBounds3d(FVector3d(0.0), FVector3d(2.0), 2.2), );
 	{
-		ADD_ATTRIBUTE_STORAGE_NOCOMPARE(FSphere3f, FSphere3f(FVector3f(2.0f), 2.2f))
+		FSphere3f RefValue;
+		RefValue.Center = FVector3f(2.0f);
+		RefValue.W = 2.2f;
+		FAttributeKey Key = CreateUniqueKey();
+		if (!IsAttributeStorageResultSuccess(TestStorage.RegisterAttribute(Key, RefValue)))
+		{
+			AddError(FString(TEXT("`AttributeStorage` must handle adding FSphere3f attribute")));
+		}
 		FSphere3f StoredValue;
 		uint32 bobafet = -45;
 		TestStorage.GetAttributeHandle<FSphere3f>(Key).Get(StoredValue);
@@ -215,7 +188,14 @@ bool FAttributeStorageTest::RunTest(const FString& Parameters)
 		}
 	}
 	{
-		ADD_ATTRIBUTE_STORAGE_NOCOMPARE(FSphere3d, FSphere3d(FVector3d(2.0), 2.2))
+		FSphere3d RefValue;
+		RefValue.Center = FVector3d(2.0);
+		RefValue.W = 2.2;
+		FAttributeKey Key = CreateUniqueKey();
+		if (!IsAttributeStorageResultSuccess(TestStorage.RegisterAttribute(Key, RefValue)))
+		{
+			AddError(FString(TEXT("`AttributeStorage` must handle adding FSphere3d attribute")));
+		}
 		FSphere3d StoredValue;
 		uint32 bobafet = -45;
 		TestStorage.GetAttributeHandle<FSphere3d>(Key).Get(StoredValue);
