split cached computation data out from voronoi container, so that voro++ can compute voronoi cells in parallel, and update the Voronoi class to support parallel versions of computing all cells / computing cell edges

#rnx

diff --git a/Engine/Source/Runtime/Experimental/Voronoi/Private/Voronoi/Voronoi.cpp b/Engine/Source/Runtime/Experimental/Voronoi/Private/Voronoi/Voronoi.cpp
--- a/Engine/Source/Runtime/Experimental/Voronoi/Private/Voronoi/Voronoi.cpp
+++ b/Engine/Source/Runtime/Experimental/Voronoi/Private/Voronoi/Voronoi.cpp
@@ -1,6 +1,7 @@
 // Copyright Epic Games, Inc. All Rights Reserved.
 #include "Voronoi/Voronoi.h"
 
+#include "Async/ParallelFor.h"
 
 THIRD_PARTY_INCLUDES_START
 #include "voro++/voro++.hh"
@@ -44,6 +45,7 @@ namespace {
 	int32 PutSitesWithDistanceCheck(voro::container *Container, const TArrayView<const FVector>& Sites, int32 Offset, float SquaredDistThreshold = 1e-4)
 	{
 		int32 SkippedPts = 0;
+		voro::voro_compute<voro::container> VoroCompute = Container->make_compute();
 		for (int i = 0; i < Sites.Num(); i++)
 		{
 			const FVector &V = Sites[i];
@@ -56,7 +58,7 @@ namespace {
 			{
 				double EX, EY, EZ;
 				int ExistingPtID;
-				if (Container->find_voronoi_cell(V.X, V.Y, V.Z, EX, EY, EZ, ExistingPtID))
+				if (Container->find_voronoi_cell(V.X, V.Y, V.Z, EX, EY, EZ, ExistingPtID, VoroCompute))
 				{
 					float dx = V.X - EX;
 					float dy = V.Y - EY;
@@ -109,6 +111,7 @@ namespace {
 
 FVoronoiDiagram::~FVoronoiDiagram() = default;
 
+// TODO: ExtraBoundingSpace unused
 FVoronoiDiagram::FVoronoiDiagram(const TArrayView<const FVector>& Sites, float ExtraBoundingSpace, float SquaredDistSkipPtThreshold) : NumSites(0), Container(nullptr)
 {
 	FBox BoundingBox = SafeInitBounds(Sites);
@@ -117,6 +120,7 @@ FVoronoiDiagram::FVoronoiDiagram(const TArrayView<const FVector>& Sites, float E
 	NumSites = Sites.Num();
 }
 
+// TODO: ExtraBoundingSpace unused
 FVoronoiDiagram::FVoronoiDiagram(const TArrayView<const FVector>& Sites, const FBox &Bounds, float ExtraBoundingSpace, float SquaredDistSkipPtThreshold) : NumSites(0), Container(nullptr)
 {
 	FBox BoundingBox(Bounds);
@@ -136,13 +140,14 @@ bool VoronoiNeighbors(const TArrayView<const FVector> &Sites, TArray<TArray<int>
 	Neighbors.Empty(NumSites);
 	Neighbors.SetNum(NumSites, true);
 
+	voro::voro_compute<voro::container> VoroCompute = Container->make_compute();
 	voro::c_loop_all CellIterator(*Container);
 	voro::voronoicell_neighbor cell;
 	if (CellIterator.start())
 	{
 		do
 		{
-			bool bCouldComputeCell = Container->compute_cell(cell, CellIterator);
+			bool bCouldComputeCell = Container->compute_cell(cell, CellIterator, VoroCompute);
 			ensureMsgf(bCouldComputeCell, TEXT("Failed to compute a Voronoi cell -- this may indicate sites positioned directly on top of other sites, which is not valid for a Voronoi diagram"));
 			if (bCouldComputeCell)
 			{
@@ -156,93 +161,122 @@ bool VoronoiNeighbors(const TArrayView<const FVector> &Sites, TArray<TArray<int>
 	return true;
 }
 
-bool GetVoronoiEdges(const TArrayView<const FVector> &Sites, const FBox& Bounds, TArray<TTuple<FVector, FVector>> &Edges, TArray<int32> &CellMember, float SquaredDistSkipPtThreshold)
+bool GetVoronoiEdges(const TArrayView<const FVector>& Sites, const FBox& Bounds, TArray<TTuple<FVector, FVector>>& Edges, TArray<int32>& CellMember, float SquaredDistSkipPtThreshold)
 {
 	int32 NumSites = Sites.Num();
 	FBox BoundingBox(Bounds);
-	auto Container = StandardVoroContainerInit(Sites, BoundingBox, FVoronoiDiagram::DefaultBoundingBoxSlack, SquaredDistSkipPtThreshold);
+	FVoronoiDiagram Diagram(Sites, Bounds, FVoronoiDiagram::DefaultBoundingBoxSlack, SquaredDistSkipPtThreshold);
+	Diagram.ComputeCellEdges(Edges, CellMember, Diagram.ApproxSitesPerThreadWithDefault(-1));
+	return true;
+}
 
-	voro::container &con = *Container;
-	voro::c_loop_all CellIterator(con);
-	voro::voronoicell cell;
 
-	int32 id = 0;
+// TODO: maybe make this a "SetSites" instead of an Add?
+void FVoronoiDiagram::AddSites(const TArrayView<const FVector>& AddSites, float SquaredDistSkipPtThreshold)
+{
+	int32 OrigSitesNum = NumSites;
+	if (SquaredDistSkipPtThreshold > 0)
+	{
+		PutSitesWithDistanceCheck(Container.Get(), AddSites, OrigSitesNum, SquaredDistSkipPtThreshold);
+	}
+	else
+	{
+		PutSites(Container.Get(), AddSites, OrigSitesNum);
+	}
+	NumSites += AddSites.Num();
+}
+
+void FVoronoiDiagram::ComputeAllCellsSerial(TArray<FVoronoiCellInfo>& AllCells)
+{
+	check(Container);
+	voro::voro_compute<voro::container> VoroCompute = Container->make_compute();
 
-	std::vector<double> Vertices;
-	std::vector<int> FaceVertices;
+	AllCells.SetNum(NumSites);
+
+	voro::c_loop_all CellIterator(*Container);
+	voro::voronoicell_neighbor cell;
 
 	if (CellIterator.start())
 	{
 		do
 		{
-			bool bCouldComputeCell = Container->compute_cell(cell, CellIterator);
+			bool bCouldComputeCell = Container->compute_cell(cell, CellIterator, VoroCompute);
+			ensureMsgf(bCouldComputeCell, TEXT("Failed to compute a Voronoi cell -- this may indicate sites positioned directly on top of other sites, which is not valid for a Voronoi diagram"));
 			if (bCouldComputeCell)
 			{
-				const double *pp = con.p[CellIterator.ijk] + con.ps * CellIterator.q;
-				const FVector Center(*pp, pp[1], pp[2]);
-				
-				cell.vertices(Center.X, Center.Y, Center.Z, Vertices);
-				cell.face_vertices(FaceVertices);
-
-				int32 FaceOffset = 0;
-				for (size_t ii = 0, ni = FaceVertices.size(); ii < ni; ii += FaceVertices[ii] + 1)
-				{
-					int32 VertCount = FaceVertices[ii];
-					int32 PreviousVertexIndex = FaceVertices[FaceOffset + VertCount] * 3;
-					for (int32 kk = 0; kk < VertCount; ++kk)
-					{
-						CellMember.Emplace(id);
-						int32 VertexIndex = FaceVertices[1 + FaceOffset + kk] * 3; // Index of vertex X coordinate in raw coordinates array
-
-						Edges.Emplace(
-							FVector(Vertices[PreviousVertexIndex], Vertices[PreviousVertexIndex + 1], Vertices[PreviousVertexIndex + 2]),
-							FVector(Vertices[VertexIndex], Vertices[VertexIndex + 1], Vertices[VertexIndex + 2])
-						);
-						PreviousVertexIndex = VertexIndex;
-					}
-					FaceOffset += VertCount + 1;
-				}
+				int32 id = CellIterator.pid();
+				double x, y, z;
+				CellIterator.pos(x, y, z);
 
-				++id;
+				FVoronoiCellInfo& Cell = AllCells[id];
+				cell.extractCellInfo(FVector(x,y,z), Cell.Vertices, Cell.Faces, Cell.Neighbors, Cell.Normals);
 			}
 		} while (CellIterator.inc());
 	}
-	return true;
 }
 
+TArray<int32> FVoronoiDiagram::GetParallelBlockRanges(int32 ApproxSitesPerThread)
+{
+	TArray<int32> BlockBounds;
 
+	ApproxSitesPerThread = ApproxSitesPerThreadWithDefault(ApproxSitesPerThread);
+	if (ApproxSitesPerThread >= NumSites * 1.5)
+	{
+		BlockBounds.Add(0); // just single-thread; put everything in a single chunk
+		BlockBounds.Add(Container->nxyz);
+		return BlockBounds;
+	}
 
-
-// TODO: maybe make this a "SetSites" instead of an Add?
-void FVoronoiDiagram::AddSites(const TArrayView<const FVector>& AddSites, float SquaredDistSkipPtThreshold)
-{
-	int32 OrigSitesNum = NumSites;
-	if (SquaredDistSkipPtThreshold > 0)
+	BlockBounds.Reserve(2 + NumSites / ApproxSitesPerThread);
+	BlockBounds.Add(0);
+	int32 AccumSites = 0;
+	for (int32 ijk = 0; ijk < Container->nxyz; ijk++)
 	{
-		PutSitesWithDistanceCheck(Container.Get(), AddSites, OrigSitesNum, SquaredDistSkipPtThreshold);
+		if (AccumSites >= ApproxSitesPerThread)
+		{
+			BlockBounds.Add(ijk);
+			AccumSites = 0;
+		}
+		AccumSites += Container->co[ijk];
 	}
-	else
+	// final bound should cover the full container, either by extending the existing final bound to cover it or creating a new final bound
+	if (AccumSites < ApproxSitesPerThread * .5)
 	{
-		PutSites(Container.Get(), AddSites, OrigSitesNum);
+		BlockBounds.Last() = Container->nxyz;
 	}
-	NumSites += AddSites.Num();
+	else // enough sites to merit a separate 'thread'
+	{
+		BlockBounds.Add(Container->nxyz);
+	}
+	return BlockBounds;
 }
 
-void FVoronoiDiagram::ComputeAllCells(TArray<FVoronoiCellInfo> &AllCells)
+void FVoronoiDiagram::ComputeAllCells(TArray<FVoronoiCellInfo>& AllCells, int32 ApproxSitesPerThread)
 {
 	check(Container);
 
-	AllCells.SetNum(NumSites);
+	TArray<int32> ChunkBounds = GetParallelBlockRanges(ApproxSitesPerThread);
+	if (ChunkBounds.Num() < 3)
+	{
+		ComputeAllCellsSerial(AllCells);
+		return;
+	}
 
-	voro::c_loop_all CellIterator(*Container);
-	voro::voronoicell_neighbor cell;
+	AllCells.SetNum(NumSites);
 
-	// TODO: multithread?
-	if (CellIterator.start())
+	ParallelFor(ChunkBounds.Num() - 1, [&](int ChunkIdx)
 	{
+		voro::voro_compute<voro::container> VoroCompute = Container->make_compute();
+		voro::c_loop_block_range CellIterator(*Container);
+		CellIterator.setup_range(ChunkBounds[ChunkIdx], ChunkBounds[ChunkIdx + 1]);
+		voro::voronoicell_neighbor cell;
+		if (!CellIterator.start())
+		{
+			return;
+		}
 		do
 		{
-			bool bCouldComputeCell = Container->compute_cell(cell, CellIterator);
+			bool bCouldComputeCell = Container->compute_cell(cell, CellIterator, VoroCompute);
 			ensureMsgf(bCouldComputeCell, TEXT("Failed to compute a Voronoi cell -- this may indicate sites positioned directly on top of other sites, which is not valid for a Voronoi diagram"));
 			if (bCouldComputeCell)
 			{
@@ -251,25 +285,150 @@ void FVoronoiDiagram::ComputeAllCells(TArray<FVoronoiCellInfo> &AllCells)
 				CellIterator.pos(x, y, z);
 
 				FVoronoiCellInfo& Cell = AllCells[id];
+				check(Cell.Faces.Num() == 0); // TODO RM
 				cell.extractCellInfo(FVector(x,y,z), Cell.Vertices, Cell.Faces, Cell.Neighbors, Cell.Normals);
 			}
 		} while (CellIterator.inc());
+	});
+}
+
+void FVoronoiDiagram::ComputeCellEdgesSerial(TArray<TTuple<FVector, FVector>>& Edges, TArray<int32>& CellMember)
+{
+	check(Container);
+
+	TArray<FVector> Vertices;
+	TArray<int32> FacesVertices;
+
+	voro::c_loop_all CellIterator(*Container);
+	voro::voronoicell cell;
+	voro::voro_compute<voro::container> VoroCompute = Container->make_compute();
+
+	if (CellIterator.start())
+	{
+		if (!CellIterator.start())
+		{
+			return;
+		}
+		do
+		{
+			bool bCouldComputeCell = Container->compute_cell(cell, CellIterator, VoroCompute);
+			ensureMsgf(bCouldComputeCell, TEXT("Failed to compute a Voronoi cell -- this may indicate sites positioned directly on top of other sites, which is not valid for a Voronoi diagram"));
+			if (bCouldComputeCell)
+			{
+				int32 ID = CellIterator.pid();
+				double X, Y, Z;
+				CellIterator.pos(X, Y, Z);
+				cell.extractCellInfo(FVector(X, Y, Z), Vertices, FacesVertices);
+
+				int32 FaceOffset = 0;
+				for (size_t ii = 0, ni = FacesVertices.Num(); ii < ni; ii += FacesVertices[ii] + 1)
+				{
+					int32 VertCount = FacesVertices[ii];
+					int32 PreviousVertexIndex = FacesVertices[FaceOffset + VertCount];
+					for (int32 kk = 0; kk < VertCount; ++kk)
+					{
+						CellMember.Emplace(ID);
+						int32 VertexIndex = FacesVertices[1 + FaceOffset + kk]; // Index of vertex X coordinate in raw coordinates array
+
+						Edges.Emplace(Vertices[PreviousVertexIndex], Vertices[VertexIndex]);
+						PreviousVertexIndex = VertexIndex;
+					}
+					FaceOffset += VertCount + 1;
+				}
+			}
+		} while (CellIterator.inc());
 	}
 }
 
-int32 FVoronoiDiagram::FindCell(const FVector& Pos)
+void FVoronoiDiagram::ComputeCellEdges(TArray<TTuple<FVector, FVector>>& Edges, TArray<int32>& CellMember, int32 ApproxSitesPerThread)
 {
 	check(Container);
-	double rx, ry, rz; // holds position of the found Voronoi site (if any); currently unused
-	int pid;
-	bool found = Container->find_voronoi_cell(Pos.X, Pos.Y, Pos.Z, rx, ry, rz, pid);
-	return found ? pid : -1;
+
+	TArray<int32> ChunkBounds = GetParallelBlockRanges(ApproxSitesPerThread);
+
+	if (ChunkBounds.Num() < 3)
+	{
+		ComputeCellEdgesSerial(Edges, CellMember);
+		return;
+	}
+
+	struct FChunkOutput
+	{
+		TArray<TTuple<FVector, FVector>> Edges;
+		TArray<int32> CellMember;
+	};
+	TArray<FChunkOutput> ChunkOutputs;
+	ChunkOutputs.SetNum(ChunkBounds.Num() - 1);
+
+	ParallelFor(ChunkBounds.Num() - 1, [this, &ChunkOutputs, &ChunkBounds](int ChunkIdx)
+	{
+		FChunkOutput& Chunk = ChunkOutputs[ChunkIdx];
+		TArray<FVector> Vertices;
+		TArray<int32> FacesVertices;
+
+		voro::voro_compute<voro::container> VoroCompute = Container->make_compute();
+		voro::c_loop_block_range CellIterator(*Container);
+		CellIterator.setup_range(ChunkBounds[ChunkIdx], ChunkBounds[ChunkIdx + 1]);
+		voro::voronoicell cell;
+		if (!CellIterator.start())
+		{
+			return;
+		}
+		do
+		{
+			bool bCouldComputeCell = Container->compute_cell(cell, CellIterator, VoroCompute);
+			ensureMsgf(bCouldComputeCell, TEXT("Failed to compute a Voronoi cell -- this may indicate sites positioned directly on top of other sites, which is not valid for a Voronoi diagram"));
+			if (bCouldComputeCell)
+			{
+				int32 ID = CellIterator.pid();
+				double X, Y, Z;
+				CellIterator.pos(X, Y, Z);
+				cell.extractCellInfo(FVector(X, Y, Z), Vertices, FacesVertices);
+
+				int32 FaceOffset = 0;
+				for (size_t ii = 0, ni = FacesVertices.Num(); ii < ni; ii += FacesVertices[ii] + 1)
+				{
+					int32 VertCount = FacesVertices[ii];
+					int32 PreviousVertexIndex = FacesVertices[FaceOffset + VertCount];
+					for (int32 kk = 0; kk < VertCount; ++kk)
+					{
+						Chunk.CellMember.Emplace(ID);
+						int32 VertexIndex = FacesVertices[1 + FaceOffset + kk]; // Index of vertex X coordinate in raw coordinates array
+
+						Chunk.Edges.Emplace(Vertices[PreviousVertexIndex], Vertices[VertexIndex]);
+						PreviousVertexIndex = VertexIndex;
+					}
+					FaceOffset += VertCount + 1;
+				}
+			}
+		} while (CellIterator.inc());
+	});
+
+	int ToReserve = 0;
+	for (const FChunkOutput& Chunk : ChunkOutputs)
+	{
+		ToReserve += Chunk.Edges.Num();
+	}
+	Edges.Reset(ToReserve);
+	CellMember.Reset(ToReserve);
+
+	for (const FChunkOutput& Chunk : ChunkOutputs)
+	{
+		Edges.Append(Chunk.Edges);
+		CellMember.Append(Chunk.CellMember);
+	}
 }
 
-//void FVoronoiDiagram::GetCells(const TArrayView<const FVector> &Sites, TArray<FVoronoiCellInfo> &AllCells)
+//// TODO: expose separate compute data as something we can pass into this fn, and uncomment it
+//int32 FVoronoiDiagram::FindCell(const FVector& Pos)
 //{
-//	FVoronoiDiagram Voronoi(Sites, FVoronoiDiagram::DefaultBoundingBoxSlack);
-//	Voronoi.ComputeAllCells(AllCells);
+//	check(Container);
+//	voro::voro_compute<voro::container> VoroCompute = Container->make_compute();
+//	double rx, ry, rz; // holds position of the found Voronoi site (if any); currently unused
+//	int pid;
+//	bool found = Container->find_voronoi_cell(Pos.X, Pos.Y, Pos.Z, rx, ry, rz, pid, VoroCompute);
+//	return found ? pid : -1;
 //}
 
-const float FVoronoiDiagram::DefaultBoundingBoxSlack = .1f;
\ No newline at end of file
+const float FVoronoiDiagram::DefaultBoundingBoxSlack = .1f;
+const int FVoronoiDiagram::MinDefaultSitesPerThread = 150;
\ No newline at end of file
diff --git a/Engine/Source/Runtime/Experimental/Voronoi/Private/voro++/src/cell.cc b/Engine/Source/Runtime/Experimental/Voronoi/Private/voro++/src/cell.cc
--- a/Engine/Source/Runtime/Experimental/Voronoi/Private/voro++/src/cell.cc
+++ b/Engine/Source/Runtime/Experimental/Voronoi/Private/voro++/src/cell.cc
@@ -2067,7 +2067,7 @@ void voronoicell_neighbor::neighborsTArray(TArray<int> &nbrs, bool excludeBounds
 	reset_edges();
 }
 /** Computes Unreal TArrays of all cell info. */
-void voronoicell_neighbor::extractCellInfo(FVector CellPosition, TArray<FVector> &Vertices, TArray<int32> &FaceVertexIndices, TArray<int32> &Nbrs, TArray<FVector> &Normals)
+void voronoicell_neighbor::extractCellInfo(const FVector& CellPosition, TArray<FVector> &Vertices, TArray<int32> &FaceVertexIndices, TArray<int32> &Nbrs, TArray<FVector> &Normals)
 {
 	Vertices.SetNumUninitialized(p);
 	double *ptsp = pts;
@@ -2079,7 +2079,7 @@ void voronoicell_neighbor::extractCellInfo(FVector CellPosition, TArray<FVector>
 	}
 	
 	
-	FaceVertexIndices.Empty(); Nbrs.Empty(); Normals.Empty();
+	FaceVertexIndices.Reset(); Nbrs.Reset(); Normals.Reset();
 
 	// Vertex indices code (face_vertices)
 	int i, j, k, l, m, vp(0), vn;
@@ -2146,6 +2146,46 @@ void voronoicell_neighbor::extractCellInfo(FVector CellPosition, TArray<FVector>
 	}
 	reset_edges();
 }
+/** Computes Unreal TArrays of all cell info. */
+void voronoicell::extractCellInfo(const FVector& CellPosition, TArray<FVector>& Vertices, TArray<int32>& FaceVertexIndices)
+{
+	Vertices.SetNumUninitialized(p);
+	double* ptsp = pts;
+	for (int i = 0; i < p; i++)
+	{
+		Vertices[i].X = CellPosition.X + *(ptsp++) * 0.5;
+		Vertices[i].Y = CellPosition.Y + *(ptsp++) * 0.5;
+		Vertices[i].Z = CellPosition.Z + *(ptsp++) * 0.5;
+	}
+
+	FaceVertexIndices.Reset();
+
+	// Vertex indices code (face_vertices)
+	int i, j, k, l, m, vp(0), vn;
+	for (i = 1; i < p; i++) for (j = 0; j < nu[i]; j++) {
+		k = ed[i][j];
+		if (k >= 0) { // on a new face
+			FaceVertexIndices.Add(0); // placeholder for face vertex count
+			FaceVertexIndices.Add(i);
+
+			ed[i][j] = -1 - k;
+			l = cycle_up(ed[i][nu[i] + j], k);
+			do {
+				FaceVertexIndices.Add(k);
+				m = ed[k][l];
+				ed[k][l] = -1 - m;
+
+				l = cycle_up(ed[k][nu[k] + l], m);
+				k = m;
+			} while (k != i);
+
+			vn = FaceVertexIndices.Num();
+			FaceVertexIndices[vp] = vn - vp - 1;
+			vp = vn;
+		}
+	}
+	reset_edges();
+}
 
 /** Prints the vertices, their edges, the relation table, and also notifies if
  * any memory errors are visible. */
diff --git a/Engine/Source/Runtime/Experimental/Voronoi/Private/voro++/src/container.cc b/Engine/Source/Runtime/Experimental/Voronoi/Private/voro++/src/container.cc
--- a/Engine/Source/Runtime/Experimental/Voronoi/Private/voro++/src/container.cc
+++ b/Engine/Source/Runtime/Experimental/Voronoi/Private/voro++/src/container.cc
@@ -54,6 +54,17 @@ container_base::~container_base() {
 	delete [] mem;
 }
 
+voro_compute<container> container::make_compute() const
+{
+	return voro_compute<container>((container&)*this, xperiodic ? 2 * nx + 1 : nx, yperiodic ? 2 * ny + 1 : ny, zperiodic ? 2 * nz + 1 : nz);
+}
+
+
+voro_compute<container_poly> container_poly::make_compute() const
+{
+	return voro_compute<container_poly>((container_poly&)*this, xperiodic ? 2 * nx + 1 : nx, yperiodic ? 2 * ny + 1 : ny, zperiodic ? 2 * nz + 1 : nz);
+}
+
 /** The class constructor sets up the geometry of container.
  * \param[in] (ax_,bx_) the minimum and maximum x coordinates.
  * \param[in] (ay_,by_) the minimum and maximum y coordinates.
@@ -66,8 +77,7 @@ container_base::~container_base() {
  * \param[in] init_mem the initial memory allocation for each block. */
 container::container(double ax_,double bx_,double ay_,double by_,double az_,double bz_,
 	int nx_,int ny_,int nz_,bool xperiodic_,bool yperiodic_,bool zperiodic_,int init_mem)
-	: container_base(ax_,bx_,ay_,by_,az_,bz_,nx_,ny_,nz_,xperiodic_,yperiodic_,zperiodic_,init_mem,3),
-	vc(*this,xperiodic_?2*nx_+1:nx_,yperiodic_?2*ny_+1:ny_,zperiodic_?2*nz_+1:nz_) {}
+	: container_base(ax_,bx_,ay_,by_,az_,bz_,nx_,ny_,nz_,xperiodic_,yperiodic_,zperiodic_,init_mem,3) {}
 
 /** The class constructor sets up the geometry of container.
  * \param[in] (ax_,bx_) the minimum and maximum x coordinates.
@@ -81,8 +91,7 @@ container::container(double ax_,double bx_,double ay_,double by_,double az_,doub
  * \param[in] init_mem the initial memory allocation for each block. */
 container_poly::container_poly(double ax_,double bx_,double ay_,double by_,double az_,double bz_,
 	int nx_,int ny_,int nz_,bool xperiodic_,bool yperiodic_,bool zperiodic_,int init_mem)
-	: container_base(ax_,bx_,ay_,by_,az_,bz_,nx_,ny_,nz_,xperiodic_,yperiodic_,zperiodic_,init_mem,4),
-	vc(*this,xperiodic_?2*nx_+1:nx_,yperiodic_?2*ny_+1:ny_,zperiodic_?2*nz_+1:nz_) {ppr=p;}
+	: container_base(ax_,bx_,ay_,by_,az_,bz_,nx_,ny_,nz_,xperiodic_,yperiodic_,zperiodic_,init_mem,4) {ppr=p;}
 
 /** Put a particle into the correct region of the container.
  * \param[in] n the numerical ID of the inserted particle.
@@ -171,7 +180,7 @@ bool container_base::put_locate_block(int &ijk,double &x,double &y,double &z) {
  *                        domain if necessary.
  * \return True if the particle can be successfully placed into the container,
  * false otherwise. */
-inline bool container_base::put_remap(int &ijk,double &x,double &y,double &z) {
+inline bool container_base::put_remap(int &ijk,double &x,double &y,double &z) const {
 	int l;
 
 	ijk=step_int((x-ax)*xsp);
@@ -200,7 +209,7 @@ inline bool container_base::put_remap(int &ijk,double &x,double &y,double &z) {
  * \param[out] ijk the block index that the vector is within.
  * \return True if the particle is within the container or can be remapped into
  * it, false if it lies outside of the container bounds. */
-inline bool container_base::remap(int &ai,int &aj,int &ak,int &ci,int &cj,int &ck,double &x,double &y,double &z,int &ijk) {
+inline bool container_base::remap(int &ai,int &aj,int &ak,int &ci,int &cj,int &ck,double &x,double &y,double &z,int &ijk) const {
 	ci=step_int((x-ax)*xsp);
 	if(ci<0||ci>=nx) {
 		if(xperiodic) {ai=step_div(ci,nx);x-=ai*(bx-ax);ci-=ai*nx;}
@@ -234,7 +243,7 @@ inline bool container_base::remap(int &ai,int &aj,int &ak,int &ci,int &cj,int &c
  * \param[out] pid the ID of the particle.
  * \return True if a particle was found. If the container has no particles,
  * then the search will not find a Voronoi cell and false is returned. */
-bool container::find_voronoi_cell(double x,double y,double z,double &rx,double &ry,double &rz,int &pid) {
+bool container::find_voronoi_cell(double x,double y,double z,double &rx,double &ry,double &rz,int &pid,voro_compute<container>& vc) const {
 	int ai,aj,ak,ci,cj,ck,ijk;
 	particle_record w;
 	double mrs;
@@ -272,7 +281,7 @@ bool container::find_voronoi_cell(double x,double y,double z,double &rx,double &
  * \param[out] pid the ID of the particle.
  * \return True if a particle was found. If the container has no particles,
  * then the search will not find a Voronoi cell and false is returned. */
-bool container_poly::find_voronoi_cell(double x,double y,double z,double &rx,double &ry,double &rz,int &pid) {
+bool container_poly::find_voronoi_cell(double x,double y,double z,double &rx,double &ry,double &rz,int &pid,voro_compute<container>& vc) const {
 	int ai,aj,ak,ci,cj,ck,ijk;
 	particle_record w;
 	double mrs;
@@ -327,7 +336,7 @@ void container_base::add_particle_memory(int i) {
 
 /** Outputs the a list of all the container regions along with the number of
  * particles stored within each. */
-void container_base::region_count() {
+void container_base::region_count() const {
 	int i,j,k,*cop=co;
 	for(k=0;k<nz;k++) for(j=0;j<ny;j++) for(i=0;i<nx;i++)
 		printf("Region (%d,%d,%d): %d particles\n",i,j,k,*(cop++));
@@ -349,10 +358,10 @@ void container_poly::clear() {
  * with the output. It is useful for measuring the pure computation time
  * of the Voronoi algorithm, without any additional calculations such as
  * volume evaluation or cell output. */
-void container::compute_all_cells() {
+void container::compute_all_cells(voro_compute<container>& vc) {
 	voronoicell c;
 	c_loop_all vl(*this);
-	if(vl.start()) do compute_cell(c,vl);
+	if(vl.start()) do compute_cell(c,vl,vc);
 	while(vl.inc());
 }
 
@@ -360,21 +369,21 @@ void container::compute_all_cells() {
  * with the output. It is useful for measuring the pure computation time
  * of the Voronoi algorithm, without any additional calculations such as
  * volume evaluation or cell output. */
-void container_poly::compute_all_cells() {
+void container_poly::compute_all_cells(voro_compute<container>& vc) {
 	voronoicell c;
 	c_loop_all vl(*this);
-	if(vl.start()) do compute_cell(c,vl);while(vl.inc());
+	if(vl.start()) do compute_cell(c,vl,vc);while(vl.inc());
 }
 
 /** Calculates all of the Voronoi cells and sums their volumes. In most cases
  * without walls, the sum of the Voronoi cell volumes should equal the volume
  * of the container to numerical precision.
  * \return The sum of all of the computed Voronoi volumes. */
-double container::sum_cell_volumes() {
+double container::sum_cell_volumes(voro_compute<container>& vc) {
 	voronoicell c;
 	double vol=0;
 	c_loop_all vl(*this);
-	if(vl.start()) do if(compute_cell(c,vl)) vol+=c.volume();while(vl.inc());
+	if(vl.start()) do if(compute_cell(c,vl,vc)) vol+=c.volume();while(vl.inc());
 	return vol;
 }
 
@@ -382,11 +391,11 @@ double container::sum_cell_volumes() {
  * without walls, the sum of the Voronoi cell volumes should equal the volume
  * of the container to numerical precision.
  * \return The sum of all of the computed Voronoi volumes. */
-double container_poly::sum_cell_volumes() {
+double container_poly::sum_cell_volumes(voro_compute<container>& vc) {
 	voronoicell c;
 	double vol=0;
 	c_loop_all vl(*this);
-	if(vl.start()) do if(compute_cell(c,vl)) vol+=c.volume();while(vl.inc());
+	if(vl.start()) do if(compute_cell(c,vl,vc)) vol+=c.volume();while(vl.inc());
 	return vol;
 }
 
@@ -395,7 +404,7 @@ double container_poly::sum_cell_volumes() {
  * \param[in] (x,y,z) the position vector to be tested.
  * \return True if the point is inside the container, false if the point is
  *         outside. */
-bool container_base::point_inside(double x,double y,double z) {
+bool container_base::point_inside(double x,double y,double z) const {
 	if(x<ax||x>bx||y<ay||y>by||z<az||z>bz) return false;
 	return point_inside_walls(x,y,z);
 }
diff --git a/Engine/Source/Runtime/Experimental/Voronoi/Private/voro++/src/container_prd.cc b/Engine/Source/Runtime/Experimental/Voronoi/Private/voro++/src/container_prd.cc
deleted file mode 100644
diff --git a/Engine/Source/Runtime/Experimental/Voronoi/Private/voro++/src/v_base.cc b/Engine/Source/Runtime/Experimental/Voronoi/Private/voro++/src/v_base.cc
--- a/Engine/Source/Runtime/Experimental/Voronoi/Private/voro++/src/v_base.cc
+++ b/Engine/Source/Runtime/Experimental/Voronoi/Private/voro++/src/v_base.cc
@@ -100,7 +100,7 @@ void voro_base::compute_minimum(double &minr,double &xlo,double &xhi,double &ylo
  * whether neighbor information is required or not.
  * \param[in] format the format string to check.
  * \return True if a "%n" is found, false otherwise. */
-bool voro_base::contains_neighbor(const char *format) {
+bool voro_base::contains_neighbor(const char *format) const {
 	char *fmp=(const_cast<char*>(format));
 
 	// Check to see if "%n" appears in the format sequence
diff --git a/Engine/Source/Runtime/Experimental/Voronoi/Private/voro++/src/v_compute.cc b/Engine/Source/Runtime/Experimental/Voronoi/Private/voro++/src/v_compute.cc
--- a/Engine/Source/Runtime/Experimental/Voronoi/Private/voro++/src/v_compute.cc
+++ b/Engine/Source/Runtime/Experimental/Voronoi/Private/voro++/src/v_compute.cc
@@ -14,7 +14,6 @@
 #include "v_compute.hh"
 #include "rad_option.hh"
 #include "container.hh"
-#include "container_prd.hh"
 
 namespace voro {
 
@@ -300,7 +299,8 @@ bool voro_compute<c_class>::compute_cell(v_cell &c,int ijk,int s,int ci,int cj,i
 	unsigned int q,*e,*mijk;
 
 	if(!con.initialize_voronoicell(c,ijk,s,ci,cj,ck,i,j,k,x,y,z,disp)) return false;
-	con.r_init(ijk,s);
+	typename c_class::radius_info radi;
+	con.r_init(ijk,s,radi);
 
 	// Initialize the Voronoi cell to fill the entire container
 	double crs,mrs;
@@ -312,7 +312,7 @@ bool voro_compute<c_class>::compute_cell(v_cell &c,int ijk,int s,int ci,int cj,i
 		x1=p[ijk][ps*l]-x;
 		y1=p[ijk][ps*l+1]-y;
 		z1=p[ijk][ps*l+2]-z;
-		rs=con.r_scale(x1*x1+y1*y1+z1*z1,ijk,l);
+		rs=con.r_scale(x1*x1+y1*y1+z1*z1,ijk,l,radi);
 		if(!c.nplane(x1,y1,z1,rs,id[ijk][l])) return false;
 	}
 	l++;
@@ -320,7 +320,7 @@ bool voro_compute<c_class>::compute_cell(v_cell &c,int ijk,int s,int ci,int cj,i
 		x1=p[ijk][ps*l]-x;
 		y1=p[ijk][ps*l+1]-y;
 		z1=p[ijk][ps*l+2]-z;
-		rs=con.r_scale(x1*x1+y1*y1+z1*z1,ijk,l);
+		rs=con.r_scale(x1*x1+y1*y1+z1*z1,ijk,l,radi);
 		if(!c.nplane(x1,y1,z1,rs,id[ijk][l])) return false;
 		l++;
 	}
@@ -378,7 +378,7 @@ bool voro_compute<c_class>::compute_cell(v_cell &c,int ijk,int s,int ci,int cj,i
 
 		// If mrs is less than the minimum distance to any untested
 		// block, then we are done
-		if(con.r_ctest(radp[g],mrs)) return true;
+		if(con.r_ctest(radp[g],mrs,radi)) return true;
 		g++;
 
 		// Load in a block off the worklist, permute it with the
@@ -399,7 +399,7 @@ bool voro_compute<c_class>::compute_cell(v_cell &c,int ijk,int s,int ci,int cj,i
 		// current mrs, in which case we skip this block and move on.
 		// Otherwise, it computes the maximum distance to the block and
 		// returns it in crs.
-		if(compute_min_max_radius(di,dj,dk,fx,fy,fz,gxs,gys,gzs,crs,mrs)) continue;
+		if(compute_min_max_radius(di,dj,dk,fx,fy,fz,gxs,gys,gzs,crs,mrs,radi)) continue;
 
 		// Now compute which region we are going to loop over, adding a
 		// displacement for the periodic cases
@@ -411,12 +411,12 @@ bool voro_compute<c_class>::compute_cell(v_cell &c,int ijk,int s,int ci,int cj,i
 		// those particles which can't possibly intersect the block.
 		if(co[ijk]>0) {
 			l=0;x2=x-qx;y2=y-qy;z2=z-qz;
-			if(!con.r_ctest(crs,mrs)) {
+			if(!con.r_ctest(crs,mrs,radi)) {
 				do {
 					x1=p[ijk][ps*l]-x2;
 					y1=p[ijk][ps*l+1]-y2;
 					z1=p[ijk][ps*l+2]-z2;
-					rs=con.r_scale(x1*x1+y1*y1+z1*z1,ijk,l);
+					rs=con.r_scale(x1*x1+y1*y1+z1*z1,ijk,l,radi);
 					if(!c.nplane(x1,y1,z1,rs,id[ijk][l])) return false;
 					l++;
 				} while (l<co[ijk]);
@@ -426,7 +426,7 @@ bool voro_compute<c_class>::compute_cell(v_cell &c,int ijk,int s,int ci,int cj,i
 					y1=p[ijk][ps*l+1]-y2;
 					z1=p[ijk][ps*l+2]-z2;
 					rs=x1*x1+y1*y1+z1*z1;
-					if(con.r_scale_check(rs,mrs,ijk,l)&&!c.nplane(x1,y1,z1,rs,id[ijk][l])) return false;
+					if(con.r_scale_check(rs,mrs,ijk,l,radi)&&!c.nplane(x1,y1,z1,rs,id[ijk][l])) return false;
 					l++;
 				} while (l<co[ijk]);
 			}
@@ -460,7 +460,7 @@ bool voro_compute<c_class>::compute_cell(v_cell &c,int ijk,int s,int ci,int cj,i
 
 		// If mrs is less than the minimum distance to any untested
 		// block, then we are done
-		if(con.r_ctest(radp[g],mrs)) return true;
+		if(con.r_ctest(radp[g],mrs,radi)) return true;
 		g++;
 
 		// Load in a block off the worklist, permute it with the
@@ -485,7 +485,7 @@ bool voro_compute<c_class>::compute_cell(v_cell &c,int ijk,int s,int ci,int cj,i
 		// current mrs, in which case we skip this block and move on.
 		// Otherwise, it computes the maximum distance to the block and
 		// returns it in crs.
-		if(compute_min_max_radius(di,dj,dk,fx,fy,fz,gxs,gys,gzs,crs,mrs)) continue;
+		if(compute_min_max_radius(di,dj,dk,fx,fy,fz,gxs,gys,gzs,crs,mrs,radi)) continue;
 
 		// Now compute which region we are going to loop over, adding a
 		// displacement for the periodic cases
@@ -497,12 +497,12 @@ bool voro_compute<c_class>::compute_cell(v_cell &c,int ijk,int s,int ci,int cj,i
 		// those particles which can't possibly intersect the block.
 		if(co[ijk]>0) {
 			l=0;x2=x-qx;y2=y-qy;z2=z-qz;
-			if(!con.r_ctest(crs,mrs)) {
+			if(!con.r_ctest(crs,mrs,radi)) {
 				do {
 					x1=p[ijk][ps*l]-x2;
 					y1=p[ijk][ps*l+1]-y2;
 					z1=p[ijk][ps*l+2]-z2;
-					rs=con.r_scale(x1*x1+y1*y1+z1*z1,ijk,l);
+					rs=con.r_scale(x1*x1+y1*y1+z1*z1,ijk,l,radi);
 					if(!c.nplane(x1,y1,z1,rs,id[ijk][l])) return false;
 					l++;
 				} while (l<co[ijk]);
@@ -512,7 +512,7 @@ bool voro_compute<c_class>::compute_cell(v_cell &c,int ijk,int s,int ci,int cj,i
 					y1=p[ijk][ps*l+1]-y2;
 					z1=p[ijk][ps*l+2]-z2;
 					rs=x1*x1+y1*y1+z1*z1;
-					if(con.r_scale_check(rs,mrs,ijk,l)&&!c.nplane(x1,y1,z1,rs,id[ijk][l])) return false;
+					if(con.r_scale_check(rs,mrs,ijk,l,radi)&&!c.nplane(x1,y1,z1,rs,id[ijk][l])) return false;
 					l++;
 				} while (l<co[ijk]);
 			}
@@ -529,7 +529,7 @@ bool voro_compute<c_class>::compute_cell(v_cell &c,int ijk,int s,int ci,int cj,i
 	}
 
 	// Do a check to see if we've reached the radius cutoff
-	if(con.r_ctest(radp[g],mrs)) return true;
+	if(con.r_ctest(radp[g],mrs,radi)) return true;
 
 	// We were unable to completely compute the cell based on the blocks in
 	// the worklist, so now we have to go block by block, reading in items
@@ -551,44 +551,44 @@ bool voro_compute<c_class>::compute_cell(v_cell &c,int ijk,int s,int ci,int cj,i
 		// could possibly intersect the cell
 		if(ei>i) {
 			if(ej>j) {
-				if(ek>k) {if(corner_test(c,xlo,ylo,zlo,xhi,yhi,zhi)) continue;}
-				else if(ek<k) {if(corner_test(c,xlo,ylo,zhi,xhi,yhi,zlo)) continue;}
-				else {if(edge_z_test(c,xlo,ylo,zlo,xhi,yhi,zhi)) continue;}
+				if(ek>k) {if(corner_test(c,xlo,ylo,zlo,xhi,yhi,zhi,radi)) continue;}
+				else if(ek<k) {if(corner_test(c,xlo,ylo,zhi,xhi,yhi,zlo,radi)) continue;}
+				else {if(edge_z_test(c,xlo,ylo,zlo,xhi,yhi,zhi,radi)) continue;}
 			} else if(ej<j) {
-				if(ek>k) {if(corner_test(c,xlo,yhi,zlo,xhi,ylo,zhi)) continue;}
-				else if(ek<k) {if(corner_test(c,xlo,yhi,zhi,xhi,ylo,zlo)) continue;}
-				else {if(edge_z_test(c,xlo,yhi,zlo,xhi,ylo,zhi)) continue;}
+				if(ek>k) {if(corner_test(c,xlo,yhi,zlo,xhi,ylo,zhi,radi)) continue;}
+				else if(ek<k) {if(corner_test(c,xlo,yhi,zhi,xhi,ylo,zlo,radi)) continue;}
+				else {if(edge_z_test(c,xlo,yhi,zlo,xhi,ylo,zhi,radi)) continue;}
 			} else {
-				if(ek>k) {if(edge_y_test(c,xlo,ylo,zlo,xhi,yhi,zhi)) continue;}
-				else if(ek<k) {if(edge_y_test(c,xlo,ylo,zhi,xhi,yhi,zlo)) continue;}
-				else {if(face_x_test(c,xlo,ylo,zlo,yhi,zhi)) continue;}
+				if(ek>k) {if(edge_y_test(c,xlo,ylo,zlo,xhi,yhi,zhi,radi)) continue;}
+				else if(ek<k) {if(edge_y_test(c,xlo,ylo,zhi,xhi,yhi,zlo,radi)) continue;}
+				else {if(face_x_test(c,xlo,ylo,zlo,yhi,zhi,radi)) continue;}
 			}
 		} else if(ei<i) {
 			if(ej>j) {
-				if(ek>k) {if(corner_test(c,xhi,ylo,zlo,xlo,yhi,zhi)) continue;}
-				else if(ek<k) {if(corner_test(c,xhi,ylo,zhi,xlo,yhi,zlo)) continue;}
-				else {if(edge_z_test(c,xhi,ylo,zlo,xlo,yhi,zhi)) continue;}
+				if(ek>k) {if(corner_test(c,xhi,ylo,zlo,xlo,yhi,zhi,radi)) continue;}
+				else if(ek<k) {if(corner_test(c,xhi,ylo,zhi,xlo,yhi,zlo,radi)) continue;}
+				else {if(edge_z_test(c,xhi,ylo,zlo,xlo,yhi,zhi,radi)) continue;}
 			} else if(ej<j) {
-				if(ek>k) {if(corner_test(c,xhi,yhi,zlo,xlo,ylo,zhi)) continue;}
-				else if(ek<k) {if(corner_test(c,xhi,yhi,zhi,xlo,ylo,zlo)) continue;}
-				else {if(edge_z_test(c,xhi,yhi,zlo,xlo,ylo,zhi)) continue;}
+				if(ek>k) {if(corner_test(c,xhi,yhi,zlo,xlo,ylo,zhi,radi)) continue;}
+				else if(ek<k) {if(corner_test(c,xhi,yhi,zhi,xlo,ylo,zlo,radi)) continue;}
+				else {if(edge_z_test(c,xhi,yhi,zlo,xlo,ylo,zhi,radi)) continue;}
 			} else {
-				if(ek>k) {if(edge_y_test(c,xhi,ylo,zlo,xlo,yhi,zhi)) continue;}
-				else if(ek<k) {if(edge_y_test(c,xhi,ylo,zhi,xlo,yhi,zlo)) continue;}
-				else {if(face_x_test(c,xhi,ylo,zlo,yhi,zhi)) continue;}
+				if(ek>k) {if(edge_y_test(c,xhi,ylo,zlo,xlo,yhi,zhi,radi)) continue;}
+				else if(ek<k) {if(edge_y_test(c,xhi,ylo,zhi,xlo,yhi,zlo,radi)) continue;}
+				else {if(face_x_test(c,xhi,ylo,zlo,yhi,zhi,radi)) continue;}
 			}
 		} else {
 			if(ej>j) {
-				if(ek>k) {if(edge_x_test(c,xlo,ylo,zlo,xhi,yhi,zhi)) continue;}
-				else if(ek<k) {if(edge_x_test(c,xlo,ylo,zhi,xhi,yhi,zlo)) continue;}
-				else {if(face_y_test(c,xlo,ylo,zlo,xhi,zhi)) continue;}
+				if(ek>k) {if(edge_x_test(c,xlo,ylo,zlo,xhi,yhi,zhi,radi)) continue;}
+				else if(ek<k) {if(edge_x_test(c,xlo,ylo,zhi,xhi,yhi,zlo,radi)) continue;}
+				else {if(face_y_test(c,xlo,ylo,zlo,xhi,zhi,radi)) continue;}
 			} else if(ej<j) {
-				if(ek>k) {if(edge_x_test(c,xlo,yhi,zlo,xhi,ylo,zhi)) continue;}
-				else if(ek<k) {if(edge_x_test(c,xlo,yhi,zhi,xhi,ylo,zlo)) continue;}
-				else {if(face_y_test(c,xlo,yhi,zlo,xhi,zhi)) continue;}
+				if(ek>k) {if(edge_x_test(c,xlo,yhi,zlo,xhi,ylo,zhi,radi)) continue;}
+				else if(ek<k) {if(edge_x_test(c,xlo,yhi,zhi,xhi,ylo,zlo,radi)) continue;}
+				else {if(face_y_test(c,xlo,yhi,zlo,xhi,zhi,radi)) continue;}
 			} else {
-				if(ek>k) {if(face_z_test(c,xlo,ylo,zlo,xhi,yhi)) continue;}
-				else if(ek<k) {if(face_z_test(c,xlo,ylo,zhi,xhi,yhi)) continue;}
+				if(ek>k) {if(face_z_test(c,xlo,ylo,zlo,xhi,yhi,radi)) continue;}
+				else if(ek<k) {if(face_z_test(c,xlo,ylo,zhi,xhi,yhi,radi)) continue;}
 				else voro_fatal_error("Compute cell routine revisiting central block, which should never\nhappen.",VOROPP_INTERNAL_ERROR);
 			}
 		}
@@ -606,7 +606,7 @@ bool voro_compute<c_class>::compute_cell(v_cell &c,int ijk,int s,int ci,int cj,i
 				x1=p[ijk][ps*l]-x2;
 				y1=p[ijk][ps*l+1]-y2;
 				z1=p[ijk][ps*l+2]-z2;
-				rs=con.r_scale(x1*x1+y1*y1+z1*z1,ijk,l);
+				rs=con.r_scale(x1*x1+y1*y1+z1*z1,ijk,l,radi);
 				if(!c.nplane(x1,y1,z1,rs,id[ijk][l])) return false;
 				l++;
 			} while (l<co[ijk]);
@@ -634,14 +634,14 @@ bool voro_compute<c_class>::compute_cell(v_cell &c,int ijk,int s,int ci,int cj,i
  * \return False if the block may intersect, true if does not. */
 template<class c_class>
 template<class v_cell>
-bool voro_compute<c_class>::corner_test(v_cell &c,double xl,double yl,double zl,double xh,double yh,double zh) {
-	con.r_prime(xl*xl+yl*yl+zl*zl);
-	if(c.plane_intersects_guess(xh,yl,zl,con.r_cutoff(xl*xh+yl*yl+zl*zl))) return false;
-	if(c.plane_intersects(xh,yh,zl,con.r_cutoff(xl*xh+yl*yh+zl*zl))) return false;
-	if(c.plane_intersects(xl,yh,zl,con.r_cutoff(xl*xl+yl*yh+zl*zl))) return false;
-	if(c.plane_intersects(xl,yh,zh,con.r_cutoff(xl*xl+yl*yh+zl*zh))) return false;
-	if(c.plane_intersects(xl,yl,zh,con.r_cutoff(xl*xl+yl*yl+zl*zh))) return false;
-	if(c.plane_intersects(xh,yl,zh,con.r_cutoff(xl*xh+yl*yl+zl*zh))) return false;
+bool voro_compute<c_class>::corner_test(v_cell &c,double xl,double yl,double zl,double xh,double yh,double zh,typename c_class::radius_info& radi) const {
+	con.r_prime(xl*xl+yl*yl+zl*zl,radi);
+	if(c.plane_intersects_guess(xh,yl,zl,con.r_cutoff(xl*xh+yl*yl+zl*zl,radi))) return false;
+	if(c.plane_intersects(xh,yh,zl,con.r_cutoff(xl*xh+yl*yh+zl*zl,radi))) return false;
+	if(c.plane_intersects(xl,yh,zl,con.r_cutoff(xl*xl+yl*yh+zl*zl,radi))) return false;
+	if(c.plane_intersects(xl,yh,zh,con.r_cutoff(xl*xl+yl*yh+zl*zh,radi))) return false;
+	if(c.plane_intersects(xl,yl,zh,con.r_cutoff(xl*xl+yl*yl+zl*zh,radi))) return false;
+	if(c.plane_intersects(xh,yl,zh,con.r_cutoff(xl*xh+yl*yl+zl*zh,radi))) return false;
 	return true;
 }
 
@@ -659,14 +659,14 @@ bool voro_compute<c_class>::corner_test(v_cell &c,double xl,double yl,double zl,
  * \return False if the block may intersect, true if does not. */
 template<class c_class>
 template<class v_cell>
-inline bool voro_compute<c_class>::edge_x_test(v_cell &c,double x0,double yl,double zl,double x1,double yh,double zh) {
-	con.r_prime(yl*yl+zl*zl);
-	if(c.plane_intersects_guess(x0,yl,zh,con.r_cutoff(yl*yl+zl*zh))) return false;
-	if(c.plane_intersects(x1,yl,zh,con.r_cutoff(yl*yl+zl*zh))) return false;
-	if(c.plane_intersects(x1,yl,zl,con.r_cutoff(yl*yl+zl*zl))) return false;
-	if(c.plane_intersects(x0,yl,zl,con.r_cutoff(yl*yl+zl*zl))) return false;
-	if(c.plane_intersects(x0,yh,zl,con.r_cutoff(yl*yh+zl*zl))) return false;
-	if(c.plane_intersects(x1,yh,zl,con.r_cutoff(yl*yh+zl*zl))) return false;
+inline bool voro_compute<c_class>::edge_x_test(v_cell &c,double x0,double yl,double zl,double x1,double yh,double zh,typename c_class::radius_info& radi) const {
+	con.r_prime(yl*yl+zl*zl,radi);
+	if(c.plane_intersects_guess(x0,yl,zh,con.r_cutoff(yl*yl+zl*zh,radi))) return false;
+	if(c.plane_intersects(x1,yl,zh,con.r_cutoff(yl*yl+zl*zh,radi))) return false;
+	if(c.plane_intersects(x1,yl,zl,con.r_cutoff(yl*yl+zl*zl,radi))) return false;
+	if(c.plane_intersects(x0,yl,zl,con.r_cutoff(yl*yl+zl*zl,radi))) return false;
+	if(c.plane_intersects(x0,yh,zl,con.r_cutoff(yl*yh+zl*zl,radi))) return false;
+	if(c.plane_intersects(x1,yh,zl,con.r_cutoff(yl*yh+zl*zl,radi))) return false;
 	return true;
 }
 
@@ -684,14 +684,14 @@ inline bool voro_compute<c_class>::edge_x_test(v_cell &c,double x0,double yl,dou
  * \return False if the block may intersect, true if does not. */
 template<class c_class>
 template<class v_cell>
-inline bool voro_compute<c_class>::edge_y_test(v_cell &c,double xl,double y0,double zl,double xh,double y1,double zh) {
-	con.r_prime(xl*xl+zl*zl);
-	if(c.plane_intersects_guess(xl,y0,zh,con.r_cutoff(xl*xl+zl*zh))) return false;
-	if(c.plane_intersects(xl,y1,zh,con.r_cutoff(xl*xl+zl*zh))) return false;
-	if(c.plane_intersects(xl,y1,zl,con.r_cutoff(xl*xl+zl*zl))) return false;
-	if(c.plane_intersects(xl,y0,zl,con.r_cutoff(xl*xl+zl*zl))) return false;
-	if(c.plane_intersects(xh,y0,zl,con.r_cutoff(xl*xh+zl*zl))) return false;
-	if(c.plane_intersects(xh,y1,zl,con.r_cutoff(xl*xh+zl*zl))) return false;
+inline bool voro_compute<c_class>::edge_y_test(v_cell &c,double xl,double y0,double zl,double xh,double y1,double zh,typename c_class::radius_info& radi) const {
+	con.r_prime(xl*xl+zl*zl,radi);
+	if(c.plane_intersects_guess(xl,y0,zh,con.r_cutoff(xl*xl+zl*zh,radi))) return false;
+	if(c.plane_intersects(xl,y1,zh,con.r_cutoff(xl*xl+zl*zh,radi))) return false;
+	if(c.plane_intersects(xl,y1,zl,con.r_cutoff(xl*xl+zl*zl,radi))) return false;
+	if(c.plane_intersects(xl,y0,zl,con.r_cutoff(xl*xl+zl*zl,radi))) return false;
+	if(c.plane_intersects(xh,y0,zl,con.r_cutoff(xl*xh+zl*zl,radi))) return false;
+	if(c.plane_intersects(xh,y1,zl,con.r_cutoff(xl*xh+zl*zl,radi))) return false;
 	return true;
 }
 
@@ -708,14 +708,14 @@ inline bool voro_compute<c_class>::edge_y_test(v_cell &c,double xl,double y0,dou
  * \return False if the block may intersect, true if does not. */
 template<class c_class>
 template<class v_cell>
-inline bool voro_compute<c_class>::edge_z_test(v_cell &c,double xl,double yl,double z0,double xh,double yh,double z1) {
-	con.r_prime(xl*xl+yl*yl);
-	if(c.plane_intersects_guess(xl,yh,z0,con.r_cutoff(xl*xl+yl*yh))) return false;
-	if(c.plane_intersects(xl,yh,z1,con.r_cutoff(xl*xl+yl*yh))) return false;
-	if(c.plane_intersects(xl,yl,z1,con.r_cutoff(xl*xl+yl*yl))) return false;
-	if(c.plane_intersects(xl,yl,z0,con.r_cutoff(xl*xl+yl*yl))) return false;
-	if(c.plane_intersects(xh,yl,z0,con.r_cutoff(xl*xh+yl*yl))) return false;
-	if(c.plane_intersects(xh,yl,z1,con.r_cutoff(xl*xh+yl*yl))) return false;
+inline bool voro_compute<c_class>::edge_z_test(v_cell &c,double xl,double yl,double z0,double xh,double yh,double z1,typename c_class::radius_info& radi) const {
+	con.r_prime(xl*xl+yl*yl,radi);
+	if(c.plane_intersects_guess(xl,yh,z0,con.r_cutoff(xl*xl+yl*yh,radi))) return false;
+	if(c.plane_intersects(xl,yh,z1,con.r_cutoff(xl*xl+yl*yh,radi))) return false;
+	if(c.plane_intersects(xl,yl,z1,con.r_cutoff(xl*xl+yl*yl,radi))) return false;
+	if(c.plane_intersects(xl,yl,z0,con.r_cutoff(xl*xl+yl*yl,radi))) return false;
+	if(c.plane_intersects(xh,yl,z0,con.r_cutoff(xl*xh+yl*yl,radi))) return false;
+	if(c.plane_intersects(xh,yl,z1,con.r_cutoff(xl*xh+yl*yl,radi))) return false;
 	return true;
 }
 
@@ -731,12 +731,12 @@ inline bool voro_compute<c_class>::edge_z_test(v_cell &c,double xl,double yl,dou
  * \return False if the block may intersect, true if does not. */
 template<class c_class>
 template<class v_cell>
-inline bool voro_compute<c_class>::face_x_test(v_cell &c,double xl,double y0,double z0,double y1,double z1) {
-	con.r_prime(xl*xl);
-	if(c.plane_intersects_guess(xl,y0,z0,con.r_cutoff(xl*xl))) return false;
-	if(c.plane_intersects(xl,y0,z1,con.r_cutoff(xl*xl))) return false;
-	if(c.plane_intersects(xl,y1,z1,con.r_cutoff(xl*xl))) return false;
-	if(c.plane_intersects(xl,y1,z0,con.r_cutoff(xl*xl))) return false;
+inline bool voro_compute<c_class>::face_x_test(v_cell &c,double xl,double y0,double z0,double y1,double z1,typename c_class::radius_info& radi) const {
+	con.r_prime(xl*xl,radi);
+	if(c.plane_intersects_guess(xl,y0,z0,con.r_cutoff(xl*xl,radi))) return false;
+	if(c.plane_intersects(xl,y0,z1,con.r_cutoff(xl*xl,radi))) return false;
+	if(c.plane_intersects(xl,y1,z1,con.r_cutoff(xl*xl,radi))) return false;
+	if(c.plane_intersects(xl,y1,z0,con.r_cutoff(xl*xl,radi))) return false;
 	return true;
 }
 
@@ -752,12 +752,12 @@ inline bool voro_compute<c_class>::face_x_test(v_cell &c,double xl,double y0,dou
  * \return False if the block may intersect, true if does not. */
 template<class c_class>
 template<class v_cell>
-inline bool voro_compute<c_class>::face_y_test(v_cell &c,double x0,double yl,double z0,double x1,double z1) {
-	con.r_prime(yl*yl);
-	if(c.plane_intersects_guess(x0,yl,z0,con.r_cutoff(yl*yl))) return false;
-	if(c.plane_intersects(x0,yl,z1,con.r_cutoff(yl*yl))) return false;
-	if(c.plane_intersects(x1,yl,z1,con.r_cutoff(yl*yl))) return false;
-	if(c.plane_intersects(x1,yl,z0,con.r_cutoff(yl*yl))) return false;
+inline bool voro_compute<c_class>::face_y_test(v_cell &c,double x0,double yl,double z0,double x1,double z1,typename c_class::radius_info& radi) const {
+	con.r_prime(yl*yl,radi);
+	if(c.plane_intersects_guess(x0,yl,z0,con.r_cutoff(yl*yl,radi))) return false;
+	if(c.plane_intersects(x0,yl,z1,con.r_cutoff(yl*yl,radi))) return false;
+	if(c.plane_intersects(x1,yl,z1,con.r_cutoff(yl*yl,radi))) return false;
+	if(c.plane_intersects(x1,yl,z0,con.r_cutoff(yl*yl,radi))) return false;
 	return true;
 }
 
@@ -773,12 +773,12 @@ inline bool voro_compute<c_class>::face_y_test(v_cell &c,double x0,double yl,dou
  * \return False if the block may intersect, true if does not. */
 template<class c_class>
 template<class v_cell>
-inline bool voro_compute<c_class>::face_z_test(v_cell &c,double x0,double y0,double zl,double x1,double y1) {
-	con.r_prime(zl*zl);
-	if(c.plane_intersects_guess(x0,y0,zl,con.r_cutoff(zl*zl))) return false;
-	if(c.plane_intersects(x0,y1,zl,con.r_cutoff(zl*zl))) return false;
-	if(c.plane_intersects(x1,y1,zl,con.r_cutoff(zl*zl))) return false;
-	if(c.plane_intersects(x1,y0,zl,con.r_cutoff(zl*zl))) return false;
+inline bool voro_compute<c_class>::face_z_test(v_cell &c,double x0,double y0,double zl,double x1,double y1,typename c_class::radius_info& radi) const {
+	con.r_prime(zl*zl,radi);
+	if(c.plane_intersects_guess(x0,y0,zl,con.r_cutoff(zl*zl,radi))) return false;
+	if(c.plane_intersects(x0,y1,zl,con.r_cutoff(zl*zl,radi))) return false;
+	if(c.plane_intersects(x1,y1,zl,con.r_cutoff(zl*zl,radi))) return false;
+	if(c.plane_intersects(x1,y0,zl,con.r_cutoff(zl*zl,radi))) return false;
 	return true;
 }
 
@@ -798,7 +798,7 @@ inline bool voro_compute<c_class>::face_z_test(v_cell &c,double x0,double y0,dou
  * \return True if the region is further away than mrs, false if the region in
  *         within mrs. */
 template<class c_class>
-bool voro_compute<c_class>::compute_min_max_radius(int di,int dj,int dk,double fx,double fy,double fz,double gxs,double gys,double gzs,double &crs,double mrs) {
+bool voro_compute<c_class>::compute_min_max_radius(int di,int dj,int dk,double fx,double fy,double fz,double gxs,double gys,double gzs,double &crs,double mrs,const typename c_class::radius_info& radi) const {
 	double xlo,ylo,zlo;
 	if(di>0) {
 		xlo=di*boxx-fx;
@@ -808,14 +808,14 @@ bool voro_compute<c_class>::compute_min_max_radius(int di,int dj,int dk,double f
 			crs+=ylo*ylo;
 			if(dk>0) {
 				zlo=dk*boxz-fz;
-				crs+=zlo*zlo;if(con.r_ctest(crs,mrs)) return true;
+				crs+=zlo*zlo;if(con.r_ctest(crs,mrs,radi)) return true;
 				crs+=bxsq+2*(boxx*xlo+boxy*ylo+boxz*zlo);
 			} else if(dk<0) {
 				zlo=(dk+1)*boxz-fz;
-				crs+=zlo*zlo;if(con.r_ctest(crs,mrs)) return true;
+				crs+=zlo*zlo;if(con.r_ctest(crs,mrs,radi)) return true;
 				crs+=bxsq+2*(boxx*xlo+boxy*ylo-boxz*zlo);
 			} else {
-				if(con.r_ctest(crs,mrs)) return true;
+				if(con.r_ctest(crs,mrs,radi)) return true;
 				crs+=boxx*(2*xlo+boxx)+boxy*(2*ylo+boxy)+gzs;
 			}
 		} else if(dj<0) {
@@ -823,27 +823,27 @@ bool voro_compute<c_class>::compute_min_max_radius(int di,int dj,int dk,double f
 			crs+=ylo*ylo;
 			if(dk>0) {
 				zlo=dk*boxz-fz;
-				crs+=zlo*zlo;if(con.r_ctest(crs,mrs)) return true;
+				crs+=zlo*zlo;if(con.r_ctest(crs,mrs,radi)) return true;
 				crs+=bxsq+2*(boxx*xlo-boxy*ylo+boxz*zlo);
 			} else if(dk<0) {
 				zlo=(dk+1)*boxz-fz;
-				crs+=zlo*zlo;if(con.r_ctest(crs,mrs)) return true;
+				crs+=zlo*zlo;if(con.r_ctest(crs,mrs,radi)) return true;
 				crs+=bxsq+2*(boxx*xlo-boxy*ylo-boxz*zlo);
 			} else {
-				if(con.r_ctest(crs,mrs)) return true;
+				if(con.r_ctest(crs,mrs,radi)) return true;
 				crs+=boxx*(2*xlo+boxx)+boxy*(-2*ylo+boxy)+gzs;
 			}
 		} else {
 			if(dk>0) {
 				zlo=dk*boxz-fz;
-				crs+=zlo*zlo;if(con.r_ctest(crs,mrs)) return true;
+				crs+=zlo*zlo;if(con.r_ctest(crs,mrs,radi)) return true;
 				crs+=boxz*(2*zlo+boxz);
 			} else if(dk<0) {
 				zlo=(dk+1)*boxz-fz;
-				crs+=zlo*zlo;if(con.r_ctest(crs,mrs)) return true;
+				crs+=zlo*zlo;if(con.r_ctest(crs,mrs,radi)) return true;
 				crs+=boxz*(-2*zlo+boxz);
 			} else {
-				if(con.r_ctest(crs,mrs)) return true;
+				if(con.r_ctest(crs,mrs,radi)) return true;
 				crs+=gzs;
 			}
 			crs+=gys+boxx*(2*xlo+boxx);
@@ -856,14 +856,14 @@ bool voro_compute<c_class>::compute_min_max_radius(int di,int dj,int dk,double f
 			crs+=ylo*ylo;
 			if(dk>0) {
 				zlo=dk*boxz-fz;
-				crs+=zlo*zlo;if(con.r_ctest(crs,mrs)) return true;
+				crs+=zlo*zlo;if(con.r_ctest(crs,mrs,radi)) return true;
 				crs+=bxsq+2*(-boxx*xlo+boxy*ylo+boxz*zlo);
 			} else if(dk<0) {
 				zlo=(dk+1)*boxz-fz;
-				crs+=zlo*zlo;if(con.r_ctest(crs,mrs)) return true;
+				crs+=zlo*zlo;if(con.r_ctest(crs,mrs,radi)) return true;
 				crs+=bxsq+2*(-boxx*xlo+boxy*ylo-boxz*zlo);
 			} else {
-				if(con.r_ctest(crs,mrs)) return true;
+				if(con.r_ctest(crs,mrs,radi)) return true;
 				crs+=boxx*(-2*xlo+boxx)+boxy*(2*ylo+boxy)+gzs;
 			}
 		} else if(dj<0) {
@@ -871,27 +871,27 @@ bool voro_compute<c_class>::compute_min_max_radius(int di,int dj,int dk,double f
 			crs+=ylo*ylo;
 			if(dk>0) {
 				zlo=dk*boxz-fz;
-				crs+=zlo*zlo;if(con.r_ctest(crs,mrs)) return true;
+				crs+=zlo*zlo;if(con.r_ctest(crs,mrs,radi)) return true;
 				crs+=bxsq+2*(-boxx*xlo-boxy*ylo+boxz*zlo);
 			} else if(dk<0) {
 				zlo=(dk+1)*boxz-fz;
-				crs+=zlo*zlo;if(con.r_ctest(crs,mrs)) return true;
+				crs+=zlo*zlo;if(con.r_ctest(crs,mrs,radi)) return true;
 				crs+=bxsq+2*(-boxx*xlo-boxy*ylo-boxz*zlo);
 			} else {
-				if(con.r_ctest(crs,mrs)) return true;
+				if(con.r_ctest(crs,mrs,radi)) return true;
 				crs+=boxx*(-2*xlo+boxx)+boxy*(-2*ylo+boxy)+gzs;
 			}
 		} else {
 			if(dk>0) {
 				zlo=dk*boxz-fz;
-				crs+=zlo*zlo;if(con.r_ctest(crs,mrs)) return true;
+				crs+=zlo*zlo;if(con.r_ctest(crs,mrs,radi)) return true;
 				crs+=boxz*(2*zlo+boxz);
 			} else if(dk<0) {
 				zlo=(dk+1)*boxz-fz;
-				crs+=zlo*zlo;if(con.r_ctest(crs,mrs)) return true;
+				crs+=zlo*zlo;if(con.r_ctest(crs,mrs,radi)) return true;
 				crs+=boxz*(-2*zlo+boxz);
 			} else {
-				if(con.r_ctest(crs,mrs)) return true;
+				if(con.r_ctest(crs,mrs,radi)) return true;
 				crs+=gzs;
 			}
 			crs+=gys+boxx*(-2*xlo+boxx);
@@ -902,14 +902,14 @@ bool voro_compute<c_class>::compute_min_max_radius(int di,int dj,int dk,double f
 			crs=ylo*ylo;
 			if(dk>0) {
 				zlo=dk*boxz-fz;
-				crs+=zlo*zlo;if(con.r_ctest(crs,mrs)) return true;
+				crs+=zlo*zlo;if(con.r_ctest(crs,mrs,radi)) return true;
 				crs+=boxz*(2*zlo+boxz);
 			} else if(dk<0) {
 				zlo=(dk+1)*boxz-fz;
-				crs+=zlo*zlo;if(con.r_ctest(crs,mrs)) return true;
+				crs+=zlo*zlo;if(con.r_ctest(crs,mrs,radi)) return true;
 				crs+=boxz*(-2*zlo+boxz);
 			} else {
-				if(con.r_ctest(crs,mrs)) return true;
+				if(con.r_ctest(crs,mrs,radi)) return true;
 				crs+=gzs;
 			}
 			crs+=boxy*(2*ylo+boxy);
@@ -918,23 +918,23 @@ bool voro_compute<c_class>::compute_min_max_radius(int di,int dj,int dk,double f
 			crs=ylo*ylo;
 			if(dk>0) {
 				zlo=dk*boxz-fz;
-				crs+=zlo*zlo;if(con.r_ctest(crs,mrs)) return true;
+				crs+=zlo*zlo;if(con.r_ctest(crs,mrs,radi)) return true;
 				crs+=boxz*(2*zlo+boxz);
 			} else if(dk<0) {
 				zlo=(dk+1)*boxz-fz;
-				crs+=zlo*zlo;if(con.r_ctest(crs,mrs)) return true;
+				crs+=zlo*zlo;if(con.r_ctest(crs,mrs,radi)) return true;
 				crs+=boxz*(-2*zlo+boxz);
 			} else {
-				if(con.r_ctest(crs,mrs)) return true;
+				if(con.r_ctest(crs,mrs,radi)) return true;
 				crs+=gzs;
 			}
 			crs+=boxy*(-2*ylo+boxy);
 		} else {
 			if(dk>0) {
-				zlo=dk*boxz-fz;crs=zlo*zlo;if(con.r_ctest(crs,mrs)) return true;
+				zlo=dk*boxz-fz;crs=zlo*zlo;if(con.r_ctest(crs,mrs,radi)) return true;
 				crs+=boxz*(2*zlo+boxz);
 			} else if(dk<0) {
-				zlo=(dk+1)*boxz-fz;crs=zlo*zlo;if(con.r_ctest(crs,mrs)) return true;
+				zlo=(dk+1)*boxz-fz;crs=zlo*zlo;if(con.r_ctest(crs,mrs,radi)) return true;
 				crs+=boxz*(-2*zlo+boxz);
 			} else {
 				crs=0;
@@ -948,7 +948,7 @@ bool voro_compute<c_class>::compute_min_max_radius(int di,int dj,int dk,double f
 }
 
 template<class c_class>
-bool voro_compute<c_class>::compute_min_radius(int di,int dj,int dk,double fx,double fy,double fz,double mrs) {
+bool voro_compute<c_class>::compute_min_radius(int di,int dj,int dk,double fx,double fy,double fz,double mrs) const {
 	double t,crs;
 
 	if(di>0) {t=di*boxx-fx;crs=t*t;}
@@ -996,14 +996,4 @@ template bool voro_compute<container_poly>::compute_cell(voronoicell&,int,int,in
 template bool voro_compute<container_poly>::compute_cell(voronoicell_neighbor&,int,int,int,int,int);
 template void voro_compute<container_poly>::find_voronoi_cell(double,double,double,int,int,int,int,particle_record&,double&);
 
-// Explicit template instantiation
-template voro_compute<container_periodic>::voro_compute(container_periodic&,int,int,int);
-template voro_compute<container_periodic_poly>::voro_compute(container_periodic_poly&,int,int,int);
-template bool voro_compute<container_periodic>::compute_cell(voronoicell&,int,int,int,int,int);
-template bool voro_compute<container_periodic>::compute_cell(voronoicell_neighbor&,int,int,int,int,int);
-template void voro_compute<container_periodic>::find_voronoi_cell(double,double,double,int,int,int,int,particle_record&,double&);
-template bool voro_compute<container_periodic_poly>::compute_cell(voronoicell&,int,int,int,int,int);
-template bool voro_compute<container_periodic_poly>::compute_cell(voronoicell_neighbor&,int,int,int,int,int);
-template void voro_compute<container_periodic_poly>::find_voronoi_cell(double,double,double,int,int,int,int,particle_record&,double&);
-
 }
diff --git a/Engine/Source/Runtime/Experimental/Voronoi/Private/voro++/src/wall.cc b/Engine/Source/Runtime/Experimental/Voronoi/Private/voro++/src/wall.cc
--- a/Engine/Source/Runtime/Experimental/Voronoi/Private/voro++/src/wall.cc
+++ b/Engine/Source/Runtime/Experimental/Voronoi/Private/voro++/src/wall.cc
@@ -17,7 +17,7 @@ namespace voro {
 /** Tests to see whether a point is inside the sphere wall object.
  * \param[in,out] (x,y,z) the vector to test.
  * \return True if the point is inside, false if the point is outside. */
-bool wall_sphere::point_inside(double x,double y,double z) {
+bool wall_sphere::point_inside(double x,double y,double z) const {
 	return (x-xc)*(x-xc)+(y-yc)*(y-yc)+(z-zc)*(z-zc)<rc*rc;
 }
 
@@ -29,7 +29,7 @@ bool wall_sphere::point_inside(double x,double y,double z) {
  * \param[in] (x,y,z) the location of the Voronoi cell.
  * \return True if the cell still exists, false if the cell is deleted. */
 template<class v_cell>
-bool wall_sphere::cut_cell_base(v_cell &c,double x,double y,double z) {
+bool wall_sphere::cut_cell_base(v_cell &c,double x,double y,double z) const {
 	double xd=x-xc,yd=y-yc,zd=z-zc,dq=xd*xd+yd*yd+zd*zd;
 	if (dq>1e-5) {
 		dq=2*(sqrt(dq)*rc-dq);
@@ -41,7 +41,7 @@ bool wall_sphere::cut_cell_base(v_cell &c,double x,double y,double z) {
 /** Tests to see whether a point is inside the plane wall object.
  * \param[in] (x,y,z) the vector to test.
  * \return True if the point is inside, false if the point is outside. */
-bool wall_plane::point_inside(double x,double y,double z) {
+bool wall_plane::point_inside(double x,double y,double z) const {
 	return x*xc+y*yc+z*zc<ac;
 }
 
@@ -50,7 +50,7 @@ bool wall_plane::point_inside(double x,double y,double z) {
  * \param[in] (x,y,z) the location of the Voronoi cell.
  * \return True if the cell still exists, false if the cell is deleted. */
 template<class v_cell>
-bool wall_plane::cut_cell_base(v_cell &c,double x,double y,double z) {
+bool wall_plane::cut_cell_base(v_cell &c,double x,double y,double z) const {
 	double dq=2*(ac-x*xc-y*yc-z*zc);
 	return c.nplane(xc,yc,zc,dq,w_id);
 }
@@ -58,7 +58,7 @@ bool wall_plane::cut_cell_base(v_cell &c,double x,double y,double z) {
 /** Tests to see whether a point is inside the cylindrical wall object.
  * \param[in] (x,y,z) the vector to test.
  * \return True if the point is inside, false if the point is outside. */
-bool wall_cylinder::point_inside(double x,double y,double z) {
+bool wall_cylinder::point_inside(double x,double y,double z) const {
 	double xd=x-xc,yd=y-yc,zd=z-zc;
 	double pa=(xd*xa+yd*ya+zd*za)*asi;
 	xd-=xa*pa;yd-=ya*pa;zd-=za*pa;
@@ -74,7 +74,7 @@ bool wall_cylinder::point_inside(double x,double y,double z) {
  * \param[in] (x,y,z) the location of the Voronoi cell.
  * \return True if the cell still exists, false if the cell is deleted. */
 template<class v_cell>
-bool wall_cylinder::cut_cell_base(v_cell &c,double x,double y,double z) {
+bool wall_cylinder::cut_cell_base(v_cell &c,double x,double y,double z) const {
 	double xd=x-xc,yd=y-yc,zd=z-zc,pa=(xd*xa+yd*ya+zd*za)*asi;
 	xd-=xa*pa;yd-=ya*pa;zd-=za*pa;
 	pa=xd*xd+yd*yd+zd*zd;
@@ -88,7 +88,7 @@ bool wall_cylinder::cut_cell_base(v_cell &c,double x,double y,double z) {
 /** Tests to see whether a point is inside the cone wall object.
  * \param[in] (x,y,z) the vector to test.
  * \return True if the point is inside, false if the point is outside. */
-bool wall_cone::point_inside(double x,double y,double z) {
+bool wall_cone::point_inside(double x,double y,double z) const {
 	double xd=x-xc,yd=y-yc,zd=z-zc,pa=(xd*xa+yd*ya+zd*za)*asi;
 	xd-=xa*pa;yd-=ya*pa;zd-=za*pa;
 	pa*=gra;
@@ -106,7 +106,7 @@ bool wall_cone::point_inside(double x,double y,double z) {
  * \param[in] (x,y,z) the location of the Voronoi cell.
  * \return True if the cell still exists, false if the cell is deleted. */
 template<class v_cell>
-bool wall_cone::cut_cell_base(v_cell &c,double x,double y,double z) {
+bool wall_cone::cut_cell_base(v_cell &c,double x,double y,double z) const {
 	double xd=x-xc,yd=y-yc,zd=z-zc,xf,yf,zf,q,pa=(xd*xa+yd*ya+zd*za)*asi;
 	xd-=xa*pa;yd-=ya*pa;zd-=za*pa;
 	pa=xd*xd+yd*yd+zd*zd;
@@ -123,13 +123,13 @@ bool wall_cone::cut_cell_base(v_cell &c,double x,double y,double z) {
 }
 
 // Explicit instantiation
-template bool wall_sphere::cut_cell_base(voronoicell&,double,double,double);
-template bool wall_sphere::cut_cell_base(voronoicell_neighbor&,double,double,double);
-template bool wall_plane::cut_cell_base(voronoicell&,double,double,double);
-template bool wall_plane::cut_cell_base(voronoicell_neighbor&,double,double,double);
-template bool wall_cylinder::cut_cell_base(voronoicell&,double,double,double);
-template bool wall_cylinder::cut_cell_base(voronoicell_neighbor&,double,double,double);
-template bool wall_cone::cut_cell_base(voronoicell&,double,double,double);
-template bool wall_cone::cut_cell_base(voronoicell_neighbor&,double,double,double);
+template bool wall_sphere::cut_cell_base(voronoicell&,double,double,double) const;
+template bool wall_sphere::cut_cell_base(voronoicell_neighbor&,double,double,double) const;
+template bool wall_plane::cut_cell_base(voronoicell&,double,double,double) const;
+template bool wall_plane::cut_cell_base(voronoicell_neighbor&,double,double,double) const;
+template bool wall_cylinder::cut_cell_base(voronoicell&,double,double,double) const;
+template bool wall_cylinder::cut_cell_base(voronoicell_neighbor&,double,double,double) const;
+template bool wall_cone::cut_cell_base(voronoicell&,double,double,double) const;
+template bool wall_cone::cut_cell_base(voronoicell_neighbor&,double,double,double) const;
 
 }
diff --git a/Engine/Source/Runtime/Experimental/Voronoi/Public/Voronoi/Voronoi.h b/Engine/Source/Runtime/Experimental/Voronoi/Public/Voronoi/Voronoi.h
--- a/Engine/Source/Runtime/Experimental/Voronoi/Public/Voronoi/Voronoi.h
+++ b/Engine/Source/Runtime/Experimental/Voronoi/Public/Voronoi/Voronoi.h
@@ -28,6 +28,7 @@ class VORONOI_API FVoronoiDiagram
 public:
 	// we typically add extra space to the bounds of the Voronoi diagram, to avoid numerical issues of a Voronoi site being on the boundary
 	const static float DefaultBoundingBoxSlack;
+	const static int MinDefaultSitesPerThread;
 
 	/**
 	 * @param Sites							Voronoi sites for the diagram
@@ -42,7 +43,6 @@ public:
 	* @param ExtraBoundingSpace				Voronoi diagram will be computed within the input Bounds + this amount of extra space in each dimension
 	* @param SquaredDistSkipPtThreshold		A safety threshold to avoid creating invalid cells: sites that are within this distance of an already-added site will not be added.
 	*										(If you know there will be no duplicate sites, can set to zero for faster perf.)
-	*
 	*/
 	FVoronoiDiagram(const TArrayView<const FVector>& Sites, const FBox &Bounds, float ExtraBoundingSpace, float SquaredDistSkipPtThreshold = KINDA_SMALL_NUMBER);
 	~FVoronoiDiagram();
@@ -54,12 +54,29 @@ public:
 
 	void AddSites(const TArrayView<const FVector>& Sites, float SquaredDistSkipPtThreshold = 0.0f);
 
-	void ComputeAllCells(TArray<FVoronoiCellInfo> &AllCells);
+	void ComputeAllCellsSerial(TArray<FVoronoiCellInfo>& AllCells);
+	void ComputeAllCells(TArray<FVoronoiCellInfo>& AllCells, int32 ApproxSitesPerThread = -1);
+
+	void ComputeCellEdgesSerial(TArray<TTuple<FVector, FVector>>& Edges, TArray<int32>& CellMember);
+	void ComputeCellEdges(TArray<TTuple<FVector, FVector>>& Edges, TArray<int32>& CellMember, int32 ApproxSitesPerThread = -1);
 
 	/**
 	 * Find the id of the Voronoi cell containing the given position (or -1 if position is outside diagram)
+	 * TODO: expose a "compute" parameter to store cached state across calls to this function, and un-comment it back
 	 */
-	int32 FindCell(const FVector& Pos);
+	//int32 FindCell(const FVector& Pos);
+
+	int32 ApproxSitesPerThreadWithDefault(int32 ApproxSitesPerThreadIn)
+	{
+		if (ApproxSitesPerThreadIn < 0)
+		{
+			return FMath::Max(MinDefaultSitesPerThread, NumSites / 64);
+		}
+		return ApproxSitesPerThreadIn;
+	}
+
+private:
+	TArray<int32> GetParallelBlockRanges(int32 ApproxSitesPerThread);
 };
 
 
