GeometryCollection : Convex Generation
- Defaulting the geometry collection assets collision geometry to convex
- Updating the convex bodies during the destruction of the geometry collections edit object.

diff --git a/Engine/Plugins/Experimental/ChaosEditor/Source/FractureEditor/Public/FractureTool.cpp b/Engine/Plugins/Experimental/ChaosEditor/Source/FractureEditor/Public/FractureTool.cpp
--- a/Engine/Plugins/Experimental/ChaosEditor/Source/FractureEditor/Public/FractureTool.cpp
+++ b/Engine/Plugins/Experimental/ChaosEditor/Source/FractureEditor/Public/FractureTool.cpp
@@ -356,6 +356,8 @@ void UFractureModalTool::Execute(TWeakPtr<FFractureEditorModeToolkit> InToolkit)
 				FractureContext.GetFracturedGeometryCollection()->InitResources();
 			}
 
+			FractureContext.GetGeometryCollectionComponent()->EditRestCollection(GeometryCollection::EEditUpdate::RestPhysicsDynamic);
+
 			if (bNeedPhysicsUpdate)
 			{
 				FractureContext.GetGeometryCollectionComponent()->RecreatePhysicsState();
diff --git a/Engine/Source/Runtime/Experimental/GeometryCollectionEngine/Private/GeometryCollection/GeometryCollectionComponent.cpp b/Engine/Source/Runtime/Experimental/GeometryCollectionEngine/Private/GeometryCollection/GeometryCollectionComponent.cpp
--- a/Engine/Source/Runtime/Experimental/GeometryCollectionEngine/Private/GeometryCollection/GeometryCollectionComponent.cpp
+++ b/Engine/Source/Runtime/Experimental/GeometryCollectionEngine/Private/GeometryCollection/GeometryCollectionComponent.cpp
@@ -2241,7 +2241,7 @@ FGeometryCollectionEdit::FGeometryCollectionEdit(UGeometryCollectionComponent* I
 
 	if (EnumHasAnyFlags(EditUpdate, GeometryCollection::EEditUpdate::Rest) && GetRestCollection())
 	{
-		GetRestCollection()->Modify();
+		Component->Modify();
 	}
 }
 
@@ -2255,6 +2255,12 @@ FGeometryCollectionEdit::~FGeometryCollectionEdit()
 			Component->ResetDynamicCollection();
 		}
 
+		if (EnumHasAnyFlags(EditUpdate, GeometryCollection::EEditUpdate::Rest) && GetRestCollection())
+		{
+			GetRestCollection()->UpdateConvexGeometry();
+			GetRestCollection()->InvalidateCollection();
+		}
+
 		if (EnumHasAnyFlags(EditUpdate, GeometryCollection::EEditUpdate::Physics) && bHadPhysicsState)
 		{
 			Component->RecreatePhysicsState();
diff --git a/Engine/Source/Runtime/Experimental/GeometryCollectionEngine/Private/GeometryCollection/GeometryCollectionObject.cpp b/Engine/Source/Runtime/Experimental/GeometryCollectionEngine/Private/GeometryCollection/GeometryCollectionObject.cpp
--- a/Engine/Source/Runtime/Experimental/GeometryCollectionEngine/Private/GeometryCollection/GeometryCollectionObject.cpp
+++ b/Engine/Source/Runtime/Experimental/GeometryCollectionEngine/Private/GeometryCollection/GeometryCollectionObject.cpp
@@ -21,6 +21,7 @@
 
 #if WITH_EDITOR
 #include "GeometryCollection/DerivedDataGeometryCollectionCooker.h"
+#include "GeometryCollection/GeometryCollectionConvexUtility.h"
 #include "GeometryCollection/GeometryCollectionComponent.h"
 #include "DerivedDataCacheInterface.h"
 #include "Serialization/MemoryReader.h"
@@ -69,7 +70,7 @@ UGeometryCollection::UGeometryCollection(const FObjectInitializer& ObjectInitial
 	, EnableNanite(false)
 #if WITH_EDITORONLY_DATA
 	, CollisionType_DEPRECATED(ECollisionTypeEnum::Chaos_Volumetric)
-	, ImplicitType_DEPRECATED(EImplicitTypeEnum::Chaos_Implicit_Box)
+	, ImplicitType_DEPRECATED(EImplicitTypeEnum::Chaos_Implicit_Convex)
 	, MinLevelSetResolution_DEPRECATED(10)
 	, MaxLevelSetResolution_DEPRECATED(10)
 	, MinClusterLevelSetResolution_DEPRECATED(50)
@@ -111,7 +112,7 @@ FGeometryCollectionCollisionParticleData::FGeometryCollectionCollisionParticleDa
 
 FGeometryCollectionCollisionTypeData::FGeometryCollectionCollisionTypeData()
 	: CollisionType(ECollisionTypeEnum::Chaos_Volumetric)
-	, ImplicitType(EImplicitTypeEnum::Chaos_Implicit_Box)
+	, ImplicitType(EImplicitTypeEnum::Chaos_Implicit_Convex)
 	, LevelSet()
 	, CollisionParticles()
 	, CollisionObjectReductionPercentage(0.0f)
@@ -124,7 +125,7 @@ FGeometryCollectionSizeSpecificData::FGeometryCollectionSizeSpecificData()
 	, CollisionShapes({ FGeometryCollectionCollisionTypeData()})
 #if WITH_EDITORONLY_DATA
 	, CollisionType_DEPRECATED(ECollisionTypeEnum::Chaos_Volumetric)
-	, ImplicitType_DEPRECATED(EImplicitTypeEnum::Chaos_Implicit_Box)
+	, ImplicitType_DEPRECATED(EImplicitTypeEnum::Chaos_Implicit_Convex)
 	, MinLevelSetResolution_DEPRECATED(5)
 	, MaxLevelSetResolution_DEPRECATED(10)
 	, MinClusterLevelSetResolution_DEPRECATED(25)
@@ -211,7 +212,7 @@ FGeometryCollectionSizeSpecificData UGeometryCollection::GeometryCollectionSizeS
 	if (Data.CollisionShapes.Num())
 	{
 		Data.CollisionShapes[0].CollisionType = ECollisionTypeEnum::Chaos_Volumetric;
-		Data.CollisionShapes[0].ImplicitType = EImplicitTypeEnum::Chaos_Implicit_Sphere;
+		Data.CollisionShapes[0].ImplicitType = EImplicitTypeEnum::Chaos_Implicit_Capsule;
 		Data.CollisionShapes[0].LevelSet.MinLevelSetResolution = 5;
 		Data.CollisionShapes[0].LevelSet.MaxLevelSetResolution = 10;
 		Data.CollisionShapes[0].LevelSet.MinClusterLevelSetResolution = 25;
@@ -265,6 +266,18 @@ void UGeometryCollection::ValidateSizeSpecificDataDefaults()
 	check(SizeSpecificData.Num());
 }
 
+void UGeometryCollection::UpdateConvexGeometry()
+{
+#if WITH_EDITOR
+	if (GeometryCollection)
+	{
+		FGeometryCollectionConvexPropertiesInterface::FConvexCreationProperties ConvexProperties = GeometryCollection->GetConvexProperties();
+		FGeometryCollectionConvexUtility::CreateNonOverlappingConvexHullData(GeometryCollection.Get(), ConvexProperties.FractionRemove, ConvexProperties.SimplificationThreshold);
+		InvalidateCollection();
+	}
+#endif
+}
+
 
 float KgCm3ToKgM3(float Density)
 {
@@ -1150,6 +1163,7 @@ bool UGeometryCollection::Modify(bool bAlwaysMarkDirty /*= true*/)
 	UPackage* Package = GetOutermost();
 	if (Package->IsDirty())
 	{
+		UpdateConvexGeometry();
 		InvalidateCollection();
 		ValidateSizeSpecificDataDefaults();
 	}
diff --git a/Engine/Source/Runtime/Experimental/GeometryCollectionEngine/Public/GeometryCollection/GeometryCollectionComponent.h b/Engine/Source/Runtime/Experimental/GeometryCollectionEngine/Public/GeometryCollection/GeometryCollectionComponent.h
--- a/Engine/Source/Runtime/Experimental/GeometryCollectionEngine/Public/GeometryCollection/GeometryCollectionComponent.h
+++ b/Engine/Source/Runtime/Experimental/GeometryCollectionEngine/Public/GeometryCollection/GeometryCollectionComponent.h
@@ -151,11 +151,11 @@ namespace GeometryCollection
 		/** No update. */
 		None = 0,
 		/** Mark the rest collection as changed. */
-		Rest = 1,
+		Rest = 1 << 1,
 		/** Recreate the physics state (proxy). */
-		Physics = 2,
+		Physics = 1 << 2,
 		/** Reset the dynamic collection. */
-		Dynamic = 4,
+		Dynamic = 1 << 4,
 		/** Mark the rest collection as changed, and recreate the physics state (proxy). */
 		RestPhysics = Rest | Physics,
 		/** Reset dynamic collection, mark the rest collection as changed, and recreate the physics state (proxy). */
diff --git a/Engine/Source/Runtime/Experimental/GeometryCollectionEngine/Public/GeometryCollection/GeometryCollectionObject.h b/Engine/Source/Runtime/Experimental/GeometryCollectionEngine/Public/GeometryCollection/GeometryCollectionObject.h
--- a/Engine/Source/Runtime/Experimental/GeometryCollectionEngine/Public/GeometryCollection/GeometryCollectionObject.h
+++ b/Engine/Source/Runtime/Experimental/GeometryCollectionEngine/Public/GeometryCollection/GeometryCollectionObject.h
@@ -583,6 +583,11 @@ public:
 	TObjectPtr<class UThumbnailInfo> ThumbnailInfo;
 #endif // WITH_EDITORONLY_DATA
 
+	/*
+	* Update the convex geometry on the collection.
+	*/
+	void UpdateConvexGeometry();
+
 private:
 #if WITH_EDITOR
 	void CreateSimulationDataImp(bool bCopyFromDDC);
@@ -590,7 +595,7 @@ private:
 
 	/*
 	* Used to transfer deprecated properties to the size specific structures during serialization
-	* and to add back the default size specific data when deleted. 
+	* and to add back the default size specific data when deleted.
 	*/
 	void ValidateSizeSpecificDataDefaults();
 
