Geometry Collection - Editing Updates
- Updated collection to support enabled flag
- Updated modify to only invalidate the UObject
- Added UObject->Modify to ComponentEdit constructor.

diff --git a/Engine/Source/Runtime/Experimental/Chaos/Private/GeometryCollection/GeometryCollectionConvexPropertiesInterface.cpp b/Engine/Source/Runtime/Experimental/Chaos/Private/GeometryCollection/GeometryCollectionConvexPropertiesInterface.cpp
--- a/Engine/Source/Runtime/Experimental/Chaos/Private/GeometryCollection/GeometryCollectionConvexPropertiesInterface.cpp
+++ b/Engine/Source/Runtime/Experimental/Chaos/Private/GeometryCollection/GeometryCollectionConvexPropertiesInterface.cpp
@@ -5,6 +5,7 @@
 
 const FName FGeometryCollectionConvexPropertiesInterface::ConvexPropertiesGroup = "ConvexProperties";
 const FName FGeometryCollectionConvexPropertiesInterface::ConvexIndexAttribute = "TransformGroupIndex";
+const FName FGeometryCollectionConvexPropertiesInterface::ConvexEnable = "Enable";
 const FName FGeometryCollectionConvexPropertiesInterface::ConvexFractionRemoveAttribute = "FractionRemove";
 const FName FGeometryCollectionConvexPropertiesInterface::ConvexSimplificationThresholdAttribute = "SimplificationThreshold";
 
@@ -26,6 +27,11 @@ FGeometryCollectionConvexPropertiesInterface::InitializeInterface()
 		ManagedCollection->AddAttribute<int32>(ConvexIndexAttribute, ConvexPropertiesGroup, AttributeDependency);
 	}
 
+	if (!ManagedCollection->HasAttribute(ConvexEnable, ConvexPropertiesGroup))
+	{
+		ManagedCollection->AddAttribute<bool>(ConvexEnable, ConvexPropertiesGroup);
+	}
+
 	if (!ManagedCollection->HasAttribute(ConvexFractionRemoveAttribute, ConvexPropertiesGroup))
 	{
 		ManagedCollection->AddAttribute<float>(ConvexFractionRemoveAttribute, ConvexPropertiesGroup);
@@ -45,9 +51,6 @@ void FGeometryCollectionConvexPropertiesInterface::SetDefaultProperty()
 	TManagedArray<int32>& Index = ManagedCollection->GetAttribute<int32>(ConvexIndexAttribute, ConvexPropertiesGroup);
 	if (Index.Num() == 0 || !Index.Contains(INDEX_NONE))
 	{
-		TManagedArray<float>& FractionRemove = ManagedCollection->GetAttribute<float>(ConvexFractionRemoveAttribute, ConvexPropertiesGroup);
-		TManagedArray<float>& SimplificationThreshold = ManagedCollection->GetAttribute<float>(ConvexSimplificationThresholdAttribute, ConvexPropertiesGroup);
-
 		int32 DefaultAttributeIndex = ManagedCollection->AddElements(1, ConvexPropertiesGroup);
 		Index[DefaultAttributeIndex] = INDEX_NONE;
 		SetConvexProperties(DefaultProperty);
@@ -61,12 +64,14 @@ FGeometryCollectionConvexPropertiesInterface::GetDefaultProperty() const
 	const TManagedArray<int32>& Index = ManagedCollection->GetAttribute<int32>(ConvexIndexAttribute, ConvexPropertiesGroup);
 	if (Index.Contains(INDEX_NONE))
 	{
+		TManagedArray<bool>& Enable = ManagedCollection->GetAttribute<bool>(ConvexEnable, ConvexPropertiesGroup);
 		const TManagedArray<float>& FractionRemove = ManagedCollection->GetAttribute<float>(ConvexFractionRemoveAttribute, ConvexPropertiesGroup);
 		const TManagedArray<float>& SimplificationThreshold = ManagedCollection->GetAttribute<float>(ConvexSimplificationThresholdAttribute, ConvexPropertiesGroup);
 
 		int32 DefaultIndex = Index.Find(INDEX_NONE);
 		if (0 <= DefaultIndex && DefaultIndex < Index.Num())
 		{
+			DefaultProperty.Enable = Enable[DefaultIndex];
 			DefaultProperty.FractionRemove = FractionRemove[DefaultIndex];
 			DefaultProperty.SimplificationThreshold = SimplificationThreshold[DefaultIndex];
 		}
@@ -95,9 +100,11 @@ FGeometryCollectionConvexPropertiesInterface::GetConvexProperties(int TransformG
 	int32 PropIndex = Index.Find(TransformGroupIndex);
 	if (0 <= PropIndex && PropIndex < Index.Num())
 	{
+		TManagedArray<bool>& Enable = ManagedCollection->GetAttribute<bool>(ConvexEnable, ConvexPropertiesGroup);
 		const TManagedArray<float>& FractionRemove = ManagedCollection->GetAttribute<float>(ConvexFractionRemoveAttribute, ConvexPropertiesGroup);
 		const TManagedArray<float>& SimplificationThreshold = ManagedCollection->GetAttribute<float>(ConvexSimplificationThresholdAttribute, ConvexPropertiesGroup);
 
+		ConvexProperty.Enable = Enable[PropIndex];
 		ConvexProperty.FractionRemove = FractionRemove[PropIndex];
 		ConvexProperty.SimplificationThreshold = SimplificationThreshold[PropIndex];
 	}
@@ -109,6 +116,7 @@ FGeometryCollectionConvexPropertiesInterface::SetConvexProperties(const FConvexC
 {
 	FConvexCreationProperties Property;
 	TManagedArray<int32>& Index = ManagedCollection->GetAttribute<int32>(ConvexIndexAttribute, ConvexPropertiesGroup);
+	TManagedArray<bool>& Enable = ManagedCollection->GetAttribute<bool>(ConvexEnable, ConvexPropertiesGroup);
 	TManagedArray<float>& FractionRemove = ManagedCollection->GetAttribute<float>(ConvexFractionRemoveAttribute, ConvexPropertiesGroup);
 	TManagedArray<float>& SimplificationThreshold = ManagedCollection->GetAttribute<float>(ConvexSimplificationThresholdAttribute, ConvexPropertiesGroup);
 
@@ -125,6 +133,7 @@ FGeometryCollectionConvexPropertiesInterface::SetConvexProperties(const FConvexC
 	if (0 <= AttributeIndex && AttributeIndex < Index.Num())
 	{
 		Index[AttributeIndex] = TransformGroupIndex;
+		Enable[AttributeIndex] = InConvexAttributes.Enable;
 		FractionRemove[AttributeIndex] = InConvexAttributes.FractionRemove;
 		SimplificationThreshold[AttributeIndex] = InConvexAttributes.SimplificationThreshold;
 	}
diff --git a/Engine/Source/Runtime/Experimental/Chaos/Public/GeometryCollection/GeometryCollectionConvexPropertiesInterface.h b/Engine/Source/Runtime/Experimental/Chaos/Public/GeometryCollection/GeometryCollectionConvexPropertiesInterface.h
--- a/Engine/Source/Runtime/Experimental/Chaos/Public/GeometryCollection/GeometryCollectionConvexPropertiesInterface.h
+++ b/Engine/Source/Runtime/Experimental/Chaos/Public/GeometryCollection/GeometryCollectionConvexPropertiesInterface.h
@@ -18,12 +18,15 @@ public :
 	// Attribute
 	static const FName ConvexIndexAttribute;
 	// Attribute
+	static const FName ConvexEnable;
+	// Attribute
 	static const FName ConvexFractionRemoveAttribute;
 	// Attribute
 	static const FName ConvexSimplificationThresholdAttribute;
 
 
 	struct FConvexCreationProperties {
+		bool Enable = true;
 		float FractionRemove = 0.5f;
 		float SimplificationThreshold = 10.0f;
 	};
diff --git a/Engine/Source/Runtime/Experimental/GeometryCollectionEngine/Private/GeometryCollection/GeometryCollectionComponent.cpp b/Engine/Source/Runtime/Experimental/GeometryCollectionEngine/Private/GeometryCollection/GeometryCollectionComponent.cpp
--- a/Engine/Source/Runtime/Experimental/GeometryCollectionEngine/Private/GeometryCollection/GeometryCollectionComponent.cpp
+++ b/Engine/Source/Runtime/Experimental/GeometryCollectionEngine/Private/GeometryCollection/GeometryCollectionComponent.cpp
@@ -2242,6 +2242,7 @@ FGeometryCollectionEdit::FGeometryCollectionEdit(UGeometryCollectionComponent* I
 	if (EnumHasAnyFlags(EditUpdate, GeometryCollection::EEditUpdate::Rest) && GetRestCollection())
 	{
 		Component->Modify();
+		GetRestCollection()->Modify();
 	}
 }
 
diff --git a/Engine/Source/Runtime/Experimental/GeometryCollectionEngine/Private/GeometryCollection/GeometryCollectionObject.cpp b/Engine/Source/Runtime/Experimental/GeometryCollectionEngine/Private/GeometryCollection/GeometryCollectionObject.cpp
--- a/Engine/Source/Runtime/Experimental/GeometryCollectionEngine/Private/GeometryCollection/GeometryCollectionObject.cpp
+++ b/Engine/Source/Runtime/Experimental/GeometryCollectionEngine/Private/GeometryCollection/GeometryCollectionObject.cpp
@@ -1163,9 +1163,7 @@ bool UGeometryCollection::Modify(bool bAlwaysMarkDirty /*= true*/)
 	UPackage* Package = GetOutermost();
 	if (Package->IsDirty())
 	{
-		UpdateConvexGeometry();
 		InvalidateCollection();
-		ValidateSizeSpecificDataDefaults();
 	}
 
 	return bSuperResult;
diff --git a/Engine/Source/Runtime/Experimental/GeometryCollectionEngine/Public/GeometryCollection/GeometryCollectionComponent.h b/Engine/Source/Runtime/Experimental/GeometryCollectionEngine/Public/GeometryCollection/GeometryCollectionComponent.h
--- a/Engine/Source/Runtime/Experimental/GeometryCollectionEngine/Public/GeometryCollection/GeometryCollectionComponent.h
+++ b/Engine/Source/Runtime/Experimental/GeometryCollectionEngine/Public/GeometryCollection/GeometryCollectionComponent.h
@@ -151,11 +151,11 @@ namespace GeometryCollection
 		/** No update. */
 		None = 0,
 		/** Mark the rest collection as changed. */
-		Rest = 1 << 1,
+		Rest = 1 << 0,
 		/** Recreate the physics state (proxy). */
-		Physics = 1 << 2,
+		Physics = 1 << 1,
 		/** Reset the dynamic collection. */
-		Dynamic = 1 << 4,
+		Dynamic = 1 << 2,
 		/** Mark the rest collection as changed, and recreate the physics state (proxy). */
 		RestPhysics = Rest | Physics,
 		/** Reset dynamic collection, mark the rest collection as changed, and recreate the physics state (proxy). */
