Chaos - prune invalid edge collisions in triangle meshes and heighfields
- add option to prune invalid edge collisions between different shapes (disabled by default)

diff --git a/Engine/Source/Runtime/Experimental/Chaos/Private/Chaos/ChaosDebugDraw.cpp b/Engine/Source/Runtime/Experimental/Chaos/Private/Chaos/ChaosDebugDraw.cpp
--- a/Engine/Source/Runtime/Experimental/Chaos/Private/Chaos/ChaosDebugDraw.cpp
+++ b/Engine/Source/Runtime/Experimental/Chaos/Private/Chaos/ChaosDebugDraw.cpp
@@ -845,7 +845,8 @@ namespace Chaos
 					// Static friction, no restitution = green
 					// Inactive = gray
 					FColor DiscColor = FColor(200, 0, 0);
-					FColor NormalColor = FColor(200, 0, 0);
+					FColor PlaneNormalColor = FColor(200, 0, 0);
+					FColor EdgeNormalColor = FColor(200, 150, 0);
 					FColor ImpulseColor = FColor(0, 0, 200);
 					FColor PushOutColor = FColor(200, 200, 0);
 					FColor PushOutImpusleColor = FColor(0, 200, 200);
@@ -856,7 +857,8 @@ namespace Chaos
 					if (!bIsActive)
 					{
 						DiscColor = FColor(100, 100, 100);
-						NormalColor = FColor(100, 100, 100);
+						PlaneNormalColor = FColor(100, 0, 0);
+						EdgeNormalColor = FColor(100, 80, 0);
 					}
 
 					const FVec3 WorldPointLocation = SpaceTransform.TransformPosition(PointLocation);
@@ -885,6 +887,7 @@ namespace Chaos
 					}
 					if (Settings.ContactLen > 0)
 					{
+						FColor NormalColor = ((ManifoldPoint.ContactPoint.ContactType != EContactPointType::EdgeEdge) ? PlaneNormalColor : EdgeNormalColor);
 						FColor C1 = (ColorScale * NormalColor).ToFColor(false);
 						FDebugDrawQueue::GetInstance().DrawDebugLine(WorldPlaneLocation, WorldPlaneLocation + Settings.DrawScale * Settings.ContactLen * WorldPlaneNormal, C1, false, KINDA_SMALL_NUMBER, Settings.DrawPriority, Settings.LineThickness);
 					}
@@ -898,10 +901,13 @@ namespace Chaos
 					FDebugDrawQueue::GetInstance().DrawDebugCircle(WorldPointLocation, 0.5f * Settings.DrawScale * Settings.ContactWidth, 12, DiscColor, false, KINDA_SMALL_NUMBER, Settings.DrawPriority, Settings.LineThickness, Axes.GetUnitAxis(EAxis::Y), Axes.GetUnitAxis(EAxis::Z), false);
 
 					// Previous points
-					const FVec3 WorldPrevPointLocation = SpaceTransform.TransformPosition(PointTransform.TransformPosition(ManifoldPoint.ShapeAnchorPoints[ContactPointOwner]));
-					const FVec3 WorldPrevPlaneLocation = SpaceTransform.TransformPosition(PlaneTransform.TransformPosition(ManifoldPoint.ShapeAnchorPoints[ContactPlaneOwner]));
-					FDebugDrawQueue::GetInstance().DrawDebugLine(WorldPrevPointLocation, WorldPointLocation, FColor::White, false, KINDA_SMALL_NUMBER, Settings.DrawPriority, Settings.LineThickness);
-					FDebugDrawQueue::GetInstance().DrawDebugLine(WorldPrevPlaneLocation, WorldPlaneLocation, FColor::White, false, KINDA_SMALL_NUMBER, Settings.DrawPriority, Settings.LineThickness);
+					if (bIsActive)
+					{
+						const FVec3 WorldPrevPointLocation = SpaceTransform.TransformPosition(PointTransform.TransformPosition(ManifoldPoint.ShapeAnchorPoints[ContactPointOwner]));
+						const FVec3 WorldPrevPlaneLocation = SpaceTransform.TransformPosition(PlaneTransform.TransformPosition(ManifoldPoint.ShapeAnchorPoints[ContactPlaneOwner]));
+						FDebugDrawQueue::GetInstance().DrawDebugLine(WorldPrevPointLocation, WorldPointLocation, FColor::White, false, KINDA_SMALL_NUMBER, Settings.DrawPriority, Settings.LineThickness);
+						FDebugDrawQueue::GetInstance().DrawDebugLine(WorldPrevPlaneLocation, WorldPlaneLocation, FColor::White, false, KINDA_SMALL_NUMBER, Settings.DrawPriority, Settings.LineThickness);
+					}
 
 					// Whether restored
 					if (Settings.ContactInfoWidth > 0)
@@ -1496,10 +1502,11 @@ namespace Chaos
 		{
 			if (FDebugDrawQueue::IsDebugDrawingEnabled())
 			{
-				CollisionAllocator.VisitCollisions(
-					[&](const FPBDCollisionConstraint* Collision)
+				CollisionAllocator.VisitConstCollisions(
+					[&](const FPBDCollisionConstraint& Collision)
 					{
-						DrawCollisionImpl(SpaceTransform, Collision, ColorScale, GetChaosDebugDrawSettings(Settings));
+						DrawCollisionImpl(SpaceTransform, &Collision, ColorScale, GetChaosDebugDrawSettings(Settings));
+						return ECollisionVisitorResult::Continue;
 					});
 			}
 		}
diff --git a/Engine/Source/Runtime/Experimental/Chaos/Private/Chaos/Collision/ParticleCollisions.cpp b/Engine/Source/Runtime/Experimental/Chaos/Private/Chaos/Collision/ParticleCollisions.cpp
--- a/Engine/Source/Runtime/Experimental/Chaos/Private/Chaos/Collision/ParticleCollisions.cpp
+++ b/Engine/Source/Runtime/Experimental/Chaos/Private/Chaos/Collision/ParticleCollisions.cpp
@@ -1,5 +1,6 @@
 // Copyright Epic Games, Inc. All Rights Reserved.
 #include "Chaos/Collision/ParticleCollisions.h"
+#include "Chaos/Collision/CollisionConstraintAllocator.h"
 #include "Chaos/Collision/ParticlePairMidPhase.h"
 #include "Chaos/ParticleHandle.h"
 #include "Chaos/PBDCollisionConstraints.h"
@@ -43,13 +44,4 @@ namespace Chaos
 		}
 	}
 
-
-	void FParticleCollisions::VisitCollisions(const FPBDCollisionVisitor& Visitor) const
-	{
-		VisitConstMidPhases([&Visitor](const FParticlePairMidPhase& MidPhase)
-			{
-				MidPhase.VisitCollisions(Visitor);
-			});
-	}
-
 }
\ No newline at end of file
diff --git a/Engine/Source/Runtime/Experimental/Chaos/Private/Chaos/Collision/ParticlePairMidPhase.cpp b/Engine/Source/Runtime/Experimental/Chaos/Private/Chaos/Collision/ParticlePairMidPhase.cpp
--- a/Engine/Source/Runtime/Experimental/Chaos/Private/Chaos/Collision/ParticlePairMidPhase.cpp
+++ b/Engine/Source/Runtime/Experimental/Chaos/Private/Chaos/Collision/ParticlePairMidPhase.cpp
@@ -1,5 +1,6 @@
 // Copyright Epic Games, Inc. All Rights Reserved.
 #include "Chaos/Collision/ParticlePairMidPhase.h"
+#include "Chaos/Collision/CollisionConstraintAllocator.h"
 #include "Chaos/Collision/CollisionContext.h"
 #include "Chaos/Collision/CollisionFilter.h"
 #include "Chaos/Collision/PBDCollisionConstraint.h"
@@ -153,12 +154,6 @@ namespace Chaos
 	{
 	}
 
-	inline bool FSingleShapePairCollisionDetector::IsUsedSince(const int32 Epoch) const
-	{
-		// If we have no constraint it was never used, so this check is always false
-		return (Constraint != nullptr) && (LastUsedEpoch >= Epoch);
-	}
-
 	bool FSingleShapePairCollisionDetector::DoBoundsOverlap(const FReal CullDistance)
 	{
 		PHYSICS_CSV_SCOPED_EXPENSIVE(PhysicsVerbose, NarrowPhase_ShapeBounds);
@@ -623,18 +618,6 @@ namespace Chaos
 		}
 	}
 
-	void FMultiShapePairCollisionDetector::VisitCollisions(const int32 LastEpoch, const FPBDCollisionVisitor& Visitor) const
-	{
-		for (auto& KVP : Constraints)
-		{
-			const TUniquePtr<FPBDCollisionConstraint>& Constraint = KVP.Value;
-			if (Constraint->GetContainerCookie().LastUsedEpoch >= LastEpoch)
-			{
-				Visitor(Constraint.Get());
-			}
-		}
-	}
-
 
 	////////////////////////////////////////////////////////////////////////////////////////
 	////////////////////////////////////////////////////////////////////////////////////////
@@ -663,6 +646,11 @@ namespace Chaos
 		Reset();
 	}
 
+	int32 FParticlePairMidPhase::GetCurrentEpoch() const
+	{
+		return CollisionAllocator->GetCurrentEpoch();
+	}
+
 	void FParticlePairMidPhase::DetachParticle(FGeometryParticleHandle* Particle)
 	{
 		Reset();
@@ -901,24 +889,6 @@ namespace Chaos
 		}
 	}
 
-	void FParticlePairMidPhase::VisitCollisions(const FPBDCollisionVisitor& Visitor) const
-	{
-		const int32 LastEpoch = IsSleeping() ? LastUsedEpoch : CollisionAllocator->GetCurrentEpoch();
-		for (const FSingleShapePairCollisionDetector& ShapePair : ShapePairDetectors)
-		{
-			if (ShapePair.IsUsedSince(LastEpoch))
-			{
-				Visitor(ShapePair.GetConstraint());
-			}
-		}
-
-		for (const FMultiShapePairCollisionDetector& MultiShapePair : MultiShapePairDetectors)
-		{
-			MultiShapePair.VisitCollisions(LastEpoch, Visitor);
-		}
-
-	}
-
 }
 
 
diff --git a/Engine/Source/Runtime/Experimental/Chaos/Private/Chaos/CollisionOneShotManifolds.cpp b/Engine/Source/Runtime/Experimental/Chaos/Private/Chaos/CollisionOneShotManifolds.cpp
--- a/Engine/Source/Runtime/Experimental/Chaos/Private/Chaos/CollisionOneShotManifolds.cpp
+++ b/Engine/Source/Runtime/Experimental/Chaos/Private/Chaos/CollisionOneShotManifolds.cpp
@@ -51,7 +51,12 @@ namespace Chaos
 	FAutoConsoleVariableRef CVarChaos_Collision_GJKEpsilon(TEXT("p.Chaos.Collision.GJKEpsilon"), Chaos_Collision_GJKEpsilon, TEXT(""));
 	FAutoConsoleVariableRef CVarChaos_Collision_EPAEpsilon(TEXT("p.Chaos.Collision.EPAEpsilon"), Chaos_Collision_EPAEpsilon, TEXT(""));
 
-	
+	// Whether to prune spurious edge contacts in trimesh nd heighfield contacts (pretty much essential)
+	bool bChaos_Collision_EnableEdgePrune = true;
+	FRealSingle Chaos_Collision_EdgePrunePlaneDistance = 3.0;
+	FAutoConsoleVariableRef CVarChaos_Collision_EnableEdgePrune(TEXT("p.Chaos.Collision.EnableEdgePrune"), bChaos_Collision_EnableEdgePrune, TEXT(""));
+	FAutoConsoleVariableRef CVarChaos_Collision_EdgePrunePlaneDistance(TEXT("p.Chaos.Collision.EdgePrunePlaneDistance"), Chaos_Collision_EdgePrunePlaneDistance, TEXT(""));
+
 	namespace Collisions
 	{
 		// Forward delarations we need from CollisionRestitution.cpp
@@ -240,11 +245,64 @@ namespace Chaos
 			return OutPointCount; // This should always be 4
 		}
 
+		// Remove spurious edge contact by removing edge-edge contact that are in the plane of another vertex-plane contact.
+		void PruneEdgeContactPoints(TArray<FContactPoint>& ContactPoints, const FReal Epsilon)
+		{
+			if (!bChaos_Collision_EnableEdgePrune)
+			{
+				return;
+			}
+
+			int32 NumContactPoints = ContactPoints.Num();
+			for (int32 EdgeContactIndex = 0; EdgeContactIndex < NumContactPoints; ++EdgeContactIndex)
+			{
+				const FContactPoint& EdgeContactPoint = ContactPoints[EdgeContactIndex];
+
+				if (EdgeContactPoint.ContactType == EContactPointType::EdgeEdge)
+				{
+					for (int32 PlaneContactIndex = 0; PlaneContactIndex < NumContactPoints; ++PlaneContactIndex)
+					{
+						const FContactPoint PlaneContactPoint = ContactPoints[PlaneContactIndex];
+
+						// Does the edge-edge contact lie in the plane of the vertex-plane contact...?
+						int32 PlaneShapeIndex = INDEX_NONE;
+						if (PlaneContactPoint.ContactType == EContactPointType::PlaneVertex)
+						{
+							PlaneShapeIndex = 0;
+						}
+						else if (PlaneContactPoint.ContactType == EContactPointType::VertexPlane)
+						{
+							PlaneShapeIndex = 1;
+						}
+						if (PlaneShapeIndex != INDEX_NONE)
+						{
+							const FVec3 DeltaPos = PlaneContactPoint.ShapeContactPoints[PlaneShapeIndex] - EdgeContactPoint.ShapeContactPoints[PlaneShapeIndex];
+							const FReal Distance = FVec3::DotProduct(DeltaPos, PlaneContactPoint.ShapeContactNormal);
+							if (FMath::Abs(Distance) < Epsilon)
+							{
+								// Note: Cannot use RemoveAtSwap unless we re-sort at the end. See ReduceManifoldContactPointsTriangeMesh
+								const bool bAllowShrinking = false;
+								ContactPoints.RemoveAt(EdgeContactIndex, 1, bAllowShrinking);
+								--EdgeContactIndex;
+								--NumContactPoints;
+								break;
+							}
+						}
+					}
+				}
+			}
+		}
+
 		// Reduce the number of contact points (in place)
 		// Prerequisites to calling this function:
 		// ContactPoints are sorted on phi (ascending)
 		void ReduceManifoldContactPointsTriangeMesh(TArray<FContactPoint>& ContactPoints)
 		{
+			// Remove edge contacts that are "hidden" by face contacts
+			// EdgePruneDistance should be some fraction of the convex margin...
+			const FReal EdgePruneDistance = Chaos_Collision_EdgePrunePlaneDistance;
+			PruneEdgeContactPoints(ContactPoints, EdgePruneDistance);
+
 			if (ContactPoints.Num() <= 4)
 			{
 				return;
@@ -746,6 +804,9 @@ namespace Chaos
 					return;
 				}
 
+				// This is an edge-edge contact
+				GJKContactPoint.ContactType = EContactPointType::EdgeEdge;
+
 				// @todo(chaos): remove this if we ditch convex margins
 				// We only need to find the best edges when we have convex margins enabled because then the GJK result
 				// is a contact on the rounded margin-reduced shape and not the actual shape. If we are using zero
@@ -767,10 +828,21 @@ namespace Chaos
 					const FVec3 EdgePos1In2 = Convex2ToConvex1Transform.InverseTransformPositionNoScale(ShapeEdgePos1);
 					const FVec3 EdgePos2In2 = ShapeEdgePos2;
 					const FReal EdgePhi = FVec3::DotProduct(EdgePos1In2 - EdgePos2In2, GJKContactPoint.ShapeContactNormal);
-					const FVec3 EdgePosIn2 = FReal(0.5) * (EdgePos1In2 + EdgePos2In2);
 
-					GJKContactPoint.ShapeContactPoints[0] = Convex2ToConvex1Transform.TransformPositionNoScale(EdgePosIn2 + FReal(0.5) * EdgePhi * GJKContactPoint.ShapeContactNormal);
-					GJKContactPoint.ShapeContactPoints[1] = EdgePosIn2 - FReal(0.5) * EdgePhi * GJKContactPoint.ShapeContactNormal;
+					// We now have an accurate separation, so reject points beyond cull distance.
+					// This is quite important because if we use a larger cull distance for edge-edge contacts than for vertex-plane
+					// contacts, the pruning of unnecessary edge-edge contacts does not work (Tri Mesh and heighfield)
+					if (EdgePhi > Constraint.GetCullDistance())
+					{
+						return;
+					}
+
+					// @todo(chaos): we leave the contact point on the second shape because this is better for triangle meshes (see PruneContactEdges). For other shapes it should probably be the average
+					//const FVec3 EdgePosIn2 = FReal(0.5) * (EdgePos1In2 + EdgePos2In2);
+					//GJKContactPoint.ShapeContactPoints[0] = Convex2ToConvex1Transform.TransformPositionNoScale(EdgePosIn2 + FReal(0.5) * EdgePhi * GJKContactPoint.ShapeContactNormal);
+					//GJKContactPoint.ShapeContactPoints[1] = EdgePosIn2 - FReal(0.5) * EdgePhi * GJKContactPoint.ShapeContactNormal;
+					GJKContactPoint.ShapeContactPoints[0] = Convex2ToConvex1Transform.TransformPositionNoScale(EdgePos2In2 + EdgePhi * GJKContactPoint.ShapeContactNormal);
+					GJKContactPoint.ShapeContactPoints[1] = EdgePos2In2;
 					GJKContactPoint.Phi = EdgePhi;
 					// Normal unchanged from GJK result
 				}
@@ -861,6 +933,9 @@ namespace Chaos
 				}
 			}
 
+			// This is a vertex-plane contact
+			EContactPointType ContactType = ReferenceFaceConvex1 ? EContactPointType::PlaneVertex : EContactPointType::VertexPlane;
+
 			// Generate the contact points from the clipped vertices
 			{
 				SCOPE_CYCLE_COUNTER_MANIFOLD_ADDFACEVERTEX();
@@ -879,6 +954,7 @@ namespace Chaos
 					ContactPoint.ShapeContactPoints[1] = ReferenceFaceConvex1 ? ClippedPointInOtherCoordinates : PointProjectedOntoReferenceFace;
 					ContactPoint.ShapeContactNormal = SeparationDirectionLocalConvex2;
 					ContactPoint.Phi = FVec3::DotProduct(PointProjectedOntoReferenceFace - VertexInReferenceCoordinates, ReferenceFaceConvex1 ? SeparationDirectionLocalConvex1 : -SeparationDirectionLocalConvex2);				
+					ContactPoint.ContactType = ContactType;
 
 					Constraint.AddOneshotManifoldContact(ContactPoint);
 				}
diff --git a/Engine/Source/Runtime/Experimental/Chaos/Private/Chaos/PBDCollisionConstraints.cpp b/Engine/Source/Runtime/Experimental/Chaos/Private/Chaos/PBDCollisionConstraints.cpp
--- a/Engine/Source/Runtime/Experimental/Chaos/Private/Chaos/PBDCollisionConstraints.cpp
+++ b/Engine/Source/Runtime/Experimental/Chaos/Private/Chaos/PBDCollisionConstraints.cpp
@@ -43,7 +43,7 @@
 
 namespace Chaos
 {
-	extern int32 UseLevelsetCollision;
+	extern FRealSingle Chaos_Collision_EdgePrunePlaneDistance;
 
 	int32 CollisionParticlesBVHDepth = 4;
 	FAutoConsoleVariableRef CVarCollisionParticlesBVHDepth(TEXT("p.CollisionParticlesBVHDepth"), CollisionParticlesBVHDepth, TEXT("The maximum depth for collision particles bvh"));
@@ -115,6 +115,7 @@ namespace Chaos
 		, bEnableCollisions(true)
 		, bEnableRestitution(true)
 		, bHandlesEnabled(true)
+		, bEnableEdgePruning(true)
 		, bIsDeterministic(false)
 		, bCanDisableContacts(true)
 		, GravityDirection(FVec3(0,0,-1))
@@ -281,6 +282,9 @@ namespace Chaos
 		// Prune the unused contacts
 		ConstraintAllocator.EndDetectCollisions();
 
+		// Disable any edge collisions that are hidden by face collisions
+		PruneEdgeCollisions();
+
 		if (bIsDeterministic)
 		{
 			ConstraintAllocator.SortConstraintsHandles();
@@ -652,4 +656,89 @@ namespace Chaos
 
 		return *GetConstraints()[Index];
 	}
+
+	void FPBDCollisionConstraints::PruneEdgeCollisions()
+	{
+		if (bEnableEdgePruning)
+		{
+			for (auto& ParticleHandle : Particles.GetNonDisabledDynamicView())
+			{
+				if ((ParticleHandle.CollisionConstraintFlags() & (uint32)ECollisionConstraintFlags::CCF_SmoothEdgeCollisions) != 0)
+				{
+					PruneParticleEdgeCollisions(ParticleHandle.Handle());
+				}
+			}
+		}
+	}
+
+	void FPBDCollisionConstraints::PruneParticleEdgeCollisions(FGeometryParticleHandle* Particle)
+	{
+		FParticleCollisions& ParticleCollision = Particle->ParticleCollisions();
+
+		const FReal EdgePlaneTolerance = Chaos_Collision_EdgePrunePlaneDistance;
+
+		// Loop over edge collisions, then all plane collisions
+		// and remove the edge collision if it is hidden by a plane collision
+		// @todo(chaos): this should probably only disable individual manifold points
+		// @todo(chaos): perf issue: this processes contacts in world space, but we don't calculated that data until Gather. Fix this.
+		ParticleCollision.VisitCollisions(
+			[&ParticleCollision, EdgePlaneTolerance](FPBDCollisionConstraint& EdgeCollision)
+			{
+				if (EdgeCollision.IsEnabled())
+				{
+					for (const FManifoldPoint& EdgeManifoldPoint : EdgeCollision.GetManifoldPoints())
+					{
+						if (EdgeManifoldPoint.ContactPoint.ContactType == EContactPointType::EdgeEdge)
+						{
+							const FVec3 EdgePos0 = EdgeCollision.GetShapeWorldTransform0().TransformPositionNoScale(EdgeManifoldPoint.ContactPoint.ShapeContactPoints[0]);
+							const FVec3 EdgePos1 = EdgeCollision.GetShapeWorldTransform1().TransformPositionNoScale(EdgeManifoldPoint.ContactPoint.ShapeContactPoints[1]);
+
+							// Loop over plane collisions
+							ECollisionVisitorResult PlaneResult = ParticleCollision.VisitConstCollisions(
+								[&EdgeCollision, &EdgePos0, &EdgePos1, EdgePlaneTolerance](const FPBDCollisionConstraint& PlaneCollision)
+								{
+									if ((&PlaneCollision != &EdgeCollision) && PlaneCollision.IsEnabled())
+									{
+										for (const FManifoldPoint& PlaneManifoldPoint : PlaneCollision.GetManifoldPoints())
+										{
+											// If the edge position is in the plane, disable it
+											FVec3 PlanePos;
+											FVec3 PlaneNormal;
+											if (PlaneManifoldPoint.ContactPoint.ContactType == EContactPointType::PlaneVertex)
+											{
+												PlanePos = PlaneCollision.GetShapeWorldTransform0().TransformPositionNoScale(PlaneManifoldPoint.ContactPoint.ShapeContactPoints[0]);
+												PlaneNormal = PlaneCollision.GetShapeWorldTransform0().TransformVectorNoScale(PlaneManifoldPoint.ContactPoint.ShapeContactNormal);
+											}
+											else if (PlaneManifoldPoint.ContactPoint.ContactType == EContactPointType::VertexPlane)
+											{
+												PlanePos = PlaneCollision.GetShapeWorldTransform1().TransformPositionNoScale(PlaneManifoldPoint.ContactPoint.ShapeContactPoints[1]);
+												PlaneNormal = PlaneCollision.GetShapeWorldTransform1().TransformVectorNoScale(PlaneManifoldPoint.ContactPoint.ShapeContactNormal);
+											}
+											const FVec3 EdgePlaneDelta0 = EdgePos0 - PlanePos;
+											const FVec3 EdgePlaneDelta1 = EdgePos1 - PlanePos;
+											const FReal EdgePlaneDistance0 = FVec3::DotProduct(EdgePlaneDelta0, PlaneNormal);
+											const FReal EdgePlaneDistance1 = FVec3::DotProduct(EdgePlaneDelta1, PlaneNormal);
+											if ((FMath::Abs(EdgePlaneDistance0) < EdgePlaneTolerance) || (FMath::Abs(EdgePlaneDistance1) < EdgePlaneTolerance))
+											{
+												// The edge contact is hidden by a plane contact so disable it and stop the inner loop
+												EdgeCollision.SetDisabled(true);
+												return ECollisionVisitorResult::Stop;
+											}
+										}
+									}
+									return ECollisionVisitorResult::Continue;
+								});
+
+							// If we disabled this constraint, move to the next one and ignore remaining manifold points
+							if (PlaneResult == ECollisionVisitorResult::Stop)
+							{
+								break;
+							}
+						}
+					}
+				}
+				return ECollisionVisitorResult::Continue;
+			});
+	}
+
 }
diff --git a/Engine/Source/Runtime/Experimental/Chaos/Private/Chaos/PBDRigidsEvolutionGBF.cpp b/Engine/Source/Runtime/Experimental/Chaos/Private/Chaos/PBDRigidsEvolutionGBF.cpp
--- a/Engine/Source/Runtime/Experimental/Chaos/Private/Chaos/PBDRigidsEvolutionGBF.cpp
+++ b/Engine/Source/Runtime/Experimental/Chaos/Private/Chaos/PBDRigidsEvolutionGBF.cpp
@@ -756,11 +756,11 @@ void FPBDRigidsEvolutionGBF::TransferJointConstraintCollisions()
 				const FRigidTransform3 ChildToParentTransform = ChildTransform.GetRelativeTransform(ParentTransform);
 
 				ChildParticle->Handle()->ParticleCollisions().VisitCollisions(
-					[&](const FPBDCollisionConstraint* ChildCollisionConstraint)
+					[&](const FPBDCollisionConstraint& ChildCollisionConstraint)
 					{
-						if (ChildCollisionConstraint->GetCCDType() != ECollisionCCDType::Disabled)
+						if (ChildCollisionConstraint.GetCCDType() != ECollisionCCDType::Disabled)
 						{
-							return;
+							return ECollisionVisitorResult::Continue;
 						}
 
 						// @todo(chaos): implemeent this
@@ -775,7 +775,7 @@ void FPBDRigidsEvolutionGBF::TransferJointConstraintCollisions()
 						//		to replace a legit collision with our fake one...probably
 						ensure(false);
 
-						const FGeometryParticleHandle* NewParentParticleConst = (ChildCollisionConstraint->GetParticle0() == ChildParticle->Handle()) ? ChildCollisionConstraint->GetParticle1() : ChildCollisionConstraint->GetParticle0();
+						const FGeometryParticleHandle* NewParentParticleConst = (ChildCollisionConstraint.GetParticle0() == ChildParticle->Handle()) ? ChildCollisionConstraint.GetParticle1() : ChildCollisionConstraint.GetParticle0();
 						
 						FGeometryParticleHandle* NewParticleA = const_cast<FGeometryParticleHandle*>(NewParentParticleConst);
 						FGeometryParticleHandle* NewParticleB = ParentParticle->Handle();
@@ -790,6 +790,8 @@ void FPBDRigidsEvolutionGBF::TransferJointConstraintCollisions()
 						// Add collision to the system
 						//FParticlePairMidPhase* MidPhase = CollisionAllocator.GetParticlePairMidPhase(NewParticleA, NewParticleB);
 						//MidPhase->InjectCollision(NewCollision);
+
+						return ECollisionVisitorResult::Continue;
 					});
 			}
 		}
diff --git a/Engine/Source/Runtime/Experimental/Chaos/Public/Chaos/Collision/CollisionConstraintAllocator.h b/Engine/Source/Runtime/Experimental/Chaos/Public/Chaos/Collision/CollisionConstraintAllocator.h
--- a/Engine/Source/Runtime/Experimental/Chaos/Public/Chaos/Collision/CollisionConstraintAllocator.h
+++ b/Engine/Source/Runtime/Experimental/Chaos/Public/Chaos/Collision/CollisionConstraintAllocator.h
@@ -282,20 +282,25 @@ namespace Chaos
 			Particle->ParticleCollisions().VisitMidPhases([Particle](FParticlePairMidPhase& MidPhase)
 				{
 					MidPhase.DetachParticle(Particle);
+					return ECollisionVisitorResult::Continue;
 				});
 		}
 
 		/**
 		 * @brief Iterate over all collisions, including sleeping ones
 		*/
-		void VisitCollisions(const TFunction<void(const FPBDCollisionConstraint*)>& Visitor) const
+		template<typename TLambda>
+		void VisitConstCollisions(const TLambda& Visitor) const
 		{
 			for (const auto& MidPhase : ParticlePairMidPhases)
 			{
-				MidPhase->VisitCollisions(Visitor);
+				if (MidPhase->VisitConstCollisions(Visitor) == ECollisionVisitorResult::Stop)
+				{
+					return;
+				}
 			}
 		}
-
+		
 	private:
 
 		void ProcessNewItems()
@@ -441,4 +446,155 @@ namespace Chaos
 		mutable TLockFreePointerListLIFO<FParticlePairMidPhase> NewParticlePairMidPhases;
 	};
 
+
+	///////////////////////////////////////////////////////////////////////////////////////////////
+	///////////////////////////////////////////////////////////////////////////////////////////////
+	//
+	// Below here is for code that is moved to avoid circular header dependencies
+	// See ParticlePairMidPhase.h and ParticleCollisions.h
+	//
+	///////////////////////////////////////////////////////////////////////////////////////////////
+	///////////////////////////////////////////////////////////////////////////////////////////////
+
+	template<typename TLambda>
+	inline ECollisionVisitorResult FMultiShapePairCollisionDetector::VisitCollisions(const int32 LastEpoch, const TLambda& Visitor)
+	{
+		for (auto& KVP : Constraints)
+		{
+			const TUniquePtr<FPBDCollisionConstraint>& Constraint = KVP.Value;
+			if (Constraint->GetContainerCookie().LastUsedEpoch >= LastEpoch)
+			{
+				if (Visitor(*Constraint) == ECollisionVisitorResult::Stop)
+				{
+					return ECollisionVisitorResult::Stop;
+				}
+			}
+		}
+		return ECollisionVisitorResult::Continue;
+	}
+
+	template<typename TLambda>
+	inline ECollisionVisitorResult FMultiShapePairCollisionDetector::VisitConstCollisions(const int32 LastEpoch, const TLambda& Visitor) const
+	{
+		for (auto& KVP : Constraints)
+		{
+			const TUniquePtr<FPBDCollisionConstraint>& Constraint = KVP.Value;
+			if (Constraint->GetContainerCookie().LastUsedEpoch >= LastEpoch)
+			{
+				if (Visitor(*Constraint) == ECollisionVisitorResult::Stop)
+				{
+					return ECollisionVisitorResult::Stop;
+				}
+			}
+		}
+		return ECollisionVisitorResult::Continue;
+	}
+
+	template<typename TLambda>
+	inline ECollisionVisitorResult FParticlePairMidPhase::VisitCollisions(const TLambda& Visitor)
+	{
+		const int32 LastEpoch = IsSleeping() ? LastUsedEpoch : GetCurrentEpoch();
+		for (FSingleShapePairCollisionDetector& ShapePair : ShapePairDetectors)
+		{
+			if (ShapePair.IsUsedSince(LastEpoch))
+			{
+				if (Visitor(*ShapePair.GetConstraint()) == ECollisionVisitorResult::Stop)
+				{
+					return ECollisionVisitorResult::Stop;
+				}
+			}
+		}
+
+		for (FMultiShapePairCollisionDetector& MultiShapePair : MultiShapePairDetectors)
+		{
+			if (MultiShapePair.VisitCollisions(LastEpoch, Visitor) == ECollisionVisitorResult::Stop)
+			{
+				return ECollisionVisitorResult::Stop;
+			}
+		}
+
+		return ECollisionVisitorResult::Continue;
+	}
+
+
+	template<typename TLambda>
+	inline ECollisionVisitorResult FParticlePairMidPhase::VisitConstCollisions(const TLambda& Visitor) const
+	{
+		const int32 LastEpoch = IsSleeping() ? LastUsedEpoch : GetCurrentEpoch();
+		for (const FSingleShapePairCollisionDetector& ShapePair : ShapePairDetectors)
+		{
+			if (ShapePair.IsUsedSince(LastEpoch))
+			{
+				if (Visitor(*ShapePair.GetConstraint()) == ECollisionVisitorResult::Stop)
+				{
+					return ECollisionVisitorResult::Stop;
+				}
+			}
+		}
+
+		for (const FMultiShapePairCollisionDetector& MultiShapePair : MultiShapePairDetectors)
+		{
+			if (MultiShapePair.VisitConstCollisions(LastEpoch, Visitor) == ECollisionVisitorResult::Stop)
+			{
+				return ECollisionVisitorResult::Stop;
+			}
+		}
+
+		return ECollisionVisitorResult::Continue;
+	}
+
+
+	template<typename TLambda>
+	inline ECollisionVisitorResult FParticleCollisions::VisitMidPhases(const TLambda& Lambda)
+	{
+		for (int32 Index = 0; Index < MidPhases.Num(); ++Index)
+		{
+			if (Lambda(*MidPhases[Index].Value) == ECollisionVisitorResult::Stop)
+			{
+				return ECollisionVisitorResult::Stop;
+			}
+		}
+		return ECollisionVisitorResult::Continue;
+	}
+
+	template<typename TLambda>
+	inline ECollisionVisitorResult FParticleCollisions::VisitConstMidPhases(const TLambda& Lambda) const
+	{
+		for (int32 Index = 0; Index < MidPhases.Num(); ++Index)
+		{
+			if (Lambda(*MidPhases[Index].Value) == ECollisionVisitorResult::Stop)
+			{
+				return ECollisionVisitorResult::Stop;
+			}
+		}
+		return ECollisionVisitorResult::Continue;
+	}
+
+	template<typename TLambda>
+	inline ECollisionVisitorResult FParticleCollisions::VisitCollisions(const TLambda& Visitor)
+	{
+		return VisitMidPhases([&Visitor](FParticlePairMidPhase& MidPhase)
+			{
+				if (MidPhase.VisitCollisions(Visitor) == ECollisionVisitorResult::Stop)
+				{
+					return ECollisionVisitorResult::Stop;
+				}
+				return ECollisionVisitorResult::Continue;
+			});
+	}
+
+	template<typename TLambda>
+	inline ECollisionVisitorResult FParticleCollisions::VisitConstCollisions(const TLambda& Visitor) const
+	{
+		return VisitConstMidPhases([&Visitor](const FParticlePairMidPhase& MidPhase)
+			{
+				if (MidPhase.VisitConstCollisions(Visitor) == ECollisionVisitorResult::Stop)
+				{
+					return ECollisionVisitorResult::Stop;
+				}
+				return ECollisionVisitorResult::Continue;
+			});
+	}
+
+
 }
\ No newline at end of file
diff --git a/Engine/Source/Runtime/Experimental/Chaos/Public/Chaos/Collision/CollisionVisitor.h b/Engine/Source/Runtime/Experimental/Chaos/Public/Chaos/Collision/CollisionVisitor.h
--- a/Engine/Source/Runtime/Experimental/Chaos/Public/Chaos/Collision/CollisionVisitor.h
+++ b/Engine/Source/Runtime/Experimental/Chaos/Public/Chaos/Collision/CollisionVisitor.h
@@ -8,5 +8,13 @@ namespace Chaos
 {
 	class FPBDCollisionConstraint;
 
-	using FPBDCollisionVisitor = TFunction<void(const FPBDCollisionConstraint*)>;
+	/**
+	 * @brief Return value in collision visitors to indicate whether we should continue visiting for stop
+	*/
+	enum class ECollisionVisitorResult
+	{
+		Stop,
+		Continue
+	};
+
 }
\ No newline at end of file
diff --git a/Engine/Source/Runtime/Experimental/Chaos/Public/Chaos/Collision/ContactPoint.h b/Engine/Source/Runtime/Experimental/Chaos/Public/Chaos/Collision/ContactPoint.h
--- a/Engine/Source/Runtime/Experimental/Chaos/Public/Chaos/Collision/ContactPoint.h
+++ b/Engine/Source/Runtime/Experimental/Chaos/Public/Chaos/Collision/ContactPoint.h
@@ -10,17 +10,20 @@
 
 namespace Chaos
 {
-	// 
-	// In principle, the low-level collision detection functions should not need to calculate the world-space contact data, but some of them
-	// do some of their work in world space and already have the results we need. To avoid duplicate work, all collision detection functions
-	// should also fill in the world-space data
+	/**
+	 * @brief Used in FContactPoint to indicate whether the contact is vertex-plane, edge-edge, etc
+	*/
+	enum class EContactPointType : int8
+	{
+		Unknown,
+		VertexPlane,
+		PlaneVertex,
+		EdgeEdge,
+		VertexVertex,
+	};
 
 	/**
 	 * @brief Data returned by the low-level collision functions
-	 * 
-	 * @note In principle, the low-level collision detection functions should not need to calculate the world-space contact data, but some of them
-	 * do some of their work in world space and already have the results we need. To avoid duplicate work, all collision detection functions
-	 * should also fill in the world-space data
 	*/
 	class CHAOS_API FContactPoint
 	{
@@ -37,9 +40,13 @@ namespace Chaos
 		// Face index of the shape we hit. Only valid for Heightfield and Trimesh contact points, otherwise INDEX_NONE
 		int32 FaceIndex;
 
+		// Whether this is a vertex-plane contact, edge-edge contact etc.
+		EContactPointType ContactType;
+
 		FContactPoint()
 			: Phi(TNumericLimits<FReal>::Max())
 			, FaceIndex(INDEX_NONE)
+			, ContactType(EContactPointType::Unknown)
 		{
 		}
 
diff --git a/Engine/Source/Runtime/Experimental/Chaos/Public/Chaos/Collision/ParticleCollisions.h b/Engine/Source/Runtime/Experimental/Chaos/Public/Chaos/Collision/ParticleCollisions.h
--- a/Engine/Source/Runtime/Experimental/Chaos/Public/Chaos/Collision/ParticleCollisions.h
+++ b/Engine/Source/Runtime/Experimental/Chaos/Public/Chaos/Collision/ParticleCollisions.h
@@ -11,7 +11,6 @@ namespace Chaos
 {
 	class FParticlePairMidPhase;
 
-
 	/**
 	 * @brief Knows about all the collisions detectors associated with a particular particle.
 	 * Used when particles are destroyed to remove perisstent collisions from the system, or
@@ -33,7 +32,7 @@ namespace Chaos
 		FParticleCollisions();
 		~FParticleCollisions();
 
-		int32 Num() const 
+		inline int32 Num() const 
 		{ 
 			return MidPhases.Num();
 		}
@@ -41,7 +40,7 @@ namespace Chaos
 		/**
 		 * @brief Clear the list of midphases. Only for use in shutdown.
 		*/
-		void Reset()
+		inline void Reset()
 		{
 			MidPhases.Reset();
 		}
@@ -78,36 +77,39 @@ namespace Chaos
 
 		/**
 		 * @brief Visit all of the midphases on the particle and call the specified function
+		 * @tparam TLambda visitor type with signature ECollisionVisitorResult(FParticlePairMidPhase&)
+		 *
 		 * @note Do not call RemoveMidPhase from the visitor
-		 * 
-		 * Lambda must have signature void(FParticlePairMidPhase&)
 		*/
 		template<typename TLambda>
-		void VisitMidPhases(const TLambda& Lambda)
-		{
-			for (int32 Index = 0; Index < MidPhases.Num(); ++Index)
-			{
-				Lambda(*MidPhases[Index].Value);
-			}
-		}
+		inline ECollisionVisitorResult VisitMidPhases(const TLambda& Lambda);
 
 		/**
 		 * @brief Visit all of the midphases on the particle and call the specified function
+		 * @tparam TLambda visitor type with signature void(const FParticlePairMidPhase&)
+		 * 
 		 * @note Do not call RemoveMidPhase from the visitor
-		 *
-		 * Lambda must have signature void(const FParticlePairMidPhase&)
 		*/
 		template<typename TLambda>
-		void VisitConstMidPhases(const TLambda& Lambda) const
-		{
-			for (int32 Index = 0; Index < MidPhases.Num(); ++Index)
-			{
-				Lambda(*MidPhases[Index].Value);
-			}
-		}
+		inline ECollisionVisitorResult VisitConstMidPhases(const TLambda& Lambda) const;
 
-		void VisitCollisions(const FPBDCollisionVisitor& Visitor) const;
+		/**
+		 * @brief Visit all the collisions on this particle
+		 * @tparam TLambda visitor type with signature void(FPBDCollisionParticle&)
+		 * 
+		 * @note do not delete constraint from the lambda. You may disable them though.
+		*/
+		template<typename TLambda>
+		inline ECollisionVisitorResult VisitCollisions(const TLambda& Visitor);
 
+		/**
+		 * @brief Visit all the collisions on this particle
+		 * @tparam TLambda visitor type with signature void(const FPBDCollisionParticle&)
+		 * 
+		 * @note do not delete constraint from the lambda. You may disable them though.
+		*/
+		template<typename TLambda>
+		inline ECollisionVisitorResult VisitConstCollisions(const TLambda& Visitor) const;
 
 	private:
 		FContainerType MidPhases;
diff --git a/Engine/Source/Runtime/Experimental/Chaos/Public/Chaos/Collision/ParticlePairMidPhase.h b/Engine/Source/Runtime/Experimental/Chaos/Public/Chaos/Collision/ParticlePairMidPhase.h
--- a/Engine/Source/Runtime/Experimental/Chaos/Public/Chaos/Collision/ParticlePairMidPhase.h
+++ b/Engine/Source/Runtime/Experimental/Chaos/Public/Chaos/Collision/ParticlePairMidPhase.h
@@ -56,7 +56,10 @@ namespace Chaos
 		/**
 		 * @brief Have we run collision detection since this Epoch (inclusive)
 		*/
-		bool IsUsedSince(const int32 CurrentEpoch) const;
+		inline bool IsUsedSince(const int32 Epoch) const
+		{
+			return (Constraint != nullptr) && (LastUsedEpoch >= Epoch);
+		}
 
 		/**
 		 * @brief Perform a bounds check and run the narrow phase if necessary
@@ -187,7 +190,17 @@ namespace Chaos
 		*/
 		void WakeCollisions(const int32 SleepEpoch);
 
-		void VisitCollisions(const int32 LastEpoch, const FPBDCollisionVisitor& Visitor) const;
+		/**
+		 * @brief Visit all the collisions
+		*/
+		template<typename TLambda>
+		ECollisionVisitorResult VisitCollisions(const int32 LastEpoch, const TLambda& Visitor);
+
+		/**
+		 * @brief Visit all the collisions
+		*/
+		template<typename TLambda>
+		ECollisionVisitorResult VisitConstCollisions(const int32 LastEpoch, const TLambda& Visitor) const;
 
 	private:
 		FPBDCollisionConstraint* FindConstraint(const FCollisionParticlePairConstraintKey& Key);
@@ -315,12 +328,22 @@ namespace Chaos
 		*/
 		void InjectCollision(const FPBDCollisionConstraint& Constraint);
 
+		int32 GetCurrentEpoch() const;
+
 		/**
-		 * @brief Call a function on each active collision constraint
-		 * This including sleeping constraints, but not constraints that are were not used on the last awake tick
+		 * @brief Call a lambda on each active collision constraint
+		 * This includes sleeping constraints, but skips constraints that are were not used on the last awake tick
 		 * but are still kept around as an optimization.
 		*/
-		void VisitCollisions(const FPBDCollisionVisitor& Visitor) const;
+		template<typename TLambda>
+		ECollisionVisitorResult VisitCollisions(const TLambda& Visitor);
+
+
+		/**
+		 * @brief Call a lambda on each active collision constraint
+		 */
+		template<typename TLambda>
+		ECollisionVisitorResult VisitConstCollisions(const TLambda& Visitor) const;
 
 		/**
 		 * @brief Cookie for use by FParticleCollisions
diff --git a/Engine/Source/Runtime/Experimental/Chaos/Public/Chaos/PBDCollisionConstraints.h b/Engine/Source/Runtime/Experimental/Chaos/Public/Chaos/PBDCollisionConstraints.h
--- a/Engine/Source/Runtime/Experimental/Chaos/Public/Chaos/PBDCollisionConstraints.h
+++ b/Engine/Source/Runtime/Experimental/Chaos/Public/Chaos/PBDCollisionConstraints.h
@@ -293,6 +293,12 @@ protected:
 	Collisions::FContactParticleParameters GetContactParticleParameters(const FReal Dt);
 	Collisions::FContactIterationParameters GetContactIterationParameters(const FReal Dt, const int32 Iteration, const int32 NumIterations, const int32 NumPairIterations, bool& bNeedsAnotherIteration);
 
+	// Call PruneParticleEdgeCollisions on all particles with ECollisionConstraintFlags::CCF_SmoothEdgeCollisions set in CollisionFlags
+	void PruneEdgeCollisions();
+
+	// Remove all edge collision on a particle that cannot be hit because a plane contact would prevent it
+	void PruneParticleEdgeCollisions(FGeometryParticleHandle* Particle);
+
 	// The "Legacy" functions handle the older solver types (GbfPbd and StandardPbd)
 	// @todo(chaos): remove legacy methods when the new solver is fully operational and used everywhere (RBAN)
 	void LegacyGatherInput(const FReal Dt, FPBDCollisionConstraint& Constraint, const int32 Particle0Level, const int32 Particle1Level, FPBDIslandSolverData& SolverData);
@@ -319,6 +325,7 @@ private:
 	bool bEnableCollisions;
 	bool bEnableRestitution;
 	bool bHandlesEnabled;
+	bool bEnableEdgePruning;
 	bool bIsDeterministic;
 
 	// This is passed to IterationParameters. If true, then an iteration can cull a contact
