Fix for inverted normal direction for triangle meshes with negative scales.
Caveat: if the original asset has inward normals and user tries to use negative scales to invert them, it won't work any more.

diff --git a/Engine/Source/Runtime/Experimental/Chaos/Private/Chaos/TriangleMeshImplicitObject.cpp b/Engine/Source/Runtime/Experimental/Chaos/Private/Chaos/TriangleMeshImplicitObject.cpp
--- a/Engine/Source/Runtime/Experimental/Chaos/Private/Chaos/TriangleMeshImplicitObject.cpp
+++ b/Engine/Source/Runtime/Experimental/Chaos/Private/Chaos/TriangleMeshImplicitObject.cpp
@@ -871,7 +871,7 @@ template <typename QueryGeomType, typename IdxType>
 struct FTriangleMeshSweepVisitor
 {
 	FTriangleMeshSweepVisitor(const FTriangleMeshImplicitObject& InTriMesh, const TArray<TVec3<IdxType>>& InElements, const QueryGeomType& InQueryGeom, const TRigidTransform<FReal,3>& InStartTM, const FVec3& InDir,
-		const FVec3& InScaledDirNormalized, const FReal InLengthScale, const FRigidTransform3& InScaledStartTM, const FReal InThickness, const bool InComputeMTD, FVec3 InTriMeshScale, bool InCullsBackFaceSweeps)
+		const FVec3& InScaledDirNormalized, const FReal InLengthScale, const FRigidTransform3& InScaledStartTM, const FReal InThickness, const bool InComputeMTD, FVec3 InTriMeshScale, FReal InCullsBackFaceSweepsCode)
 	: TriMesh(InTriMesh)
 	, Elements(InElements)
 	, StartTM(InStartTM)
@@ -879,7 +879,7 @@ struct FTriangleMeshSweepVisitor
 	, Dir(InDir)
 	, Thickness(InThickness)
 	, bComputeMTD(InComputeMTD)
-	, bCullsBackFaceSweeps(InCullsBackFaceSweeps)
+	, CullsBackFaceSweepsCode(InCullsBackFaceSweepsCode)
 	, ScaledDirNormalized(InScaledDirNormalized)
 	, LengthScale(InLengthScale)
 	, ScaledStartTM(InScaledStartTM)
@@ -921,10 +921,10 @@ struct FTriangleMeshSweepVisitor
 		TransformVertsHelper(TriMeshScale,TriIdx,TriMesh.MParticles,Elements,A,B,C);
 		FTriangle Tri(A, B, C);
 
-		if(bCullsBackFaceSweeps)
+		if(CullsBackFaceSweepsCode != 0)
 		{
 			const FVec3 TriNormal = FVec3::CrossProduct(B - A, C - A);
-			if (FVec3::DotProduct(TriNormal, ScaledDirNormalized) > 0)
+			if (FVec3::DotProduct(TriNormal, ScaledDirNormalized) * CullsBackFaceSweepsCode > 0)
 			{
 				return true;
 			}
@@ -964,7 +964,7 @@ struct FTriangleMeshSweepVisitor
 	const FVec3& Dir;
 	const FReal Thickness;
 	const bool bComputeMTD;
-	const bool bCullsBackFaceSweeps;
+	const FReal CullsBackFaceSweepsCode; // 0: no culling, 1/-1: winding order
 
 	// Cache these values for Scaled Triangle Mesh, as they are needed for transformation when sweeping against triangles.
 	FVec3 ScaledDirNormalized;
@@ -995,6 +995,12 @@ void ComputeScaledSweepInputs(FVec3 TriMeshScale, const FRigidTransform3& StartT
 	OutScaledStartTM = FRigidTransform3(StartTM.GetLocation() * TriMeshScale, StartTM.GetRotation());
 }
 
+FReal GetWindingOrder(const FVec3& Scale)
+{
+	const FVec3 SignVector = Scale.GetSignVector();
+	return SignVector.X * SignVector.Y * SignVector.Z;
+}
+
 template <typename QueryGeomType>
 bool FTriangleMeshImplicitObject::SweepGeomImp(const QueryGeomType& QueryGeom, const FRigidTransform3& StartTM, const FVec3& Dir, const FReal Length,
 	FReal& OutTime, FVec3& OutPosition, FVec3& OutNormal, int32& OutFaceIndex, const FReal Thickness, const bool bComputeMTD, FVec3 TriMeshScale) const
@@ -1009,8 +1015,9 @@ bool FTriangleMeshImplicitObject::SweepGeomImp(const QueryGeomType& QueryGeom, c
 	bool bHit = false;
 	auto LambdaHelper = [&](const auto& Elements)
 	{
+		const FReal CullsBackFaceRaycastCode = bCullsBackFaceRaycast ? GetWindingOrder(TriMeshScale) : 0.f;
 		using VisitorType = FTriangleMeshSweepVisitor<QueryGeomType,decltype(Elements[0][0])>;
-		VisitorType SQVisitor(*this,Elements, QueryGeom,StartTM,Dir,ScaledDirNormalized,LengthScale,ScaledStartTM,Thickness,bComputeMTD, TriMeshScale, bCullsBackFaceRaycast);
+		VisitorType SQVisitor(*this,Elements, QueryGeom,StartTM,Dir,ScaledDirNormalized,LengthScale,ScaledStartTM,Thickness,bComputeMTD, TriMeshScale, CullsBackFaceRaycastCode);
 
 
 		const FAABB3 QueryBounds = QueryGeom.BoundingBox().TransformedAABB(FRigidTransform3(FVec3::ZeroVector,StartTM.GetRotation()));
