update fracture tool convex to account for convexes autogenerating ("make convex hulls" is now "update convex hulls" and just transfers the tool settings to the geometry collection convex parameters)
+ add CanExceedFraction parameter to the convex parameters

#rnx

diff --git a/Engine/Plugins/Experimental/ChaosEditor/Source/FractureEditor/Private/FractureToolConvex.cpp b/Engine/Plugins/Experimental/ChaosEditor/Source/FractureEditor/Private/FractureToolConvex.cpp
--- a/Engine/Plugins/Experimental/ChaosEditor/Source/FractureEditor/Private/FractureToolConvex.cpp
+++ b/Engine/Plugins/Experimental/ChaosEditor/Source/FractureEditor/Private/FractureToolConvex.cpp
@@ -160,12 +160,12 @@ bool UFractureToolConvex::CanExecute() const
 
 FText UFractureToolConvex::GetDisplayText() const
 {
-	return FText(LOCTEXT("FractureToolConvex", "Make Convex Collision Volumes"));
+	return FText(LOCTEXT("FractureToolConvex", "Update Convex Collision Volumes"));
 }
 
 FText UFractureToolConvex::GetTooltipText() const
 {
-	return FText(LOCTEXT("FractureToolConvexTooltip", "This tool creates (non-overlapping) convex volumes for the bones of geometry collections"));
+	return FText(LOCTEXT("FractureToolConvexTooltip", "This tool visualizes and updates settings for convex hull generation on the geometry collections"));
 }
 
 FSlateIcon UFractureToolConvex::GetToolIcon() const
@@ -175,7 +175,7 @@ FSlateIcon UFractureToolConvex::GetToolIcon() const
 
 void UFractureToolConvex::RegisterUICommand(FFractureEditorCommands* BindingContext)
 {
-	UI_COMMAND_EXT(BindingContext, UICommandInfo, "Convex", "Convex", "Create (and visualize) a hierarchy of non-overlapping convex collision volumes for the bones of geometry collections.", EUserInterfaceActionType::ToggleButton, FInputChord());
+	UI_COMMAND_EXT(BindingContext, UICommandInfo, "Convex", "Convex", "Update (and visualize) a hierarchy of non-overlapping convex collision volumes for the bones of geometry collections.", EUserInterfaceActionType::ToggleButton, FInputChord());
 	BindingContext->MakeConvex = UICommandInfo;
 }
 
@@ -291,7 +291,14 @@ void UFractureToolConvex::AutoComputeConvex(const FFractureToolContext& Fracture
 
 int32 UFractureToolConvex::ExecuteFracture(const FFractureToolContext& FractureContext)
 {
-	AutoComputeConvex(FractureContext);
+	if (FractureContext.GetGeometryCollection().IsValid())
+	{
+		FGeometryCollectionConvexPropertiesInterface::FConvexCreationProperties Properties = FractureContext.GetGeometryCollection()->GetConvexProperties();
+		Properties.FractionRemove = ConvexSettings->FractionAllowRemove;
+		Properties.SimplificationThreshold = ConvexSettings->SimplificationDistanceThreshold;
+		Properties.CanExceedFraction = ConvexSettings->CanExceedFraction;
+		FractureContext.GetGeometryCollection()->SetConvexProperties(Properties);
+	}
 
 	return INDEX_NONE;
 }
diff --git a/Engine/Plugins/Experimental/ChaosEditor/Source/FractureEditor/Private/FractureToolConvex.h b/Engine/Plugins/Experimental/ChaosEditor/Source/FractureEditor/Private/FractureToolConvex.h
--- a/Engine/Plugins/Experimental/ChaosEditor/Source/FractureEditor/Private/FractureToolConvex.h
+++ b/Engine/Plugins/Experimental/ChaosEditor/Source/FractureEditor/Private/FractureToolConvex.h
@@ -65,7 +65,7 @@ public:
 	virtual FText GetDisplayText() const override;
 	virtual FText GetTooltipText() const override;
 	/** This is the Text that will appear on the button to execute the fracture **/
-	virtual FText GetApplyText() const override { return FText(NSLOCTEXT("Convex", "ExecuteConvex", "Make Convex Hulls")); }
+	virtual FText GetApplyText() const override { return FText(NSLOCTEXT("Convex", "ExecuteConvex", "Update Convex Hulls")); }
 	virtual FSlateIcon GetToolIcon() const override;
 	virtual void RegisterUICommand(FFractureEditorCommands* BindingContext) override;
 	virtual TArray<UObject*> GetSettingsObjects() const override;
diff --git a/Engine/Source/Runtime/Experimental/Chaos/Private/GeometryCollection/GeometryCollectionConvexPropertiesInterface.cpp b/Engine/Source/Runtime/Experimental/Chaos/Private/GeometryCollection/GeometryCollectionConvexPropertiesInterface.cpp
--- a/Engine/Source/Runtime/Experimental/Chaos/Private/GeometryCollection/GeometryCollectionConvexPropertiesInterface.cpp
+++ b/Engine/Source/Runtime/Experimental/Chaos/Private/GeometryCollection/GeometryCollectionConvexPropertiesInterface.cpp
@@ -8,6 +8,7 @@ const FName FGeometryCollectionConvexPropertiesInterface::ConvexIndexAttribute =
 const FName FGeometryCollectionConvexPropertiesInterface::ConvexEnable = "Enable";
 const FName FGeometryCollectionConvexPropertiesInterface::ConvexFractionRemoveAttribute = "FractionRemove";
 const FName FGeometryCollectionConvexPropertiesInterface::ConvexSimplificationThresholdAttribute = "SimplificationThreshold";
+const FName FGeometryCollectionConvexPropertiesInterface::ConvexCanExceedFractionAttribute = "CanExceedFraction";
 
 FGeometryCollectionConvexPropertiesInterface::FGeometryCollectionConvexPropertiesInterface(FGeometryCollection* InGeometryCollection)
 	: FManagedArrayInterface(InGeometryCollection)
@@ -42,6 +43,11 @@ FGeometryCollectionConvexPropertiesInterface::InitializeInterface()
 		ManagedCollection->AddAttribute<float>(ConvexSimplificationThresholdAttribute, ConvexPropertiesGroup);
 	}
 
+	if (!ManagedCollection->HasAttribute(ConvexCanExceedFractionAttribute, ConvexPropertiesGroup))
+	{
+		ManagedCollection->AddAttribute<float>(ConvexCanExceedFractionAttribute, ConvexPropertiesGroup);
+	}
+
 	SetDefaultProperty();
 }
 
@@ -67,6 +73,7 @@ FGeometryCollectionConvexPropertiesInterface::GetDefaultProperty() const
 		TManagedArray<bool>& Enable = ManagedCollection->GetAttribute<bool>(ConvexEnable, ConvexPropertiesGroup);
 		const TManagedArray<float>& FractionRemove = ManagedCollection->GetAttribute<float>(ConvexFractionRemoveAttribute, ConvexPropertiesGroup);
 		const TManagedArray<float>& SimplificationThreshold = ManagedCollection->GetAttribute<float>(ConvexSimplificationThresholdAttribute, ConvexPropertiesGroup);
+		const TManagedArray<float>& CanExceedFraction = ManagedCollection->GetAttribute<float>(ConvexCanExceedFractionAttribute, ConvexPropertiesGroup);
 
 		int32 DefaultIndex = Index.Find(INDEX_NONE);
 		if (0 <= DefaultIndex && DefaultIndex < Index.Num())
@@ -74,6 +81,7 @@ FGeometryCollectionConvexPropertiesInterface::GetDefaultProperty() const
 			DefaultProperty.Enable = Enable[DefaultIndex];
 			DefaultProperty.FractionRemove = FractionRemove[DefaultIndex];
 			DefaultProperty.SimplificationThreshold = SimplificationThreshold[DefaultIndex];
+			DefaultProperty.CanExceedFraction = CanExceedFraction[DefaultIndex];
 		}
 	}
 	return DefaultProperty;
@@ -103,10 +111,12 @@ FGeometryCollectionConvexPropertiesInterface::GetConvexProperties(int TransformG
 		TManagedArray<bool>& Enable = ManagedCollection->GetAttribute<bool>(ConvexEnable, ConvexPropertiesGroup);
 		const TManagedArray<float>& FractionRemove = ManagedCollection->GetAttribute<float>(ConvexFractionRemoveAttribute, ConvexPropertiesGroup);
 		const TManagedArray<float>& SimplificationThreshold = ManagedCollection->GetAttribute<float>(ConvexSimplificationThresholdAttribute, ConvexPropertiesGroup);
+		TManagedArray<float>& CanExceedFraction = ManagedCollection->GetAttribute<float>(ConvexCanExceedFractionAttribute, ConvexPropertiesGroup);
 
 		ConvexProperty.Enable = Enable[PropIndex];
 		ConvexProperty.FractionRemove = FractionRemove[PropIndex];
 		ConvexProperty.SimplificationThreshold = SimplificationThreshold[PropIndex];
+		ConvexProperty.CanExceedFraction = CanExceedFraction[PropIndex];
 	}
 	return ConvexProperty;
 }
@@ -119,6 +129,7 @@ FGeometryCollectionConvexPropertiesInterface::SetConvexProperties(const FConvexC
 	TManagedArray<bool>& Enable = ManagedCollection->GetAttribute<bool>(ConvexEnable, ConvexPropertiesGroup);
 	TManagedArray<float>& FractionRemove = ManagedCollection->GetAttribute<float>(ConvexFractionRemoveAttribute, ConvexPropertiesGroup);
 	TManagedArray<float>& SimplificationThreshold = ManagedCollection->GetAttribute<float>(ConvexSimplificationThresholdAttribute, ConvexPropertiesGroup);
+	TManagedArray<float>& CanExceedFraction = ManagedCollection->GetAttribute<float>(ConvexCanExceedFractionAttribute, ConvexPropertiesGroup);
 
 	int32 AttributeIndex = INDEX_NONE;
 	if (!Index.Contains(TransformGroupIndex))
@@ -136,5 +147,6 @@ FGeometryCollectionConvexPropertiesInterface::SetConvexProperties(const FConvexC
 		Enable[AttributeIndex] = InConvexAttributes.Enable;
 		FractionRemove[AttributeIndex] = InConvexAttributes.FractionRemove;
 		SimplificationThreshold[AttributeIndex] = InConvexAttributes.SimplificationThreshold;
+		CanExceedFraction[AttributeIndex] = InConvexAttributes.CanExceedFraction;
 	}
 }
diff --git a/Engine/Source/Runtime/Experimental/Chaos/Public/GeometryCollection/GeometryCollectionConvexPropertiesInterface.h b/Engine/Source/Runtime/Experimental/Chaos/Public/GeometryCollection/GeometryCollectionConvexPropertiesInterface.h
--- a/Engine/Source/Runtime/Experimental/Chaos/Public/GeometryCollection/GeometryCollectionConvexPropertiesInterface.h
+++ b/Engine/Source/Runtime/Experimental/Chaos/Public/GeometryCollection/GeometryCollectionConvexPropertiesInterface.h
@@ -23,12 +23,15 @@ public :
 	static const FName ConvexFractionRemoveAttribute;
 	// Attribute
 	static const FName ConvexSimplificationThresholdAttribute;
+	// Attribute
+	static const FName ConvexCanExceedFractionAttribute;
 
 
 	struct FConvexCreationProperties {
 		bool Enable = true;
 		float FractionRemove = 0.5f;
 		float SimplificationThreshold = 10.0f;
+		float CanExceedFraction = 0.5f;
 	};
 
 	FGeometryCollectionConvexPropertiesInterface(FGeometryCollection* InGeometryCollection);
diff --git a/Engine/Source/Runtime/Experimental/GeometryCollectionEngine/Private/GeometryCollection/GeometryCollectionObject.cpp b/Engine/Source/Runtime/Experimental/GeometryCollectionEngine/Private/GeometryCollection/GeometryCollectionObject.cpp
--- a/Engine/Source/Runtime/Experimental/GeometryCollectionEngine/Private/GeometryCollection/GeometryCollectionObject.cpp
+++ b/Engine/Source/Runtime/Experimental/GeometryCollectionEngine/Private/GeometryCollection/GeometryCollectionObject.cpp
@@ -272,7 +272,7 @@ void UGeometryCollection::UpdateConvexGeometry()
 	if (GeometryCollection)
 	{
 		FGeometryCollectionConvexPropertiesInterface::FConvexCreationProperties ConvexProperties = GeometryCollection->GetConvexProperties();
-		FGeometryCollectionConvexUtility::CreateNonOverlappingConvexHullData(GeometryCollection.Get(), ConvexProperties.FractionRemove, ConvexProperties.SimplificationThreshold);
+		FGeometryCollectionConvexUtility::CreateNonOverlappingConvexHullData(GeometryCollection.Get(), ConvexProperties.FractionRemove, ConvexProperties.SimplificationThreshold, ConvexProperties.CanExceedFraction);
 		InvalidateCollection();
 	}
 #endif
