Chaos - fix initial shape transform on wake

diff --git a/Engine/Source/Runtime/Experimental/Chaos/Private/Chaos/Collision/ParticlePairMidPhase.cpp b/Engine/Source/Runtime/Experimental/Chaos/Private/Chaos/Collision/ParticlePairMidPhase.cpp
--- a/Engine/Source/Runtime/Experimental/Chaos/Private/Chaos/Collision/ParticlePairMidPhase.cpp
+++ b/Engine/Source/Runtime/Experimental/Chaos/Private/Chaos/Collision/ParticlePairMidPhase.cpp
@@ -338,11 +338,16 @@ namespace Chaos
 	{
 		if ((Constraint != nullptr) && IsUsedSince(SleepEpoch))
 		{
-			// We just need to refresh the epoch so that the constraint state will be used as the previous
+			// We need to refresh the epoch so that the constraint state will be used as the previous
 			// state iof the pair is still colliding in the next tick
 			const int32 CurrentEpoch = MidPhase.GetCollisionAllocator().GetCurrentEpoch();
 			Constraint->GetContainerCookie().LastUsedEpoch = MidPhase.GetCollisionAllocator().GetCurrentEpoch();
 			LastUsedEpoch = CurrentEpoch;
+
+			// We have skipped collision detection for this particle because it was asleep, so we need to update the transforms...
+			// NOTE: this relies on the shape world transforms being up-to-date. They are usually updated in Integarte which
+			// is also skipped for sleeping particles, so they must be updated manually when waking partciles (see IslandManager)
+			Constraint->SetShapeWorldTransforms(Shape0->GetLeafWorldTransform(), Shape1->GetLeafWorldTransform());
 		}
 	}
 
diff --git a/Engine/Source/Runtime/Experimental/Chaos/Private/Chaos/Island/IslandManager.cpp b/Engine/Source/Runtime/Experimental/Chaos/Private/Chaos/Island/IslandManager.cpp
--- a/Engine/Source/Runtime/Experimental/Chaos/Private/Chaos/Island/IslandManager.cpp
+++ b/Engine/Source/Runtime/Experimental/Chaos/Private/Chaos/Island/IslandManager.cpp
@@ -118,6 +118,20 @@ inline void UpdateSleepState(FPBDIslandSolver* IslandSolver, FPBDRigidsSOAs& Par
 			if (!bIsSleeping && IslandParticle->Sleeping())
 			{
 				Particles.ActivateParticle(IslandParticle, true);
+
+				// When we wake particles, we have skipped their integrate step which causes some issues:
+				//	- we have zero velocity (no gravity or external forces applied)
+				//	- the world transforms cached in the ShapesArray will be at the last post-integrate positions
+				//	  which doesn't match what the velocity is telling us
+				// This causes problems for the solver - essentially we have an "initial overlap" situation.
+				// @todo(chaos): We could just run (partial) integrate here for this particle, but we don't know about the Evolution - fix this
+				const FRigidTransform3 ParticleWorldTransform = FParticleUtilities::GetActorWorldTransform(FConstGenericParticleHandle(IslandParticle));
+				for (const TUniquePtr<FPerShapeData>& Shape : IslandParticle->ShapesArray())
+				{
+					const FRigidTransform3 ShapeWorldTransform = Shape->GetLeafRelativeTransform() * ParticleWorldTransform;
+					Shape->SetLeafWorldTransform(ShapeWorldTransform);
+				}
+
 				bNeedRebuild = true;
 			}
 			// If sleeping we deactivate the dynamic particles
diff --git a/Engine/Source/Runtime/Experimental/Chaos/Public/Chaos/GeometryParticles.h b/Engine/Source/Runtime/Experimental/Chaos/Public/Chaos/GeometryParticles.h
--- a/Engine/Source/Runtime/Experimental/Chaos/Public/Chaos/GeometryParticles.h
+++ b/Engine/Source/Runtime/Experimental/Chaos/Public/Chaos/GeometryParticles.h
@@ -93,6 +93,7 @@ namespace Chaos
 
 		// The world-space transform of the leaf geometry (from the last call to UpdateWorldSpaceState)
 		const FRigidTransform3& GetLeafWorldTransform() const { return LeafWorldTransform; }
+		void SetLeafWorldTransform(const FRigidTransform3& InTransform) { LeafWorldTransform = InTransform; }
 
 		// The margin to use with the leaf shape (which could be from its Instanced wrapper if it had one)
 		const FReal GetLeafMargin() const { return LeafMargin; }
