Chaos: Deprecate geometrycollection debuyg draw component

diff --git a/Engine/Plugins/Experimental/GeometryCollectionPlugin/Source/GeometryCollectionEditor/Private/GeometryCollection/DetailCustomizations/SelectedRigidBodyCustomization.cpp b/Engine/Plugins/Experimental/GeometryCollectionPlugin/Source/GeometryCollectionEditor/Private/GeometryCollection/DetailCustomizations/SelectedRigidBodyCustomization.cpp
deleted file mode 100644
diff --git a/Engine/Plugins/Experimental/GeometryCollectionPlugin/Source/GeometryCollectionEditor/Private/GeometryCollection/DetailCustomizations/SelectedRigidBodyCustomization.h b/Engine/Plugins/Experimental/GeometryCollectionPlugin/Source/GeometryCollectionEditor/Private/GeometryCollection/DetailCustomizations/SelectedRigidBodyCustomization.h
deleted file mode 100644
diff --git a/Engine/Plugins/Experimental/GeometryCollectionPlugin/Source/GeometryCollectionEditor/Private/GeometryCollection/DetailCustomizations/WarningMessageCustomization.cpp b/Engine/Plugins/Experimental/GeometryCollectionPlugin/Source/GeometryCollectionEditor/Private/GeometryCollection/DetailCustomizations/WarningMessageCustomization.cpp
--- a/Engine/Plugins/Experimental/GeometryCollectionPlugin/Source/GeometryCollectionEditor/Private/GeometryCollection/DetailCustomizations/WarningMessageCustomization.cpp
+++ b/Engine/Plugins/Experimental/GeometryCollectionPlugin/Source/GeometryCollectionEditor/Private/GeometryCollection/DetailCustomizations/WarningMessageCustomization.cpp
@@ -27,21 +27,11 @@ void FWarningMessageCustomization::CustomizeHeader(TSharedRef<IPropertyHandle> S
 	bool bHasBegunPlay = false;
 	for (const TWeakObjectPtr<UObject>& SelectedObject: StructCustomizationUtils.GetPropertyUtilities()->GetSelectedObjects())
 	{
-		if (AGeometryCollectionDebugDrawActor* const Actor = Cast<AGeometryCollectionDebugDrawActor>(SelectedObject.Get()))
-		{
-			bHasBegunPlay = Actor->HasActorBegunPlay();
-			break;
-		}
 		if (AGeometryCollectionActor* const Actor = Cast<AGeometryCollectionActor>(SelectedObject.Get()))
 		{
 			bHasBegunPlay = Actor->HasActorBegunPlay();
 			break;
 		}
-		if (UGeometryCollectionDebugDrawComponent* const Component = Cast<UGeometryCollectionDebugDrawComponent>(SelectedObject.Get()))
-		{
-			bHasBegunPlay = Component->HasBegunPlay();
-			break;
-		}
 	}
 
 	if (!bHasBegunPlay)
diff --git a/Engine/Plugins/Experimental/GeometryCollectionPlugin/Source/GeometryCollectionEditor/Private/GeometryCollection/GeometryCollectionEditorPlugin.cpp b/Engine/Plugins/Experimental/GeometryCollectionPlugin/Source/GeometryCollectionEditor/Private/GeometryCollection/GeometryCollectionEditorPlugin.cpp
--- a/Engine/Plugins/Experimental/GeometryCollectionPlugin/Source/GeometryCollectionEditor/Private/GeometryCollection/GeometryCollectionEditorPlugin.cpp
+++ b/Engine/Plugins/Experimental/GeometryCollectionPlugin/Source/GeometryCollectionEditor/Private/GeometryCollection/GeometryCollectionEditorPlugin.cpp
@@ -21,7 +21,6 @@
 #include "HAL/ConsoleManager.h"
 #include "Features/IModularFeatures.h"
 #include "GeometryCollection/DetailCustomizations/GeomComponentCacheCustomization.h"
-#include "GeometryCollection/DetailCustomizations/SelectedRigidBodyCustomization.h"
 #include "GeometryCollection/DetailCustomizations/WarningMessageCustomization.h"
 #include "GeometryCollection/DetailCustomizations/GeometryCollectionCustomization.h"
 #include "Styling/SlateStyle.h"
@@ -223,7 +222,6 @@ void IGeometryCollectionEditorPlugin::StartupModule()
 	if(PropertyModule)
 	{
 		PropertyModule->RegisterCustomPropertyTypeLayout("GeomComponentCacheParameters", FOnGetPropertyTypeCustomizationInstance::CreateStatic(&FGeomComponentCacheParametersCustomization::MakeInstance));
-		PropertyModule->RegisterCustomPropertyTypeLayout("GeometryCollectionDebugDrawActorSelectedRigidBody", FOnGetPropertyTypeCustomizationInstance::CreateStatic(&FSelectedRigidBodyCustomization::MakeInstance));
 		PropertyModule->RegisterCustomPropertyTypeLayout("GeometryCollectionDebugDrawWarningMessage", FOnGetPropertyTypeCustomizationInstance::CreateStatic(&FWarningMessageCustomization::MakeInstance));
 
 		const FName GeometryCollectionName = UGeometryCollection::StaticClass()->GetFName();
diff --git a/Engine/Source/Runtime/Experimental/GeometryCollectionEngine/Private/GeometryCollection/GeometryCollectionActor.cpp b/Engine/Source/Runtime/Experimental/GeometryCollectionEngine/Private/GeometryCollection/GeometryCollectionActor.cpp
--- a/Engine/Source/Runtime/Experimental/GeometryCollectionEngine/Private/GeometryCollection/GeometryCollectionActor.cpp
+++ b/Engine/Source/Runtime/Experimental/GeometryCollectionEngine/Private/GeometryCollection/GeometryCollectionActor.cpp
@@ -32,12 +32,7 @@ AGeometryCollectionActor::AGeometryCollectionActor(const FObjectInitializer& Obj
 	GeometryCollectionComponent = CreateDefaultSubobject<UGeometryCollectionComponent>(TEXT("GeometryCollectionComponent0"));
 	RootComponent = GeometryCollectionComponent;
 
-#if GEOMETRYCOLLECTION_DEBUG_DRAW
-	GeometryCollectionDebugDrawComponent = CreateDefaultSubobject<UGeometryCollectionDebugDrawComponent>(TEXT("GeometryCollectionDrawComponent0"));
-	GeometryCollectionDebugDrawComponent->GeometryCollectionComponent = GeometryCollectionComponent;
-#else
-	GeometryCollectionDebugDrawComponent = nullptr;
-#endif
+	GeometryCollectionDebugDrawComponent_DEPRECATED = nullptr;
 
 	PrimaryActorTick.bCanEverTick = true;
 	SetActorTickEnabled(true);
diff --git a/Engine/Source/Runtime/Experimental/GeometryCollectionEngine/Private/GeometryCollection/GeometryCollectionDebugDrawActor.cpp b/Engine/Source/Runtime/Experimental/GeometryCollectionEngine/Private/GeometryCollection/GeometryCollectionDebugDrawActor.cpp
--- a/Engine/Source/Runtime/Experimental/GeometryCollectionEngine/Private/GeometryCollection/GeometryCollectionDebugDrawActor.cpp
+++ b/Engine/Source/Runtime/Experimental/GeometryCollectionEngine/Private/GeometryCollection/GeometryCollectionDebugDrawActor.cpp
@@ -24,2920 +24,16 @@
 
 DEFINE_LOG_CATEGORY_STATIC(LogGeometryCollectionDebugDrawActor, Log, All);
 
-// Constants
-namespace GeometryCollectionDebugDrawActorConstants
-{
-	// Invariables
-	static const bool bPersistent = true;  // Debug draw needs persistency to work well within the editor.
-	static const float LifeTime = -1.0f;  // Lifetime is infinite.
-	static const uint8 DepthPriority = 0;
-	static const uint32 CircleSegments = 32;
-	static const bool bDrawCircleAxis = true;
-	static const TArray<FIntPoint> BoxEdges(
-	{
-		FIntPoint(0, 1), FIntPoint(1, 2), FIntPoint(2, 3), FIntPoint(3, 0),
-		FIntPoint(4, 5), FIntPoint(5, 6), FIntPoint(6, 7), FIntPoint(7, 4),
-		FIntPoint(0, 4), FIntPoint(1, 5), FIntPoint(2, 6), FIntPoint(3, 7)
-	});
-
-	// Base colors
-	static const FLinearColor DarkerTintFactor(1.0f, 1.0f, 0.7f);  // Darker HSV multiplier
-	static const FLinearColor LighterTintFactor(1.0f, 1.0f, 2.0f);  // Lighter HSV multiplier
-	static const FLinearColor RigidBodyTint(0.8f, 0.1f, 0.1f);  // Red
-	static const FLinearColor ClusteringTint(0.6, 0.4f, 0.2f);  // Orange
-	static const FLinearColor GeometryTint(0.4f, 0.2f, 0.6f);  // Purple
-	static const FLinearColor SingleFaceTint(0.6, 0.2f, 0.4f);  // Pink
-	static const FLinearColor VertexTint(0.2f, 0.4f, 0.6f);  // Blue
-
-	// Defaults
-	static const FString SelectedRigidBodySolverDefault = FName(NAME_None).ToString();
-	static const int32 SelectedRigidBodyIdDefault = INDEX_NONE;
-	static const int32 DebugDrawWholeCollectionDefault = 0;
-	static const int32 DebugDrawHierarchyDefault = 0;
-	static const int32 DebugDrawClusteringDefault = 0;
-	static const int32 HideGeometryDefault = int32(EGeometryCollectionDebugDrawActorHideGeometry::HideWithCollision);
-	static const int32 ShowRigidBodyCollisionDefault = 0;
-	static const int32 ShowRigidBodyIdDefault = 0;
-	static const int32 CollisionAtOriginDefault = 0;
-	static const int32 ShowRigidBodyTransformDefault = 0;
-	static const int32 ShowRigidBodyInertiaDefault = 0;
-	static const int32 ShowRigidBodyVelocityDefault = 0;
-	static const int32 ShowRigidBodyForceDefault = 0;
-	static const int32 ShowRigidBodyInfosDefault = 0;
-	static const int32 ShowTransformIndexDefault = 0;
-	static const int32 ShowTransformDefault = 0;
-	static const int32 ShowParentDefault = 0;
-	static const int32 ShowLevelDefault = 0;
-	static const int32 ShowConnectivityEdgesDefault = 0;
-	static const int32 ShowGeometryIndexDefault = 0;
-	static const int32 ShowGeometryTransformDefault = 0;
-	static const int32 ShowBoundingBoxDefault = 0;
-	static const int32 ShowFacesDefault = 0;
-	static const int32 ShowFaceIndicesDefault = 0;
-	static const int32 ShowFaceNormalsDefault = 0;
-	static const int32 ShowSingleFaceDefault = 0;
-	static const int32 SingleFaceIndexDefault = 0;
-	static const int32 ShowVerticesDefault = 0;
-	static const int32 ShowVertexIndicesDefault = 0;
-	static const int32 ShowVertexNormalsDefault = 0;
-	static const int32 UseActiveVisualizationDefault = true;
-	static const float PointThicknessDefault = 6.0f;
-	static const float LineThicknessDefault = 1.0f;
-	static const int32 TextShadowDefault = 1;
-	static const float TextScaleDefault = 1.0f;
-	static const float NormalScaleDefault = 10.0f;
-	static const float AxisScaleDefault = 20.0f;
-	static const float ArrowScaleDefault = 2.5f;
-	static const float TransformScaleDefault = 1.f;
-	static const FColor RigidBodyIdsColorDefault = (RigidBodyTint.LinearRGBToHSV() * LighterTintFactor).HSVToLinearRGB().ToFColor(true);
-	static const FColor RigidBodyCollisionColorDefault = RigidBodyTint.ToFColor(true);
-	static const FColor RigidBodyInertiaColorDefault = (RigidBodyTint.LinearRGBToHSV() * LighterTintFactor).HSVToLinearRGB().ToFColor(true);
-	static const FColor RigidBodyVelocityColorDefault = (RigidBodyTint.LinearRGBToHSV() * DarkerTintFactor).HSVToLinearRGB().ToFColor(true);
-	static const FColor RigidBodyForceColorDefault = (RigidBodyTint.LinearRGBToHSV() * DarkerTintFactor).HSVToLinearRGB().ToFColor(true);
-	static const FColor RigidBodyInfoColorDefault = (RigidBodyTint.LinearRGBToHSV() * LighterTintFactor).HSVToLinearRGB().ToFColor(true);
-	static const FColor TransformIndexColorDefault = (ClusteringTint.LinearRGBToHSV() * LighterTintFactor).HSVToLinearRGB().ToFColor(true);
-	static const FColor LevelColorDefault = (ClusteringTint.LinearRGBToHSV() * LighterTintFactor).HSVToLinearRGB().ToFColor(true);
-	static const FColor ParentColorDefault = ClusteringTint.ToFColor(true);
-	static const FColor GeometryIndexColorDefault = (GeometryTint.LinearRGBToHSV() * LighterTintFactor).HSVToLinearRGB().ToFColor(true);
-	static const FColor BoundingBoxColorDefault = (GeometryTint.LinearRGBToHSV() * DarkerTintFactor).HSVToLinearRGB().ToFColor(true);
-	static const FColor FaceColorDefault = GeometryTint.ToFColor(true);
-	static const FColor FaceIndexColorDefault = (GeometryTint.LinearRGBToHSV() * LighterTintFactor).HSVToLinearRGB().ToFColor(true);
-	static const FColor FaceNormalColorDefault = (GeometryTint.LinearRGBToHSV() * DarkerTintFactor).HSVToLinearRGB().ToFColor(true);
-	static const FColor SingleFaceColorDefault = (SingleFaceTint.LinearRGBToHSV() * LighterTintFactor).HSVToLinearRGB().ToFColor(true);
-	static const FColor VertexColorDefault = VertexTint.ToFColor(true);
-	static const FColor VertexIndexColorDefault = (VertexTint.LinearRGBToHSV() * LighterTintFactor).HSVToLinearRGB().ToFColor(true);
-	static const FColor VertexNormalColorDefault = (VertexTint.LinearRGBToHSV() * DarkerTintFactor).HSVToLinearRGB().ToFColor(true);
-}
-
-namespace GeometryCollectionDebugDrawActorCVars
-{
-	// Console variables, also exposed as settings in this actor
-	static TAutoConsoleVariable<FString> SelectedRigidBodySolver (TEXT("p.gc.SelectedRigidBodySolver" ), GeometryCollectionDebugDrawActorConstants::SelectedRigidBodySolverDefault , TEXT("Geometry Collection debug draw, visualize debug informations for the selected rigid body solver.\nDefault = None"), ECVF_Cheat);
-	static TAutoConsoleVariable<int32  > SelectedRigidBodyId     (TEXT("p.gc.SelectedRigidBodyId"     ), GeometryCollectionDebugDrawActorConstants::SelectedRigidBodyIdDefault     , TEXT("Geometry Collection debug draw, visualize debug informations for the selected rigid body ids.\nDefault = -1"), ECVF_Cheat);
-	//static TAutoConsoleVariable<FGuid  > SelectedRigidBodyId     (TEXT("p.gc.SelectedRigidBodyId"     ), GeometryCollectionDebugDrawActorConstants::SelectedRigidBodyIdDefault     , TEXT("Geometry Collection debug draw, visualize debug informations for the selected rigid body ids.\nDefault = 0:0:0:0"), ECVF_Cheat);
-	static TAutoConsoleVariable<int32  > DebugDrawWholeCollection(TEXT("p.gc.DebugDrawWholeCollection"), GeometryCollectionDebugDrawActorConstants::DebugDrawWholeCollectionDefault, TEXT("Geometry Collection debug draw, show debug visualization for the rest of the geometry collection related to the current rigid body id selection.\nDefault = 0"), ECVF_Cheat);
-	static TAutoConsoleVariable<int32  > DebugDrawHierarchy      (TEXT("p.gc.DebugDrawHierarchy"      ), GeometryCollectionDebugDrawActorConstants::DebugDrawHierarchyDefault      , TEXT("Geometry Collection debug draw, show debug visualization for the top level node rather than the bottom leaf nodes of a cluster's hierarchy..\nDefault = 0"), ECVF_Cheat);
-	static TAutoConsoleVariable<int32  > DebugDrawClustering     (TEXT("p.gc.DebugDrawClustering"     ), GeometryCollectionDebugDrawActorConstants::DebugDrawClusteringDefault     , TEXT("Geometry Collection debug draw, show debug visualization for all clustered children associated to the current rigid body id selection.\nDefault = 0"), ECVF_Cheat);
-	static TAutoConsoleVariable<int32  > HideGeometry            (TEXT("p.gc.HideGeometry"            ), GeometryCollectionDebugDrawActorConstants::HideGeometryDefault            , TEXT("Geometry Collection debug draw, geometry visibility setting, select the part of the geometry to hide in order to better visualize the debug information.\n0: Do not hide any geometries.\n1: Hide the geometry associated to the rigid bodies selected for collision display.\n2: Hide the geometry associated to the selected rigid bodies.\n3: Hide the entire geometry collection associated to the selected rigid bodies.\n4: Hide all geometry collections.\nDefault = 1"), ECVF_Cheat);
-	static TAutoConsoleVariable<int32  > ShowRigidBodyId         (TEXT("p.gc.ShowRigidBodyId"         ), GeometryCollectionDebugDrawActorConstants::ShowRigidBodyIdDefault         , TEXT("Geometry Collection debug draw, show the rigid body id(s).\nDefault = 0"), ECVF_Cheat);
-	static TAutoConsoleVariable<int32  > ShowRigidBodyCollision  (TEXT("p.gc.ShowRigidBodyCollision"  ), GeometryCollectionDebugDrawActorConstants::ShowRigidBodyCollisionDefault  , TEXT("Geometry Collection debug draw, show the selected's rigid body's collision volume.\nDefault = 0"), ECVF_Cheat);
-	static TAutoConsoleVariable<int32  > CollisionAtOrigin       (TEXT("p.gc.CollisionAtOrigin"       ), GeometryCollectionDebugDrawActorConstants::CollisionAtOriginDefault       , TEXT("Geometry Collection debug draw, show any collision volume at the origin, in local space.\nDefault = 0"), ECVF_Cheat);
-	static TAutoConsoleVariable<int32  > ShowRigidBodyTransform  (TEXT("p.gc.ShowRigidBodyTransform"  ), GeometryCollectionDebugDrawActorConstants::ShowRigidBodyTransformDefault  , TEXT("Geometry Collection debug draw, show the selected's rigid body's transform.\nDefault = 0"), ECVF_Cheat);
-	static TAutoConsoleVariable<int32  > ShowRigidBodyInertia    (TEXT("p.gc.ShowRigidBodyInertia"    ), GeometryCollectionDebugDrawActorConstants::ShowRigidBodyInertiaDefault    , TEXT("Geometry Collection debug draw, show the selected's rigid body's inertia tensor box.\nDefault = 0"), ECVF_Cheat);
-	static TAutoConsoleVariable<int32  > ShowRigidBodyVelocity   (TEXT("p.gc.ShowRigidBodyVelocity"   ), GeometryCollectionDebugDrawActorConstants::ShowRigidBodyVelocityDefault   , TEXT("Geometry Collection debug draw, show the selected's rigid body's linear and angular velocities.\nDefault = 0"), ECVF_Cheat);
-	static TAutoConsoleVariable<int32  > ShowRigidBodyForce      (TEXT("p.gc.ShowRigidBodyForce"      ), GeometryCollectionDebugDrawActorConstants::ShowRigidBodyForceDefault      , TEXT("Geometry Collection debug draw, show the selected's rigid body's applied force and torque.\nDefault = 0"), ECVF_Cheat);
-	static TAutoConsoleVariable<int32  > ShowRigidBodyInfos      (TEXT("p.gc.ShowRigidBodyInfos"      ), GeometryCollectionDebugDrawActorConstants::ShowRigidBodyInfosDefault      , TEXT("Geometry Collection debug draw, show the selected's rigid body's information.\nDefault = 0"), ECVF_Cheat);
-	static TAutoConsoleVariable<int32  > ShowTransformIndex      (TEXT("p.gc.ShowTransformIndex"      ), GeometryCollectionDebugDrawActorConstants::ShowTransformIndexDefault      , TEXT("Geometry Collection debug draw, show the transform index for the selected rigid body's associated cluster nodes.\nDefault = 0"), ECVF_Cheat);
-	static TAutoConsoleVariable<int32  > ShowTransform           (TEXT("p.gc.ShowTransform"           ), GeometryCollectionDebugDrawActorConstants::ShowTransformDefault           , TEXT("Geometry Collection debug draw, show the transform for the selected rigid body's associated cluster nodes.\nDefault = 0"), ECVF_Cheat);
-	static TAutoConsoleVariable<int32  > ShowParent              (TEXT("p.gc.ShowParent"              ), GeometryCollectionDebugDrawActorConstants::ShowParentDefault              , TEXT("Geometry Collection debug draw, show a link from the selected rigid body's associated cluster nodes to their parent's nodes.\nDefault = 0"), ECVF_Cheat);
-	static TAutoConsoleVariable<int32  > ShowLevel               (TEXT("p.gc.ShowLevel"               ), GeometryCollectionDebugDrawActorConstants::ShowLevelDefault               , TEXT("Geometry Collection debug draw, show the hierarchical level for the selected rigid body's associated cluster nodes.\nDefault = 0"), ECVF_Cheat);
-	static TAutoConsoleVariable<int32  > ShowConnectivityEdges   (TEXT("p.gc.ShowConnectivityEdges"   ), GeometryCollectionDebugDrawActorConstants::ShowConnectivityEdgesDefault   , TEXT("Geometry Collection debug draw, show the connectivity edges for the rigid body's associated cluster nodes.\nDefault = 0"), ECVF_Cheat);
-	static TAutoConsoleVariable<int32  > ShowGeometryIndex       (TEXT("p.gc.ShowGeometryIndex"       ), GeometryCollectionDebugDrawActorConstants::ShowGeometryIndexDefault       , TEXT("Geometry Collection debug draw, show the geometry index for the selected rigid body's associated geometries.\nDefault = 0"), ECVF_Cheat);
-	static TAutoConsoleVariable<int32  > ShowGeometryTransform   (TEXT("p.gc.ShowGeometryTransform"   ), GeometryCollectionDebugDrawActorConstants::ShowGeometryTransformDefault   , TEXT("Geometry Collection debug draw, show the geometry transform for the selected rigid body's associated geometries.\nDefault = 0"), ECVF_Cheat);
-	static TAutoConsoleVariable<int32  > ShowBoundingBox         (TEXT("p.gc.ShowBoundingBox"         ), GeometryCollectionDebugDrawActorConstants::ShowBoundingBoxDefault         , TEXT("Geometry Collection debug draw, show the bounding box for the selected rigid body's associated geometries.\nDefault = 0"), ECVF_Cheat);
-	static TAutoConsoleVariable<int32  > ShowFaces               (TEXT("p.gc.ShowFaces"               ), GeometryCollectionDebugDrawActorConstants::ShowFacesDefault               , TEXT("Geometry Collection debug draw, show the faces for the selected rigid body's associated geometries.\nDefault = 0"), ECVF_Cheat);
-	static TAutoConsoleVariable<int32  > ShowFaceIndices         (TEXT("p.gc.ShowFaceIndices"         ), GeometryCollectionDebugDrawActorConstants::ShowFaceIndicesDefault         , TEXT("Geometry Collection debug draw, show the face indices for the selected rigid body's associated geometries.\nDefault = 0"), ECVF_Cheat);
-	static TAutoConsoleVariable<int32  > ShowFaceNormals         (TEXT("p.gc.ShowFaceNormals"         ), GeometryCollectionDebugDrawActorConstants::ShowFaceNormalsDefault         , TEXT("Geometry Collection debug draw, show the face normals for the selected rigid body's associated geometries.\nDefault = 0"), ECVF_Cheat);
-	static TAutoConsoleVariable<int32  > ShowSingleFace          (TEXT("p.gc.ShowSingleFace"          ), GeometryCollectionDebugDrawActorConstants::ShowSingleFaceDefault          , TEXT("Geometry Collection debug draw, enable single face visualization for the selected rigid body's associated geometries.\nDefault = 0"), ECVF_Cheat);
-	static TAutoConsoleVariable<int32  > SingleFaceIndex         (TEXT("p.gc.SingleFaceIndex"         ), GeometryCollectionDebugDrawActorConstants::SingleFaceIndexDefault         , TEXT("Geometry Collection debug draw, the index of the single face to visualize.\nDefault = 0"), ECVF_Cheat);
-	static TAutoConsoleVariable<int32  > ShowVertices            (TEXT("p.gc.ShowVertices"            ), GeometryCollectionDebugDrawActorConstants::ShowVerticesDefault            , TEXT("Geometry Collection debug draw, show the vertices for the selected rigid body's associated geometries.\nDefault = 0"), ECVF_Cheat);
-	static TAutoConsoleVariable<int32  > ShowVertexIndices       (TEXT("p.gc.ShowVertexIndices"       ), GeometryCollectionDebugDrawActorConstants::ShowVertexIndicesDefault       , TEXT("Geometry Collection debug draw, show the vertex index for the selected rigid body's associated geometries.\nDefault = 0"), ECVF_Cheat);
-	static TAutoConsoleVariable<int32  > ShowVertexNormals       (TEXT("p.gc.ShowVertexNormals"       ), GeometryCollectionDebugDrawActorConstants::ShowVertexNormalsDefault       , TEXT("Geometry Collection debug draw, show the vertex normals for the selected rigid body's associated geometries.\nDefault = 0"), ECVF_Cheat);
-	static TAutoConsoleVariable<int32  > UseActiveVisualization  (TEXT("p.gc.UseActiveVisualization"  ), GeometryCollectionDebugDrawActorConstants::UseActiveVisualizationDefault  , TEXT("Geometry Collection debug draw, adapt visualization depending of the cluster nodes' hierarchical level..\nDefault = 1."), ECVF_Cheat);
-	static TAutoConsoleVariable<float  > PointThickness          (TEXT("p.gc.PointThickness"          ), GeometryCollectionDebugDrawActorConstants::PointThicknessDefault          , TEXT("Geometry Collection debug draw, point thickness.\nDefault = 6."), ECVF_Cheat);
-	static TAutoConsoleVariable<float  > LineThickness           (TEXT("p.gc.LineThickness"           ), GeometryCollectionDebugDrawActorConstants::LineThicknessDefault           , TEXT("Geometry Collection debug draw, line thickness.\nDefault = 1."), ECVF_Cheat);
-	static TAutoConsoleVariable<int32  > TextShadow              (TEXT("p.gc.TextShadow"              ), GeometryCollectionDebugDrawActorConstants::TextShadowDefault              , TEXT("Geometry Collection debug draw, text shadow under indices for better readability.\nDefault = 1."), ECVF_Cheat);
-	static TAutoConsoleVariable<float  > TextScale               (TEXT("p.gc.TextScale"               ), GeometryCollectionDebugDrawActorConstants::TextScaleDefault               , TEXT("Geometry Collection debug draw, text scale.\nDefault = 1."), ECVF_Cheat);
-	static TAutoConsoleVariable<float  > NormalScale             (TEXT("p.gc.NormalScale"             ), GeometryCollectionDebugDrawActorConstants::NormalScaleDefault             , TEXT("Geometry Collection debug draw, normal size.\nDefault = 10."), ECVF_Cheat);
-	static TAutoConsoleVariable<float  > AxisScale               (TEXT("p.gc.AxisScale"               ), GeometryCollectionDebugDrawActorConstants::AxisScaleDefault               , TEXT("Geometry Collection debug draw, size of the axis used for visualizing all transforms.\nDefault = 20."), ECVF_Cheat);
-	static TAutoConsoleVariable<float  > ArrowScale              (TEXT("p.gc.ArrowScale"              ), GeometryCollectionDebugDrawActorConstants::ArrowScaleDefault              , TEXT("Geometry Collection debug draw, arrow size for normals.\nDefault = 2.5."), ECVF_Cheat);
-}
-
+PRAGMA_DISABLE_DEPRECATION_WARNINGS
 FString FGeometryCollectionDebugDrawActorSelectedRigidBody::GetSolverName() const
 {
 	return !Solver ? FName(NAME_None).ToString() : Solver->GetName();
 }
 
-AGeometryCollectionDebugDrawActor* AGeometryCollectionDebugDrawActor::FindOrCreate(UWorld* World)
-{
-	AGeometryCollectionDebugDrawActor* Actor = nullptr;
-	if (!World)
-	{
-		UE_LOG(LogGeometryCollectionDebugDrawActor, Warning, TEXT("No valid World for where to search for an existing GeometryCollectionDebugDrawActor singleton actor."));
-	}
-	else
-	{
-		const TActorIterator<AGeometryCollectionDebugDrawActor> ActorIterator(World);
-		if (ActorIterator)
-		{
-			Actor = *ActorIterator;
-		}
-		else
-		{
-			FActorSpawnParameters SpawnInfo;
-			SpawnInfo.SpawnCollisionHandlingOverride = ESpawnActorCollisionHandlingMethod::AlwaysSpawn;
-			Actor = World->SpawnActor<AGeometryCollectionDebugDrawActor>(SpawnInfo);
-		}
-		if (!Actor)
-		{
-			UE_LOG(LogGeometryCollectionDebugDrawActor, Warning, TEXT("No GeometryCollectionDebugDrawActor singleton actor could be found or created."));
-		}
-	}
-	return Actor;
-}
-
 AGeometryCollectionDebugDrawActor::AGeometryCollectionDebugDrawActor(const FObjectInitializer& ObjectInitializer)
 	: Super(ObjectInitializer)
-	, SelectedRigidBody        (  GeometryCollectionDebugDrawActorConstants::SelectedRigidBodyIdDefault     )
-	, bDebugDrawWholeCollection(!!GeometryCollectionDebugDrawActorConstants::DebugDrawWholeCollectionDefault)
-	, bDebugDrawHierarchy      (!!GeometryCollectionDebugDrawActorConstants::DebugDrawHierarchyDefault      )
-	, bDebugDrawClustering     (!!GeometryCollectionDebugDrawActorConstants::DebugDrawClusteringDefault     )
-	, HideGeometry             ( EGeometryCollectionDebugDrawActorHideGeometry(
-		                          GeometryCollectionDebugDrawActorConstants::HideGeometryDefault            ))
-	, bShowRigidBodyId         (!!GeometryCollectionDebugDrawActorConstants::ShowRigidBodyIdDefault         )
-	, bShowRigidBodyCollision  (!!GeometryCollectionDebugDrawActorConstants::ShowRigidBodyCollisionDefault  )
-	, bCollisionAtOrigin       (!!GeometryCollectionDebugDrawActorConstants::CollisionAtOriginDefault       )
-	, bShowRigidBodyTransform  (!!GeometryCollectionDebugDrawActorConstants::ShowRigidBodyTransformDefault  )
-	, bShowRigidBodyInertia    (!!GeometryCollectionDebugDrawActorConstants::ShowRigidBodyInertiaDefault    )
-	, bShowRigidBodyVelocity   (!!GeometryCollectionDebugDrawActorConstants::ShowRigidBodyVelocityDefault   )
-	, bShowRigidBodyForce      (!!GeometryCollectionDebugDrawActorConstants::ShowRigidBodyForceDefault      )
-	, bShowRigidBodyInfos      (!!GeometryCollectionDebugDrawActorConstants::ShowRigidBodyInfosDefault      )
-	, bShowTransformIndex      (!!GeometryCollectionDebugDrawActorConstants::ShowTransformIndexDefault      )
-	, bShowTransform           (!!GeometryCollectionDebugDrawActorConstants::ShowTransformDefault           )
-	, bShowParent              (!!GeometryCollectionDebugDrawActorConstants::ShowParentDefault              )
-	, bShowLevel               (!!GeometryCollectionDebugDrawActorConstants::ShowLevelDefault               )
-	, bShowConnectivityEdges   (!!GeometryCollectionDebugDrawActorConstants::ShowConnectivityEdgesDefault   )
-	, bShowGeometryIndex       (!!GeometryCollectionDebugDrawActorConstants::ShowGeometryIndexDefault       )
-	, bShowGeometryTransform   (!!GeometryCollectionDebugDrawActorConstants::ShowGeometryTransformDefault   )
-	, bShowBoundingBox         (!!GeometryCollectionDebugDrawActorConstants::ShowBoundingBoxDefault         )
-	, bShowFaces               (!!GeometryCollectionDebugDrawActorConstants::ShowFacesDefault               )
-	, bShowFaceIndices         (!!GeometryCollectionDebugDrawActorConstants::ShowFaceIndicesDefault         )
-	, bShowFaceNormals         (!!GeometryCollectionDebugDrawActorConstants::ShowFaceNormalsDefault         )
-	, bShowVertices            (!!GeometryCollectionDebugDrawActorConstants::ShowVerticesDefault            )
-	, bShowVertexIndices       (!!GeometryCollectionDebugDrawActorConstants::ShowVertexIndicesDefault       )
-	, bShowVertexNormals       (!!GeometryCollectionDebugDrawActorConstants::ShowVertexNormalsDefault       )
-	, bUseActiveVisualization  (!!GeometryCollectionDebugDrawActorConstants::UseActiveVisualizationDefault  )
-	, PointThickness           (  GeometryCollectionDebugDrawActorConstants::PointThicknessDefault          )
-	, LineThickness            (  GeometryCollectionDebugDrawActorConstants::LineThicknessDefault           )
-	, bTextShadow              (!!GeometryCollectionDebugDrawActorConstants::TextShadowDefault              )
-	, TextScale                (  GeometryCollectionDebugDrawActorConstants::TextScaleDefault               )
-	, NormalScale              (  GeometryCollectionDebugDrawActorConstants::NormalScaleDefault             )
-	, AxisScale                (  GeometryCollectionDebugDrawActorConstants::AxisScaleDefault               )
-	, ArrowScale               (  GeometryCollectionDebugDrawActorConstants::ArrowScaleDefault              )
-	, RigidBodyIdColor         (  GeometryCollectionDebugDrawActorConstants::RigidBodyIdsColorDefault       )
-	, RigidBodyTransformScale  (  GeometryCollectionDebugDrawActorConstants::TransformScaleDefault          )
-	, RigidBodyCollisionColor  (  GeometryCollectionDebugDrawActorConstants::RigidBodyCollisionColorDefault )
-	, RigidBodyInertiaColor    (  GeometryCollectionDebugDrawActorConstants::RigidBodyInertiaColorDefault   )
-	, RigidBodyVelocityColor   (  GeometryCollectionDebugDrawActorConstants::RigidBodyVelocityColorDefault  )
-	, RigidBodyForceColor      (  GeometryCollectionDebugDrawActorConstants::RigidBodyForceColorDefault     )
-	, RigidBodyInfoColor       (  GeometryCollectionDebugDrawActorConstants::RigidBodyInfoColorDefault      )
-	, TransformIndexColor      (  GeometryCollectionDebugDrawActorConstants::TransformIndexColorDefault     )
-	, TransformScale           (  GeometryCollectionDebugDrawActorConstants::TransformScaleDefault          )
-	, LevelColor               (  GeometryCollectionDebugDrawActorConstants::LevelColorDefault              )
-	, ParentColor              (  GeometryCollectionDebugDrawActorConstants::ParentColorDefault             )
-	, ConnectivityEdgeThickness(  GeometryCollectionDebugDrawActorConstants::LineThicknessDefault           )
-	, GeometryIndexColor       (  GeometryCollectionDebugDrawActorConstants::GeometryIndexColorDefault      )
-	, GeometryTransformScale   (  GeometryCollectionDebugDrawActorConstants::TransformScaleDefault          )
-	, BoundingBoxColor         (  GeometryCollectionDebugDrawActorConstants::BoundingBoxColorDefault        )
-	, FaceColor                (  GeometryCollectionDebugDrawActorConstants::FaceColorDefault               )
-	, FaceIndexColor           (  GeometryCollectionDebugDrawActorConstants::FaceIndexColorDefault          )
-	, FaceNormalColor          (  GeometryCollectionDebugDrawActorConstants::FaceNormalColorDefault         )
-	, SingleFaceColor          (  GeometryCollectionDebugDrawActorConstants::SingleFaceColorDefault         )
-	, VertexColor              (  GeometryCollectionDebugDrawActorConstants::VertexColorDefault             )
-	, VertexIndexColor         (  GeometryCollectionDebugDrawActorConstants::VertexIndexColorDefault        )
-	, VertexNormalColor        (  GeometryCollectionDebugDrawActorConstants::VertexNormalColorDefault       )
-	, SpriteComponent()
-	, ConsoleVariableSinkHandle()
-	, DebugDrawTextDelegateHandle()
-	, DebugDrawTexts()
-	, bNeedsDebugLinesFlush(false)
-#if WITH_EDITOR
-	, bWasEditorPaused(false)
-#endif
-{
-	// Enable game tick calls
-	PrimaryActorTick.bCanEverTick = true;
-	PrimaryActorTick.bTickEvenWhenPaused = true;  // Debug draw must still runs while paused
-	SetActorTickEnabled(true);
-
-	// Register console variable sink
-	ConsoleVariableSinkHandle = IConsoleManager::Get().RegisterConsoleVariableSink_Handle(FConsoleCommandDelegate::CreateUObject(this, &AGeometryCollectionDebugDrawActor::OnCVarsChanged));
-
-	// Create scene component
-	RootComponent = ObjectInitializer.CreateDefaultSubobject<USceneComponent>(this, TEXT("SceneComponent"));
-	RootComponent->Mobility = EComponentMobility::Static;
-
-#if WITH_EDITORONLY_DATA
-	// Create and attach sprite
-	SpriteComponent = ObjectInitializer.CreateEditorOnlyDefaultSubobject<UBillboardComponent>(this, TEXT("Sprite"));
-	if (SpriteComponent)
-	{
-		// Structure to hold one-time sprite initialization
-		struct FConstructorStatics
-		{
-			ConstructorHelpers::FObjectFinderOptional<UTexture2D> NoteTextureObject;  // A helper class object used to find target UTexture2D object in resource package
-			FName ID_Notes;  // Icon sprite category name
-			FText NAME_Notes;  // Icon sprite display name
-
-			FConstructorStatics()
-				: NoteTextureObject(TEXT("/Engine/EditorResources/S_Note"))
-				, ID_Notes(TEXT("Notes"))
-				, NAME_Notes(NSLOCTEXT("SpriteCategory", "Notes", "Notes"))
-			{}
-		};
-		static FConstructorStatics ConstructorStatics;
-		SpriteComponent->Sprite = ConstructorStatics.NoteTextureObject.Get();
-		SpriteComponent->SpriteInfo.Category = ConstructorStatics.ID_Notes;
-		SpriteComponent->SpriteInfo.DisplayName = ConstructorStatics.NAME_Notes;
-		SpriteComponent->AttachToComponent(RootComponent, FAttachmentTransformRules::KeepRelativeTransform);
-		SpriteComponent->Mobility = EComponentMobility::Static;
-	}
-#endif // WITH_EDITORONLY_DATA
-}
-
-void AGeometryCollectionDebugDrawActor::BeginDestroy()
-{
-	// Unregister console variable sink
-	IConsoleManager::Get().UnregisterConsoleVariableSink_Handle(ConsoleVariableSinkHandle);
-
-	Super::BeginDestroy();
-}
-
-void AGeometryCollectionDebugDrawActor::Tick(float DeltaSeconds)
-{
-	Super::Tick(DeltaSeconds);
-
-	// Clear all persistent strings and debug lines.
-	Flush();
-
-	UWorld* const World = GetWorld();
-#if WITH_EDITOR && ENABLE_DRAW_DEBUG
-	// Check editor pause status and force a dynamic update on all components to catchup with the physics thread
-	// This can't be done in the GeometryCollectionDebugDrawComponent since it doesn't tick at every frame,
-	// and can't be done in GeometryCollectionComponent either since it doesn't usually tick while paused.
-	const bool bIsEditorPaused = World && World->IsPlayInEditor() && World->bDebugPauseExecution;
-	if (bIsEditorPaused && !bWasEditorPaused)
-	{
-		// For dynamic update of transforms
-		for (TActorIterator<AGeometryCollectionActor> ActorIterator(World); ActorIterator; ++ActorIterator)
-		{
-			UGeometryCollectionDebugDrawComponent* const GeometryCollectionDebugDrawComponent = ActorIterator->GetGeometryCollectionDebugDrawComponent();
-			if (GeometryCollectionDebugDrawComponent &&
-				GeometryCollectionDebugDrawComponent->GeometryCollectionDebugDrawActor == this &&
-				ensure(GeometryCollectionDebugDrawComponent->GeometryCollectionComponent))
-			{
-				GeometryCollectionDebugDrawComponent->GeometryCollectionComponent->ForceRenderUpdateDynamicData();
-			}
-		}
-	}
-	bWasEditorPaused = bIsEditorPaused;
-#endif  // #if WITH_EDITOR && ENABLE_DRAW_DEBUG
-
-#if GEOMETRYCOLLECTION_DEBUG_DRAW
-	// Check badly synced collections in case it is still looking for an id match
-	if (World && SelectedRigidBody.Id != INDEX_NONE && !SelectedRigidBody.GeometryCollection)
-	//if (World && SelectedRigidBody.Id.IsValid() && !SelectedRigidBody.GeometryCollection)
-	{
-#if TODO_REIMPLEMENT_GET_RIGID_PARTICLES
-		// Check the id is within the selected solver range
-		const Chaos::FPBDRigidsSolver* const Solver = 
-			SelectedRigidBody.Solver ? SelectedRigidBody.Solver->GetSolver() :  // Selected solver
-			World->PhysicsScene_Chaos ? World->PhysicsScene_Chaos->GetSolver() :  // Default world solver
-			nullptr;  // No solver
-
-		const bool IsWithinRange = Solver ? (uint32(SelectedRigidBody.Id) < Solver->GetRigidParticles().Size()): false;
-		if (!IsWithinRange)
-		{
-			UE_LOG(LogGeometryCollectionDebugDrawActor, VeryVerbose, TEXT("The selection id is out of range."));
-		}
-		else  // Statement continues below...
-#endif  // #if TODO_REIMPLEMENT_GET_RIGID_PARTICLES
-		{
-			UE_LOG(LogGeometryCollectionDebugDrawActor, VeryVerbose, TEXT("The selection couldn't be found. The property update will run on all components still containing any invalid rigid body ids."));
-
-			// Check for delayed Rigid Body Id array initializations
-			for (TActorIterator<AGeometryCollectionActor> ActorIterator(World); ActorIterator; ++ActorIterator)
-			{
-				if (UGeometryCollectionDebugDrawComponent* const GeometryCollectionDebugDrawComponent = ActorIterator->GetGeometryCollectionDebugDrawComponent())
-				{
-					if (GeometryCollectionDebugDrawComponent->GeometryCollectionDebugDrawActor == this &&
-						GeometryCollectionDebugDrawComponent->HasIncompleteRigidBodyIdSync())
-					{
-						const bool bIsSelected = GeometryCollectionDebugDrawComponent->OnDebugDrawPropertiesChanged(false);
-						if (bIsSelected)
-						{
-							SelectedRigidBody.GeometryCollection = *ActorIterator;
-							UE_LOG(LogGeometryCollectionDebugDrawActor, Verbose, TEXT("Selection found. Stopping continuous property update."));
-							break;
-						}
-					}
-				}
-			}
-		}
-	}
-#endif  // #if GEOMETRYCOLLECTION_DEBUG_DRAW
-}
-
-void AGeometryCollectionDebugDrawActor::BeginPlay()
-{
-	Super::BeginPlay();
-#if GEOMETRYCOLLECTION_DEBUG_DRAW && ENABLE_DRAW_DEBUG
-	// Initialize text renderer
-	const FDebugDrawDelegate DebugDrawTextDelegate = FDebugDrawDelegate::CreateUObject(this, &AGeometryCollectionDebugDrawActor::DebugDrawText);
-	DebugDrawTextDelegateHandle = UDebugDrawService::Register(TEXT("TextRender"), DebugDrawTextDelegate);  // TextRender is an engine show flag that works in both editor and game modes
-#endif  // #if GEOMETRYCOLLECTION_DEBUG_DRAW && ENABLE_DRAW_DEBUG
-}
-
-void AGeometryCollectionDebugDrawActor::EndPlay(EEndPlayReason::Type ReasonEnd)
-{
-#if GEOMETRYCOLLECTION_DEBUG_DRAW && ENABLE_DRAW_DEBUG
-	UDebugDrawService::Unregister(DebugDrawTextDelegateHandle);
-#endif  // #if GEOMETRYCOLLECTION_DEBUG_DRAW && ENABLE_DRAW_DEBUG
-	Super::EndPlay(ReasonEnd);
-}
-
-void AGeometryCollectionDebugDrawActor::PostLoad()
-{
-	static const EConsoleVariableFlags SetBy = ECVF_SetByConsole;  // Can't use the default ECVF_SetByCode as otherwise it won't update the global console variable.
-
-	GeometryCollectionDebugDrawActorCVars::SelectedRigidBodySolver ->Set(*SelectedRigidBody.GetSolverName(), SetBy);
-	GeometryCollectionDebugDrawActorCVars::SelectedRigidBodyId     ->Set( SelectedRigidBody.Id             , SetBy);
-	//GeometryCollectionDebugDrawActorCVars::SelectedRigidBodyId     ->Set(*SelectedRigidBody.Id.ToString()  , SetBy);
-	GeometryCollectionDebugDrawActorCVars::DebugDrawWholeCollection->Set(int32(bDebugDrawWholeCollection  ), SetBy);
-	GeometryCollectionDebugDrawActorCVars::DebugDrawHierarchy      ->Set(int32(bDebugDrawHierarchy        ), SetBy);
-	GeometryCollectionDebugDrawActorCVars::DebugDrawClustering     ->Set(int32(bDebugDrawClustering       ), SetBy);
-	GeometryCollectionDebugDrawActorCVars::HideGeometry            ->Set(int32(HideGeometry               ), SetBy);
-	GeometryCollectionDebugDrawActorCVars::ShowRigidBodyId         ->Set(int32(bShowRigidBodyId           ), SetBy);
-	GeometryCollectionDebugDrawActorCVars::ShowRigidBodyCollision  ->Set(int32(bShowRigidBodyCollision    ), SetBy);
-	GeometryCollectionDebugDrawActorCVars::CollisionAtOrigin       ->Set(int32(bCollisionAtOrigin         ), SetBy);
-	GeometryCollectionDebugDrawActorCVars::ShowRigidBodyTransform  ->Set(int32(bShowRigidBodyTransform    ), SetBy);
-	GeometryCollectionDebugDrawActorCVars::ShowRigidBodyInertia    ->Set(int32(bShowRigidBodyInertia      ), SetBy);
-	GeometryCollectionDebugDrawActorCVars::ShowRigidBodyVelocity   ->Set(int32(bShowRigidBodyVelocity     ), SetBy);
-	GeometryCollectionDebugDrawActorCVars::ShowRigidBodyForce      ->Set(int32(bShowRigidBodyForce        ), SetBy);
-	GeometryCollectionDebugDrawActorCVars::ShowRigidBodyInfos      ->Set(int32(bShowRigidBodyInfos        ), SetBy);
-	GeometryCollectionDebugDrawActorCVars::ShowTransformIndex      ->Set(int32(bShowTransformIndex        ), SetBy);
-	GeometryCollectionDebugDrawActorCVars::ShowTransform           ->Set(int32(bShowTransform             ), SetBy);
-	GeometryCollectionDebugDrawActorCVars::ShowParent              ->Set(int32(bShowParent                ), SetBy);
-	GeometryCollectionDebugDrawActorCVars::ShowLevel               ->Set(int32(bShowLevel                 ), SetBy);
-	GeometryCollectionDebugDrawActorCVars::ShowConnectivityEdges   ->Set(int32(bShowConnectivityEdges     ), SetBy);
-	GeometryCollectionDebugDrawActorCVars::ShowGeometryIndex       ->Set(int32(bShowGeometryIndex         ), SetBy);
-	GeometryCollectionDebugDrawActorCVars::ShowGeometryTransform   ->Set(int32(bShowGeometryTransform     ), SetBy);
-	GeometryCollectionDebugDrawActorCVars::ShowBoundingBox         ->Set(int32(bShowBoundingBox           ), SetBy);
-	GeometryCollectionDebugDrawActorCVars::ShowFaces               ->Set(int32(bShowFaces                 ), SetBy);
-	GeometryCollectionDebugDrawActorCVars::ShowFaceIndices         ->Set(int32(bShowFaceIndices           ), SetBy);
-	GeometryCollectionDebugDrawActorCVars::ShowFaceNormals         ->Set(int32(bShowFaceNormals           ), SetBy);
-	GeometryCollectionDebugDrawActorCVars::ShowVertices            ->Set(int32(bShowVertices              ), SetBy);
-	GeometryCollectionDebugDrawActorCVars::ShowVertexIndices       ->Set(int32(bShowVertexIndices         ), SetBy);
-	GeometryCollectionDebugDrawActorCVars::ShowVertexNormals       ->Set(int32(bShowVertexNormals         ), SetBy);
-	GeometryCollectionDebugDrawActorCVars::UseActiveVisualization  ->Set(int32(bUseActiveVisualization    ), SetBy);
-	GeometryCollectionDebugDrawActorCVars::PointThickness          ->Set(      PointThickness              , SetBy);
-	GeometryCollectionDebugDrawActorCVars::LineThickness           ->Set(      LineThickness               , SetBy);
-	GeometryCollectionDebugDrawActorCVars::TextShadow              ->Set(int32(bTextShadow                ), SetBy);
-	GeometryCollectionDebugDrawActorCVars::TextScale               ->Set(      TextScale                   , SetBy);
-	GeometryCollectionDebugDrawActorCVars::NormalScale             ->Set(      NormalScale                 , SetBy);
-	GeometryCollectionDebugDrawActorCVars::AxisScale               ->Set(      AxisScale                   , SetBy);
-	GeometryCollectionDebugDrawActorCVars::ArrowScale              ->Set(      ArrowScale                  , SetBy);
-
-	Super::PostLoad();
-}
-
-FColor AGeometryCollectionDebugDrawActor::MakeDarker(const FColor& Color, int32 Level)
-{
-	FLinearColor LinearColor(Color);
-	LinearColor = LinearColor.LinearRGBToHSV();
-	LinearColor.Component(2) /= float(1L << Level);
-	return LinearColor.HSVToLinearRGB().ToFColor(true);
-}
-
-int32 AGeometryCollectionDebugDrawActor::GetLevel(int32 TransformIndex, const TManagedArray<int32>& Parents)
-{
-	check(TransformIndex != FGeometryCollection::Invalid);
-	int32 Level = 0;
-	while ((TransformIndex = Parents[TransformIndex]) != FGeometryCollection::Invalid)
-	{
-		++Level;
-	}
-	return Level;
-}
-
-void AGeometryCollectionDebugDrawActor::OnPropertyChanged(bool bForceVisibilityUpdate)
-{
-#if GEOMETRYCOLLECTION_DEBUG_DRAW
-	// Reset selected rigid body's actor
-	SelectedRigidBody.GeometryCollection = nullptr;
-
-	// Update component states
-	UWorld* const World = GetWorld();
-	if (!World || !World->HasBegunPlay() || !HasActorBegunPlay()) { return; }
-
-	for (TActorIterator<AGeometryCollectionActor> ActorIterator(World); ActorIterator; ++ActorIterator)
-	{
-		UGeometryCollectionDebugDrawComponent* const GeometryCollectionDebugDrawComponent = ActorIterator->GetGeometryCollectionDebugDrawComponent();
-		if (GeometryCollectionDebugDrawComponent)
-		{
-			check(GeometryCollectionDebugDrawComponent->GeometryCollectionDebugDrawActor == this);
-			const bool bIsSelected = GeometryCollectionDebugDrawComponent->OnDebugDrawPropertiesChanged(bForceVisibilityUpdate);
-			if (bIsSelected)
-			{
-				SelectedRigidBody.GeometryCollection = *ActorIterator;
-			}
-		}
-	}
-#endif
-}
-
-template<typename T1, typename T2>
-void AGeometryCollectionDebugDrawActor::UpdatePropertyValue(T1& PropertyValue, const TAutoConsoleVariable<T2>& ConsoleVariable, bool& bHasChanged)
-{
-	const T1 NewValue = static_cast<T1>(ConsoleVariable.GetValueOnGameThread());
-	if (PropertyValue != NewValue)
-	{
-		bHasChanged = true;
-		PropertyValue = NewValue;
-	}
-}
-
-template<typename T, decltype(ImplicitConv<AChaosSolverActor*>(std::declval<T>()))*>
-void AGeometryCollectionDebugDrawActor::UpdatePropertyValue(T& PropertyValue, const TAutoConsoleVariable<FString>& ConsoleVariable, bool& bHasChanged)
-{
-	AChaosSolverActor* NewValue = nullptr;
-
-	if (UWorld* const World = GetWorld())
-	{
-		const FString SolverName = ConsoleVariable.GetValueOnGameThread();
-		for (TActorIterator<AChaosSolverActor> ActorIterator(World); ActorIterator; ++ActorIterator)
-		{
-			if (ActorIterator->GetName() == SolverName)
-			{
-				NewValue = *ActorIterator;
-				break;
-			}
-		}
-	}
-
-	if (PropertyValue != NewValue)
-	{
-		bHasChanged = true;
-		PropertyValue = NewValue;
-	}
-}
-
-void AGeometryCollectionDebugDrawActor::OnCVarsChanged()
-{
-	// Discard callback if this actor isn't in any world
-	if (!GetWorld()) { return; }
-
-	// Update properties from cvars
-	bool bHavePropertiesChanged = false;
-	bool bHasDebugDrawClusteringChanged = false;
-
-	UpdatePropertyValue(SelectedRigidBody.Solver , GeometryCollectionDebugDrawActorCVars::SelectedRigidBodySolver , bHavePropertiesChanged);
-	UpdatePropertyValue(SelectedRigidBody.Id     , GeometryCollectionDebugDrawActorCVars::SelectedRigidBodyId     , bHavePropertiesChanged);
-	UpdatePropertyValue(bDebugDrawWholeCollection, GeometryCollectionDebugDrawActorCVars::DebugDrawWholeCollection, bHavePropertiesChanged);
-	UpdatePropertyValue(bDebugDrawHierarchy      , GeometryCollectionDebugDrawActorCVars::DebugDrawHierarchy      , bHavePropertiesChanged);
-	UpdatePropertyValue(bDebugDrawClustering     , GeometryCollectionDebugDrawActorCVars::DebugDrawClustering     , bHasDebugDrawClusteringChanged);
-	UpdatePropertyValue(HideGeometry             , GeometryCollectionDebugDrawActorCVars::HideGeometry            , bHavePropertiesChanged);
-	UpdatePropertyValue(bShowRigidBodyId         , GeometryCollectionDebugDrawActorCVars::ShowRigidBodyId         , bHavePropertiesChanged);
-	UpdatePropertyValue(bShowRigidBodyCollision  , GeometryCollectionDebugDrawActorCVars::ShowRigidBodyCollision  , bHavePropertiesChanged);
-	UpdatePropertyValue(bCollisionAtOrigin       , GeometryCollectionDebugDrawActorCVars::CollisionAtOrigin       , bHavePropertiesChanged);
-	UpdatePropertyValue(bShowRigidBodyTransform  , GeometryCollectionDebugDrawActorCVars::ShowRigidBodyTransform  , bHavePropertiesChanged);
-	UpdatePropertyValue(bShowRigidBodyInertia    , GeometryCollectionDebugDrawActorCVars::ShowRigidBodyInertia    , bHavePropertiesChanged);
-	UpdatePropertyValue(bShowRigidBodyVelocity   , GeometryCollectionDebugDrawActorCVars::ShowRigidBodyVelocity   , bHavePropertiesChanged);
-	UpdatePropertyValue(bShowRigidBodyForce      , GeometryCollectionDebugDrawActorCVars::ShowRigidBodyForce      , bHavePropertiesChanged);
-	UpdatePropertyValue(bShowRigidBodyInfos      , GeometryCollectionDebugDrawActorCVars::ShowRigidBodyInfos      , bHavePropertiesChanged);
-	UpdatePropertyValue(bShowTransformIndex      , GeometryCollectionDebugDrawActorCVars::ShowTransformIndex      , bHavePropertiesChanged);
-	UpdatePropertyValue(bShowTransform           , GeometryCollectionDebugDrawActorCVars::ShowTransform           , bHavePropertiesChanged);
-	UpdatePropertyValue(bShowParent              , GeometryCollectionDebugDrawActorCVars::ShowParent              , bHavePropertiesChanged);
-	UpdatePropertyValue(bShowLevel               , GeometryCollectionDebugDrawActorCVars::ShowLevel               , bHavePropertiesChanged);
-	UpdatePropertyValue(bShowConnectivityEdges   , GeometryCollectionDebugDrawActorCVars::ShowConnectivityEdges   , bHavePropertiesChanged);
-	UpdatePropertyValue(bShowGeometryIndex       , GeometryCollectionDebugDrawActorCVars::ShowGeometryIndex       , bHavePropertiesChanged);
-	UpdatePropertyValue(bShowGeometryTransform   , GeometryCollectionDebugDrawActorCVars::ShowGeometryTransform   , bHavePropertiesChanged);
-	UpdatePropertyValue(bShowBoundingBox         , GeometryCollectionDebugDrawActorCVars::ShowBoundingBox         , bHavePropertiesChanged);
-	UpdatePropertyValue(bShowFaces               , GeometryCollectionDebugDrawActorCVars::ShowFaces               , bHavePropertiesChanged);
-	UpdatePropertyValue(bShowFaceIndices         , GeometryCollectionDebugDrawActorCVars::ShowFaceIndices         , bHavePropertiesChanged);
-	UpdatePropertyValue(bShowFaceNormals         , GeometryCollectionDebugDrawActorCVars::ShowFaceNormals         , bHavePropertiesChanged);
-	UpdatePropertyValue(bShowSingleFace          , GeometryCollectionDebugDrawActorCVars::ShowSingleFace          , bHavePropertiesChanged);
-	UpdatePropertyValue(SingleFaceIndex          , GeometryCollectionDebugDrawActorCVars::SingleFaceIndex         , bHavePropertiesChanged);
-	UpdatePropertyValue(bShowVertices            , GeometryCollectionDebugDrawActorCVars::ShowVertices            , bHavePropertiesChanged);
-	UpdatePropertyValue(bShowVertexIndices       , GeometryCollectionDebugDrawActorCVars::ShowVertexIndices       , bHavePropertiesChanged);
-	UpdatePropertyValue(bShowVertexNormals       , GeometryCollectionDebugDrawActorCVars::ShowVertexNormals       , bHavePropertiesChanged);
-	UpdatePropertyValue(bUseActiveVisualization  , GeometryCollectionDebugDrawActorCVars::UseActiveVisualization  , bHavePropertiesChanged);
-	UpdatePropertyValue(PointThickness           , GeometryCollectionDebugDrawActorCVars::PointThickness          , bHavePropertiesChanged);
-	UpdatePropertyValue(LineThickness            , GeometryCollectionDebugDrawActorCVars::LineThickness           , bHavePropertiesChanged);
-	UpdatePropertyValue(bTextShadow              , GeometryCollectionDebugDrawActorCVars::TextShadow              , bHavePropertiesChanged);
-	UpdatePropertyValue(TextScale                , GeometryCollectionDebugDrawActorCVars::TextScale               , bHavePropertiesChanged);
-	UpdatePropertyValue(NormalScale              , GeometryCollectionDebugDrawActorCVars::NormalScale             , bHavePropertiesChanged);
-	UpdatePropertyValue(AxisScale                , GeometryCollectionDebugDrawActorCVars::AxisScale               , bHavePropertiesChanged);
-	UpdatePropertyValue(ArrowScale               , GeometryCollectionDebugDrawActorCVars::ArrowScale              , bHavePropertiesChanged);
-
-	// Update geometry collection component, but only if this actor has begun play
-	if (bHavePropertiesChanged || bHasDebugDrawClusteringChanged)
-	{
-		OnPropertyChanged(bHasDebugDrawClusteringChanged);
-	}
-}
-
-#if WITH_EDITOR
-void AGeometryCollectionDebugDrawActor::PostEditChangeProperty(FPropertyChangedEvent& PropertyChangedEvent)
-{
-	// Synchronize the command variables to this Actor's properties if the property name matches.
-	static const EConsoleVariableFlags SetBy = ECVF_SetByConsole;  // Can't use the default ECVF_SetByCode as otherwise changing the UI won't update the global console variable.
-	const FName PropertyName = PropertyChangedEvent.Property ? PropertyChangedEvent.Property->GetFName() : NAME_None;
-
-	bool bForceVisibilityUpdate = false;
-
-	if      (PropertyName == GET_MEMBER_NAME_CHECKED(AGeometryCollectionDebugDrawActor, SelectedRigidBody        )) { GeometryCollectionDebugDrawActorCVars::SelectedRigidBodySolver ->Set(*SelectedRigidBody.GetSolverName(), SetBy);
-																													  GeometryCollectionDebugDrawActorCVars::SelectedRigidBodyId     ->Set( SelectedRigidBody.Id             , SetBy); }
-//																													  GeometryCollectionDebugDrawActorCVars::SelectedRigidBodyId     ->Set(*SelectedRigidBody.Id.ToString()  , SetBy); }
-	else if (PropertyName == GET_MEMBER_NAME_CHECKED(FGeometryCollectionDebugDrawActorSelectedRigidBody, Solver  )) { GeometryCollectionDebugDrawActorCVars::SelectedRigidBodySolver ->Set(*SelectedRigidBody.GetSolverName(), SetBy); }
-	else if (PropertyName == GET_MEMBER_NAME_CHECKED(FGeometryCollectionDebugDrawActorSelectedRigidBody, Id      )) { GeometryCollectionDebugDrawActorCVars::SelectedRigidBodyId     ->Set( SelectedRigidBody.Id             , SetBy); }
-//	else if (PropertyName == GET_MEMBER_NAME_CHECKED(FGeometryCollectionDebugDrawActorSelectedRigidBody, Id      )) { GeometryCollectionDebugDrawActorCVars::SelectedRigidBodyId     ->Set(*SelectedRigidBody.Id.ToString()  , SetBy); }
-	else if (PropertyName == GET_MEMBER_NAME_CHECKED(AGeometryCollectionDebugDrawActor, bDebugDrawWholeCollection)) { GeometryCollectionDebugDrawActorCVars::DebugDrawWholeCollection->Set(int32(bDebugDrawWholeCollection  ), SetBy); }
-	else if (PropertyName == GET_MEMBER_NAME_CHECKED(AGeometryCollectionDebugDrawActor, bDebugDrawHierarchy      )) { GeometryCollectionDebugDrawActorCVars::DebugDrawHierarchy      ->Set(int32(bDebugDrawHierarchy        ), SetBy); }
-	else if (PropertyName == GET_MEMBER_NAME_CHECKED(AGeometryCollectionDebugDrawActor, bDebugDrawClustering     )) { GeometryCollectionDebugDrawActorCVars::DebugDrawClustering     ->Set(int32(bDebugDrawClustering       ), SetBy); bForceVisibilityUpdate = true; }
-	else if (PropertyName == GET_MEMBER_NAME_CHECKED(AGeometryCollectionDebugDrawActor, HideGeometry             )) { GeometryCollectionDebugDrawActorCVars::HideGeometry            ->Set(int32(HideGeometry               ), SetBy); }
-	else if (PropertyName == GET_MEMBER_NAME_CHECKED(AGeometryCollectionDebugDrawActor, bShowRigidBodyId         )) { GeometryCollectionDebugDrawActorCVars::ShowRigidBodyId         ->Set(int32(bShowRigidBodyId           ), SetBy); }
-	else if (PropertyName == GET_MEMBER_NAME_CHECKED(AGeometryCollectionDebugDrawActor, bShowRigidBodyCollision  )) { GeometryCollectionDebugDrawActorCVars::ShowRigidBodyCollision  ->Set(int32(bShowRigidBodyCollision    ), SetBy); }
-	else if (PropertyName == GET_MEMBER_NAME_CHECKED(AGeometryCollectionDebugDrawActor, bCollisionAtOrigin       )) { GeometryCollectionDebugDrawActorCVars::CollisionAtOrigin       ->Set(int32(bCollisionAtOrigin         ), SetBy); }
-	else if (PropertyName == GET_MEMBER_NAME_CHECKED(AGeometryCollectionDebugDrawActor, bShowRigidBodyTransform  )) { GeometryCollectionDebugDrawActorCVars::ShowRigidBodyTransform  ->Set(int32(bShowRigidBodyTransform    ), SetBy); }
-	else if (PropertyName == GET_MEMBER_NAME_CHECKED(AGeometryCollectionDebugDrawActor, bShowRigidBodyInertia    )) { GeometryCollectionDebugDrawActorCVars::ShowRigidBodyInertia    ->Set(int32(bShowRigidBodyInertia      ), SetBy); }
-	else if (PropertyName == GET_MEMBER_NAME_CHECKED(AGeometryCollectionDebugDrawActor, bShowRigidBodyVelocity   )) { GeometryCollectionDebugDrawActorCVars::ShowRigidBodyVelocity   ->Set(int32(bShowRigidBodyVelocity     ), SetBy); }
-	else if (PropertyName == GET_MEMBER_NAME_CHECKED(AGeometryCollectionDebugDrawActor, bShowRigidBodyForce      )) { GeometryCollectionDebugDrawActorCVars::ShowRigidBodyForce      ->Set(int32(bShowRigidBodyForce        ), SetBy); }
-	else if (PropertyName == GET_MEMBER_NAME_CHECKED(AGeometryCollectionDebugDrawActor, bShowRigidBodyInfos      )) { GeometryCollectionDebugDrawActorCVars::ShowRigidBodyInfos      ->Set(int32(bShowRigidBodyInfos        ), SetBy); }
-	else if (PropertyName == GET_MEMBER_NAME_CHECKED(AGeometryCollectionDebugDrawActor, bShowTransformIndex      )) { GeometryCollectionDebugDrawActorCVars::ShowTransformIndex      ->Set(int32(bShowTransformIndex        ), SetBy); }
-	else if (PropertyName == GET_MEMBER_NAME_CHECKED(AGeometryCollectionDebugDrawActor, bShowTransform           )) { GeometryCollectionDebugDrawActorCVars::ShowTransform           ->Set(int32(bShowTransform             ), SetBy); }
-	else if (PropertyName == GET_MEMBER_NAME_CHECKED(AGeometryCollectionDebugDrawActor, bShowParent              )) { GeometryCollectionDebugDrawActorCVars::ShowParent              ->Set(int32(bShowParent                ), SetBy); }
-	else if (PropertyName == GET_MEMBER_NAME_CHECKED(AGeometryCollectionDebugDrawActor, bShowLevel               )) { GeometryCollectionDebugDrawActorCVars::ShowLevel               ->Set(int32(bShowLevel                 ), SetBy); }
-	else if (PropertyName == GET_MEMBER_NAME_CHECKED(AGeometryCollectionDebugDrawActor, bShowConnectivityEdges   )) { GeometryCollectionDebugDrawActorCVars::ShowConnectivityEdges   ->Set(int32(bShowConnectivityEdges     ), SetBy); }
-	else if (PropertyName == GET_MEMBER_NAME_CHECKED(AGeometryCollectionDebugDrawActor, bShowGeometryIndex       )) { GeometryCollectionDebugDrawActorCVars::ShowGeometryIndex       ->Set(int32(bShowGeometryIndex         ), SetBy); }
-	else if (PropertyName == GET_MEMBER_NAME_CHECKED(AGeometryCollectionDebugDrawActor, bShowGeometryTransform   )) { GeometryCollectionDebugDrawActorCVars::ShowGeometryTransform   ->Set(int32(bShowGeometryTransform     ), SetBy); }
-	else if (PropertyName == GET_MEMBER_NAME_CHECKED(AGeometryCollectionDebugDrawActor, bShowBoundingBox         )) { GeometryCollectionDebugDrawActorCVars::ShowBoundingBox         ->Set(int32(bShowBoundingBox           ), SetBy); }
-	else if (PropertyName == GET_MEMBER_NAME_CHECKED(AGeometryCollectionDebugDrawActor, bShowFaces               )) { GeometryCollectionDebugDrawActorCVars::ShowFaces               ->Set(int32(bShowFaces                 ), SetBy); }
-	else if (PropertyName == GET_MEMBER_NAME_CHECKED(AGeometryCollectionDebugDrawActor, bShowFaceIndices         )) { GeometryCollectionDebugDrawActorCVars::ShowFaceIndices         ->Set(int32(bShowFaceIndices           ), SetBy); }
-	else if (PropertyName == GET_MEMBER_NAME_CHECKED(AGeometryCollectionDebugDrawActor, bShowFaceNormals         )) { GeometryCollectionDebugDrawActorCVars::ShowFaceNormals         ->Set(int32(bShowFaceNormals           ), SetBy); }
-	else if (PropertyName == GET_MEMBER_NAME_CHECKED(AGeometryCollectionDebugDrawActor, bShowSingleFace          )) { GeometryCollectionDebugDrawActorCVars::ShowSingleFace          ->Set(int32(bShowSingleFace            ), SetBy); }
-	else if (PropertyName == GET_MEMBER_NAME_CHECKED(AGeometryCollectionDebugDrawActor, SingleFaceIndex          )) { GeometryCollectionDebugDrawActorCVars::SingleFaceIndex         ->Set(      SingleFaceIndex             , SetBy); }
-	else if (PropertyName == GET_MEMBER_NAME_CHECKED(AGeometryCollectionDebugDrawActor, bShowVertices            )) { GeometryCollectionDebugDrawActorCVars::ShowVertices            ->Set(int32(bShowVertices              ), SetBy); }
-	else if (PropertyName == GET_MEMBER_NAME_CHECKED(AGeometryCollectionDebugDrawActor, bShowVertexIndices       )) { GeometryCollectionDebugDrawActorCVars::ShowVertexIndices       ->Set(int32(bShowVertexIndices         ), SetBy); }
-	else if (PropertyName == GET_MEMBER_NAME_CHECKED(AGeometryCollectionDebugDrawActor, bShowVertexNormals       )) { GeometryCollectionDebugDrawActorCVars::ShowVertexNormals       ->Set(int32(bShowVertexNormals         ), SetBy); }
-	else
-	{
-		// These properties are cosmetic changes and don't require visibility updates or enabling the component tick
-		if      (PropertyName == GET_MEMBER_NAME_CHECKED(AGeometryCollectionDebugDrawActor, bUseActiveVisualization  )) { GeometryCollectionDebugDrawActorCVars::UseActiveVisualization  ->Set(int32(bUseActiveVisualization    ), SetBy); }
-		else if (PropertyName == GET_MEMBER_NAME_CHECKED(AGeometryCollectionDebugDrawActor, PointThickness           )) { GeometryCollectionDebugDrawActorCVars::PointThickness          ->Set(      PointThickness              , SetBy); }
-		else if (PropertyName == GET_MEMBER_NAME_CHECKED(AGeometryCollectionDebugDrawActor, LineThickness            )) { GeometryCollectionDebugDrawActorCVars::LineThickness           ->Set(      LineThickness               , SetBy); }
-		else if (PropertyName == GET_MEMBER_NAME_CHECKED(AGeometryCollectionDebugDrawActor, bTextShadow              )) { GeometryCollectionDebugDrawActorCVars::TextShadow              ->Set(int32(bTextShadow                ), SetBy); }
-		else if (PropertyName == GET_MEMBER_NAME_CHECKED(AGeometryCollectionDebugDrawActor, TextScale                )) { GeometryCollectionDebugDrawActorCVars::TextScale               ->Set(      TextScale                   , SetBy); }
-		else if (PropertyName == GET_MEMBER_NAME_CHECKED(AGeometryCollectionDebugDrawActor, NormalScale              )) { GeometryCollectionDebugDrawActorCVars::NormalScale             ->Set(      NormalScale                 , SetBy); }
-		else if (PropertyName == GET_MEMBER_NAME_CHECKED(AGeometryCollectionDebugDrawActor, AxisScale                )) { GeometryCollectionDebugDrawActorCVars::AxisScale               ->Set(      AxisScale                   , SetBy); }
-		else if (PropertyName == GET_MEMBER_NAME_CHECKED(AGeometryCollectionDebugDrawActor, ArrowScale               )) { GeometryCollectionDebugDrawActorCVars::ArrowScale              ->Set(      ArrowScale                  , SetBy); }
-
-		Super::PostEditChangeProperty(PropertyChangedEvent);
-		return;  // Don't call OnPropertyChanged()
-	}
-	OnPropertyChanged(bForceVisibilityUpdate);
-	Super::PostEditChangeProperty(PropertyChangedEvent);
-}
-
-bool AGeometryCollectionDebugDrawActor::CanEditChange(const FProperty* InProperty) const
-{
-	const FName PropertyName = InProperty->GetFName();
-	if (PropertyName == GET_MEMBER_NAME_CHECKED(AGeometryCollectionDebugDrawActor, bDebugDrawWholeCollection) ||
-		PropertyName == GET_MEMBER_NAME_CHECKED(AGeometryCollectionDebugDrawActor, bShowSingleFace))
-	{
-		return SelectedRigidBody.Id != INDEX_NONE;
-		//return SelectedRigidBody.Id.IsValid();
-	}
-	return Super::CanEditChange(InProperty);
-}
-#endif  // #if WITH_EDITOR
-
-void AGeometryCollectionDebugDrawActor::AddDebugText(const FString& Text, const FVector& Position, const FColor& Color, float Scale, bool bDrawShadow)
-{
-#if GEOMETRYCOLLECTION_DEBUG_DRAW && ENABLE_DRAW_DEBUG
-	DebugDrawTexts.Add({ Text, Position, Color, Scale, bDrawShadow });
-#endif  // #if GEOMETRYCOLLECTION_DEBUG_DRAW && ENABLE_DRAW_DEBUG
-}
-
-void AGeometryCollectionDebugDrawActor::DebugDrawText(UCanvas* Canvas, APlayerController* PlayerController)
-{
-#if GEOMETRYCOLLECTION_DEBUG_DRAW && ENABLE_DRAW_DEBUG
-	for (const FDebugDrawText& DebugDrawText : DebugDrawTexts)
-	{
-		const FVector Position = Canvas->Project(DebugDrawText.Position);
-		if (Position.Z < KINDA_SMALL_NUMBER) { continue; }  // Don't draw behind the camera
-
-		const FVector2D Postion2D(FMath::CeilToFloat(Position.X), FMath::CeilToFloat(Position.Y));
-		const FText Text = FText::FromString(DebugDrawText.Text);
-
-		FCanvasTextItem TextItem(Postion2D, Text, GEngine->GetSmallFont(), DebugDrawText.Color);
-		TextItem.Scale = FVector2D(DebugDrawText.Scale, DebugDrawText.Scale);
-		if (DebugDrawText.bDrawShadow)
-		{
-			TextItem.EnableShadow(FLinearColor::Black);
-		}
-		else
-		{
-			TextItem.DisableShadow();
-		}
-		TextItem.Draw(Canvas->Canvas);
-	}
-#endif  // #if GEOMETRYCOLLECTION_DEBUG_DRAW && ENABLE_DRAW_DEBUG
-}
-
-void AGeometryCollectionDebugDrawActor::Flush()
-{
-#if GEOMETRYCOLLECTION_DEBUG_DRAW && ENABLE_DRAW_DEBUG
-	DebugDrawTexts.Reset();
-
-	// Note that a flush will remove all the other persistent debug draw elements, so best to only do them when needed
-	if (bNeedsDebugLinesFlush)  // Only flush if a geometry collection debug draw function has been drawing lines
-	{
-		const UWorld* const World = GetWorld();
-		check(World);
-		FlushPersistentDebugLines(World);
-		bNeedsDebugLinesFlush = false;
-	}
-#endif  // #if GEOMETRYCOLLECTION_DEBUG_DRAW && ENABLE_DRAW_DEBUG
-}
-
-void AGeometryCollectionDebugDrawActor::DrawVertices(const TArray<FTransform>& GlobalTransforms, const UGeometryCollectionComponent* GeometryCollectionComponent, const FColor& Color)
-{
-#if GEOMETRYCOLLECTION_DEBUG_DRAW && ENABLE_DRAW_DEBUG
-	check(GeometryCollectionComponent);
-	check(GeometryCollectionComponent->RestCollection);
-
-	const UWorld* const World = GetWorld();
-	check(World);
-
-	const TManagedArray<int32>& TransformIndexArray = GeometryCollectionComponent->GetTransformIndexArray();
-	const TManagedArray<FVector3f>& VertexArray = GeometryCollectionComponent->GetVertexArray();
-	const TManagedArray<int32>& VertexStartArray = GeometryCollectionComponent->GetVertexStartArray();
-	const TManagedArray<int32>& VertexCountArray = GeometryCollectionComponent->GetVertexCountArray();
-
-	const TManagedArray<int32>& ParentArray = GeometryCollectionComponent->GetParentArray();
-	const TManagedArray<TSet<int32>>& ChildrenArray = GeometryCollectionComponent->GetChildrenArray();
-	const TManagedArray<TSet<int32>>& ChildrenArrayRest = GeometryCollectionComponent->GetChildrenArrayRest();
-
-	const int32 NumGeometries = GeometryCollectionComponent->GetNumElements(FGeometryCollection::GeometryGroup);
-	for (int32 GeometryIndex = 0; GeometryIndex < NumGeometries; ++GeometryIndex)
-	{
-		const int32 TransformIndex = TransformIndexArray[GeometryIndex];
-
-		bool bIsVisible;
-		const bool bIsLeafNodeRest = (ChildrenArrayRest[TransformIndex].Num() == 0);
-		if (bDebugDrawHierarchy)
-		{
-			const bool bIsLeafNode = (ChildrenArray[TransformIndex].Num() == 0);
-			const bool bHasParent = (ParentArray[TransformIndex] != FGeometryCollection::Invalid);
-			bIsVisible = (!bHasParent && bIsLeafNode == bIsLeafNodeRest) || (bHasParent && bDebugDrawClustering);
-		}
-		else
-		{
-			bIsVisible = bIsLeafNodeRest;
-		}
-		if (bIsVisible)
-		{
-			const FTransform& Transform = GlobalTransforms[TransformIndex];
-
-			const FColor ActiveColor = bUseActiveVisualization ? MakeDarker(Color, GetLevel(TransformIndex, ParentArray)): Color;
-
-			const int32 VertexStart = VertexStartArray[GeometryIndex];
-			const int32 VertexEnd = VertexStart + VertexCountArray[GeometryIndex];
-
-			for (int32 VertexIndex = VertexStart; VertexIndex < VertexEnd; ++VertexIndex)
-			{
-				const FVector Position = Transform.TransformPosition(VertexArray[VertexIndex]);
-				DrawDebugPoint(World, Position, PointThickness, ActiveColor, GeometryCollectionDebugDrawActorConstants::bPersistent, GeometryCollectionDebugDrawActorConstants::LifeTime, GeometryCollectionDebugDrawActorConstants::DepthPriority);
-			}
-			bNeedsDebugLinesFlush = true;
-		}
-	}
-#endif  // #if GEOMETRYCOLLECTION_DEBUG_DRAW && ENABLE_DRAW_DEBUG
-}
-
-void AGeometryCollectionDebugDrawActor::DrawVertices(const TArray<FTransform>& GlobalTransforms, const UGeometryCollectionComponent* GeometryCollectionComponent, int32 TransformIndex, const FColor& Color)
-{
-#if GEOMETRYCOLLECTION_DEBUG_DRAW && ENABLE_DRAW_DEBUG
-	check(GeometryCollectionComponent);
-	check(GeometryCollectionComponent->RestCollection);
-	check(TransformIndex >= 0);
-
-	const UWorld* const World = GetWorld();
-	check(World);
-
-	const TManagedArray<int32>& TransformToGeometryIndexArray = GeometryCollectionComponent->GetTransformToGeometryIndexArray();
-	const int32 GeometryIndex = TransformToGeometryIndexArray[TransformIndex];
-
-	const TManagedArray<TSet<int32>>& ChildrenArrayRest = GeometryCollectionComponent->GetChildrenArrayRest();
-	const bool bIsLeafNodeRest = (ChildrenArrayRest[TransformIndex].Num() == 0);
-
-	// Debug draw this geometry
-	if (GeometryIndex != INDEX_NONE && (bIsLeafNodeRest || bDebugDrawHierarchy))
-	{
-		const TManagedArray<FVector3f>& VertexArray = GeometryCollectionComponent->GetVertexArray();
-		const TManagedArray<int32>& VertexStartArray = GeometryCollectionComponent->GetVertexStartArray();
-		const TManagedArray<int32>& VertexCountArray = GeometryCollectionComponent->GetVertexCountArray();
-		const int32 VertexStart = VertexStartArray[GeometryIndex];
-		const int32 VertexEnd = VertexStart + VertexCountArray[GeometryIndex];
-
-		const FTransform& Transform = GlobalTransforms[TransformIndex];
-
-		for (int32 VertexIndex = VertexStart; VertexIndex < VertexEnd; ++VertexIndex)
-		{
-			const FVector Position = Transform.TransformPosition(VertexArray[VertexIndex]);
-			DrawDebugPoint(World, Position, PointThickness, Color, GeometryCollectionDebugDrawActorConstants::bPersistent, GeometryCollectionDebugDrawActorConstants::LifeTime, GeometryCollectionDebugDrawActorConstants::DepthPriority);
-		}
-		bNeedsDebugLinesFlush = true;
-	}
-
-	// Debug draw children if the cluster mode is on, or if there is no geometry attached to this node
-	if (!(bIsLeafNodeRest || bDebugDrawHierarchy) || bDebugDrawClustering)
-	{
-		const FColor ActiveColor = bUseActiveVisualization ? MakeDarker(Color) : Color;
-		for (int32 ChildTransformIndex : ChildrenArrayRest[TransformIndex])
-		{
-			DrawVertices(GlobalTransforms, GeometryCollectionComponent, ChildTransformIndex, ActiveColor);
-		}
-	}
-#endif  // #if GEOMETRYCOLLECTION_DEBUG_DRAW && ENABLE_DRAW_DEBUG
-}
-
-void AGeometryCollectionDebugDrawActor::DrawVertexIndices(const TArray<FTransform>& GlobalTransforms, const UGeometryCollectionComponent* GeometryCollectionComponent, const FColor& Color)
-{
-#if GEOMETRYCOLLECTION_DEBUG_DRAW && ENABLE_DRAW_DEBUG
-	check(GeometryCollectionComponent);
-	check(GeometryCollectionComponent->RestCollection);
-
-	const UWorld* const World = GetWorld();
-	check(World);
-
-	const TManagedArray<int32>& TransformIndexArray = GeometryCollectionComponent->GetTransformIndexArray();
-	const TManagedArray<FVector3f>& VertexArray = GeometryCollectionComponent->GetVertexArray();
-	const TManagedArray<int32>& VertexStartArray = GeometryCollectionComponent->GetVertexStartArray();
-	const TManagedArray<int32>& VertexCountArray = GeometryCollectionComponent->GetVertexCountArray();
-
-	const TManagedArray<int32>& ParentArray = GeometryCollectionComponent->GetParentArray();
-	const TManagedArray<TSet<int32>>& ChildrenArray = GeometryCollectionComponent->GetChildrenArray();
-	const TManagedArray<TSet<int32>>& ChildrenArrayRest = GeometryCollectionComponent->GetChildrenArrayRest();
-
-	const int32 NumGeometries = GeometryCollectionComponent->GetNumElements(FGeometryCollection::GeometryGroup);
-	for (int32 GeometryIndex = 0; GeometryIndex < NumGeometries; ++GeometryIndex)
-	{
-		const int32 TransformIndex = TransformIndexArray[GeometryIndex];
-
-		bool bIsVisible;
-		const bool bIsLeafNodeRest = (ChildrenArrayRest[TransformIndex].Num() == 0);
-		if (bDebugDrawHierarchy)
-		{
-			const bool bIsLeafNode = (ChildrenArray[TransformIndex].Num() == 0);
-			const bool bHasParent = (ParentArray[TransformIndex] != FGeometryCollection::Invalid);
-			bIsVisible = (!bHasParent && bIsLeafNode == bIsLeafNodeRest) || (bHasParent && bDebugDrawClustering);
-		}
-		else
-		{
-			bIsVisible = bIsLeafNodeRest;
-		}
-		if (bIsVisible)
-		{
-			const FTransform& Transform = GlobalTransforms[TransformIndex];
-
-			const FColor ActiveColor = bUseActiveVisualization ? MakeDarker(Color, GetLevel(TransformIndex, ParentArray)) : Color;
-
-			const int32 VertexStart = VertexStartArray[GeometryIndex];
-			const int32 VertexEnd = VertexStart + VertexCountArray[GeometryIndex];
-
-			for (int32 VertexIndex = VertexStart; VertexIndex < VertexEnd; ++VertexIndex)
-			{
-				const FVector Position = Transform.TransformPosition(VertexArray[VertexIndex]);
-				const FString Text = FString::Printf(TEXT("%d"), VertexIndex);
-				AddDebugText(Text, Position, ActiveColor, TextScale, bTextShadow);
-			}
-		}
-	}
-#endif  // #if GEOMETRYCOLLECTION_DEBUG_DRAW && ENABLE_DRAW_DEBUG
-}
-
-void AGeometryCollectionDebugDrawActor::DrawVertexIndices(const TArray<FTransform>& GlobalTransforms, const UGeometryCollectionComponent* GeometryCollectionComponent, int32 TransformIndex, const FColor& Color)
-{
-#if GEOMETRYCOLLECTION_DEBUG_DRAW && ENABLE_DRAW_DEBUG
-	check(GeometryCollectionComponent);
-	check(GeometryCollectionComponent->RestCollection);
-	check(TransformIndex >= 0);
-
-	const UWorld* const World = GetWorld();
-	check(World);
-
-	const TManagedArray<int32>& TransformToGeometryIndexArray = GeometryCollectionComponent->GetTransformToGeometryIndexArray();
-	const int32 GeometryIndex = TransformToGeometryIndexArray[TransformIndex];
-
-	const TManagedArray<TSet<int32>>& ChildrenArrayRest = GeometryCollectionComponent->GetChildrenArrayRest();
-	const bool bIsLeafNodeRest = (ChildrenArrayRest[TransformIndex].Num() == 0);
-
-	// Debug draw this geometry
-	if (GeometryIndex != INDEX_NONE && (bIsLeafNodeRest || bDebugDrawHierarchy))
-	{
-		const FTransform& Transform = GlobalTransforms[TransformIndex];
-
-		const TManagedArray<FVector3f>& VertexArray = GeometryCollectionComponent->GetVertexArray();
-		const TManagedArray<int32>& VertexStartArray = GeometryCollectionComponent->GetVertexStartArray();
-		const TManagedArray<int32>& VertexCountArray = GeometryCollectionComponent->GetVertexCountArray();
-		const int32 VertexStart = VertexStartArray[GeometryIndex];
-		const int32 VertexEnd = VertexStart + VertexCountArray[GeometryIndex];
-
-		for (int32 VertexIndex = VertexStart; VertexIndex < VertexEnd; ++VertexIndex)
-		{
-			const FVector Position = Transform.TransformPosition(VertexArray[VertexIndex]);
-			const FString Text = FString::Printf(TEXT("%d"), VertexIndex);
-			AddDebugText(Text, Position, Color, TextScale, bTextShadow);
-		}
-	}
-
-	// Debug draw children if the cluster mode is on, or if there is no geometry attached to this node
-	if (!(bIsLeafNodeRest || bDebugDrawHierarchy) || bDebugDrawClustering)
-	{
-		const FColor ActiveColor = bUseActiveVisualization ? MakeDarker(Color) : Color;
-		for (int32 ChildTransformIndex : ChildrenArrayRest[TransformIndex])
-		{
-			DrawVertexIndices(GlobalTransforms, GeometryCollectionComponent, ChildTransformIndex, ActiveColor);
-		}
-	}
-#endif  // #if GEOMETRYCOLLECTION_DEBUG_DRAW && ENABLE_DRAW_DEBUG
-}
-
-void AGeometryCollectionDebugDrawActor::DrawVertexNormals(const TArray<FTransform>& GlobalTransforms, const UGeometryCollectionComponent* GeometryCollectionComponent, const FColor& Color)
-{
-#if GEOMETRYCOLLECTION_DEBUG_DRAW && ENABLE_DRAW_DEBUG
-	check(GeometryCollectionComponent);
-	check(GeometryCollectionComponent->RestCollection);
-
-	const UWorld* const World = GetWorld();
-	check(World);
-
-	const TManagedArray<int32>& TransformIndexArray = GeometryCollectionComponent->GetTransformIndexArray();
-	const TManagedArray<FVector3f>& VertexArray = GeometryCollectionComponent->GetVertexArray();
-	const TManagedArray<FVector3f>& NormalArray = GeometryCollectionComponent->GetNormalArray();
-	const TManagedArray<int32>& VertexStartArray = GeometryCollectionComponent->GetVertexStartArray();
-	const TManagedArray<int32>& VertexCountArray = GeometryCollectionComponent->GetVertexCountArray();
-
-	const TManagedArray<int32>& ParentArray = GeometryCollectionComponent->GetParentArray();
-	const TManagedArray<TSet<int32>>& ChildrenArray = GeometryCollectionComponent->GetChildrenArray();
-	const TManagedArray<TSet<int32>>& ChildrenArrayRest = GeometryCollectionComponent->GetChildrenArrayRest();
-
-	const int32 NumGeometries = GeometryCollectionComponent->GetNumElements(FGeometryCollection::GeometryGroup);
-	for (int32 GeometryIndex = 0; GeometryIndex < NumGeometries; ++GeometryIndex)
-	{
-		const int32 TransformIndex = TransformIndexArray[GeometryIndex];
-
-		bool bIsVisible;
-		const bool bIsLeafNodeRest = (ChildrenArrayRest[TransformIndex].Num() == 0);
-		if (bDebugDrawHierarchy)
-		{
-			const bool bIsLeafNode = (ChildrenArray[TransformIndex].Num() == 0);
-			const bool bHasParent = (ParentArray[TransformIndex] != FGeometryCollection::Invalid);
-			bIsVisible = (!bHasParent && bIsLeafNode == bIsLeafNodeRest) || (bHasParent && bDebugDrawClustering);
-		}
-		else
-		{
-			bIsVisible = bIsLeafNodeRest;
-		}
-		if (bIsVisible)
-		{
-			const FTransform& Transform = GlobalTransforms[TransformIndex];
-
-			const FColor ActiveColor = bUseActiveVisualization ? MakeDarker(Color, GetLevel(TransformIndex, ParentArray)) : Color;
-
-			const int32 VertexStart = VertexStartArray[GeometryIndex];
-			const int32 VertexEnd = VertexStart + VertexCountArray[GeometryIndex];
-
-			for (int32 VertexIndex = VertexStart; VertexIndex < VertexEnd; ++VertexIndex)
-			{
-				const FVector LineStart = Transform.TransformPosition(VertexArray[VertexIndex]);
-				const FVector VertexNormal = Transform.TransformVector(NormalArray[VertexIndex]).GetSafeNormal();
-				const FVector LineEnd = LineStart + VertexNormal * NormalScale;
-				DrawDebugDirectionalArrow(World, LineStart, LineEnd, ArrowScale, ActiveColor, GeometryCollectionDebugDrawActorConstants::bPersistent, GeometryCollectionDebugDrawActorConstants::LifeTime, GeometryCollectionDebugDrawActorConstants::DepthPriority, LineThickness);
-			}
-			bNeedsDebugLinesFlush = true;
-		}
-	}
-#endif  // #if GEOMETRYCOLLECTION_DEBUG_DRAW && ENABLE_DRAW_DEBUG
-}
-
-void AGeometryCollectionDebugDrawActor::DrawVertexNormals(const TArray<FTransform>& GlobalTransforms, const UGeometryCollectionComponent* GeometryCollectionComponent, int32 TransformIndex, const FColor& Color)
-{
-#if GEOMETRYCOLLECTION_DEBUG_DRAW && ENABLE_DRAW_DEBUG
-	check(GeometryCollectionComponent);
-	check(GeometryCollectionComponent->RestCollection);
-	check(TransformIndex >= 0);
-
-	const UWorld* const World = GetWorld();
-	check(World);
-
-	const TManagedArray<int32>& TransformToGeometryIndexArray = GeometryCollectionComponent->GetTransformToGeometryIndexArray();
-	const int32 GeometryIndex = TransformToGeometryIndexArray[TransformIndex];
-
-	const TManagedArray<TSet<int32>>& ChildrenArrayRest = GeometryCollectionComponent->GetChildrenArrayRest();
-	const bool bIsLeafNodeRest = (ChildrenArrayRest[TransformIndex].Num() == 0);
-
-	// Debug draw this geometry
-	if (GeometryIndex != INDEX_NONE && (bIsLeafNodeRest || bDebugDrawHierarchy))
-	{
-		const FTransform& Transform = GlobalTransforms[TransformIndex];
-
-		const TManagedArray<FVector3f>& VertexArray = GeometryCollectionComponent->GetVertexArray();
-		const TManagedArray<FVector3f>& NormalArray = GeometryCollectionComponent->GetNormalArray();
-		const TManagedArray<int32>& VertexStartArray = GeometryCollectionComponent->GetVertexStartArray();
-		const TManagedArray<int32>& VertexCountArray = GeometryCollectionComponent->GetVertexCountArray();
-		const int32 VertexStart = VertexStartArray[GeometryIndex];
-		const int32 VertexEnd = VertexStart + VertexCountArray[GeometryIndex];
-
-		for (int32 VertexIndex = VertexStart; VertexIndex < VertexEnd; ++VertexIndex)
-		{
-			const FVector LineStart = Transform.TransformPosition(VertexArray[VertexIndex]);
-			const FVector VertexNormal = Transform.TransformVector(NormalArray[VertexIndex]).GetSafeNormal();
-			const FVector LineEnd = LineStart + VertexNormal * NormalScale;
-
-			DrawDebugDirectionalArrow(World, LineStart, LineEnd, ArrowScale, Color, GeometryCollectionDebugDrawActorConstants::bPersistent, GeometryCollectionDebugDrawActorConstants::LifeTime, GeometryCollectionDebugDrawActorConstants::DepthPriority, LineThickness);
-		}
-		bNeedsDebugLinesFlush = true;
-	}
-
-	// Debug draw children if the cluster mode is on, or if there is no geometry attached to this node
-	if (!(bIsLeafNodeRest || bDebugDrawHierarchy) || bDebugDrawClustering)
-	{
-		const FColor ActiveColor = bUseActiveVisualization ? MakeDarker(Color) : Color;
-		for (int32 ChildTransformIndex : ChildrenArrayRest[TransformIndex])
-		{
-			DrawVertexNormals(GlobalTransforms, GeometryCollectionComponent, ChildTransformIndex, ActiveColor);
-		}
-	}
-#endif  // #if GEOMETRYCOLLECTION_DEBUG_DRAW && ENABLE_DRAW_DEBUG
-}
-
-void AGeometryCollectionDebugDrawActor::DrawFaces(const TArray<FTransform>& GlobalTransforms, const UGeometryCollectionComponent* GeometryCollectionComponent, const FColor& Color)
-{
-#if GEOMETRYCOLLECTION_DEBUG_DRAW && ENABLE_DRAW_DEBUG
-	check(GeometryCollectionComponent);
-	check(GeometryCollectionComponent->RestCollection);
-
-	const UWorld* const World = GetWorld();
-	check(World);
-
-	const TManagedArray<FVector3f>& VertexArray = GeometryCollectionComponent->GetVertexArray();
-	const TManagedArray<FIntVector>& IndicesArray = GeometryCollectionComponent->GetIndicesArray();
-	const TManagedArray<int32>& TransformIndexArray = GeometryCollectionComponent->GetTransformIndexArray();
-	const TManagedArray<int32>& FaceStartArray = GeometryCollectionComponent->GetFaceStartArray();
-	const TManagedArray<int32>& FaceCountArray = GeometryCollectionComponent->GetFaceCountArray();
-
-	const TManagedArray<int32>& ParentArray = GeometryCollectionComponent->GetParentArray();
-	const TManagedArray<TSet<int32>>& ChildrenArray = GeometryCollectionComponent->GetChildrenArray();
-	const TManagedArray<TSet<int32>>& ChildrenArrayRest = GeometryCollectionComponent->GetChildrenArrayRest();
-
-	const int32 NumGeometries = GeometryCollectionComponent->GetNumElements(FGeometryCollection::GeometryGroup);
-	for (int32 GeometryIndex = 0; GeometryIndex < NumGeometries; ++GeometryIndex)
-	{
-		const int32 TransformIndex = TransformIndexArray[GeometryIndex];
-
-		bool bIsVisible;
-		const bool bIsLeafNodeRest = (ChildrenArrayRest[TransformIndex].Num() == 0);
-		if (bDebugDrawHierarchy)
-		{
-			const bool bIsLeafNode = (ChildrenArray[TransformIndex].Num() == 0);
-			const bool bHasParent = (ParentArray[TransformIndex] != FGeometryCollection::Invalid);
-			bIsVisible = (!bHasParent && bIsLeafNode == bIsLeafNodeRest) || (bHasParent && bDebugDrawClustering);
-		}
-		else
-		{
-			bIsVisible = bIsLeafNodeRest;
-		}
-		if (bIsVisible)
-		{
-			const FTransform& Transform = GlobalTransforms[TransformIndex];
-
-			const FColor ActiveColor = bUseActiveVisualization ? MakeDarker(Color, GetLevel(TransformIndex, ParentArray)) : Color;
-
-			const int32 FaceStart = FaceStartArray[GeometryIndex];
-			const int32 FaceEnd = FaceStart + FaceCountArray[GeometryIndex];
-
-			for (int32 FaceIndex = FaceStart; FaceIndex < FaceEnd; ++FaceIndex)
-			{
-				const FVector Vertex0 = Transform.TransformPosition(VertexArray[IndicesArray[FaceIndex][0]]);
-				const FVector Vertex1 = Transform.TransformPosition(VertexArray[IndicesArray[FaceIndex][1]]);
-				const FVector Vertex2 = Transform.TransformPosition(VertexArray[IndicesArray[FaceIndex][2]]);
-
-				DrawDebugLine(World, Vertex0, Vertex1, ActiveColor, GeometryCollectionDebugDrawActorConstants::bPersistent, GeometryCollectionDebugDrawActorConstants::LifeTime, GeometryCollectionDebugDrawActorConstants::DepthPriority, LineThickness);
-				DrawDebugLine(World, Vertex0, Vertex2, ActiveColor, GeometryCollectionDebugDrawActorConstants::bPersistent, GeometryCollectionDebugDrawActorConstants::LifeTime, GeometryCollectionDebugDrawActorConstants::DepthPriority, LineThickness);
-				DrawDebugLine(World, Vertex1, Vertex2, ActiveColor, GeometryCollectionDebugDrawActorConstants::bPersistent, GeometryCollectionDebugDrawActorConstants::LifeTime, GeometryCollectionDebugDrawActorConstants::DepthPriority, LineThickness);
-			}
-			bNeedsDebugLinesFlush = true;
-		}
-	}
-#endif  // #if GEOMETRYCOLLECTION_DEBUG_DRAW && ENABLE_DRAW_DEBUG
-}
-
-void AGeometryCollectionDebugDrawActor::DrawFaces(const TArray<FTransform>& GlobalTransforms, const UGeometryCollectionComponent* GeometryCollectionComponent, int32 TransformIndex, const FColor& Color)
-{
-#if GEOMETRYCOLLECTION_DEBUG_DRAW && ENABLE_DRAW_DEBUG
-	check(GeometryCollectionComponent);
-	check(GeometryCollectionComponent->RestCollection);
-	check(TransformIndex >= 0);
-
-	const UWorld* const World = GetWorld();
-	check(World);
-
-	const TManagedArray<int32>& TransformToGeometryIndexArray = GeometryCollectionComponent->GetTransformToGeometryIndexArray();
-	const int32 GeometryIndex = TransformToGeometryIndexArray[TransformIndex];
-
-	const TManagedArray<TSet<int32>>& ChildrenArrayRest = GeometryCollectionComponent->GetChildrenArrayRest();
-	const bool bIsLeafNodeRest = (ChildrenArrayRest[TransformIndex].Num() == 0);
-
-	// Debug draw this geometry
-	if (GeometryIndex != INDEX_NONE && (bIsLeafNodeRest || bDebugDrawHierarchy))
-	{
-		const FTransform& Transform = GlobalTransforms[TransformIndex];
-
-		const TManagedArray<FVector3f>& VertexArray = GeometryCollectionComponent->GetVertexArray();
-		const TManagedArray<FIntVector>& IndicesArray = GeometryCollectionComponent->GetIndicesArray();
-		const TManagedArray<int32>& FaceStartArray = GeometryCollectionComponent->GetFaceStartArray();
-		const TManagedArray<int32>& FaceCountArray = GeometryCollectionComponent->GetFaceCountArray();
-		const int32 FaceStart = FaceStartArray[GeometryIndex];
-		const int32 FaceEnd = FaceStart + FaceCountArray[GeometryIndex];
-
-		for (int32 FaceIndex = FaceStart; FaceIndex < FaceEnd; ++FaceIndex)
-		{
-			const FVector Vertex0 = Transform.TransformPosition(VertexArray[IndicesArray[FaceIndex][0]]);
-			const FVector Vertex1 = Transform.TransformPosition(VertexArray[IndicesArray[FaceIndex][1]]);
-			const FVector Vertex2 = Transform.TransformPosition(VertexArray[IndicesArray[FaceIndex][2]]);
-
-			DrawDebugLine(World, Vertex0, Vertex1, Color, GeometryCollectionDebugDrawActorConstants::bPersistent, GeometryCollectionDebugDrawActorConstants::LifeTime, GeometryCollectionDebugDrawActorConstants::DepthPriority, LineThickness);
-			DrawDebugLine(World, Vertex0, Vertex2, Color, GeometryCollectionDebugDrawActorConstants::bPersistent, GeometryCollectionDebugDrawActorConstants::LifeTime, GeometryCollectionDebugDrawActorConstants::DepthPriority, LineThickness);
-			DrawDebugLine(World, Vertex1, Vertex2, Color, GeometryCollectionDebugDrawActorConstants::bPersistent, GeometryCollectionDebugDrawActorConstants::LifeTime, GeometryCollectionDebugDrawActorConstants::DepthPriority, LineThickness);
-		}
-		bNeedsDebugLinesFlush = true;
-	}
-
-	// Debug draw children if the cluster mode is on, or if there is no geometry attached to this node
-	if (!(bIsLeafNodeRest || bDebugDrawHierarchy) || bDebugDrawClustering)
-	{
-		const FColor ActiveColor = bUseActiveVisualization ? MakeDarker(Color) : Color;
-		for (int32 ChildTransformIndex : ChildrenArrayRest[TransformIndex])
-		{
-			DrawFaces(GlobalTransforms, GeometryCollectionComponent, ChildTransformIndex, ActiveColor);
-		}
-	}
-#endif  // #if GEOMETRYCOLLECTION_DEBUG_DRAW && ENABLE_DRAW_DEBUG
-}
-
-void AGeometryCollectionDebugDrawActor::DrawFaceIndices(const TArray<FTransform>& GlobalTransforms, const UGeometryCollectionComponent* GeometryCollectionComponent, const FColor& Color)
-{
-#if GEOMETRYCOLLECTION_DEBUG_DRAW && ENABLE_DRAW_DEBUG
-	check(GeometryCollectionComponent);
-	check(GeometryCollectionComponent->RestCollection);
-
-	const UWorld* const World = GetWorld();
-	check(World);
-
-	const TManagedArray<FVector3f>& VertexArray = GeometryCollectionComponent->GetVertexArray();
-	const TManagedArray<FIntVector>& IndicesArray = GeometryCollectionComponent->GetIndicesArray();
-	const TManagedArray<int32>& TransformIndexArray = GeometryCollectionComponent->GetTransformIndexArray();
-	const TManagedArray<int32>& FaceStartArray = GeometryCollectionComponent->GetFaceStartArray();
-	const TManagedArray<int32>& FaceCountArray = GeometryCollectionComponent->GetFaceCountArray();
-
-	const TManagedArray<int32>& ParentArray = GeometryCollectionComponent->GetParentArray();
-	const TManagedArray<TSet<int32>>& ChildrenArray = GeometryCollectionComponent->GetChildrenArray();
-	const TManagedArray<TSet<int32>>& ChildrenArrayRest = GeometryCollectionComponent->GetChildrenArrayRest();
-
-	const int32 NumGeometries = GeometryCollectionComponent->GetNumElements(FGeometryCollection::GeometryGroup);
-	for (int32 GeometryIndex = 0; GeometryIndex < NumGeometries; ++GeometryIndex)
-	{
-		const int32 TransformIndex = TransformIndexArray[GeometryIndex];
-
-		bool bIsVisible;
-		const bool bIsLeafNodeRest = (ChildrenArrayRest[TransformIndex].Num() == 0);
-		if (bDebugDrawHierarchy)
-		{
-			const bool bIsLeafNode = (ChildrenArray[TransformIndex].Num() == 0);
-			const bool bHasParent = (ParentArray[TransformIndex] != FGeometryCollection::Invalid);
-			bIsVisible = (!bHasParent && bIsLeafNode == bIsLeafNodeRest) || (bHasParent && bDebugDrawClustering);
-		}
-		else
-		{
-			bIsVisible = bIsLeafNodeRest;
-		}
-		if (bIsVisible)
-		{
-			const FTransform& Transform = GlobalTransforms[TransformIndex];
-
-			const FColor ActiveColor = bUseActiveVisualization ? MakeDarker(Color, GetLevel(TransformIndex, ParentArray)) : Color;
-
-			const int32 FaceStart = FaceStartArray[GeometryIndex];
-			const int32 FaceEnd = FaceStart + FaceCountArray[GeometryIndex];
-
-			for (int32 FaceIndex = FaceStart; FaceIndex < FaceEnd; ++FaceIndex)
-			{
-				const FVector Vertex0 = VertexArray[IndicesArray[FaceIndex][0]];
-				const FVector Vertex1 = VertexArray[IndicesArray[FaceIndex][1]];
-				const FVector Vertex2 = VertexArray[IndicesArray[FaceIndex][2]];
-
-				const FVector FaceCenter = (Vertex0 + Vertex1 + Vertex2) / 3.f;
-
-				const FVector Position = Transform.TransformPosition(FaceCenter);
-
-				const FString Text = FString::Printf(TEXT("%d"), FaceIndex);
-				AddDebugText(Text, Position, ActiveColor, TextScale, bTextShadow);
-			}
-		}
-	}
-#endif  // #if GEOMETRYCOLLECTION_DEBUG_DRAW && ENABLE_DRAW_DEBUG
-}
-
-void AGeometryCollectionDebugDrawActor::DrawFaceIndices(const TArray<FTransform>& GlobalTransforms, const UGeometryCollectionComponent* GeometryCollectionComponent, int32 TransformIndex, const FColor& Color)
-{
-#if GEOMETRYCOLLECTION_DEBUG_DRAW && ENABLE_DRAW_DEBUG
-	check(GeometryCollectionComponent);
-	check(GeometryCollectionComponent->RestCollection);
-	check(TransformIndex >= 0);
-
-	const TManagedArray<int32>& TransformToGeometryIndexArray = GeometryCollectionComponent->GetTransformToGeometryIndexArray();
-	const int32 GeometryIndex = TransformToGeometryIndexArray[TransformIndex];
-
-	const TManagedArray<TSet<int32>>& ChildrenArrayRest = GeometryCollectionComponent->GetChildrenArrayRest();
-	const bool bIsLeafNodeRest = (ChildrenArrayRest[TransformIndex].Num() == 0);
-
-	// Debug draw this geometry
-	if (GeometryIndex != INDEX_NONE && (bIsLeafNodeRest || bDebugDrawHierarchy))
-	{
-		const UWorld* const World = GetWorld();
-		check(World);
-
-		const FTransform& Transform = GlobalTransforms[TransformIndex];
-
-		const TManagedArray<FVector3f>& VertexArray = GeometryCollectionComponent->GetVertexArray();
-		const TManagedArray<FIntVector>& IndicesArray = GeometryCollectionComponent->GetIndicesArray();
-		const TManagedArray<int32>& FaceStartArray = GeometryCollectionComponent->GetFaceStartArray();
-		const TManagedArray<int32>& FaceCountArray = GeometryCollectionComponent->GetFaceCountArray();
-		const int32 FaceStart = FaceStartArray[GeometryIndex];
-		const int32 FaceEnd = FaceStart + FaceCountArray[GeometryIndex];
-
-		for (int32 FaceIndex = FaceStart; FaceIndex < FaceEnd; ++FaceIndex)
-		{
-			const FVector Vertex0 = VertexArray[IndicesArray[FaceIndex][0]];
-			const FVector Vertex1 = VertexArray[IndicesArray[FaceIndex][1]];
-			const FVector Vertex2 = VertexArray[IndicesArray[FaceIndex][2]];
-			const FVector FaceCenter = (Vertex0 + Vertex1 + Vertex2) / 3.f;
-
-			const FVector Position = Transform.TransformPosition(FaceCenter);
-			const FString Text = FString::Printf(TEXT("%d"), FaceIndex);
-			AddDebugText(Text, Position, Color, TextScale, bTextShadow);
-		}
-	}
-
-	// Debug draw children if the cluster mode is on, or if there is no geometry attached to this node
-	if (!(bIsLeafNodeRest || bDebugDrawHierarchy) || bDebugDrawClustering)
-	{
-		const FColor ActiveColor = bUseActiveVisualization ? MakeDarker(Color) : Color;
-		for (int32 ChildTransformIndex : ChildrenArrayRest[TransformIndex])
-		{
-			DrawFaceIndices(GlobalTransforms, GeometryCollectionComponent, ChildTransformIndex, ActiveColor);
-		}
-	}
-#endif  // #if GEOMETRYCOLLECTION_DEBUG_DRAW && ENABLE_DRAW_DEBUG
-}
-
-void AGeometryCollectionDebugDrawActor::DrawFaceNormals(const TArray<FTransform>& GlobalTransforms, const UGeometryCollectionComponent* GeometryCollectionComponent, const FColor& Color)
-{
-#if GEOMETRYCOLLECTION_DEBUG_DRAW && ENABLE_DRAW_DEBUG
-	check(GeometryCollectionComponent);
-	check(GeometryCollectionComponent->RestCollection);
-
-	const UWorld* const World = GetWorld();
-	check(World);
-
-	const TManagedArray<FVector3f>& VertexArray = GeometryCollectionComponent->GetVertexArray();
-	const TManagedArray<FIntVector>& IndicesArray = GeometryCollectionComponent->GetIndicesArray();
-	const TManagedArray<int32>& TransformIndexArray = GeometryCollectionComponent->GetTransformIndexArray();
-	const TManagedArray<int32>& FaceStartArray = GeometryCollectionComponent->GetFaceStartArray();
-	const TManagedArray<int32>& FaceCountArray = GeometryCollectionComponent->GetFaceCountArray();
-
-	const TManagedArray<int32>& ParentArray = GeometryCollectionComponent->GetParentArray();
-	const TManagedArray<TSet<int32>>& ChildrenArray = GeometryCollectionComponent->GetChildrenArray();
-	const TManagedArray<TSet<int32>>& ChildrenArrayRest = GeometryCollectionComponent->GetChildrenArrayRest();
-
-	const int32 NumGeometries = GeometryCollectionComponent->GetNumElements(FGeometryCollection::GeometryGroup);
-	for (int32 GeometryIndex = 0; GeometryIndex < NumGeometries; ++GeometryIndex)
-	{
-		const int32 TransformIndex = TransformIndexArray[GeometryIndex];
-
-		bool bIsVisible;
-		const bool bIsLeafNodeRest = (ChildrenArrayRest[TransformIndex].Num() == 0);
-		if (bDebugDrawHierarchy)
-		{
-			const bool bIsLeafNode = (ChildrenArray[TransformIndex].Num() == 0);
-			const bool bHasParent = (ParentArray[TransformIndex] != FGeometryCollection::Invalid);
-			bIsVisible = (!bHasParent && bIsLeafNode == bIsLeafNodeRest) || (bHasParent && bDebugDrawClustering);
-		}
-		else
-		{
-			bIsVisible = bIsLeafNodeRest;
-		}
-		if (bIsVisible)
-		{
-			const FTransform& Transform = GlobalTransforms[TransformIndex];
-
-			const FColor ActiveColor = bUseActiveVisualization ? MakeDarker(Color, GetLevel(TransformIndex, ParentArray)) : Color;
-
-			const int32 FaceStart = FaceStartArray[GeometryIndex];
-			const int32 FaceEnd = FaceStart + FaceCountArray[GeometryIndex];
-
-			for (int32 FaceIndex = FaceStart; FaceIndex < FaceEnd; ++FaceIndex)
-			{
-				const FVector Vertex0 = VertexArray[IndicesArray[FaceIndex][0]];
-				const FVector Vertex1 = VertexArray[IndicesArray[FaceIndex][1]];
-				const FVector Vertex2 = VertexArray[IndicesArray[FaceIndex][2]];
-
-				const FVector FaceCenter = (Vertex0 + Vertex1 + Vertex2) / 3.f;
-
-				const FVector Edge1 = Vertex1 - Vertex0;
-				const FVector Edge2 = -(Vertex2 - Vertex1);
-
-				const FVector FaceNormal = Transform.TransformVector(Edge1 ^ Edge2).GetSafeNormal();
-
-				const FVector LineStart = Transform.TransformPosition(FaceCenter);
-				const FVector LineEnd = LineStart + FaceNormal * NormalScale;
-				DrawDebugDirectionalArrow(World, LineStart, LineEnd, ArrowScale, ActiveColor, GeometryCollectionDebugDrawActorConstants::bPersistent, GeometryCollectionDebugDrawActorConstants::LifeTime, GeometryCollectionDebugDrawActorConstants::DepthPriority, LineThickness);
-			}
-			bNeedsDebugLinesFlush = true;
-		}
-	}
-#endif  // #if GEOMETRYCOLLECTION_DEBUG_DRAW && ENABLE_DRAW_DEBUG
-}
-
-void AGeometryCollectionDebugDrawActor::DrawFaceNormals(const TArray<FTransform>& GlobalTransforms, const UGeometryCollectionComponent* GeometryCollectionComponent, int32 TransformIndex, const FColor& Color)
-{
-#if GEOMETRYCOLLECTION_DEBUG_DRAW && ENABLE_DRAW_DEBUG
-	check(GeometryCollectionComponent);
-	check(GeometryCollectionComponent->RestCollection);
-	check(TransformIndex >= 0);
-
-	const UWorld* const World = GetWorld();
-	check(World);
-
-	const TManagedArray<int32>& TransformToGeometryIndexArray = GeometryCollectionComponent->GetTransformToGeometryIndexArray();
-	const int32 GeometryIndex = TransformToGeometryIndexArray[TransformIndex];
-
-	const TManagedArray<TSet<int32>>& ChildrenArrayRest = GeometryCollectionComponent->GetChildrenArrayRest();
-	const bool bIsLeafNodeRest = (ChildrenArrayRest[TransformIndex].Num() == 0);
-
-	// Debug draw this geometry
-	if (GeometryIndex != INDEX_NONE && (bIsLeafNodeRest || bDebugDrawHierarchy))
-	{
-		const FTransform& Transform = GlobalTransforms[TransformIndex];
-
-		const TManagedArray<FVector3f>& VertexArray = GeometryCollectionComponent->GetVertexArray();
-		const TManagedArray<FIntVector>& IndicesArray = GeometryCollectionComponent->GetIndicesArray();
-		const TManagedArray<int32>& FaceStartArray = GeometryCollectionComponent->GetFaceStartArray();
-		const TManagedArray<int32>& FaceCountArray = GeometryCollectionComponent->GetFaceCountArray();
-		const int32 FaceStart = FaceStartArray[GeometryIndex];
-		const int32 FaceEnd = FaceStart + FaceCountArray[GeometryIndex];
-
-		for (int32 FaceIndex = FaceStart; FaceIndex < FaceEnd; ++FaceIndex)
-		{
-			const FVector Vertex0 = VertexArray[IndicesArray[FaceIndex][0]];
-			const FVector Vertex1 = VertexArray[IndicesArray[FaceIndex][1]];
-			const FVector Vertex2 = VertexArray[IndicesArray[FaceIndex][2]];
-
-			const FVector FaceCenter = (Vertex0 + Vertex1 + Vertex2) / 3.f;
-
-			const FVector Edge1 = Vertex1 - Vertex0;
-			const FVector Edge2 = -(Vertex2 - Vertex1);
-
-			const FVector FaceNormal = Transform.TransformVector(Edge1 ^ Edge2).GetSafeNormal();
-
-			const FVector LineStart = Transform.TransformPosition(FaceCenter);
-			const FVector LineEnd = LineStart + FaceNormal * NormalScale;
-			DrawDebugDirectionalArrow(World, LineStart, LineEnd, ArrowScale, Color, GeometryCollectionDebugDrawActorConstants::bPersistent, GeometryCollectionDebugDrawActorConstants::LifeTime, GeometryCollectionDebugDrawActorConstants::DepthPriority, LineThickness);
-		}
-		bNeedsDebugLinesFlush = true;
-	}
-
-	// Debug draw children if the cluster mode is on, or if there is no geometry attached to this node
-	if (!(bIsLeafNodeRest || bDebugDrawHierarchy) || bDebugDrawClustering)
-	{
-		const FColor ActiveColor = bUseActiveVisualization ? MakeDarker(Color) : Color;
-		for (int32 ChildTransformIndex : ChildrenArrayRest[TransformIndex])
-		{
-			DrawFaceNormals(GlobalTransforms, GeometryCollectionComponent, ChildTransformIndex, ActiveColor);
-		}
-	}
-#endif  // #if GEOMETRYCOLLECTION_DEBUG_DRAW && ENABLE_DRAW_DEBUG
-}
-
-void AGeometryCollectionDebugDrawActor::DrawSingleFace(const TArray<FTransform>& GlobalTransforms, const UGeometryCollectionComponent* GeometryCollectionComponent, const int32 FaceIndex, const FColor& Color)
-{
-#if GEOMETRYCOLLECTION_DEBUG_DRAW && ENABLE_DRAW_DEBUG
-	check(GeometryCollectionComponent);
-	check(GeometryCollectionComponent->RestCollection);
-
-	const UWorld* const World = GetWorld();
-	check(World);
-
-	const int32 NumFaces = GeometryCollectionComponent->GetNumElements(FGeometryCollection::FacesGroup);
-	if (FaceIndex < 0 || FaceIndex >= NumFaces) { return; }
-
-	const TManagedArray<FVector3f>& VertexArray = GeometryCollectionComponent->GetVertexArray();
-	const TManagedArray<int32>& BoneMapArray = GeometryCollectionComponent->GetBoneMapArray();
-	const TManagedArray<FIntVector>& IndicesArray = GeometryCollectionComponent->GetIndicesArray();
-
-	const int32 TransformIndex = BoneMapArray[IndicesArray[FaceIndex][0]];
-	const FTransform& Transform = GlobalTransforms[TransformIndex];
-
-	const FVector Vertex0 = Transform.TransformPosition(VertexArray[IndicesArray[FaceIndex][0]]);
-	const FVector Vertex1 = Transform.TransformPosition(VertexArray[IndicesArray[FaceIndex][1]]);
-	const FVector Vertex2 = Transform.TransformPosition(VertexArray[IndicesArray[FaceIndex][2]]);
-
-	DrawDebugLine(World, Vertex0, Vertex1, Color, GeometryCollectionDebugDrawActorConstants::bPersistent, GeometryCollectionDebugDrawActorConstants::LifeTime, GeometryCollectionDebugDrawActorConstants::DepthPriority, LineThickness * 2.f);
-	DrawDebugLine(World, Vertex0, Vertex2, Color, GeometryCollectionDebugDrawActorConstants::bPersistent, GeometryCollectionDebugDrawActorConstants::LifeTime, GeometryCollectionDebugDrawActorConstants::DepthPriority, LineThickness * 2.f);
-	DrawDebugLine(World, Vertex1, Vertex2, Color, GeometryCollectionDebugDrawActorConstants::bPersistent, GeometryCollectionDebugDrawActorConstants::LifeTime, GeometryCollectionDebugDrawActorConstants::DepthPriority, LineThickness * 2.f);
-	bNeedsDebugLinesFlush = true;
-#endif  // #if GEOMETRYCOLLECTION_DEBUG_DRAW && ENABLE_DRAW_DEBUG
-}
-
-void AGeometryCollectionDebugDrawActor::DrawGeometryIndices(const TArray<FTransform>& GlobalTransforms, const UGeometryCollectionComponent* GeometryCollectionComponent, const FColor& Color)
-{
-#if GEOMETRYCOLLECTION_DEBUG_DRAW && ENABLE_DRAW_DEBUG
-	check(GeometryCollectionComponent);
-	check(GeometryCollectionComponent->RestCollection);
-
-	const UWorld* const World = GetWorld();
-	check(World);
-
-	const TManagedArray<int32>& TransformIndexArray = GeometryCollectionComponent->GetTransformIndexArray();
-
-	const TManagedArray<int32>& ParentArray = GeometryCollectionComponent->GetParentArray();
-	const TManagedArray<TSet<int32>>& ChildrenArray = GeometryCollectionComponent->GetChildrenArray();
-	const TManagedArray<TSet<int32>>& ChildrenArrayRest = GeometryCollectionComponent->GetChildrenArrayRest();
-
-	// Iterate though all geometries, and find those who needs to be visualized
-	const int32 NumGeometries = GeometryCollectionComponent->GetNumElements(FGeometryCollection::GeometryGroup);
-	for (int32 GeometryIndex = 0; GeometryIndex < NumGeometries; ++GeometryIndex)
-	{
-		const int32 TransformIndex = TransformIndexArray[GeometryIndex];
-
-		bool bIsVisible;
-		const bool bIsLeafNodeRest = (ChildrenArrayRest[TransformIndex].Num() == 0);
-		if (bDebugDrawHierarchy)
-		{
-			const bool bIsLeafNode = (ChildrenArray[TransformIndex].Num() == 0);
-			const bool bHasParent = (ParentArray[TransformIndex] != FGeometryCollection::Invalid);
-			bIsVisible = (!bHasParent && bIsLeafNode == bIsLeafNodeRest) || (bHasParent && bDebugDrawClustering);
-		}
-		else
-		{
-			bIsVisible = bIsLeafNodeRest;
-		}
-		if (bIsVisible)
-		{
-			const FTransform& Transform = GlobalTransforms[TransformIndex];
-			const FVector Position = Transform.GetLocation();
-
-			const FColor ActiveColor = bUseActiveVisualization ? MakeDarker(Color, GetLevel(TransformIndex, ParentArray)) : Color;
-
-			const FString Text = FString::Printf(TEXT("%d"), GeometryIndex);
-			AddDebugText(Text, Position, ActiveColor, TextScale, bTextShadow);
-		}
-	}
-#endif  // #if GEOMETRYCOLLECTION_DEBUG_DRAW && ENABLE_DRAW_DEBUG
-}
-
-void AGeometryCollectionDebugDrawActor::DrawGeometryIndex(const TArray<FTransform>& GlobalTransforms, const UGeometryCollectionComponent* GeometryCollectionComponent, int32 TransformIndex, const FColor& Color)
-{
-#if GEOMETRYCOLLECTION_DEBUG_DRAW && ENABLE_DRAW_DEBUG
-	check(GeometryCollectionComponent);
-	check(GeometryCollectionComponent->RestCollection);
-	check(TransformIndex >= 0);
-
-	const UWorld* const World = GetWorld();
-	check(World);
-
-	const TManagedArray<int32>& TransformToGeometryIndexArray = GeometryCollectionComponent->GetTransformToGeometryIndexArray();
-	const int32 GeometryIndex = TransformToGeometryIndexArray[TransformIndex];
-
-	const TManagedArray<TSet<int32>>& ChildrenArrayRest = GeometryCollectionComponent->GetChildrenArrayRest();
-	const bool bIsLeafNodeRest = (ChildrenArrayRest[TransformIndex].Num() == 0);
-
-	// Debug draw this geometry
-	if (GeometryIndex != INDEX_NONE && (bIsLeafNodeRest || bDebugDrawHierarchy))
-	{
-		const FTransform& Transform = GlobalTransforms[TransformIndex];
-
-		const FVector Position = Transform.GetLocation();
-		const FString Text = FString::Printf(TEXT("%d"), GeometryIndex);
-		AddDebugText(Text, Position, Color, TextScale, bTextShadow);
-	}
-
-	// Debug draw children if the cluster mode is on, or if there is no geometry attached to this node
-	if (!(bIsLeafNodeRest || bDebugDrawHierarchy) || bDebugDrawClustering)
-	{
-		const FColor ActiveColor = bUseActiveVisualization ? MakeDarker(Color) : Color;
-		for (int32 ChildTransformIndex : ChildrenArrayRest[TransformIndex])
-		{
-			DrawGeometryIndex(GlobalTransforms, GeometryCollectionComponent, ChildTransformIndex, ActiveColor);
-		}
-	}
-#endif  // #if GEOMETRYCOLLECTION_DEBUG_DRAW && ENABLE_DRAW_DEBUG
-}
-
-void AGeometryCollectionDebugDrawActor::DrawTransforms(const TArray<FTransform>& GlobalTransforms, const UGeometryCollectionComponent* GeometryCollectionComponent, float Scale)
-{
-#if GEOMETRYCOLLECTION_DEBUG_DRAW && ENABLE_DRAW_DEBUG
-	Scale *= AxisScale;
-
-	check(GeometryCollectionComponent);
-	check(GeometryCollectionComponent->RestCollection);
-
-	const UWorld* const World = GetWorld();
-	check(World);
-
-	const TManagedArray<int32>& TransformToGeometryIndexArray = GeometryCollectionComponent->GetTransformToGeometryIndexArray();
-
-	const TManagedArray<int32>& ParentArray = GeometryCollectionComponent->GetParentArray();
-	const TManagedArray<TSet<int32>>& ChildrenArray = GeometryCollectionComponent->GetChildrenArray();
-	const TManagedArray<TSet<int32>>& ChildrenArrayRest = GeometryCollectionComponent->GetChildrenArrayRest();
-
-	// Iterate though all transforms, and find those who are geometries and needs to be visualized
-	const int32 NumTransforms = GeometryCollectionComponent->GetNumElements(FGeometryCollection::TransformGroup);
-	for (int32 TransformIndex = 0; TransformIndex < NumTransforms; ++TransformIndex)
-	{
-		bool bIsVisible;
-		const bool bIsLeafNodeRest = (ChildrenArrayRest[TransformIndex].Num() == 0);
-		if (bDebugDrawHierarchy)
-		{
-			const bool bIsLeafNode = (ChildrenArray[TransformIndex].Num() == 0);
-			const bool bHasParent = (ParentArray[TransformIndex] != FGeometryCollection::Invalid);
-			bIsVisible = (!bHasParent && bIsLeafNode == bIsLeafNodeRest) || (bHasParent && bDebugDrawClustering);
-		}
-		else
-		{
-			const int32 GeometryIndex = TransformToGeometryIndexArray[TransformIndex];
-			bIsVisible = (GeometryIndex != INDEX_NONE && bIsLeafNodeRest);
-		}
-		if (bIsVisible)
-		{
-			const FTransform& Transform = GlobalTransforms[TransformIndex];
-			const FVector Position = Transform.GetLocation();
-			const FRotator Rotator = Transform.Rotator();
-
-			const float ActiveScale = bUseActiveVisualization ? MakeSmaller(Scale, GetLevel(TransformIndex, ParentArray)) : Scale;
-
-			DrawDebugCoordinateSystem(World, Position, Rotator, ActiveScale, GeometryCollectionDebugDrawActorConstants::bPersistent, GeometryCollectionDebugDrawActorConstants::LifeTime, GeometryCollectionDebugDrawActorConstants::DepthPriority, LineThickness);
-			bNeedsDebugLinesFlush = true;
-		}
-	}
-#endif  // #if GEOMETRYCOLLECTION_DEBUG_DRAW && ENABLE_DRAW_DEBUG
-}
-
-void AGeometryCollectionDebugDrawActor::DrawTransform(const TArray<FTransform>& GlobalTransforms, const UGeometryCollectionComponent* GeometryCollectionComponent, int32 TransformIndex, float Scale)
-{
-#if GEOMETRYCOLLECTION_DEBUG_DRAW && ENABLE_DRAW_DEBUG
-	check(GeometryCollectionComponent);
-	check(GeometryCollectionComponent->RestCollection);
-	check(TransformIndex >= 0);
-
-	const UWorld* const World = GetWorld();
-	check(World);
-
-	const TManagedArray<TSet<int32>>& ChildrenArrayRest = GeometryCollectionComponent->GetChildrenArrayRest();
-	const bool bIsLeafNodeRest = (ChildrenArrayRest[TransformIndex].Num() == 0);
-
-	// Debug draw this geometry
-	if (bIsLeafNodeRest || bDebugDrawHierarchy)
-	{
-		const FTransform& Transform = GlobalTransforms[TransformIndex];
-
-		const FVector Position = Transform.GetLocation();
-		const FRotator Rotator = Transform.Rotator();
-
-		DrawDebugCoordinateSystem(World, Position, Rotator, Scale * AxisScale, GeometryCollectionDebugDrawActorConstants::bPersistent, GeometryCollectionDebugDrawActorConstants::LifeTime, GeometryCollectionDebugDrawActorConstants::DepthPriority, LineThickness);
-		bNeedsDebugLinesFlush = true;
-	}
-
-	// Debug draw children if the cluster mode is on, or if there is no geometry attached to this node
-	if (!(bIsLeafNodeRest || bDebugDrawHierarchy) || bDebugDrawClustering)
-	{
-		const float ActiveScale = bUseActiveVisualization ? MakeSmaller(Scale) : Scale;
-		for (int32 ChildTransformIndex : ChildrenArrayRest[TransformIndex])
-		{
-			DrawTransform(GlobalTransforms, GeometryCollectionComponent, ChildTransformIndex, ActiveScale);
-		}
-	}
-#endif  // #if GEOMETRYCOLLECTION_DEBUG_DRAW && ENABLE_DRAW_DEBUG
-}
-
-void AGeometryCollectionDebugDrawActor::DrawTransformIndices(const TArray<FTransform>& GlobalTransforms, const UGeometryCollectionComponent* GeometryCollectionComponent, const FColor& Color)
-{
-#if GEOMETRYCOLLECTION_DEBUG_DRAW && ENABLE_DRAW_DEBUG
-	check(GeometryCollectionComponent);
-	check(GeometryCollectionComponent->RestCollection);
-
-	const UWorld* const World = GetWorld();
-	check(World);
-
-	const TManagedArray<int32>& TransformToGeometryIndexArray = GeometryCollectionComponent->GetTransformToGeometryIndexArray();
-
-	const TManagedArray<int32>& ParentArray = GeometryCollectionComponent->GetParentArray();
-	const TManagedArray<TSet<int32>>& ChildrenArray = GeometryCollectionComponent->GetChildrenArray();
-	const TManagedArray<TSet<int32>>& ChildrenArrayRest = GeometryCollectionComponent->GetChildrenArrayRest();
-
-	// Iterate though all transforms, and find those who are geometries and needs to be visualized
-	const int32 NumTransforms = GeometryCollectionComponent->GetNumElements(FGeometryCollection::TransformGroup);
-	for (int32 TransformIndex = 0; TransformIndex < NumTransforms; ++TransformIndex)
-	{
-		bool bIsVisible;
-		const bool bIsLeafNodeRest = (ChildrenArrayRest[TransformIndex].Num() == 0);
-		if (bDebugDrawHierarchy)
-		{
-			const bool bIsLeafNode = (ChildrenArray[TransformIndex].Num() == 0);
-			const bool bHasParent = (ParentArray[TransformIndex] != FGeometryCollection::Invalid);
-			bIsVisible = (!bHasParent && bIsLeafNode == bIsLeafNodeRest) || (bHasParent && bDebugDrawClustering);
-		}
-		else
-		{
-			const int32 GeometryIndex = TransformToGeometryIndexArray[TransformIndex];
-			bIsVisible = (GeometryIndex != INDEX_NONE && bIsLeafNodeRest);
-		}
-		if (bIsVisible)
-		{
-			const FTransform& Transform = GlobalTransforms[TransformIndex];
-			const FVector Position = Transform.GetLocation();
-
-			const FColor ActiveColor = bUseActiveVisualization ? MakeDarker(Color, GetLevel(TransformIndex, ParentArray)) : Color;
-
-			const FString Text = FString::Printf(TEXT("%d"), TransformIndex);
-			AddDebugText(Text, Position, ActiveColor, TextScale, bTextShadow);
-		}
-	}
-#endif  // #if GEOMETRYCOLLECTION_DEBUG_DRAW && ENABLE_DRAW_DEBUG
-}
-
-void AGeometryCollectionDebugDrawActor::DrawTransformIndex(const TArray<FTransform>& GlobalTransforms, const UGeometryCollectionComponent* GeometryCollectionComponent, int32 TransformIndex, const FColor& Color)
-{
-#if GEOMETRYCOLLECTION_DEBUG_DRAW && ENABLE_DRAW_DEBUG
-	check(GeometryCollectionComponent);
-	check(GeometryCollectionComponent->RestCollection);
-	check(TransformIndex >= 0);
-
-	const UWorld* const World = GetWorld();
-	check(World);
-
-	const TManagedArray<TSet<int32>>& ChildrenArrayRest = GeometryCollectionComponent->GetChildrenArrayRest();
-	const bool bIsLeafNodeRest = (ChildrenArrayRest[TransformIndex].Num() == 0);
-
-	// Debug draw this geometry
-	if (bIsLeafNodeRest || bDebugDrawHierarchy)
-	{
-		const FTransform& Transform = GlobalTransforms[TransformIndex];
-
-		const FVector Position = Transform.GetLocation();
-		const FString Text = FString::Printf(TEXT("%d"), TransformIndex);
-		AddDebugText(Text, Position, Color, TextScale, bTextShadow);
-	}
-
-	// Debug draw children if the cluster mode is on, or if there is no geometry attached to this node
-	if (!(bIsLeafNodeRest || bDebugDrawHierarchy) || bDebugDrawClustering)
-	{
-		const FColor ActiveColor = bUseActiveVisualization ? MakeDarker(Color) : Color;
-		for (int32 ChildTransformIndex : ChildrenArrayRest[TransformIndex])
-		{
-			DrawTransformIndex(GlobalTransforms, GeometryCollectionComponent, ChildTransformIndex, ActiveColor);
-		}
-	}
-#endif  // #if GEOMETRYCOLLECTION_DEBUG_DRAW && ENABLE_DRAW_DEBUG
-}
-
-void AGeometryCollectionDebugDrawActor::DrawLevels(const TArray<FTransform>& GlobalTransforms, const UGeometryCollectionComponent* GeometryCollectionComponent, const FColor& Color)
-{
-#if GEOMETRYCOLLECTION_DEBUG_DRAW && ENABLE_DRAW_DEBUG
-	check(GeometryCollectionComponent);
-	check(GeometryCollectionComponent->RestCollection);
-
-	const UWorld* const World = GetWorld();
-	check(World);
-
-	const TManagedArray<int32>& TransformToGeometryIndexArray = GeometryCollectionComponent->GetTransformToGeometryIndexArray();
-
-	const TManagedArray<int32>& ParentArray = GeometryCollectionComponent->GetParentArray();
-	const TManagedArray<TSet<int32>>& ChildrenArray = GeometryCollectionComponent->GetChildrenArray();
-	const TManagedArray<TSet<int32>>& ChildrenArrayRest = GeometryCollectionComponent->GetChildrenArrayRest();
-
-	// Iterate though all transforms, and find those who are geometries and needs to be visualized
-	const int32 NumTransforms = GeometryCollectionComponent->GetNumElements(FGeometryCollection::TransformGroup);
-	for (int32 TransformIndex = 0; TransformIndex < NumTransforms; ++TransformIndex)
-	{
-		bool bIsVisible;
-		const bool bIsLeafNodeRest = (ChildrenArrayRest[TransformIndex].Num() == 0);
-		if (bDebugDrawHierarchy)
-		{
-			const bool bIsLeafNode = (ChildrenArray[TransformIndex].Num() == 0);
-			const bool bHasParent = (ParentArray[TransformIndex] != FGeometryCollection::Invalid);
-			bIsVisible = (!bHasParent && bIsLeafNode == bIsLeafNodeRest) || (bHasParent && bDebugDrawClustering);
-		}
-		else
-		{
-			const int32 GeometryIndex = TransformToGeometryIndexArray[TransformIndex];
-			bIsVisible = (GeometryIndex != INDEX_NONE && bIsLeafNodeRest);
-		}
-		if (bIsVisible)
-		{
-			const FColor ActiveColor = bUseActiveVisualization ? MakeDarker(Color, GetLevel(TransformIndex, ParentArray)) : Color;
-
-			const FTransform& Transform = GlobalTransforms[TransformIndex];
-
-			const FVector Position = Transform.GetLocation();
-			const int32 Level = GetLevel(TransformIndex, ParentArray);
-			const FString Text = FString::Printf(TEXT("%d"), Level);
-			AddDebugText(Text, Position, ActiveColor, TextScale, bTextShadow);
-		}
-	}
-#endif  // #if GEOMETRYCOLLECTION_DEBUG_DRAW && ENABLE_DRAW_DEBUG
-}
-
-void AGeometryCollectionDebugDrawActor::DrawLevel(const TArray<FTransform>& GlobalTransforms, const UGeometryCollectionComponent* GeometryCollectionComponent, int32 TransformIndex, const FColor& Color)
-{
-#if GEOMETRYCOLLECTION_DEBUG_DRAW && ENABLE_DRAW_DEBUG
-	check(GeometryCollectionComponent);
-	check(GeometryCollectionComponent->RestCollection);
-	check(TransformIndex >= 0);
-
-	const UWorld* const World = GetWorld();
-	check(World);
-
-	const TManagedArray<TSet<int32>>& ChildrenArrayRest = GeometryCollectionComponent->GetChildrenArrayRest();
-	const bool bIsLeafNodeRest = (ChildrenArrayRest[TransformIndex].Num() == 0);
-
-	// Debug draw this geometry
-	if (bIsLeafNodeRest || bDebugDrawHierarchy)
-	{
-		const FTransform& Transform = GlobalTransforms[TransformIndex];
-		const FVector Position = Transform.GetLocation();
-
-		const TManagedArray<int32>& ParentArray = GeometryCollectionComponent->GetParentArray();
-		const int32 Level = GetLevel(TransformIndex, ParentArray);
-
-		const FString Text = FString::Printf(TEXT("%d"), Level);
-		AddDebugText(Text, Position, Color, TextScale, bTextShadow);
-	}
-
-	// Debug draw children if the cluster mode is on, or if there is no geometry attached to this node
-	if (!(bIsLeafNodeRest || bDebugDrawHierarchy) || bDebugDrawClustering)
-	{
-		const FColor ActiveColor = bUseActiveVisualization ? MakeDarker(Color) : Color;
-		for (int32 ChildTransformIndex : ChildrenArrayRest[TransformIndex])
-		{
-			DrawLevel(GlobalTransforms, GeometryCollectionComponent, ChildTransformIndex, ActiveColor);
-		}
-	}
-#endif  // #if GEOMETRYCOLLECTION_DEBUG_DRAW && ENABLE_DRAW_DEBUG
-}
-
-void AGeometryCollectionDebugDrawActor::DrawParents(const TArray<FTransform>& GlobalTransforms, const UGeometryCollectionComponent* GeometryCollectionComponent, const FColor& Color)
-{
-#if GEOMETRYCOLLECTION_DEBUG_DRAW && ENABLE_DRAW_DEBUG
-	check(GeometryCollectionComponent);
-	check(GeometryCollectionComponent->RestCollection);
-
-	const UWorld* const World = GetWorld();
-	check(World);
-
-	const TManagedArray<int32>& TransformToGeometryIndexArray = GeometryCollectionComponent->GetTransformToGeometryIndexArray();
-
-	const TManagedArray<int32>& ParentArray = GeometryCollectionComponent->GetParentArray();
-	const TManagedArray<TSet<int32>>& ChildrenArrayRest = GeometryCollectionComponent->GetChildrenArrayRest();
-
-	// Iterate though all transforms, and find those who are geometries and needs to be visualized
-	const int32 NumTransforms = GeometryCollectionComponent->GetNumElements(FGeometryCollection::TransformGroup);
-	for (int32 TransformIndex = 0; TransformIndex < NumTransforms; ++TransformIndex)
-	{
-		const int32 ParentTransformIndex = ParentArray[TransformIndex];
-		const bool bHasParent = (ParentTransformIndex != FGeometryCollection::Invalid);
-		if (bHasParent)
-		{
-			bool bIsVisible;
-			if (bDebugDrawHierarchy)
-			{
-				bIsVisible = bDebugDrawClustering;
-			}
-			else
-			{
-				const int32 GeometryIndex = TransformToGeometryIndexArray[TransformIndex];
-				const bool bIsLeafNodeRest = (ChildrenArrayRest[TransformIndex].Num() == 0);
-				bIsVisible = (GeometryIndex != INDEX_NONE && bIsLeafNodeRest);
-			}
-			if (bIsVisible)
-			{
-				const FTransform& Transform = GlobalTransforms[TransformIndex];
-				const FVector Position = Transform.GetLocation();
-
-				const FTransform ParentTransform = GlobalTransforms[ParentTransformIndex];
-				const FVector ParentPosition = ParentTransform.GetLocation();
-
-				const FColor ActiveColor = bUseActiveVisualization ? MakeDarker(Color, GetLevel(TransformIndex, ParentArray)) : Color;
-
-				DrawDebugLine(World, ParentPosition, Position, ActiveColor, GeometryCollectionDebugDrawActorConstants::bPersistent, GeometryCollectionDebugDrawActorConstants::LifeTime, GeometryCollectionDebugDrawActorConstants::DepthPriority, LineThickness);
-	
-				bNeedsDebugLinesFlush = true;
-			}
-		}
-	}
-#endif  // #if GEOMETRYCOLLECTION_DEBUG_DRAW && ENABLE_DRAW_DEBUG
-}
-
-void AGeometryCollectionDebugDrawActor::DrawParent(const TArray<FTransform>& GlobalTransforms, const UGeometryCollectionComponent* GeometryCollectionComponent, int32 TransformIndex, const FColor& Color)
-{
-#if GEOMETRYCOLLECTION_DEBUG_DRAW && ENABLE_DRAW_DEBUG
-	check(GeometryCollectionComponent);
-	check(GeometryCollectionComponent->RestCollection);
-	check(TransformIndex >= 0);
-
-	const UWorld* const World = GetWorld();
-	check(World);
-
-	const TManagedArray<TSet<int32>>& ChildrenArrayRest = GeometryCollectionComponent->GetChildrenArrayRest();
-	const bool bIsLeafNodeRest = (ChildrenArrayRest[TransformIndex].Num() == 0);
-
-	const TManagedArray<int32>& ParentArray = GeometryCollectionComponent->GetParentArray();
-	const int32 ParentTransformIndex = ParentArray[TransformIndex];
-	const bool bHasParent = (ParentTransformIndex != FGeometryCollection::Invalid);
-
-	// Debug draw this geometry
-	if (bHasParent && (bIsLeafNodeRest || bDebugDrawHierarchy))
-	{
-		const FTransform& Transform = GlobalTransforms[TransformIndex];
-		const FVector Position = Transform.GetLocation();
-
-		const FTransform ParentTransform = GlobalTransforms[ParentTransformIndex];
-		const FVector ParentPosition = ParentTransform.GetLocation();
-
-		const float Scale = ArrowScale * FVector::Dist(ParentPosition, Position);
-
-		DrawDebugLine(World, ParentPosition, Position, Color, GeometryCollectionDebugDrawActorConstants::bPersistent, GeometryCollectionDebugDrawActorConstants::LifeTime, GeometryCollectionDebugDrawActorConstants::DepthPriority, LineThickness);
-
-		bNeedsDebugLinesFlush = true;
-	}
-
-	// Debug draw children if the cluster mode is on, or if there is no geometry attached to this node
-	if (!(bIsLeafNodeRest || bDebugDrawHierarchy) || bDebugDrawClustering)
-	{
-		const FColor ActiveColor = bUseActiveVisualization ? MakeDarker(Color) : Color;
-		for (int32 ChildTransformIndex : ChildrenArrayRest[TransformIndex])
-		{
-			DrawParent(GlobalTransforms, GeometryCollectionComponent, ChildTransformIndex, ActiveColor);
-		}
-	}
-#endif  // #if GEOMETRYCOLLECTION_DEBUG_DRAW && ENABLE_DRAW_DEBUG
-}
-
-void AGeometryCollectionDebugDrawActor::DrawBoundingBoxes(const TArray<FTransform>& GlobalTransforms, const UGeometryCollectionComponent* GeometryCollectionComponent, const FColor& Color)
-{
-#if GEOMETRYCOLLECTION_DEBUG_DRAW && ENABLE_DRAW_DEBUG
-	check(GeometryCollectionComponent);
-	check(GeometryCollectionComponent->RestCollection);
-
-	const UWorld* const World = GetWorld();
-	check(World);
-
-	const TManagedArray<FBox>& BoundingBoxArray = GeometryCollectionComponent->GetBoundingBoxArray();
-	const TManagedArray<int32>& TransformIndexArray = GeometryCollectionComponent->GetTransformIndexArray();
-
-	const TManagedArray<int32>& ParentArray = GeometryCollectionComponent->GetParentArray();
-	const TManagedArray<TSet<int32>>& ChildrenArray = GeometryCollectionComponent->GetChildrenArray();
-	const TManagedArray<TSet<int32>>& ChildrenArrayRest = GeometryCollectionComponent->GetChildrenArrayRest();
-
-	// Iterate though all geometries, and find those who are active and needs to be visualized
-	const int32 NumGeometries = GeometryCollectionComponent->GetNumElements(FGeometryCollection::GeometryGroup);
-	for (int32 GeometryIndex = 0; GeometryIndex < NumGeometries; ++GeometryIndex)
-	{
-		const int32 TransformIndex = TransformIndexArray[GeometryIndex];
-
-		bool bIsVisible;
-		const bool bIsLeafNodeRest = (ChildrenArrayRest[TransformIndex].Num() == 0);
-		if (bDebugDrawHierarchy)
-		{
-			const bool bIsLeafNode = (ChildrenArray[TransformIndex].Num() == 0);
-			const bool bHasParent = (ParentArray[TransformIndex] != FGeometryCollection::Invalid);
-			bIsVisible = (!bHasParent && bIsLeafNode == bIsLeafNodeRest) || (bHasParent && bDebugDrawClustering);
-		}
-		else
-		{
-			bIsVisible = bIsLeafNodeRest;
-		}
-		if (bIsVisible)
-		{
-			const FTransform& Transform = GlobalTransforms[TransformIndex];
-
-			const FBox& BBox = BoundingBoxArray[GeometryIndex];
-			const FVector Vertices[8] =
-			{
-				Transform.TransformPosition(BBox.Min),
-				Transform.TransformPosition(FVector(BBox.Max.X, BBox.Min.Y, BBox.Min.Z)),
-				Transform.TransformPosition(FVector(BBox.Max.X, BBox.Max.Y, BBox.Min.Z)),
-				Transform.TransformPosition(FVector(BBox.Min.X, BBox.Max.Y, BBox.Min.Z)),
-				Transform.TransformPosition(FVector(BBox.Min.X, BBox.Min.Y, BBox.Max.Z)),
-				Transform.TransformPosition(FVector(BBox.Max.X, BBox.Min.Y, BBox.Max.Z)),
-				Transform.TransformPosition(BBox.Max),
-				Transform.TransformPosition(FVector(BBox.Min.X, BBox.Max.Y, BBox.Max.Z))
-			};
-
-			const FColor ActiveColor = bUseActiveVisualization ? MakeDarker(Color, GetLevel(TransformIndex, ParentArray)) : Color;
-
-			for (const auto& BoxEdge : GeometryCollectionDebugDrawActorConstants::BoxEdges)
-			{
-				DrawDebugLine(World, Vertices[BoxEdge(0)], Vertices[BoxEdge(1)], ActiveColor, GeometryCollectionDebugDrawActorConstants::bPersistent, GeometryCollectionDebugDrawActorConstants::LifeTime, GeometryCollectionDebugDrawActorConstants::DepthPriority, LineThickness);
-			}
-
-			bNeedsDebugLinesFlush = true;
-		}
-	}
-#endif  // #if GEOMETRYCOLLECTION_DEBUG_DRAW && ENABLE_DRAW_DEBUG
-}
-
-void AGeometryCollectionDebugDrawActor::DrawBoundingBox(const TArray<FTransform>& GlobalTransforms, const UGeometryCollectionComponent* GeometryCollectionComponent, int32 TransformIndex, const FColor& Color)
 {
-#if GEOMETRYCOLLECTION_DEBUG_DRAW && ENABLE_DRAW_DEBUG
-	check(GeometryCollectionComponent);
-	check(GeometryCollectionComponent->RestCollection);
-	check(TransformIndex >= 0);
-
-	const UWorld* const World = GetWorld();
-	check(World);
-
-	const TManagedArray<int32>& TransformToGeometryIndexArray = GeometryCollectionComponent->GetTransformToGeometryIndexArray();
-	const int32 GeometryIndex = TransformToGeometryIndexArray[TransformIndex];
-
-	const TManagedArray<TSet<int32>>& ChildrenArrayRest = GeometryCollectionComponent->GetChildrenArrayRest();
-	const bool bIsLeafNodeRest = (ChildrenArrayRest[TransformIndex].Num() == 0);
-
-	// Debug draw this geometry
-	if (GeometryIndex != INDEX_NONE && (bIsLeafNodeRest || bDebugDrawHierarchy))
-	{
-		const FTransform& Transform = GlobalTransforms[TransformIndex];
-
-		const TManagedArray<FBox>& BoundingBoxArray = GeometryCollectionComponent->GetBoundingBoxArray();
-		const FBox& BBox = BoundingBoxArray[GeometryIndex];
-		const FVector Vertices[8] =
-		{
-			Transform.TransformPosition(BBox.Min),
-			Transform.TransformPosition(FVector(BBox.Max.X, BBox.Min.Y, BBox.Min.Z)),
-			Transform.TransformPosition(FVector(BBox.Max.X, BBox.Max.Y, BBox.Min.Z)),
-			Transform.TransformPosition(FVector(BBox.Min.X, BBox.Max.Y, BBox.Min.Z)),
-			Transform.TransformPosition(FVector(BBox.Min.X, BBox.Min.Y, BBox.Max.Z)),
-			Transform.TransformPosition(FVector(BBox.Max.X, BBox.Min.Y, BBox.Max.Z)),
-			Transform.TransformPosition(BBox.Max),
-			Transform.TransformPosition(FVector(BBox.Min.X, BBox.Max.Y, BBox.Max.Z))
-		};
-
-		for (const auto& BoxEdge : GeometryCollectionDebugDrawActorConstants::BoxEdges)
-		{
-			DrawDebugLine(World, Vertices[BoxEdge(0)], Vertices[BoxEdge(1)], Color, GeometryCollectionDebugDrawActorConstants::bPersistent, GeometryCollectionDebugDrawActorConstants::LifeTime, GeometryCollectionDebugDrawActorConstants::DepthPriority, LineThickness);
-		}
-
-		bNeedsDebugLinesFlush = true;
-	}
-
-	// Debug draw children if the cluster mode is on, or if there is no geometry attached to this node
-	if (!(bIsLeafNodeRest || bDebugDrawHierarchy) || bDebugDrawClustering)
-	{
-		const FColor ActiveColor = bUseActiveVisualization ? MakeDarker(Color) : Color;
-		for (int32 ChildTransformIndex : ChildrenArrayRest[TransformIndex])
-		{
-			DrawBoundingBox(GlobalTransforms, GeometryCollectionComponent, ChildTransformIndex, ActiveColor);
-		}
-	}
-#endif  // #if GEOMETRYCOLLECTION_DEBUG_DRAW && ENABLE_DRAW_DEBUG
-}
-
-FTransform AGeometryCollectionDebugDrawActor::GetParticleTransform(const UGeometryCollectionComponent* GeometryCollectionComponent, int32 TransformIndex, const FGeometryCollectionParticlesData& ParticlesData)
-{
-#if GEOMETRYCOLLECTION_DEBUG_DRAW && ENABLE_DRAW_DEBUG
-	check(GeometryCollectionComponent);
-	check(GeometryCollectionComponent->RestCollection);
-	check(TransformIndex >= 0);
-
-	// Check particle sync status
-	bool bSynced;
-	bSynced = ParticlesData.RequestSyncedData(EGeometryCollectionParticlesData::X);
-	bSynced = ParticlesData.RequestSyncedData(EGeometryCollectionParticlesData::R) && bSynced;
-	bSynced = ParticlesData.RequestSyncedData(EGeometryCollectionParticlesData::ChildToParentMap) && bSynced;
-	if (bSynced)
-	{
-		return GetParticleTransformNoChecks(GeometryCollectionComponent, TransformIndex, ParticlesData);
-	}
-#endif  // #if GEOMETRYCOLLECTION_DEBUG_DRAW && ENABLE_DRAW_DEBUG
-	return FTransform::Identity;
-}
-
-FTransform AGeometryCollectionDebugDrawActor::GetParticleTransformNoChecks(const UGeometryCollectionComponent* GeometryCollectionComponent, int32 TransformIndex, const FGeometryCollectionParticlesData& ParticlesData)
-{
-	const TManagedArray<int32>& ParentArray = GeometryCollectionComponent->GetParentArray();
-
-	// Retrieve particle transform
-	Chaos::FVec3 Translation = Chaos::FVec3::ZeroVector;
-	Chaos::FRotation3 Rotation = Chaos::FRotation3::Identity;
-#if GEOMETRYCOLLECTION_DEBUG_DRAW && ENABLE_DRAW_DEBUG
-	int32 Index;
-	for (Index = TransformIndex;
-		ParentArray[Index] != FGeometryCollection::Invalid;
-		Index = ParentArray[Index])
-	{
-		const Chaos::FRigidTransform3& ChildToParentTransform = ParticlesData.GetChildToParentMap(Index);
-		Translation = ChildToParentTransform.GetTranslation() + ChildToParentTransform.GetRotation().RotateVector(Translation);
-		Rotation = ChildToParentTransform.GetRotation() * Rotation;
-	}
-	const Chaos::FVec3& RootTranslation = ParticlesData.GetX(Index);
-	const Chaos::FRotation3& RootRotation = ParticlesData.GetR(Index);
-	Translation = RootTranslation + RootRotation.RotateVector(Translation);
-	Rotation = RootRotation * Rotation;
-#endif  // #if GEOMETRYCOLLECTION_DEBUG_DRAW && ENABLE_DRAW_DEBUG
-	return FTransform(Rotation, Translation);
-}
-
-//void AGeometryCollectionDebugDrawActor::DrawRigidBodiesId(const UGeometryCollectionComponent* GeometryCollectionComponent, const FGeometryCollectionParticlesData& ParticlesData, const TManagedArray<int32>& RigidBodyIdArray, const FColor& Color)
-void AGeometryCollectionDebugDrawActor::DrawRigidBodiesId(const UGeometryCollectionComponent* GeometryCollectionComponent, const FGeometryCollectionParticlesData& ParticlesData, const TManagedArray<FGuid>& RigidBodyIdArray, const FColor& Color)
-{
-#if GEOMETRYCOLLECTION_DEBUG_DRAW && ENABLE_DRAW_DEBUG
-	check(GeometryCollectionComponent);
-	check(GeometryCollectionComponent->RestCollection);
-	check(GetWorld());
-
-	// Check rigid body array sync status
-	if (RigidBodyIdArray.Num() == 0) { return; }
-
-	// Check particle sync status
-	bool bSynced;
-	bSynced = ParticlesData.RequestSyncedData(EGeometryCollectionParticlesData::X);
-	bSynced = ParticlesData.RequestSyncedData(EGeometryCollectionParticlesData::R) && bSynced;
-	bSynced = ParticlesData.RequestSyncedData(EGeometryCollectionParticlesData::ChildToParentMap) && bSynced;
-	bSynced = ParticlesData.RequestSyncedData(EGeometryCollectionParticlesData::Disabled) && bSynced;
-	if (bSynced)
-	{
-		const TManagedArray<int32>& ParentArray = GeometryCollectionComponent->GetParentArray();
-		const TManagedArray<TSet<int32>>& ChildrenArray = GeometryCollectionComponent->GetChildrenArray();
-		const TManagedArray<int32>& TransformToGeometryIndexArray = GeometryCollectionComponent->GetTransformToGeometryIndexArray();
-		const TManagedArray<TSet<int32>>& ChildrenArrayRest = GeometryCollectionComponent->GetChildrenArrayRest();
-
-		// Iterate though all transforms, and find those who are active (has no parent, is a leaf node with geometry or has children) and needs to be visualized
-		const int32 NumTransforms = FGenericPlatformMath::Min(GeometryCollectionComponent->GetNumElements(FGeometryCollection::TransformGroup), RigidBodyIdArray.Num());
-		for (int32 TransformIndex = 0; TransformIndex < NumTransforms; ++TransformIndex)
-		{
-			const bool bHasChildren = (ChildrenArray[TransformIndex].Num() > 0);
-			const bool bIsLeafNodeRest = (ChildrenArrayRest[TransformIndex].Num() == 0);
-			const bool bIsGeometry = (TransformToGeometryIndexArray[TransformIndex] != FGeometryCollection::Invalid);
-			const bool bHasParent = (ParentArray[TransformIndex] != FGeometryCollection::Invalid);
-			if ((bHasParent && bDebugDrawClustering) || (!bHasParent && (bHasChildren || (bIsLeafNodeRest && bIsGeometry))))
-			{
-				const FColor ActiveColor = bUseActiveVisualization ? MakeDarker(Color, GetLevel(TransformIndex, ParentArray)) : Color;
-				DrawRigidBodyIdNoChecks(GeometryCollectionComponent, TransformIndex, ParticlesData, RigidBodyIdArray, ActiveColor);
-			}
-		}
-	}
-#endif  // #if GEOMETRYCOLLECTION_DEBUG_DRAW && ENABLE_DRAW_DEBUG
-}
-
-//void AGeometryCollectionDebugDrawActor::DrawRigidBodyId(const UGeometryCollectionComponent* GeometryCollectionComponent, int32 TransformIndex, const FGeometryCollectionParticlesData& ParticlesData, const TManagedArray<int32>& RigidBodyIdArray, const FColor& Color)
-void AGeometryCollectionDebugDrawActor::DrawRigidBodyId(const UGeometryCollectionComponent* GeometryCollectionComponent, int32 TransformIndex, const FGeometryCollectionParticlesData& ParticlesData, const TManagedArray<FGuid>& RigidBodyIdArray, const FColor& Color)
-{
-#if GEOMETRYCOLLECTION_DEBUG_DRAW && ENABLE_DRAW_DEBUG
-	check(GeometryCollectionComponent);
-	check(GeometryCollectionComponent->RestCollection);
-	check(TransformIndex >= 0);
-	check(GetWorld());
-
-	// Request/check sync status
-	bool bSynced;
-	bSynced = ParticlesData.RequestSyncedData(EGeometryCollectionParticlesData::X);
-	bSynced = ParticlesData.RequestSyncedData(EGeometryCollectionParticlesData::R) && bSynced;
-	bSynced = ParticlesData.RequestSyncedData(EGeometryCollectionParticlesData::ChildToParentMap) && bSynced;
-	bSynced = ParticlesData.RequestSyncedData(EGeometryCollectionParticlesData::Disabled) && bSynced;
-	if (bSynced)
-	{
-		DrawRigidBodyIdNoChecks(GeometryCollectionComponent, TransformIndex, ParticlesData, RigidBodyIdArray, Color);
-
-		// Debug draw children if the cluster mode is on
-		if (bDebugDrawClustering)
-		{
-			const FColor ActiveColor = bUseActiveVisualization ? MakeDarker(Color) : Color;
-			const TManagedArray<TSet<int32>>& ChildrenArrayRest = GeometryCollectionComponent->GetChildrenArrayRest();
-			for (int32 ChildTransformIndex : ChildrenArrayRest[TransformIndex])
-			{
-				DrawRigidBodyId(GeometryCollectionComponent, ChildTransformIndex, ParticlesData, RigidBodyIdArray, ActiveColor);
-			}
-		}
-	}
-#endif  // #if GEOMETRYCOLLECTION_DEBUG_DRAW && ENABLE_DRAW_DEBUG
-}
-
-//void AGeometryCollectionDebugDrawActor::DrawRigidBodyIdNoChecks(const UGeometryCollectionComponent* GeometryCollectionComponent, int32 TransformIndex, const FGeometryCollectionParticlesData& ParticlesData, const TManagedArray<int32>& RigidBodyIdArray, const FColor& Color)
-void AGeometryCollectionDebugDrawActor::DrawRigidBodyIdNoChecks(const UGeometryCollectionComponent* GeometryCollectionComponent, int32 TransformIndex, const FGeometryCollectionParticlesData& ParticlesData, const TManagedArray<FGuid>& RigidBodyIdArray, const FColor& Color)
-{
-#if GEOMETRYCOLLECTION_DEBUG_DRAW && ENABLE_DRAW_DEBUG
-	//const int32 RigidBodyId = RigidBodyIdArray[TransformIndex];
-	const FGuid& RigidBodyId = RigidBodyIdArray[TransformIndex];
-
-	// Retrieve particle transform
-	const FTransform Transform = GetParticleTransformNoChecks(GeometryCollectionComponent, TransformIndex, ParticlesData);
-	const FVector Position = Transform.GetTranslation();
-
-	// Retrieve disabled state
-	const bool bIsDisabled = ParticlesData.IsDisabled(TransformIndex);
-	const TManagedArray<int32>& ParentArray = GeometryCollectionComponent->GetParentArray();
-	const FColor& DisabledColor = bIsDisabled ? FColor::Silver: Color;
-
-	// Draw rigid body id
-	//const FString Text = FString::Printf(TEXT("%d"), RigidBodyId);
-	const FString Text = RigidBodyId.ToString();
-	AddDebugText(Text, Position, DisabledColor, TextScale, bTextShadow);
-#endif  // #if GEOMETRYCOLLECTION_DEBUG_DRAW && ENABLE_DRAW_DEBUG
-}
-
-void AGeometryCollectionDebugDrawActor::DrawRigidBodiesTransform(const UGeometryCollectionComponent* GeometryCollectionComponent, const FGeometryCollectionParticlesData& ParticlesData, float Scale)
-{
-#if GEOMETRYCOLLECTION_DEBUG_DRAW && ENABLE_DRAW_DEBUG
-	check(GeometryCollectionComponent);
-	check(GeometryCollectionComponent->RestCollection);
-	check(GetWorld());
-
-	// Request/check sync status
-	bool bSynced;
-	bSynced = ParticlesData.RequestSyncedData(EGeometryCollectionParticlesData::X);
-	bSynced = ParticlesData.RequestSyncedData(EGeometryCollectionParticlesData::R) && bSynced;
-	bSynced = ParticlesData.RequestSyncedData(EGeometryCollectionParticlesData::ChildToParentMap) && bSynced;
-	bSynced = ParticlesData.RequestSyncedData(EGeometryCollectionParticlesData::Disabled) && bSynced;
-	if (bSynced)
-	{
-		const TManagedArray<int32>& ParentArray = GeometryCollectionComponent->GetParentArray();
-		const TManagedArray<TSet<int32>>& ChildrenArray = GeometryCollectionComponent->GetChildrenArray();
-		const TManagedArray<int32>& TransformToGeometryIndexArray = GeometryCollectionComponent->GetTransformToGeometryIndexArray();
-		const TManagedArray<TSet<int32>>& ChildrenArrayRest = GeometryCollectionComponent->GetChildrenArrayRest();
-
-		const int32 NumTransforms = GeometryCollectionComponent->GetNumElements(FGeometryCollection::TransformGroup);
-		for (int32 TransformIndex = 0; TransformIndex < NumTransforms; ++TransformIndex)
-		{
-			const bool bHasChildren = (ChildrenArray[TransformIndex].Num() > 0);
-			const bool bIsLeafNodeRest = (ChildrenArrayRest[TransformIndex].Num() == 0);
-			const bool bIsGeometry = (TransformToGeometryIndexArray[TransformIndex] != FGeometryCollection::Invalid);
-			const bool bHasParent = (ParentArray[TransformIndex] != FGeometryCollection::Invalid);
-			if ((bHasParent && bDebugDrawClustering) || (!bHasParent && (bHasChildren || (bIsLeafNodeRest && bIsGeometry))))
-			{
-				const float ActiveScale = bUseActiveVisualization ? MakeSmaller(Scale, GetLevel(TransformIndex, ParentArray)) : Scale;
-				DrawRigidBodyTransformNoChecks(GeometryCollectionComponent, TransformIndex, ParticlesData, ActiveScale);
-			}
-		}
-	}
-#endif  // #if GEOMETRYCOLLECTION_DEBUG_DRAW && ENABLE_DRAW_DEBUG
-}
-
-void AGeometryCollectionDebugDrawActor::DrawRigidBodyTransform(const UGeometryCollectionComponent* GeometryCollectionComponent, int32 TransformIndex, const FGeometryCollectionParticlesData& ParticlesData, float Scale)
-{
-#if GEOMETRYCOLLECTION_DEBUG_DRAW && ENABLE_DRAW_DEBUG
-	check(GeometryCollectionComponent);
-	check(GeometryCollectionComponent->RestCollection);
-	check(TransformIndex >= 0);
-	check(GetWorld());
-
-	// Request/check sync status
-	bool bSynced;
-	bSynced = ParticlesData.RequestSyncedData(EGeometryCollectionParticlesData::X);
-	bSynced = ParticlesData.RequestSyncedData(EGeometryCollectionParticlesData::R) && bSynced;
-	bSynced = ParticlesData.RequestSyncedData(EGeometryCollectionParticlesData::ChildToParentMap) && bSynced;
-	bSynced = ParticlesData.RequestSyncedData(EGeometryCollectionParticlesData::Disabled) && bSynced;
-	if (bSynced)
-	{
-		DrawRigidBodyTransformNoChecks(GeometryCollectionComponent, TransformIndex, ParticlesData, Scale);
-
-		// Debug draw children if the cluster mode is on
-		if (bDebugDrawClustering)
-		{
-			const float ActiveScale = bUseActiveVisualization ? MakeSmaller(Scale) : Scale;
-			const TManagedArray<TSet<int32>>& ChildrenArrayRest = GeometryCollectionComponent->GetChildrenArrayRest();
-			for (int32 ChildTransformIndex : ChildrenArrayRest[TransformIndex])
-			{
-				DrawRigidBodyTransform(GeometryCollectionComponent, ChildTransformIndex, ParticlesData, ActiveScale);
-			}
-		}
-	}
-#endif  // #if GEOMETRYCOLLECTION_DEBUG_DRAW && ENABLE_DRAW_DEBUG
-}
-
-void AGeometryCollectionDebugDrawActor::DrawRigidBodyTransformNoChecks(const UGeometryCollectionComponent* GeometryCollectionComponent, int32 TransformIndex, const FGeometryCollectionParticlesData& ParticlesData, float Scale)
-{
-#if GEOMETRYCOLLECTION_DEBUG_DRAW && ENABLE_DRAW_DEBUG
-	Scale *= AxisScale;
-
-	// Retrieve particle transform
-	const FTransform Transform = GetParticleTransformNoChecks(GeometryCollectionComponent, TransformIndex, ParticlesData);
-	const FVector Position = Transform.GetTranslation();
-	const FQuat Rotation = Transform.GetRotation();
-
-	// Retrieve disabled status
-	const bool bIsDisabled = ParticlesData.IsDisabled(TransformIndex);
-
-	// Draw transform
-	if (bIsDisabled)
-	{
-		// Only visualize non clustered disabled elements
-		const TManagedArray<int32>& ParentArray = GeometryCollectionComponent->GetParentArray();
-		if (ParentArray[TransformIndex] == FGeometryCollection::Invalid)
-		{
-			const FVector Forward = Position + Rotation.RotateVector(FVector::ForwardVector) * Scale;
-			const FVector Right = Position + Rotation.RotateVector(FVector::RightVector) * Scale;
-			const FVector Up = Position + Rotation.RotateVector(FVector::UpVector) * Scale;
-
-			DrawDebugLine(GetWorld(), Position, Forward, FColor(64, 0, 0), GeometryCollectionDebugDrawActorConstants::bPersistent, GeometryCollectionDebugDrawActorConstants::LifeTime, GeometryCollectionDebugDrawActorConstants::DepthPriority, LineThickness);
-			DrawDebugLine(GetWorld(), Position, Right, FColor(0, 64, 0), GeometryCollectionDebugDrawActorConstants::bPersistent, GeometryCollectionDebugDrawActorConstants::LifeTime, GeometryCollectionDebugDrawActorConstants::DepthPriority, LineThickness);
-			DrawDebugLine(GetWorld(), Position, Up, FColor(0, 0, 64), GeometryCollectionDebugDrawActorConstants::bPersistent, GeometryCollectionDebugDrawActorConstants::LifeTime, GeometryCollectionDebugDrawActorConstants::DepthPriority, LineThickness);
-		}
-		else
-		{
-			DrawDebugPoint(GetWorld(), Position, PointThickness, FColor::Black, GeometryCollectionDebugDrawActorConstants::bPersistent, GeometryCollectionDebugDrawActorConstants::LifeTime, GeometryCollectionDebugDrawActorConstants::DepthPriority);
-		}
-	}
-	else
-	{
-		const FRotator Rotator(Rotation);
-		DrawDebugCoordinateSystem(GetWorld(), Position, Rotator, Scale, GeometryCollectionDebugDrawActorConstants::bPersistent, GeometryCollectionDebugDrawActorConstants::LifeTime, GeometryCollectionDebugDrawActorConstants::DepthPriority, LineThickness);
-	}
-	bNeedsDebugLinesFlush = true;
-#endif  // #if GEOMETRYCOLLECTION_DEBUG_DRAW && ENABLE_DRAW_DEBUG
-}
-
-void AGeometryCollectionDebugDrawActor::DrawRigidBodiesInertia(const UGeometryCollectionComponent* GeometryCollectionComponent, const FGeometryCollectionParticlesData& ParticlesData, const FColor& Color)
-{
-#if GEOMETRYCOLLECTION_DEBUG_DRAW && ENABLE_DRAW_DEBUG
-	check(GeometryCollectionComponent);
-	check(GeometryCollectionComponent->RestCollection);
-	check(GetWorld());
-
-	// Request/check sync status
-	bool bSynced;
-	bSynced = ParticlesData.RequestSyncedData(EGeometryCollectionParticlesData::X);
-	bSynced = ParticlesData.RequestSyncedData(EGeometryCollectionParticlesData::R) && bSynced;
-	bSynced = ParticlesData.RequestSyncedData(EGeometryCollectionParticlesData::ChildToParentMap) && bSynced;
-	bSynced = ParticlesData.RequestSyncedData(EGeometryCollectionParticlesData::Disabled) && bSynced;
-	bSynced = ParticlesData.RequestSyncedData(EGeometryCollectionParticlesData::I) && bSynced;
-	if (bSynced)
-	{
-		const TManagedArray<int32>& ParentArray = GeometryCollectionComponent->GetParentArray();
-		const TManagedArray<TSet<int32>>& ChildrenArray = GeometryCollectionComponent->GetChildrenArray();
-		const TManagedArray<int32>& TransformToGeometryIndexArray = GeometryCollectionComponent->GetTransformToGeometryIndexArray();
-		const TManagedArray<TSet<int32>>& ChildrenArrayRest = GeometryCollectionComponent->GetChildrenArrayRest();
-
-		const int32 NumTransforms = GeometryCollectionComponent->GetNumElements(FGeometryCollection::TransformGroup);
-		for (int32 TransformIndex = 0; TransformIndex < NumTransforms; ++TransformIndex)
-		{
-			const bool bHasChildren = (ChildrenArray[TransformIndex].Num() > 0);
-			const bool bIsLeafNodeRest = (ChildrenArrayRest[TransformIndex].Num() == 0);
-			const bool bIsGeometry = (TransformToGeometryIndexArray[TransformIndex] != FGeometryCollection::Invalid);
-			const bool bHasParent = (ParentArray[TransformIndex] != FGeometryCollection::Invalid);
-			if ((bHasParent && bDebugDrawClustering) || (!bHasParent && (bHasChildren || (bIsLeafNodeRest && bIsGeometry))))
-			{
-				const FColor ActiveColor = bUseActiveVisualization ? MakeDarker(Color, GetLevel(TransformIndex, ParentArray)) : Color;
-				DrawRigidBodyInertiaNoChecks(GeometryCollectionComponent, TransformIndex, ParticlesData, ActiveColor);
-			}
-		}
-	}
-#endif  // #if GEOMETRYCOLLECTION_DEBUG_DRAW && ENABLE_DRAW_DEBUG
-}
-
-void AGeometryCollectionDebugDrawActor::DrawRigidBodyInertia(const UGeometryCollectionComponent* GeometryCollectionComponent, int32 TransformIndex, const FGeometryCollectionParticlesData& ParticlesData, const FColor& Color)
-{
-#if GEOMETRYCOLLECTION_DEBUG_DRAW && ENABLE_DRAW_DEBUG
-	check(GeometryCollectionComponent);
-	check(GeometryCollectionComponent->RestCollection);
-	check(TransformIndex >= 0);
-	check(GetWorld());
-
-	// Request/check sync status
-	bool bSynced;
-	bSynced = ParticlesData.RequestSyncedData(EGeometryCollectionParticlesData::X);
-	bSynced = ParticlesData.RequestSyncedData(EGeometryCollectionParticlesData::R) && bSynced;
-	bSynced = ParticlesData.RequestSyncedData(EGeometryCollectionParticlesData::ChildToParentMap) && bSynced;
-	bSynced = ParticlesData.RequestSyncedData(EGeometryCollectionParticlesData::Disabled) && bSynced;
-	bSynced = ParticlesData.RequestSyncedData(EGeometryCollectionParticlesData::I) && bSynced;
-	if (bSynced)
-	{
-		DrawRigidBodyInertiaNoChecks(GeometryCollectionComponent, TransformIndex, ParticlesData, Color);
-
-		// Debug draw children if the cluster mode is on
-		if (bDebugDrawClustering)
-		{
-			const FColor ActiveColor = bUseActiveVisualization ? MakeDarker(Color) : Color;
-			const TManagedArray<TSet<int32>>& ChildrenArrayRest = GeometryCollectionComponent->GetChildrenArrayRest();
-			for (int32 ChildTransformIndex : ChildrenArrayRest[TransformIndex])
-			{
-				DrawRigidBodyInertia(GeometryCollectionComponent, ChildTransformIndex, ParticlesData, ActiveColor);
-			}
-		}
-	}
-#endif  // #if GEOMETRYCOLLECTION_DEBUG_DRAW && ENABLE_DRAW_DEBUG
-}
-
-void AGeometryCollectionDebugDrawActor::DrawRigidBodyInertiaNoChecks(const UGeometryCollectionComponent* GeometryCollectionComponent, int32 TransformIndex, const FGeometryCollectionParticlesData& ParticlesData, const FColor& Color)
-{
-#if GEOMETRYCOLLECTION_DEBUG_DRAW && ENABLE_DRAW_DEBUG
-	// Retrieve particle transform
-	const FTransform Transform = GetParticleTransformNoChecks(GeometryCollectionComponent, TransformIndex, ParticlesData);
-	const FVector Position = Transform.GetTranslation();
-	const FQuat Rotation = Transform.GetRotation();
-
-	// Retrieve disabled status
-	const bool bIsDisabled = ParticlesData.IsDisabled(TransformIndex);
-	const FColor ActiveColor = bIsDisabled ? FColor::Black : Color;
-
-	// Draw transform
-	const Chaos::FMatrix33& Inertia = ParticlesData.GetI(TransformIndex);
-	const FVector Side(
-		FMath::Sqrt(6.f * Inertia.M[1][1] + 6.f * Inertia.M[2][2] - 6.f * Inertia.M[0][0]),
-		FMath::Sqrt(6.f * Inertia.M[0][0] + 6.f * Inertia.M[2][2] - 6.f * Inertia.M[1][1]),
-		FMath::Sqrt(6.f * Inertia.M[0][0] + 6.f * Inertia.M[1][1] - 6.f * Inertia.M[2][2]));
-
-	const FVector VertexMin = Side * -0.5f;
-	const FVector VertexMax = Side * 0.5f;
-	const FVector Vertices[8] = 
-	{
-		Position + Rotation.RotateVector(VertexMin),
-		Position + Rotation.RotateVector(FVector(VertexMax.X, VertexMin.Y, VertexMin.Z)),
-		Position + Rotation.RotateVector(FVector(VertexMax.X, VertexMax.Y, VertexMin.Z)),
-		Position + Rotation.RotateVector(FVector(VertexMin.X, VertexMax.Y, VertexMin.Z)),
-		Position + Rotation.RotateVector(FVector(VertexMin.X, VertexMin.Y, VertexMax.Z)),
-		Position + Rotation.RotateVector(FVector(VertexMax.X, VertexMin.Y, VertexMax.Z)),
-		Position + Rotation.RotateVector(VertexMax),
-		Position + Rotation.RotateVector(FVector(VertexMin.X, VertexMax.Y, VertexMax.Z))
-	};
-	for (const auto& BoxEdge : GeometryCollectionDebugDrawActorConstants::BoxEdges)
-	{
-		DrawDebugLine(GetWorld(), Vertices[BoxEdge(0)], Vertices[BoxEdge(1)], ActiveColor, GeometryCollectionDebugDrawActorConstants::bPersistent, GeometryCollectionDebugDrawActorConstants::LifeTime, GeometryCollectionDebugDrawActorConstants::DepthPriority, LineThickness);
-	}
-	bNeedsDebugLinesFlush = true;
-#endif  // #if GEOMETRYCOLLECTION_DEBUG_DRAW && ENABLE_DRAW_DEBUG
-}
-
-void AGeometryCollectionDebugDrawActor::DrawRigidBodiesCollision(const UGeometryCollectionComponent* GeometryCollectionComponent, const FGeometryCollectionParticlesData& ParticlesData, const FColor& Color)
-{
-#if GEOMETRYCOLLECTION_DEBUG_DRAW && ENABLE_DRAW_DEBUG
-	check(GeometryCollectionComponent);
-	check(GeometryCollectionComponent->RestCollection);
-	check(GetWorld());
-
-	// Request/check sync status
-	bool bSynced;
-	bSynced = ParticlesData.RequestSyncedData(EGeometryCollectionParticlesData::X);
-	bSynced = ParticlesData.RequestSyncedData(EGeometryCollectionParticlesData::R) && bSynced;
-	bSynced = ParticlesData.RequestSyncedData(EGeometryCollectionParticlesData::ChildToParentMap) && bSynced;
-	bSynced = ParticlesData.RequestSyncedData(EGeometryCollectionParticlesData::GeometryType) && bSynced;
-	bSynced = ParticlesData.RequestSyncedData(EGeometryCollectionParticlesData::GeometryBoxMin) && bSynced;
-	bSynced = ParticlesData.RequestSyncedData(EGeometryCollectionParticlesData::GeometryBoxMax) && bSynced;
-	bSynced = ParticlesData.RequestSyncedData(EGeometryCollectionParticlesData::GeometrySphereCenter) && bSynced;
-	bSynced = ParticlesData.RequestSyncedData(EGeometryCollectionParticlesData::GeometrySphereRadius) && bSynced;
-	bSynced = ParticlesData.RequestSyncedData(EGeometryCollectionParticlesData::Disabled) && bSynced;
-	if (bSynced)
-	{
-		const TManagedArray<int32>& ParentArray = GeometryCollectionComponent->GetParentArray();
-		const TManagedArray<TSet<int32>>& ChildrenArray = GeometryCollectionComponent->GetChildrenArray();
-		const TManagedArray<int32>& TransformToGeometryIndexArray = GeometryCollectionComponent->GetTransformToGeometryIndexArray();
-		const TManagedArray<TSet<int32>>& ChildrenArrayRest = GeometryCollectionComponent->GetChildrenArrayRest();
-
-		const int32 NumTransforms = GeometryCollectionComponent->GetNumElements(FGeometryCollection::TransformGroup);
-		for (int32 TransformIndex = 0; TransformIndex < NumTransforms; ++TransformIndex)
-		{
-			const bool bHasChildren = (ChildrenArray[TransformIndex].Num() > 0);
-			const bool bIsLeafNodeRest = (ChildrenArrayRest[TransformIndex].Num() == 0);
-			const bool bIsGeometry = (TransformToGeometryIndexArray[TransformIndex] != FGeometryCollection::Invalid);
-			const bool bHasParent = (ParentArray[TransformIndex] != FGeometryCollection::Invalid);
-			const bool bHasParentUnion = (bHasParent && ParticlesData.GetGeometryType(ParentArray[TransformIndex]) == Chaos::ImplicitObjectType::Union);
-			if ((bHasParent && bDebugDrawClustering) || (!bHasParent && (bHasChildren || (bIsLeafNodeRest && bIsGeometry))) || bHasParentUnion)
-			{
-				const FColor ActiveColor = bUseActiveVisualization ? MakeDarker(Color, GetLevel(TransformIndex, ParentArray)) : Color;
-				DrawRigidBodyCollisionNoChecks(GeometryCollectionComponent, TransformIndex, ParticlesData, ActiveColor);
-			}
-		}
-	}
-#endif  // #if GEOMETRYCOLLECTION_DEBUG_DRAW && ENABLE_DRAW_DEBUG
-}
-
-void AGeometryCollectionDebugDrawActor::DrawRigidBodyCollision(const UGeometryCollectionComponent* GeometryCollectionComponent, int32 TransformIndex, const FGeometryCollectionParticlesData& ParticlesData, const FColor& Color)
-{
-#if GEOMETRYCOLLECTION_DEBUG_DRAW && ENABLE_DRAW_DEBUG
-	check(GeometryCollectionComponent);
-	check(GeometryCollectionComponent->RestCollection);
-	check(TransformIndex >= 0);
-	check(GetWorld());
-
-	// Request/check sync status
-	bool bSynced;
-	bSynced = ParticlesData.RequestSyncedData(EGeometryCollectionParticlesData::X);
-	bSynced = ParticlesData.RequestSyncedData(EGeometryCollectionParticlesData::R) && bSynced;
-	bSynced = ParticlesData.RequestSyncedData(EGeometryCollectionParticlesData::ChildToParentMap) && bSynced;
-	bSynced = ParticlesData.RequestSyncedData(EGeometryCollectionParticlesData::GeometryType) && bSynced;
-	bSynced = ParticlesData.RequestSyncedData(EGeometryCollectionParticlesData::GeometryBoxMin) && bSynced;
-	bSynced = ParticlesData.RequestSyncedData(EGeometryCollectionParticlesData::GeometryBoxMax) && bSynced;
-	bSynced = ParticlesData.RequestSyncedData(EGeometryCollectionParticlesData::GeometrySphereCenter) && bSynced;
-	bSynced = ParticlesData.RequestSyncedData(EGeometryCollectionParticlesData::GeometrySphereRadius) && bSynced;
-	bSynced = ParticlesData.RequestSyncedData(EGeometryCollectionParticlesData::Disabled) && bSynced;
-	if (bSynced)
-	{
-		DrawRigidBodyCollisionNoChecks(GeometryCollectionComponent, TransformIndex, ParticlesData, Color);
-
-		if (bDebugDrawClustering)
-		{
-			// Debug draw all children
-			const FColor ActiveColor = bUseActiveVisualization ? MakeDarker(Color) : Color;
-			const TManagedArray<TSet<int32>>& ChildrenArrayRest = GeometryCollectionComponent->GetChildrenArrayRest();
-			for (int32 ChildTransformIndex : ChildrenArrayRest[TransformIndex])
-			{
-				DrawRigidBodyCollision(GeometryCollectionComponent, ChildTransformIndex, ParticlesData, ActiveColor);
-			}
-		}
-		else if (ParticlesData.GetGeometryType(TransformIndex) == Chaos::ImplicitObjectType::Union)
-		{
-			// Only debug draw children that are still attached
-			const FColor ActiveColor = bUseActiveVisualization ? MakeDarker(Color) : Color;
-			const TManagedArray<TSet<int32>>& ChildrenArray = GeometryCollectionComponent->GetChildrenArray();
-			for (int32 ChildTransformIndex : ChildrenArray[TransformIndex])
-			{
-				DrawRigidBodyCollision(GeometryCollectionComponent, ChildTransformIndex, ParticlesData, ActiveColor);
-			}
-		}
-	}
-#endif  // #if GEOMETRYCOLLECTION_DEBUG_DRAW && ENABLE_DRAW_DEBUG
-}
-
-void AGeometryCollectionDebugDrawActor::DrawRigidBodyCollisionNoChecks(const UGeometryCollectionComponent* GeometryCollectionComponent, int32 TransformIndex, const FGeometryCollectionParticlesData& ParticlesData, const FColor& Color)
-{
-#if GEOMETRYCOLLECTION_DEBUG_DRAW && ENABLE_DRAW_DEBUG
-	const UWorld* const World = GetWorld();
-
-	// Retrieve particle transform
-	FVector Position;
-	FQuat Rotation;
-	if (bCollisionAtOrigin)
-	{
-		Position = FVector::ZeroVector;
-		Rotation = FQuat::Identity;
-	}
-	else
-	{
-		// Retrieve particle transform
-		const FTransform Transform = GetParticleTransformNoChecks(GeometryCollectionComponent, TransformIndex, ParticlesData);
-		Position = Transform.GetTranslation();
-		Rotation = Transform.GetRotation();
-	}
-
-	// Set active color depending on parent type and disabled state
-	const TManagedArray<int32>& ParentArray = GeometryCollectionComponent->GetParentArray();
-	const int32 ParentTransformIndex = ParentArray[TransformIndex];
-	const bool bIsParentUnion = (ParentTransformIndex != FGeometryCollection::Invalid &&
-		ParticlesData.GetGeometryType(ParentTransformIndex) == Chaos::ImplicitObjectType::Union);
-	const bool bIsDisabled = ParticlesData.IsDisabled(TransformIndex);
-	const FColor& ActiveColor = (bIsDisabled && !bIsParentUnion) ? FColor::Black: Color;
-
-	// Draw collision volume
-	const Chaos::EImplicitObjectType GeometryType = ParticlesData.GetGeometryType(TransformIndex);
-	switch (GeometryType)
-	{
-	case Chaos::ImplicitObjectType::Box:
-		{
-			const FVector VertexMin = ParticlesData.GetGeometryBoxMin(TransformIndex);
-			const FVector VertexMax = ParticlesData.GetGeometryBoxMax(TransformIndex);
-			const FVector Vertices[8] = 
-			{
-				Position + Rotation.RotateVector(VertexMin),
-				Position + Rotation.RotateVector(FVector(VertexMax.X, VertexMin.Y, VertexMin.Z)),
-				Position + Rotation.RotateVector(FVector(VertexMax.X, VertexMax.Y, VertexMin.Z)),
-				Position + Rotation.RotateVector(FVector(VertexMin.X, VertexMax.Y, VertexMin.Z)),
-				Position + Rotation.RotateVector(FVector(VertexMin.X, VertexMin.Y, VertexMax.Z)),
-				Position + Rotation.RotateVector(FVector(VertexMax.X, VertexMin.Y, VertexMax.Z)),
-				Position + Rotation.RotateVector(VertexMax),
-				Position + Rotation.RotateVector(FVector(VertexMin.X, VertexMax.Y, VertexMax.Z))
-			};
-			for (const auto& BoxEdge : GeometryCollectionDebugDrawActorConstants::BoxEdges)
-			{
-				DrawDebugLine(World, Vertices[BoxEdge(0)], Vertices[BoxEdge(1)], ActiveColor, GeometryCollectionDebugDrawActorConstants::bPersistent, GeometryCollectionDebugDrawActorConstants::LifeTime, GeometryCollectionDebugDrawActorConstants::DepthPriority, LineThickness);
-			}
-		}
-		break;
-	case Chaos::ImplicitObjectType::Sphere:
-		{
-			const FVector Center = Position + Rotation.RotateVector(ParticlesData.GetGeometrySphereCenter(TransformIndex));
-			const float Radius = ParticlesData.GetGeometrySphereRadius(TransformIndex);
-			DrawDebugSphere(World, Center, Radius, GeometryCollectionDebugDrawActorConstants::CircleSegments, ActiveColor, GeometryCollectionDebugDrawActorConstants::bPersistent, GeometryCollectionDebugDrawActorConstants::LifeTime, GeometryCollectionDebugDrawActorConstants::DepthPriority, LineThickness);
-		}
-		break;
-	default:
-		DrawDebugPoint(World, Position, PointThickness, ActiveColor, GeometryCollectionDebugDrawActorConstants::bPersistent, GeometryCollectionDebugDrawActorConstants::LifeTime, GeometryCollectionDebugDrawActorConstants::DepthPriority);
-		break;
-	}
-	bNeedsDebugLinesFlush = true;
-#endif  // #if GEOMETRYCOLLECTION_DEBUG_DRAW && ENABLE_DRAW_DEBUG
-}
-
-void AGeometryCollectionDebugDrawActor::DrawRigidBodiesInfo(const UGeometryCollectionComponent* GeometryCollectionComponent, const FGeometryCollectionParticlesData& ParticlesData, const FColor& Color)
-{
-#if GEOMETRYCOLLECTION_DEBUG_DRAW && ENABLE_DRAW_DEBUG
-	check(GeometryCollectionComponent);
-	check(GeometryCollectionComponent->RestCollection);
-	check(GetWorld());
-
-	// Request all data sync
-	ParticlesData.SetAllDataSyncFlag();
-
-	// Check sync status (only need to check position, since ParticlesData.ToString() returns whatever data has been synced
-	if (ParticlesData.HasSyncedData(EGeometryCollectionParticlesData::X) &&
-		ParticlesData.HasSyncedData(EGeometryCollectionParticlesData::R) &&
-		ParticlesData.HasSyncedData(EGeometryCollectionParticlesData::ChildToParentMap))
-	{
-		const TManagedArray<int32>& ParentArray = GeometryCollectionComponent->GetParentArray();
-		const TManagedArray<TSet<int32>>& ChildrenArray = GeometryCollectionComponent->GetChildrenArray();
-		const TManagedArray<int32>& TransformToGeometryIndexArray = GeometryCollectionComponent->GetTransformToGeometryIndexArray();
-		const TManagedArray<TSet<int32>>& ChildrenArrayRest = GeometryCollectionComponent->GetChildrenArrayRest();
-
-		const int32 NumTransforms = GeometryCollectionComponent->GetNumElements(FGeometryCollection::TransformGroup);
-		for (int32 TransformIndex = 0; TransformIndex < NumTransforms; ++TransformIndex)
-		{
-			const bool bHasChildren = (ChildrenArray[TransformIndex].Num() > 0);
-			const bool bIsLeafNodeRest = (ChildrenArrayRest[TransformIndex].Num() == 0);
-			const bool bIsGeometry = (TransformToGeometryIndexArray[TransformIndex] != FGeometryCollection::Invalid);
-			const bool bHasParent = (ParentArray[TransformIndex] != FGeometryCollection::Invalid);
-			if ((bHasParent && bDebugDrawClustering) || (!bHasParent && (bHasChildren || (bIsLeafNodeRest && bIsGeometry))))
-			{
-				const FColor ActiveColor = bUseActiveVisualization ? MakeDarker(Color, GetLevel(TransformIndex, ParentArray)) : Color;
-				DrawRigidBodyInfoNoChecks(GeometryCollectionComponent, TransformIndex, ParticlesData, ActiveColor);
-			}
-		}
-	}
-#endif  // #if GEOMETRYCOLLECTION_DEBUG_DRAW && ENABLE_DRAW_DEBUG
-}
-
-void AGeometryCollectionDebugDrawActor::DrawRigidBodyInfo(const UGeometryCollectionComponent* GeometryCollectionComponent, int32 TransformIndex, const FGeometryCollectionParticlesData& ParticlesData, const FColor& Color)
-{
-#if GEOMETRYCOLLECTION_DEBUG_DRAW && ENABLE_DRAW_DEBUG
-	check(GeometryCollectionComponent);
-	check(GeometryCollectionComponent->RestCollection);
-	check(TransformIndex >= 0);
-	check(GetWorld());
-
-	// Only visualize non clustered elements
-	const TManagedArray<int32>& ParentArray = GeometryCollectionComponent->GetParentArray();
-	if (ParentArray[TransformIndex] == FGeometryCollection::Invalid)
-	{
-		// Request all data sync
-		ParticlesData.SetAllDataSyncFlag();
-
-		// Check sync status (only need to check position, since ParticlesData.ToString() returns whatever data has been synced
-		if (ParticlesData.HasSyncedData(EGeometryCollectionParticlesData::X) &&
-			ParticlesData.HasSyncedData(EGeometryCollectionParticlesData::R) &&
-			ParticlesData.HasSyncedData(EGeometryCollectionParticlesData::ChildToParentMap))
-		{
-			DrawRigidBodyInfoNoChecks(GeometryCollectionComponent, TransformIndex, ParticlesData, Color);
-
-			// Debug draw children if the cluster mode is on
-			if (bDebugDrawClustering)
-			{
-				const FColor ActiveColor = bUseActiveVisualization ? MakeDarker(Color) : Color;
-				const TManagedArray<TSet<int32>>& ChildrenArrayRest = GeometryCollectionComponent->GetChildrenArrayRest();
-				for (int32 ChildTransformIndex : ChildrenArrayRest[TransformIndex])
-				{
-					DrawRigidBodyInfo(GeometryCollectionComponent, ChildTransformIndex, ParticlesData, ActiveColor);
-				}
-			}
-		}
-	}
-#endif  // #if GEOMETRYCOLLECTION_DEBUG_DRAW && ENABLE_DRAW_DEBUG
-}
-
-void AGeometryCollectionDebugDrawActor::DrawRigidBodyInfoNoChecks(const UGeometryCollectionComponent* GeometryCollectionComponent, int32 TransformIndex, const FGeometryCollectionParticlesData& ParticlesData, const FColor& Color)
-{
-#if GEOMETRYCOLLECTION_DEBUG_DRAW && ENABLE_DRAW_DEBUG
-	// Retrieve particle information
-	const FString Infos = TEXT("\n") + ParticlesData.ToString(TransformIndex, TEXT("\n"));  // First line is skipped for particle Id
-
-	// Retrieve particle transform
-	const FTransform Transform = GetParticleTransformNoChecks(GeometryCollectionComponent, TransformIndex, ParticlesData);
-	const FVector Position = Transform.GetTranslation();
-
-	// Draw string
-	AddDebugText(Infos, Position, Color, TextScale, bTextShadow);
-#endif  // #if GEOMETRYCOLLECTION_DEBUG_DRAW && ENABLE_DRAW_DEBUG
-}
-
-//void AGeometryCollectionDebugDrawActor::DrawConnectivityEdges(const UGeometryCollectionComponent* GeometryCollectionComponent, const FGeometryCollectionParticlesData& ParticlesData, const TManagedArray<int32>& RigidBodyIdArray)
-void AGeometryCollectionDebugDrawActor::DrawConnectivityEdges(const UGeometryCollectionComponent* GeometryCollectionComponent, const FGeometryCollectionParticlesData& ParticlesData, const TManagedArray<FGuid>& RigidBodyIdArray)
-{
-#if GEOMETRYCOLLECTION_DEBUG_DRAW && ENABLE_DRAW_DEBUG
-	check(GeometryCollectionComponent);
-	check(GeometryCollectionComponent->RestCollection);
-	check(GetWorld());
-
-	// Request/check sync status
-	bool bSynced;
-	bSynced = ParticlesData.RequestSyncedData(EGeometryCollectionParticlesData::X);
-	bSynced = ParticlesData.RequestSyncedData(EGeometryCollectionParticlesData::R) && bSynced;
-	bSynced = ParticlesData.RequestSyncedData(EGeometryCollectionParticlesData::ChildToParentMap) && bSynced;
-	bSynced = ParticlesData.RequestSyncedData(EGeometryCollectionParticlesData::ConnectivityEdges) && bSynced;
-	if (bSynced)
-	{
-		const TManagedArray<int32>& ParentArray = GeometryCollectionComponent->GetParentArray();
-
-		uint8 Hue = 0;
-
-		const int32 NumTransforms = GeometryCollectionComponent->GetNumElements(FGeometryCollection::TransformGroup);
-		for (int32 TransformIndex = 0; TransformIndex < NumTransforms; ++TransformIndex)
-		{
-			const bool bHasParent = (ParentArray[TransformIndex] != FGeometryCollection::Invalid);
-			if (bHasParent)  // Root nodes aren't clustered
-			{
-				Hue = (Hue + 157) % 256;  // 157 is a prime number that gives a good spread of colors without getting too similar as a rand might do.
-				const FColor RandomColor = FLinearColor::MakeFromHSV8(Hue, 160, 128).ToFColor(true);
-				const FColor ActiveColor = bUseActiveVisualization ? MakeDarker(RandomColor, GetLevel(TransformIndex, ParentArray)) : RandomColor;
-				DrawConnectivityEdgesNoChecks(GeometryCollectionComponent, TransformIndex, ParticlesData, RigidBodyIdArray, ActiveColor);
-			}
-		}
-	}
-#endif  // #if GEOMETRYCOLLECTION_DEBUG_DRAW && ENABLE_DRAW_DEBUG
-}
-
-//void AGeometryCollectionDebugDrawActor::DrawConnectivityEdges(const UGeometryCollectionComponent* GeometryCollectionComponent, int32 TransformIndex, const FGeometryCollectionParticlesData& ParticlesData, const TManagedArray<int32>& RigidBodyIdArray, FColor HSVColor)
-void AGeometryCollectionDebugDrawActor::DrawConnectivityEdges(const UGeometryCollectionComponent* GeometryCollectionComponent, int32 TransformIndex, const FGeometryCollectionParticlesData& ParticlesData, const TManagedArray<FGuid>& RigidBodyIdArray, FColor HSVColor)
-{
-#if GEOMETRYCOLLECTION_DEBUG_DRAW && ENABLE_DRAW_DEBUG
-	check(GeometryCollectionComponent);
-	check(GeometryCollectionComponent->RestCollection);
-	check(TransformIndex >= 0);
-	check(GetWorld());
-
-	// Request/check sync status
-	bool bSynced;
-	bSynced = ParticlesData.RequestSyncedData(EGeometryCollectionParticlesData::X);
-	bSynced = ParticlesData.RequestSyncedData(EGeometryCollectionParticlesData::R) && bSynced;
-	bSynced = ParticlesData.RequestSyncedData(EGeometryCollectionParticlesData::ChildToParentMap) && bSynced;
-	bSynced = ParticlesData.RequestSyncedData(EGeometryCollectionParticlesData::ConnectivityEdges) && bSynced;
-	if (bSynced)
-	{
-		// Debug draw connectivity edges if it has a parent
-		const TManagedArray<int32>& ParentArray = GeometryCollectionComponent->GetParentArray();
-		const bool bHasParent = (ParentArray[TransformIndex] != FGeometryCollection::Invalid);
-		if (bHasParent)
-		{
-			const FLinearColor LinearColor = FLinearColor::MakeFromHSV8(HSVColor.R, HSVColor.G, HSVColor.B);  // HSV stored as RGB values
-			DrawConnectivityEdgesNoChecks(GeometryCollectionComponent, TransformIndex, ParticlesData, RigidBodyIdArray, LinearColor.ToFColor(true));
-		}
-
-		// Debug draw children if the cluster mode is on
-		if (bDebugDrawClustering)
-		{
-			if (bUseActiveVisualization) { HSVColor.B /= 2; }  // HSV stored as RGB values, this makes the color darker
-			const TManagedArray<TSet<int32>>& ChildrenArrayRest = GeometryCollectionComponent->GetChildrenArrayRest();
-			for (int32 ChildTransformIndex : ChildrenArrayRest[TransformIndex])
-			{
-				HSVColor.R = (HSVColor.R + 157) % 256;   // HSV stored as RGB values, this moves to the next "random" hue
-				DrawConnectivityEdges(GeometryCollectionComponent, ChildTransformIndex, ParticlesData, RigidBodyIdArray, HSVColor);
-			}
-		}
-	}
-#endif  // #if GEOMETRYCOLLECTION_DEBUG_DRAW && ENABLE_DRAW_DEBUG
-}
-
-//void AGeometryCollectionDebugDrawActor::DrawConnectivityEdgesNoChecks(const UGeometryCollectionComponent* GeometryCollectionComponent, int32 TransformIndex, const FGeometryCollectionParticlesData& ParticlesData, const TManagedArray<int32>& RigidBodyIdArray, const FColor& Color)
-void AGeometryCollectionDebugDrawActor::DrawConnectivityEdgesNoChecks(const UGeometryCollectionComponent* GeometryCollectionComponent, int32 TransformIndex, const FGeometryCollectionParticlesData& ParticlesData, const TManagedArray<FGuid>& RigidBodyIdArray, const FColor& Color)
-{
-#if GEOMETRYCOLLECTION_DEBUG_DRAW && ENABLE_DRAW_DEBUG
-#if TODO_REIMPLEMENT_RIGID_CLUSTERING
-	const UWorld* const World = GetWorld();
-
-	// Get parent index
-	const TManagedArray<int32>& ParentArray = GeometryCollectionComponent->GetParentArray();
-	const TManagedArray<TSet<int32>>& ChildrenArray = GeometryCollectionComponent->GetChildrenArray();
-	check(ParentArray[TransformIndex] != FGeometryCollection::Invalid);
-
-	// Retrieve mass to local transform so to draw edges from local origin rather than from particle location
-	const TManagedArray<FTransform>& MassToLocalArray = GeometryCollectionComponent->RestCollection->GetGeometryCollection()->GetAttribute<FTransform>("MassToLocal", FTransformCollection::TransformGroup);
-
-	// Retrieve parent particle transform
-	const FTransform Transform = GetParticleTransformNoChecks(GeometryCollectionComponent, ParentArray[TransformIndex], ParticlesData);
-	const FVector ParentPosition = Transform.GetTranslation();
-	const FQuat ParentRotation = Transform.GetRotation();
-
-	// Retrieve local transform
-	const Chaos::TRigidTransform<float, 3>& ChildToParentMap = ParticlesData.GetChildToParentMap(TransformIndex);
-	const FVector Position = !ChildrenArray[TransformIndex].Num() ?
-		ParentPosition + ParentRotation.RotateVector(ChildToParentMap.TransformPositionNoScale(-MassToLocalArray[TransformIndex].GetLocation())):
-		ParentPosition + ParentRotation.RotateVector(ChildToParentMap.GetTranslation());
-
-	// Retrieve connectivity edges information
-	const TArray<Chaos::TConnectivityEdge<float>>& ConnectivityEdges = ParticlesData.GetConnectivityEdges(TransformIndex);
-
-	// Edge thickness
-	const float Thickness = ConnectivityEdgeThickness * LineThickness;
-
-	// Draw connectivity information
-	for (const auto& ConnectivityEdge: ConnectivityEdges)
-	{
-		// Retrieve the sibling's transform index in the geometry collection
-		const int32 SiblingId = ConnectivityEdge.Sibling;
-		int32 SiblingTransformIndex = FGeometryCollection::Invalid;
-		for (int32 i = 0; i < RigidBodyIdArray.Num(); ++i)
-		{
-			if (RigidBodyIdArray[i] == SiblingId)
-			{
-				SiblingTransformIndex = i;
-				break;
-			}
-		}
-
-		// Draw connection
-		if (SiblingTransformIndex != FGeometryCollection::Invalid)
-		{
-			// Retrieve local transform for sibling
-			const Chaos::TRigidTransform<float, 3>& SiblingToParentMap = ParticlesData.GetChildToParentMap(SiblingTransformIndex);
-			const FVector SiblingPosition = !ChildrenArray[SiblingTransformIndex].Num() ?
-				ParentPosition + ParentRotation.RotateVector(SiblingToParentMap.TransformPositionNoScale(-MassToLocalArray[SiblingTransformIndex].GetLocation())):
-				ParentPosition + ParentRotation.RotateVector(SiblingToParentMap.GetTranslation());
-
-			// Draw half line
-			const FVector HalfPosition = (Position + SiblingPosition) * 0.5f;
-			DrawDebugLine(World, Position, HalfPosition, Color, GeometryCollectionDebugDrawActorConstants::bPersistent, GeometryCollectionDebugDrawActorConstants::LifeTime, GeometryCollectionDebugDrawActorConstants::DepthPriority, Thickness);
-			DrawDebugPoint(World, SiblingPosition, PointThickness, Color, GeometryCollectionDebugDrawActorConstants::bPersistent, GeometryCollectionDebugDrawActorConstants::LifeTime, GeometryCollectionDebugDrawActorConstants::DepthPriority);
-			bNeedsDebugLinesFlush = true;
-		}
-	}
-#endif // TODO_REIMPLEMENT_RIGID_CLUSTERING
-#endif  // #if GEOMETRYCOLLECTION_DEBUG_DRAW && ENABLE_DRAW_DEBUG
-}
-
-void AGeometryCollectionDebugDrawActor::DrawRigidBodiesVelocity(const UGeometryCollectionComponent* GeometryCollectionComponent, const FGeometryCollectionParticlesData& ParticlesData, const FColor& Color)
-{
-#if GEOMETRYCOLLECTION_DEBUG_DRAW && ENABLE_DRAW_DEBUG
-	check(GeometryCollectionComponent);
-	check(GeometryCollectionComponent->RestCollection);
-	check(GetWorld());
-
-	// Request/check sync status
-	bool bSynced;
-	bSynced = ParticlesData.RequestSyncedData(EGeometryCollectionParticlesData::X);
-	bSynced = ParticlesData.RequestSyncedData(EGeometryCollectionParticlesData::R);
-	bSynced = ParticlesData.RequestSyncedData(EGeometryCollectionParticlesData::ChildToParentMap) && bSynced;
-	bSynced = ParticlesData.RequestSyncedData(EGeometryCollectionParticlesData::V) && bSynced;
-	bSynced = ParticlesData.RequestSyncedData(EGeometryCollectionParticlesData::W) && bSynced;
-	bSynced = ParticlesData.RequestSyncedData(EGeometryCollectionParticlesData::Disabled) && bSynced;
-	if (bSynced)
-	{
-		const TManagedArray<int32>& ParentArray = GeometryCollectionComponent->GetParentArray();
-		const TManagedArray<TSet<int32>>& ChildrenArray = GeometryCollectionComponent->GetChildrenArray();
-		const TManagedArray<int32>& TransformToGeometryIndexArray = GeometryCollectionComponent->GetTransformToGeometryIndexArray();
-		const TManagedArray<TSet<int32>>& ChildrenArrayRest = GeometryCollectionComponent->GetChildrenArrayRest();
-
-		const int32 NumTransforms = GeometryCollectionComponent->GetNumElements(FGeometryCollection::TransformGroup);
-		for (int32 TransformIndex = 0; TransformIndex < NumTransforms; ++TransformIndex)
-		{
-			const bool bHasChildren = (ChildrenArray[TransformIndex].Num() > 0);
-			const bool bIsLeafNodeRest = (ChildrenArrayRest[TransformIndex].Num() == 0);
-			const bool bIsGeometry = (TransformToGeometryIndexArray[TransformIndex] != FGeometryCollection::Invalid);
-			const bool bHasParent = (ParentArray[TransformIndex] != FGeometryCollection::Invalid);
-			if ((bHasParent && bDebugDrawClustering) || (!bHasParent && (bHasChildren || (bIsLeafNodeRest && bIsGeometry))))
-			{
-				const FColor ActiveColor = bUseActiveVisualization ? MakeDarker(Color, GetLevel(TransformIndex, ParentArray)) : Color;
-				DrawRigidBodyVelocityNoChecks(GeometryCollectionComponent, TransformIndex, ParticlesData, ActiveColor);
-			}
-		}
-	}
-#endif  // #if GEOMETRYCOLLECTION_DEBUG_DRAW && ENABLE_DRAW_DEBUG
-}
-
-void AGeometryCollectionDebugDrawActor::DrawRigidBodyVelocity(const UGeometryCollectionComponent* GeometryCollectionComponent, int32 TransformIndex, const FGeometryCollectionParticlesData& ParticlesData, const FColor& Color)
-{
-#if GEOMETRYCOLLECTION_DEBUG_DRAW && ENABLE_DRAW_DEBUG
-	check(GeometryCollectionComponent);
-	check(GeometryCollectionComponent->RestCollection);
-	check(TransformIndex >= 0);
-	check(GetWorld());
-
-	// Only visualize non clustered elements
-	const TManagedArray<int32>& ParentArray = GeometryCollectionComponent->GetParentArray();
-	if (ParentArray[TransformIndex] == FGeometryCollection::Invalid)
-	{
-		// Request/check sync status
-		bool bSynced;
-		bSynced = ParticlesData.RequestSyncedData(EGeometryCollectionParticlesData::X);
-		bSynced = ParticlesData.RequestSyncedData(EGeometryCollectionParticlesData::R);
-		bSynced = ParticlesData.RequestSyncedData(EGeometryCollectionParticlesData::ChildToParentMap) && bSynced;
-		bSynced = ParticlesData.RequestSyncedData(EGeometryCollectionParticlesData::V) && bSynced;
-		bSynced = ParticlesData.RequestSyncedData(EGeometryCollectionParticlesData::W) && bSynced;
-		bSynced = ParticlesData.RequestSyncedData(EGeometryCollectionParticlesData::Disabled) && bSynced;
-		if (bSynced)
-		{
-			DrawRigidBodyVelocityNoChecks(GeometryCollectionComponent, TransformIndex, ParticlesData, Color);
-
-			// Debug draw children if the cluster mode is on
-			if (bDebugDrawClustering)
-			{
-				const FColor ActiveColor = bUseActiveVisualization ? MakeDarker(Color) : Color;
-				const TManagedArray<TSet<int32>>& ChildrenArrayRest = GeometryCollectionComponent->GetChildrenArrayRest();
-				for (int32 ChildTransformIndex : ChildrenArrayRest[TransformIndex])
-				{
-					DrawRigidBodyVelocity(GeometryCollectionComponent, ChildTransformIndex, ParticlesData, ActiveColor);
-				}
-			}
-		}
-	}
-#endif  // #if GEOMETRYCOLLECTION_DEBUG_DRAW && ENABLE_DRAW_DEBUG
-}
-
-void AGeometryCollectionDebugDrawActor::DrawRigidBodyVelocityNoChecks(const UGeometryCollectionComponent* GeometryCollectionComponent, int32 TransformIndex, const FGeometryCollectionParticlesData& ParticlesData, const FColor& Color)
-{
-#if GEOMETRYCOLLECTION_DEBUG_DRAW && ENABLE_DRAW_DEBUG
-	// Retrieve particle transform
-	const FTransform Transform = GetParticleTransformNoChecks(GeometryCollectionComponent, TransformIndex, ParticlesData);
-	const FVector Position = Transform.GetTranslation();
-
-	// Retrieve disabled state
-	const bool bIsDisabled = ParticlesData.IsDisabled(TransformIndex);
-	const FColor& ActiveColor = bIsDisabled ? FColor::Black: Color;
-
-	// Retrieve particle velocities
-	const FVector& LinearVelocity = ParticlesData.GetV(TransformIndex);
-	const FVector& AngularVelocity  = ParticlesData.GetW(TransformIndex);
-
-	// Draw position
-	DrawDebugPoint(GetWorld(), Position, PointThickness, ActiveColor, GeometryCollectionDebugDrawActorConstants::bPersistent, GeometryCollectionDebugDrawActorConstants::LifeTime, GeometryCollectionDebugDrawActorConstants::DepthPriority);
-
-	// Draw linear velocity
-	const FVector LinearEnd = Position + LinearVelocity;
-	const float Scale = ArrowScale * LinearVelocity.Size();
-	DrawDebugDirectionalArrow(GetWorld(), Position, LinearEnd, Scale, ActiveColor, GeometryCollectionDebugDrawActorConstants::bPersistent, GeometryCollectionDebugDrawActorConstants::LifeTime, GeometryCollectionDebugDrawActorConstants::DepthPriority, LineThickness);
-
-	// Draw angular velocity
-	const float Radius = AngularVelocity.Size();
-	if (Radius > KINDA_SMALL_NUMBER)
-	{
-		FVector YAxis, ZAxis;
-		AngularVelocity.GetUnsafeNormal().FindBestAxisVectors(YAxis, ZAxis);
-		const FVector AngularEnd = Position + AngularVelocity;
-		DrawDebugLine(GetWorld(), Position, AngularEnd, ActiveColor, GeometryCollectionDebugDrawActorConstants::bPersistent, GeometryCollectionDebugDrawActorConstants::LifeTime, GeometryCollectionDebugDrawActorConstants::DepthPriority, LineThickness);
-		DrawDebugCircle(GetWorld(), AngularEnd, Radius, GeometryCollectionDebugDrawActorConstants::CircleSegments, ActiveColor, GeometryCollectionDebugDrawActorConstants::bPersistent, GeometryCollectionDebugDrawActorConstants::LifeTime, GeometryCollectionDebugDrawActorConstants::DepthPriority, LineThickness, YAxis, ZAxis, false);
-	}
-	bNeedsDebugLinesFlush = true;
-#endif  // #if GEOMETRYCOLLECTION_DEBUG_DRAW && ENABLE_DRAW_DEBUG
-}
-
-void AGeometryCollectionDebugDrawActor::DrawRigidBodiesForce(const UGeometryCollectionComponent* GeometryCollectionComponent, const FGeometryCollectionParticlesData& ParticlesData, const FColor& Color)
-{
-#if GEOMETRYCOLLECTION_DEBUG_DRAW && ENABLE_DRAW_DEBUG
-	check(GeometryCollectionComponent);
-	check(GeometryCollectionComponent->RestCollection);
-	check(GetWorld());
-
-	// Request/check sync status
-	bool bSynced;
-	bSynced = ParticlesData.RequestSyncedData(EGeometryCollectionParticlesData::X);
-	bSynced = ParticlesData.RequestSyncedData(EGeometryCollectionParticlesData::R);
-	bSynced = ParticlesData.RequestSyncedData(EGeometryCollectionParticlesData::ChildToParentMap) && bSynced;
-	bSynced = ParticlesData.RequestSyncedData(EGeometryCollectionParticlesData::F) && bSynced;
-	bSynced = ParticlesData.RequestSyncedData(EGeometryCollectionParticlesData::Torque) && bSynced;
-	bSynced = ParticlesData.RequestSyncedData(EGeometryCollectionParticlesData::Disabled) && bSynced;
-	if (bSynced)
-	{
-		const TManagedArray<int32>& ParentArray = GeometryCollectionComponent->GetParentArray();
-		const TManagedArray<TSet<int32>>& ChildrenArray = GeometryCollectionComponent->GetChildrenArray();
-		const TManagedArray<int32>& TransformToGeometryIndexArray = GeometryCollectionComponent->GetTransformToGeometryIndexArray();
-		const TManagedArray<TSet<int32>>& ChildrenArrayRest = GeometryCollectionComponent->GetChildrenArrayRest();
-
-		const int32 NumTransforms = GeometryCollectionComponent->GetNumElements(FGeometryCollection::TransformGroup);
-		for (int32 TransformIndex = 0; TransformIndex < NumTransforms; ++TransformIndex)
-		{
-			const bool bHasChildren = (ChildrenArray[TransformIndex].Num() > 0);
-			const bool bIsLeafNodeRest = (ChildrenArrayRest[TransformIndex].Num() == 0);
-			const bool bIsGeometry = (TransformToGeometryIndexArray[TransformIndex] != FGeometryCollection::Invalid);
-			const bool bHasParent = (ParentArray[TransformIndex] != FGeometryCollection::Invalid);
-			if ((bHasParent && bDebugDrawClustering) || (!bHasParent && (bHasChildren || (bIsLeafNodeRest && bIsGeometry))))
-			{
-				const FColor ActiveColor = bUseActiveVisualization ? MakeDarker(Color, GetLevel(TransformIndex, ParentArray)) : Color;
-				DrawRigidBodyForceNoChecks(GeometryCollectionComponent, TransformIndex, ParticlesData, ActiveColor);
-			}
-		}
-	}
-#endif  // #if GEOMETRYCOLLECTION_DEBUG_DRAW && ENABLE_DRAW_DEBUG
-}
-
-void AGeometryCollectionDebugDrawActor::DrawRigidBodyForce(const UGeometryCollectionComponent* GeometryCollectionComponent, int32 TransformIndex, const FGeometryCollectionParticlesData& ParticlesData, const FColor& Color)
-{
-#if GEOMETRYCOLLECTION_DEBUG_DRAW && ENABLE_DRAW_DEBUG
-	check(GeometryCollectionComponent);
-	check(GeometryCollectionComponent->RestCollection);
-	check(TransformIndex >= 0);
-	check(GetWorld());
-
-	// Only visualize non clustered elements
-	const TManagedArray<int32>& ParentArray = GeometryCollectionComponent->GetParentArray();
-	if (ParentArray[TransformIndex] == FGeometryCollection::Invalid)
-	{
-		// Request/check sync status
-		bool bSynced;
-		bSynced = ParticlesData.RequestSyncedData(EGeometryCollectionParticlesData::X);
-		bSynced = ParticlesData.RequestSyncedData(EGeometryCollectionParticlesData::R);
-		bSynced = ParticlesData.RequestSyncedData(EGeometryCollectionParticlesData::ChildToParentMap) && bSynced;
-		bSynced = ParticlesData.RequestSyncedData(EGeometryCollectionParticlesData::F) && bSynced;
-		bSynced = ParticlesData.RequestSyncedData(EGeometryCollectionParticlesData::Torque) && bSynced;
-		bSynced = ParticlesData.RequestSyncedData(EGeometryCollectionParticlesData::Disabled) && bSynced;
-		if (bSynced)
-		{
-			DrawRigidBodyForceNoChecks(GeometryCollectionComponent, TransformIndex, ParticlesData, Color);
-
-			// Debug draw children if the cluster mode is on
-			if (bDebugDrawClustering)
-			{
-				const FColor ActiveColor = bUseActiveVisualization ? MakeDarker(Color) : Color;
-				const TManagedArray<TSet<int32>>& ChildrenArrayRest = GeometryCollectionComponent->GetChildrenArrayRest();
-				for (int32 ChildTransformIndex : ChildrenArrayRest[TransformIndex])
-				{
-					DrawRigidBodyForce(GeometryCollectionComponent, ChildTransformIndex, ParticlesData, ActiveColor);
-				}
-			}
-		}
-	}
-#endif  // #if GEOMETRYCOLLECTION_DEBUG_DRAW && ENABLE_DRAW_DEBUG
-}
-
-void AGeometryCollectionDebugDrawActor::DrawRigidBodyForceNoChecks(const UGeometryCollectionComponent* GeometryCollectionComponent, int32 TransformIndex, const FGeometryCollectionParticlesData& ParticlesData, const FColor& Color)
-{
-#if GEOMETRYCOLLECTION_DEBUG_DRAW && ENABLE_DRAW_DEBUG
-	// Retrieve particle transform
-	const FTransform Transform = GetParticleTransformNoChecks(GeometryCollectionComponent, TransformIndex, ParticlesData);
-	const FVector Position = Transform.GetTranslation();
-
-	// Retrieve particle information
-	const FVector& Force = ParticlesData.GetF(TransformIndex);
-	const FVector& Torque  = ParticlesData.GetTorque(TransformIndex);
-
-	// Retrieve disabled state
-	const bool bIsDisabled = ParticlesData.IsDisabled(TransformIndex);
-	const FColor& ActiveColor = bIsDisabled ? FColor::Black: Color;
-
-	// Draw position
-	DrawDebugPoint(GetWorld(), Position, PointThickness, ActiveColor, GeometryCollectionDebugDrawActorConstants::bPersistent, GeometryCollectionDebugDrawActorConstants::LifeTime, GeometryCollectionDebugDrawActorConstants::DepthPriority);
-
-	// Draw linear velocity
-	const FVector LinearEnd = Position + Force;
-	const float Scale = ArrowScale * Force.Size();
-	DrawDebugDirectionalArrow(GetWorld(), Position, LinearEnd, Scale, ActiveColor, GeometryCollectionDebugDrawActorConstants::bPersistent, GeometryCollectionDebugDrawActorConstants::LifeTime, GeometryCollectionDebugDrawActorConstants::DepthPriority, LineThickness);
-
-	// Draw angular velocity
-	const float Radius = Torque.Size();
-	if (Radius > KINDA_SMALL_NUMBER)
-	{
-		FVector YAxis, ZAxis;
-		Torque.GetUnsafeNormal().FindBestAxisVectors(YAxis, ZAxis);
-		const FVector AngularEnd = Position + Torque;
-		DrawDebugLine(GetWorld(), Position, AngularEnd, ActiveColor, GeometryCollectionDebugDrawActorConstants::bPersistent, GeometryCollectionDebugDrawActorConstants::LifeTime, GeometryCollectionDebugDrawActorConstants::DepthPriority, LineThickness);
-		DrawDebugCircle(GetWorld(), AngularEnd, Radius, GeometryCollectionDebugDrawActorConstants::CircleSegments, ActiveColor, GeometryCollectionDebugDrawActorConstants::bPersistent, GeometryCollectionDebugDrawActorConstants::LifeTime, GeometryCollectionDebugDrawActorConstants::DepthPriority, LineThickness, YAxis, ZAxis, false);
-	}
-	bNeedsDebugLinesFlush = true;
-#endif  // #if GEOMETRYCOLLECTION_DEBUG_DRAW && ENABLE_DRAW_DEBUG
+	PrimaryActorTick.bCanEverTick = false;
+	PrimaryActorTick.bTickEvenWhenPaused = false;
 }
+PRAGMA_ENABLE_DEPRECATION_WARNINGS
\ No newline at end of file
diff --git a/Engine/Source/Runtime/Experimental/GeometryCollectionEngine/Private/GeometryCollection/GeometryCollectionDebugDrawComponent.cpp b/Engine/Source/Runtime/Experimental/GeometryCollectionEngine/Private/GeometryCollection/GeometryCollectionDebugDrawComponent.cpp
--- a/Engine/Source/Runtime/Experimental/GeometryCollectionEngine/Private/GeometryCollection/GeometryCollectionDebugDrawComponent.cpp
+++ b/Engine/Source/Runtime/Experimental/GeometryCollectionEngine/Private/GeometryCollection/GeometryCollectionDebugDrawComponent.cpp
@@ -15,891 +15,15 @@
 
 DEFINE_LOG_CATEGORY_STATIC(LogGeometryCollectionDebugDraw, Log, All);
 
-// Static member variables
-UGeometryCollectionDebugDrawComponent* UGeometryCollectionDebugDrawComponent::RenderLevelSetOwner = nullptr;
-int32 UGeometryCollectionDebugDrawComponent::LastRenderedId = INDEX_NONE;
-//FGuid UGeometryCollectionDebugDrawComponent::LastRenderedId = FGuid();
-
 UGeometryCollectionDebugDrawComponent::UGeometryCollectionDebugDrawComponent(const FObjectInitializer& ObjectInitializer)
 	: Super(ObjectInitializer)
-	, GeometryCollectionDebugDrawActor(nullptr)
+	, GeometryCollectionDebugDrawActor_DEPRECATED(nullptr)
 	, GeometryCollectionRenderLevelSetActor(nullptr)
 	, GeometryCollectionComponent(nullptr)
-#if GEOMETRYCOLLECTION_DEBUG_DRAW
-	, ParticlesData()
-	, ParentCheckSum(0)
-	, SelectedRigidBodyId(INDEX_NONE)
-	//, SelectedRigidBodyId(FGuid())
-	, SelectedTransformIndex(INDEX_NONE)
-	, HiddenTransformIndex(INDEX_NONE)
-	, bWasVisible(true)
-	, bHasIncompleteRigidBodyIdSync(false)
-	, SelectedChaosSolver(nullptr)
-#endif  // #if GEOMETRYCOLLECTION_DEBUG_DRAW
 {
 	bNavigationRelevant = false;
 	bTickInEditor = false;
-#if GEOMETRYCOLLECTION_DEBUG_DRAW
-	PrimaryComponentTick.bCanEverTick = true;
-	PrimaryComponentTick.bTickEvenWhenPaused = true;
-	PrimaryComponentTick.bStartWithTickEnabled = false;
-#else  // #if GEOMETRYCOLLECTION_DEBUG_DRAW
 	PrimaryComponentTick.bCanEverTick = false;
-#endif  // #if GEOMETRYCOLLECTION_DEBUG_DRAW #else
-}
-
-void UGeometryCollectionDebugDrawComponent::BeginPlay()
-{
-	Super::BeginPlay();
-
-#if GEOMETRYCOLLECTION_DEBUG_DRAW
-	// Should always start with ticking disabled
-	check(!PrimaryComponentTick.IsTickFunctionEnabled());
-
-	if (GeometryCollectionComponent)
-	{
-		// Reset all index variables
-		SelectedRigidBodyId = INDEX_NONE;
-		//SelectedRigidBodyId = FGuid();
-		SelectedTransformIndex = INDEX_NONE;
-		HiddenTransformIndex = INDEX_NONE;
-		bWasVisible = true;
-
-		// Find or create global debug draw actor
-		GeometryCollectionDebugDrawActor = AGeometryCollectionDebugDrawActor::FindOrCreate(GetWorld());
-		if (ensure(GeometryCollectionDebugDrawActor))
-		{
-			// Make sure to tick the debug draw actor first
-			// It is required to clear up the persistent lines before drawing a new frame
-			AddTickPrerequisiteActor(GeometryCollectionDebugDrawActor);
-		}
-
-		// Update the visibility and tick status depending on the debug draw properties currently selected
-		OnDebugDrawPropertiesChanged(false);
-
-		// Find or create level set renderer
-		GeometryCollectionRenderLevelSetActor = AGeometryCollectionRenderLevelSetActor::FindOrCreate(GetWorld());
-	}
-#endif  // #if GEOMETRYCOLLECTION_DEBUG_DRAW
-}
-
-void UGeometryCollectionDebugDrawComponent::EndPlay(EEndPlayReason::Type ReasonEnd)
-{
-#if GEOMETRYCOLLECTION_DEBUG_DRAW
-	if (GeometryCollectionComponent)
-	{
-		// Garbage collect debug draw singleton actor (needs this before updating properties to restore visibility)
-		GeometryCollectionDebugDrawActor = nullptr;
-
-		// Turn off level set renderer (needs this before updating properties to disable the level set renderer)
-		if (RenderLevelSetOwner == this)
-		{
-			LastRenderedId = INDEX_NONE;
-			//LastRenderedId = FGuid();
-		}
-
-		// Refresh states from end-play properties
-		OnDebugDrawPropertiesChanged(false);
-
-		// Garbage collect levelset rendering actor (needs this after updating properties)
-		GeometryCollectionRenderLevelSetActor = nullptr;
-	}
-#endif  // #if GEOMETRYCOLLECTION_DEBUG_DRAW
-	Super::EndPlay(ReasonEnd);
-}
-
-void UGeometryCollectionDebugDrawComponent::TickComponent(float DeltaTime, enum ELevelTick TickType, FActorComponentTickFunction *ThisTickFunction)
-{
-	Super::TickComponent(DeltaTime, TickType, ThisTickFunction);
-	UE_CLOG(GetOwner(), LogGeometryCollectionDebugDraw, VeryVerbose, TEXT("Component ticked for actor %s."), *GetOwner()->GetName());
-
-#if GEOMETRYCOLLECTION_DEBUG_DRAW
-	if (GeometryCollectionComponent && GeometryCollectionDebugDrawActor)
-	{
-		if (GeometryCollectionComponent->RestCollection)
-		{
-			// Debug draw collection
-			DebugDrawTick();
-
-			if (GeometryCollectionRenderLevelSetActor)
-			{
-				// Debug draw solver infos for this collection
-				DebugDrawChaosTick();
-			}
-
-			// Detect breaking by tracking changes in parents so that the geometry visibility can be updated if required
-			// Note: The GeometryCollectionComponent breaking callback is unsuitable for this purpose as it notifies changes before any array gets updated.
-			const int32 PrevParentCheckSum = ParentCheckSum;
-			ParentCheckSum = 0;
-			for (int32 ParentIndex : GeometryCollectionComponent->GetParentArray())
-			{
-				ParentCheckSum += ParentIndex;
-			}
-			if (ParentCheckSum != PrevParentCheckSum)
-			{
-				UE_CLOG(GetOwner(), LogGeometryCollectionDebugDraw, Verbose, TEXT("Geometry Collection has broken up for actor %s."), *GetOwner()->GetName());
-				UpdateGeometryVisibility(true);
-			}
-		}
-		else
-		{
-			UE_CLOG(GetOwner(), LogGeometryCollectionDebugDraw, Warning, TEXT("Null Rest Collection for actor %s, skipping Debug Draw Component tick."), *GetOwner()->GetName());
-		}
-	}
-	else
-	{
-		UE_CLOG(GetOwner(), LogGeometryCollectionDebugDraw, Warning, TEXT("Invalid Debug Draw Component for actor %s, tick is now disabled."), *GetOwner()->GetName());
-		SetComponentTickEnabled(false);
-	}
-#endif  // #if GEOMETRYCOLLECTION_DEBUG_DRAW
 }
 
-#if GEOMETRYCOLLECTION_DEBUG_DRAW
-bool UGeometryCollectionDebugDrawComponent::OnDebugDrawPropertiesChanged(bool bForceVisibilityUpdate)
-{
-	if (HasBegunPlay() && GeometryCollectionComponent)
-	{
-		UE_CLOG(GetOwner(), LogGeometryCollectionDebugDraw, Verbose, TEXT("OnDebugDrawPropertiesChanged for actor %s"), *GetOwner()->GetName());
-
-		// Make sure to have an up to date selected transform index
-		UpdateSelectedTransformIndex();
-
-		// Update geometry visibility
-		UpdateGeometryVisibility(bForceVisibilityUpdate);
-
-		// Turn off level set rendering when required
-		UpdateLevelSetVisibility();
-
-		// Update tick function
-		UpdateTickStatus();
-	}
-	return SelectedTransformIndex != INDEX_NONE;
-}
-
-void UGeometryCollectionDebugDrawComponent::OnClusterChanged()
-{
-	if (HasBegunPlay() && GeometryCollectionComponent && IsComponentTickEnabled())  // Only proceed if the tick is already enabled
-	{
-		UE_CLOG(GetOwner(), LogGeometryCollectionDebugDraw, Verbose, TEXT("OnClusterChanged for actor %s"), *GetOwner()->GetName());
-
-		// Make sure to have an up to date selected transform index
-		UpdateSelectedTransformIndex();
-
-		// Update geometry visibility
-		UpdateGeometryVisibility(true);  // Force visibility update
-
-		// Turn off level set rendering when required
-		UpdateLevelSetVisibility();
-	}
-}
-
-void UGeometryCollectionDebugDrawComponent::ComputeClusterTransforms(int32 Index, TArray<bool>& IsComputed, TArray<FTransform>& InOutGlobalTransforms)
-{
-	if (!IsComputed[Index])
-	{
-		FTransform& Transform = InOutGlobalTransforms[Index];
-
-		// Set to local transform
-		const TManagedArray<FTransform>& Transforms = GeometryCollectionComponent->GetTransformArray();
-		Transform = Transforms[Index];
-
-		// Recurse through parents and transform this from local space to global space
-		const TManagedArray<int32>& Parents = GeometryCollectionComponent->GetParentArray();
-		const int32 ParentIndex = Parents[Index];
-		if (ParentIndex != FGeometryCollection::Invalid)
-		{
-			ComputeClusterTransforms(ParentIndex, IsComputed, InOutGlobalTransforms);
-			Transform *= InOutGlobalTransforms[ParentIndex];
-		}
-		IsComputed[Index] = true;
-	}
-}
-
-void UGeometryCollectionDebugDrawComponent::ComputeTransforms(TArray<FTransform>& OutClusterTransforms, TArray<FTransform>& OutGeometryTransforms)
-{
-	check(GeometryCollectionComponent);
-	check(GeometryCollectionComponent->RestCollection);
-
-	const TManagedArray<FTransform>& Transforms = GeometryCollectionComponent->GetTransformArray();
-	const TManagedArray<FTransform>& RestTransforms = GeometryCollectionComponent->GetTransformArrayRest();
-	const TManagedArray<int32>& Parents = GeometryCollectionComponent->GetParentArray();
-	const TManagedArray<TSet<int32>>& Children = GeometryCollectionComponent->GetChildrenArray();
-
-	check(Transforms.Num() == RestTransforms.Num());
-	const int32 NumTransforms = Transforms.Num();
-
-	const AActor* const Actor = GetOwner();
-	check(Actor);
-	const FTransform ActorTransform = Actor->GetTransform();
-
-	// Initialize all flip flop array values to false
-	TArray<bool> FlipFlopComputeStatus;
-	FlipFlopComputeStatus.AddZeroed(NumTransforms);
-
-	// First pass, go up the hierarchy to calculate the leaf/cluster global transforms, not including the actor's tranform
-	OutClusterTransforms.SetNumUninitialized(NumTransforms, false);
-	for (int32 Index = 0; Index < NumTransforms; ++Index)
-	{
-		ComputeClusterTransforms(Index, FlipFlopComputeStatus, OutClusterTransforms);
-	}
-
-	// Second pass, starts from the leaves and go up the hierarchy to update the parents' transform using the rest transform array
-	// Also applies the actor transform to the calculated transforms
-	OutGeometryTransforms.SetNumUninitialized(NumTransforms, false);
-
-	for (int32 Index = 0; Index < NumTransforms; ++Index)
-	{
-		// Start from the leaves, since these have the only up to date geometry transforms
-		if (!Children[Index].Num())
-		{
-			// Apply actor transform to leaves
-			OutClusterTransforms[Index] *= ActorTransform;
-
-			// Copy child geometry transform
-			OutGeometryTransforms[Index] = OutClusterTransforms[Index];
-
-			// Iterate up the hierarchy to update the parent transform, stop when it finds a parent that has already been computed
-			for (int32 ChildIndex = Index, ParentIndex = Parents[Index];
-				ParentIndex != FGeometryCollection::Invalid && FlipFlopComputeStatus[ParentIndex];
-				ParentIndex = Parents[ChildIndex = ParentIndex])
-			{
-				// Finalize the remaining cluster nodes' transform by applying the actor transform
-				OutClusterTransforms[ParentIndex] *= ActorTransform;
 
-				// Calculate final geometry transform from the current child's transform
-				const FTransform& ParentToChildTransform = RestTransforms[ChildIndex];
-				OutGeometryTransforms[ParentIndex] = ParentToChildTransform.Inverse() * OutGeometryTransforms[ChildIndex];
-				
-				// Mark this parent calculation as completed
-				FlipFlopComputeStatus[ParentIndex] = false;
-			}
-		}
-	}
-}
-
-void UGeometryCollectionDebugDrawComponent::DebugDrawTick()
-{
-	check(GeometryCollectionComponent);
-	check(GeometryCollectionComponent->RestCollection);
-	check(GeometryCollectionDebugDrawActor);
-
-	const bool bIsSelected = (SelectedTransformIndex != INDEX_NONE);
-	const bool bIsOneSelected = bIsSelected && !GeometryCollectionDebugDrawActor->bDebugDrawWholeCollection;
-	const bool bAreAllSelected = (bIsSelected && GeometryCollectionDebugDrawActor->bDebugDrawWholeCollection) ||
-		(GeometryCollectionDebugDrawActor->SelectedRigidBody.Id == INDEX_NONE &&
-		//(GeometryCollectionDebugDrawActor->SelectedRigidBody.Id.IsValid() == false &&
-			GeometryCollectionDebugDrawActor->SelectedRigidBody.Solver == GeometryCollectionComponent->ChaosSolverActor);
-
-	// Compute world space geometry and cluster transforms
-	TArray<FTransform> ClusterTransforms;
-	TArray<FTransform> GeometryTransforms;
-	ComputeTransforms(ClusterTransforms, GeometryTransforms);
-
-	if (bIsOneSelected)
-	{
-		if (GeometryCollectionDebugDrawActor->bShowTransformIndex)
-		{
-			GeometryCollectionDebugDrawActor->DrawTransformIndex(ClusterTransforms, GeometryCollectionComponent, SelectedTransformIndex, GeometryCollectionDebugDrawActor->TransformIndexColor);
-		}
-		if (GeometryCollectionDebugDrawActor->bShowTransform)
-		{
-			GeometryCollectionDebugDrawActor->DrawTransform(ClusterTransforms, GeometryCollectionComponent, SelectedTransformIndex, GeometryCollectionDebugDrawActor->TransformScale);
-		}
-		if (GeometryCollectionDebugDrawActor->bShowParent)
-		{
-			GeometryCollectionDebugDrawActor->DrawParent(ClusterTransforms, GeometryCollectionComponent, SelectedTransformIndex, GeometryCollectionDebugDrawActor->ParentColor);
-		}
-		if (GeometryCollectionDebugDrawActor->bShowLevel)
-		{
-			GeometryCollectionDebugDrawActor->DrawLevel(ClusterTransforms, GeometryCollectionComponent, SelectedTransformIndex, GeometryCollectionDebugDrawActor->LevelColor);
-		}
-		if (GeometryCollectionDebugDrawActor->bShowVertices)
-		{
-			GeometryCollectionDebugDrawActor->DrawVertices(GeometryTransforms, GeometryCollectionComponent, SelectedTransformIndex, GeometryCollectionDebugDrawActor->VertexColor);
-		}
-		if (GeometryCollectionDebugDrawActor->bShowVertexIndices)
-		{
-			GeometryCollectionDebugDrawActor->DrawVertexIndices(GeometryTransforms, GeometryCollectionComponent, SelectedTransformIndex, GeometryCollectionDebugDrawActor->VertexIndexColor);
-		}
-		if (GeometryCollectionDebugDrawActor->bShowVertexNormals)
-		{
-			GeometryCollectionDebugDrawActor->DrawVertexNormals(GeometryTransforms, GeometryCollectionComponent, SelectedTransformIndex, GeometryCollectionDebugDrawActor->VertexNormalColor);
-		}
-		if (GeometryCollectionDebugDrawActor->bShowFaces)
-		{
-			GeometryCollectionDebugDrawActor->DrawFaces(GeometryTransforms, GeometryCollectionComponent, SelectedTransformIndex, GeometryCollectionDebugDrawActor->FaceColor);
-		}
-		if (GeometryCollectionDebugDrawActor->bShowFaceIndices)
-		{
-			GeometryCollectionDebugDrawActor->DrawFaceIndices(GeometryTransforms, GeometryCollectionComponent, SelectedTransformIndex, GeometryCollectionDebugDrawActor->FaceIndexColor);
-		}
-		if (GeometryCollectionDebugDrawActor->bShowFaceNormals)
-		{
-			GeometryCollectionDebugDrawActor->DrawFaceNormals(GeometryTransforms, GeometryCollectionComponent, SelectedTransformIndex, GeometryCollectionDebugDrawActor->FaceNormalColor);
-		}
-		if (GeometryCollectionDebugDrawActor->bShowGeometryIndex)
-		{
-			GeometryCollectionDebugDrawActor->DrawGeometryIndex(GeometryTransforms, GeometryCollectionComponent, SelectedTransformIndex, GeometryCollectionDebugDrawActor->GeometryIndexColor);
-		}
-		if (GeometryCollectionDebugDrawActor->bShowGeometryTransform)
-		{
-			GeometryCollectionDebugDrawActor->DrawTransform(GeometryTransforms, GeometryCollectionComponent, SelectedTransformIndex, GeometryCollectionDebugDrawActor->GeometryTransformScale);
-		}
-		if (GeometryCollectionDebugDrawActor->bShowBoundingBox)
-		{
-			GeometryCollectionDebugDrawActor->DrawBoundingBox(GeometryTransforms, GeometryCollectionComponent, SelectedTransformIndex, GeometryCollectionDebugDrawActor->BoundingBoxColor);
-		}
-	}
-	else if (bAreAllSelected)
-	{
-		if (GeometryCollectionDebugDrawActor->bShowTransformIndex)
-		{
-			GeometryCollectionDebugDrawActor->DrawTransformIndices(ClusterTransforms, GeometryCollectionComponent, GeometryCollectionDebugDrawActor->TransformIndexColor);
-		}
-		if (GeometryCollectionDebugDrawActor->bShowTransform)
-		{
-			GeometryCollectionDebugDrawActor->DrawTransforms(ClusterTransforms, GeometryCollectionComponent, GeometryCollectionDebugDrawActor->TransformScale);
-		}
-		if (GeometryCollectionDebugDrawActor->bShowParent)
-		{
-			GeometryCollectionDebugDrawActor->DrawParents(ClusterTransforms, GeometryCollectionComponent, GeometryCollectionDebugDrawActor->ParentColor);
-		}
-		if (GeometryCollectionDebugDrawActor->bShowLevel)
-		{
-			GeometryCollectionDebugDrawActor->DrawLevels(ClusterTransforms, GeometryCollectionComponent, GeometryCollectionDebugDrawActor->LevelColor);
-		}
-		if (GeometryCollectionDebugDrawActor->bShowVertices)
-		{
-			GeometryCollectionDebugDrawActor->DrawVertices(GeometryTransforms, GeometryCollectionComponent, GeometryCollectionDebugDrawActor->VertexColor);
-		}
-		if (GeometryCollectionDebugDrawActor->bShowVertexIndices)
-		{
-			GeometryCollectionDebugDrawActor->DrawVertexIndices(GeometryTransforms, GeometryCollectionComponent, GeometryCollectionDebugDrawActor->VertexIndexColor);
-		}
-		if (GeometryCollectionDebugDrawActor->bShowVertexNormals)
-		{
-			GeometryCollectionDebugDrawActor->DrawVertexNormals(GeometryTransforms, GeometryCollectionComponent, GeometryCollectionDebugDrawActor->VertexNormalColor);
-		}
-		if (GeometryCollectionDebugDrawActor->bShowFaces)
-		{
-			GeometryCollectionDebugDrawActor->DrawFaces(GeometryTransforms, GeometryCollectionComponent, GeometryCollectionDebugDrawActor->FaceColor);
-		}
-		if (GeometryCollectionDebugDrawActor->bShowFaceIndices)
-		{
-			GeometryCollectionDebugDrawActor->DrawFaceIndices(GeometryTransforms, GeometryCollectionComponent, GeometryCollectionDebugDrawActor->FaceIndexColor);
-		}
-		if (GeometryCollectionDebugDrawActor->bShowFaceNormals)
-		{
-			GeometryCollectionDebugDrawActor->DrawFaceNormals(GeometryTransforms, GeometryCollectionComponent, GeometryCollectionDebugDrawActor->FaceNormalColor);
-		}
-		if (GeometryCollectionDebugDrawActor->bShowGeometryIndex)
-		{
-			GeometryCollectionDebugDrawActor->DrawGeometryIndices(GeometryTransforms, GeometryCollectionComponent, GeometryCollectionDebugDrawActor->GeometryIndexColor);
-		}
-		if (GeometryCollectionDebugDrawActor->bShowGeometryTransform)
-		{
-			GeometryCollectionDebugDrawActor->DrawTransforms(GeometryTransforms, GeometryCollectionComponent, GeometryCollectionDebugDrawActor->GeometryTransformScale);
-		}
-		if (GeometryCollectionDebugDrawActor->bShowBoundingBox)
-		{
-			GeometryCollectionDebugDrawActor->DrawBoundingBoxes(GeometryTransforms, GeometryCollectionComponent, GeometryCollectionDebugDrawActor->BoundingBoxColor);
-		}
-	}
-	// Draw single face
-	if (bIsSelected && GeometryCollectionDebugDrawActor->bShowSingleFace)
-	{
-		GeometryCollectionDebugDrawActor->DrawSingleFace(GeometryTransforms, GeometryCollectionComponent, GeometryCollectionDebugDrawActor->SingleFaceIndex, GeometryCollectionDebugDrawActor->SingleFaceColor);
-	}
-}
-
-void UGeometryCollectionDebugDrawComponent::UpdateSelectedTransformIndex()
-{
-	check(GeometryCollectionComponent);
-
-	// No actor, no selection
-	if (!GeometryCollectionDebugDrawActor)
-	{
-		SelectedTransformIndex = SelectedRigidBodyId = INDEX_NONE;
-		//SelectedTransformIndex = INDEX_NONE;
-		//SelectedRigidBodyId = FGuid();
-		return;
-	}
-
-	// Check whether the selected rigid body id, or solver has changed
-	if (!bHasIncompleteRigidBodyIdSync &&
-		SelectedRigidBodyId == GeometryCollectionDebugDrawActor->SelectedRigidBody.Id &&
-		SelectedChaosSolver == GeometryCollectionDebugDrawActor->SelectedRigidBody.Solver)
-	{
-		return;
-	}
-
-	// Default init selected transform index, in case of premature exit
-	SelectedTransformIndex = INDEX_NONE;
-
-	// Simple test to allow for an early exit when nothing has been selected
-	if (GeometryCollectionDebugDrawActor->SelectedRigidBody.Id == INDEX_NONE ||
-	//if (GeometryCollectionDebugDrawActor->SelectedRigidBody.Id.IsValid() == false ||
-		GeometryCollectionDebugDrawActor->SelectedRigidBody.Solver != GeometryCollectionComponent->ChaosSolverActor)
-	{
-		SelectedRigidBodyId = GeometryCollectionDebugDrawActor->SelectedRigidBody.Id;
-		SelectedChaosSolver = GeometryCollectionDebugDrawActor->SelectedRigidBody.Solver;
-		return;
-	}
-	/*
-	// Check rigid body id sync
-	// Note that this test alone isn't enough to ensure that the rigid body ids are valid.
-	//const TManagedArray<int32>& RigidBodyIds = GeometryCollectionComponent->RigidBodyIds;
-	const TManagedArray<FGuid>& RigidBodyIds = GeometryCollectionComponent->GetRigidBodyGuidArray();
-	if (RigidBodyIds.Num() == 0)
-	{
-		bHasIncompleteRigidBodyIdSync = !!GeometryCollectionComponent->GetTransformArray().Num();
-		UE_CLOG(bHasIncompleteRigidBodyIdSync && GetOwner(), LogGeometryCollectionDebugDraw, Verbose, TEXT("UpdateSelectedTransformIndex(): Empty RigidBodyIds array for actor %s."), *GetOwner()->GetName());
-		return;
-	}
-
-	// Find the matching transform if any (and also check the sync completion status)
-	bHasIncompleteRigidBodyIdSync = false;
-
-	const TManagedArray<TSet<int32>>& ChildrenRest = GeometryCollectionComponent->GetChildrenArrayRest();
-	const TManagedArray<TSet<int32>>& Children = GeometryCollectionComponent->GetChildrenArray();
-
-	for (int32 TransformIndex = 0; TransformIndex < RigidBodyIds.Num(); ++TransformIndex)
-	{
-		// Is this the selected id?
-		//if (RigidBodyIds[TransformIndex] == GeometryCollectionDebugDrawActor->SelectedRigidBody.Id)
-		if (static_cast<int32>(GetTypeHash(RigidBodyIds[TransformIndex])) == GeometryCollectionDebugDrawActor->SelectedRigidBody.Id)
-		{
-			SelectedTransformIndex = TransformIndex;
-			bHasIncompleteRigidBodyIdSync = false;  // Found it, the wait for a sync can be canceled
-			break;
-		}
-		// Check the reason behind any invalid index
-		//if (RigidBodyIds[TransformIndex] == INDEX_NONE)
-		if (!RigidBodyIds[TransformIndex].IsValid())
-		{
-			// Look for detached clusters in order to differentiate un-synced vs empty cluster rigid body ids.
-			int32 ChildTransformIndex = TransformIndex;
-			while (const TSet<int32>::TConstIterator ChildTransformIterator = Children[ChildTransformIndex].CreateConstIterator())
-			{
-				// Go down to the cluster's leaf level through the first child
-				ChildTransformIndex = *ChildTransformIterator;
-			}
-
-			// If this is a leaf bone, it can not be a detached cluster so it should have a valid rigid body
-			// In which case the sync has yet to happen and it might be worth trying this again later
-			if (!ChildrenRest[ChildTransformIndex].Num())
-			{
-				bHasIncompleteRigidBodyIdSync = true;
-				UE_CLOG(GetOwner(), LogGeometryCollectionDebugDraw, VeryVerbose, TEXT("UpdateSelectedTransformIndex(): Invalid rigid body id for actor %s, TransformIndex %d."), *GetOwner()->GetName(), TransformIndex);
-			}
-			else
-			{
-				// This should match the SimulationType == FST_CLUSTERED or IsClustered(int32 Element)
-				ensure(GeometryCollectionComponent->GetSimulationTypeArrayRest()[ChildTransformIndex] == FGeometryCollection::ESimulationTypes::FST_Clustered);
-				UE_CLOG(GetOwner(), LogGeometryCollectionDebugDraw, VeryVerbose, TEXT("UpdateSelectedTransformIndex(): Found empty cluster for actor %s, TransformIndex %d."), *GetOwner()->GetName(), TransformIndex);
-			}
-		}
-	}
-	UE_CLOG(bHasIncompleteRigidBodyIdSync && GetOwner(), LogGeometryCollectionDebugDraw, Verbose, TEXT("UpdateSelectedTransformIndex(): Invalid RigidBodyIds array elements for actor %s."), *GetOwner()->GetName());
-
-	// Update selected rigid body index and solver
-	SelectedRigidBodyId = GeometryCollectionDebugDrawActor->SelectedRigidBody.Id;
-	SelectedChaosSolver = GeometryCollectionDebugDrawActor->SelectedRigidBody.Solver;
-	*/
-}
-
-int32 UGeometryCollectionDebugDrawComponent::CountFaces(int32 TransformIndex, bool bDebugDrawClustering) const
-{
-	check(GeometryCollectionComponent);
-	int32 FaceCount = 0;
-	const TManagedArray<int32>& TransformToGeometryIndexArray = GeometryCollectionComponent->GetTransformToGeometryIndexArray();
-	const int32 GeometryIndex = TransformToGeometryIndexArray[TransformIndex];
-	if (GeometryIndex != INDEX_NONE)
-	{
-		const TManagedArray<int32>& FaceCountArray = GeometryCollectionComponent->GetFaceCountArray();
-		FaceCount += FaceCountArray[GeometryIndex];
-	}
-	const TManagedArray<TSet<int32>>& ChildrenArray = bDebugDrawClustering ?
-		GeometryCollectionComponent->GetChildrenArrayRest():
-		GeometryCollectionComponent->GetChildrenArray();
-	for (int32 HiddenChildIndex : ChildrenArray[TransformIndex])
-	{
-		FaceCount += CountFaces(HiddenChildIndex, bDebugDrawClustering);
-	}
-	return FaceCount;
-}
-
-void UGeometryCollectionDebugDrawComponent::HideFaces(int32 TransformIndex, bool bDebugDrawClustering)
-{
-	check(GeometryCollectionComponent);
-	const TManagedArray<int32>& TransformToGeometryIndexArray = GeometryCollectionComponent->GetTransformToGeometryIndexArray();
-	const int32 GeometryIndex = TransformToGeometryIndexArray[TransformIndex];
-	if (GeometryIndex != INDEX_NONE)
-	{
-		TManagedArray<bool>& VisibleArray = GeometryCollectionComponent->GetVisibleArrayCopyOnWrite();
-		const TManagedArray<int32>& FaceStartArray = GeometryCollectionComponent->GetFaceStartArray();
-		const TManagedArray<int32>& FaceCountArray = GeometryCollectionComponent->GetFaceCountArray();
-		const int32 FaceStart = FaceStartArray[GeometryIndex];
-		const int32 FaceCount = FaceCountArray[GeometryIndex];
-		for (int32 FaceIndex = FaceStart; FaceIndex < FaceStart + FaceCount; ++FaceIndex)
-		{
-			VisibleArray[FaceIndex] = false;
-		}
-	}
-	const TManagedArray<TSet<int32>>& ChildrenArray = bDebugDrawClustering ?
-		GeometryCollectionComponent->GetChildrenArrayRest():
-		GeometryCollectionComponent->GetChildrenArray();
-	for (int32 HiddenChildIndex : ChildrenArray[TransformIndex])
-	{
-		HideFaces(HiddenChildIndex, bDebugDrawClustering);
-	}
-}
-
-void UGeometryCollectionDebugDrawComponent::UpdateGeometryVisibility(bool bForceVisibilityUpdate)
-{
-	check(GeometryCollectionComponent);
-	if (!GeometryCollectionComponent->RestCollection)  // Required for GetTransformIndexArray
-	{
-		UE_CLOG(GetOwner(), LogGeometryCollectionDebugDraw, Warning, TEXT("UpdateGeometryVisibility(): Empty RestCollection for actor %s"), *GetOwner()->GetName());
-		return;
-	}
-
-	// Keep old hidden index
-	const int32 PrevHiddenIndex = HiddenTransformIndex;
-
-	// Extract this object's visibility arguments from the debug draw actor's hide geometry status
-	bool bIsVisible;
-	if (!GeometryCollectionDebugDrawActor)
-	{
-		bIsVisible = true;
-		HiddenTransformIndex = INDEX_NONE;
-	}
-	else
-	{
-		// Work out partial changes in visiblity
-		const bool bIsSelected = (SelectedTransformIndex != INDEX_NONE);
-		const bool bAreAllSelected = (bIsSelected && GeometryCollectionDebugDrawActor->bDebugDrawWholeCollection) ||
-			(GeometryCollectionDebugDrawActor->SelectedRigidBody.Id == INDEX_NONE && 
-			//(GeometryCollectionDebugDrawActor->SelectedRigidBody.Id.IsValid() == false && 
-			 GeometryCollectionDebugDrawActor->SelectedRigidBody.Solver == GeometryCollectionComponent->ChaosSolverActor);
-		const bool bAreAnySelected = bIsSelected || bAreAllSelected;
-
-		switch (GeometryCollectionDebugDrawActor->HideGeometry)
-		{
-		default:  // Intentional fallthrough
-		case EGeometryCollectionDebugDrawActorHideGeometry::HideNone:
-			bIsVisible = true;
-			HiddenTransformIndex = INDEX_NONE;
-			break;
-		case EGeometryCollectionDebugDrawActorHideGeometry::HideWithCollision:
-			if (bAreAnySelected && GeometryCollectionDebugDrawActor->bShowRigidBodyCollision)
-			{
-				bIsVisible = !bAreAllSelected;
-				HiddenTransformIndex = bIsVisible ? SelectedTransformIndex: INDEX_NONE;
-			}
-			else
-			{
-				bIsVisible = true;
-				HiddenTransformIndex = INDEX_NONE;
-			}
-			break;
-		case EGeometryCollectionDebugDrawActorHideGeometry::HideSelected:
-			bIsVisible = !bAreAllSelected;
-			HiddenTransformIndex = bIsVisible ? SelectedTransformIndex: INDEX_NONE;
-			break;
-		case EGeometryCollectionDebugDrawActorHideGeometry::HideWholeCollection:
-			bIsVisible = !bAreAnySelected;
-			HiddenTransformIndex = INDEX_NONE;
-			break;
-		case EGeometryCollectionDebugDrawActorHideGeometry::HideAll:
-			bIsVisible = false;
-			HiddenTransformIndex = INDEX_NONE;
-			break;
-		}
-	}
-
-	// Update face visibility
-	bool bIndexHasChanged = (HiddenTransformIndex != PrevHiddenIndex) || bForceVisibilityUpdate;
-	if (bIndexHasChanged)
-	{
-		// This must happen in both show and hidden geometry mode so that the last 
-		// hidden section does not stay hidden when switching between Selected>All>None.
-
-		// todo(ocohen): add ability to use rest collection (i.e. don't use instance copy)
-		TManagedArray<bool>& VisibleArray = GeometryCollectionComponent->GetVisibleArrayCopyOnWrite();
-
-		// Reset visibility array to default as the index has changed
-		VisibleArray.Init(GeometryCollectionComponent->GetVisibleArrayRest());
-		UE_LOG(LogGeometryCollectionDebugDraw, Verbose, TEXT("VisibleArray re-initialized."));
-
-		if (HiddenTransformIndex != INDEX_NONE)
-		{
-			const int32 NumFaces = CountFaces(HiddenTransformIndex, GeometryCollectionDebugDrawActor->bDebugDrawClustering);
-			if (NumFaces < VisibleArray.Num())
-			{
-				// Hide this geometry's faces
-				HideFaces(HiddenTransformIndex, GeometryCollectionDebugDrawActor->bDebugDrawClustering);
-				UE_LOG(LogGeometryCollectionDebugDraw, Verbose, TEXT("UpdateGeometryVisibility(): Hidding partial object."));
-			}
-			else
-			{
-				// Hide entire object
-				// Can't send zero vertices to force the vertex buffer to be empty, so hide the component instead
-				bIsVisible = false;
-				HiddenTransformIndex = INDEX_NONE;
-				// Update index has changed
-				bIndexHasChanged = (HiddenTransformIndex != PrevHiddenIndex);
-				UE_LOG(LogGeometryCollectionDebugDraw, Verbose, TEXT("UpdateGeometryVisibility(): Hidding entire object."));
-			}
-		}
-	}
-	UE_CLOG(bIndexHasChanged, LogGeometryCollectionDebugDraw, Verbose, TEXT("UpdateGeometryVisibility(): Index has changed. Prev index = %d, new index = %d."), PrevHiddenIndex, HiddenTransformIndex);
-
-	// Force component reinit
-	if (bIndexHasChanged)
-	{
-		UE_LOG(LogGeometryCollectionDebugDraw, Verbose, TEXT("UpdateGeometryVisibility(): Forcing init render data."));
-		GeometryCollectionComponent->ForceRenderUpdateConstantData();
-	}
-
-	// Update component visibility, only if the component visibility hasn't changed in between the last call (or unless the change is in sync with the component visibility)
-	const bool bIsComponentVisible = GeometryCollectionComponent->IsVisible();
-	const bool bAllowVisibilityChange = (bIsComponentVisible || !bWasVisible);
-	if (bAllowVisibilityChange)
-	{
-		const bool bVisibilityHasChanged = (bIsVisible != bIsComponentVisible);
-		if (bVisibilityHasChanged)
-		{
-			UE_LOG(LogGeometryCollectionDebugDraw, Verbose, TEXT("UpdateGeometryVisibility(): Visibility has changed. Old visibility = %d, new visibility = %d."), bWasVisible, bIsVisible);
-			GeometryCollectionComponent->SetVisibility(bIsVisible);
-		}
-		bWasVisible = bIsVisible;  // Only update when changes are allowed so that the component can stay hidden when visibility is out of sync.
-	}
-}
-
-void UGeometryCollectionDebugDrawComponent::UpdateTickStatus()
-{
-	bool bIsEnabled;
-	if (!GeometryCollectionDebugDrawActor)
-	{
-		bIsEnabled = false;
-	}
-	else
-	{
-		// Check whether anything from this component is selected for debug drawing
-		const bool bAreAnySelected = (SelectedTransformIndex != INDEX_NONE ||
-			(GeometryCollectionDebugDrawActor->SelectedRigidBody.Id == INDEX_NONE &&
-			//(GeometryCollectionDebugDrawActor->SelectedRigidBody.Id.IsValid() == false &&
-			 GeometryCollectionDebugDrawActor->SelectedRigidBody.Solver == GeometryCollectionComponent->ChaosSolverActor));
-
-		bIsEnabled =
-			bAreAnySelected && (
-				   GeometryCollectionDebugDrawActor->bShowRigidBodyId
-				|| GeometryCollectionDebugDrawActor->bShowRigidBodyCollision
-				|| GeometryCollectionDebugDrawActor->bShowRigidBodyTransform
-				|| GeometryCollectionDebugDrawActor->bShowRigidBodyInertia
-				|| GeometryCollectionDebugDrawActor->bShowRigidBodyVelocity
-				|| GeometryCollectionDebugDrawActor->bShowRigidBodyForce
-				|| GeometryCollectionDebugDrawActor->bShowRigidBodyInfos
-				|| GeometryCollectionDebugDrawActor->bShowTransformIndex
-				|| GeometryCollectionDebugDrawActor->bShowTransform
-				|| GeometryCollectionDebugDrawActor->bShowParent
-				|| GeometryCollectionDebugDrawActor->bShowLevel
-				|| GeometryCollectionDebugDrawActor->bShowConnectivityEdges
-				|| GeometryCollectionDebugDrawActor->bShowGeometryIndex
-				|| GeometryCollectionDebugDrawActor->bShowGeometryTransform
-				|| GeometryCollectionDebugDrawActor->bShowBoundingBox
-				|| GeometryCollectionDebugDrawActor->bShowFaces
-				|| GeometryCollectionDebugDrawActor->bShowFaceIndices
-				|| GeometryCollectionDebugDrawActor->bShowFaceNormals
-				|| GeometryCollectionDebugDrawActor->bShowSingleFace
-				|| GeometryCollectionDebugDrawActor->SingleFaceIndex
-				|| GeometryCollectionDebugDrawActor->bShowVertices
-				|| GeometryCollectionDebugDrawActor->bShowVertexIndices
-				|| GeometryCollectionDebugDrawActor->bShowVertexNormals);
-	}
-	// Update component's ability to tick 
-	SetComponentTickEnabled(bIsEnabled);
-	UE_CLOG(GetOwner(), LogGeometryCollectionDebugDraw, Verbose, TEXT("Debug Draw Tick Component bIsEnabled = %d for actor %s"), bIsEnabled, *GetOwner()->GetName());
-}
-
-void UGeometryCollectionDebugDrawComponent::DebugDrawChaosTick()
-{
-	/*
-	check(GeometryCollectionComponent);
-	check(GeometryCollectionDebugDrawActor);
-	check(GeometryCollectionRenderLevelSetActor);
-
-	AActor* const Actor = GetOwner();
-	check(Actor);
-
-	
-	// Retrieve synced particle and clustering data
-	//const TManagedArray<int32>& RigidBodyIds = GeometryCollectionComponent->RigidBodyIds;
-	const TManagedArray<FGuid>& RigidBodyIds = GeometryCollectionComponent->GetRigidBodyGuidArray();
-	const FGeometryCollectionPhysicsProxy* const PhysicsProxy = GeometryCollectionComponent->GetPhysicsProxy();
-	if (PhysicsProxy)
-	{
-		ParticlesData.Sync(PhysicsProxy->GetSolver<Chaos::FPhysicsSolver>(), RigidBodyIds);
-	}
-
-	// Visualize single rigid body
-	const bool bIsSelected = (SelectedTransformIndex != INDEX_NONE);
-	if (bIsSelected)
-	{
-		// Visualize the level set collision volume when synced data are available and set to the correct type
-		if (GeometryCollectionDebugDrawActor->bShowRigidBodyCollision &&
-			ParticlesData.RequestSyncedData(EGeometryCollectionParticlesData::GeometryType) &&
-			ParticlesData.GetGeometryType(SelectedTransformIndex) == Chaos::ImplicitObjectType::LevelSet)
-		{
-			// Get the transform for the current piece
-			FTransform Transform = FTransform::Identity;
-
-			// Update the transform if we are rendering the level set aligned with the simulated geometry
-			bool bSynced;
-			bSynced = ParticlesData.RequestSyncedData(EGeometryCollectionParticlesData::X);
-			bSynced = ParticlesData.RequestSyncedData(EGeometryCollectionParticlesData::R) && bSynced;
-			bSynced = ParticlesData.RequestSyncedData(EGeometryCollectionParticlesData::ChildToParentMap) && bSynced;
-
-			if (!GeometryCollectionDebugDrawActor->bCollisionAtOrigin && bSynced)
-			{
-				// Retrieve particle transform
-				Transform = AGeometryCollectionDebugDrawActor::GetParticleTransform(GeometryCollectionComponent, SelectedTransformIndex, ParticlesData);
-			}
-
-			// If the level set index has changed at run time, then reload the volume
-			// because someone wants to visualize another piece
-			const bool bLevelSetTextureDirty = (RenderLevelSetOwner != this || LastRenderedId == INDEX_NONE || LastRenderedId != SelectedRigidBodyId);
-			//const bool bLevelSetTextureDirty = (RenderLevelSetOwner != this || LastRenderedId.IsValid() == false || LastRenderedId != SelectedRigidBodyId);
-			if (!bLevelSetTextureDirty)
-			{
-				// If we are only updating the transform, or also loading the volume
-				GeometryCollectionRenderLevelSetActor->SyncLevelSetTransform(Transform);
-			}
-			else if (ParticlesData.RequestSyncedData(EGeometryCollectionParticlesData::Geometry))
-			{
-				// Retrieve level set pointer from sync
-				const Chaos::FLevelSet* LevelSet = static_cast<const Chaos::FLevelSet*>(ParticlesData.GetGeometry(SelectedTransformIndex));
-
-				// Build the volume texture
-				// @note: we only want to do this once, so we have a state variable on the component to ensure that
-				const bool Success = GeometryCollectionRenderLevelSetActor->SetLevelSetToRender(*LevelSet, Transform);
-				if (!Success)
-				{
-					UE_LOG(LogGeometryCollectionDebugDraw, Warning, TEXT("Levelset generation failed: %s"), *GetFullName());
-					LastRenderedId = INDEX_NONE;
-					//LastRenderedId = FGuid();
-				}
-				else
-				{
-					// Take ownership
-					RenderLevelSetOwner = this;
-					// Turn on the volume rendering
-					GeometryCollectionRenderLevelSetActor->SetEnabled(true);
-					// Update last rendered index for next dirty test
-					LastRenderedId = SelectedRigidBodyId;
-				}
-			}
-			else {}  // Wait a game tick for synced data
-		}
-	}
-
-	// Visualize other rigid body debug draw informations
-	const bool bIsOneSelected = bIsSelected && !GeometryCollectionDebugDrawActor->bDebugDrawWholeCollection;
-	const bool bAreAllSelected = (bIsSelected && GeometryCollectionDebugDrawActor->bDebugDrawWholeCollection) ||
-		(GeometryCollectionDebugDrawActor->SelectedRigidBody.Id == INDEX_NONE &&
-		//(GeometryCollectionDebugDrawActor->SelectedRigidBody.Id.IsValid() == false &&
-		 GeometryCollectionDebugDrawActor->SelectedRigidBody.Solver == GeometryCollectionComponent->ChaosSolverActor);
-
-	if (bIsOneSelected)
-	{
-		if (GeometryCollectionDebugDrawActor->bShowRigidBodyId)
-		{
-			GeometryCollectionDebugDrawActor->DrawRigidBodyId(GeometryCollectionComponent, SelectedTransformIndex, ParticlesData, RigidBodyIds, GeometryCollectionDebugDrawActor->RigidBodyIdColor);
-		}
-		if (GeometryCollectionDebugDrawActor->bShowRigidBodyTransform)
-		{
-			GeometryCollectionDebugDrawActor->DrawRigidBodyTransform(GeometryCollectionComponent, SelectedTransformIndex, ParticlesData, GeometryCollectionDebugDrawActor->RigidBodyTransformScale);
-		}
-		const bool bIsShowingLevelSet = (RenderLevelSetOwner == this && LastRenderedId == SelectedRigidBodyId);  // Only draw single collision whenever there isn't a level set being already rendered
-		if (!bIsShowingLevelSet && GeometryCollectionDebugDrawActor->bShowRigidBodyCollision)
-		{
-			GeometryCollectionDebugDrawActor->DrawRigidBodyCollision(GeometryCollectionComponent, SelectedTransformIndex, ParticlesData, GeometryCollectionDebugDrawActor->RigidBodyCollisionColor);
-		}
-		if (GeometryCollectionDebugDrawActor->bShowRigidBodyInertia)
-		{
-			GeometryCollectionDebugDrawActor->DrawRigidBodyInertia(GeometryCollectionComponent, SelectedTransformIndex, ParticlesData, GeometryCollectionDebugDrawActor->RigidBodyInertiaColor);
-		}
-		if (GeometryCollectionDebugDrawActor->bShowRigidBodyVelocity)
-		{
-			GeometryCollectionDebugDrawActor->DrawRigidBodyVelocity(GeometryCollectionComponent, SelectedTransformIndex, ParticlesData, GeometryCollectionDebugDrawActor->RigidBodyVelocityColor);
-		}
-		if (GeometryCollectionDebugDrawActor->bShowRigidBodyForce)
-		{
-			GeometryCollectionDebugDrawActor->DrawRigidBodyForce(GeometryCollectionComponent, SelectedTransformIndex, ParticlesData, GeometryCollectionDebugDrawActor->RigidBodyForceColor);
-		}
-		if (GeometryCollectionDebugDrawActor->bShowRigidBodyInfos)
-		{
-			GeometryCollectionDebugDrawActor->DrawRigidBodyInfo(GeometryCollectionComponent, SelectedTransformIndex, ParticlesData, GeometryCollectionDebugDrawActor->RigidBodyInfoColor);
-		}
-		if (GeometryCollectionDebugDrawActor->bShowConnectivityEdges)
-		{
-			GeometryCollectionDebugDrawActor->DrawConnectivityEdges(GeometryCollectionComponent, SelectedTransformIndex, ParticlesData, RigidBodyIds);
-		}
-	}
-	else if (bAreAllSelected)
-	{
-		if (GeometryCollectionDebugDrawActor->bShowRigidBodyCollision)
-		{
-			GeometryCollectionDebugDrawActor->DrawRigidBodiesCollision(GeometryCollectionComponent, ParticlesData, GeometryCollectionDebugDrawActor->RigidBodyCollisionColor);
-		}
-		if (GeometryCollectionDebugDrawActor->bShowRigidBodyInertia)
-		{
-			GeometryCollectionDebugDrawActor->DrawRigidBodiesInertia(GeometryCollectionComponent, ParticlesData, GeometryCollectionDebugDrawActor->RigidBodyInertiaColor);
-		}
-		if (GeometryCollectionDebugDrawActor->bShowRigidBodyVelocity)
-		{
-			GeometryCollectionDebugDrawActor->DrawRigidBodiesVelocity(GeometryCollectionComponent, ParticlesData, GeometryCollectionDebugDrawActor->RigidBodyVelocityColor);
-		}
-		if (GeometryCollectionDebugDrawActor->bShowRigidBodyForce)
-		{
-			GeometryCollectionDebugDrawActor->DrawRigidBodiesForce(GeometryCollectionComponent, ParticlesData, GeometryCollectionDebugDrawActor->RigidBodyForceColor);
-		}
-		if (GeometryCollectionDebugDrawActor->bShowRigidBodyInfos)
-		{
-			GeometryCollectionDebugDrawActor->DrawRigidBodiesInfo(GeometryCollectionComponent, ParticlesData, GeometryCollectionDebugDrawActor->RigidBodyInfoColor);
-		}
-		if (GeometryCollectionDebugDrawActor->bShowConnectivityEdges)
-		{
-			GeometryCollectionDebugDrawActor->DrawConnectivityEdges(GeometryCollectionComponent, ParticlesData, RigidBodyIds);
-		}
-		if (GeometryCollectionDebugDrawActor->bShowRigidBodyTransform)
-		{
-			GeometryCollectionDebugDrawActor->DrawRigidBodiesTransform(GeometryCollectionComponent, ParticlesData, GeometryCollectionDebugDrawActor->RigidBodyTransformScale);
-		}
-		if (GeometryCollectionDebugDrawActor->bShowRigidBodyId)
-		{
-			GeometryCollectionDebugDrawActor->DrawRigidBodiesId(GeometryCollectionComponent, ParticlesData, RigidBodyIds, GeometryCollectionDebugDrawActor->RigidBodyIdColor);
-		}
-	}
-	*/
-}
-
-void UGeometryCollectionDebugDrawComponent::UpdateLevelSetVisibility()
-{
-	const bool bIsSelected = (SelectedTransformIndex != INDEX_NONE);
-	const bool bShowCollision = bIsSelected && GeometryCollectionDebugDrawActor && GeometryCollectionDebugDrawActor->bShowRigidBodyCollision;
-	
-	if (RenderLevelSetOwner == this && (LastRenderedId == INDEX_NONE || !bShowCollision))
-	//if (RenderLevelSetOwner == this && (LastRenderedId.IsValid() == false || !bShowCollision))
-	{
-		// Disable rendering
-		GeometryCollectionRenderLevelSetActor->SetEnabled(false);
-
-		// Disown renderer
-		RenderLevelSetOwner = nullptr;
-		LastRenderedId = INDEX_NONE;
-		//LastRenderedId = FGuid();
-	}
-}
-#endif  // #if GEOMETRYCOLLECTION_DEBUG_DRAW
diff --git a/Engine/Source/Runtime/Experimental/GeometryCollectionEngine/Public/GeometryCollection/GeometryCollectionActor.h b/Engine/Source/Runtime/Experimental/GeometryCollectionEngine/Public/GeometryCollection/GeometryCollectionActor.h
--- a/Engine/Source/Runtime/Experimental/GeometryCollectionEngine/Public/GeometryCollection/GeometryCollectionActor.h
+++ b/Engine/Source/Runtime/Experimental/GeometryCollectionEngine/Public/GeometryCollection/GeometryCollectionActor.h
@@ -34,10 +34,11 @@ public:
 	TObjectPtr<UGeometryCollectionComponent> GeometryCollectionComponent;
 	UGeometryCollectionComponent* GetGeometryCollectionComponent() const { return GeometryCollectionComponent; }
 
-	UPROPERTY(VisibleAnywhere, Category = Destruction, meta = (ExposeFunctionCategories = "Components|GeometryCollection", AllowPrivateAccess = "true"))
-	TObjectPtr<UGeometryCollectionDebugDrawComponent> GeometryCollectionDebugDrawComponent;
-	UGeometryCollectionDebugDrawComponent* GetGeometryCollectionDebugDrawComponent() const { return GeometryCollectionDebugDrawComponent; }
-
+	PRAGMA_DISABLE_DEPRECATION_WARNINGS
+	UPROPERTY()
+	TObjectPtr<UGeometryCollectionDebugDrawComponent> GeometryCollectionDebugDrawComponent_DEPRECATED;
+	PRAGMA_ENABLE_DEPRECATION_WARNINGS
+	
 	UFUNCTION(BlueprintCallable, Category = "Physics")
 	bool RaycastSingle(FVector Start, FVector End, FHitResult& OutHit) const;
 
diff --git a/Engine/Source/Runtime/Experimental/GeometryCollectionEngine/Public/GeometryCollection/GeometryCollectionDebugDrawActor.h b/Engine/Source/Runtime/Experimental/GeometryCollectionEngine/Public/GeometryCollection/GeometryCollectionDebugDrawActor.h
--- a/Engine/Source/Runtime/Experimental/GeometryCollectionEngine/Public/GeometryCollection/GeometryCollectionDebugDrawActor.h
+++ b/Engine/Source/Runtime/Experimental/GeometryCollectionEngine/Public/GeometryCollection/GeometryCollectionDebugDrawActor.h
@@ -40,7 +40,7 @@ enum class EGeometryCollectionDebugDrawActorHideGeometry : uint8
 *   Empty structure used to embed a warning message in the UI through a detail customization.
 */
 USTRUCT()
-struct FGeometryCollectionDebugDrawWarningMessage
+struct UE_DEPRECATED(5.0, "Deprecated. Use normal debug draw Chaos Physics commands") FGeometryCollectionDebugDrawWarningMessage
 {
 	GENERATED_USTRUCT_BODY()
 };
@@ -50,7 +50,7 @@ struct FGeometryCollectionDebugDrawWarningMessage
 *   Structure used to select a rigid body id with a picking tool through a detail customization.
 */
 USTRUCT()
-struct FGeometryCollectionDebugDrawActorSelectedRigidBody
+struct UE_DEPRECATED(5.0, "Deprecated. Use normal debug draw Chaos Physics commands") FGeometryCollectionDebugDrawActorSelectedRigidBody
 {
 	GENERATED_USTRUCT_BODY()
 
@@ -81,20 +81,22 @@ struct FGeometryCollectionDebugDrawActorSelectedRigidBody
 *   Only one actor is to be used in the world, and should be automatically 
 *   spawned by any GeometryDebugDrawComponent that needs it.
 */
+class UE_DEPRECATED(5.0, "Deprecated. Use normal debug draw Chaos Physics commands") AGeometryCollectionDebugDrawActor;
 UCLASS(HideCategories = ("Rendering", "Replication", "Input", "Actor", "Collision", "LOD", "Cooking"))
 class GEOMETRYCOLLECTIONENGINE_API AGeometryCollectionDebugDrawActor : public AActor
 {
 	GENERATED_UCLASS_BODY()
 public:
-
+	PRAGMA_DISABLE_DEPRECATION_WARNINGS
 	/** Warning message to explain that the debug draw properties have no effect until starting playing/simulating. */
-	UPROPERTY(EditAnywhere, Category = "Debug Draw")
-	FGeometryCollectionDebugDrawWarningMessage WarningMessage;
+	UPROPERTY()
+	FGeometryCollectionDebugDrawWarningMessage WarningMessage_DEPRECATED;
 
 	/** Picking tool used to select a rigid body id. */
-	UPROPERTY(EditAnywhere, Category = "Debug Draw")
-	FGeometryCollectionDebugDrawActorSelectedRigidBody SelectedRigidBody;
-
+	UPROPERTY()
+	FGeometryCollectionDebugDrawActorSelectedRigidBody SelectedRigidBody_DEPRECATED;
+	PRAGMA_ENABLE_DEPRECATION_WARNINGS
+	
 	/** Show debug visualization for the rest of the geometry collection related to the current rigid body id selection. */
 	UPROPERTY(EditAnywhere, Category = "Debug Draw")
 	bool bDebugDrawWholeCollection;
@@ -331,241 +333,4 @@ public:
 	/** Display icon in the editor. */
 	UPROPERTY()
 	TObjectPtr<UBillboardComponent> SpriteComponent;
-
-	/** Game tick callback. This tick function is required to clean up the persistent debug lines. */
-	static AGeometryCollectionDebugDrawActor* FindOrCreate(UWorld* World);
-
-	/** Game tick callback. This tick function is required to clean up the persistent debug lines. */
-	virtual void Tick(float DeltaSeconds) override;
-
-	/** Actor destruction callback. Used here to clear up the command callbacks. */
-	virtual void BeginDestroy() override;
-
-	/** Register debug draw service. */
-	virtual void BeginPlay() override;
-
-	/** Unregister debug draw service. */
-	virtual void EndPlay(EEndPlayReason::Type ReasonEnd) override;
-
-	/** Reset command variables from the newly loaded properties. */
-	virtual void PostLoad() override;
-
-#if WITH_EDITOR
-	/** Property changed callback. Required to synchronize the command variables to this Actor's properties. */
-	virtual void PostEditChangeProperty(FPropertyChangedEvent& PropertyChangedEvent) override;
-
-	/** Some properties are unlocked depending on the value of the indices not being -1. */
-	virtual bool CanEditChange(const FProperty* InProperty) const override;
-#endif
-
-	/** Draw vertices. */
-	void DrawVertices(const TArray<FTransform>& GlobalTransforms, const UGeometryCollectionComponent* GeometryCollectionComponent, const FColor& Color);
-
-	/** Draw vertices for the part of the geometry attached to the specified transform index. */
-	void DrawVertices(const TArray<FTransform>& GlobalTransforms, const UGeometryCollectionComponent* GeometryCollectionComponent, int32 TransformIndex, const FColor& Color);
-
-	/** Draw vertex indices. */
-	void DrawVertexIndices(const TArray<FTransform>& GlobalTransforms, const UGeometryCollectionComponent* GeometryCollectionComponent, const FColor& Color);
-
-	/** Draw vertex indices for the part of the geometry attached to the specified transform index. */
-	void DrawVertexIndices(const TArray<FTransform>& GlobalTransforms, const UGeometryCollectionComponent* GeometryCollectionComponent, int32 TransformIndex, const FColor& Color);
-
-	/** Draw vertex normals. */
-	void DrawVertexNormals(const TArray<FTransform>& GlobalTransforms, const UGeometryCollectionComponent* GeometryCollectionComponent, const FColor& Color);
-
-	/** Draw vertex normals for the part of the geometry attached to the specified transform index. */
-	void DrawVertexNormals(const TArray<FTransform>& GlobalTransforms, const UGeometryCollectionComponent* GeometryCollectionComponent, int32 TransformIndex, const FColor& Color);
-
-	/** Draw faces. */
-	void DrawFaces(const TArray<FTransform>& GlobalTransforms, const UGeometryCollectionComponent* GeometryCollectionComponent, const FColor& Color);
-
-	/** Draw faces for the part of the geometry attached to the specified transform index. */
-	void DrawFaces(const TArray<FTransform>& GlobalTransforms, const UGeometryCollectionComponent* GeometryCollectionComponent, int32 TransformIndex, const FColor& Color);
-
-	/** Draw face indices. */
-	void DrawFaceIndices(const TArray<FTransform>& GlobalTransforms, const UGeometryCollectionComponent* GeometryCollectionComponent, const FColor& Color);
-
-	/** Draw face indices for the part of the geometry attached to the specified transform index. */
-	void DrawFaceIndices(const TArray<FTransform>& GlobalTransforms, const UGeometryCollectionComponent* GeometryCollectionComponent, int32 TransformIndex, const FColor& Color);
-
-	/** Draw face normals. */
-	void DrawFaceNormals(const TArray<FTransform>& GlobalTransforms, const UGeometryCollectionComponent* GeometryCollectionComponent, const FColor& Color);
-
-	/** Draw face normals for the part of the geometry attached to the specified transform index. */
-	void DrawFaceNormals(const TArray<FTransform>& GlobalTransforms, const UGeometryCollectionComponent* GeometryCollectionComponent, int32 TransformIndex, const FColor& Color);
-
-	/** Draw single face. */
-	void DrawSingleFace(const TArray<FTransform>& GlobalTransforms, const UGeometryCollectionComponent* GeometryCollectionComponent, const int32 FaceIndex, const FColor& Color);
-
-	/** Draw geometry indices. */
-	void DrawGeometryIndices(const TArray<FTransform>& GlobalTransforms, const UGeometryCollectionComponent* GeometryCollectionComponent, const FColor& Color);
-
-	/** Draw the geometry index for the part of the geometry attached to the specified transform index. */
-	void DrawGeometryIndex(const TArray<FTransform>& GlobalTransforms, const UGeometryCollectionComponent* GeometryCollectionComponent, int32 TransformIndex, const FColor& Color);
-
-	/** Draw transforms. */
-	void DrawTransforms(const TArray<FTransform>& GlobalTransforms, const UGeometryCollectionComponent* GeometryCollectionComponent, float Scale);
-
-	/** Draw the transform for the part of the geometry attached to the specified transform index. */
-	void DrawTransform(const TArray<FTransform>& GlobalTransforms, const UGeometryCollectionComponent* GeometryCollectionComponent, int32 TransformIndex, float Scale);
-
-	/** Draw transform indices. */
-	void DrawTransformIndices(const TArray<FTransform>& GlobalTransforms, const UGeometryCollectionComponent* GeometryCollectionComponent, const FColor& Color);
-
-	/** Draw the transform index for the part of the geometry attached to the specified transform index. */
-	void DrawTransformIndex(const TArray<FTransform>& GlobalTransforms, const UGeometryCollectionComponent* GeometryCollectionComponent, int32 TransformIndex, const FColor& Color);
-
-	/** Draw hierarchical levels. */
-	void DrawLevels(const TArray<FTransform>& GlobalTransforms, const UGeometryCollectionComponent* GeometryCollectionComponent, const FColor& Color);
-
-	/** Draw the hierarchical level for the part of the geometry attached to the specified transform index. */
-	void DrawLevel(const TArray<FTransform>& GlobalTransforms, const UGeometryCollectionComponent* GeometryCollectionComponent, int32 TransformIndex, const FColor& Color);
-
-	/** Draw links to the parent. */
-	void DrawParents(const TArray<FTransform>& GlobalTransforms, const UGeometryCollectionComponent* GeometryCollectionComponent, const FColor& Color);
-
-	/** Draw a link to the parent for the part of the geometry attached to the specified transform index. */
-	void DrawParent(const TArray<FTransform>& GlobalTransforms, const UGeometryCollectionComponent* GeometryCollectionComponent, int32 TransformIndex, const FColor& Color);
-
-	/** Draw bounding boxes. */
-	void DrawBoundingBoxes(const TArray<FTransform>& GlobalTransforms, const UGeometryCollectionComponent* GeometryCollectionComponent, const FColor& Color);
-
-	/** Draw the bounding box for the part of the geometry attached to the specified transform index. */
-	void DrawBoundingBox(const TArray<FTransform>& GlobalTransforms, const UGeometryCollectionComponent* GeometryCollectionComponent, int32 TransformIndex, const FColor& Color);
-
-	/** Return the concatenated transform for the specified particle. */
-	static FTransform GetParticleTransform(const UGeometryCollectionComponent* GeometryCollectionComponent, int32 TransformIndex, const FGeometryCollectionParticlesData& ParticlesData);
-
-	/** Draw Chaos' rigid body ids. */
-	//void DrawRigidBodiesId(const UGeometryCollectionComponent* GeometryCollectionComponent,  const FGeometryCollectionParticlesData& ParticlesData, const TManagedArray<int32>& RigidBodyIdArray, const FColor& Color);
-	void DrawRigidBodiesId(const UGeometryCollectionComponent* GeometryCollectionComponent,  const FGeometryCollectionParticlesData& ParticlesData, const TManagedArray<FGuid>& RigidBodyIdArray, const FColor& Color);
-
-	/** Draw Chaos' rigid body id. */
-	//void DrawRigidBodyId(const UGeometryCollectionComponent* GeometryCollectionComponent, int32 TransformIndex, const FGeometryCollectionParticlesData& ParticlesData, const TManagedArray<int32>& RigidBodyIdArray, const FColor& Color);
-	void DrawRigidBodyId(const UGeometryCollectionComponent* GeometryCollectionComponent, int32 TransformIndex, const FGeometryCollectionParticlesData& ParticlesData, const TManagedArray<FGuid>& RigidBodyIdArray, const FColor& Color);
-
-	/** Draw Chaos' rigid body transform. */
-	void DrawRigidBodiesTransform(const UGeometryCollectionComponent* GeometryCollectionComponent, const FGeometryCollectionParticlesData& ParticlesData, float Scale);
-
-	/** Draw Chaos' single rigid body transform. */
-	void DrawRigidBodyTransform(const UGeometryCollectionComponent* GeometryCollectionComponent, int32 TransformIndex, const FGeometryCollectionParticlesData& ParticlesData, float Scale);
-
-	/** Draw Chaos' implicit collision primitives. */
-	void DrawRigidBodiesCollision(const UGeometryCollectionComponent* GeometryCollectionComponent, const FGeometryCollectionParticlesData& ParticlesData, const FColor& Color);
-
-	/** Draw Chaos' implicit single collision primitive. */
-	void DrawRigidBodyCollision(const UGeometryCollectionComponent* GeometryCollectionComponent, int32 TransformIndex, const FGeometryCollectionParticlesData& ParticlesData, const FColor& Color);
-
-	/** Draw Chaos' inertia tensors. */
-	void DrawRigidBodiesInertia(const UGeometryCollectionComponent* GeometryCollectionComponent, const FGeometryCollectionParticlesData& ParticlesData, const FColor& Color);
-
-	/** Draw Chaos' single rigid body inertia tensor. */
-	void DrawRigidBodyInertia(const UGeometryCollectionComponent* GeometryCollectionComponent, int32 TransformIndex, const FGeometryCollectionParticlesData& ParticlesData, const FColor& Color);
-
-	/** Draw Chaos' rigid body informations. */
-	void DrawRigidBodiesInfo(const UGeometryCollectionComponent* GeometryCollectionComponent, const FGeometryCollectionParticlesData& ParticlesData, const FColor& Color);
-
-	/** Draw Chaos' single rigid body informations. */
-	void DrawRigidBodyInfo(const UGeometryCollectionComponent* GeometryCollectionComponent, int32 TransformIndex, const FGeometryCollectionParticlesData& ParticlesData, const FColor& Color);
-
-	/** Draw Chaos' rigid clustering's connectivities edges */
-	//void DrawConnectivityEdges(const UGeometryCollectionComponent* GeometryCollectionComponent, const FGeometryCollectionParticlesData& ParticlesData, const TManagedArray<int32>& RigidBodyIdArray);
-	void DrawConnectivityEdges(const UGeometryCollectionComponent* GeometryCollectionComponent, const FGeometryCollectionParticlesData& ParticlesData, const TManagedArray<FGuid>& RigidBodyIdArray);
-
-	/** Draw Chaos' single rigid clustering's connectivity edges. */
-	//void DrawConnectivityEdges(const UGeometryCollectionComponent* GeometryCollectionComponent, int32 TransformIndex, const FGeometryCollectionParticlesData& ParticlesData, const TManagedArray<int32>& RigidBodyIdArray, FColor HSVColor = FColor(157, 160, 128));
-	void DrawConnectivityEdges(const UGeometryCollectionComponent* GeometryCollectionComponent, int32 TransformIndex, const FGeometryCollectionParticlesData& ParticlesData, const TManagedArray<FGuid>& RigidBodyIdArray, FColor HSVColor = FColor(157, 160, 128));
-
-	/** Draw Chaos' rigid body velocity. */
-	void DrawRigidBodiesVelocity(const UGeometryCollectionComponent* GeometryCollectionComponent, const FGeometryCollectionParticlesData& ParticlesData, const FColor& Color);
-
-	/** Draw Chaos' single rigid body velocity. */
-	void DrawRigidBodyVelocity(const UGeometryCollectionComponent* GeometryCollectionComponent, int32 TransformIndex, const FGeometryCollectionParticlesData& ParticlesData, const FColor& Color);
-
-	/** Draw Chaos' rigid body applied force and torque. */
-	void DrawRigidBodiesForce(const UGeometryCollectionComponent* GeometryCollectionComponent, const FGeometryCollectionParticlesData& ParticlesData, const FColor& Color);
-
-	/** Draw Chaos' single rigid body applied force and torque. */
-	void DrawRigidBodyForce(const UGeometryCollectionComponent* GeometryCollectionComponent, int32 TransformIndex, const FGeometryCollectionParticlesData& ParticlesData, const FColor& Color);
-
-private:
-	/** Return a darker color depending on level. */
-	static FColor MakeDarker(const FColor& Color, int32 Level = 1);
-
-	/** Return a smaller scale depending on level. */
-	static float MakeSmaller(float Scale, int32 Level = 1) { while (--Level >= 0) { Scale *= 0.666666f; } return Scale; }
-
-	/** Return the hierarchy level for this transform index. */
-	static int32 GetLevel(int32 TransformIndex, const TManagedArray<int32>& Parents);
-
-	/** Callback on property changes. */
-	void OnPropertyChanged(bool bForceVisibilityUpdate);
-
-	/** Property update function helper. */
-	template<typename T1, typename T2>
-	void UpdatePropertyValue(T1& PropertyValue, const TAutoConsoleVariable<T2>& ConsoleVariable, bool& bHasChanged);
-
-	/** Property update function helper. */
-	template<typename T, decltype(ImplicitConv<AChaosSolverActor*>(std::declval<T>()))* = nullptr>
-	void UpdatePropertyValue(T& PropertyValue, const TAutoConsoleVariable<FString>& ConsoleVariable, bool& bHasChanged);
-
-	/** Callback on console variable. */
-	void OnCVarsChanged();
-
-	/** Add debug text output. */
-	void AddDebugText(const FString& Text, const FVector& Position = FVector::ZeroVector, const FColor& Color = FColor::White, float Scale = 1.f, bool bDrawShadow = false);
-
-	/** Draw all text output. */
-	void DebugDrawText(UCanvas* Canvas, APlayerController* PlayerController);
-
-	/** Clear all persistent strings and debug lines. */
-	void Flush();
-
-	/** Return the concatenated transform for the specified particle. Must have the ParticleData X, R, and ChildToParentMap already synced. */
-	static FTransform GetParticleTransformNoChecks(const UGeometryCollectionComponent* GeometryCollectionComponent, int32 TransformIndex, const FGeometryCollectionParticlesData& ParticlesData);
-
-	/** Draw Chaos' single rigid body id without synchronization checks. */
-	//void DrawRigidBodyIdNoChecks(const UGeometryCollectionComponent* GeometryCollectionComponent, int32 TransformIndex, const FGeometryCollectionParticlesData& ParticlesData, const TManagedArray<int32>& RigidBodyIdArray, const FColor& Color);
-	void DrawRigidBodyIdNoChecks(const UGeometryCollectionComponent* GeometryCollectionComponent, int32 TransformIndex, const FGeometryCollectionParticlesData& ParticlesData, const TManagedArray<FGuid>& RigidBodyIdArray, const FColor& Color);
-
-	/** Draw Chaos' single rigid body transform without synchronization checks. */
-	void DrawRigidBodyTransformNoChecks(const UGeometryCollectionComponent* GeometryCollectionComponent, int32 TransformIndex, const FGeometryCollectionParticlesData& ParticlesData, float Scale);
-
-	/** Draw Chaos' implicit single collision primitive without synchronization checks. */
-	void DrawRigidBodyCollisionNoChecks(const UGeometryCollectionComponent* GeometryCollectionComponent, int32 TransformIndex, const FGeometryCollectionParticlesData& ParticlesData, const FColor& Color);
-
-	/** Draw Chaos' implicit inertia tensor without synchronization checks. */
-	void DrawRigidBodyInertiaNoChecks(const UGeometryCollectionComponent* GeometryCollectionComponent, int32 TransformIndex, const FGeometryCollectionParticlesData& ParticlesData, const FColor& Color);
-
-	/** Draw Chaos' single rigid body informations without synchronization checks. */
-	void DrawRigidBodyInfoNoChecks(const UGeometryCollectionComponent* GeometryCollectionComponent, int32 TransformIndex, const FGeometryCollectionParticlesData& ParticlesData, const FColor& Color);
-
-	/** Draw Chaos' single rigid clustering's connectivity edges without synchronization checks. */
-	//void DrawConnectivityEdgesNoChecks(const UGeometryCollectionComponent* GeometryCollectionComponent, int32 TransformIndex, const FGeometryCollectionParticlesData& ParticlesData, const TManagedArray<int32>& RigidBodyIdArray, const FColor& Color);
-	void DrawConnectivityEdgesNoChecks(const UGeometryCollectionComponent* GeometryCollectionComponent, int32 TransformIndex, const FGeometryCollectionParticlesData& ParticlesData, const TManagedArray<FGuid>& RigidBodyIdArray, const FColor& Color);
-
-	/** Draw Chaos' single rigid body velocity without synchronization checks. */
-	void DrawRigidBodyVelocityNoChecks(const UGeometryCollectionComponent* GeometryCollectionComponent, int32 TransformIndex, const FGeometryCollectionParticlesData& ParticlesData, const FColor& Color);
-
-	/** Draw Chaos' single rigid body applied force and torque without synchronization checks. */
-	void DrawRigidBodyForceNoChecks(const UGeometryCollectionComponent* GeometryCollectionComponent, int32 TransformIndex, const FGeometryCollectionParticlesData& ParticlesData, const FColor& Color);
-
-private:
-	struct FDebugDrawText
-	{
-		FString Text;
-		FVector Position;
-		FColor Color;
-		float Scale;
-		bool bDrawShadow;
-	};
-
-	FConsoleVariableSinkHandle ConsoleVariableSinkHandle;
-	FDelegateHandle DebugDrawTextDelegateHandle;
-	TArray<FDebugDrawText> DebugDrawTexts;
-
-	bool bNeedsDebugLinesFlush;
-#if WITH_EDITOR
-	bool bWasEditorPaused;
-#endif  // #if WITH_EDITOR
 };
diff --git a/Engine/Source/Runtime/Experimental/GeometryCollectionEngine/Public/GeometryCollection/GeometryCollectionDebugDrawComponent.h b/Engine/Source/Runtime/Experimental/GeometryCollectionEngine/Public/GeometryCollection/GeometryCollectionDebugDrawComponent.h
--- a/Engine/Source/Runtime/Experimental/GeometryCollectionEngine/Public/GeometryCollection/GeometryCollectionDebugDrawComponent.h
+++ b/Engine/Source/Runtime/Experimental/GeometryCollectionEngine/Public/GeometryCollection/GeometryCollectionDebugDrawComponent.h
@@ -22,86 +22,22 @@ namespace Chaos { template<class T, int d> class TPBDRigidParticles; }
 *   Component adding debug drawing functionality to a GeometryCollectionActor.
 *   This component is automatically added to every GeometryCollectionActor.
 */
+class UE_DEPRECATED(5.0, "Deprecated. Use normal debug draw Chaos Physics commands") UGeometryCollectionDebugDrawComponent;
 UCLASS(meta = (BlueprintSpawnableComponent), HideCategories = ("Tags", "Activation", "Cooking", "AssetUserData", "Collision"))
 class GEOMETRYCOLLECTIONENGINE_API UGeometryCollectionDebugDrawComponent : public UActorComponent
 {
 	GENERATED_UCLASS_BODY()
 
 public:
-
+	PRAGMA_DISABLE_DEPRECATION_WARNINGS
 	/** Singleton actor, containing the debug draw properties. Automatically populated at play time unless explicitly set. */
-	UPROPERTY(EditAnywhere, Category = "Debug Draw", AdvancedDisplay)
-	TObjectPtr<AGeometryCollectionDebugDrawActor> GeometryCollectionDebugDrawActor;
-
+	UPROPERTY()
+	TObjectPtr<AGeometryCollectionDebugDrawActor> GeometryCollectionDebugDrawActor_DEPRECATED;
+	PRAGMA_ENABLE_DEPRECATION_WARNINGS
+	
 	/** Level Set singleton actor, containing the Render properties. Automatically populated at play time unless explicitly set. */
 	UPROPERTY(EditAnywhere, Category = "Debug Draw", AdvancedDisplay)
 	TObjectPtr<AGeometryCollectionRenderLevelSetActor> GeometryCollectionRenderLevelSetActor;
 
 	UGeometryCollectionComponent* GeometryCollectionComponent;  // the component we are debug rendering for, set by the GeometryCollectionActor after creation
-
-	virtual void BeginPlay() override;
-	virtual void EndPlay(EEndPlayReason::Type ReasonEnd) override;
-	virtual void TickComponent(float DeltaTime, enum ELevelTick TickType, FActorComponentTickFunction *ThisTickFunction) override;
-	
-#if GEOMETRYCOLLECTION_DEBUG_DRAW
-	/** Update selection and visibility after any change in properties. Also enable/disable this component tick update. Return true if this geometry collection is selected. */
-	bool OnDebugDrawPropertiesChanged(bool bForceVisibilityUpdate);
-
-	/** Update selection and visibility after a change in cluster. Only handled when the debug drawing is active (the component is ticking). */
-	void OnClusterChanged();
-
-	/** Return whether the geometry collection rigid body id array is not completely initialized. This can happen when running the physics multithreaded. */
-	FORCEINLINE bool HasIncompleteRigidBodyIdSync() const { return bHasIncompleteRigidBodyIdSync;  }
-
-private:
-	/** Recursively compute global cluster transforms. Only gives geometry transforms for the leaf nodes, mid-level transforms are those of the clusters. */
-	void ComputeClusterTransforms(int32 Index, TArray<bool>& IsComputed, TArray<FTransform>& InOutGlobalTransforms);
-
-	/** 
-	* Compute global transforms.
-	* Unlike GeometryCollectionAlgo::GlobalMatrices(), this also calculates the correct mid-level geometry transforms and includes the actor transform.
-	*/
-	void ComputeTransforms(TArray<FTransform>& OutClusterTransforms, TArray<FTransform>& OutGeometryTransforms);
-
-	/** Geometry collection debug draw. */
-	void DebugDrawTick();
-
-	/** Update the transform index dependending on the current filter settings. */
-	void UpdateSelectedTransformIndex();
-
-	/** Return the number of faces for the given geometry (includes its children, and includes its detached children when bDebugDrawClustering is true). */
-	int32 CountFaces(int32 TransformIndex, bool bDebugDrawClustering) const;
-
-	/** Update visible array to hide the selected geometry and its children, and includes its detached children when bDebugDrawClustering is true)*/
-	void HideFaces(int32 TransformIndex, bool bDebugDrawClustering);
-
-	/** Update geometry visibility. Set bForceVisibilityUpdate to true to force the visibility array update. */
-	void UpdateGeometryVisibility(bool bForceVisibilityUpdate = false);
-
-	/** Update ticking status. */
-	void UpdateTickStatus();
-
-	/** Chaos dependent debug draw. */
-	void DebugDrawChaosTick();
-
-	/** Update level set visibility. */
-	void UpdateLevelSetVisibility();
-#endif  // #if GEOMETRYCOLLECTION_DEBUG_DRAW
-
-private:
-	static UGeometryCollectionDebugDrawComponent* RenderLevelSetOwner;
-	static int32 LastRenderedId;
-	//static FGuid LastRenderedId;
-
-#if GEOMETRYCOLLECTION_DEBUG_DRAW
-	FGeometryCollectionParticlesData ParticlesData;
-	int32 ParentCheckSum;
-	int32 SelectedRigidBodyId;
-	//FGuid SelectedRigidBodyId;
-	int32 SelectedTransformIndex;
-	int32 HiddenTransformIndex;
-	bool bWasVisible;
-	bool bHasIncompleteRigidBodyIdSync;
-	AChaosSolverActor* SelectedChaosSolver;
-#endif  // #if GEOMETRYCOLLECTION_DEBUG_DRAW
 };
