Chaos - fix edge pruning

diff --git a/Engine/Source/Runtime/Experimental/Chaos/Private/Chaos/CollisionOneShotManifolds.cpp b/Engine/Source/Runtime/Experimental/Chaos/Private/Chaos/CollisionOneShotManifolds.cpp
--- a/Engine/Source/Runtime/Experimental/Chaos/Private/Chaos/CollisionOneShotManifolds.cpp
+++ b/Engine/Source/Runtime/Experimental/Chaos/Private/Chaos/CollisionOneShotManifolds.cpp
@@ -246,7 +246,9 @@ namespace Chaos
 		}
 
 		// Remove spurious edge contact by removing edge-edge contact that are in the plane of another vertex-plane contact.
-		void PruneEdgeContactPoints(TArray<FContactPoint>& ContactPoints, const FReal Epsilon)
+		// Note: this assumes that the first shape in the contacts is the same convex shape (e.g., a box), and the second
+		// can be from any other shape(s), e.g., triangles of a non-convex shape.
+		void PruneEdgeContactPoints(TArray<FContactPoint>& ContactPoints, const FReal MaxPlaneDistance)
 		{
 			if (!bChaos_Collision_EnableEdgePrune)
 			{
@@ -265,21 +267,18 @@ namespace Chaos
 						const FContactPoint PlaneContactPoint = ContactPoints[PlaneContactIndex];
 
 						// Does the edge-edge contact lie in the plane of the vertex-plane contact...?
-						int32 PlaneShapeIndex = INDEX_NONE;
-						if (PlaneContactPoint.ContactType == EContactPointType::PlaneVertex)
-						{
-							PlaneShapeIndex = 0;
-						}
-						else if (PlaneContactPoint.ContactType == EContactPointType::VertexPlane)
-						{
-							PlaneShapeIndex = 1;
-						}
-						if (PlaneShapeIndex != INDEX_NONE)
+						// We only consider Vertex-Plane contacts. We ignore Plane-Vertex because the first
+						// shape is always the convex object.
+						if (PlaneContactPoint.ContactType == EContactPointType::VertexPlane)
 						{
+							const int32 PlaneShapeIndex = 1;
 							const FVec3 DeltaPos = PlaneContactPoint.ShapeContactPoints[PlaneShapeIndex] - EdgeContactPoint.ShapeContactPoints[PlaneShapeIndex];
-							const FReal Distance = FVec3::DotProduct(DeltaPos, PlaneContactPoint.ShapeContactNormal);
-							if (FMath::Abs(Distance) < Epsilon)
+							const FReal PlaneDistance = FVec3::DotProduct(DeltaPos, PlaneContactPoint.ShapeContactNormal);
+							if (FMath::Abs(PlaneDistance) < MaxPlaneDistance)
 							{
+								// @todo(chaos): Pretty sure we'll need to do an absolute distance check here too so we don;t reject edge
+								// contacts because of plane contacts a long way away.
+
 								// Note: Cannot use RemoveAtSwap unless we re-sort at the end. See ReduceManifoldContactPointsTriangeMesh
 								const bool bAllowShrinking = false;
 								ContactPoints.RemoveAt(EdgeContactIndex, 1, bAllowShrinking);
@@ -298,11 +297,6 @@ namespace Chaos
 		// ContactPoints are sorted on phi (ascending)
 		void ReduceManifoldContactPointsTriangeMesh(TArray<FContactPoint>& ContactPoints)
 		{
-			// Remove edge contacts that are "hidden" by face contacts
-			// EdgePruneDistance should be some fraction of the convex margin...
-			const FReal EdgePruneDistance = Chaos_Collision_EdgePrunePlaneDistance;
-			PruneEdgeContactPoints(ContactPoints, EdgePruneDistance);
-
 			if (ContactPoints.Num() <= 4)
 			{
 				return;
diff --git a/Engine/Source/Runtime/Experimental/Chaos/Private/Chaos/CollisionOneShotManifolds.h b/Engine/Source/Runtime/Experimental/Chaos/Private/Chaos/CollisionOneShotManifolds.h
--- a/Engine/Source/Runtime/Experimental/Chaos/Private/Chaos/CollisionOneShotManifolds.h
+++ b/Engine/Source/Runtime/Experimental/Chaos/Private/Chaos/CollisionOneShotManifolds.h
@@ -15,6 +15,7 @@ namespace Chaos
 	{
 		uint32 BoxBoxClipVerticesAgainstPlane(const FVec3* InputVertexBuffer, FVec3* outputVertexBuffer, uint32 ClipPointCount, int32 ClippingAxis, FReal Distance);
 		uint32 ReduceManifoldContactPoints(FVec3* Points, uint32 PointCount);
+		void PruneEdgeContactPoints(TArray<FContactPoint>& ContactPoints, const FReal MaxPlaneDistance);
 		void ReduceManifoldContactPointsTriangeMesh(TArray<FContactPoint>& ContactPoints);
 
 		void ConstructBoxBoxOneShotManifold(
diff --git a/Engine/Source/Runtime/Experimental/Chaos/Private/Chaos/HeightField.cpp b/Engine/Source/Runtime/Experimental/Chaos/Private/Chaos/HeightField.cpp
--- a/Engine/Source/Runtime/Experimental/Chaos/Private/Chaos/HeightField.cpp
+++ b/Engine/Source/Runtime/Experimental/Chaos/Private/Chaos/HeightField.cpp
@@ -17,6 +17,8 @@
 
 namespace Chaos
 {
+	extern FRealSingle Chaos_Collision_EdgePrunePlaneDistance;
+
 	int32 bOneSidedHeightField = 1;
 	static FAutoConsoleVariableRef CVarOneSidedHeightField(TEXT("p.Chaos.OneSidedHeightField"), bOneSidedHeightField, TEXT("When enabled, extra steps will ensure that FHeightField::GJKContactPointImp never results in internal-facing contact data."));
 
@@ -1489,6 +1491,11 @@ namespace Chaos
 			}
 		}
 
+		// Remove edge contacts that are "hidden" by face contacts
+		// EdgePruneDistance should be some fraction of the convex margin...
+		const FReal EdgePruneDistance = Chaos_Collision_EdgePrunePlaneDistance;
+		Collisions::PruneEdgeContactPoints(ContactPoints, EdgePruneDistance);
+
 		// Remove all points (except for the deepest one, and ones with phis similar to it)
 		const FReal CullMargin = 0.1f;
 		int32 NewContactPointCount = ContactPoints.Num() > 0 ? 1 : 0;
diff --git a/Engine/Source/Runtime/Experimental/Chaos/Private/Chaos/PBDCollisionConstraints.cpp b/Engine/Source/Runtime/Experimental/Chaos/Private/Chaos/PBDCollisionConstraints.cpp
--- a/Engine/Source/Runtime/Experimental/Chaos/Private/Chaos/PBDCollisionConstraints.cpp
+++ b/Engine/Source/Runtime/Experimental/Chaos/Private/Chaos/PBDCollisionConstraints.cpp
@@ -677,52 +677,62 @@ namespace Chaos
 
 		const FReal EdgePlaneTolerance = Chaos_Collision_EdgePrunePlaneDistance;
 
-		// Loop over edge collisions, then all plane collisions
-		// and remove the edge collision if it is hidden by a plane collision
+		// Loop over edge collisions, then all plane collisions and remove the edge collision if it is hidden by a plane collision
+		// NOTE: We only look at plane collisions where the other shape owns the plane.
 		// @todo(chaos): this should probably only disable individual manifold points
+		// @todo(chaos): we should probably only reject edges if the plane contact is also close to the edge contact
+		// @todo(chaos): we should also try to eliminate face contacts from sub-surface faces
 		// @todo(chaos): perf issue: this processes contacts in world space, but we don't calculated that data until Gather. Fix this.
 		ParticleCollision.VisitCollisions(
-			[&ParticleCollision, EdgePlaneTolerance](FPBDCollisionConstraint& EdgeCollision)
+			[Particle, &ParticleCollision, EdgePlaneTolerance](FPBDCollisionConstraint& EdgeCollision)
 			{
 				if (EdgeCollision.IsEnabled())
 				{
+					const int32 EdgeOtherShapeIndex = (EdgeCollision.GetParticle0() == Particle) ? 1 : 0;
+
 					for (const FManifoldPoint& EdgeManifoldPoint : EdgeCollision.GetManifoldPoints())
 					{
 						if (EdgeManifoldPoint.ContactPoint.ContactType == EContactPointType::EdgeEdge)
 						{
-							const FVec3 EdgePos0 = EdgeCollision.GetShapeWorldTransform0().TransformPositionNoScale(EdgeManifoldPoint.ContactPoint.ShapeContactPoints[0]);
-							const FVec3 EdgePos1 = EdgeCollision.GetShapeWorldTransform1().TransformPositionNoScale(EdgeManifoldPoint.ContactPoint.ShapeContactPoints[1]);
+							const FRigidTransform3& EdgeTransform = (EdgeOtherShapeIndex == 0) ? EdgeCollision.GetShapeWorldTransform0() : EdgeCollision.GetShapeWorldTransform1();
+							const FVec3 EdgePos = EdgeTransform.TransformPositionNoScale(EdgeManifoldPoint.ContactPoint.ShapeContactPoints[EdgeOtherShapeIndex]);
 
 							// Loop over plane collisions
 							ECollisionVisitorResult PlaneResult = ParticleCollision.VisitConstCollisions(
-								[&EdgeCollision, &EdgePos0, &EdgePos1, EdgePlaneTolerance](const FPBDCollisionConstraint& PlaneCollision)
+								[Particle, &EdgeCollision, &EdgePos, EdgePlaneTolerance](const FPBDCollisionConstraint& PlaneCollision)
 								{
 									if ((&PlaneCollision != &EdgeCollision) && PlaneCollision.IsEnabled())
 									{
+										const int32 PlaneOtherShapeIndex = (PlaneCollision.GetParticle0() == Particle) ? 1 : 0;
+
 										for (const FManifoldPoint& PlaneManifoldPoint : PlaneCollision.GetManifoldPoints())
 										{
 											// If the edge position is in the plane, disable it
 											FVec3 PlanePos;
 											FVec3 PlaneNormal;
-											if (PlaneManifoldPoint.ContactPoint.ContactType == EContactPointType::PlaneVertex)
+											bool bHavePlane = false;
+											if ((PlaneOtherShapeIndex == 0) && (PlaneManifoldPoint.ContactPoint.ContactType == EContactPointType::PlaneVertex))
 											{
 												PlanePos = PlaneCollision.GetShapeWorldTransform0().TransformPositionNoScale(PlaneManifoldPoint.ContactPoint.ShapeContactPoints[0]);
-												PlaneNormal = PlaneCollision.GetShapeWorldTransform0().TransformVectorNoScale(PlaneManifoldPoint.ContactPoint.ShapeContactNormal);
+												bHavePlane = true;
 											}
-											else if (PlaneManifoldPoint.ContactPoint.ContactType == EContactPointType::VertexPlane)
+											else if ((PlaneOtherShapeIndex == 1) && (PlaneManifoldPoint.ContactPoint.ContactType == EContactPointType::VertexPlane))
 											{
 												PlanePos = PlaneCollision.GetShapeWorldTransform1().TransformPositionNoScale(PlaneManifoldPoint.ContactPoint.ShapeContactPoints[1]);
-												PlaneNormal = PlaneCollision.GetShapeWorldTransform1().TransformVectorNoScale(PlaneManifoldPoint.ContactPoint.ShapeContactNormal);
+												bHavePlane = true;
 											}
-											const FVec3 EdgePlaneDelta0 = EdgePos0 - PlanePos;
-											const FVec3 EdgePlaneDelta1 = EdgePos1 - PlanePos;
-											const FReal EdgePlaneDistance0 = FVec3::DotProduct(EdgePlaneDelta0, PlaneNormal);
-											const FReal EdgePlaneDistance1 = FVec3::DotProduct(EdgePlaneDelta1, PlaneNormal);
-											if ((FMath::Abs(EdgePlaneDistance0) < EdgePlaneTolerance) || (FMath::Abs(EdgePlaneDistance1) < EdgePlaneTolerance))
+											PlaneNormal = PlaneCollision.GetShapeWorldTransform1().TransformVectorNoScale(PlaneManifoldPoint.ContactPoint.ShapeContactNormal);
+
+											if (bHavePlane)
 											{
-												// The edge contact is hidden by a plane contact so disable it and stop the inner loop
-												EdgeCollision.SetDisabled(true);
-												return ECollisionVisitorResult::Stop;
+												const FVec3 EdgePlaneDelta = EdgePos - PlanePos;
+												const FReal EdgePlaneDistance = FVec3::DotProduct(EdgePlaneDelta, PlaneNormal);
+												if (FMath::Abs(EdgePlaneDistance) < EdgePlaneTolerance)
+												{
+													// The edge contact is hidden by a plane contact so disable it and stop the inner loop
+													EdgeCollision.SetDisabled(true);
+													return ECollisionVisitorResult::Stop;
+												}
 											}
 										}
 									}
diff --git a/Engine/Source/Runtime/Experimental/Chaos/Private/Chaos/TriangleMeshImplicitObject.cpp b/Engine/Source/Runtime/Experimental/Chaos/Private/Chaos/TriangleMeshImplicitObject.cpp
--- a/Engine/Source/Runtime/Experimental/Chaos/Private/Chaos/TriangleMeshImplicitObject.cpp
+++ b/Engine/Source/Runtime/Experimental/Chaos/Private/Chaos/TriangleMeshImplicitObject.cpp
@@ -15,6 +15,7 @@
 
 namespace Chaos
 {
+	extern FRealSingle Chaos_Collision_EdgePrunePlaneDistance;
 
 	// Note that if this is re-enabled when previously off, the cooked trimeshes won't have the vertex map serialized, so the change will not take effect until re-cooked.
 	bool TriMeshPerPolySupport = 1;
@@ -437,6 +438,11 @@ bool FTriangleMeshImplicitObject::ContactManifoldImp(const GeomType& QueryGeom,
 
 		}
 
+		// Remove edge contacts that are "hidden" by face contacts
+		// EdgePruneDistance should be some fraction of the convex margin...
+		const FReal EdgePruneDistance = Chaos_Collision_EdgePrunePlaneDistance;
+		Collisions::PruneEdgeContactPoints(ContactPoints, EdgePruneDistance);
+
 		// Remove all points (except for the deepest one, and ones with phis similar to it)
 		const FReal CullMargin = 0.1f;
 		int32 NewContactPointCount = ContactPoints.Num() > 0 ? 1 : 0;
