make convex calculation compute proximity data if it's missing

#rnx

diff --git a/Engine/Source/Runtime/Experimental/Chaos/Private/GeometryCollection/GeometryCollectionConvexUtility.cpp b/Engine/Source/Runtime/Experimental/Chaos/Private/GeometryCollection/GeometryCollectionConvexUtility.cpp
--- a/Engine/Source/Runtime/Experimental/Chaos/Private/GeometryCollection/GeometryCollectionConvexUtility.cpp
+++ b/Engine/Source/Runtime/Experimental/Chaos/Private/GeometryCollection/GeometryCollectionConvexUtility.cpp
@@ -5,6 +5,7 @@
 #include "Chaos/GJK.h"
 #include "GeometryCollection/GeometryCollection.h"
 #include "GeometryCollection/GeometryCollectionAlgo.h"
+#include "GeometryCollection/GeometryCollectionProximityUtility.h"
 #include "CompGeom/ConvexHull3.h"
 #include "Templates/Sorting.h"
 #include "Spatial/PointHashGrid3.h"
@@ -1067,6 +1068,11 @@ FGeometryCollectionConvexUtility::FGeometryCollectionConvexData FGeometryCollect
 	TFunctionRef<bool(int32)> HasCustomConvexFn = (CustomConvexFlags != nullptr) ? (TFunctionRef<bool(int32)>)ConvexFlagsFromArray : (TFunctionRef<bool(int32)>)ConvexFlagsAlwaysFalse;
 
 	bool bHasProximity = GeometryCollection->HasAttribute("Proximity", FGeometryCollection::GeometryGroup);
+	if (!bHasProximity)
+	{
+		FGeometryCollectionProximityUtility ProximityUtility(GeometryCollection);
+		ProximityUtility.UpdateProximity();
+	}
 	const TManagedArray<TSet<int32>>* GCProximity = GeometryCollection->FindAttribute<TSet<int32>>("Proximity", FGeometryCollection::GeometryGroup);
 	const TManagedArray<float>* Volume = GeometryCollection->FindAttribute<float>("Volume", FGeometryCollection::TransformGroup);
 	TArray<TUniquePtr<Chaos::FConvex>> Convexes;
