Chaos : Geometry Collection
-Defaults Updated

diff --git a/Engine/Plugins/Experimental/GeometryCollectionPlugin/Source/GeometryCollectionEditor/Private/GeometryCollection/ActorFactoryGeometryCollection.cpp b/Engine/Plugins/Experimental/GeometryCollectionPlugin/Source/GeometryCollectionEditor/Private/GeometryCollection/ActorFactoryGeometryCollection.cpp
--- a/Engine/Plugins/Experimental/GeometryCollectionPlugin/Source/GeometryCollectionEditor/Private/GeometryCollection/ActorFactoryGeometryCollection.cpp
+++ b/Engine/Plugins/Experimental/GeometryCollectionPlugin/Source/GeometryCollectionEditor/Private/GeometryCollection/ActorFactoryGeometryCollection.cpp
@@ -53,7 +53,6 @@ void UActorFactoryGeometryCollection::PostSpawnActor(UObject* Asset, AActor* New
 	NewGeometryCollectionActor->GetGeometryCollectionComponent()->EnableClustering = GeometryCollection->EnableClustering;
 	NewGeometryCollectionActor->GetGeometryCollectionComponent()->ClusterGroupIndex = GeometryCollection->ClusterGroupIndex;
 	NewGeometryCollectionActor->GetGeometryCollectionComponent()->MaxClusterLevel = GeometryCollection->MaxClusterLevel;
-	NewGeometryCollectionActor->GetGeometryCollectionComponent()->DamageThreshold = GeometryCollection->DamageThreshold;
 
 	// Init Component
 	NewGeometryCollectionActor->GetGeometryCollectionComponent()->RegisterComponent();
diff --git a/Engine/Source/Runtime/Experimental/Chaos/Public/GeometryCollection/GeometryCollectionSimulationCoreTypes.h b/Engine/Source/Runtime/Experimental/Chaos/Public/GeometryCollection/GeometryCollectionSimulationCoreTypes.h
--- a/Engine/Source/Runtime/Experimental/Chaos/Public/GeometryCollection/GeometryCollectionSimulationCoreTypes.h
+++ b/Engine/Source/Runtime/Experimental/Chaos/Public/GeometryCollection/GeometryCollectionSimulationCoreTypes.h
@@ -67,7 +67,7 @@ struct FSharedSimulationSizeSpecificData
 	FSharedSimulationSizeSpecificData()
 		: MaxSize(0.f)
 		, CollisionShapesData({ FCollectionCollisionTypeData() })
-		, DamageThreshold(250.f)
+		, DamageThreshold(5000.f)
 	{
 	}
 
@@ -175,7 +175,7 @@ struct FSimulationParameters
 		, ClusterGroupIndex(0)
 		, MaxClusterLevel(100)
 		, bUseSizeSpecificDamageThresholds(false)
-		, DamageThreshold({250.f})
+		, DamageThreshold({500000.f, 50000.f, 5000.f})
 		, ClusterConnectionMethod(Chaos::FClusterCreationParameters::EConnectionMethod::PointImplicit)
 		, CollisionGroup(0)
 		, CollisionSampleFraction(1.0)
diff --git a/Engine/Source/Runtime/Experimental/GeometryCollectionEngine/Private/GeometryCollection/GeometryCollectionComponent.cpp b/Engine/Source/Runtime/Experimental/GeometryCollectionEngine/Private/GeometryCollection/GeometryCollectionComponent.cpp
--- a/Engine/Source/Runtime/Experimental/GeometryCollectionEngine/Private/GeometryCollection/GeometryCollectionComponent.cpp
+++ b/Engine/Source/Runtime/Experimental/GeometryCollectionEngine/Private/GeometryCollection/GeometryCollectionComponent.cpp
@@ -230,7 +230,7 @@ UGeometryCollectionComponent::UGeometryCollectionComponent(const FObjectInitiali
 	, EnableClustering(true)
 	, ClusterGroupIndex(0)
 	, MaxClusterLevel(100)
-	, DamageThreshold({250.0})
+	, DamageThreshold({ 500000.f, 50000.f, 5000.f })
 	, bUseSizeSpecificDamageThreshold(false)
 	, ClusterConnectionType_DEPRECATED(EClusterConnectionTypeEnum::Chaos_MinimalSpanningSubsetDelaunayTriangulation)
 	, CollisionGroup(0)
diff --git a/Engine/Source/Runtime/Experimental/GeometryCollectionEngine/Private/GeometryCollection/GeometryCollectionObject.cpp b/Engine/Source/Runtime/Experimental/GeometryCollectionEngine/Private/GeometryCollection/GeometryCollectionObject.cpp
--- a/Engine/Source/Runtime/Experimental/GeometryCollectionEngine/Private/GeometryCollection/GeometryCollectionObject.cpp
+++ b/Engine/Source/Runtime/Experimental/GeometryCollectionEngine/Private/GeometryCollection/GeometryCollectionObject.cpp
@@ -63,7 +63,7 @@ UGeometryCollection::UGeometryCollection(const FObjectInitializer& ObjectInitial
 	, EnableClustering(true)
 	, ClusterGroupIndex(0)
 	, MaxClusterLevel(100)
-	, DamageThreshold({ 250.0 })
+	, DamageThreshold_DEPRECATED({ 500000.f, 50000.f, 5000.f })
 	, ClusterConnectionType(EClusterConnectionTypeEnum::Chaos_MinimalSpanningSubsetDelaunayTriangulation)
 	, bUseFullPrecisionUVs(false)
 	, bStripOnCook(false)
@@ -77,12 +77,12 @@ UGeometryCollection::UGeometryCollection(const FObjectInitializer& ObjectInitial
 	, MaxClusterLevelSetResolution_DEPRECATED(50)
 	, CollisionObjectReductionPercentage_DEPRECATED(0.0f)
 #endif
-	, bMassAsDensity(false)
-	, Mass(1.0f)
+	, bMassAsDensity(true)
+	, Mass(2500.0f)
 	, MinimumMassClamp(0.1f)
 	, bRemoveOnMaxSleep(false)
-	, MaximumSleepTime(10.0)
-	, RemovalDuration(5.0)
+	, MaximumSleepTime(5.0, 10.0)
+	, RemovalDuration(2.5, 5.0)
 	, EnableRemovePiecesOnFracture(false)
 	, GeometryCollection(new FGeometryCollection())
 {
@@ -121,7 +121,7 @@ FGeometryCollectionCollisionTypeData::FGeometryCollectionCollisionTypeData()
 }
 
 FGeometryCollectionSizeSpecificData::FGeometryCollectionSizeSpecificData()
-	: MaxSize(FLT_MAX)
+	: MaxSize(99999.9)
 	, CollisionShapes({ FGeometryCollectionCollisionTypeData()})
 #if WITH_EDITORONLY_DATA
 	, CollisionType_DEPRECATED(ECollisionTypeEnum::Chaos_Volumetric)
@@ -134,7 +134,7 @@ FGeometryCollectionSizeSpecificData::FGeometryCollectionSizeSpecificData()
 	, CollisionParticlesFraction_DEPRECATED(1.f)
 	, MaximumCollisionParticles_DEPRECATED(60)
 #endif
-	, DamageThreshold(250.0)
+	, DamageThreshold(5000.0)
 {
 }
 
@@ -208,7 +208,7 @@ FGeometryCollectionSizeSpecificData UGeometryCollection::GeometryCollectionSizeS
 {
 	FGeometryCollectionSizeSpecificData Data;
 
-	Data.MaxSize = FLT_MAX;
+	Data.MaxSize = 99999.9;
 	if (Data.CollisionShapes.Num())
 	{
 		Data.CollisionShapes[0].CollisionType = ECollisionTypeEnum::Chaos_Volumetric;
@@ -222,7 +222,7 @@ FGeometryCollectionSizeSpecificData UGeometryCollection::GeometryCollectionSizeS
 		Data.CollisionShapes[0].CollisionParticles.CollisionParticlesFraction = 1.0;
 		Data.CollisionShapes[0].CollisionParticles.MaximumCollisionParticles = 60;
 	}
-	Data.DamageThreshold = 250.0f;
+	Data.DamageThreshold = 5000.0f;
 	return Data;
 }
 
@@ -233,7 +233,7 @@ void UGeometryCollection::ValidateSizeSpecificDataDefaults()
 	{
 		for (const FGeometryCollectionSizeSpecificData& Data : DatasIn)
 		{
-			if (FMath::IsNearlyEqual(Data.MaxSize, FLT_MAX))
+			if (Data.MaxSize >= 99999.9)
 			{
 				return true;
 			}
@@ -322,7 +322,7 @@ void UGeometryCollection::GetSharedSimulationParams(FSharedSimulationParameters&
 		InfSize.CollisionShapes[0].CollisionParticles.CollisionParticlesFraction = SizeSpecificDefault.CollisionShapes[0].CollisionParticles.CollisionParticlesFraction;
 		InfSize.CollisionShapes[0].CollisionParticles.MaximumCollisionParticles = SizeSpecificDefault.CollisionShapes[0].CollisionParticles.MaximumCollisionParticles;
 	}
-	InfSize.MaxSize = TNumericLimits<float>::Max();
+	InfSize.MaxSize = 99999.9;
 	OutParams.SizeSpecificData.SetNum(SizeSpecificData.Num() + 1);
 	FillSharedSimulationSizeSpecificData(OutParams.SizeSpecificData[0], InfSize);
 
@@ -1122,30 +1122,65 @@ void UGeometryCollection::PostEditChangeProperty(struct FPropertyChangedEvent& P
 {
 	if (PropertyChangedEvent.Property)
 	{
+		FName PropertyName = PropertyChangedEvent.Property->GetFName();
+
+		bool bDoInvalidateCollection = false;
+		bool bDoEnsureDataIsCooked = false;
+		bool bValidateSizeSpecificDataDefaults = false;
+		bool bDoUpdateConvexGeometry = false;
 		bool bRebuildSimulationData = false;
 
 		if (PropertyChangedEvent.Property->GetFName() == GET_MEMBER_NAME_CHECKED(UGeometryCollection, EnableNanite))
 		{
-			InvalidateCollection();
-			EnsureDataIsCooked();
+			bDoInvalidateCollection = true;
+			bDoEnsureDataIsCooked = true;
 		}
 		else if (PropertyChangedEvent.Property->GetFName() == GET_MEMBER_NAME_CHECKED(UGeometryCollection, bUseFullPrecisionUVs))
 		{
-			InvalidateCollection();
+			bDoInvalidateCollection = true;
 		}
 		else if (PropertyChangedEvent.Property->GetFName() == GET_MEMBER_NAME_CHECKED(UGeometryCollection, SizeSpecificData))
 		{
-			ValidateSizeSpecificDataDefaults();
-			InvalidateCollection();
+			bDoInvalidateCollection = true;
+			bDoUpdateConvexGeometry = true;
+			bValidateSizeSpecificDataDefaults = true;
 			bRebuildSimulationData = true;
+		}		
+		else if (PropertyName.ToString().Contains(FString("ImplicitType")))
+			//SizeSpecificData.Num() && SizeSpecificData[0].CollisionShapes.Num() &&
+			//	PropertyChangedEvent.Property->GetFName() == GET_MEMBER_NAME_CHECKED(UGeometryCollection, SizeSpecificData[0].CollisionShapes[0].ImplicitType))
+		{
+				bDoInvalidateCollection = true;
+				bDoUpdateConvexGeometry = true;
+				bRebuildSimulationData = true;
 		}
 		else if (PropertyChangedEvent.Property->GetFName() != GET_MEMBER_NAME_CHECKED(UGeometryCollection, Materials))
 		{
-			InvalidateCollection();
+			bDoInvalidateCollection = true;
 			bRebuildSimulationData = true;
 		}
 
 
+		if (bDoInvalidateCollection)
+		{
+			InvalidateCollection();
+		}
+
+		if (bValidateSizeSpecificDataDefaults)
+		{
+			ValidateSizeSpecificDataDefaults();
+		}
+
+		if (bDoUpdateConvexGeometry)
+		{
+			UpdateConvexGeometry();
+		}
+
+		if (bDoEnsureDataIsCooked)
+		{
+			EnsureDataIsCooked();
+		}
+
 		if (bRebuildSimulationData)
 		{
 			if (!bManualDataCreate)
diff --git a/Engine/Source/Runtime/Experimental/GeometryCollectionEngine/Public/GeometryCollection/GeometryCollectionObject.h b/Engine/Source/Runtime/Experimental/GeometryCollectionEngine/Public/GeometryCollection/GeometryCollectionObject.h
--- a/Engine/Source/Runtime/Experimental/GeometryCollectionEngine/Public/GeometryCollection/GeometryCollectionObject.h
+++ b/Engine/Source/Runtime/Experimental/GeometryCollectionEngine/Public/GeometryCollection/GeometryCollectionObject.h
@@ -424,8 +424,8 @@ public:
 	int32 MaxClusterLevel;
 
 	/** Damage threshold for clusters at different levels. */
-	UPROPERTY(EditAnywhere, Category = "Clustering")
-	TArray<float> DamageThreshold;
+	UPROPERTY(meta = (DeprecatedProperty, DeprecationMessage = "This property is deprecated. Use the default SizeSpecificData instead."))
+	TArray<float> DamageThreshold_DEPRECATED;
 
 	/** */
 	UPROPERTY(EditAnywhere, Category = "Clustering")
@@ -540,7 +540,7 @@ public:
 	bool bRemoveOnMaxSleep;
 	
 	/** How long may the particle sleep before initiating removal (in seconds). */
-	UPROPERTY(EditAnywhere, BlueprintReadOnly, Category = Removal)
+	UPROPERTY(EditAnywhere, BlueprintReadOnly, Category = Removal, DisplayName = "Sleep Min Max")
 	FVector2D MaximumSleepTime;
 
 	/** How long does the removal process take (in seconds). */
