fix fracture reset sometimes highlighting bones that are not selected

#rnx

diff --git a/Engine/Plugins/Experimental/ChaosEditor/Source/FractureEditor/Private/FractureToolGenerators.cpp b/Engine/Plugins/Experimental/ChaosEditor/Source/FractureEditor/Private/FractureToolGenerators.cpp
--- a/Engine/Plugins/Experimental/ChaosEditor/Source/FractureEditor/Private/FractureToolGenerators.cpp
+++ b/Engine/Plugins/Experimental/ChaosEditor/Source/FractureEditor/Private/FractureToolGenerators.cpp
@@ -450,7 +450,7 @@ void UFractureToolResetAsset::Execute(TWeakPtr<FFractureEditorModeToolkit> InToo
 		
 		GeometryCollectionComponent->InitializeEmbeddedGeometry();
 		
-		FScopedColorEdit EditBoneColor = GeometryCollectionComponent->EditBoneSelection();
+		FScopedColorEdit EditBoneColor = GeometryCollectionComponent->EditBoneSelection(true);
 		EditBoneColor.ResetBoneSelection();
 		EditBoneColor.ResetHighlightedBones();
 	}
diff --git a/Engine/Source/Runtime/Experimental/GeometryCollectionEngine/Public/GeometryCollection/GeometryCollectionComponent.h b/Engine/Source/Runtime/Experimental/GeometryCollectionEngine/Public/GeometryCollection/GeometryCollectionComponent.h
--- a/Engine/Source/Runtime/Experimental/GeometryCollectionEngine/Public/GeometryCollection/GeometryCollectionComponent.h
+++ b/Engine/Source/Runtime/Experimental/GeometryCollectionEngine/Public/GeometryCollection/GeometryCollectionComponent.h
@@ -392,7 +392,7 @@ public:
 	FORCEINLINE const UGeometryCollection* GetRestCollection() const { return RestCollection; }
 	FORCEINLINE FGeometryCollectionEdit EditRestCollection(GeometryCollection::EEditUpdate EditUpdate = GeometryCollection::EEditUpdate::RestPhysics, bool bShapeIsUnchanged = false) { return FGeometryCollectionEdit(this, EditUpdate, bShapeIsUnchanged); }
 #if WITH_EDITOR
-	FORCEINLINE FScopedColorEdit EditBoneSelection() { return FScopedColorEdit(this); }
+	FORCEINLINE FScopedColorEdit EditBoneSelection(bool bForceUpdate = false) { return FScopedColorEdit(this, bForceUpdate); }
 
 	/** Propagate bone selection to embedded geometry components. */
 	void SelectEmbeddedGeometry();
