UE-140719: Improve collision detection cache efficiency

diff --git a/Engine/Source/Runtime/Experimental/Chaos/Private/Chaos/Collision/ParticlePairMidPhase.cpp b/Engine/Source/Runtime/Experimental/Chaos/Private/Chaos/Collision/ParticlePairMidPhase.cpp
--- a/Engine/Source/Runtime/Experimental/Chaos/Private/Chaos/Collision/ParticlePairMidPhase.cpp
+++ b/Engine/Source/Runtime/Experimental/Chaos/Private/Chaos/Collision/ParticlePairMidPhase.cpp
@@ -633,16 +633,16 @@ namespace Chaos
 	FParticlePairMidPhase::FParticlePairMidPhase()
 		: Particle0(nullptr)
 		, Particle1(nullptr)
+		, CullDistanceScale(1)
+		, Flags()
 		, Key()
-		, ShapePairDetectors()
-		, MultiShapePairDetectors()
 		, CollisionAllocator(nullptr)
-		, Flags()
 		, LastUsedEpoch(INDEX_NONE)
 		, NumActiveConstraints(0)
 		, ParticleCollisionsIndex0(INDEX_NONE)
 		, ParticleCollisionsIndex1(INDEX_NONE)
-		, CullDistanceScale(1)
+		, MultiShapePairDetectors()
+		, ShapePairDetectors()
 	{
 	}
 
@@ -792,7 +792,7 @@ namespace Chaos
 		const FReal MaxBoundsSize1 = bIsDynamic1 ? Particle1->LocalBounds().Extents().GetMax() : FReal(0);
 		const FReal CullDistanceScale0 = MaxBoundsSize0 * CullDistanceReferenceSizeInv;
 		const FReal CullDistanceScale1 = MaxBoundsSize1 * CullDistanceReferenceSizeInv;
-		CullDistanceScale = FMath::Max3(CullDistanceScale0, CullDistanceScale1, MinCullDistanceScale);
+		CullDistanceScale = (FRealSingle) FMath::Max3(CullDistanceScale0, CullDistanceScale1, MinCullDistanceScale);
 	}
 
 	void FParticlePairMidPhase::GenerateCollisions(
diff --git a/Engine/Source/Runtime/Experimental/Chaos/Public/Chaos/Collision/NarrowPhase.h b/Engine/Source/Runtime/Experimental/Chaos/Public/Chaos/Collision/NarrowPhase.h
--- a/Engine/Source/Runtime/Experimental/Chaos/Public/Chaos/Collision/NarrowPhase.h
+++ b/Engine/Source/Runtime/Experimental/Chaos/Public/Chaos/Collision/NarrowPhase.h
@@ -44,6 +44,17 @@ namespace Chaos
 			}
 		}
 
+		FParticlePairMidPhase* GetParticlePairMidPhase(TGeometryParticleHandle<FReal, 3>* Particle0, TGeometryParticleHandle<FReal, 3>* Particle1)
+		{
+			return CollisionAllocator->GetParticlePairMidPhase(Particle0, Particle1);
+		}
+
+		// Use this function if a Mid phase pair is already allocated
+		void GenerateCollisions(FReal Dt, FParticlePairMidPhase* MidPhase, const bool bForceDisableCCD)
+		{
+			Context.bForceDisableCCD = bForceDisableCCD;
+			MidPhase->GenerateCollisions(GetBoundsExpansion(), Dt, GetContext());
+		}
 	private:
 		FCollisionContext Context;
 		FCollisionConstraintAllocator* CollisionAllocator;
diff --git a/Engine/Source/Runtime/Experimental/Chaos/Public/Chaos/Collision/ParticlePairMidPhase.h b/Engine/Source/Runtime/Experimental/Chaos/Public/Chaos/Collision/ParticlePairMidPhase.h
--- a/Engine/Source/Runtime/Experimental/Chaos/Public/Chaos/Collision/ParticlePairMidPhase.h
+++ b/Engine/Source/Runtime/Experimental/Chaos/Public/Chaos/Collision/ParticlePairMidPhase.h
@@ -280,6 +280,12 @@ namespace Chaos
 
 		inline bool IsValid() const { return (Particle0 != nullptr) && (Particle1 != nullptr); }
 
+		// Prefetch the memory in this class for later use
+		inline void CachePrefetch()
+		{
+			FPlatformMisc::PrefetchBlock(this, sizeof(*this));
+		}
+
 		/**
 		 * @brief Have we run collision detection since this Epoch (inclusive)
 		*/
@@ -400,14 +406,10 @@ namespace Chaos
 
 		void InitThresholds();
 
-
-		FGeometryParticleHandle* Particle0;
-		FGeometryParticleHandle* Particle1;
-		FCollisionParticlePairKey Key;
-
-		TArray<FSingleShapePairCollisionDetector, TInlineAllocator<1>> ShapePairDetectors;
-		TArray<FMultiShapePairCollisionDetector> MultiShapePairDetectors;
-		FCollisionConstraintAllocator* CollisionAllocator;
+		FGeometryParticleHandle* Particle0; // 8 bytes
+		FGeometryParticleHandle* Particle1; // 8 bytes
+		// A number based on the size of the dynamic objects used to scale cull distance
+		FRealSingle CullDistanceScale; // 4 bytes
 
 		union FFlags
 		{
@@ -419,15 +421,19 @@ namespace Chaos
 				uint32 bIsSleeping : 1;
 			};
 			uint32 Bits;
-		} Flags;
-		int32 LastUsedEpoch;
-		int32 NumActiveConstraints;
+		} Flags; // 4 bytes
+
+		FCollisionParticlePairKey Key; //8 bytes
+		FCollisionConstraintAllocator* CollisionAllocator; // 8 bytes
+
+		int32 LastUsedEpoch;  //4 bytes
+		int32 NumActiveConstraints; // 4 bytes
 
 		// Indices into the arrays of collisions on the particles. This is a cookie for use by FParticleCollisions
-		int32 ParticleCollisionsIndex0;
-		int32 ParticleCollisionsIndex1;
+		int32 ParticleCollisionsIndex0; // 4 bytes
+		int32 ParticleCollisionsIndex1; // 4 bytes
 
-		// A number based on the size of the dynamic objects used to scale cull distance
-		FReal CullDistanceScale;
+		TArray<FMultiShapePairCollisionDetector> MultiShapePairDetectors; // 16 bytes
+		TArray<FSingleShapePairCollisionDetector, TInlineAllocator<1>> ShapePairDetectors; //88 bytes
 	};
 }
diff --git a/Engine/Source/Runtime/Experimental/Chaos/Public/Chaos/Collision/SpatialAccelerationBroadPhase.h b/Engine/Source/Runtime/Experimental/Chaos/Public/Chaos/Collision/SpatialAccelerationBroadPhase.h
--- a/Engine/Source/Runtime/Experimental/Chaos/Public/Chaos/Collision/SpatialAccelerationBroadPhase.h
+++ b/Engine/Source/Runtime/Experimental/Chaos/Public/Chaos/Collision/SpatialAccelerationBroadPhase.h
@@ -237,6 +237,7 @@ namespace Chaos
 				SCOPE_CYCLE_COUNTER(STAT_Collisions_Filtering);
 				const int32 NumPotentials = PotentialIntersections.Num();
 				int32 NumIntoNarrowPhase = 0;
+				TArray<FParticlePairMidPhase*> MidPhasePairs;
 				for (int32 i = 0; i < NumPotentials; ++i)
 				{
 					auto& Particle2 = *PotentialIntersections[i].GetGeometryParticleHandle_PhysicsThread();
@@ -372,8 +373,24 @@ namespace Chaos
 					{
 						Swap(ParticleA, ParticleB);
 					}
-					
-					NarrowPhase.GenerateCollisions(Dt, ParticleA, ParticleB, false);
+
+					FParticlePairMidPhase* MidPhase = NarrowPhase.GetParticlePairMidPhase(ParticleA, ParticleB);
+					if (MidPhase)
+					{
+						MidPhasePairs.Add(MidPhase);
+					}
+				}
+
+				for (int32 Index = 0; Index < MidPhasePairs.Num(); Index++)
+				{
+					// Prefetch next pair
+					if (Index != MidPhasePairs.Num() - 1)
+					{
+						MidPhasePairs[Index + 1]->CachePrefetch();
+					}
+
+					FParticlePairMidPhase* MidPhasePair = MidPhasePairs[Index];
+					NarrowPhase.GenerateCollisions(Dt, MidPhasePair, false);
 				}
 
 				PHYSICS_CSV_CUSTOM_EXPENSIVE(PhysicsCounters, NumFromBroadphase, NumPotentials, ECsvCustomStatOp::Accumulate);
