Chaos Cloth - Move back the post LWC solver to float for performance reason, since cloth is being simulated in local space.

diff --git a/Engine/Source/Runtime/Experimental/Chaos/Private/Chaos/GraphColoring.cpp b/Engine/Source/Runtime/Experimental/Chaos/Private/Chaos/GraphColoring.cpp
--- a/Engine/Source/Runtime/Experimental/Chaos/Private/Chaos/GraphColoring.cpp
+++ b/Engine/Source/Runtime/Experimental/Chaos/Private/Chaos/GraphColoring.cpp
@@ -3,9 +3,8 @@
 #include "Chaos/Array.h"
 #include "ChaosLog.h"
 
-using namespace Chaos;
-
-bool VerifyGraph(TArray<TArray<int32>> ColorGraph, const TArray<TVec2<int32>>& Graph, const FDynamicParticles& InParticles)
+template<typename T>
+static bool VerifyGraph(TArray<TArray<int32>> ColorGraph, const TArray<Chaos::TVec2<int32>>& Graph, const Chaos::TDynamicParticles<T, 3>& InParticles)
 {
 	for (int32 i = 0; i < ColorGraph.Num(); ++i)
 	{
@@ -37,7 +36,8 @@ bool VerifyGraph(TArray<TArray<int32>> ColorGraph, const TArray<TVec2<int32>>& G
 	return true;
 }
 
-bool VerifyGraph(TArray<TArray<int32>> ColorGraph, const TArray<TVec3<int32>>& Graph, const FDynamicParticles& InParticles)
+template<typename T>
+static bool VerifyGraph(TArray<TArray<int32>> ColorGraph, const TArray<Chaos::TVec3<int32>>& Graph, const Chaos::TDynamicParticles<T, 3>& InParticles)
 {
 	for (int32 i = 0; i < ColorGraph.Num(); ++i)
 	{
@@ -79,8 +79,8 @@ bool VerifyGraph(TArray<TArray<int32>> ColorGraph, const TArray<TVec3<int32>>& G
 	return true;
 }
 
-
-TArray<TArray<int32>> FGraphColoring::ComputeGraphColoring(const TArray<TVec2<int32>>& Graph, const FDynamicParticles& InParticles)
+template<typename T>
+TArray<TArray<int32>> Chaos::FGraphColoring::ComputeGraphColoring(const TArray<Chaos::TVec2<int32>>& Graph, const Chaos::TDynamicParticles<T, 3>& InParticles)
 {
 	TArray<TArray<int32>> ColorGraph;
 	TArray<FGraphNode> Nodes;
@@ -103,7 +103,7 @@ TArray<TArray<int32>> FGraphColoring::ComputeGraphColoring(const TArray<TVec2<in
 
 	for (uint32 ParticleNodeIndex = 0; ParticleNodeIndex < InParticles.Size(); ++ParticleNodeIndex)
 	{
-		const bool bIsParticleDynamic = InParticles.InvM(ParticleNodeIndex) != 0;
+		const bool bIsParticleDynamic = InParticles.InvM(ParticleNodeIndex) != (T)0.;
 		if (ProcessedNodes.Contains(ParticleNodeIndex) || !bIsParticleDynamic)
 		{
 			continue;
@@ -152,7 +152,7 @@ TArray<TArray<int32>> FGraphColoring::ComputeGraphColoring(const TArray<TVec2<in
 				{
 					FGraphNode& OtherNode = Nodes[OtherNodeIndex];
 
-					const bool bIsOtherGraphNodeDynamic = InParticles.InvM(OtherNodeIndex) != 0;
+					const bool bIsOtherGraphNodeDynamic = InParticles.InvM(OtherNodeIndex) != (T)0.;
 					if (bIsOtherGraphNodeDynamic)
 					{
 						while (OtherNode.UsedColors.Contains(ColorToUse) || GraphNode.UsedColors.Contains(ColorToUse))
@@ -182,7 +182,7 @@ TArray<TArray<int32>> FGraphColoring::ComputeGraphColoring(const TArray<TVec2<in
 				if (OtherNodeIndex != INDEX_NONE)
 				{
 					FGraphNode& OtherGraphNode = Nodes[OtherNodeIndex];
-					const bool bIsOtherGraphNodeDynamic = InParticles.InvM(OtherNodeIndex) != 0;
+					const bool bIsOtherGraphNodeDynamic = InParticles.InvM(OtherNodeIndex) != (T)0.;
 					if (bIsOtherGraphNodeDynamic)
 					{
 						// Mark other node as not allowing use of this color
@@ -206,8 +206,11 @@ TArray<TArray<int32>> FGraphColoring::ComputeGraphColoring(const TArray<TVec2<in
 	return ColorGraph;
 }
 
-TArray<TArray<int32>> FGraphColoring::ComputeGraphColoring(const TArray<TVec3<int32>>& Graph, const FDynamicParticles& InParticles)
+template<typename T>
+TArray<TArray<int32>> Chaos::FGraphColoring::ComputeGraphColoring(const TArray<TVec3<int32>>& Graph, const Chaos::TDynamicParticles<T, 3>& InParticles)
 {
+	using namespace Chaos;
+
 	TArray<TArray<int32>> ColorGraph;
 	TArray<FGraphNode> Nodes;
 	TArray<FGraph3dEdge> Edges;
@@ -231,7 +234,7 @@ TArray<TArray<int32>> FGraphColoring::ComputeGraphColoring(const TArray<TVec3<in
 
 	for (uint32 ParticleNodeIndex = 0; ParticleNodeIndex < InParticles.Size(); ++ParticleNodeIndex)
 	{
-		const bool bIsParticleDynamic = InParticles.InvM(ParticleNodeIndex) != 0;
+		const bool bIsParticleDynamic = InParticles.InvM(ParticleNodeIndex) != (T)0.;
 		if (ProcessedNodes.Contains(ParticleNodeIndex) || !bIsParticleDynamic)
 		{
 			continue;
@@ -288,7 +291,7 @@ TArray<TArray<int32>> FGraphColoring::ComputeGraphColoring(const TArray<TVec3<in
 				{
 					FGraphNode& OtherNode = Nodes[OtherNodeIndex];
 
-					const bool bIsOtherGraphNodeDynamic = InParticles.InvM(OtherNodeIndex) != 0;
+					const bool bIsOtherGraphNodeDynamic = InParticles.InvM(OtherNodeIndex) != (T)0.;
 					if (bIsOtherGraphNodeDynamic)
 					{
 						while (OtherNode.UsedColors.Contains(ColorToUse) || GraphNode.UsedColors.Contains(ColorToUse))
@@ -301,7 +304,7 @@ TArray<TArray<int32>> FGraphColoring::ComputeGraphColoring(const TArray<TVec3<in
 				{
 					FGraphNode& OtherNode = Nodes[OtherNodeIndex2];
 
-					const bool bIsOtherGraphNodeDynamic = InParticles.InvM(OtherNodeIndex2) != 0;
+					const bool bIsOtherGraphNodeDynamic = InParticles.InvM(OtherNodeIndex2) != (T)0.;
 					if (bIsOtherGraphNodeDynamic)
 					{
 						if (OtherNodeIndex == INDEX_NONE)
@@ -342,7 +345,7 @@ TArray<TArray<int32>> FGraphColoring::ComputeGraphColoring(const TArray<TVec3<in
 				if (OtherNodeIndex != INDEX_NONE)
 				{
 					FGraphNode& OtherGraphNode = Nodes[OtherNodeIndex];
-					const bool bIsOtherGraphNodeDynamic = InParticles.InvM(OtherNodeIndex) != 0;
+					const bool bIsOtherGraphNodeDynamic = InParticles.InvM(OtherNodeIndex) != (T)0.;
 					if (bIsOtherGraphNodeDynamic)
 					{
 						// Mark other node as not allowing use of this color
@@ -361,7 +364,7 @@ TArray<TArray<int32>> FGraphColoring::ComputeGraphColoring(const TArray<TVec3<in
 				if (OtherNodeIndex2 != INDEX_NONE)
 				{
 					FGraphNode& OtherGraphNode = Nodes[OtherNodeIndex2];
-					const bool bIsOtherGraphNodeDynamic = InParticles.InvM(OtherNodeIndex2) != 0;
+					const bool bIsOtherGraphNodeDynamic = InParticles.InvM(OtherNodeIndex2) != (T)0.;
 					if (bIsOtherGraphNodeDynamic)
 					{
 						// Mark other node as not allowing use of this color
@@ -385,3 +388,7 @@ TArray<TArray<int32>> FGraphColoring::ComputeGraphColoring(const TArray<TVec3<in
 	return ColorGraph;
 }
 
+template CHAOS_API TArray<TArray<int32>> Chaos::FGraphColoring::ComputeGraphColoring<Chaos::FRealSingle>(const TArray<Chaos::TVector<int32, 2>>&, const Chaos::TDynamicParticles<Chaos::FRealSingle, 3>&);
+template CHAOS_API TArray<TArray<int32>> Chaos::FGraphColoring::ComputeGraphColoring<Chaos::FRealDouble>(const TArray<Chaos::TVector<int32, 2>>&, const Chaos::TDynamicParticles<Chaos::FRealDouble, 3>&);
+template CHAOS_API TArray<TArray<int32>> Chaos::FGraphColoring::ComputeGraphColoring<Chaos::FRealSingle>(const TArray<Chaos::TVector<int32, 3>>&, const Chaos::TDynamicParticles<Chaos::FRealSingle, 3>&);
+template CHAOS_API TArray<TArray<int32>> Chaos::FGraphColoring::ComputeGraphColoring<Chaos::FRealDouble>(const TArray<Chaos::TVector<int32, 3>>&, const Chaos::TDynamicParticles<Chaos::FRealDouble, 3>&);
diff --git a/Engine/Source/Runtime/Experimental/Chaos/Private/Chaos/PBDAxialSpringConstraints.cpp b/Engine/Source/Runtime/Experimental/Chaos/Private/Chaos/PBDAxialSpringConstraints.cpp
--- a/Engine/Source/Runtime/Experimental/Chaos/Private/Chaos/PBDAxialSpringConstraints.cpp
+++ b/Engine/Source/Runtime/Experimental/Chaos/Private/Chaos/PBDAxialSpringConstraints.cpp
@@ -1,8 +1,8 @@
 // Copyright Epic Games, Inc. All Rights Reserved.
 #include "Chaos/PBDAxialSpringConstraints.h"
 #include "Chaos/GraphColoring.h"
+#include "Chaos/Framework/Parallel.h"
 #include "ChaosStats.h"
-#include "ChaosLog.h"
 #if INTEL_ISPC
 #include "PBDAxialSpringConstraints.ispc.generated.h"
 #endif
@@ -12,15 +12,18 @@ DECLARE_CYCLE_STAT(TEXT("Chaos PBD Axial Spring Constraint"), STAT_PBD_AxialSpri
 #if INTEL_ISPC && !UE_BUILD_SHIPPING
 bool bChaos_AxialSpring_ISPC_Enabled = true;
 FAutoConsoleVariableRef CVarChaosAxialSpringISPCEnabled(TEXT("p.Chaos.AxialSpring.ISPC"), bChaos_AxialSpring_ISPC_Enabled, TEXT("Whether to use ISPC optimizations in AxialSpring constraints"));
+
+static_assert(sizeof(ispc::FVector3f) == sizeof(Chaos::Softs::FSolverVec3), "sizeof(ispc::FVector3f) != sizeof(Chaos::Softs::FSolverVec3");
+static_assert(sizeof(ispc::FIntVector) == sizeof(Chaos::TVec3<int32>), "sizeof(ispc::FIntVector) != sizeof(Chaos::TVec3<int32>");
 #endif
 
-using namespace Chaos;
+namespace Chaos::Softs {
 
 // @todo(chaos): the parallel threshold (or decision to run parallel) should probably be owned by the solver and passed to the constraint container
 int32 Chaos_AxialSpring_ParallelConstraintCount = 100;
 FAutoConsoleVariableRef CVarChaosAxialSpringParallelConstraintCount(TEXT("p.Chaos.AxialSpring.ParallelConstraintCount"), Chaos_AxialSpring_ParallelConstraintCount, TEXT("If we have more constraints than this, use parallel-for in Apply."));
 
-void FPBDAxialSpringConstraints::InitColor(const FPBDParticles& InParticles)
+void FPBDAxialSpringConstraints::InitColor(const FSolverParticles& InParticles)
 {
 	// In dev builds we always color so we can tune the system without restarting. See Apply()
 #if UE_BUILD_SHIPPING || UE_BUILD_TEST
@@ -31,29 +34,29 @@ void FPBDAxialSpringConstraints::InitColor(const FPBDParticles& InParticles)
 	}
 }
 
-void FPBDAxialSpringConstraints::ApplyHelper(FPBDParticles& Particles, const FReal Dt, const int32 ConstraintIndex, const FReal ExpStiffnessValue) const
+void FPBDAxialSpringConstraints::ApplyHelper(FSolverParticles& Particles, const FSolverReal Dt, const int32 ConstraintIndex, const FSolverReal ExpStiffnessValue) const
 {
-		const auto& constraint = Constraints[ConstraintIndex];
-		const int32 i1 = constraint[0];
-		const int32 i2 = constraint[1];
-		const int32 i3 = constraint[2];
-		const FVec3 Delta = Base::GetDelta(Particles, ConstraintIndex, ExpStiffnessValue);
-		const FReal Multiplier = 2 / (FMath::Max(Barys[ConstraintIndex], 1 - Barys[ConstraintIndex]) + 1);
-		if (Particles.InvM(i1) > (FReal)0.)
+		const TVec3<int32>& Constraint = Constraints[ConstraintIndex];
+		const int32 i1 = Constraint[0];
+		const int32 i2 = Constraint[1];
+		const int32 i3 = Constraint[2];
+		const FSolverVec3 Delta = Base::GetDelta(Particles, ConstraintIndex, ExpStiffnessValue);
+		const FSolverReal Multiplier = (FSolverReal)2. / (FMath::Max(Barys[ConstraintIndex], (FSolverReal)1. - Barys[ConstraintIndex]) + (FSolverReal)1.);
+		if (Particles.InvM(i1) > (FSolverReal)0.)
 		{
 			Particles.P(i1) -= Multiplier * Particles.InvM(i1) * Delta;
 		}
-		if (Particles.InvM(i2) > (FReal)0.)
+		if (Particles.InvM(i2) > (FSolverReal)0.)
 		{
 			Particles.P(i2) += Multiplier * Particles.InvM(i2) * Barys[ConstraintIndex] * Delta;
 		}
-		if (Particles.InvM(i3) > (FReal)0.)
+		if (Particles.InvM(i3) > (FSolverReal)0.)
 		{
-			Particles.P(i3) += Multiplier * Particles.InvM(i3) * ((FReal)1. - Barys[ConstraintIndex]) * Delta;
+			Particles.P(i3) += Multiplier * Particles.InvM(i3) * ((FSolverReal)1. - Barys[ConstraintIndex]) * Delta;
 		}
 }
 
-void FPBDAxialSpringConstraints::Apply(FPBDParticles& Particles, const FReal Dt) const
+void FPBDAxialSpringConstraints::Apply(FSolverParticles& Particles, const FSolverReal Dt) const
 {
 	TRACE_CPUPROFILER_EVENT_SCOPE(FPBDAxialSpringConstraints_Apply);
 	SCOPE_CYCLE_COUNTER(STAT_PBD_AxialSpring);
@@ -61,7 +64,7 @@ void FPBDAxialSpringConstraints::Apply(FPBDParticles& Particles, const FReal Dt)
 	{
 		if (!Stiffness.HasWeightMap())
 		{
-			const FReal ExpStiffnessValue = (FReal)Stiffness;
+			const FSolverReal ExpStiffnessValue = (FSolverReal)Stiffness;
 
 #if INTEL_ISPC
 			if (bRealTypeCompatibleWithISPC && bChaos_AxialSpring_ISPC_Enabled)
@@ -69,7 +72,7 @@ void FPBDAxialSpringConstraints::Apply(FPBDParticles& Particles, const FReal Dt)
 				for (const TArray<int32>& ConstraintBatch : ConstraintsPerColor)
 				{
 					ispc::ApplyAxialSpringConstraints(
-						(ispc::FVector*)&Particles.GetP()[0],
+						(ispc::FVector3f*)&Particles.GetP()[0],
 						(ispc::FIntVector*)&Constraints.GetData()[0],
 						&ConstraintBatch.GetData()[0],
 						&Particles.GetInvM().GetData()[0],
@@ -100,7 +103,7 @@ void FPBDAxialSpringConstraints::Apply(FPBDParticles& Particles, const FReal Dt)
 				for (const TArray<int32>& ConstraintBatch : ConstraintsPerColor)
 				{
 					ispc::ApplyAxialSpringConstraintsWithWeightMaps(
-						(ispc::FVector*)&Particles.GetP()[0],
+						(ispc::FVector3f*)&Particles.GetP()[0],
 						(ispc::FIntVector*)&Constraints.GetData()[0],
 						&ConstraintBatch.GetData()[0],
 						&Particles.GetInvM().GetData()[0],
@@ -119,7 +122,7 @@ void FPBDAxialSpringConstraints::Apply(FPBDParticles& Particles, const FReal Dt)
 					PhysicsParallelFor(ConstraintBatch.Num(), [&](const int32 Index)
 					{
 						const int32 ConstraintIndex = ConstraintBatch[Index];
-						const FReal ExpStiffnessValue = Stiffness[ConstraintIndex];
+						const FSolverReal ExpStiffnessValue = Stiffness[ConstraintIndex];
 						ApplyHelper(Particles, Dt, ConstraintIndex, ExpStiffnessValue);
 					});
 				}
@@ -130,7 +133,7 @@ void FPBDAxialSpringConstraints::Apply(FPBDParticles& Particles, const FReal Dt)
 	{
 		if (!Stiffness.HasWeightMap())
 		{
-			const FReal ExpStiffnessValue = (FReal)Stiffness;
+			const FSolverReal ExpStiffnessValue = (FSolverReal)Stiffness;
 			for (int32 ConstraintIndex = 0; ConstraintIndex < Constraints.Num(); ++ConstraintIndex)
 			{
 				ApplyHelper(Particles, Dt, ConstraintIndex, ExpStiffnessValue);
@@ -140,10 +143,11 @@ void FPBDAxialSpringConstraints::Apply(FPBDParticles& Particles, const FReal Dt)
 		{
 			for (int32 ConstraintIndex = 0; ConstraintIndex < Constraints.Num(); ++ConstraintIndex)
 			{
-				const FReal ExpStiffnessValue = Stiffness[ConstraintIndex];
+				const FSolverReal ExpStiffnessValue = Stiffness[ConstraintIndex];
 				ApplyHelper(Particles, Dt, ConstraintIndex, ExpStiffnessValue);
 			}
 		}
 	}
-
 }
+
+}  // End namespace Chaos::Softs
diff --git a/Engine/Source/Runtime/Experimental/Chaos/Private/Chaos/PBDAxialSpringConstraints.ispc b/Engine/Source/Runtime/Experimental/Chaos/Private/Chaos/PBDAxialSpringConstraints.ispc
--- a/Engine/Source/Runtime/Experimental/Chaos/Private/Chaos/PBDAxialSpringConstraints.ispc
+++ b/Engine/Source/Runtime/Experimental/Chaos/Private/Chaos/PBDAxialSpringConstraints.ispc
@@ -4,13 +4,13 @@
 
 #include "Math/Vector.isph"
 
-export void ApplyAxialSpringConstraints(uniform FVector P[],
+export void ApplyAxialSpringConstraints(uniform FVector3f P[],
 										const uniform FIntVector AllConstraints[],
 										const uniform int32 ActiveConstraints[],
-										const uniform FReal InvM[],
-										const uniform FReal Barys[],
-										const uniform FReal Dists[],
-										const uniform FReal Stiffness,
+										const uniform float InvM[],
+										const uniform float Barys[],
+										const uniform float Dists[],
+										const uniform float Stiffness,
 										const uniform int32 NumConstraints)
 {
 	varying int32 ActiveConstraint;
@@ -34,16 +34,16 @@ export void ApplyAxialSpringConstraints(uniform FVector P[],
 #endif
 
 		#pragma ignore warning(perf)
-		const varying FReal Bary = Barys[ActiveConstraint];
+		const varying float Bary = Barys[ActiveConstraint];
 		#pragma ignore warning(perf)
-		const varying FReal Dist = Dists[ActiveConstraint];
+		const varying float Dist = Dists[ActiveConstraint];
 
 		const varying int32 i1 = Constraint.V[0];
 		const varying int32 i2 = Constraint.V[1];
 		const varying int32 i3 = Constraint.V[2];
 
-		const varying FVector P2 = VectorGather(&P[i2]);
-		const varying FVector P3 = VectorGather(&P[i3]);
+		const varying FVector3f P2 = VectorGather(&P[i2]);
+		const varying FVector3f P3 = VectorGather(&P[i3]);
 
 #if HW_GATHER_SUPPORTED == 1
 		if(i + programCount < NumConstraints)
@@ -54,29 +54,29 @@ export void ApplyAxialSpringConstraints(uniform FVector P[],
 		}
 #endif
 
-		const varying FVector P1 = VectorGather(&P[i1]);
+		const varying FVector3f P1 = VectorGather(&P[i1]);
 
-		const FVector P0 = (P2 - P3) * Bary + P3;
-		const FVector Difference = P1 - P0;
+		const FVector3f P0 = (P2 - P3) * Bary + P3;
+		const FVector3f Difference = P1 - P0;
 
-		const FReal Distance = VectorSize(Difference);
-		const FReal Multiplier = 2 / (max(Bary, 1 - Bary) + 1);
+		const float Distance = VectorSize(Difference);
+		const float Multiplier = 2 / (max(Bary, 1 - Bary) + 1);
 
 		#pragma ignore warning(perf)
-		const varying FReal M1 = InvM[i1];
+		const varying float M1 = InvM[i1];
 		#pragma ignore warning(perf)
-		const varying FReal M2 = InvM[i2];
+		const varying float M2 = InvM[i2];
 		#pragma ignore warning(perf)
-		const varying FReal M3 = InvM[i3];
+		const varying float M3 = InvM[i3];
 
-		const FVector Direction = Difference / Distance;
-		const FVector InnerDelta = (Distance - Dist) * Direction;
+		const FVector3f Direction = Difference / Distance;
+		const FVector3f InnerDelta = (Distance - Dist) * Direction;
 
-		const FReal PInvMass = M3 * (1 - Bary) + M2 * Bary;
-		const FReal CombinedInvMass = PInvMass + M1;
-		const FVector ComputedDelta = Stiffness * InnerDelta / CombinedInvMass;
+		const float PInvMass = M3 * (1 - Bary) + M2 * Bary;
+		const float CombinedInvMass = PInvMass + M1;
+		const FVector3f ComputedDelta = Stiffness * InnerDelta / CombinedInvMass;
 
-		const FVector Delta = VectorSelect(((M1 == 0 && PInvMass == 0) || Distance <= SMALL_NUMBER), ZeroVector, ComputedDelta);
+		const FVector3f Delta = VectorSelect(((M1 == 0 && PInvMass == 0) || Distance <= FLOAT_SMALL_NUMBER), FloatZeroVector, ComputedDelta);
 
 		if (M1 > 0)
 		{
@@ -93,14 +93,14 @@ export void ApplyAxialSpringConstraints(uniform FVector P[],
 	}
 }
 
-export void ApplyAxialSpringConstraintsWithWeightMaps(uniform FVector P[],
+export void ApplyAxialSpringConstraintsWithWeightMaps(uniform FVector3f P[],
 										const uniform FIntVector AllConstraints[],
 										const uniform int32 ActiveConstraints[],
-										const uniform FReal InvM[],
-										const uniform FReal Barys[],
-										const uniform FReal Dists[],
+										const uniform float InvM[],
+										const uniform float Barys[],
+										const uniform float Dists[],
 										const uniform uint8 StiffnessIndices[],
-										const uniform FReal StiffnessTable[],
+										const uniform float StiffnessTable[],
 										const uniform int32 NumConstraints)
 {
 	varying int32 ActiveConstraint;
@@ -124,9 +124,9 @@ export void ApplyAxialSpringConstraintsWithWeightMaps(uniform FVector P[],
 #endif
 
 		#pragma ignore warning(perf)
-		const varying FReal Bary = Barys[ActiveConstraint];
+		const varying float Bary = Barys[ActiveConstraint];
 		#pragma ignore warning(perf)
-		const varying FReal Dist = Dists[ActiveConstraint];
+		const varying float Dist = Dists[ActiveConstraint];
 		#pragma ignore warning(perf)
 		const varying uint8 StiffnessIndex = StiffnessIndices[ActiveConstraint];
 
@@ -134,8 +134,8 @@ export void ApplyAxialSpringConstraintsWithWeightMaps(uniform FVector P[],
 		const varying int32 i2 = Constraint.V[1];
 		const varying int32 i3 = Constraint.V[2];
 
-		const varying FVector P2 = VectorGather(&P[i2]);
-		const varying FVector P3 = VectorGather(&P[i3]);
+		const varying FVector3f P2 = VectorGather(&P[i2]);
+		const varying FVector3f P3 = VectorGather(&P[i3]);
 
 #if HW_GATHER_SUPPORTED == 1
 		if(i + programCount < NumConstraints)
@@ -146,32 +146,32 @@ export void ApplyAxialSpringConstraintsWithWeightMaps(uniform FVector P[],
 		}
 #endif
 
-		const varying FVector P1 = VectorGather(&P[i1]);
+		const varying FVector3f P1 = VectorGather(&P[i1]);
 
-		const FVector P0 = (P2 - P3) * Bary + P3;
-		const FVector Difference = P1 - P0;
+		const FVector3f P0 = (P2 - P3) * Bary + P3;
+		const FVector3f Difference = P1 - P0;
 
-		const FReal Distance = VectorSize(Difference);
-		const FReal Multiplier = 2 / (max(Bary, 1 - Bary) + 1);
+		const float Distance = VectorSize(Difference);
+		const float Multiplier = 2 / (max(Bary, 1 - Bary) + 1);
 
 		#pragma ignore warning(perf)
-		const varying FReal M1 = InvM[i1];
+		const varying float M1 = InvM[i1];
 		#pragma ignore warning(perf)
-		const varying FReal M2 = InvM[i2];
+		const varying float M2 = InvM[i2];
 		#pragma ignore warning(perf)
-		const varying FReal M3 = InvM[i3];
+		const varying float M3 = InvM[i3];
 
 		#pragma ignore warning(perf)
-		const varying FReal Stiffness = StiffnessTable[StiffnessIndex];
+		const varying float Stiffness = StiffnessTable[StiffnessIndex];
 
-		const FVector Direction = Difference / Distance;
-		const FVector InnerDelta = (Distance - Dist) * Direction;
+		const FVector3f Direction = Difference / Distance;
+		const FVector3f InnerDelta = (Distance - Dist) * Direction;
 
-		const FReal PInvMass = M3 * (1 - Bary) + M2 * Bary;
-		const FReal CombinedInvMass = PInvMass + M1;
-		const FVector ComputedDelta = Stiffness * InnerDelta / CombinedInvMass;
+		const float PInvMass = M3 * (1 - Bary) + M2 * Bary;
+		const float CombinedInvMass = PInvMass + M1;
+		const FVector3f ComputedDelta = Stiffness * InnerDelta / CombinedInvMass;
 
-		const FVector Delta = VectorSelect(((M1 == 0 && PInvMass == 0) || Distance <= SMALL_NUMBER), ZeroVector, ComputedDelta);
+		const FVector3f Delta = VectorSelect(((M1 == 0 && PInvMass == 0) || Distance <= FLOAT_SMALL_NUMBER), FloatZeroVector, ComputedDelta);
 
 		if (M1 > 0)
 		{
diff --git a/Engine/Source/Runtime/Experimental/Chaos/Private/Chaos/PBDCollisionSpringConstraintsBase.cpp b/Engine/Source/Runtime/Experimental/Chaos/Private/Chaos/PBDCollisionSpringConstraintsBase.cpp
--- a/Engine/Source/Runtime/Experimental/Chaos/Private/Chaos/PBDCollisionSpringConstraintsBase.cpp
+++ b/Engine/Source/Runtime/Experimental/Chaos/Private/Chaos/PBDCollisionSpringConstraintsBase.cpp
@@ -7,7 +7,7 @@
 #include "kDOP.h"
 #endif
 
-using namespace Chaos;
+namespace Chaos::Softs {
 
 #if PLATFORM_DESKTOP && PLATFORM_64BITS
 class FMeshBuildDataProvider
@@ -40,7 +40,7 @@ public:
 		return FMatrix::Identity;
 	}
 
-	FORCEINLINE FReal GetDeterminant() const
+	FORCEINLINE FSolverReal GetDeterminant() const
 	{
 		return 1.0f;
 	}
@@ -55,8 +55,8 @@ FPBDCollisionSpringConstraintsBase::FPBDCollisionSpringConstraintsBase(
 	const int32 InNumParticles,
 	const TArray<TVec3<int32>>& InElements,
 	TSet<TVec2<int32>>&& InDisabledCollisionElements,
-	const FReal InThickness,
-	const FReal InStiffness)
+	const FSolverReal InThickness,
+	const FSolverReal InStiffness)
 	: Elements(InElements)
 	, DisabledCollisionElements(InDisabledCollisionElements)
 	, Offset(InOffset)
@@ -66,7 +66,7 @@ FPBDCollisionSpringConstraintsBase::FPBDCollisionSpringConstraintsBase(
 {
 }
 
-void FPBDCollisionSpringConstraintsBase::Init(const FPBDParticles& Particles)
+void FPBDCollisionSpringConstraintsBase::Init(const FSolverParticles& Particles)
 {
 #if PLATFORM_DESKTOP && PLATFORM_64BITS
 	if (!Elements.Num())
@@ -84,18 +84,18 @@ void FPBDCollisionSpringConstraintsBase::Init(const FPBDParticles& Particles)
 	for (int32 i = 0; i < Elements.Num(); ++i)
 	{
 		const TVector<int32, 3>& Elem = Elements[i];
-		const FVec3& P0 = Particles.X(Elem[0]);
-		const FVec3& P1 = Particles.X(Elem[1]);
-		const FVec3& P2 = Particles.X(Elem[2]);
+		const FSolverVec3& P0 = Particles.X(Elem[0]);
+		const FSolverVec3& P1 = Particles.X(Elem[1]);
+		const FSolverVec3& P2 = Particles.X(Elem[2]);
 
-		BuildTriangleArray.Add(FkDOPBuildCollisionTriangle<uint32>(i, P0, P1, P2));
+		BuildTriangleArray.Add(FkDOPBuildCollisionTriangle<uint32>(i, FVec3(P0), FVec3(P1), FVec3(P2)));
 	}
 	DopTree.Build(BuildTriangleArray);
 
 	FMeshBuildDataProvider DopDataProvider(DopTree);
 	FCriticalSection CriticalSection;
 
-	const FReal Height = Thickness + Thickness;
+	const FSolverReal Height = Thickness + Thickness;
 
 	PhysicsParallelFor(NumParticles,
 		[this, &Particles, &CriticalSection, Height, &DopDataProvider, &DopTree](int32 i)
@@ -103,9 +103,9 @@ void FPBDCollisionSpringConstraintsBase::Init(const FPBDParticles& Particles)
 			const int32 Index = i + Offset;
 
 			FkHitResult Result;
-			const FVec3& Start = Particles.X(Index);
-			const FVec3 Direction = Particles.V(Index).GetSafeNormal();
-			const FVec3 End = Particles.P(Index) + Direction * Height;
+			const FSolverVec3& Start = Particles.X(Index);
+			const FSolverVec3 Direction = Particles.V(Index).GetSafeNormal();
+			const FSolverVec3 End = Particles.P(Index) + Direction * Height;
 			using FVec4Real = decltype(FVector4::X);
 			const FVector4 Start4((FVec4Real)Start.X, (FVec4Real)Start.Y, (FVec4Real)Start.Z, 0);
 			const FVector4 End4((FVec4Real)End.X, (FVec4Real)End.Y, (FVec4Real)End.Z, 0);
@@ -120,16 +120,16 @@ void FPBDCollisionSpringConstraintsBase::Init(const FPBDParticles& Particles)
 					return;
 				}
 
-				const FVec3& P = Particles.X(Index);
-				const FVec3& P0 = Particles.X(Elem[0]);
-				const FVec3& P1 = Particles.X(Elem[1]);
-				const FVec3& P2 = Particles.X(Elem[2]);
-				const FVec3 P10 = P1 - P0;
-				const FVec3 P20 = P2 - P0;
-				const FVec3 PP0 = P - P0;
-				const FReal Size10 = P10.SizeSquared();
-				const FReal Size20 = P20.SizeSquared();
-				FVec3 Bary;
+				const FSolverVec3& P = Particles.X(Index);
+				const FSolverVec3& P0 = Particles.X(Elem[0]);
+				const FSolverVec3& P1 = Particles.X(Elem[1]);
+				const FSolverVec3& P2 = Particles.X(Elem[2]);
+				const FSolverVec3 P10 = P1 - P0;
+				const FSolverVec3 P20 = P2 - P0;
+				const FSolverVec3 PP0 = P - P0;
+				const FSolverReal Size10 = P10.SizeSquared();
+				const FSolverReal Size20 = P20.SizeSquared();
+				FSolverVec3 Bary;
 				if (Size10 < SMALL_NUMBER)
 				{
 					if (Size20 < SMALL_NUMBER)
@@ -139,7 +139,7 @@ void FPBDCollisionSpringConstraintsBase::Init(const FPBDParticles& Particles)
 					}
 					else
 					{
-						const FReal Size = PP0.SizeSquared();
+						const FSolverReal Size = PP0.SizeSquared();
 						if (Size < SMALL_NUMBER)
 						{
 							Bary.Y = Bary.Z = 0.f;
@@ -147,7 +147,7 @@ void FPBDCollisionSpringConstraintsBase::Init(const FPBDParticles& Particles)
 						}
 						else
 						{
-							const FReal ProjP2 = FVec3::DotProduct(PP0, P20);
+							const FSolverReal ProjP2 = FSolverVec3::DotProduct(PP0, P20);
 							Bary.Y = 0.f;
 							Bary.Z = ProjP2 * FMath::InvSqrt(Size20 * Size);
 							Bary.X = 1.0f - Bary.Z;
@@ -156,7 +156,7 @@ void FPBDCollisionSpringConstraintsBase::Init(const FPBDParticles& Particles)
 				}
 				else if (Size20 < SMALL_NUMBER)
 				{
-					const FReal Size = PP0.SizeSquared();
+					const FSolverReal Size = PP0.SizeSquared();
 					if (Size < SMALL_NUMBER)
 					{
 						Bary.Y = Bary.Z = 0.f;
@@ -164,7 +164,7 @@ void FPBDCollisionSpringConstraintsBase::Init(const FPBDParticles& Particles)
 					}
 					else
 					{
-						const FReal ProjP1 = FVec3::DotProduct(PP0, P10);
+						const FSolverReal ProjP1 = FSolverVec3::DotProduct(PP0, P10);
 						Bary.Y = ProjP1 * FMath::InvSqrt(Size10 * PP0.SizeSquared());
 						Bary.Z = 0.f;
 						Bary.X = 1.0f - Bary.Y;
@@ -172,17 +172,17 @@ void FPBDCollisionSpringConstraintsBase::Init(const FPBDParticles& Particles)
 				}
 				else
 				{
-					const FReal ProjSides = FVec3::DotProduct(P10, P20);
-					const FReal ProjP1 = FVec3::DotProduct(PP0, P10);
-					const FReal ProjP2 = FVec3::DotProduct(PP0, P20);
-					const FReal Denom = Size10 * Size20 - ProjSides * ProjSides;
+					const FSolverReal ProjSides = FSolverVec3::DotProduct(P10, P20);
+					const FSolverReal ProjP1 = FSolverVec3::DotProduct(PP0, P10);
+					const FSolverReal ProjP2 = FSolverVec3::DotProduct(PP0, P20);
+					const FSolverReal Denom = Size10 * Size20 - ProjSides * ProjSides;
 					Bary.Y = (Size20 * ProjP1 - ProjSides * ProjP2) / Denom;
 					Bary.Z = (Size10 * ProjP2 - ProjSides * ProjP1) / Denom;
 					Bary.X = 1.0f - Bary.Z - Bary.Y;
 				}
 
-				FVec3 Normal = FVec3::CrossProduct(P10, P20).GetSafeNormal();
-				Normal = (FVec3::DotProduct(Normal, PP0) > 0) ? Normal : -Normal;
+				FSolverVec3 Normal = FSolverVec3::CrossProduct(P10, P20).GetSafeNormal();
+				Normal = (FSolverVec3::DotProduct(Normal, PP0) > 0) ? Normal : -Normal;
 
 				CriticalSection.Lock();
 				Constraints.Add({Index, Elem[0], Elem[1], Elem[2]});
@@ -194,7 +194,7 @@ void FPBDCollisionSpringConstraintsBase::Init(const FPBDParticles& Particles)
 #endif  // #if PLATFORM_DESKTOP && PLATFORM_64BITS
 }
 
-FVec3 FPBDCollisionSpringConstraintsBase::GetDelta(const FPBDParticles& Particles, const int32 i) const
+FSolverVec3 FPBDCollisionSpringConstraintsBase::GetDelta(const FSolverParticles& Particles, const int32 i) const
 {
 	const TVec4<int32>& Constraint = Constraints[i];
 	const int32 i1 = Constraint[0];
@@ -202,31 +202,33 @@ FVec3 FPBDCollisionSpringConstraintsBase::GetDelta(const FPBDParticles& Particle
 	const int32 i3 = Constraint[2];
 	const int32 i4 = Constraint[3];
 
-	const FReal CombinedMass = Particles.InvM(i1) +
+	const FSolverReal CombinedMass = Particles.InvM(i1) +
 		Particles.InvM(i2) * Barys[i][0] +
 		Particles.InvM(i3) * Barys[i][1] +
 		Particles.InvM(i4) * Barys[i][2];
-	if (CombinedMass <= (FReal)1e-7)
+	if (CombinedMass <= (FSolverReal)1e-7)
 	{
-		return FVec3(0);
+		return FSolverVec3(0);
 	}
 
-	const FVec3& P1 = Particles.P(i1);
-	const FVec3& P2 = Particles.P(i2);
-	const FVec3& P3 = Particles.P(i3);
-	const FVec3& P4 = Particles.P(i4);
+	const FSolverVec3& P1 = Particles.P(i1);
+	const FSolverVec3& P2 = Particles.P(i2);
+	const FSolverVec3& P3 = Particles.P(i3);
+	const FSolverVec3& P4 = Particles.P(i4);
 
-	const FReal Height = Thickness + Thickness;
-	const FVec3 P = Barys[i][0] * P2 + Barys[i][1] * P3 + Barys[i][2] * P4 + Height * Normals[i];
+	const FSolverReal Height = Thickness + Thickness;
+	const FSolverVec3 P = Barys[i][0] * P2 + Barys[i][1] * P3 + Barys[i][2] * P4 + Height * Normals[i];
 
-	const FVec3 Difference = P1 - P;
-	if (FVec3::DotProduct(Difference, Normals[i]) > 0)
+	const FSolverVec3 Difference = P1 - P;
+	if (FSolverVec3::DotProduct(Difference, Normals[i]) > 0)
 	{
-		return FVec3(0);
+		return FSolverVec3(0);
 	}
-	const FReal Distance = Difference.Size();
-	const FVec3 Delta = Distance * Normals[i];
+	const FSolverReal Distance = Difference.Size();
+	const FSolverVec3 Delta = Distance * Normals[i];
 	return Stiffness * Delta / CombinedMass;
 }
 
+}  // End namespace Chaos::Softs
+
 #endif
diff --git a/Engine/Source/Runtime/Experimental/Chaos/Private/Chaos/PBDEvolution.cpp b/Engine/Source/Runtime/Experimental/Chaos/Private/Chaos/PBDEvolution.cpp
--- a/Engine/Source/Runtime/Experimental/Chaos/Private/Chaos/PBDEvolution.cpp
+++ b/Engine/Source/Runtime/Experimental/Chaos/Private/Chaos/PBDEvolution.cpp
@@ -2,16 +2,11 @@
 #include "Chaos/PBDEvolution.h"
 
 #include "Chaos/Framework/Parallel.h"
-#include "Chaos/PBDCollisionSphereConstraints.h"
 #include "Chaos/PerParticleDampVelocity.h"
-#include "Chaos/PerParticleEulerStepVelocity.h"
 #include "Chaos/PerParticleGravity.h"
-#include "Chaos/PerParticleInitForce.h"
-#include "Chaos/PerParticlePBDCollisionConstraint.h"
 #include "Chaos/PerParticlePBDCCDCollisionConstraint.h"
-#include "Chaos/PerParticlePBDEulerStep.h"
-#include "Chaos/PerParticlePBDGroundConstraint.h"
-#include "Chaos/PerParticlePBDUpdateFromDeltaPosition.h"
+#include "Chaos/PerParticlePBDCollisionConstraint.h"
+#include "Chaos/VelocityField.h"
 #include "ChaosStats.h"
 #include "HAL/IConsoleManager.h"
 #if INTEL_ISPC
@@ -23,6 +18,7 @@ DECLARE_CYCLE_STAT(TEXT("Chaos PBD Velocity Damping State Update"), STAT_ChaosPB
 DECLARE_CYCLE_STAT(TEXT("Chaos PBD Velocity Field Update Forces"), STAT_ChaosPBDVelocityFieldUpdateForces, STATGROUP_Chaos);
 DECLARE_CYCLE_STAT(TEXT("Chaos PBD Velocity Damping"), STAT_ChaosPBDVelocityDampUpdate, STATGROUP_Chaos);
 DECLARE_CYCLE_STAT(TEXT("Chaos PBD Pre Iteration Updates"), STAT_ChaosPBDPreIterationUpdates, STATGROUP_Chaos);
+DECLARE_CYCLE_STAT(TEXT("Chaos PBD Integrate"), STAT_ChaosClothSolverIntegrate, STATGROUP_Chaos);
 DECLARE_CYCLE_STAT(TEXT("Chaos PBD Iteration Loop"), STAT_ChaosPBDIterationLoop, STATGROUP_Chaos);
 DECLARE_CYCLE_STAT(TEXT("Chaos PBD Post Iteration Updates"), STAT_ChaosPBDPostIterationUpdates, STATGROUP_Chaos);
 DECLARE_CYCLE_STAT(TEXT("Chaos PBD Constraint Rules"), STAT_ChaosPBDConstraintRule, STATGROUP_Chaos);
@@ -38,13 +34,16 @@ TAutoConsoleVariable<bool> CVarChaosPBDEvolutionFastPositionBasedFriction(TEXT("
 TAutoConsoleVariable<bool> CVarChaosPBDEvolutionUseSmoothTimeStep(TEXT("p.Chaos.PBDEvolution.UseSmoothTimeStep"), true, TEXT(""), ECVF_Cheat);
 TAutoConsoleVariable<int32> CVarChaosPBDEvolutionMinParallelBatchSize(TEXT("p.Chaos.PBDEvolution.MinParallelBatchSize"), 300, TEXT(""), ECVF_Cheat);
 TAutoConsoleVariable<bool> CVarChaosPBDEvolutionWriteCCDContacts(TEXT("p.Chaos.PBDEvolution.WriteCCDContacts"), false, TEXT("Write CCD collision contacts and normals potentially causing the CCD collision threads to lock, allowing for debugging of these contacts."), ECVF_Cheat);
+TAutoConsoleVariable<bool> CVarChaosPBDEvolutionParallelIntegrate(TEXT("p.Chaos.PBDEvolution.ParalleIntegrate"), false, TEXT("Run the integration step in parallel for."), ECVF_Cheat);
 
 #if INTEL_ISPC && !UE_BUILD_SHIPPING
 bool bChaos_PostIterationUpdates_ISPC_Enabled = true;
 FAutoConsoleVariableRef CVarChaosPostIterationUpdatesISPCEnabled(TEXT("p.Chaos.PostIterationUpdates.ISPC"), bChaos_PostIterationUpdates_ISPC_Enabled, TEXT("Whether to use ISPC optimizations in PBD Post iteration updates"));
+
+static_assert(sizeof(ispc::FVector3f) == sizeof(Chaos::Softs::FSolverVec3), "sizeof(ispc::FVector3f) != sizeof(Chaos::Softs::FSolverVec3");
 #endif
 
-using namespace Chaos;
+namespace Chaos::Softs {
 
 void FPBDEvolution::AddGroups(int32 NumGroups)
 {
@@ -55,7 +54,7 @@ void FPBDEvolution::AddGroups(int32 NumGroups)
 	// Set defaults
 	for (uint32 GroupId = Offset; GroupId < TArrayCollection::Size(); ++GroupId)
 	{
-		MGroupGravityForces[GroupId].SetAcceleration(MGravity);
+		MGroupGravityAccelerations[GroupId] = MGravity;
 		MGroupCollisionThicknesses[GroupId] = MCollisionThickness;
 		MGroupSelfCollisionThicknesses[GroupId] = MSelfCollisionThickness;
 		MGroupCoefficientOfFrictions[GroupId] = MCoefficientOfFriction;
@@ -70,8 +69,8 @@ void FPBDEvolution::ResetGroups()
 	AddGroups(1);  // Add default group
 }
 
-FPBDEvolution::FPBDEvolution(FPBDParticles&& InParticles, FKinematicGeometryClothParticles&& InGeometryParticles, TArray<TVec3<int32>>&& CollisionTriangles,
-    int32 NumIterations, FReal CollisionThickness, FReal SelfCollisionThickness, FReal CoefficientOfFriction, FReal Damping)
+FPBDEvolution::FPBDEvolution(FSolverParticles&& InParticles, FSolverRigidParticles&& InGeometryParticles, TArray<TVec3<int32>>&& CollisionTriangles,
+    int32 NumIterations, FSolverReal CollisionThickness, FSolverReal SelfCollisionThickness, FSolverReal CoefficientOfFriction, FSolverReal Damping)
     : MParticles(MoveTemp(InParticles))
 	, MParticlesActiveView(MParticles)
 	, MCollisionParticles(MoveTemp(InGeometryParticles))
@@ -80,7 +79,7 @@ FPBDEvolution::FPBDEvolution(FPBDParticles&& InParticles, FKinematicGeometryClot
 	, MConstraintInitsActiveView(MConstraintInits)
 	, MConstraintRulesActiveView(MConstraintRules)
 	, MNumIterations(NumIterations)
-	, MGravity(FVec3((FReal)0., (FReal)0., (FReal)-980.665))
+	, MGravity(FSolverVec3((FSolverReal)0., (FSolverReal)0., (FSolverReal)-980.665))
 	, MCollisionThickness(CollisionThickness)
 	, MSelfCollisionThickness(SelfCollisionThickness)
 	, MCoefficientOfFriction(CoefficientOfFriction)
@@ -89,7 +88,7 @@ FPBDEvolution::FPBDEvolution(FPBDParticles&& InParticles, FKinematicGeometryClot
 	, MSmoothDt(1.f / 30.f)  // Initialize filtered timestep at 30fps
 {
 	// Add group arrays
-	TArrayCollection::AddArray(&MGroupGravityForces);
+	TArrayCollection::AddArray(&MGroupGravityAccelerations);
 	TArrayCollection::AddArray(&MGroupVelocityFields);
 	TArrayCollection::AddArray(&MGroupForceRules);
 	TArrayCollection::AddArray(&MGroupCollisionThicknesses);
@@ -193,14 +192,14 @@ int32 FPBDEvolution::AddConstraintRuleRange(int32 NumConstraints, bool bActivate
 
 template<bool bForceRule, bool bVelocityField, bool bDampVelocityRule>
 void FPBDEvolution::PreIterationUpdate(
-	const FReal Dt,
+	const FSolverReal Dt,
 	const int32 Offset,
 	const int32 Range,
 	const int32 MinParallelBatchSize)
 {
 	const uint32 ParticleGroupId = MParticleGroupIds[Offset];
-	const TFunction<void(FPBDParticles&, const FReal, const int32)>& ForceRule = MGroupForceRules[ParticleGroupId];
-	const FVec3& Gravity = MGroupGravityForces[ParticleGroupId].GetAcceleration();
+	const TFunction<void(FSolverParticles&, const FSolverReal, const int32)>& ForceRule = MGroupForceRules[ParticleGroupId];
+	const FSolverVec3& Gravity = MGroupGravityAccelerations[ParticleGroupId];
 	FVelocityField& VelocityField = MGroupVelocityFields[ParticleGroupId];
 
 	if (bVelocityField)
@@ -216,48 +215,52 @@ void FPBDEvolution::PreIterationUpdate(
 		DampVelocityRule.UpdatePositionBasedState(MParticles, Offset, Range);
 	}
 
-	const int32 RangeSize = Range - Offset;
-	PhysicsParallelFor(RangeSize,
-		[this, &Offset, &ForceRule, &Gravity, &VelocityField, &DampVelocityRule, Dt](int32 i)
-		{
-			const int32 Index = Offset + i;
-			if (MParticles.InvM(Index) != (FReal)0.)  // Process dynamic particles
-			{
-				// Init forces with GravityForces
-				MParticles.Acceleration(Index) = Gravity;
+	{
+		SCOPE_CYCLE_COUNTER(STAT_ChaosClothSolverIntegrate);
 
-				// Force Rule
-				if (bForceRule)
+		const int32 RangeSize = Range - Offset;
+		PhysicsParallelFor(RangeSize,
+			[this, &Offset, &ForceRule, &Gravity, &VelocityField, &DampVelocityRule, Dt](int32 i)
+			{
+				const int32 Index = Offset + i;
+				if (MParticles.InvM(Index) != (FSolverReal)0.)  // Process dynamic particles
 				{
-					ForceRule(MParticles, Dt, Index); // F += M * A
-				}
+					// Init forces with GravityForces
+					MParticles.Acceleration(Index) = Gravity;
 
-				// Velocity Field
-				if (bVelocityField)
-				{
-					VelocityField.Apply(MParticles, Dt, Index);
-				}
+					// Force Rule
+					if (bForceRule)
+					{
+						ForceRule(MParticles, Dt, Index); // F += M * A
+					}
+
+					// Velocity Field
+					if (bVelocityField)
+					{
+						VelocityField.Apply(MParticles, Dt, Index);
+					}
 
-				// Euler Step Velocity
-				MParticles.V(Index) += MParticles.Acceleration(Index) * MSmoothDt;
+					// Euler Step Velocity
+					MParticles.V(Index) += MParticles.Acceleration(Index) * MSmoothDt;
 
-				// Damp Velocity Rule
-				if (bDampVelocityRule)
+					// Damp Velocity Rule
+					if (bDampVelocityRule)
+					{
+						DampVelocityRule.ApplyFast(MParticles, Dt, Index);
+					}
+
+					// Euler Step
+					MParticles.P(Index) = MParticles.X(Index) + MParticles.V(Index) * Dt;
+				}
+				else  // Process kinematic particles
 				{
-					DampVelocityRule.ApplyFast(MParticles, Dt, Index);
+					MKinematicUpdate(MParticles, Dt, MTime, Index);
 				}
-
-				// Euler Step
-				MParticles.P(Index) = MParticles.X(Index) + MParticles.V(Index) * Dt;
-			}
-			else  // Process kinematic particles
-			{
-				MKinematicUpdate(MParticles, Dt, MTime, Index);
-			}
-		}, RangeSize < MinParallelBatchSize);
+			}, RangeSize < MinParallelBatchSize);
+	}
 }
 
-void FPBDEvolution::AdvanceOneTimeStep(const FReal Dt, const bool bSmoothDt)
+void FPBDEvolution::AdvanceOneTimeStep(const FSolverReal Dt, const bool bSmoothDt)
 {
 	TRACE_CPUPROFILER_EVENT_SCOPE(FPBDEvolution_AdvanceOneTimeStep);
 	SCOPE_CYCLE_COUNTER(STAT_ChaosPBDVAdvanceTime);
@@ -268,7 +271,7 @@ void FPBDEvolution::AdvanceOneTimeStep(const FReal Dt, const bool bSmoothDt)
 	// Filter delta time to smoothen time variations and prevent unwanted vibrations, works best on Forces
 	if (bSmoothDt && CVarChaosPBDEvolutionUseSmoothTimeStep.GetValueOnAnyThread())
 	{
-		constexpr FReal DeltaTimeDecay = (FReal)0.1;
+		constexpr FSolverReal DeltaTimeDecay = (FSolverReal)0.1;
 		MSmoothDt += (Dt - MSmoothDt) * DeltaTimeDecay;
 	}
 	else
@@ -278,7 +281,9 @@ void FPBDEvolution::AdvanceOneTimeStep(const FReal Dt, const bool bSmoothDt)
 
 	// Don't bother with threaded execution if we don't have enough work to make it worth while.
 	const bool bUseSingleThreadedRange = !CVarChaosPBDEvolutionUseNestedParallelFor.GetValueOnAnyThread();
-	const int32 MinParallelBatchSize = CVarChaosPBDEvolutionMinParallelBatchSize.GetValueOnAnyThread(); // TODO: 1000 is a guess, tune this!
+	const int32 MinParallelBatchSize = !CVarChaosPBDEvolutionParallelIntegrate.GetValueOnAnyThread() ?
+		TNumericLimits<int32>::Max() :  // Disable
+		CVarChaosPBDEvolutionMinParallelBatchSize.GetValueOnAnyThread(); // TODO: 1000 is a guess, tune this!
 	const bool bWriteCCDContacts = CVarChaosPBDEvolutionWriteCCDContacts.GetValueOnAnyThread();
 
 	{
@@ -286,13 +291,13 @@ void FPBDEvolution::AdvanceOneTimeStep(const FReal Dt, const bool bSmoothDt)
 		SCOPE_CYCLE_COUNTER(STAT_ChaosPBDPreIterationUpdates);
 
 		MParticlesActiveView.RangeFor(
-			[this, Dt, MinParallelBatchSize](FPBDParticles& Particles, int32 Offset, int32 Range)
+			[this, Dt, MinParallelBatchSize](FSolverParticles& Particles, int32 Offset, int32 Range)
 			{
 				const uint32 ParticleGroupId = MParticleGroupIds[Offset];
 
 				if (MGroupVelocityFields[ParticleGroupId].IsActive())
 				{
-					if (MGroupDampings[ParticleGroupId] > (FReal)0.)
+					if (MGroupDampings[ParticleGroupId] > (FSolverReal)0.)
 					{
 						if (MGroupForceRules[ParticleGroupId])  // VeloctiyFields, Damping, Forces  // Damping?????
 						{
@@ -317,7 +322,7 @@ void FPBDEvolution::AdvanceOneTimeStep(const FReal Dt, const bool bSmoothDt)
 				}
 				else   // No Velocity Fields
 				{
-					if (MGroupDampings[ParticleGroupId] > (FReal)0.)
+					if (MGroupDampings[ParticleGroupId] > (FSolverReal)0.)
 					{
 						if (MGroupForceRules[ParticleGroupId])  // VeloctiyFields, Damping, Forces
 						{
@@ -351,10 +356,10 @@ void FPBDEvolution::AdvanceOneTimeStep(const FReal Dt, const bool bSmoothDt)
 			SCOPE_CYCLE_COUNTER(STAT_ChaosPBDCollisionKinematicUpdate);
 
 			MCollisionParticlesActiveView.SequentialFor(
-				[this, Dt](FKinematicGeometryClothParticles& CollisionParticles, int32 Index)
+				[this, Dt](FSolverRigidParticles& CollisionParticles, int32 Index)
 				{
 					// Store active collision particle frames prior to the kinematic update for CCD collisions
-					MCollisionTransforms[Index] = FRigidTransform3(CollisionParticles.X(Index), CollisionParticles.R(Index));
+					MCollisionTransforms[Index] = FSolverRigidTransform3(CollisionParticles.X(Index), CollisionParticles.R(Index));
 
 					// Update collision transform and velocity
 					MCollisionKinematicUpdate(CollisionParticles, Dt, MTime, Index);
@@ -373,7 +378,7 @@ void FPBDEvolution::AdvanceOneTimeStep(const FReal Dt, const bool bSmoothDt)
 		TRACE_CPUPROFILER_EVENT_SCOPE(ChaosXPBDConstraintsInit);
 		SCOPE_CYCLE_COUNTER(STAT_ChaosXPBDConstraintsInit);
 		MConstraintInitsActiveView.SequentialFor(
-			[this, Dt](TArray<TFunction<void(const FPBDParticles&, const FReal)>>& ConstraintInits, int32 Index)
+			[this, Dt](TArray<TFunction<void(const FSolverParticles&, const FSolverReal)>>& ConstraintInits, int32 Index)
 			{
 				ConstraintInits[Index](MParticles, Dt);
 			});
@@ -383,7 +388,7 @@ void FPBDEvolution::AdvanceOneTimeStep(const FReal Dt, const bool bSmoothDt)
 	MCollisionContacts.Reset();
 	MCollisionNormals.Reset();
 
-	TPerParticlePBDCollisionConstraint<EGeometryParticlesSimType::Other> CollisionRule(
+	FPerParticlePBDCollisionConstraint CollisionRule(
 		MCollisionParticlesActiveView,
 		MCollided,
 		MParticleGroupIds,
@@ -391,7 +396,7 @@ void FPBDEvolution::AdvanceOneTimeStep(const FReal Dt, const bool bSmoothDt)
 		MGroupCollisionThicknesses,
 		MGroupCoefficientOfFrictions);
 
-	TPerParticlePBDCCDCollisionConstraint<EGeometryParticlesSimType::Other> CCDCollisionRule(
+	FPerParticlePBDCCDCollisionConstraint CCDCollisionRule(
 		MCollisionParticlesActiveView,
 		MCollisionTransforms,
 		MCollided,
@@ -411,7 +416,7 @@ void FPBDEvolution::AdvanceOneTimeStep(const FReal Dt, const bool bSmoothDt)
 		for (int32 i = 0; i < MNumIterations; ++i)
 		{
 			MConstraintRulesActiveView.RangeFor(
-				[this, Dt](TArray<TFunction<void(FPBDParticles&, const FReal)>>& ConstraintRules, int32 Offset, int32 Range)
+				[this, Dt](TArray<TFunction<void(FSolverParticles&, const FSolverReal)>>& ConstraintRules, int32 Offset, int32 Range)
 				{
 					TRACE_CPUPROFILER_EVENT_SCOPE(ChaosPBDConstraintRule);
 					SCOPE_CYCLE_COUNTER(STAT_ChaosPBDConstraintRule);
@@ -425,7 +430,7 @@ void FPBDEvolution::AdvanceOneTimeStep(const FReal Dt, const bool bSmoothDt)
 				TRACE_CPUPROFILER_EVENT_SCOPE(ChaosPBDCollisionRule);
 				SCOPE_CYCLE_COUNTER(STAT_ChaosPBDCollisionRule);
 				MParticlesActiveView.RangeFor(
-					[this, &CollisionRule, &CCDCollisionRule, Dt](FPBDParticles& Particles, int32 Offset, int32 Range)
+					[this, &CollisionRule, &CCDCollisionRule, Dt](FSolverParticles& Particles, int32 Offset, int32 Range)
 					{
 						TRACE_CPUPROFILER_EVENT_SCOPE(ChaosPBDCollisionRuleP);
 						const uint32 DynamicGroupId = MParticleGroupIds[Offset];  // Particle group Id, must be the same across the entire range
@@ -446,26 +451,26 @@ void FPBDEvolution::AdvanceOneTimeStep(const FReal Dt, const bool bSmoothDt)
 			SCOPE_CYCLE_COUNTER(STAT_ChaosPBDPostIterationUpdates);
 
 			// Particle update, V = (P - X) / Dt; X = P;
+#if INTEL_ISPC
 			if (bChaos_PostIterationUpdates_ISPC_Enabled)
 			{
-#if INTEL_ISPC
 				MParticlesActiveView.RangeFor(
-					[Dt](FPBDParticles& Particles, int32 Offset, int32 Range)
+					[Dt](FSolverParticles& Particles, int32 Offset, int32 Range)
 					{
 						ispc::PostIterationUpdates(
-							(ispc::FVector*)Particles.GetV().GetData(),
-							(ispc::FVector*)Particles.XArray().GetData(),
-							(const ispc::FVector*)Particles.GetP().GetData(),
+							(ispc::FVector3f*)Particles.GetV().GetData(),
+							(ispc::FVector3f*)Particles.XArray().GetData(),
+							(const ispc::FVector3f*)Particles.GetP().GetData(),
 							Dt,
 							Offset,
 							Range);
 					});
-#endif
 			}
 			else
+#endif
 			{
 				MParticlesActiveView.ParallelFor(
-					[Dt](FPBDParticles& Particles, int32 Index)
+					[Dt](FSolverParticles& Particles, int32 Index)
 					{
 						Particles.V(Index) = (Particles.P(Index) - Particles.X(Index)) / Dt;
 						Particles.X(Index) = Particles.P(Index);
@@ -479,10 +484,12 @@ void FPBDEvolution::AdvanceOneTimeStep(const FReal Dt, const bool bSmoothDt)
 	{
 		SCOPE_CYCLE_COUNTER(STAT_ChaosPBDCollisionRuleFriction);
 		MParticlesActiveView.ParallelFor(
-			[&CollisionRule, Dt](FPBDParticles& Particles, int32 Index)
+			[&CollisionRule, Dt](FSolverParticles& Particles, int32 Index)
 			{
 				TRACE_CPUPROFILER_EVENT_SCOPE(ChaosPBDCollisionRuleFriction);
 				CollisionRule.ApplyFriction(Particles, Dt, Index);
 			}, bUseSingleThreadedRange, MinParallelBatchSize);
 	}
-}
\ No newline at end of file
+}
+
+}  // End namespace Chaos::Softs
diff --git a/Engine/Source/Runtime/Experimental/Chaos/Private/Chaos/PBDEvolution.ispc b/Engine/Source/Runtime/Experimental/Chaos/Private/Chaos/PBDEvolution.ispc
--- a/Engine/Source/Runtime/Experimental/Chaos/Private/Chaos/PBDEvolution.ispc
+++ b/Engine/Source/Runtime/Experimental/Chaos/Private/Chaos/PBDEvolution.ispc
@@ -4,29 +4,29 @@
 
 #include "Math/Vector.isph"
 
-export void PostIterationUpdates(uniform FVector ParticlesV[],
-								uniform FVector ParticlesX[],
-								const uniform FVector ParticlesP[],
-								const uniform FReal Dt,
+export void PostIterationUpdates(uniform FVector3f ParticlesV[],
+								uniform FVector3f ParticlesX[],
+								const uniform FVector3f ParticlesP[],
+								const uniform float Dt,
 								const uniform int32 Offset,
 								const uniform int32 Range)
 {
 	const uniform int RangeBase = (Range - Offset) & ~(programCount-1);
-	const varying FReal vDt = Dt;
+	const varying float vDt = Dt;
 
 	for(uniform int32 Index = Offset; Index < Offset + RangeBase; Index+=programCount)
 	{
-		uniform FReal *uniform VPtr = (uniform FReal *uniform)&ParticlesV[Index];
-		uniform FReal *uniform XPtr = (uniform FReal *uniform)&ParticlesX[Index];
-		const uniform FReal *uniform PPtr = (uniform FReal *uniform)&ParticlesP[Index];
+		uniform float *uniform VPtr = (uniform float *uniform)&ParticlesV[Index];
+		uniform float *uniform XPtr = (uniform float *uniform)&ParticlesX[Index];
+		const uniform float *uniform PPtr = (uniform float *uniform)&ParticlesP[Index];
 
-		const varying FReal P0 = PPtr[programIndex];
-		const varying FReal P1 = PPtr[programCount + programIndex];
-		const varying FReal P2 = PPtr[2*programCount + programIndex];
+		const varying float P0 = PPtr[programIndex];
+		const varying float P1 = PPtr[programCount + programIndex];
+		const varying float P2 = PPtr[2*programCount + programIndex];
 
-		const varying FReal X0 = XPtr[programIndex];
-		const varying FReal X1 = XPtr[programCount + programIndex];
-		const varying FReal X2 = XPtr[2*programCount + programIndex];
+		const varying float X0 = XPtr[programIndex];
+		const varying float X1 = XPtr[programCount + programIndex];
+		const varying float X2 = XPtr[2*programCount + programIndex];
 
 		VPtr[programIndex] = (P0 - X0) / vDt;
 		VPtr[programCount + programIndex] = (P1 - X1) / vDt;
@@ -39,8 +39,8 @@ export void PostIterationUpdates(uniform FVector ParticlesV[],
 
 	for(uniform int32 Index = Offset + RangeBase; Index < Range; ++Index)
 	{
-		const uniform FVector P = ParticlesP[Index];
-		const uniform FVector X = ParticlesX[Index];
+		const uniform FVector3f P = ParticlesP[Index];
+		const uniform FVector3f X = ParticlesX[Index];
 
 		ParticlesV[Index] = (P - X) / Dt;
 		ParticlesX[Index] = P;
diff --git a/Engine/Source/Runtime/Experimental/Chaos/Private/Chaos/PBDLongRangeConstraints.cpp b/Engine/Source/Runtime/Experimental/Chaos/Private/Chaos/PBDLongRangeConstraints.cpp
--- a/Engine/Source/Runtime/Experimental/Chaos/Private/Chaos/PBDLongRangeConstraints.cpp
+++ b/Engine/Source/Runtime/Experimental/Chaos/Private/Chaos/PBDLongRangeConstraints.cpp
@@ -6,16 +6,16 @@
 #endif
 
 #if INTEL_ISPC && !UE_BUILD_SHIPPING
-static_assert(sizeof(ispc::FVector) == sizeof(Chaos::FVec3), "sizeof(ispc::FVector) != sizeof(Chaos::FVec3)");
-static_assert(sizeof(ispc::FTether) == sizeof(Chaos::FPBDLongRangeConstraintsBase::FTether), "sizeof(ispc::FTether) != sizeof(Chaos::FPBDLongRangeConstraintsBase::FTether)");
+static_assert(sizeof(ispc::FVector3f) == sizeof(Chaos::Softs::FSolverVec3), "sizeof(ispc::FVector3f) != sizeof(Chaos::Softs::FSolverVec3)");
+static_assert(sizeof(ispc::FTether) == sizeof(Chaos::Softs::FPBDLongRangeConstraintsBase::FTether), "sizeof(ispc::FTether) != sizeof(Chaos::Softs::FPBDLongRangeConstraintsBase::FTether)");
 
 bool bChaos_LongRange_ISPC_Enabled = true;
 FAutoConsoleVariableRef CVarChaosLongRangeISPCEnabled(TEXT("p.Chaos.LongRange.ISPC"), bChaos_LongRange_ISPC_Enabled, TEXT("Whether to use ISPC optimizations in long range constraints"));
 #endif
 
-using namespace Chaos;
+namespace Chaos::Softs {
 
-void FPBDLongRangeConstraints::Apply(FPBDParticles& Particles, const FReal /*Dt*/) const
+void FPBDLongRangeConstraints::Apply(FSolverParticles& Particles, const FSolverReal /*Dt*/) const
 {
 	TRACE_CPUPROFILER_EVENT_SCOPE(FPBDLongRangeConstraints_Apply);
 	SCOPE_CYCLE_COUNTER(STAT_PBD_LongRange);
@@ -23,10 +23,10 @@ void FPBDLongRangeConstraints::Apply(FPBDParticles& Particles, const FReal /*Dt*
 
 	if (!Stiffness.HasWeightMap())
 	{
-		const FReal ExpStiffnessValue = (FReal)Stiffness;
+		const FSolverReal ExpStiffnessValue = (FSolverReal)Stiffness;
 		if (!HasScaleWeightMap())
 		{
-			const FReal ScaleValue = ScaleTable[0];
+			const FSolverReal ScaleValue = ScaleTable[0];
 #if INTEL_ISPC
 			if (bRealTypeCompatibleWithISPC && bChaos_LongRange_ISPC_Enabled)
 			{
@@ -34,7 +34,7 @@ void FPBDLongRangeConstraints::Apply(FPBDParticles& Particles, const FReal /*Dt*
 				for (const TConstArrayView<FTether>& TetherBatch : Tethers)
 				{
 					ispc::ApplyLongRangeConstraints(
-						(ispc::FVector*)Particles.GetP().GetData(),
+						(ispc::FVector3f*)Particles.GetP().GetData(),
 						(const ispc::FTether*)TetherBatch.GetData(),
 						ExpStiffnessValue,
 						ScaleValue,
@@ -65,7 +65,7 @@ void FPBDLongRangeConstraints::Apply(FPBDParticles& Particles, const FReal /*Dt*
 				for (const TConstArrayView<FTether>& TetherBatch : Tethers)
 				{
 					ispc::ApplyLongRangeConstraintsScaleWeightmap(
-						(ispc::FVector*)Particles.GetP().GetData(),
+						(ispc::FVector3f*)Particles.GetP().GetData(),
 						(const ispc::FTether*)TetherBatch.GetData(),
 						ExpStiffnessValue,
 						ScaleIndices.GetData(),
@@ -84,7 +84,7 @@ void FPBDLongRangeConstraints::Apply(FPBDParticles& Particles, const FReal /*Dt*
 						{
 							const FTether& Tether = TetherBatch[Index];
 							const int32 LocalParticleIndex = GetEndIndex(Tether);
-							const FReal ScaleValue = ScaleTable[ScaleIndices[LocalParticleIndex]];
+							const FSolverReal ScaleValue = ScaleTable[ScaleIndices[LocalParticleIndex]];
 							Particles.P(ParticleOffset + LocalParticleIndex) += ExpStiffnessValue * GetDelta(Particles, Tether, ScaleValue);
 						}, TetherBatch.Num() < MinParallelSize);
 				}
@@ -95,7 +95,7 @@ void FPBDLongRangeConstraints::Apply(FPBDParticles& Particles, const FReal /*Dt*
 	{
 		if (!HasScaleWeightMap())
 		{
-			const FReal ScaleValue = ScaleTable[0];
+			const FSolverReal ScaleValue = ScaleTable[0];
 
 #if INTEL_ISPC
 			if (bRealTypeCompatibleWithISPC && bChaos_LongRange_ISPC_Enabled)
@@ -104,7 +104,7 @@ void FPBDLongRangeConstraints::Apply(FPBDParticles& Particles, const FReal /*Dt*
 				for (const TConstArrayView<FTether>& TetherBatch : Tethers)
 				{
 					ispc::ApplyLongRangeConstraintsStiffnessWeightmap(
-						(ispc::FVector*)Particles.GetP().GetData(),
+						(ispc::FVector3f*)Particles.GetP().GetData(),
 						(const ispc::FTether*)TetherBatch.GetData(),
 						Stiffness.GetIndices().GetData(),
 						Stiffness.GetTable().GetData(),
@@ -123,7 +123,7 @@ void FPBDLongRangeConstraints::Apply(FPBDParticles& Particles, const FReal /*Dt*
 						{
 							const FTether& Tether = TetherBatch[Index];
 							const int32 LocalParticleIndex = GetEndIndex(Tether);
-							const FReal ExpStiffnessValue = Stiffness[LocalParticleIndex];
+							const FSolverReal ExpStiffnessValue = Stiffness[LocalParticleIndex];
 							Particles.P(ParticleOffset + LocalParticleIndex) += ExpStiffnessValue * GetDelta(Particles, Tether, ScaleValue);
 						}, TetherBatch.Num() < MinParallelSize);
 				}
@@ -138,7 +138,7 @@ void FPBDLongRangeConstraints::Apply(FPBDParticles& Particles, const FReal /*Dt*
 				for (const TConstArrayView<FTether>& TetherBatch : Tethers)
 				{
 					ispc::ApplyLongRangeConstraintsStiffnessScaleWeightmaps(
-						(ispc::FVector*)Particles.GetP().GetData(),
+						(ispc::FVector3f*)Particles.GetP().GetData(),
 						(const ispc::FTether*)TetherBatch.GetData(),
 						Stiffness.GetIndices().GetData(),
 						Stiffness.GetTable().GetData(),
@@ -158,8 +158,8 @@ void FPBDLongRangeConstraints::Apply(FPBDParticles& Particles, const FReal /*Dt*
 						{
 							const FTether& Tether = TetherBatch[Index];
 							const int32 LocalParticleIndex = GetEndIndex(Tether);
-							const FReal ExpStiffnessValue = Stiffness[LocalParticleIndex];
-							const FReal ScaleValue = ScaleTable[ScaleIndices[LocalParticleIndex]];
+							const FSolverReal ExpStiffnessValue = Stiffness[LocalParticleIndex];
+							const FSolverReal ScaleValue = ScaleTable[ScaleIndices[LocalParticleIndex]];
 							Particles.P(ParticleOffset + LocalParticleIndex) += ExpStiffnessValue * GetDelta(Particles, Tether, ScaleValue);
 						}, TetherBatch.Num() < MinParallelSize);
 				}
@@ -168,3 +168,4 @@ void FPBDLongRangeConstraints::Apply(FPBDParticles& Particles, const FReal /*Dt*
 	}
 }
 
+}  // End namespace Chaos::Softs
diff --git a/Engine/Source/Runtime/Experimental/Chaos/Private/Chaos/PBDLongRangeConstraints.ispc b/Engine/Source/Runtime/Experimental/Chaos/Private/Chaos/PBDLongRangeConstraints.ispc
--- a/Engine/Source/Runtime/Experimental/Chaos/Private/Chaos/PBDLongRangeConstraints.ispc
+++ b/Engine/Source/Runtime/Experimental/Chaos/Private/Chaos/PBDLongRangeConstraints.ispc
@@ -21,10 +21,10 @@ static inline varying FTether LoadTetherHelper(const uniform FTether *uniform Sr
 	return Tether;
 }
 
-export void ApplyLongRangeConstraints(uniform FVector P[],
+export void ApplyLongRangeConstraints(uniform FVector3f P[],
 										const uniform FTether Tethers[],
-										const uniform FReal Stiffness,
-										const uniform FReal Scale,
+										const uniform float Stiffness,
+										const uniform float Scale,
 										const uniform int32 NumConstraints,
 										const uniform int32 ParticleOffset)
 {
@@ -49,8 +49,8 @@ export void ApplyLongRangeConstraints(uniform FVector P[],
 		const varying int32 i1 = Start + ParticleOffset;
 		const varying int32 i2 = End + ParticleOffset;
 
-		const varying FVector P1 = VectorGather(&P[i1]);
-		const varying FVector P2 = VectorGather(&P[i2]);
+		const varying FVector3f P1 = VectorGather(&P[i1]);
+		const varying FVector3f P2 = VectorGather(&P[i2]);
 
 #if HW_GATHER_SUPPORTED == 1
 		if(i + programCount < NumConstraints)
@@ -59,24 +59,24 @@ export void ApplyLongRangeConstraints(uniform FVector P[],
 		}
 #endif
 
-		FVector Direction = P1 - P2;
-		const FReal Length = VectorSize(Direction);
-		Direction = VectorSelect(Length < KINDA_SMALL_NUMBER, ForwardVector, Direction / Length);
+		FVector3f Direction = P1 - P2;
+		const float Length = VectorSize(Direction);
+		Direction = VectorSelect(Length < FLOAT_KINDA_SMALL_NUMBER, FloatForwardVector, Direction / Length);
 
-		const FReal Offset = Length - RefLength * Scale;
-		const FVector ComputedDelta = Stiffness * Offset * Direction;
+		const float Offset = Length - RefLength * Scale;
+		const FVector3f ComputedDelta = Stiffness * Offset * Direction;
 
-		const FVector Delta = VectorSelect(Offset < ZERO, ZeroVector, ComputedDelta);
+		const FVector3f Delta = VectorSelect(Offset < FLOAT_ZERO, FloatZeroVector, ComputedDelta);
 
 		VectorScatter(&P[i2], P2 + Delta);
 	}
 }
 
-export void ApplyLongRangeConstraintsStiffnessWeightmap(uniform FVector P[],
+export void ApplyLongRangeConstraintsStiffnessWeightmap(uniform FVector3f P[],
 														const uniform FTether Tethers[],
 														const uniform uint8 StiffnessIndices[],
-														const uniform FReal StiffnessTable[],
-														const uniform FReal Scale,
+														const uniform float StiffnessTable[],
+														const uniform float Scale,
 														const uniform int32 NumConstraints,
 														const uniform int32 ParticleOffset)
 {
@@ -104,8 +104,8 @@ export void ApplyLongRangeConstraintsStiffnessWeightmap(uniform FVector P[],
 		const varying int32 i1 = Start + ParticleOffset;
 		const varying int32 i2 = End + ParticleOffset;
 
-		const varying FVector P1 = VectorGather(&P[i1]);
-		const varying FVector P2 = VectorGather(&P[i2]);
+		const varying FVector3f P1 = VectorGather(&P[i1]);
+		const varying FVector3f P2 = VectorGather(&P[i2]);
 
 #if HW_GATHER_SUPPORTED == 1
 		if(i + programCount < NumConstraints)
@@ -115,26 +115,26 @@ export void ApplyLongRangeConstraintsStiffnessWeightmap(uniform FVector P[],
 #endif
 
 		#pragma ignore warning(perf)
-		const varying FReal Stiffness = StiffnessTable[StiffnessIndex];
+		const varying float Stiffness = StiffnessTable[StiffnessIndex];
 
-		FVector Direction = P1 - P2;
-		const FReal Length = VectorSize(Direction);
-		Direction = VectorSelect(Length < KINDA_SMALL_NUMBER, ForwardVector, Direction / Length);
+		FVector3f Direction = P1 - P2;
+		const float Length = VectorSize(Direction);
+		Direction = VectorSelect(Length < FLOAT_KINDA_SMALL_NUMBER, FloatForwardVector, Direction / Length);
 
-		const FReal Offset = Length - RefLength * Scale;
-		const FVector ComputedDelta = Stiffness * Offset * Direction;
+		const float Offset = Length - RefLength * Scale;
+		const FVector3f ComputedDelta = Stiffness * Offset * Direction;
 
-		const FVector Delta = VectorSelect(Offset < ZERO, ZeroVector, ComputedDelta);
+		const FVector3f Delta = VectorSelect(Offset < FLOAT_ZERO, FloatZeroVector, ComputedDelta);
 
 		VectorScatter(&P[i2], P2 + Delta);
 	}
 }
 
-export void ApplyLongRangeConstraintsScaleWeightmap(uniform FVector P[],
+export void ApplyLongRangeConstraintsScaleWeightmap(uniform FVector3f P[],
 													const uniform FTether Tethers[],
-													const uniform FReal Stiffness,
+													const uniform float Stiffness,
 													const uniform uint8 ScaleIndices[],
-													const uniform FReal ScaleTable[],
+													const uniform float ScaleTable[],
 													const uniform int32 NumConstraints,
 													const uniform int32 ParticleOffset)
 {
@@ -162,8 +162,8 @@ export void ApplyLongRangeConstraintsScaleWeightmap(uniform FVector P[],
 		const varying int32 i1 = Start + ParticleOffset;
 		const varying int32 i2 = End + ParticleOffset;
 
-		const varying FVector P1 = VectorGather(&P[i1]);
-		const varying FVector P2 = VectorGather(&P[i2]);
+		const varying FVector3f P1 = VectorGather(&P[i1]);
+		const varying FVector3f P2 = VectorGather(&P[i2]);
 
 #if HW_GATHER_SUPPORTED == 1
 		if(i + programCount < NumConstraints)
@@ -173,28 +173,28 @@ export void ApplyLongRangeConstraintsScaleWeightmap(uniform FVector P[],
 #endif
 
 		#pragma ignore warning(perf)
-		const varying FReal Scale = ScaleTable[ScaleIndex];
+		const varying float Scale = ScaleTable[ScaleIndex];
 
-		FVector Direction = P1 - P2;
-		const FReal Length = VectorSize(Direction);
-		Direction = VectorSelect(Length < KINDA_SMALL_NUMBER, ForwardVector, Direction / Length);
+		FVector3f Direction = P1 - P2;
+		const float Length = VectorSize(Direction);
+		Direction = VectorSelect(Length < FLOAT_KINDA_SMALL_NUMBER, FloatForwardVector, Direction / Length);
 
-		const FReal Offset = Length - RefLength * Scale;
-		const FVector ComputedDelta = Stiffness * Offset * Direction;
+		const float Offset = Length - RefLength * Scale;
+		const FVector3f ComputedDelta = Stiffness * Offset * Direction;
 
-		const FVector Delta = VectorSelect(Offset < ZERO, ZeroVector, ComputedDelta);
+		const FVector3f Delta = VectorSelect(Offset < FLOAT_ZERO, FloatZeroVector, ComputedDelta);
 
 		VectorScatter(&P[i2], P2 + Delta);
 	}
 }
 
 
-export void ApplyLongRangeConstraintsStiffnessScaleWeightmaps(uniform FVector P[],
+export void ApplyLongRangeConstraintsStiffnessScaleWeightmaps(uniform FVector3f P[],
 																const uniform FTether Tethers[],
 																const uniform uint8 StiffnessIndices[],
-																const uniform FReal StiffnessTable[],
+																const uniform float StiffnessTable[],
 																const uniform uint8 ScaleIndices[],
-																const uniform FReal ScaleTable[],
+																const uniform float ScaleTable[],
 																const uniform int32 NumConstraints,
 																const uniform int32 ParticleOffset)
 {
@@ -224,8 +224,8 @@ export void ApplyLongRangeConstraintsStiffnessScaleWeightmaps(uniform FVector P[
 		const varying int32 i1 = Start + ParticleOffset;
 		const varying int32 i2 = End + ParticleOffset;
 
-		const varying FVector P1 = VectorGather(&P[i1]);
-		const varying FVector P2 = VectorGather(&P[i2]);
+		const varying FVector3f P1 = VectorGather(&P[i1]);
+		const varying FVector3f P2 = VectorGather(&P[i2]);
 
 #if HW_GATHER_SUPPORTED == 1
 		if(i + programCount < NumConstraints)
@@ -235,18 +235,18 @@ export void ApplyLongRangeConstraintsStiffnessScaleWeightmaps(uniform FVector P[
 #endif
 
 		#pragma ignore warning(perf)
-		const varying FReal Stiffness = StiffnessTable[StiffnessIndex];
+		const varying float Stiffness = StiffnessTable[StiffnessIndex];
 		#pragma ignore warning(perf)
-		const varying FReal Scale = ScaleTable[ScaleIndex];
+		const varying float Scale = ScaleTable[ScaleIndex];
 
-		FVector Direction = P1 - P2;
-		const FReal Length = VectorSize(Direction);
-		Direction = VectorSelect(Length < KINDA_SMALL_NUMBER, ForwardVector, Direction / Length);
+		FVector3f Direction = P1 - P2;
+		const float Length = VectorSize(Direction);
+		Direction = VectorSelect(Length < FLOAT_KINDA_SMALL_NUMBER, FloatForwardVector, Direction / Length);
 
-		const FReal Offset = Length - RefLength * Scale;
-		const FVector ComputedDelta = Stiffness * Offset * Direction;
+		const float Offset = Length - RefLength * Scale;
+		const FVector3f ComputedDelta = Stiffness * Offset * Direction;
 
-		const FVector Delta = VectorSelect(Offset < ZERO, ZeroVector, ComputedDelta);
+		const FVector3f Delta = VectorSelect(Offset < FLOAT_ZERO, FloatZeroVector, ComputedDelta);
 
 		VectorScatter(&P[i2], P2 + Delta);
 	}
diff --git a/Engine/Source/Runtime/Experimental/Chaos/Private/Chaos/PBDLongRangeConstraintsBase.cpp b/Engine/Source/Runtime/Experimental/Chaos/Private/Chaos/PBDLongRangeConstraintsBase.cpp
--- a/Engine/Source/Runtime/Experimental/Chaos/Private/Chaos/PBDLongRangeConstraintsBase.cpp
+++ b/Engine/Source/Runtime/Experimental/Chaos/Private/Chaos/PBDLongRangeConstraintsBase.cpp
@@ -17,7 +17,7 @@ FAutoConsoleVariableRef CVarChaosPBDLongRangeConstraintsMinParallelBatchSize(
 	TEXT("The minimum number of long range tethers in a batch to process in parallel."),
 	ECVF_Cheat);
 
-using namespace Chaos;
+namespace Chaos::Softs {
 
 int32 FPBDLongRangeConstraintsBase::GetMinParallelBatchSize()
 {
@@ -25,14 +25,14 @@ int32 FPBDLongRangeConstraintsBase::GetMinParallelBatchSize()
 }
 
 FPBDLongRangeConstraintsBase::FPBDLongRangeConstraintsBase(
-	const FPBDParticles& Particles,
+	const FSolverParticles& Particles,
 	const int32 InParticleOffset,
 	const int32 InParticleCount,
 	const TArray<TConstArrayView<TTuple<int32, int32, FRealSingle>>>& InTethers,
 	const TConstArrayView<FRealSingle>& StiffnessMultipliers,
 	const TConstArrayView<FRealSingle>& ScaleMultipliers,
-	const FVec2& InStiffness,
-	const FVec2& InScale)
+	const FSolverVec2& InStiffness,
+	const FSolverVec2& InScale)
 	: Tethers(InTethers)
 	, Stiffness(InStiffness, StiffnessMultipliers, InParticleCount)
 	, Scale(InScale)
@@ -62,24 +62,25 @@ FPBDLongRangeConstraintsBase::FPBDLongRangeConstraintsBase(
 	}
 
 	// Apply default properties
-	ApplyProperties((FReal)(1. / FPBDStiffness::ParameterFrequency), 1);
+	ApplyProperties((FSolverReal)(1. / FPBDStiffness::ParameterFrequency), 1);
 }
 
-void FPBDLongRangeConstraintsBase::ApplyProperties(const FReal Dt, const int32 NumIterations)
+void FPBDLongRangeConstraintsBase::ApplyProperties(const FSolverReal Dt, const int32 NumIterations)
 {
 	// Apply stiffness
 	Stiffness.ApplyValues(Dt, NumIterations);
 
 	// Apply scale
-	const FReal Offset = Scale[0];
-	const FReal Range = Scale[1] - Scale[0];
+	const FSolverReal Offset = Scale[0];
+	const FSolverReal Range = Scale[1] - Scale[0];
 	const int32 ScaleTableSize = ScaleTable.Num();
-	const FReal WeightIncrement = (ScaleTableSize > 1) ? (FReal)1. / (FReal)(ScaleTableSize - 1) : (FReal)1.; // Must allow full range from 0 to 1 included
+	const FSolverReal WeightIncrement = (ScaleTableSize > 1) ? (FSolverReal)1. / (FSolverReal)(ScaleTableSize - 1) : (FSolverReal)1.; // Must allow full range from 0 to 1 included
 	for (int32 Index = 0; Index < ScaleTableSize; ++Index)
 	{
-		const FReal Weight = (FReal)Index * WeightIncrement;
+		const FSolverReal Weight = (FSolverReal)Index * WeightIncrement;
 		ScaleTable[Index] = Offset + Weight * Range;
 	}
 }
 
+}  // End namespace Chaos::Softs
 
diff --git a/Engine/Source/Runtime/Experimental/Chaos/Private/Chaos/PBDSphericalConstraint.cpp b/Engine/Source/Runtime/Experimental/Chaos/Private/Chaos/PBDSphericalConstraint.cpp
--- a/Engine/Source/Runtime/Experimental/Chaos/Private/Chaos/PBDSphericalConstraint.cpp
+++ b/Engine/Source/Runtime/Experimental/Chaos/Private/Chaos/PBDSphericalConstraint.cpp
@@ -12,11 +12,13 @@
 #if INTEL_ISPC && !UE_BUILD_SHIPPING
 bool bChaos_Spherical_ISPC_Enabled = true;
 FAutoConsoleVariableRef CVarChaosSphericalISPCEnabled(TEXT("p.Chaos.Spherical.ISPC"), bChaos_Spherical_ISPC_Enabled, TEXT("Whether to use ISPC optimizations in spherical constraints"));
+
+static_assert(sizeof(ispc::FVector3f) == sizeof(Chaos::Softs::FSolverVec3), "sizeof(ispc::FVector3f) != sizeof(Chaos::Softs::FSolverVec3)");
 #endif
 
-using namespace Chaos;
+namespace Chaos::Softs {
 
-void FPBDSphericalConstraint::ApplyHelperISPC(FPBDParticles& Particles, const FReal Dt) const
+void FPBDSphericalConstraint::ApplyHelperISPC(FSolverParticles& Particles, const FSolverReal /*Dt*/) const
 {
 	check(bRealTypeCompatibleWithISPC);
 
@@ -24,8 +26,8 @@ void FPBDSphericalConstraint::ApplyHelperISPC(FPBDParticles& Particles, const FR
 
 #if INTEL_ISPC
 	ispc::ApplySphericalConstraints(
-		(ispc::FVector*)Particles.GetP().GetData(),
-		(const ispc::FVector*)AnimationPositions.GetData(),
+		(ispc::FVector3f*)Particles.GetP().GetData(),
+		(const ispc::FVector3f*)AnimationPositions.GetData(),
 		Particles.GetInvM().GetData(),
 		SphereRadii.GetData(),
 		SphereRadiiMultiplier,
@@ -33,7 +35,7 @@ void FPBDSphericalConstraint::ApplyHelperISPC(FPBDParticles& Particles, const FR
 		ParticleCount);
 #endif
 }
-void FPBDSphericalBackstopConstraint::ApplyLegacyHelperISPC(FPBDParticles& Particles, const FReal Dt) const
+void FPBDSphericalBackstopConstraint::ApplyLegacyHelperISPC(FSolverParticles& Particles, const FSolverReal /*Dt*/) const
 {
 	check(bRealTypeCompatibleWithISPC);
 
@@ -41,9 +43,9 @@ void FPBDSphericalBackstopConstraint::ApplyLegacyHelperISPC(FPBDParticles& Parti
 
 #if INTEL_ISPC
 	ispc::ApplyLegacySphericalBackstopConstraints(
-		(ispc::FVector*)Particles.GetP().GetData(),
-		(const ispc::FVector*)AnimationPositions.GetData(),
-		(const ispc::FVector*)AnimationNormals.GetData(),
+		(ispc::FVector3f*)Particles.GetP().GetData(),
+		(const ispc::FVector3f*)AnimationPositions.GetData(),
+		(const ispc::FVector3f*)AnimationNormals.GetData(),
 		Particles.GetInvM().GetData(),
 		SphereOffsetDistances.GetData(),
 		SphereRadii.GetData(),
@@ -53,16 +55,16 @@ void FPBDSphericalBackstopConstraint::ApplyLegacyHelperISPC(FPBDParticles& Parti
 #endif
 }
 
-void FPBDSphericalBackstopConstraint::ApplyHelperISPC(FPBDParticles& Particles, const FReal Dt) const
+void FPBDSphericalBackstopConstraint::ApplyHelperISPC(FSolverParticles& Particles, const FSolverReal /*Dt*/) const
 {
 	check(bRealTypeCompatibleWithISPC);
 	const int32 ParticleCount = SphereRadii.Num();
 
 #if INTEL_ISPC
 	ispc::ApplySphericalBackstopConstraints(
-		(ispc::FVector*)Particles.GetP().GetData(),
-		(const ispc::FVector*)AnimationPositions.GetData(),
-		(const ispc::FVector*)AnimationNormals.GetData(),
+		(ispc::FVector3f*)Particles.GetP().GetData(),
+		(const ispc::FVector3f*)AnimationPositions.GetData(),
+		(const ispc::FVector3f*)AnimationNormals.GetData(),
 		Particles.GetInvM().GetData(),
 		SphereOffsetDistances.GetData(),
 		SphereRadii.GetData(),
@@ -71,3 +73,5 @@ void FPBDSphericalBackstopConstraint::ApplyHelperISPC(FPBDParticles& Particles,
 		ParticleCount);
 #endif
 }
+
+}  // End namespace Chaos::Softs
diff --git a/Engine/Source/Runtime/Experimental/Chaos/Private/Chaos/PBDSphericalConstraint.ispc b/Engine/Source/Runtime/Experimental/Chaos/Private/Chaos/PBDSphericalConstraint.ispc
--- a/Engine/Source/Runtime/Experimental/Chaos/Private/Chaos/PBDSphericalConstraint.ispc
+++ b/Engine/Source/Runtime/Experimental/Chaos/Private/Chaos/PBDSphericalConstraint.ispc
@@ -4,20 +4,20 @@
 
 #include "Math/Vector.isph"
 
-export void ApplySphericalConstraints(uniform FVector P[],
-									const uniform FVector AnimationPositions[],
-									const uniform FReal InvM[],
+export void ApplySphericalConstraints(uniform FVector3f P[],
+									const uniform FVector3f AnimationPositions[],
+									const uniform float InvM[],
 									const uniform float SphereRadii[],
-									const uniform FReal SphereRadiiMultiplier,
+									const uniform float SphereRadiiMultiplier,
 									const uniform int32 ParticleOffset,
 									const uniform int32 ParticleCount)
 {
-	static const FReal DeadZoneSquareRadius = SMALL_NUMBER; // We will not push the particle away in the dead zone
+	static const float DeadZoneSquareRadius = FLOAT_SMALL_NUMBER; // We will not push the particle away in the dead zone
 
 	foreach(Index = 0 ... ParticleCount)
 	{
 		const int32 ParticleIndex = ParticleOffset + Index;
-		const FReal IM = InvM[ParticleIndex];
+		const float IM = InvM[ParticleIndex];
 
 		if(all(IM == 0))
 		{
@@ -26,15 +26,15 @@ export void ApplySphericalConstraints(uniform FVector P[],
 
 		const uniform int32 Offset = extract(ParticleIndex, 0);
 		
-		const FVector Center = VectorLoad(&AnimationPositions[Offset]);
-		const FVector Position = VectorLoad(&P[Offset]);
+		const FVector3f Center = VectorLoad(&AnimationPositions[Offset]);
+		const FVector3f Position = VectorLoad(&P[Offset]);
 
-		const FVector CenterToParticle = Position - Center;
-		const FReal DistanceSquared = VectorSizeSquared(CenterToParticle);
-		const FReal Distance = sqrt(DistanceSquared);
+		const FVector3f CenterToParticle = Position - Center;
+		const float DistanceSquared = VectorSizeSquared(CenterToParticle);
+		const float Distance = sqrt(DistanceSquared);
 
-		const FReal Radius = SphereRadii[Index] * SphereRadiiMultiplier;
-		const FVector FinalPosition = ((Radius / Distance) * CenterToParticle) + Center;
+		const float Radius = SphereRadii[Index] * SphereRadiiMultiplier;
+		const FVector3f FinalPosition = ((Radius / Distance) * CenterToParticle) + Center;
 
 		if (IM != 0 && DistanceSquared > ((Radius * Radius) + DeadZoneSquareRadius))
 		{
@@ -43,42 +43,42 @@ export void ApplySphericalConstraints(uniform FVector P[],
 	}
 }
 
-export void ApplyLegacySphericalBackstopConstraints(uniform FVector P[],
-													const uniform FVector AnimationPositions[],
-													const uniform FVector AnimationNormals[],
-													const uniform FReal InvM[],
+export void ApplyLegacySphericalBackstopConstraints(uniform FVector3f P[],
+													const uniform FVector3f AnimationPositions[],
+													const uniform FVector3f AnimationNormals[],
+													const uniform float InvM[],
 													const uniform float SphereOffsetDistances[],
 													const uniform float SphereRadii[],
-													const uniform FReal SphereRadiiMultiplier,
+													const uniform float SphereRadiiMultiplier,
 													const uniform int32 ParticleOffset,
 													const uniform int32 ParticleCount)
 {
-	static const FReal DeadZoneSquareRadius = SMALL_NUMBER; // We will not push the particle away in the dead zone
+	static const float DeadZoneSquareRadius = FLOAT_SMALL_NUMBER; // We will not push the particle away in the dead zone
 
 	foreach(Index = 0 ... ParticleCount)
 	{
 		const int32 ParticleIndex = ParticleOffset + Index;
-		const FReal IM = InvM[ParticleIndex];
+		const float IM = InvM[ParticleIndex];
 
 		if(all(IM == 0))
 		{
 			continue;
 		}
 
-		const FVector AnimationPosition = VectorLoad(&AnimationPositions[extract(ParticleIndex, 0)]);
-		const FVector AnimationNormal = VectorLoad(&AnimationNormals[extract(ParticleIndex, 0)]);
+		const FVector3f AnimationPosition = VectorLoad(&AnimationPositions[extract(ParticleIndex, 0)]);
+		const FVector3f AnimationNormal = VectorLoad(&AnimationNormals[extract(ParticleIndex, 0)]);
 
 		const float SphereOffsetDistance = SphereOffsetDistances[Index];
-		const FReal Radius = SphereRadii[Index] * SphereRadiiMultiplier;
+		const float Radius = SphereRadii[Index] * SphereRadiiMultiplier;
 
-		const FVector Center = AnimationPosition - SphereOffsetDistance * AnimationNormal;  // Legacy version already includes the radius within the distance
-		const FVector CenterToParticle = VectorLoad(&P[extract(ParticleIndex, 0)]) - Center;
-		const FReal DistanceSquared = VectorSizeSquared(CenterToParticle);
+		const FVector3f Center = AnimationPosition - SphereOffsetDistance * AnimationNormal;  // Legacy version already includes the radius within the distance
+		const FVector3f CenterToParticle = VectorLoad(&P[extract(ParticleIndex, 0)]) - Center;
+		const float DistanceSquared = VectorSizeSquared(CenterToParticle);
 
-		const FVector DistSquaredLessRadiusSquared = ((Radius / sqrt(DistanceSquared)) * CenterToParticle) + Center;
-		const FVector DistSquaredLessDeadZoneRadius = AnimationPosition - (SphereOffsetDistance - Radius) * AnimationNormal;  // Legacy version already includes the radius to the distance
+		const FVector3f DistSquaredLessRadiusSquared = ((Radius / sqrt(DistanceSquared)) * CenterToParticle) + Center;
+		const FVector3f DistSquaredLessDeadZoneRadius = AnimationPosition - (SphereOffsetDistance - Radius) * AnimationNormal;  // Legacy version already includes the radius to the distance
 
-		FVector NewP = VectorSelect(DistanceSquared < Radius * Radius, DistSquaredLessRadiusSquared, ZeroVector);
+		FVector3f NewP = VectorSelect(DistanceSquared < Radius * Radius, DistSquaredLessRadiusSquared, FloatZeroVector);
 		NewP = VectorSelect(DistanceSquared < DeadZoneSquareRadius, DistSquaredLessDeadZoneRadius, NewP);
 
 		if (IM != 0 && NewP.V[0] != 0)
@@ -88,42 +88,42 @@ export void ApplyLegacySphericalBackstopConstraints(uniform FVector P[],
 	}
 }
 
-export void ApplySphericalBackstopConstraints(uniform FVector P[],
-												const uniform FVector AnimationPositions[],
-												const uniform FVector AnimationNormals[],
-												const uniform FReal InvM[],
+export void ApplySphericalBackstopConstraints(uniform FVector3f P[],
+												const uniform FVector3f AnimationPositions[],
+												const uniform FVector3f AnimationNormals[],
+												const uniform float InvM[],
 												const uniform float SphereOffsetDistances[],
 												const uniform float SphereRadii[],
-												const uniform FReal SphereRadiiMultiplier,
+												const uniform float SphereRadiiMultiplier,
 												const uniform int32 ParticleOffset,
 												const uniform int32 ParticleCount)
 {
-	static const FReal DeadZoneSquareRadius = SMALL_NUMBER; // We will not push the particle away in the dead zone
+	static const float DeadZoneSquareRadius = FLOAT_SMALL_NUMBER; // We will not push the particle away in the dead zone
 
 	foreach(Index = 0 ... ParticleCount)
 	{
 		const int32 ParticleIndex = ParticleOffset + Index;
-		const FReal IM = InvM[ParticleIndex];
+		const float IM = InvM[ParticleIndex];
 
 		if(all(IM == 0))
 		{
 			continue;
 		}
 
-		const FVector AnimationPosition = VectorLoad(&AnimationPositions[extract(ParticleIndex, 0)]);
-		const FVector AnimationNormal = VectorLoad(&AnimationNormals[extract(ParticleIndex, 0)]);
+		const FVector3f AnimationPosition = VectorLoad(&AnimationPositions[extract(ParticleIndex, 0)]);
+		const FVector3f AnimationNormal = VectorLoad(&AnimationNormals[extract(ParticleIndex, 0)]);
 
 		const float SphereOffsetDistance = SphereOffsetDistances[Index];
-		const FReal Radius = SphereRadii[Index] * SphereRadiiMultiplier;
+		const float Radius = SphereRadii[Index] * SphereRadiiMultiplier;
 
-		const FVector Center = AnimationPosition - (Radius + SphereOffsetDistance) * AnimationNormal;  // Non legacy version adds radius to the distance
-		const FVector CenterToParticle = VectorLoad(&P[extract(ParticleIndex, 0)]) - Center;
-		const FReal DistanceSquared = VectorSizeSquared(CenterToParticle);
+		const FVector3f Center = AnimationPosition - (Radius + SphereOffsetDistance) * AnimationNormal;  // Non legacy version adds radius to the distance
+		const FVector3f CenterToParticle = VectorLoad(&P[extract(ParticleIndex, 0)]) - Center;
+		const float DistanceSquared = VectorSizeSquared(CenterToParticle);
 
-		const FVector DistSquaredLessRadiusSquared = ((Radius / sqrt(DistanceSquared)) * CenterToParticle) + Center;
-		const FVector DistSquaredLessDeadZoneRadius = AnimationPosition - SphereOffsetDistance * AnimationNormal;  // Non legacy version adds radius to the distance
+		const FVector3f DistSquaredLessRadiusSquared = ((Radius / sqrt(DistanceSquared)) * CenterToParticle) + Center;
+		const FVector3f DistSquaredLessDeadZoneRadius = AnimationPosition - SphereOffsetDistance * AnimationNormal;  // Non legacy version adds radius to the distance
 
-		FVector NewP = VectorSelect(DistanceSquared < Radius * Radius, DistSquaredLessRadiusSquared, ZeroVector);
+		FVector3f NewP = VectorSelect(DistanceSquared < Radius * Radius, DistSquaredLessRadiusSquared, FloatZeroVector);
 		NewP = VectorSelect(DistanceSquared < DeadZoneSquareRadius, DistSquaredLessDeadZoneRadius, NewP);
 
 		if (IM != 0 && NewP.V[0] != 0)
diff --git a/Engine/Source/Runtime/Experimental/Chaos/Private/Chaos/PBDSpringConstraints.cpp b/Engine/Source/Runtime/Experimental/Chaos/Private/Chaos/PBDSpringConstraints.cpp
--- a/Engine/Source/Runtime/Experimental/Chaos/Private/Chaos/PBDSpringConstraints.cpp
+++ b/Engine/Source/Runtime/Experimental/Chaos/Private/Chaos/PBDSpringConstraints.cpp
@@ -15,14 +15,14 @@
 DECLARE_CYCLE_STAT(TEXT("Chaos PBD Spring Constraint"), STAT_PBD_Spring, STATGROUP_Chaos);
 
 #if INTEL_ISPC && !UE_BUILD_SHIPPING
-static_assert(sizeof(ispc::FVector) == sizeof(Chaos::FVec3), "sizeof(ispc::FVector) != sizeof(Chaos::FVec3)");
+static_assert(sizeof(ispc::FVector3f) == sizeof(Chaos::Softs::FSolverVec3), "sizeof(ispc::FVector3f) != sizeof(Chaos::Softs::FSolverVec3)");
 static_assert(sizeof(ispc::FIntVector2) == sizeof(Chaos::TVec2<int32>), "sizeof(ispc::FIntVector2) != sizeof(Chaos::TVec2<int32>)");
 
 bool bChaos_Spring_ISPC_Enabled = true;
 FAutoConsoleVariableRef CVarChaosSpringISPCEnabled(TEXT("p.Chaos.Spring.ISPC"), bChaos_Spring_ISPC_Enabled, TEXT("Whether to use ISPC optimizations in Spring constraints"));
 #endif
 
-using namespace Chaos;
+namespace Chaos::Softs {
 
 // @todo(chaos): the parallel threshold (or decision to run parallel) should probably be owned by the solver and passed to the constraint container
 static int32 Chaos_Spring_ParallelConstraintCount = 100;
@@ -30,7 +30,7 @@ static int32 Chaos_Spring_ParallelConstraintCount = 100;
 FAutoConsoleVariableRef CVarChaosSpringParallelConstraintCount(TEXT("p.Chaos.Spring.ParallelConstraintCount"), Chaos_Spring_ParallelConstraintCount, TEXT("If we have more constraints than this, use parallel-for in Apply."));
 #endif
 
-void FPBDSpringConstraints::InitColor(const FPBDParticles& Particles)
+void FPBDSpringConstraints::InitColor(const FSolverParticles& Particles)
 {
 	// In dev builds we always color so we can tune the system without restarting. See Apply()
 #if UE_BUILD_SHIPPING || UE_BUILD_TEST
@@ -41,30 +41,30 @@ void FPBDSpringConstraints::InitColor(const FPBDParticles& Particles)
 	}
 }
 
-void FPBDSpringConstraints::ApplyHelper(FPBDParticles& Particles, const FReal Dt, const int32 ConstraintIndex, const FReal ExpStiffnessValue) const
+void FPBDSpringConstraints::ApplyHelper(FSolverParticles& Particles, const FSolverReal Dt, const int32 ConstraintIndex, const FSolverReal ExpStiffnessValue) const
 {
 	const TVec2<int32>& Constraint = Constraints[ConstraintIndex];
 	const int32 i1 = Constraint[0];
 	const int32 i2 = Constraint[1];
-	const FVec3 Delta =  Base::GetDelta(Particles, ConstraintIndex, ExpStiffnessValue);
-	if (Particles.InvM(i1) > (FReal)0.)
+	const FSolverVec3 Delta =  Base::GetDelta(Particles, ConstraintIndex, ExpStiffnessValue);
+	if (Particles.InvM(i1) > (FSolverReal)0.)
 	{
 		Particles.P(i1) -= Particles.InvM(i1) * Delta;
 	}
-	if (Particles.InvM(i2) > (FReal)0.)
+	if (Particles.InvM(i2) > (FSolverReal)0.)
 	{
 		Particles.P(i2) += Particles.InvM(i2) * Delta;
 	}
 }
 
-void FPBDSpringConstraints::Apply(FPBDParticles& Particles, const FReal Dt) const
+void FPBDSpringConstraints::Apply(FSolverParticles& Particles, const FSolverReal Dt) const
 {
 	SCOPE_CYCLE_COUNTER(STAT_PBD_Spring);
 	if ((ConstraintsPerColor.Num() > 0) && (Constraints.Num() > Chaos_Spring_ParallelConstraintCount))
 	{
 		if (!Stiffness.HasWeightMap())
 		{
-			const FReal ExpStiffnessValue = (FReal)Stiffness;
+			const FSolverReal ExpStiffnessValue = (FSolverReal)Stiffness;
 
 #if INTEL_ISPC
 			if (bRealTypeCompatibleWithISPC && bChaos_Spring_ISPC_Enabled)
@@ -72,7 +72,7 @@ void FPBDSpringConstraints::Apply(FPBDParticles& Particles, const FReal Dt) cons
 				for (const TArray<int32>& ConstraintBatch : ConstraintsPerColor)
 				{
 					ispc::ApplySpringConstraints(
-						(ispc::FVector*) & Particles.GetP()[0],
+						(ispc::FVector3f*) & Particles.GetP()[0],
 						(ispc::FIntVector2*) & Constraints.GetData()[0],
 						&ConstraintBatch.GetData()[0],
 						&Particles.GetInvM().GetData()[0],
@@ -102,7 +102,7 @@ void FPBDSpringConstraints::Apply(FPBDParticles& Particles, const FReal Dt) cons
 				for (const TArray<int32>& ConstraintBatch : ConstraintsPerColor)
 				{
 					ispc::ApplySpringConstraintsWithWeightMaps(
-						(ispc::FVector*) & Particles.GetP()[0],
+						(ispc::FVector3f*) & Particles.GetP()[0],
 						(ispc::FIntVector2*) & Constraints.GetData()[0],
 						&ConstraintBatch.GetData()[0],
 						&Particles.GetInvM().GetData()[0],
@@ -120,7 +120,7 @@ void FPBDSpringConstraints::Apply(FPBDParticles& Particles, const FReal Dt) cons
 					PhysicsParallelFor(ConstraintBatch.Num(), [&](const int32 Index)
 					{
 						const int32 ConstraintIndex = ConstraintBatch[Index];
-						const FReal ExpStiffnessValue = Stiffness[ConstraintIndex];
+						const FSolverReal ExpStiffnessValue = Stiffness[ConstraintIndex];
 						ApplyHelper(Particles, Dt, ConstraintIndex, ExpStiffnessValue);
 					});
 				}
@@ -131,7 +131,7 @@ void FPBDSpringConstraints::Apply(FPBDParticles& Particles, const FReal Dt) cons
 	{
 		if (!Stiffness.HasWeightMap())
 		{
-			const FReal ExpStiffnessValue = (FReal)Stiffness;
+			const FSolverReal ExpStiffnessValue = (FSolverReal)Stiffness;
 			for (int32 ConstraintIndex = 0; ConstraintIndex < Constraints.Num(); ++ConstraintIndex)
 			{
 				ApplyHelper(Particles, Dt, ConstraintIndex, ExpStiffnessValue);
@@ -141,9 +141,11 @@ void FPBDSpringConstraints::Apply(FPBDParticles& Particles, const FReal Dt) cons
 		{
 			for (int32 ConstraintIndex = 0; ConstraintIndex < Constraints.Num(); ++ConstraintIndex)
 			{
-				const FReal ExpStiffnessValue = Stiffness[ConstraintIndex];
+				const FSolverReal ExpStiffnessValue = Stiffness[ConstraintIndex];
 				ApplyHelper(Particles, Dt, ConstraintIndex, ExpStiffnessValue);
 			}
 		}
 	}
 }
+
+} // End namespace Chaos::Softs
diff --git a/Engine/Source/Runtime/Experimental/Chaos/Private/Chaos/PBDSpringConstraints.ispc b/Engine/Source/Runtime/Experimental/Chaos/Private/Chaos/PBDSpringConstraints.ispc
--- a/Engine/Source/Runtime/Experimental/Chaos/Private/Chaos/PBDSpringConstraints.ispc
+++ b/Engine/Source/Runtime/Experimental/Chaos/Private/Chaos/PBDSpringConstraints.ispc
@@ -4,20 +4,20 @@
 
 #include "Math/Vector.isph"
 
-static inline FReal SafeNormalize(FVector &Direction)
+static inline float SafeNormalize(FVector3f &Direction)
 {
-	const FReal SizeSquared = VectorSizeSquared(Direction);
-	const FReal Size = sqrt(SizeSquared);
-	Direction = VectorSelect((SizeSquared < KINDA_SMALL_NUMBER), ForwardVector, Direction / Size);
-	return (SizeSquared < KINDA_SMALL_NUMBER) ? ZERO : Size;
+	const float SizeSquared = VectorSizeSquared(Direction);
+	const float Size = sqrt(SizeSquared);
+	Direction = VectorSelect((SizeSquared < FLOAT_KINDA_SMALL_NUMBER), FloatForwardVector, Direction / Size);
+	return (SizeSquared < FLOAT_KINDA_SMALL_NUMBER) ? FLOAT_ZERO : Size;
 }
 
-export void ApplySpringConstraints(uniform FVector P[],
+export void ApplySpringConstraints(uniform FVector3f P[],
 									const uniform FIntVector2 AllConstraints[],
 									const uniform int32 ActiveConstraints[],
-									const uniform FReal InvM[],
-									const uniform FReal Dists[],
-									const uniform FReal Stiffness,
+									const uniform float InvM[],
+									const uniform float Dists[],
+									const uniform float Stiffness,
 									const uniform int32 NumConstraints)
 {
 	varying int32 ActiveConstraint;
@@ -43,14 +43,14 @@ export void ApplySpringConstraints(uniform FVector P[],
 		const varying int32 i1 = Constraint.V[0];
 		const varying int32 i2 = Constraint.V[1];
 
-		const varying FVector P1 = VectorGather(&P[i1]);
-		const varying FVector P2 = VectorGather(&P[i2]);
+		const varying FVector3f P1 = VectorGather(&P[i1]);
+		const varying FVector3f P2 = VectorGather(&P[i2]);
 
-		FVector Direction = P1 - P2;
-		const FReal Distance = SafeNormalize(Direction);
+		FVector3f Direction = P1 - P2;
+		const float Distance = SafeNormalize(Direction);
 
 		#pragma ignore warning(perf)
-		const varying FReal Dist = Dists[ActiveConstraint];
+		const varying float Dist = Dists[ActiveConstraint];
 
 #if HW_GATHER_SUPPORTED == 1
 		if(i + programCount < NumConstraints)
@@ -61,36 +61,36 @@ export void ApplySpringConstraints(uniform FVector P[],
 		}
 #endif
 
-		const FVector InnerDelta = (Distance - Dist) * Direction;
+		const FVector3f InnerDelta = (Distance - Dist) * Direction;
 
 		#pragma ignore warning(perf)
-		const varying FReal InvM1 = InvM[i1];
+		const varying float InvM1 = InvM[i1];
 		#pragma ignore warning(perf)
-		const varying FReal InvM2 = InvM[i2];
+		const varying float InvM2 = InvM[i2];
 
-		const FReal CombinedInvMass = InvM2 + InvM1;
-		const FVector ComputedDelta = Stiffness * InnerDelta / CombinedInvMass;
+		const float CombinedInvMass = InvM2 + InvM1;
+		const FVector3f ComputedDelta = Stiffness * InnerDelta / CombinedInvMass;
 
-		const FVector Delta = VectorSelect((InvM2 == ZERO && InvM1 == ZERO), ZeroVector, ComputedDelta);
+		const FVector3f Delta = VectorSelect((InvM2 == FLOAT_ZERO && InvM1 == FLOAT_ZERO), FloatZeroVector, ComputedDelta);
 
-		if (InvM1 > ZERO)
+		if (InvM1 > FLOAT_ZERO)
 		{
 			VectorScatter(&P[i1], P1 - (InvM1 * Delta));
 		}
-		if (InvM2 > ZERO)
+		if (InvM2 > FLOAT_ZERO)
 		{
 			VectorScatter(&P[i2], P2 + (InvM2 * Delta));
 		}
 	}
 }
 
-export void ApplySpringConstraintsWithWeightMaps(uniform FVector P[],
+export void ApplySpringConstraintsWithWeightMaps(uniform FVector3f P[],
 									const uniform FIntVector2 AllConstraints[],
 									const uniform int32 ActiveConstraints[],
-									const uniform FReal InvM[],
-									const uniform FReal Dists[],
+									const uniform float InvM[],
+									const uniform float Dists[],
 									const uniform uint8 StiffnessIndices[],
-									const uniform FReal StiffnessTable[],
+									const uniform float StiffnessTable[],
 									const uniform int32 NumConstraints)
 {
 	varying int32 ActiveConstraint;
@@ -116,14 +116,14 @@ export void ApplySpringConstraintsWithWeightMaps(uniform FVector P[],
 		const varying int32 i1 = Constraint.V[0];
 		const varying int32 i2 = Constraint.V[1];
 
-		const varying FVector P1 = VectorGather(&P[i1]);
-		const varying FVector P2 = VectorGather(&P[i2]);
+		const varying FVector3f P1 = VectorGather(&P[i1]);
+		const varying FVector3f P2 = VectorGather(&P[i2]);
 
-		FVector Direction = P1 - P2;
-		const FReal Distance = SafeNormalize(Direction);
+		FVector3f Direction = P1 - P2;
+		const float Distance = SafeNormalize(Direction);
 
 		#pragma ignore warning(perf)
-		const varying FReal Dist = Dists[ActiveConstraint];
+		const varying float Dist = Dists[ActiveConstraint];
 
 		#pragma ignore warning(perf)
 		const varying uint8 StiffnessIndex = StiffnessIndices[ActiveConstraint];
@@ -137,26 +137,26 @@ export void ApplySpringConstraintsWithWeightMaps(uniform FVector P[],
 		}
 #endif
 
-		const FVector InnerDelta = (Distance - Dist) * Direction;
+		const FVector3f InnerDelta = (Distance - Dist) * Direction;
 
 		#pragma ignore warning(perf)
-		const varying FReal InvM1 = InvM[i1];
+		const varying float InvM1 = InvM[i1];
 		#pragma ignore warning(perf)
-		const varying FReal InvM2 = InvM[i2];
+		const varying float InvM2 = InvM[i2];
 
 		#pragma ignore warning(perf)
-		const varying FReal Stiffness = StiffnessTable[StiffnessIndex];
+		const varying float Stiffness = StiffnessTable[StiffnessIndex];
 
-		const FReal CombinedInvMass = InvM2 + InvM1;
-		const FVector ComputedDelta = Stiffness * InnerDelta / CombinedInvMass;
+		const float CombinedInvMass = InvM2 + InvM1;
+		const FVector3f ComputedDelta = Stiffness * InnerDelta / CombinedInvMass;
 
-		const FVector Delta = VectorSelect((InvM2 == ZERO && InvM1 == ZERO), ZeroVector, ComputedDelta);
+		const FVector3f Delta = VectorSelect((InvM2 == FLOAT_ZERO && InvM1 == FLOAT_ZERO), FloatZeroVector, ComputedDelta);
 
-		if (InvM1 > ZERO)
+		if (InvM1 > FLOAT_ZERO)
 		{
 			VectorScatter(&P[i1], P1 - (InvM1 * Delta));
 		}
-		if (InvM2 > ZERO)
+		if (InvM2 > FLOAT_ZERO)
 		{
 			VectorScatter(&P[i2], P2 + (InvM2 * Delta));
 		}
diff --git a/Engine/Source/Runtime/Experimental/Chaos/Private/Chaos/PerParticleDampVelocity.cpp b/Engine/Source/Runtime/Experimental/Chaos/Private/Chaos/PerParticleDampVelocity.cpp
--- a/Engine/Source/Runtime/Experimental/Chaos/Private/Chaos/PerParticleDampVelocity.cpp
+++ b/Engine/Source/Runtime/Experimental/Chaos/Private/Chaos/PerParticleDampVelocity.cpp
@@ -1,83 +1,87 @@
 // Copyright Epic Games, Inc. All Rights Reserved.
 
 #include "Chaos/PerParticleDampVelocity.h"
+#include "Chaos/Matrix.h"
+#include "GenericPlatform/GenericPlatformMath.h"
 
 #if INTEL_ISPC
 #include "PerParticleDampVelocity.ispc.generated.h"
 #endif
 
 #if INTEL_ISPC && !UE_BUILD_SHIPPING
-static_assert(sizeof(ispc::FVector) == sizeof(Chaos::FVec3), "sizeof(ispc::FVector) != sizeof(Chaos::FVec3)");
+static_assert(sizeof(ispc::FVector3f) == sizeof(Chaos::Softs::FSolverVec3), "sizeof(ispc::FVector3f) != sizeof(Chaos::Softs::FSolverVec3)");
 
 bool bChaos_DampVelocity_ISPC_Enabled = true;
 FAutoConsoleVariableRef CVarChaosDampVelocityISPCEnabled(TEXT("p.Chaos.DampVelocity.ISPC"), bChaos_DampVelocity_ISPC_Enabled, TEXT("Whether to use ISPC optimizations in per particle damp velocity calculation"));
 #endif
 
-using namespace Chaos;
+namespace Chaos::Softs {
 
-void FPerParticleDampVelocity::UpdatePositionBasedState(const FPBDParticles& Particles, const int32 Offset, const int32 Range)
+void FPerParticleDampVelocity::UpdatePositionBasedState(const FSolverParticles& Particles, const int32 Offset, const int32 Range)
 {
+#if INTEL_ISPC
 	if (bRealTypeCompatibleWithISPC && bChaos_DampVelocity_ISPC_Enabled)
 	{
-#if INTEL_ISPC
 		ispc::UpdatePositionBasedState(
-			(ispc::FVector&)MXcm,
-			(ispc::FVector&)MVcm,
-			(ispc::FVector&)MOmega,
-			(const ispc::FVector*)Particles.XArray().GetData(),
-			(const ispc::FVector*)Particles.GetV().GetData(),
+			(ispc::FVector3f&)Xcm,
+			(ispc::FVector3f&)Vcm,
+			(ispc::FVector3f&)Omega,
+			(const ispc::FVector3f*)Particles.XArray().GetData(),
+			(const ispc::FVector3f*)Particles.GetV().GetData(),
 			Particles.GetM().GetData(),
 			Particles.GetInvM().GetData(),
 			Offset,
 			Range);
-#endif
 	}
 	else
+#endif
 	{
-		MXcm = FVec3(0);
-		MVcm = FVec3(0);
-		FReal Mcm = (FReal)0;
+		Xcm = FSolverVec3(0.);
+		Vcm = FSolverVec3(0.);
+		FSolverReal Mcm = (FSolverReal)0.;
 
 		for (int32 Index = Offset; Index < Range; ++Index)
 		{
-			if (Particles.InvM(Index) == (FReal)0)
+			if (Particles.InvM(Index) == (FSolverReal)0.)
 			{
 				continue;
 			}
 
-			MXcm += Particles.X(Index) * Particles.M(Index);
-			MVcm += Particles.V(Index) * Particles.M(Index);
+			Xcm += Particles.X(Index) * Particles.M(Index);
+			Vcm += Particles.V(Index) * Particles.M(Index);
 			Mcm += Particles.M(Index);
 		}
 
-		if (Mcm != (FReal)0.0)
+		if (Mcm != (FSolverReal)0.)
 		{
-			MXcm /= Mcm;
-			MVcm /= Mcm;
+			Xcm /= Mcm;
+			Vcm /= Mcm;
 		}
 
-		FVec3 L = FVec3(0);
-		FMatrix33 I(0);
+		FSolverVec3 L = FSolverVec3(0.);
+		FSolverMatrix33 I(0.);
 		for (int32 Index = Offset; Index < Range; ++Index)
 		{
-			if (Particles.InvM(Index) == (FReal)0)
+			if (Particles.InvM(Index) == (FSolverReal)0.)
 			{
 				continue;
 			}
 
-			const FVec3 V = Particles.X(Index) - MXcm;
-			L += FVec3::CrossProduct(V, Particles.M(Index) * Particles.V(Index));
-			const FMatrix33 M(0, V[2], -V[1], -V[2], 0, V[0], V[1], -V[0], 0);
+			const FSolverVec3 V = Particles.X(Index) - Xcm;
+			L += FSolverVec3::CrossProduct(V, Particles.M(Index) * Particles.V(Index));
+			const FSolverMatrix33 M(0, V[2], -V[1], -V[2], 0, V[0], V[1], -V[0], 0);
 			I += M.GetTransposed() * M * Particles.M(Index);
 		}
 
+		const FSolverReal Det = I.Determinant();
+		Omega = Det < (FSolverReal)SMALL_NUMBER || !FMath::IsFinite(Det) ?
+			FSolverVec3(0.) :
 #if COMPILE_WITHOUT_UNREAL_SUPPORT
-		MOmega = I.Determinant() > 1e-7 ? FRigidTransform3(I).InverseTransformVector(L) : FVec3(0);
+			FSolverRigidTransform3(I).InverseTransformVector(L);
 #else
-		const FReal Det = I.Determinant();
-		MOmega = Det < SMALL_NUMBER || !FGenericPlatformMath::IsFinite(Det) ?
-			FVec3(0) :
 			I.InverseTransformVector(L); // Calls FMatrix::InverseFast(), which tests against SMALL_NUMBER
 #endif
 	}
-}
\ No newline at end of file
+}
+
+}  // End namespace Chaos::Softs
