Chaos Cloth - Move back the post LWC solver to float for performance reason, since cloth is being simulated in local space.

diff --git a/Engine/Source/Runtime/Experimental/Chaos/Private/Chaos/PerParticleDampVelocity.ispc b/Engine/Source/Runtime/Experimental/Chaos/Private/Chaos/PerParticleDampVelocity.ispc
--- a/Engine/Source/Runtime/Experimental/Chaos/Private/Chaos/PerParticleDampVelocity.ispc
+++ b/Engine/Source/Runtime/Experimental/Chaos/Private/Chaos/PerParticleDampVelocity.ispc
@@ -11,28 +11,23 @@ static inline uniform bool IsFinite(uniform float A)
 	return (intbits(A) & 0x7F800000U) != 0x7F800000U;
 }
 
-static inline uniform bool IsFinite(uniform double A)
-{
-	return (intbits(A) & 0x7FF0000000000000ULL) != 0x7FF0000000000000ULL;
-}
-
-export void UpdatePositionBasedState(uniform FVector& MXcm,
-									uniform FVector& MVcm,
-									uniform FVector& MOmega,
-									const uniform FVector X[],
-									const uniform FVector V[],
-									const uniform FReal M[],
-									const uniform FReal InvM[],
+export void UpdatePositionBasedState(uniform FVector3f& MXcm,
+									uniform FVector3f& MVcm,
+									uniform FVector3f& MOmega,
+									const uniform FVector3f X[],
+									const uniform FVector3f V[],
+									const uniform float M[],
+									const uniform float InvM[],
 									const uniform int Offset,
 									const uniform int32 Range)
 {
-	varying FVector Xcm = ZeroVector;
-	varying FVector Vcm = ZeroVector;
-	varying FReal Mcm = 0.;
+	varying FVector3f Xcm = FloatZeroVector;
+	varying FVector3f Vcm = FloatZeroVector;
+	varying float Mcm = 0.;
 
 	foreach(Index = Offset ... Range)
 	{
-		const varying FReal IM = InvM[Index];
+		const varying float IM = InvM[Index];
 
 		if(all(!IM))
 		{
@@ -44,9 +39,9 @@ export void UpdatePositionBasedState(uniform FVector& MXcm,
 			continue;
 		}
 
-		const varying FVector MX = VectorLoad(&X[extract(Index, 0)]);
-		const varying FVector MV = VectorLoad(&V[extract(Index, 0)]);
-		const varying FReal MM = M[Index];
+		const varying FVector3f MX = VectorLoad(&X[extract(Index, 0)]);
+		const varying FVector3f MV = VectorLoad(&V[extract(Index, 0)]);
+		const varying float MM = M[Index];
 
 		Xcm = Xcm + (MX * MM);
 		Vcm = Vcm + (MV * MM);
@@ -55,20 +50,20 @@ export void UpdatePositionBasedState(uniform FVector& MXcm,
 
 	MXcm = VectorReduceAdd(Xcm);
 	MVcm = VectorReduceAdd(Vcm);
-	uniform FReal MMcm = reduce_add(Mcm);
+	uniform float MMcm = reduce_add(Mcm);
 
-	if (MMcm != ZERO)
+	if (MMcm != FLOAT_ZERO)
 	{
 		MXcm = MXcm / MMcm;
 		MVcm = MVcm / MMcm;
 	}
 
-	varying FVector VL = ZeroVector;
-	varying FMatrix33 VI = Matrix33Zero;
+	varying FVector3f VL = FloatZeroVector;
+	varying FMatrix33f VI = Matrix33FloatZero;
 
 	foreach(Index = Offset ... Range)
 	{
-		const varying FReal IM = InvM[Index];
+		const varying float IM = InvM[Index];
 
 		if(all(!IM))
 		{
@@ -80,21 +75,22 @@ export void UpdatePositionBasedState(uniform FVector& MXcm,
 			continue;
 		}
 
-		const varying FVector MX = VectorLoad(&X[extract(Index, 0)]);
-		const varying FVector MV = VectorLoad(&V[extract(Index, 0)]);
-		const varying FReal MM = M[Index];
+		const varying FVector3f MX = VectorLoad(&X[extract(Index, 0)]);
+		const varying FVector3f MV = VectorLoad(&V[extract(Index, 0)]);
+		const varying float MM = M[Index];
 
-		const varying FVector VX = MX - MXcm;
+		const varying FVector3f VX = MX - MXcm;
 		VL = VL + VectorCross(VX, MM * MV);
 
-		const varying FMatrix33 Mat = SetMatrix33(0, VX.V[2], -VX.V[1], -VX.V[2], 0, VX.V[0], VX.V[1], -VX.V[0], 0);
+		const varying FMatrix33f Mat = SetMatrix33(0, VX.V[2], -VX.V[1], -VX.V[2], 0, VX.V[0], VX.V[1], -VX.V[0], 0);
 		VI = AddAB(VI, Multiply(MultiplyABt(Mat, Mat), MM));
 	}
 
-	const uniform FVector L = VectorReduceAdd(VL);
-	const uniform FMatrix33 I = MatrixReduceAdd(VI);
+	const uniform FVector3f L = VectorReduceAdd(VL);
+	const uniform FMatrix33f I = MatrixReduceAdd(VI);
 
-	const uniform FReal Det = MatrixDeterminant(I);
+	const uniform FVector3f InverseTransform = MatrixInverseTransformVector(I, L);
+	const uniform float Det = MatrixDeterminant(I);
 
-	MOmega = VectorSelect(Det < SMALL_NUMBER || !IsFinite(Det), ZeroVector, MatrixInverseTransformVector(I, L)); // Calls FMatrix::InverseFast(), which tests against SMALL_NUMBER
+	MOmega = VectorSelect(Det < FLOAT_SMALL_NUMBER || !IsFinite(Det), FloatZeroVector, InverseTransform); // Calls FMatrix::InverseFast(), which tests against SMALL_NUMBER
 }
diff --git a/Engine/Source/Runtime/Experimental/Chaos/Private/Chaos/PerParticlePBDCollisionConstraint.cpp b/Engine/Source/Runtime/Experimental/Chaos/Private/Chaos/PerParticlePBDCollisionConstraint.cpp
--- a/Engine/Source/Runtime/Experimental/Chaos/Private/Chaos/PerParticlePBDCollisionConstraint.cpp
+++ b/Engine/Source/Runtime/Experimental/Chaos/Private/Chaos/PerParticlePBDCollisionConstraint.cpp
@@ -8,8 +8,8 @@
 #endif
 
 #if INTEL_ISPC && !UE_BUILD_SHIPPING
-static_assert(sizeof(ispc::FVector) == sizeof(Chaos::FVec3), "sizeof(ispc::FVector) != sizeof(Chaos::FVec3)");
-static_assert(sizeof(ispc::FVector4) == sizeof(Chaos::FRotation3), "sizeof(ispc::FVector4) != sizeof(Chaos::FRotation3)");
+static_assert(sizeof(ispc::FVector3f) == sizeof(Chaos::Softs::FSolverVec3), "sizeof(ispc::FVector3f) != sizeof(Chaos::Softs::FSolverVec3)");
+static_assert(sizeof(ispc::FVector4f) == sizeof(Chaos::Softs::FSolverRotation3), "sizeof(ispc::FVector4f) != sizeof(Chaos::Softs::FSolverRotation3)");
 
 bool bChaos_PerParticleCollision_ISPC_Enabled = true;
 FAutoConsoleVariableRef CVarChaosPerParticleCollisionISPCEnabled(TEXT("p.Chaos.PerParticleCollision.ISPC"), bChaos_PerParticleCollision_ISPC_Enabled, TEXT("Whether to use ISPC optimizations in per particle collisions"));
@@ -22,32 +22,27 @@ FAutoConsoleVariableRef CVarChaosPerParticleCollisionISPCParallelBatchSize(TEXT(
 FAutoConsoleVariableRef CVarChaosPerParticleCollisionISPCFastFriction(TEXT("p.Chaos.PerParticleCollision.ISPC.FastFriction"), Chaos_PerParticleCollision_ISPC_FastFriction, TEXT("Faster friction ISPC"));
 #endif
 
-using namespace Chaos;
-
-template<EGeometryParticlesSimType SimType>
-void TPerParticlePBDCollisionConstraint<SimType>::ApplyHelperISPC(FPBDParticles& Particles, const FReal Dt, const int32 Offset, const int32 Range) const
-{
-	ApplyHelper(Particles, Dt, Offset, Range); // default to non ISPC
-}
+namespace Chaos::Softs {
 
 // Helper function to call PhiWithNormal and return data to ISPC
-extern "C" void GetPhiWithNormal(const uint8* CollisionParticles, const FReal* InV, FReal* Normal, FReal* Phi, const int32 i, const int32 ProgramCount, const int32 Mask)
+extern "C" void GetPhiWithNormal(const uint8* CollisionParticles, const FSolverReal* InV, FSolverReal* Normal, FSolverReal* Phi, const int32 i, const int32 ProgramCount, const int32 Mask)
 {
-	const TKinematicGeometryParticlesImp<FReal, 3, EGeometryParticlesSimType::Other>& C = *(const TKinematicGeometryParticlesImp<FReal, 3, EGeometryParticlesSimType::Other>*)CollisionParticles;
+	const TKinematicGeometryParticlesImp<FSolverReal, 3, EGeometryParticlesSimType::Other>& C = *(const TKinematicGeometryParticlesImp<FSolverReal, 3, EGeometryParticlesSimType::Other>*)CollisionParticles;
 	
 	for (int32 Index = 0; Index < ProgramCount; ++Index)
 	{
 		if (Mask & (1 << Index))
 		{
-			FVec3 V;
+			FSolverVec3 V;
 
 			// aos_to_soa3
 			V.X = InV[Index];
 			V.Y = InV[Index + ProgramCount];
 			V.Z = InV[Index + 2 * ProgramCount];
 
-			FVec3 Norm;
-			Phi[Index] = C.Geometry(i)->PhiWithNormal(V, Norm);
+			FVec3 ImplicitNormal;
+			Phi[Index] = (FSolverReal)C.Geometry(i)->PhiWithNormal(FVec3(V), ImplicitNormal);
+			FSolverVec3 Norm(ImplicitNormal);
 
 			// aos_to_soa3
 			Normal[Index] = Norm.X;
@@ -57,16 +52,15 @@ extern "C" void GetPhiWithNormal(const uint8* CollisionParticles, const FReal* I
 	}
 }
 
-template<>
-void TPerParticlePBDCollisionConstraint<EGeometryParticlesSimType::Other>::ApplyHelperISPC(FPBDParticles& InParticles, const FReal Dt, int32 Offset, int32 Range) const
+void FPerParticlePBDCollisionConstraint::ApplyHelperISPC(FSolverParticles& InParticles, const FSolverReal Dt, const int32 Offset, const int32 Range) const
 {
 	check(bRealTypeCompatibleWithISPC);
 
 	const uint32 DynamicGroupId = MDynamicGroupIds[Offset];
-	const FReal PerGroupFriction = MPerGroupFriction[DynamicGroupId];
-	const FReal PerGroupThickness = MPerGroupThickness[DynamicGroupId];
+	const FSolverReal PerGroupFriction = MPerGroupFriction[DynamicGroupId];
+	const FSolverReal PerGroupThickness = MPerGroupThickness[DynamicGroupId];
 
-	const int32 NumBatches = FMath::CeilToInt((FReal)(Range - Offset) / (FReal)Chaos_PerParticleCollision_ISPC_ParallelBatchSize);
+	const int32 NumBatches = FMath::CeilToInt((FSolverReal)(Range - Offset) / (FSolverReal)Chaos_PerParticleCollision_ISPC_ParallelBatchSize);
 
 	if (Chaos_PerParticleCollision_ISPC_FastFriction)
 	{
@@ -79,16 +73,16 @@ void TPerParticlePBDCollisionConstraint<EGeometryParticlesSimType::Other>::Apply
 
 #if INTEL_ISPC
 				MCollisionParticlesActiveView.RangeFor(
-					[this, &InParticles, Dt, BatchBegin, BatchEnd, DynamicGroupId, PerGroupFriction, PerGroupThickness](FCollisionParticles& CollisionParticles, int32 CollisionOffset, int32 CollisionRange)
+					[this, &InParticles, Dt, BatchBegin, BatchEnd, DynamicGroupId, PerGroupFriction, PerGroupThickness](FSolverRigidParticles& CollisionParticles, int32 CollisionOffset, int32 CollisionRange)
 					{
 						ispc::ApplyPerParticleCollisionFastFriction(
-							(ispc::FVector*)InParticles.GetP().GetData(),
-							(const ispc::FVector*)InParticles.XArray().GetData(),
+							(ispc::FVector3f*)InParticles.GetP().GetData(),
+							(const ispc::FVector3f*)InParticles.XArray().GetData(),
 							InParticles.GetInvM().GetData(),
-							(const ispc::FVector*)CollisionParticles.AllV().GetData(),
-							(const ispc::FVector*)CollisionParticles.XArray().GetData(),
-							(const ispc::FVector*)CollisionParticles.AllW().GetData(),
-							(const ispc::FVector4*)CollisionParticles.AllR().GetData(),
+							(const ispc::FVector3f*)CollisionParticles.AllV().GetData(),
+							(const ispc::FVector3f*)CollisionParticles.XArray().GetData(),
+							(const ispc::FVector3f*)CollisionParticles.AllW().GetData(),
+							(const ispc::FVector4f*)CollisionParticles.AllR().GetData(),
 							DynamicGroupId,
 							MKinematicGroupIds.GetData(),
 							PerGroupFriction,
@@ -116,16 +110,16 @@ void TPerParticlePBDCollisionConstraint<EGeometryParticlesSimType::Other>::Apply
 
 #if INTEL_ISPC
 				MCollisionParticlesActiveView.RangeFor(
-					[this, &InParticles, Dt, BatchBegin, BatchEnd, DynamicGroupId, PerGroupThickness](FCollisionParticles& CollisionParticles, int32 CollisionOffset, int32 CollisionRange)
+					[this, &InParticles, Dt, BatchBegin, BatchEnd, DynamicGroupId, PerGroupThickness](FSolverRigidParticles& CollisionParticles, int32 CollisionOffset, int32 CollisionRange)
 					{
 						ispc::ApplyPerParticleCollisionNoFriction(
-							(ispc::FVector*)InParticles.GetP().GetData(),
-							(const ispc::FVector*)InParticles.XArray().GetData(),
+							(ispc::FVector3f*)InParticles.GetP().GetData(),
+							(const ispc::FVector3f*)InParticles.XArray().GetData(),
 							InParticles.GetInvM().GetData(),
-							(const ispc::FVector*)CollisionParticles.AllV().GetData(),
-							(const ispc::FVector*)CollisionParticles.XArray().GetData(),
-							(const ispc::FVector*)CollisionParticles.AllW().GetData(),
-							(const ispc::FVector4*)CollisionParticles.AllR().GetData(),
+							(const ispc::FVector3f*)CollisionParticles.AllV().GetData(),
+							(const ispc::FVector3f*)CollisionParticles.XArray().GetData(),
+							(const ispc::FVector3f*)CollisionParticles.AllW().GetData(),
+							(const ispc::FVector4f*)CollisionParticles.AllR().GetData(),
 							DynamicGroupId,
 							MKinematicGroupIds.GetData(),
 							PerGroupThickness,
@@ -153,16 +147,16 @@ void TPerParticlePBDCollisionConstraint<EGeometryParticlesSimType::Other>::Apply
 
 #if INTEL_ISPC
 			MCollisionParticlesActiveView.RangeFor(
-				[this, &InParticles, Dt, BatchBegin, BatchEnd](FCollisionParticles& CollisionParticles, int32 CollisionOffset, int32 CollisionRange)
+				[this, &InParticles, Dt, BatchBegin, BatchEnd](FSolverRigidParticles& CollisionParticles, int32 CollisionOffset, int32 CollisionRange)
 				{
 					ispc::ApplyPerParticleCollision(
-						(ispc::FVector*)InParticles.GetP().GetData(),
-						(const ispc::FVector*)InParticles.XArray().GetData(),
+						(ispc::FVector3f*)InParticles.GetP().GetData(),
+						(const ispc::FVector3f*)InParticles.XArray().GetData(),
 						InParticles.GetInvM().GetData(),
-						(const ispc::FVector*)CollisionParticles.AllV().GetData(),
-						(const ispc::FVector*)CollisionParticles.XArray().GetData(),
-						(const ispc::FVector*)CollisionParticles.AllW().GetData(),
-						(const ispc::FVector4*)CollisionParticles.AllR().GetData(),
+						(const ispc::FVector3f*)CollisionParticles.AllV().GetData(),
+						(const ispc::FVector3f*)CollisionParticles.XArray().GetData(),
+						(const ispc::FVector3f*)CollisionParticles.AllW().GetData(),
+						(const ispc::FVector4f*)CollisionParticles.AllR().GetData(),
 						MDynamicGroupIds.GetData(),
 						MKinematicGroupIds.GetData(),
 						MPerGroupFriction.GetData(),
@@ -183,5 +177,4 @@ void TPerParticlePBDCollisionConstraint<EGeometryParticlesSimType::Other>::Apply
 	}
 }
 
-template class Chaos::TPerParticlePBDCollisionConstraint<EGeometryParticlesSimType::RigidBodySim>;
-template class Chaos::TPerParticlePBDCollisionConstraint<EGeometryParticlesSimType::Other>;
+}  // End namespace Chaos::Softs
diff --git a/Engine/Source/Runtime/Experimental/Chaos/Private/Chaos/PerParticlePBDCollisionConstraint.ispc b/Engine/Source/Runtime/Experimental/Chaos/Private/Chaos/PerParticlePBDCollisionConstraint.ispc
--- a/Engine/Source/Runtime/Experimental/Chaos/Private/Chaos/PerParticlePBDCollisionConstraint.ispc
+++ b/Engine/Source/Runtime/Experimental/Chaos/Private/Chaos/PerParticlePBDCollisionConstraint.ispc
@@ -481,11 +481,11 @@ static inline void PhiWithNormalConvex(const uniform uint8 *uniform TypedObjectP
 	const uniform int32 NumPlanes = Object->Planes.ArrayNum;
 	if (NumPlanes == 0)
 	{
-		Phi = FLT_MAX;
+		Phi = REAL_MAX;
 		return;
 	}
 
-	FReal MaxPhi = -FLT_MAX;
+	FReal MaxPhi = -REAL_MAX;
 	int32 MaxPlane = 0;
 
 	for (uniform int32 Idx = 0; Idx < NumPlanes; ++Idx)
@@ -513,7 +513,7 @@ static inline void PhiWithNormalConvex(const uniform uint8 *uniform TypedObjectP
 		return;
 	}
 
-	FReal ClosestDistance = FLT_MAX;
+	FReal ClosestDistance = REAL_MAX;
 	FVector ClosestPoint = ZeroVector;
 	for (uniform int32 Index = 0; Index < PlaneVerticesNum - 2; Index++)
 	{
@@ -544,25 +544,25 @@ static inline void PhiWithNormalConvex(const uniform uint8 *uniform TypedObjectP
 	}
 }
 
-extern "C" unmasked void GetPhiWithNormal(const uniform uint8 *uniform CollisionParticles, const uniform FReal *uniform V, uniform FReal *uniform Normal, uniform FReal *uniform Phi, const uniform int i, const uniform int ProgramCount, const uniform int Mask);
+extern "C" unmasked void GetPhiWithNormal(const uniform uint8 *uniform CollisionParticles, const uniform float *uniform V, uniform float *uniform Normal, uniform float *uniform Phi, const uniform int i, const uniform int ProgramCount, const uniform int Mask);
 
-export void ApplyPerParticleCollision(uniform FVector ParticlesP[],
-										const uniform FVector ParticlesX[],
-										const uniform FReal ParticlesInvM[],
-										const uniform FVector CollisionV[],
-										const uniform FVector CollisionX[],
-										const uniform FVector CollisionW[],
-										const uniform FVector4 CollisionR[],
+export void ApplyPerParticleCollision(uniform FVector3f ParticlesP[],
+										const uniform FVector3f ParticlesX[],
+										const uniform float ParticlesInvM[],
+										const uniform FVector3f CollisionV[],
+										const uniform FVector3f CollisionX[],
+										const uniform FVector3f CollisionW[],
+										const uniform FVector4f CollisionR[],
 										const uniform uint32 DynamicGroupIds[],
 										const uniform uint32 KinematicGroupIds[],
-										const uniform FReal PerGroupFriction[],
-										const uniform FReal PerGroupThickness[],
+										const uniform float PerGroupFriction[],
+										const uniform float PerGroupThickness[],
 										const uniform uint8 *uniform CollisionParticles,
 										const uniform uint8 *uniform Geometry,
 										const uniform int SizeofFImplicitObject,
 										const uniform int OffsetofGeometryType,
 										const uniform int OffsetOfMargin,
-										const uniform FReal Dt,
+										const uniform float Dt,
 										const uniform int InnerOffset,
 										const uniform int InnerRange,
 										const uniform int OuterOffset,
@@ -596,23 +596,28 @@ export void ApplyPerParticleCollision(uniform FVector ParticlesP[],
 
 	foreach(Index = OuterOffset ... BailAt)
 	{
-		const FReal IM = ParticlesInvM[Index];
+		const float IM = ParticlesInvM[Index];
 
 		if(all(IM == 0))
 		{
 			continue;
 		}
 
-		FVector P = VectorLoad(&ParticlesP[extract(Index, 0)]);
+		FVector3f P = VectorLoad(&ParticlesP[extract(Index, 0)]);
 
 		const uint32 DynamicGroupId = DynamicGroupIds[Index]; // Particle group Id
 
 		for(uniform int i = InnerOffset; i < InnerRange; i++) // SequentialFor
 		{
-			const uniform FTransform Frame = SetTransform(CollisionR[i], SetVector4(CollisionX[i], (uniform FReal)(0.0)), Vector1110);
-			const FVector PhiWithNormalInput = InverseTransformPosition(Frame, P);
-			FVector Normal;
-			FReal Phi;
+			// InverseTransformPosition. Can skip Scale calculation because it is 1 for xyz.
+			const FVector3f TranslatedVec = P - CollisionX[i];
+			const FVector4f InverseR = QuatInverse(CollisionR[i]);
+			const FVector3f PhiWithNormalInputFloat = VectorQuaternionRotateVector(InverseR, TranslatedVec);
+			const FVector PhiWithNormalInput = ConvertVector3fTo3Native(PhiWithNormalInputFloat);
+			FVector NormalReal;
+			FReal PhiReal;
+			FVector3f Normal;
+			float Phi;
 
 			const uniform SerializedPtr *uniform GeometryPtr = (const uniform SerializedPtr *uniform)Geometry;
 			const uniform uint8 *uniform Ptr = GeometryPtr[i].Ptr;
@@ -624,28 +629,38 @@ export void ApplyPerParticleCollision(uniform FVector ParticlesP[],
 			case Sphere:
 				{
 					const uniform FReal Margin = *((const uniform FReal *uniform)(Ptr + OffsetOfMargin));
-					PhiWithNormalSphere(TypedPtr, PhiWithNormalInput, Normal, Phi, Margin);
+					PhiWithNormalSphere(TypedPtr, PhiWithNormalInput, NormalReal, PhiReal, Margin);
+					Phi = (varying float)PhiReal;
+					Normal = ConvertVector3NativeTo3f(NormalReal);
 					break;
 				}
 			case Capsule:
 				{
 					const uniform FReal Margin = *((const uniform FReal *uniform)(Ptr + OffsetOfMargin));
-					PhiWithNormalCapsule(TypedPtr, PhiWithNormalInput, Normal, Phi, Margin);
+					PhiWithNormalCapsule(TypedPtr, PhiWithNormalInput, NormalReal, PhiReal, Margin);
+					Phi = (varying float)PhiReal;
+					Normal = ConvertVector3NativeTo3f(NormalReal);
 					break;
 				}
 			case Union:
 				{
-					PhiWithNormalUnion(TypedPtr, PhiWithNormalInput, Normal, Phi, SizeofFImplicitObject, OffsetOfMargin);
+					PhiWithNormalUnion(TypedPtr, PhiWithNormalInput, NormalReal, PhiReal, SizeofFImplicitObject, OffsetOfMargin);
+					Phi = (varying float)PhiReal;
+					Normal = ConvertVector3NativeTo3f(NormalReal);
 					break;
 				}
 			case TaperedCapsule:
 				{
-					PhiWithNormalTaperedCapsule(TypedPtr, PhiWithNormalInput, Normal, Phi);
+					PhiWithNormalTaperedCapsule(TypedPtr, PhiWithNormalInput, NormalReal, PhiReal);
+					Phi = (varying float)PhiReal;
+					Normal = ConvertVector3NativeTo3f(NormalReal);
 					break;
 				}
 			case Convex:
 				{
-					PhiWithNormalConvex(TypedPtr, PhiWithNormalInput, Normal, Phi, SizeofFImplicitObject);
+					PhiWithNormalConvex(TypedPtr, PhiWithNormalInput, NormalReal, PhiReal, SizeofFImplicitObject);
+					Phi = (varying float)PhiReal;
+					Normal = ConvertVector3NativeTo3f(NormalReal);
 					break;
 				}
 			default:
@@ -656,33 +671,34 @@ export void ApplyPerParticleCollision(uniform FVector ParticlesP[],
 						IMMask = lanemask();
 					}
 
-					GetPhiWithNormal(CollisionParticles, (const uniform FReal *uniform)&PhiWithNormalInput, (uniform FReal *uniform)&Normal, (uniform FReal *uniform)&Phi, i, programCount, IMMask);
+					GetPhiWithNormal(CollisionParticles, (const uniform float *uniform)&PhiWithNormalInputFloat, (uniform float *uniform)&Normal, (uniform float *uniform)&Phi, i, programCount, IMMask);
 					break;
 				}
 			}
 
 			#pragma ignore warning(perf)
-			const FReal Thickness = PerGroupThickness[DynamicGroupId];
+			const float Thickness = PerGroupThickness[DynamicGroupId];
 
 			if(Phi < Thickness)
 			{
-				const FVector NormalWorld = TransformVector(Frame, Normal);
+				// TransformVector. Scale is 1 so that part can be removed.
+				const FVector3f NormalWorld = VectorQuaternionRotateVector(CollisionR[i], Normal);
 				P = P + ((-Phi + Thickness) * NormalWorld);
 
 				#pragma ignore warning(perf)
-				const FReal CoefficientOfFriction = PerGroupFriction[DynamicGroupId];
+				const float CoefficientOfFriction = PerGroupFriction[DynamicGroupId];
 
-				if(CoefficientOfFriction > (uniform FReal)(0.0))
+				if(CoefficientOfFriction > 0.0f)
 				{
-					const uniform FReal MaximumFrictionCorrectionPerStep = (uniform FReal)(1.0); // 1cm absolute maximum correction the friction force can provide
-					const FReal Penetration = (-Phi + Thickness); // This is related to the Normal impulse
-
-					FVector VectorToPoint = P - CollisionX[i];
-					const FVector X = VectorLoad(&ParticlesX[extract(Index, 0)]);
-					const FVector RelativeDisplacement = (P - X) - (CollisionV[i] + VectorCross(CollisionW[i], VectorToPoint)) * Dt; // This corresponds to the tangential velocity multiplied by dt (friction will drive this to zero if it is high enough)
-					const FVector RelativeDisplacementTangent = RelativeDisplacement - VectorDot(RelativeDisplacement, NormalWorld) * NormalWorld; // Project displacement into the tangential plane
-					const FReal RelativeDisplacementTangentLength = VectorSize(RelativeDisplacementTangent);
-					FReal PositionCorrection = clamp(Penetration * CoefficientOfFriction, (uniform FReal)(0.0), RelativeDisplacementTangentLength);
+					const uniform float MaximumFrictionCorrectionPerStep = 1.0f; // 1cm absolute maximum correction the friction force can provide
+					const float Penetration = (-Phi + Thickness); // This is related to the Normal impulse
+
+					FVector3f VectorToPoint = P - CollisionX[i];
+					const FVector3f X = VectorLoad(&ParticlesX[extract(Index, 0)]);
+					const FVector3f RelativeDisplacement = (P - X) - (CollisionV[i] + VectorCross(CollisionW[i], VectorToPoint)) * Dt; // This corresponds to the tangential velocity multiplied by dt (friction will drive this to zero if it is high enough)
+					const FVector3f RelativeDisplacementTangent = RelativeDisplacement - VectorDot(RelativeDisplacement, NormalWorld) * NormalWorld; // Project displacement into the tangential plane
+					const float RelativeDisplacementTangentLength = VectorSize(RelativeDisplacementTangent);
+					float PositionCorrection = clamp(Penetration * CoefficientOfFriction, 0.0f, RelativeDisplacementTangentLength);
 					if (PositionCorrection > MaximumFrictionCorrectionPerStep)
 					{
 						PositionCorrection = MaximumFrictionCorrectionPerStep;
@@ -702,23 +718,23 @@ export void ApplyPerParticleCollision(uniform FVector ParticlesP[],
 	}
 }
 
-export void ApplyPerParticleCollisionFastFriction(uniform FVector ParticlesP[],
-													const uniform FVector ParticlesX[],
-													const uniform FReal ParticlesInvM[],
-													const uniform FVector CollisionV[],
-													const uniform FVector CollisionX[],
-													const uniform FVector CollisionW[],
-													const uniform FVector4 CollisionR[],
+export void ApplyPerParticleCollisionFastFriction(uniform FVector3f ParticlesP[],
+													const uniform FVector3f ParticlesX[],
+													const uniform float ParticlesInvM[],
+													const uniform FVector3f CollisionV[],
+													const uniform FVector3f CollisionX[],
+													const uniform FVector3f CollisionW[],
+													const uniform FVector4f CollisionR[],
 													const uniform uint32 DynamicGroupId,
 													const uniform uint32 KinematicGroupIds[],
-													const uniform FReal PerGroupFriction,
-													const uniform FReal PerGroupThickness,
+													const uniform float PerGroupFriction,
+													const uniform float PerGroupThickness,
 													const uniform uint8 *uniform CollisionParticles,
 													const uniform uint8 *uniform Geometry,
 													const uniform int SizeofFImplicitObject,
 													const uniform int OffsetofGeometryType,
 													const uniform int OffsetOfMargin,
-													const uniform FReal Dt,
+													const uniform float Dt,
 													const uniform int InnerOffset,
 													const uniform int InnerRange,
 													const uniform int OuterOffset,
@@ -750,21 +766,26 @@ export void ApplyPerParticleCollisionFastFriction(uniform FVector ParticlesP[],
 
 	foreach(Index = OuterOffset ... BailAt)
 	{
-		const FReal IM = ParticlesInvM[Index];
+		const float IM = ParticlesInvM[Index];
 
 		if(all(IM == 0))
 		{
 			continue;
 		}
 
-		FVector P = VectorLoad(&ParticlesP[extract(Index, 0)]);
+		FVector3f P = VectorLoad(&ParticlesP[extract(Index, 0)]);
 
 		for(uniform int i = InnerOffset; i < InnerRange; i++) // SequentialFor
 		{
-			const uniform FTransform Frame = SetTransform(CollisionR[i], SetVector4(CollisionX[i], (uniform FReal)(0.0)), Vector1110);
-			const FVector PhiWithNormalInput = InverseTransformPosition(Frame, P);
-			FVector Normal;
-			FReal Phi;
+			// InverseTransformPosition. Can skip Scale calculation because it is 1 for xyz.
+			const FVector3f TranslatedVec = P - CollisionX[i];
+			const FVector4f InverseR = QuatInverse(CollisionR[i]);
+			const FVector3f PhiWithNormalInputFloat = VectorQuaternionRotateVector(InverseR, TranslatedVec);
+			const FVector PhiWithNormalInput = ConvertVector3fTo3Native(PhiWithNormalInputFloat);
+			FVector NormalReal;
+			FReal PhiReal;
+			FVector3f Normal;
+			float Phi;
 
 			const uniform SerializedPtr *uniform GeometryPtr = (const uniform SerializedPtr *uniform)Geometry;
 			const uniform uint8 *uniform Ptr = GeometryPtr[i].Ptr;
@@ -776,28 +797,38 @@ export void ApplyPerParticleCollisionFastFriction(uniform FVector ParticlesP[],
 			case Sphere:
 				{
 					const uniform FReal Margin = *((const uniform FReal *uniform)(Ptr + OffsetOfMargin));
-					PhiWithNormalSphere(TypedPtr, PhiWithNormalInput, Normal, Phi, Margin);
+					PhiWithNormalSphere(TypedPtr, PhiWithNormalInput, NormalReal, PhiReal, Margin);
+					Phi = (varying float)PhiReal;
+					Normal = ConvertVector3NativeTo3f(NormalReal);
 					break;
 				}
 			case Capsule:
 				{
 					const uniform FReal Margin = *((const uniform FReal *uniform)(Ptr + OffsetOfMargin));
-					PhiWithNormalCapsule(TypedPtr, PhiWithNormalInput, Normal, Phi, Margin);
+					PhiWithNormalCapsule(TypedPtr, PhiWithNormalInput, NormalReal, PhiReal, Margin);
+					Phi = (varying float)PhiReal;
+					Normal = ConvertVector3NativeTo3f(NormalReal);
 					break;
 				}
 			case Union:
 				{
-					PhiWithNormalUnion(TypedPtr, PhiWithNormalInput, Normal, Phi, SizeofFImplicitObject, OffsetOfMargin);
+					PhiWithNormalUnion(TypedPtr, PhiWithNormalInput, NormalReal, PhiReal, SizeofFImplicitObject, OffsetOfMargin);
+					Phi = (varying float)PhiReal;
+					Normal = ConvertVector3NativeTo3f(NormalReal);
 					break;
 				}
 			case TaperedCapsule:
 				{
-					PhiWithNormalTaperedCapsule(TypedPtr, PhiWithNormalInput, Normal, Phi);
+					PhiWithNormalTaperedCapsule(TypedPtr, PhiWithNormalInput, NormalReal, PhiReal);
+					Phi = (varying float)PhiReal;
+					Normal = ConvertVector3NativeTo3f(NormalReal);
 					break;
 				}
 			case Convex:
 				{
-					PhiWithNormalConvex(TypedPtr, PhiWithNormalInput, Normal, Phi, SizeofFImplicitObject);
+					PhiWithNormalConvex(TypedPtr, PhiWithNormalInput, NormalReal, PhiReal, SizeofFImplicitObject);
+					Phi = (varying float)PhiReal;
+					Normal = ConvertVector3NativeTo3f(NormalReal);
 					break;
 				}
 			default:
@@ -808,25 +839,26 @@ export void ApplyPerParticleCollisionFastFriction(uniform FVector ParticlesP[],
 						IMMask = lanemask();
 					}
 
-					GetPhiWithNormal(CollisionParticles, (const uniform FReal *uniform)&PhiWithNormalInput, (uniform FReal *uniform)&Normal, (uniform FReal *uniform)&Phi, i, programCount, IMMask);
+					GetPhiWithNormal(CollisionParticles, (const uniform float *uniform)&PhiWithNormalInputFloat, (uniform float *uniform)&Normal, (uniform float *uniform)&Phi, i, programCount, IMMask);
 					break;
 				}
 			}
 
-			const FReal Penetration = PerGroupThickness - Phi; // This is related to the Normal impulse
-			if(Penetration > (uniform FReal)(0.0))
+			const float Penetration = PerGroupThickness - Phi; // This is related to the Normal impulse
+			if(Penetration > 0.0f)
 			{
-				const FVector NormalWorld = TransformVector(Frame, Normal);
+				// TransformVector. Scale is 1 so that part can be removed.
+				const FVector3f NormalWorld = VectorQuaternionRotateVector(CollisionR[i], Normal);
 				P = P + (Penetration * NormalWorld);
 
-				const FVector VectorToPoint = P - CollisionX[i];
-				const FVector X = VectorLoad(&ParticlesX[extract(Index, 0)]);
-				const FVector RelativeDisplacement = (P - X) - (CollisionV[i] + VectorCross(CollisionW[i], VectorToPoint)) * Dt; // This corresponds to the tangential velocity multiplied by dt (friction will drive this to zero if it is high enough)
-				const FVector RelativeDisplacementTangent = RelativeDisplacement - VectorDot(RelativeDisplacement, NormalWorld) * NormalWorld; // Project displacement into the tangential plane
-				const FReal RelativeDisplacementTangentLength = VectorSize(RelativeDisplacementTangent);
+				const FVector3f VectorToPoint = P - CollisionX[i];
+				const FVector3f X = VectorLoad(&ParticlesX[extract(Index, 0)]);
+				const FVector3f RelativeDisplacement = (P - X) - (CollisionV[i] + VectorCross(CollisionW[i], VectorToPoint)) * Dt; // This corresponds to the tangential velocity multiplied by dt (friction will drive this to zero if it is high enough)
+				const FVector3f RelativeDisplacementTangent = RelativeDisplacement - VectorDot(RelativeDisplacement, NormalWorld) * NormalWorld; // Project displacement into the tangential plane
+				const float RelativeDisplacementTangentLength = VectorSize(RelativeDisplacementTangent);
 
-				const FReal PositionCorrection = min(Penetration * PerGroupFriction, RelativeDisplacementTangentLength);
-				const FReal CorrectionRatio = select(RelativeDisplacementTangentLength < SMALL_NUMBER, (uniform FReal)(0.0), PositionCorrection / RelativeDisplacementTangentLength);
+				const float PositionCorrection = min(Penetration * PerGroupFriction, RelativeDisplacementTangentLength);
+				const float CorrectionRatio = select(RelativeDisplacementTangentLength < FLOAT_SMALL_NUMBER, 0.0f, PositionCorrection / RelativeDisplacementTangentLength);
 				P = P - (CorrectionRatio * RelativeDisplacementTangent);
 			}
 		}
@@ -838,22 +870,22 @@ export void ApplyPerParticleCollisionFastFriction(uniform FVector ParticlesP[],
 	}
 }
 
-export void ApplyPerParticleCollisionNoFriction(uniform FVector ParticlesP[],
-												const uniform FVector ParticlesX[],
-												const uniform FReal ParticlesInvM[],
-												const uniform FVector CollisionV[],
-												const uniform FVector CollisionX[],
-												const uniform FVector CollisionW[],
-												const uniform FVector4 CollisionR[],
+export void ApplyPerParticleCollisionNoFriction(uniform FVector3f ParticlesP[],
+												const uniform FVector3f ParticlesX[],
+												const uniform float ParticlesInvM[],
+												const uniform FVector3f CollisionV[],
+												const uniform FVector3f CollisionX[],
+												const uniform FVector3f CollisionW[],
+												const uniform FVector4f CollisionR[],
 												const uniform uint32 DynamicGroupId,
 												const uniform uint32 KinematicGroupIds[],
-												const uniform FReal PerGroupThickness,
+												const uniform float PerGroupThickness,
 												const uniform uint8 *uniform CollisionParticles,
 												const uniform uint8 *uniform Geometry,
 												const uniform int SizeofFImplicitObject,
 												const uniform int OffsetofGeometryType,
 												const uniform int OffsetOfMargin,
-												const uniform FReal Dt,
+												const uniform float Dt,
 												const uniform int InnerOffset,
 												const uniform int InnerRange,
 												const uniform int OuterOffset,
@@ -885,21 +917,26 @@ export void ApplyPerParticleCollisionNoFriction(uniform FVector ParticlesP[],
 
 	foreach(Index = OuterOffset ... BailAt)
 	{
-		const FReal IM = ParticlesInvM[Index];
+		const float IM = ParticlesInvM[Index];
 
 		if(all(IM == 0))
 		{
 			continue;
 		}
 
-		FVector P = VectorLoad(&ParticlesP[extract(Index, 0)]);
+		FVector3f P = VectorLoad(&ParticlesP[extract(Index, 0)]);
 
 		for(uniform int i = InnerOffset; i < InnerRange; i++) // SequentialFor
 		{
-			const uniform FTransform Frame = SetTransform(CollisionR[i], SetVector4(CollisionX[i], (uniform FReal)(0.0)), Vector1110);
-			const FVector PhiWithNormalInput = InverseTransformPosition(Frame, P);
-			FVector Normal;
-			FReal Phi;
+			// InverseTransformPosition. Can skip Scale calculation because it is 1 for xyz.
+			const FVector3f TranslatedVec = P - CollisionX[i];
+			const FVector4f InverseR = QuatInverse(CollisionR[i]);
+			const FVector3f PhiWithNormalInputFloat = VectorQuaternionRotateVector(InverseR, TranslatedVec);
+			const FVector PhiWithNormalInput = ConvertVector3fTo3Native(PhiWithNormalInputFloat);
+			FVector NormalReal;
+			FReal PhiReal;
+			FVector3f Normal;
+			float Phi;
 
 			const uniform SerializedPtr *uniform GeometryPtr = (const uniform SerializedPtr *uniform)Geometry;
 			const uniform uint8 *uniform Ptr = GeometryPtr[i].Ptr;
@@ -911,28 +948,38 @@ export void ApplyPerParticleCollisionNoFriction(uniform FVector ParticlesP[],
 			case Sphere:
 				{
 					const uniform FReal Margin = *((const uniform FReal *uniform)(Ptr + OffsetOfMargin));
-					PhiWithNormalSphere(TypedPtr, PhiWithNormalInput, Normal, Phi, Margin);
+					PhiWithNormalSphere(TypedPtr, PhiWithNormalInput, NormalReal, PhiReal, Margin);
+					Phi = (varying float)PhiReal;
+					Normal = ConvertVector3NativeTo3f(NormalReal);
 					break;
 				}
 			case Capsule:
 				{
 					const uniform FReal Margin = *((const uniform FReal *uniform)(Ptr + OffsetOfMargin));
-					PhiWithNormalCapsule(TypedPtr, PhiWithNormalInput, Normal, Phi, Margin);
+					PhiWithNormalCapsule(TypedPtr, PhiWithNormalInput, NormalReal, PhiReal, Margin);
+					Phi = (varying float)PhiReal;
+					Normal = ConvertVector3NativeTo3f(NormalReal);
 					break;
 				}
 			case Union:
 				{
-					PhiWithNormalUnion(TypedPtr, PhiWithNormalInput, Normal, Phi, SizeofFImplicitObject, OffsetOfMargin);
+					PhiWithNormalUnion(TypedPtr, PhiWithNormalInput, NormalReal, PhiReal, SizeofFImplicitObject, OffsetOfMargin);
+					Phi = (varying float)PhiReal;
+					Normal = ConvertVector3NativeTo3f(NormalReal);
 					break;
 				}
 			case TaperedCapsule:
 				{
-					PhiWithNormalTaperedCapsule(TypedPtr, PhiWithNormalInput, Normal, Phi);
+					PhiWithNormalTaperedCapsule(TypedPtr, PhiWithNormalInput, NormalReal, PhiReal);
+					Phi = (varying float)PhiReal;
+					Normal = ConvertVector3NativeTo3f(NormalReal);
 					break;
 				}
 			case Convex:
 				{
-					PhiWithNormalConvex(TypedPtr, PhiWithNormalInput, Normal, Phi, SizeofFImplicitObject);
+					PhiWithNormalConvex(TypedPtr, PhiWithNormalInput, NormalReal, PhiReal, SizeofFImplicitObject);
+					Phi = (varying float)PhiReal;
+					Normal = ConvertVector3NativeTo3f(NormalReal);
 					break;
 				}
 			default:
@@ -943,15 +990,16 @@ export void ApplyPerParticleCollisionNoFriction(uniform FVector ParticlesP[],
 						IMMask = lanemask();
 					}
 
-					GetPhiWithNormal(CollisionParticles, (const uniform FReal *uniform)&PhiWithNormalInput, (uniform FReal *uniform)&Normal, (uniform FReal *uniform)&Phi, i, programCount, IMMask);
+					GetPhiWithNormal(CollisionParticles, (const uniform float *uniform)&PhiWithNormalInputFloat, (uniform float *uniform)&Normal, (uniform float *uniform)&Phi, i, programCount, IMMask);
 					break;
 				}
 			}
 
-			const FReal Penetration = PerGroupThickness - Phi; // This is related to the Normal impulse
-			if(Penetration > (uniform FReal)(0.0))
+			const float Penetration = PerGroupThickness - Phi; // This is related to the Normal impulse
+			if(Penetration > 0.0f)
 			{
-				const FVector NormalWorld = TransformVector(Frame, Normal);
+				// TransformVector. Scale is 1 so that part can be removed.
+				const FVector3f NormalWorld = VectorQuaternionRotateVector(CollisionR[i], Normal);
 				P = P + (Penetration * NormalWorld);
 			}
 		}
diff --git a/Engine/Source/Runtime/Experimental/Chaos/Private/Chaos/Transform.cpp b/Engine/Source/Runtime/Experimental/Chaos/Private/Chaos/Transform.cpp
--- a/Engine/Source/Runtime/Experimental/Chaos/Private/Chaos/Transform.cpp
+++ b/Engine/Source/Runtime/Experimental/Chaos/Private/Chaos/Transform.cpp
@@ -5,9 +5,15 @@
 
 using namespace Chaos;
 
-PMatrix<Chaos::FReal, 4, 4> TRigidTransform<Chaos::FReal, 3>::operator*(const PMatrix<Chaos::FReal, 4, 4>& Matrix) const
+PMatrix<FRealSingle, 4, 4> TRigidTransform<FRealSingle, 3>::operator*(const PMatrix<FRealSingle, 4, 4>& Matrix) const
 {
 	// LWC_TODO: Perf pessimization
-	return ToMatrixNoScale() * static_cast<const UE::Math::TMatrix<FReal>&>(Matrix);
+	return ToMatrixNoScale() * static_cast<const UE::Math::TMatrix<FRealSingle>&>(Matrix);
+}
+
+PMatrix<FRealDouble, 4, 4> TRigidTransform<FRealDouble, 3>::operator*(const PMatrix<FRealDouble, 4, 4>& Matrix) const
+{
+	// LWC_TODO: Perf pessimization
+	return ToMatrixNoScale() * static_cast<const UE::Math::TMatrix<FRealDouble>&>(Matrix);
 }
 
diff --git a/Engine/Source/Runtime/Experimental/Chaos/Private/Chaos/TriangleMesh.cpp b/Engine/Source/Runtime/Experimental/Chaos/Private/Chaos/TriangleMesh.cpp
--- a/Engine/Source/Runtime/Experimental/Chaos/Private/Chaos/TriangleMesh.cpp
+++ b/Engine/Source/Runtime/Experimental/Chaos/Private/Chaos/TriangleMesh.cpp
@@ -18,14 +18,15 @@
 #endif
 
 #if INTEL_ISPC && !UE_BUILD_SHIPPING
-static_assert(sizeof(ispc::FVector) == sizeof(Chaos::FVec3), "sizeof(ispc::FVector) != sizeof(Chaos::FVec3)");
+static_assert(sizeof(ispc::FVector3f) == sizeof(Chaos::TVec3<Chaos::FRealSingle>), "sizeof(ispc::FVector3f) != sizeof(Chaos::TVec3<Chaos::FRealSingle>)");
 static_assert(sizeof(ispc::TArrayInt) == sizeof(TArray<int32>), "sizeof(ispc::TArrayInt) != sizeof(TArray<int32>)");
 
 bool bChaos_TriangleMesh_ISPC_Enabled = true;
 FAutoConsoleVariableRef CVarChaosTriangleMeshISPCEnabled(TEXT("p.Chaos.TriangleMesh.ISPC"), bChaos_TriangleMesh_ISPC_Enabled, TEXT("Whether to use ISPC optimizations in triangle mesh calculations"));
 #endif
 
-using namespace Chaos;
+namespace Chaos
+{
 
 FTriangleMesh::FTriangleMesh()
     : MStartIdx(0)
@@ -219,26 +220,20 @@ TArray<TVec4<int32>> FTriangleMesh::GetUniqueAdjacentElements() const
 	return BendingConstraints;
 }
 
-TArray<FVec3> FTriangleMesh::GetFaceNormals(const TConstArrayView<FVec3>& Points, const bool ReturnEmptyOnError) const
-{
-	TArray<FVec3> Normals;
-	GetFaceNormals(Normals, Points, ReturnEmptyOnError);
-	return Normals;
-}
-
 // Note:	This function assumes Counter Clockwise triangle windings in a Left Handed coordinate system
 //			If this is not the case the returned face normals may need to be inverted
-void FTriangleMesh::GetFaceNormals(TArray<FVec3>& Normals, const TConstArrayView<FVec3>& Points, const bool ReturnEmptyOnError) const
+template <typename T>
+void FTriangleMesh::GetFaceNormals(TArray<TVec3<T>>& Normals, const TConstArrayView<TVec3<T>>& Points, const bool ReturnEmptyOnError) const
 {
 	Normals.Reset(MElements.Num());
 	if (ReturnEmptyOnError)
 	{
 		for (const TVec3<int32>& Tri : MElements)
 		{
-			FVec3 p10 = Points[Tri[1]] - Points[Tri[0]];
-			FVec3 p20 = Points[Tri[2]] - Points[Tri[0]];
-			FVec3 Cross = FVec3::CrossProduct(p20, p10);
-			const FReal Size2 = Cross.SizeSquared();
+			const TVec3<T> p10 = Points[Tri[1]] - Points[Tri[0]];
+			const TVec3<T> p20 = Points[Tri[2]] - Points[Tri[0]];
+			const TVec3<T> Cross = TVec3<T>::CrossProduct(p20, p10);
+			const T Size2 = Cross.SizeSquared();
 			if (Size2 < SMALL_NUMBER)
 			{
 				//particles should not be coincident by the time they get here. Return empty to signal problem to caller
@@ -254,30 +249,41 @@ void FTriangleMesh::GetFaceNormals(TArray<FVec3>& Normals, const TConstArrayView
 	}
 	else
 	{
-		if (bRealTypeCompatibleWithISPC && bChaos_TriangleMesh_ISPC_Enabled)
+#if INTEL_ISPC
+		if (bChaos_TriangleMesh_ISPC_Enabled && TAreTypesEqual<T, FRealSingle>::Value)
 		{
 			Normals.SetNumUninitialized(MElements.Num());
-
-#if INTEL_ISPC
 			ispc::GetFaceNormals(
-				(ispc::FVector*)Normals.GetData(),
-				(ispc::FVector*)Points.GetData(),
+				(ispc::FVector3f*)Normals.GetData(),
+				(ispc::FVector3f*)Points.GetData(),
 				(ispc::FIntVector*)MElements.GetData(),
 				MElements.Num());
-#endif
 		}
 		else
+#endif
 		{
 			for (const TVec3<int32>& Tri : MElements)
 			{
-				FVec3 p10 = Points[Tri[1]] - Points[Tri[0]];
-				FVec3 p20 = Points[Tri[2]] - Points[Tri[0]];
-				FVec3 Cross = FVec3::CrossProduct(p20, p10);
+				const TVec3<FRealSingle> p10 = Points[Tri[1]] - Points[Tri[0]];
+				const TVec3<FRealSingle> p20 = Points[Tri[2]] - Points[Tri[0]];
+				const TVec3<FRealSingle> Cross = TVec3<FRealSingle>::CrossProduct(p20, p10);
 				Normals.Add(Cross.GetSafeNormal());
 			}
 		}
 	}
 }
+template CHAOS_API void FTriangleMesh::GetFaceNormals<FRealSingle>(TArray<TVec3<FRealSingle>>&, const TConstArrayView<TVec3<FRealSingle>>&, const bool) const;
+template CHAOS_API void FTriangleMesh::GetFaceNormals<FRealDouble>(TArray<TVec3<FRealDouble>>&, const TConstArrayView<TVec3<FRealDouble>>&, const bool) const;
+
+template <typename T>
+TArray<TVec3<T>> FTriangleMesh::GetFaceNormals(const TConstArrayView<TVec3<T>>& Points, const bool ReturnEmptyOnError) const
+{
+	TArray<TVec3<T>> Normals;
+	GetFaceNormals(Normals, Points, ReturnEmptyOnError);
+	return Normals;
+}
+template CHAOS_API TArray<TVec3<FRealSingle>> FTriangleMesh::GetFaceNormals<FRealSingle>(const TConstArrayView<TVec3<FRealSingle>>&, const bool) const;
+template CHAOS_API TArray<TVec3<FRealDouble>> FTriangleMesh::GetFaceNormals<FRealDouble>(const TConstArrayView<TVec3<FRealDouble>>&, const bool) const;
 
 TArray<FVec3> FTriangleMesh::GetPointNormals(const TConstArrayView<FVec3>& Points, const bool ReturnEmptyOnError)
 {
@@ -298,29 +304,53 @@ void FTriangleMesh::GetPointNormals(TArrayView<FVec3> PointNormals, const TConst
 	ConstThis->GetPointNormals(PointNormals, FaceNormals, bUseGlobalArray);
 }
 
-void FTriangleMesh::GetPointNormals(TArrayView<FVec3> PointNormals, const TConstArrayView<FVec3>& FaceNormals, const bool bUseGlobalArray) const
+template <typename T>
+void FTriangleMesh::GetPointNormals(TArrayView<TVec3<T>> PointNormals, const TConstArrayView<TVec3<T>>& FaceNormals, const bool bUseGlobalArray) const
 {
 	check(MPointToTriangleMap.Num() != 0);
 
-	if (bRealTypeCompatibleWithISPC && bChaos_TriangleMesh_ISPC_Enabled)
+	for (int32 Element = 0; Element < MNumIndices; ++Element)  // Iterate points with local indexes
 	{
+		const int32 NormalIndex = bUseGlobalArray ? LocalToGlobal(Element) : Element;  // Select whether the points normal indices match the points indices or start at 0
+		TVec3<T>& Normal = PointNormals[NormalIndex];
+		Normal = TVec3<T>(0.);
+		const TArray<int32>& TriangleMap = MPointToTriangleMap[Element];  // Access MPointToTriangleMap with local index
+		for (int32 k = 0; k < TriangleMap.Num(); ++k)
+		{
+			if (FaceNormals.IsValidIndex(TriangleMap[k]))
+			{
+				Normal += FaceNormals[TriangleMap[k]];
+			}
+		}
+		Normal = Normal.GetSafeNormal();
+	}
+}
+template CHAOS_API void FTriangleMesh::GetPointNormals<FRealDouble>(TArrayView<TVec3<FRealDouble>>, const TConstArrayView<TVec3<FRealDouble>>&, const bool) const;
+
+template <>
+CHAOS_API void FTriangleMesh::GetPointNormals<FRealSingle>(TArrayView<TVec3<FRealSingle>> PointNormals, const TConstArrayView<TVec3<FRealSingle>>& FaceNormals, const bool bUseGlobalArray) const
+{
+	check(MPointToTriangleMap.Num() != 0);
+
 #if INTEL_ISPC
+	if (bChaos_TriangleMesh_ISPC_Enabled)
+	{
 		ispc::GetPointNormals(
-			(ispc::FVector*)PointNormals.GetData(),
-			(const ispc::FVector*)FaceNormals.GetData(),
+			(ispc::FVector3f*)PointNormals.GetData(),
+			(const ispc::FVector3f*)FaceNormals.GetData(),
 			(const ispc::TArrayInt*)MPointToTriangleMap.GetData(),
 			bUseGlobalArray ? LocalToGlobal(0) : 0,
 			FaceNormals.Num(),
 			MNumIndices);
-#endif
 	}
 	else
+#endif
 	{
 		for (int32 Element = 0; Element < MNumIndices; ++Element)  // Iterate points with local indexes
 		{
 			const int32 NormalIndex = bUseGlobalArray ? LocalToGlobal(Element) : Element;  // Select whether the points normal indices match the points indices or start at 0
-			FVec3& Normal = PointNormals[NormalIndex];
-			Normal = FVec3(0);
+			TVec3<FRealSingle>& Normal = PointNormals[NormalIndex];
+			Normal = TVec3<FRealSingle>(0.);
 			const TArray<int32>& TriangleMap = MPointToTriangleMap[Element];  // Access MPointToTriangleMap with local index
 			for (int32 k = 0; k < TriangleMap.Num(); ++k)
 			{
@@ -619,7 +649,7 @@ struct OrderedEdgeKeyFuncs : BaseKeyFuncs<TVec2<int32>, TVec2<int32>, false>
 
 	static FORCEINLINE uint32 GetKeyHash(const TVec2<int32>& elem)
 	{
-		const uint32 v = HashCombine(GetTypeHash(elem[0]), GetTypeHash(elem[1]));
+		const uint32 v = HashCombine(::GetTypeHash(elem[0]), ::GetTypeHash(elem[1]));
 		return v;
 	}
 };
@@ -1395,3 +1425,5 @@ void FTriangleMesh::RemoveDegenerateElements()
 		}
 	}
 }
+
+}  // End namespace Chaos
diff --git a/Engine/Source/Runtime/Experimental/Chaos/Private/Chaos/TriangleMesh.ispc b/Engine/Source/Runtime/Experimental/Chaos/Private/Chaos/TriangleMesh.ispc
--- a/Engine/Source/Runtime/Experimental/Chaos/Private/Chaos/TriangleMesh.ispc
+++ b/Engine/Source/Runtime/Experimental/Chaos/Private/Chaos/TriangleMesh.ispc
@@ -4,8 +4,8 @@
 
 #include "Math/Vector.isph"
 
-export void GetFaceNormals(uniform FVector Normals[],
-							const uniform FVector Points[],
+export void GetFaceNormals(uniform FVector3f Normals[],
+							const uniform FVector3f Points[],
 							const uniform FIntVector Elements[],
 							const uniform int32 NumElements)
 {
@@ -13,14 +13,14 @@ export void GetFaceNormals(uniform FVector Normals[],
 	{
 		const varying FIntVector Tri = VectorLoad(&Elements[extract(i, 0)]);
 
-		const varying FVector p0 = VectorGather(&Points[Tri.V[0]]);
-		const varying FVector p1 = VectorGather(&Points[Tri.V[1]]);
-		const varying FVector p2 = VectorGather(&Points[Tri.V[2]]);
+		const varying FVector3f p0 = VectorGather(&Points[Tri.V[0]]);
+		const varying FVector3f p1 = VectorGather(&Points[Tri.V[1]]);
+		const varying FVector3f p2 = VectorGather(&Points[Tri.V[2]]);
 
-		const varying FVector p10 = p1 - p0;
-		const varying FVector p20 = p2 - p0;
-		const varying FVector Cross = VectorCross(p20, p10);
-		const varying FVector CrossNormalized = VectorGetSafeNormal(Cross);
+		const varying FVector3f p10 = p1 - p0;
+		const varying FVector3f p20 = p2 - p0;
+		const varying FVector3f Cross = VectorCross(p20, p10);
+		const varying FVector3f CrossNormalized = VectorGetSafeNormal(Cross);
 
 		VectorStore(&Normals[extract(i, 0)], CrossNormalized);
 	}
@@ -32,8 +32,8 @@ struct TArrayInt
 	int32 ArrayNum;
 };
 
-export void GetPointNormals(uniform FVector PointNormals[],
-							const uniform FVector FaceNormals[],
+export void GetPointNormals(uniform FVector3f PointNormals[],
+							const uniform FVector3f FaceNormals[],
 							const uniform TArrayInt PointToTriangleMap[],
 							const uniform int32 StartIndex,
 							const uniform int32 NumFaceNormals,
@@ -41,7 +41,7 @@ export void GetPointNormals(uniform FVector PointNormals[],
 {
 	foreach(Element = 0 ... NumIndices)
 	{
-		varying FVector Normal = ZeroVector;
+		varying FVector3f Normal = FloatZeroVector;
 
 		#pragma ignore warning(perf)
 		const varying TArrayInt TriangleMap = PointToTriangleMap[Element];  // Access MPointToTriangleMap with local index
diff --git a/Engine/Source/Runtime/Experimental/Chaos/Private/Chaos/VelocityField.cpp b/Engine/Source/Runtime/Experimental/Chaos/Private/Chaos/VelocityField.cpp
--- a/Engine/Source/Runtime/Experimental/Chaos/Private/Chaos/VelocityField.cpp
+++ b/Engine/Source/Runtime/Experimental/Chaos/Private/Chaos/VelocityField.cpp
@@ -5,9 +5,9 @@
 #if INTEL_ISPC
 #include "VelocityField.ispc.generated.h"
 
-static_assert(sizeof(ispc::FVector) == sizeof(Chaos::FVec3), "sizeof(ispc::FVector) != sizeof(Chaos::FVec3)");
+static_assert(sizeof(ispc::FVector3f) == sizeof(Chaos::Softs::FSolverVec3), "sizeof(ispc::FVector3f) != sizeof(Chaos::Softs::FSolverVec3)");
 static_assert(sizeof(ispc::FIntVector) == sizeof(Chaos::TVec3<int32>), "sizeof(ispc::FIntVector) != sizeof(Chaos::TVec3<int32>)");
-static_assert(sizeof(ispc::FVector2) == sizeof(Chaos::FVec2), "sizeof(ispc::FVector2) != sizeof(Chaos::FVec2)");
+static_assert(sizeof(ispc::FVector2f) == sizeof(Chaos::Softs::FSolverVec2), "sizeof(ispc::FVector2f) != sizeof(Chaos::Softs::FSolverVec2)");
 #endif
 
 #if INTEL_ISPC && !UE_BUILD_SHIPPING
@@ -15,7 +15,7 @@ bool bChaos_VelocityField_ISPC_Enabled = true;
 FAutoConsoleVariableRef CVarChaosVelocityFieldISPCEnabled(TEXT("p.Chaos.VelocityField.ISPC"), bChaos_VelocityField_ISPC_Enabled, TEXT("Whether to use ISPC optimizations in velocity field calculations"));
 #endif
 
-using namespace Chaos;
+namespace Chaos::Softs {
 
 void FVelocityField::SetGeometry(const FTriangleMesh* TriangleMesh, const TConstArrayView<FRealSingle>& DragMultipliers, const TConstArrayView<FRealSingle>& LiftMultipliers)
 {
@@ -47,7 +47,7 @@ void FVelocityField::SetGeometry(const FTriangleMesh* TriangleMesh, const TConst
 		}
 		else
 		{
-			static const FReal OneThird = (FReal)1. / (FReal)3.;
+			constexpr FSolverReal OneThird = (FSolverReal)1. / (FSolverReal)3.;
 
 			Multipliers.SetNumUninitialized(Elements.Num());
 
@@ -61,15 +61,15 @@ void FVelocityField::SetGeometry(const FTriangleMesh* TriangleMesh, const TConst
 				switch (WeighMaps)
 				{
 				case EWeighMaps::DragMultipliers:
-					Multipliers[ElementIndex] = FVec2((DragMultipliers[I0] + DragMultipliers[I1] + DragMultipliers[I2]) * OneThird, (FReal)0.); break;
+					Multipliers[ElementIndex] = FSolverVec2((FSolverReal)(DragMultipliers[I0] + DragMultipliers[I1] + DragMultipliers[I2]) * OneThird, (FSolverReal)0.); break;
 					break;
 				case EWeighMaps::LiftMultipliers:
-					Multipliers[ElementIndex] = FVec2((FReal)0., (LiftMultipliers[I0] + LiftMultipliers[I1] + LiftMultipliers[I2]) * OneThird); break;
+					Multipliers[ElementIndex] = FSolverVec2((FSolverReal)0., (FSolverReal)(LiftMultipliers[I0] + LiftMultipliers[I1] + LiftMultipliers[I2]) * OneThird); break;
 					break;
 				default:
-					Multipliers[ElementIndex] = FVec2(
-						DragMultipliers[I0] + DragMultipliers[I1] + DragMultipliers[I2],
-						LiftMultipliers[I0] + LiftMultipliers[I1] + LiftMultipliers[I2]) * OneThird;
+					Multipliers[ElementIndex] = FSolverVec2(
+						(FSolverReal)(DragMultipliers[I0] + DragMultipliers[I1] + DragMultipliers[I2]) * OneThird,
+						(FSolverReal)(LiftMultipliers[I0] + LiftMultipliers[I1] + LiftMultipliers[I2]) * OneThird);
 					break;
 				}
 			}
@@ -83,13 +83,13 @@ void FVelocityField::SetGeometry(const FTriangleMesh* TriangleMesh, const TConst
 		NumParticles = 0;
 		Forces.Reset();
 		Multipliers.Reset();
-		SetProperties(FVec2::ZeroVector, FVec2::ZeroVector, (FReal)0.);
+		SetProperties(FSolverVec2(0.), FSolverVec2(0.), (FSolverReal)0.);
 	}
 }
 
 
 
-void FVelocityField::UpdateForces(const FPBDParticles& InParticles, const FReal /*Dt*/)
+void FVelocityField::UpdateForces(const FSolverParticles& InParticles, const FSolverReal /*Dt*/)
 {
 	if (!Multipliers.Num())
 	{
@@ -97,11 +97,11 @@ void FVelocityField::UpdateForces(const FPBDParticles& InParticles, const FReal
 		if (bRealTypeCompatibleWithISPC && bChaos_VelocityField_ISPC_Enabled)
 		{
 			ispc::UpdateField(
-				(ispc::FVector*)Forces.GetData(),
+				(ispc::FVector3f*)Forces.GetData(),
 				(const ispc::FIntVector*)Elements.GetData(),
-				(const ispc::FVector*)InParticles.GetV().GetData(),
-				(const ispc::FVector*)InParticles.XArray().GetData(),
-				(const ispc::FVector&)Velocity,
+				(const ispc::FVector3f*)InParticles.GetV().GetData(),
+				(const ispc::FVector3f*)InParticles.XArray().GetData(),
+				(const ispc::FVector3f&)Velocity,
 				QuarterRho,
 				DragBase,
 				LiftBase,
@@ -122,12 +122,12 @@ void FVelocityField::UpdateForces(const FPBDParticles& InParticles, const FReal
 		if (bRealTypeCompatibleWithISPC && bChaos_VelocityField_ISPC_Enabled)
 		{
 			ispc::UpdateFieldWithWeightMaps(
-				(ispc::FVector*)Forces.GetData(),
+				(ispc::FVector3f*)Forces.GetData(),
 				(const ispc::FIntVector*)Elements.GetData(),
-				(const ispc::FVector*)InParticles.GetV().GetData(),
-				(const ispc::FVector*)InParticles.XArray().GetData(),
-				(const ispc::FVector2*)Multipliers.GetData(),
-				(const ispc::FVector&)Velocity,
+				(const ispc::FVector3f*)InParticles.GetV().GetData(),
+				(const ispc::FVector3f*)InParticles.XArray().GetData(),
+				(const ispc::FVector2f*)Multipliers.GetData(),
+				(const ispc::FVector3f&)Velocity,
 				QuarterRho,
 				DragBase,
 				DragRange,
@@ -140,12 +140,14 @@ void FVelocityField::UpdateForces(const FPBDParticles& InParticles, const FReal
 		{
 			for (int32 ElementIndex = 0; ElementIndex < Elements.Num(); ++ElementIndex)
 			{
-				const FVec2& Multiplier = Multipliers[ElementIndex];
-				const FReal Cd = DragBase + DragRange * Multiplier[0];
-				const FReal Cl = LiftBase + LiftRange * Multiplier[1];
+				const FSolverVec2& Multiplier = Multipliers[ElementIndex];
+				const FSolverReal Cd = DragBase + DragRange * Multiplier[0];
+				const FSolverReal Cl = LiftBase + LiftRange * Multiplier[1];
 
 				UpdateField(InParticles, ElementIndex, Velocity, Cd, Cl);
 			}
 		}
 	}
 }
+
+}  // End namespace Chaos::Softs
diff --git a/Engine/Source/Runtime/Experimental/Chaos/Private/Chaos/VelocityField.ispc b/Engine/Source/Runtime/Experimental/Chaos/Private/Chaos/VelocityField.ispc
--- a/Engine/Source/Runtime/Experimental/Chaos/Private/Chaos/VelocityField.ispc
+++ b/Engine/Source/Runtime/Experimental/Chaos/Private/Chaos/VelocityField.ispc
@@ -4,23 +4,23 @@
 
 #include "Math/Vector.isph"
 
-static const FReal OneThird = 1.0/3.0;
+static const float OneThird = 1.0f/3.0f;
 
-static inline FReal SafeNormalize(FVector &Direction)
+static inline float SafeNormalize(FVector3f &Direction)
 {
-	const FReal Size = sqrt(VectorSizeSquared(Direction));
-	Direction = VectorSelect((Size < KINDA_SMALL_NUMBER), ForwardVector, Direction / Size);
-	return (Size < KINDA_SMALL_NUMBER) ? 0 : Size;
+	const float Size = sqrt(VectorSizeSquared(Direction));
+	Direction = VectorSelect((Size < FLOAT_KINDA_SMALL_NUMBER), FloatForwardVector, Direction / Size);
+	return (Size < FLOAT_KINDA_SMALL_NUMBER) ? 0 : Size;
 }
 
-export void UpdateField(uniform FVector Forces[],
+export void UpdateField(uniform FVector3f Forces[],
 								const uniform FIntVector Elements[],
-								const uniform FVector V[],
-								const uniform FVector X[],
-								const uniform FVector& Velocity,
-								const uniform FReal QuarterRho,
-								const uniform FReal Cd,
-								const uniform FReal Cl,
+								const uniform FVector3f V[],
+								const uniform FVector3f X[],
+								const uniform FVector3f& Velocity,
+								const uniform float QuarterRho,
+								const uniform float Cd,
+								const uniform float Cl,
 								const uniform int32 NumElements)
 {
 	varying FIntVector Element;
@@ -39,25 +39,25 @@ export void UpdateField(uniform FVector Forces[],
 		Element = VectorLoad(&Elements[Offset]);
 #endif
 
-		const FVector X0 = VectorGather(&X[Element.V[0]]);
-		const FVector X1 = VectorGather(&X[Element.V[1]]);
-		const FVector X2 = VectorGather(&X[Element.V[2]]);
+		const FVector3f X0 = VectorGather(&X[Element.V[0]]);
+		const FVector3f X1 = VectorGather(&X[Element.V[1]]);
+		const FVector3f X2 = VectorGather(&X[Element.V[2]]);
 
 		// Calculate the normal and the area of the surface exposed to the flow
-		FVector N = VectorCross(X1 - X0, X2 - X0);
-		const FReal DoubleArea = SafeNormalize(N);
+		FVector3f N = VectorCross(X1 - X0, X2 - X0);
+		const float DoubleArea = SafeNormalize(N);
 
-		const FVector V0 = VectorGather(&V[Element.V[0]]);
-		const FVector V1 = VectorGather(&V[Element.V[1]]);
-		const FVector V2 = VectorGather(&V[Element.V[2]]);
+		const FVector3f V0 = VectorGather(&V[Element.V[0]]);
+		const FVector3f V1 = VectorGather(&V[Element.V[1]]);
+		const FVector3f V2 = VectorGather(&V[Element.V[2]]);
 
 		// Calculate the direction and the relative velocity of the triangle to the flow
-		const FVector SurfaceVelocity = OneThird * (V0 + V1 + V2);
-		const FVector RelVelocity = Velocity - SurfaceVelocity;
+		const FVector3f SurfaceVelocity = OneThird * (V0 + V1 + V2);
+		const FVector3f RelVelocity = Velocity - SurfaceVelocity;
 
 		// Set the aerodynamic forces
-		const FReal VDotN = VectorDot(RelVelocity, N);
-		const FReal VSquare = VectorDot(RelVelocity, RelVelocity);
+		const float VDotN = VectorDot(RelVelocity, N);
+		const float VSquare = VectorDot(RelVelocity, RelVelocity);
 
 #if HW_GATHER_SUPPORTED == 1
 		if(ElementIndex + programCount < NumElements)
@@ -66,10 +66,10 @@ export void UpdateField(uniform FVector Forces[],
 		}
 #endif
 
-		const FVector EvenForce = (Cd - Cl) * VDotN * RelVelocity + Cl * VSquare * N;
-		const FVector OddForce = (Cl - Cd) * VDotN * RelVelocity - Cl * VSquare * N;
+		const FVector3f EvenForce = (Cd - Cl) * VDotN * RelVelocity + Cl * VSquare * N;
+		const FVector3f OddForce = (Cl - Cd) * VDotN * RelVelocity - Cl * VSquare * N;
 
-		const FVector Force = QuarterRho * DoubleArea * VectorSelect(VDotN >= 0, EvenForce, OddForce);
+		const FVector3f Force = QuarterRho * DoubleArea * VectorSelect(VDotN >= 0, EvenForce, OddForce);
 
 		VectorStore(&Forces[Offset], Force);
 
@@ -77,17 +77,17 @@ export void UpdateField(uniform FVector Forces[],
 	}
 }
 
-export void UpdateFieldWithWeightMaps(uniform FVector Forces[],
+export void UpdateFieldWithWeightMaps(uniform FVector3f Forces[],
 								const uniform FIntVector Elements[],
-								const uniform FVector V[],
-								const uniform FVector X[],
-								const uniform FVector2 Multipliers[],
-								const uniform FVector& Velocity,
-								const uniform FReal QuarterRho,
-								const uniform FReal DragBase,
-								const uniform FReal DragRange,
-								const uniform FReal LiftBase,
-								const uniform FReal LiftRange,
+								const uniform FVector3f V[],
+								const uniform FVector3f X[],
+								const uniform FVector2f Multipliers[],
+								const uniform FVector3f& Velocity,
+								const uniform float QuarterRho,
+								const uniform float DragBase,
+								const uniform float DragRange,
+								const uniform float LiftBase,
+								const uniform float LiftRange,
 								const uniform int32 NumElements)
 {
 	varying FIntVector Element;
@@ -105,30 +105,30 @@ export void UpdateFieldWithWeightMaps(uniform FVector Forces[],
 #if HW_GATHER_SUPPORTED == 0
 		Element = VectorLoad(&Elements[Offset]);
 #endif
-		const FVector X0 = VectorGather(&X[Element.V[0]]);
-		const FVector X1 = VectorGather(&X[Element.V[1]]);
-		const FVector X2 = VectorGather(&X[Element.V[2]]);
+		const FVector3f X0 = VectorGather(&X[Element.V[0]]);
+		const FVector3f X1 = VectorGather(&X[Element.V[1]]);
+		const FVector3f X2 = VectorGather(&X[Element.V[2]]);
 
-		const FVector2 Multiplier = VectorGather(&Multipliers[ElementIndex]);
+		const FVector2f Multiplier = VectorGather(&Multipliers[ElementIndex]);
 
 		// Calculate the normal and the area of the surface exposed to the flow
-		FVector N = VectorCross(X1 - X0, X2 - X0);
-		const FReal DoubleArea = SafeNormalize(N);
+		FVector3f N = VectorCross(X1 - X0, X2 - X0);
+		const float DoubleArea = SafeNormalize(N);
 
-		const FVector V0 = VectorGather(&V[Element.V[0]]);
-		const FVector V1 = VectorGather(&V[Element.V[1]]);
-		const FVector V2 = VectorGather(&V[Element.V[2]]);
+		const FVector3f V0 = VectorGather(&V[Element.V[0]]);
+		const FVector3f V1 = VectorGather(&V[Element.V[1]]);
+		const FVector3f V2 = VectorGather(&V[Element.V[2]]);
 
 		// Calculate the direction and the relative velocity of the triangle to the flow
-		const FVector SurfaceVelocity = OneThird * (V0 + V1 + V2);
-		const FVector RelVelocity = Velocity - SurfaceVelocity;
+		const FVector3f SurfaceVelocity = OneThird * (V0 + V1 + V2);
+		const FVector3f RelVelocity = Velocity - SurfaceVelocity;
 
 		// Set the aerodynamic forces
-		const FReal VDotN = VectorDot(RelVelocity, N);
-		const FReal VSquare = VectorDot(RelVelocity, RelVelocity);
+		const float VDotN = VectorDot(RelVelocity, N);
+		const float VSquare = VectorDot(RelVelocity, RelVelocity);
 
-		const FReal Cd = DragBase + DragRange * Multiplier.V[0];
-		const FReal Cl = LiftBase + LiftRange * Multiplier.V[1];
+		const float Cd = DragBase + DragRange * Multiplier.V[0];
+		const float Cl = LiftBase + LiftRange * Multiplier.V[1];
 
 #if HW_GATHER_SUPPORTED == 1
 		if(ElementIndex + programCount < NumElements)
@@ -137,10 +137,10 @@ export void UpdateFieldWithWeightMaps(uniform FVector Forces[],
 		}
 #endif
 
-		const FVector EvenForce = (Cd - Cl) * VDotN * RelVelocity + Cl * VSquare * N;
-		const FVector OddForce = (Cl - Cd) * VDotN * RelVelocity - Cl * VSquare * N;
+		const FVector3f EvenForce = (Cd - Cl) * VDotN * RelVelocity + Cl * VSquare * N;
+		const FVector3f OddForce = (Cl - Cd) * VDotN * RelVelocity - Cl * VSquare * N;
 
-		const FVector Force = QuarterRho * DoubleArea * VectorSelect(VDotN >= 0, EvenForce, OddForce);
+		const FVector3f Force = QuarterRho * DoubleArea * VectorSelect(VDotN >= 0, EvenForce, OddForce);
 
 		VectorStore(&Forces[Offset], Force);
 
diff --git a/Engine/Source/Runtime/Experimental/Chaos/Public/Chaos/AABB.h b/Engine/Source/Runtime/Experimental/Chaos/Public/Chaos/AABB.h
--- a/Engine/Source/Runtime/Experimental/Chaos/Public/Chaos/AABB.h
+++ b/Engine/Source/Runtime/Experimental/Chaos/Public/Chaos/AABB.h
@@ -44,6 +44,13 @@ namespace Chaos
 		{
 		}
 
+		template<typename OtherType> 
+		explicit TAABB(const TAABB<OtherType, d>& Other)
+			: MMin(TVector<T, d>(Other.Min()))
+			, MMax(TVector<T, d>(Other.Max()))
+		{
+		}
+
 		FORCEINLINE TAABB(TAABB<T, d>&& Other)
 			: MMin(MoveTemp(Other.MMin))
 			, MMax(MoveTemp(Other.MMax))
@@ -565,14 +572,15 @@ namespace Chaos
 		TVector<T, d> MMin, MMax;
 	};
 
-
-	FORCEINLINE FChaosArchive& operator<<(FChaosArchive& Ar, TAABB<FReal, 3>& AABB)
+	template<class T, int d>
+	FORCEINLINE FChaosArchive& operator<<(FChaosArchive& Ar, TAABB<T, d>& AABB)
 	{
 		AABB.Serialize(Ar);
 		return Ar;
 	}
 
-	FORCEINLINE FArchive& operator<<(FArchive& Ar, TAABB<FReal, 3>& AABB)
+	template<class T, int d>
+	FORCEINLINE FArchive& operator<<(FArchive& Ar, TAABB<T, d>& AABB)
 	{
 		AABB.Serialize(Ar);
 		return Ar;
diff --git a/Engine/Source/Runtime/Experimental/Chaos/Public/Chaos/Box.h b/Engine/Source/Runtime/Experimental/Chaos/Public/Chaos/Box.h
--- a/Engine/Source/Runtime/Experimental/Chaos/Public/Chaos/Box.h
+++ b/Engine/Source/Runtime/Experimental/Chaos/Public/Chaos/Box.h
@@ -153,7 +153,7 @@ namespace Chaos
 		FORCEINLINE const TAABB<FReal, 3> BoundingBox() const
 		{
 			// LWC - necessary if T is different from FReal
-			return TAABB<FReal, 3>(AABB.Min(), AABB.Max());
+			return TAABB<FReal, 3>(AABB);
 		}
 
 		// Apply a limit to the specified margin that prevents the box inverting
@@ -462,23 +462,46 @@ namespace Chaos
 		}
 
 		// See comments on SerializeAsAABB
-		static void SerializeAsAABBs(FArchive& Ar, TArray<TAABB<T, d>>& AABBs)
+		template<typename OtherType>
+		static void SerializeAsAABBs(FArchive& Ar, TArray<TAABB<OtherType, d>>& AABBs)
 		{
 			Ar.UsingCustomVersion(FExternalPhysicsCustomObjectVersion::GUID);
 			if (Ar.CustomVer(FExternalPhysicsCustomObjectVersion::GUID) < FExternalPhysicsCustomObjectVersion::TBoxReplacedWithTAABB)
 			{
 				TArray<TBox<T, d>> Tmp;
 				Ar << Tmp;
-				AABBs.Reserve(Tmp.Num());
+				AABBs.Reset(Tmp.Num());
 				for (const TBox<T, d>& Box : Tmp)
 				{
-					AABBs.Add(Box.AABB);
+					AABBs.Emplace(TAABB<OtherType, d>(Box.AABB));
 				}
 			}
-			else
+			else if (TAreTypesEqual<T, OtherType>::Value)
 			{
 				Ar << AABBs;
 			}
+			else if (Ar.IsLoading())
+			{
+				// Reload to a different floating point type
+				TArray<TAABB<T, d>> TmpAABBs;
+				Ar << TmpAABBs;
+				AABBs.Reset(TmpAABBs.Num());
+				for (const TAABB<T, d>& TmpAABB : TmpAABBs)
+				{
+					AABBs.Emplace(TAABB<OtherType, d>(TmpAABB));
+				}
+			}
+			else
+			{
+				// Serialize from a different floating point type
+				TArray<TAABB<T, d>> TmpAABBs;
+				TmpAABBs.Reserve(AABBs.Num());
+				for (const TAABB<OtherType, d>& AABB : AABBs)
+				{
+					TmpAABBs.Emplace(TAABB<T, d>(AABB));
+				}
+				Ar << TmpAABBs;
+			}
 		}
 
 		// See comments on SerializeAsAABB
diff --git a/Engine/Source/Runtime/Experimental/Chaos/Public/Chaos/GeometryParticles.h b/Engine/Source/Runtime/Experimental/Chaos/Public/Chaos/GeometryParticles.h
--- a/Engine/Source/Runtime/Experimental/Chaos/Public/Chaos/GeometryParticles.h
+++ b/Engine/Source/Runtime/Experimental/Chaos/Public/Chaos/GeometryParticles.h
@@ -556,7 +556,7 @@ namespace Chaos
 			MHasBounds[Index] = (InGeometry && InGeometry->HasBoundingBox());
 			if (MHasBounds[Index])
 			{
-				MLocalBounds[Index] = InGeometry->BoundingBox();
+				MLocalBounds[Index] = TAABB<T, d>(InGeometry->BoundingBox());
 
 				// Update the world-space stat of all the shapes - must be called after UpdateShapesArray
 				// world space inflated bounds needs to take expansion into account - this is done in integrate for dynamics anyway, so
@@ -626,7 +626,7 @@ namespace Chaos
 				WorldBounds.GrowToInclude(Shape->GetWorldSpaceInflatedShapeBounds());
 			}
 
-			MWorldSpaceInflatedBounds[Index] = WorldBounds;
+			MWorldSpaceInflatedBounds[Index] = TAABB<T, d>(WorldBounds);
 		}
 
 		void UpdateWorldSpaceStateSwept(const int32 Index, const FRigidTransform3& EndWorldTransform, const FVec3& BoundsExpansion, const FVec3& DeltaX)
@@ -637,13 +637,13 @@ namespace Chaos
 
 		const TArray<TSerializablePtr<FImplicitObject>>& GetAllGeometry() const { return MGeometry; }
 
-		typedef TGeometryParticleHandle<T, d> THandleType;
+		typedef FGeometryParticleHandle THandleType;
 		FORCEINLINE THandleType* Handle(int32 Index) const { return const_cast<THandleType*>(MGeometryParticleHandle[Index].Get()); }
 
-		CHAOS_API void SetHandle(int32 Index, TGeometryParticleHandle<T, d>* Handle);
+		CHAOS_API void SetHandle(int32 Index, FGeometryParticleHandle* Handle);
 		
-		CHAOS_API TGeometryParticle<T, d>* GTGeometryParticle(const int32 Index) const { return MGeometryParticle[Index]; }
-		CHAOS_API TGeometryParticle<T, d>*& GTGeometryParticle(const int32 Index) { return MGeometryParticle[Index]; }
+		CHAOS_API FGeometryParticle* GTGeometryParticle(const int32 Index) const { return MGeometryParticle[Index]; }
+		CHAOS_API FGeometryParticle*& GTGeometryParticle(const int32 Index) { return MGeometryParticle[Index]; }
 
 		CHAOS_API const IPhysicsProxyBase* PhysicsProxy(const int32 Index) const { return MPhysicsProxy[Index];  }
 		CHAOS_API IPhysicsProxyBase* PhysicsProxy(const int32 Index) { return MPhysicsProxy[Index]; }
@@ -764,9 +764,9 @@ public:
 					MHasBounds[Idx] = MGeometry[Idx] && MGeometry[Idx]->HasBoundingBox();
 					if (MHasBounds[Idx])
 					{
-						MLocalBounds[Idx] = MGeometry[Idx]->BoundingBox();
+						MLocalBounds[Idx] = TAABB<T, d>(MGeometry[Idx]->BoundingBox());
 						//ignore velocity too, really just trying to get something reasonable)
-						UpdateWorldSpaceState(Idx, TRigidTransform<T,d>(X(Idx), R(Idx)), FVec3(0));
+						UpdateWorldSpaceState(Idx, FRigidTransform3(X(Idx), R(Idx)), FVec3(0));
 					}
 				}
 			}
@@ -824,8 +824,8 @@ public:
 		TArrayCollectionArray<TSharedPtr<const FImplicitObject, ESPMode::ThreadSafe>> MSharedGeometry;
 		// MDynamicGeometry entries are used for geo which is by the evolution. It is not set from the game side.
 		TArrayCollectionArray<TUniquePtr<FImplicitObject>> MDynamicGeometry;
-		TArrayCollectionArray<TSerializablePtr<TGeometryParticleHandle<T, d>>> MGeometryParticleHandle;
-		TArrayCollectionArray<TGeometryParticle<T, d>*> MGeometryParticle;
+		TArrayCollectionArray<TSerializablePtr<FGeometryParticleHandle>> MGeometryParticleHandle;
+		TArrayCollectionArray<FGeometryParticle*> MGeometryParticle;
 		TArrayCollectionArray<IPhysicsProxyBase*> MPhysicsProxy;
 		TArrayCollectionArray<bool> MHasCollision;
 		TArrayCollectionArray<FShapesArray> MShapesArray;
diff --git a/Engine/Source/Runtime/Experimental/Chaos/Public/Chaos/GraphColoring.h b/Engine/Source/Runtime/Experimental/Chaos/Public/Chaos/GraphColoring.h
--- a/Engine/Source/Runtime/Experimental/Chaos/Public/Chaos/GraphColoring.h
+++ b/Engine/Source/Runtime/Experimental/Chaos/Public/Chaos/GraphColoring.h
@@ -50,8 +50,10 @@ class FGraphColoring
 	};
 
   public:
-	CHAOS_API static TArray<TArray<int32>> ComputeGraphColoring(const TArray<TVector<int32, 2>>& Graph, const FDynamicParticles& InParticles);
-	CHAOS_API static TArray<TArray<int32>> ComputeGraphColoring(const TArray<TVector<int32, 3>>& Graph, const FDynamicParticles& InParticles);
+	template<typename T>
+	CHAOS_API static TArray<TArray<int32>> ComputeGraphColoring(const TArray<TVector<int32, 2>>& Graph, const TDynamicParticles<T, 3>& InParticles);
+	template<typename T>
+	CHAOS_API static TArray<TArray<int32>> ComputeGraphColoring(const TArray<TVector<int32, 3>>& Graph, const TDynamicParticles<T, 3>& InParticles);
 };
 
 }
diff --git a/Engine/Source/Runtime/Experimental/Chaos/Public/Chaos/KinematicGeometryParticles.h b/Engine/Source/Runtime/Experimental/Chaos/Public/Chaos/KinematicGeometryParticles.h
--- a/Engine/Source/Runtime/Experimental/Chaos/Public/Chaos/KinematicGeometryParticles.h
+++ b/Engine/Source/Runtime/Experimental/Chaos/Public/Chaos/KinematicGeometryParticles.h
@@ -91,6 +91,6 @@ using FKinematicGeometryParticles = TKinematicGeometryParticles<FReal, 3>;
 template <typename T, int d>
 using TKinematicGeometryClothParticles = TKinematicGeometryParticlesImp<T, d, EGeometryParticlesSimType::Other>;
 
-using FKinematicGeometryClothParticles = TKinematicGeometryClothParticles<FReal, 3>;
+using FKinematicGeometryClothParticles UE_DEPRECATED(5.0, "FKinematicGeometryClothParticles is deprecated, use Softs::FSolverRigidParticles instead.") = TKinematicGeometryClothParticles<FReal, 3>;
 }
 
diff --git a/Engine/Source/Runtime/Experimental/Chaos/Public/Chaos/KinematicTargets.h b/Engine/Source/Runtime/Experimental/Chaos/Public/Chaos/KinematicTargets.h
--- a/Engine/Source/Runtime/Experimental/Chaos/Public/Chaos/KinematicTargets.h
+++ b/Engine/Source/Runtime/Experimental/Chaos/Public/Chaos/KinematicTargets.h
@@ -94,8 +94,8 @@ namespace Chaos
 				FRigidTransform3 Transform;
 				Ar << Transform << KinematicTarget.Mode;
 
-				KinematicTarget.Position = Transform.GetLocation();
-				KinematicTarget.Rotation = Transform.GetRotation();			
+				KinematicTarget.Position = TVec3<T>(Transform.GetLocation());
+				KinematicTarget.Rotation = TRotation3<T>(Transform.GetRotation());
 			}
 
 			return Ar;
diff --git a/Engine/Source/Runtime/Experimental/Chaos/Public/Chaos/PBDActiveView.h b/Engine/Source/Runtime/Experimental/Chaos/Public/Chaos/PBDActiveView.h
--- a/Engine/Source/Runtime/Experimental/Chaos/Public/Chaos/PBDActiveView.h
+++ b/Engine/Source/Runtime/Experimental/Chaos/Public/Chaos/PBDActiveView.h
@@ -4,6 +4,7 @@
 #include "Containers/Array.h"
 #include "Algo/BinarySearch.h"
 #include "Chaos/Framework/Parallel.h"
+#include "Chaos/Vector.h"
 
 namespace Chaos
 {
diff --git a/Engine/Source/Runtime/Experimental/Chaos/Public/Chaos/PBDAnimDriveConstraint.h b/Engine/Source/Runtime/Experimental/Chaos/Public/Chaos/PBDAnimDriveConstraint.h
--- a/Engine/Source/Runtime/Experimental/Chaos/Public/Chaos/PBDAnimDriveConstraint.h
+++ b/Engine/Source/Runtime/Experimental/Chaos/Public/Chaos/PBDAnimDriveConstraint.h
@@ -2,21 +2,23 @@
 #pragma once
 
 #include "Chaos/Framework/Parallel.h"
+#include "Chaos/PBDSoftsEvolutionFwd.h"
 #include "Chaos/PBDStiffness.h"
 #include "ChaosStats.h"
 
 DECLARE_CYCLE_STAT(TEXT("Chaos PBD Anim Drive Constraint"), STAT_PBD_AnimDriveConstraint, STATGROUP_Chaos);
 
-namespace Chaos
+namespace Chaos::Softs
 {
+
 	class FPBDAnimDriveConstraint final
 	{
 	public:
 		FPBDAnimDriveConstraint(
 			const int32 InParticleOffset,
 			const int32 InParticleCount,
-			const TArray<FVec3>& InAnimationPositions,  // Use global indexation (will need adding ParticleOffset)
-			const TArray<FVec3>& InOldAnimationPositions,  // Use global indexation (will need adding ParticleOffset)
+			const TArray<FSolverVec3>& InAnimationPositions,  // Use global indexation (will need adding ParticleOffset)
+			const TArray<FSolverVec3>& InOldAnimationPositions,  // Use global indexation (will need adding ParticleOffset)
 			const TConstArrayView<FRealSingle>& StiffnessMultipliers,  // Use local indexation
 			const TConstArrayView<FRealSingle>& DampingMultipliers  // Use local indexation
 		)
@@ -24,33 +26,33 @@ namespace Chaos
 			, OldAnimationPositions(InOldAnimationPositions)
 			, ParticleOffset(InParticleOffset)
 			, ParticleCount(InParticleCount)
-			, Stiffness(FVec2::UnitVector, StiffnessMultipliers, InParticleCount)
-			, Damping(FVec2::UnitVector, DampingMultipliers, InParticleCount)
+			, Stiffness(FSolverVec2::UnitVector, StiffnessMultipliers, InParticleCount)
+			, Damping(FSolverVec2::UnitVector, DampingMultipliers, InParticleCount)
 		{
 		}
 
 		~FPBDAnimDriveConstraint() {}
 
 		// Return the stiffness input values used by the constraint
-		FVec2 GetStiffness() const { return Stiffness.GetWeightedValue(); }
+		FSolverVec2 GetStiffness() const { return Stiffness.GetWeightedValue(); }
 
 		// Return the damping input values used by the constraint
-		FVec2 GetDamping() const { return Damping.GetWeightedValue(); }
+		FSolverVec2 GetDamping() const { return Damping.GetWeightedValue(); }
 
-		inline void SetProperties(const FVec2& InStiffness, const FVec2& InDamping)
+		inline void SetProperties(const FSolverVec2& InStiffness, const FSolverVec2& InDamping)
 		{
 			Stiffness.SetWeightedValue(InStiffness);
 			Damping.SetWeightedValue(InDamping);
 		}
 
 		// Set stiffness offset and range, as well as the simulation stiffness exponent
-		inline void ApplyProperties(const FReal Dt, const int32 NumIterations)
+		inline void ApplyProperties(const FSolverReal Dt, const int32 NumIterations)
 		{
 			Stiffness.ApplyValues(Dt, NumIterations);
 			Damping.ApplyValues(Dt, NumIterations);
 		}
 
-		inline void Apply(FPBDParticles& InParticles, const FReal Dt) const
+		inline void Apply(FSolverParticles& InParticles, const FSolverReal Dt) const
 		{
 			SCOPE_CYCLE_COUNTER(STAT_PBD_AnimDriveConstraint);
 
@@ -58,38 +60,38 @@ namespace Chaos
 			{
 				if (Damping.HasWeightMap())
 				{
-					PhysicsParallelFor(ParticleCount, [&](int32 Index)  // TODO: profile needed for these parallel loop based on particle count
+					PhysicsParallelFor(ParticleCount, [this, &InParticles, &Dt](int32 Index)  // TODO: profile needed for these parallel loop based on particle count
 					{
-						const FReal ParticleStiffness = Stiffness[Index];
-						const FReal ParticleDamping = Damping[Index];
+						const FSolverReal ParticleStiffness = Stiffness[Index];
+						const FSolverReal ParticleDamping = Damping[Index];
 						ApplyHelper(InParticles, ParticleStiffness, ParticleDamping, Dt, Index);
 					});
 				}
 				else
 				{
-					const FReal ParticleDamping = (FReal)Damping;
-					PhysicsParallelFor(ParticleCount, [&](int32 Index)
+					const FSolverReal ParticleDamping = (FSolverReal)Damping;
+					PhysicsParallelFor(ParticleCount, [this, &InParticles, ParticleDamping, &Dt](int32 Index)
 					{
-						const FReal ParticleStiffness = Stiffness[Index];
+						const FSolverReal ParticleStiffness = Stiffness[Index];
 						ApplyHelper(InParticles, ParticleStiffness, ParticleDamping, Dt, Index);
 					});
 				}
 			}
 			else
 			{
-				const FReal ParticleStiffness = (FReal)Stiffness;
+				const FSolverReal ParticleStiffness = (FSolverReal)Stiffness;
 				if (Damping.HasWeightMap())
 				{
-					PhysicsParallelFor(ParticleCount, [&](int32 Index)
+					PhysicsParallelFor(ParticleCount, [this, &InParticles, &ParticleStiffness, &Dt](int32 Index)
 					{
-						const FReal ParticleDamping = Damping[Index];
+						const FSolverReal ParticleDamping = Damping[Index];
 						ApplyHelper(InParticles, ParticleStiffness, ParticleDamping, Dt, Index);
 					});
 				}
 				else
 				{
-					const FReal ParticleDamping = (FReal)Damping;
-					PhysicsParallelFor(ParticleCount, [&](int32 Index)
+					const FSolverReal ParticleDamping = (FSolverReal)Damping;
+					PhysicsParallelFor(ParticleCount, [this, &InParticles, &ParticleStiffness, &ParticleDamping, &Dt](int32 Index)
 					{
 						ApplyHelper(InParticles, ParticleStiffness, ParticleDamping, Dt, Index);
 					});
@@ -98,28 +100,28 @@ namespace Chaos
 		}
 
 	private:
-		inline void ApplyHelper(FPBDParticles& Particles, const FReal InStiffness, const FReal InDamping, const FReal Dt, const int32 Index) const
+		inline void ApplyHelper(FSolverParticles& Particles, const FSolverReal InStiffness, const FSolverReal InDamping, const FSolverReal Dt, const int32 Index) const
 		{
 			const int32 ParticleIndex = ParticleOffset + Index;
-			if (Particles.InvM(ParticleIndex) == (FReal)0.)
+			if (Particles.InvM(ParticleIndex) == (FSolverReal)0.)
 			{
 				return;
 			}
 
-			FVec3& ParticlePosition = Particles.P(ParticleIndex);
-			const FVec3& AnimationPosition = AnimationPositions[ParticleIndex];
-			const FVec3& OldAnimationPosition = OldAnimationPositions[ParticleIndex];
+			FSolverVec3& ParticlePosition = Particles.P(ParticleIndex);
+			const FSolverVec3& AnimationPosition = AnimationPositions[ParticleIndex];
+			const FSolverVec3& OldAnimationPosition = OldAnimationPositions[ParticleIndex];
 
-			const FVec3 ParticleDisplacement = ParticlePosition - Particles.X(ParticleIndex);
-			const FVec3 AnimationDisplacement = OldAnimationPosition - AnimationPosition;
-			const FVec3 RelativeDisplacement = ParticleDisplacement - AnimationDisplacement;
+			const FSolverVec3 ParticleDisplacement = ParticlePosition - Particles.X(ParticleIndex);
+			const FSolverVec3 AnimationDisplacement = OldAnimationPosition - AnimationPosition;
+			const FSolverVec3 RelativeDisplacement = ParticleDisplacement - AnimationDisplacement;
 
 			ParticlePosition -= InStiffness * (ParticlePosition - AnimationPosition) + InDamping * RelativeDisplacement;
 		}
 
 	private:
-		const TArray<FVec3>& AnimationPositions;  // Use global index (needs adding ParticleOffset)
-		const TArray<FVec3>& OldAnimationPositions;  // Use global index (needs adding ParticleOffset)
+		const TArray<FSolverVec3>& AnimationPositions;  // Use global index (needs adding ParticleOffset)
+		const TArray<FSolverVec3>& OldAnimationPositions;  // Use global index (needs adding ParticleOffset)
 		const int32 ParticleOffset;
 		const int32 ParticleCount;
 
@@ -127,6 +129,4 @@ namespace Chaos
 		FPBDStiffness Damping;
 	};
 
-	template<typename T, int d>
-	using TPBDAnimDriveConstraint UE_DEPRECATED(4.27, "Deprecated. this class is to be deleted, use FPBDAnimDriveConstraint instead") = FPBDAnimDriveConstraint;
-}
+}  // End namespace Chaos::Softs
diff --git a/Engine/Source/Runtime/Experimental/Chaos/Public/Chaos/PBDAxialSpringConstraints.h b/Engine/Source/Runtime/Experimental/Chaos/Public/Chaos/PBDAxialSpringConstraints.h
--- a/Engine/Source/Runtime/Experimental/Chaos/Public/Chaos/PBDAxialSpringConstraints.h
+++ b/Engine/Source/Runtime/Experimental/Chaos/Public/Chaos/PBDAxialSpringConstraints.h
@@ -3,7 +3,7 @@
 
 #include "Chaos/PBDAxialSpringConstraintsBase.h"
 
-namespace Chaos
+namespace Chaos::Softs
 {
 
 class CHAOS_API FPBDAxialSpringConstraints : public FPBDAxialSpringConstraintsBase
@@ -14,31 +14,31 @@ class CHAOS_API FPBDAxialSpringConstraints : public FPBDAxialSpringConstraintsBa
 
 public:
 	FPBDAxialSpringConstraints(
-		const FPBDParticles& Particles,
+		const FSolverParticles& Particles,
 		int32 ParticleOffset,
 		int32 ParticleCount,
 		const TArray<TVec3<int32>>& InConstraints,
 		const TConstArrayView<FRealSingle>& StiffnessMultipliers,
-		const FVec2& InStiffness,
+		const FSolverVec2& InStiffness,
 		bool bTrimKinematicConstraints)
 		: Base(Particles, ParticleOffset, ParticleCount, InConstraints, StiffnessMultipliers, InStiffness, bTrimKinematicConstraints)
 	{
 		InitColor(Particles);
 	}
 
-	virtual ~FPBDAxialSpringConstraints() {}
+	virtual ~FPBDAxialSpringConstraints() override {}
 
-	void Apply(FPBDParticles& InParticles, const FReal Dt) const;
+	void Apply(FSolverParticles& InParticles, const FSolverReal Dt) const;
 
 private:
-	void InitColor(const FPBDParticles& InParticles);
-	void ApplyHelper(FPBDParticles& InParticles, const FReal Dt, const int32 ConstraintIndex, const FReal ExpStiffnessValue) const;
+	void InitColor(const FSolverParticles& InParticles);
+	void ApplyHelper(FSolverParticles& InParticles, const FSolverReal Dt, const int32 ConstraintIndex, const FSolverReal ExpStiffnessValue) const;
 
 private:
 	TArray<TArray<int32>> ConstraintsPerColor;
 };
 
-}
+}  // End namespace Chaos::Softs
 
 // Support ISPC enable/disable in non-shipping builds
 #if !INTEL_ISPC
diff --git a/Engine/Source/Runtime/Experimental/Chaos/Public/Chaos/PBDAxialSpringConstraintsBase.h b/Engine/Source/Runtime/Experimental/Chaos/Public/Chaos/PBDAxialSpringConstraintsBase.h
--- a/Engine/Source/Runtime/Experimental/Chaos/Public/Chaos/PBDAxialSpringConstraintsBase.h
+++ b/Engine/Source/Runtime/Experimental/Chaos/Public/Chaos/PBDAxialSpringConstraintsBase.h
@@ -1,28 +1,29 @@
 // Copyright Epic Games, Inc. All Rights Reserved.
 #pragma once
 
+#include "Chaos/PBDSoftsEvolutionFwd.h"
 #include "Chaos/PBDParticles.h"
 #include "Chaos/PBDStiffness.h"
 #include "Containers/Array.h"
 
-namespace Chaos
+namespace Chaos::Softs
 {
 
 class FPBDAxialSpringConstraintsBase
 {
 public:
 	FPBDAxialSpringConstraintsBase(
-		const FPBDParticles& Particles,
+		const FSolverParticles& Particles,
 		int32 ParticleOffset,
 		int32 ParticleCount,
 		const TArray<TVec3<int32>>& InConstraints,
 		const TConstArrayView<FRealSingle>& StiffnessMultipliers,
-		const FVec2& InStiffness,
+		const FSolverVec2& InStiffness,
 		bool bTrimKinematicConstraints)
 		: Constraints(TrimConstraints(InConstraints, 
 			[&Particles, bTrimKinematicConstraints](int32 Index0, int32 Index1, int32 Index2)
 			{
-				return bTrimKinematicConstraints && Particles.InvM(Index0) == (FReal)0. && Particles.InvM(Index1) == (FReal)0. && Particles.InvM(Index2) == (FReal)0.;
+				return bTrimKinematicConstraints && Particles.InvM(Index0) == (FSolverReal)0. && Particles.InvM(Index1) == (FSolverReal)0. && Particles.InvM(Index2) == (FSolverReal)0.;
 			}))
 		, Stiffness(InStiffness, StiffnessMultipliers, TConstArrayView<TVec3<int32>>(Constraints), ParticleOffset, ParticleCount)
 	{
@@ -31,48 +32,48 @@ public:
 
 	virtual ~FPBDAxialSpringConstraintsBase() {}
 
-	void SetProperties(const FVec2& InStiffness) { Stiffness.SetWeightedValue(InStiffness); }
+	void SetProperties(const FSolverVec2& InStiffness) { Stiffness.SetWeightedValue(InStiffness); }
 
-	void ApplyProperties(const FReal Dt, const int32 NumIterations) { Stiffness.ApplyValues(Dt, NumIterations); }
+	void ApplyProperties(const FSolverReal Dt, const int32 NumIterations) { Stiffness.ApplyValues(Dt, NumIterations); }
 
 protected:
-	inline FVec3 GetDelta(const FPBDParticles& Particles, const int32 ConstraintIndex, const FReal ExpStiffnessValue) const
+	inline FSolverVec3 GetDelta(const FSolverParticles& Particles, const int32 ConstraintIndex, const FSolverReal ExpStiffnessValue) const
 	{
 		const TVec3<int32>& Constraint = Constraints[ConstraintIndex];
 		const int32 i1 = Constraint[0];
 		const int32 i2 = Constraint[1];
 		const int32 i3 = Constraint[2];
-		const FReal PInvMass = Particles.InvM(i3) * ((FReal)1. - Barys[ConstraintIndex]) + Particles.InvM(i2) * Barys[ConstraintIndex];
-		if (Particles.InvM(i1) == (FReal)0. && PInvMass == (FReal)0.)
+		const FSolverReal PInvMass = Particles.InvM(i3) * ((FSolverReal)1. - Barys[ConstraintIndex]) + Particles.InvM(i2) * Barys[ConstraintIndex];
+		if (Particles.InvM(i1) == (FSolverReal)0. && PInvMass == (FSolverReal)0.)
 		{
-			return FVec3((FReal)0.);
+			return FSolverVec3((FSolverReal)0.);
 		}
-		const FVec3& P1 = Particles.P(i1);
-		const FVec3& P2 = Particles.P(i2);
-		const FVec3& P3 = Particles.P(i3);
-		const FVec3 P = (P2 - P3) * Barys[ConstraintIndex] + P3;
-		const FVec3 Difference = P1 - P;
-		const FReal Distance = Difference.Size();
-		if (UNLIKELY(Distance <= (FReal)SMALL_NUMBER))
+		const FSolverVec3& P1 = Particles.P(i1);
+		const FSolverVec3& P2 = Particles.P(i2);
+		const FSolverVec3& P3 = Particles.P(i3);
+		const FSolverVec3 P = (P2 - P3) * Barys[ConstraintIndex] + P3;
+		const FSolverVec3 Difference = P1 - P;
+		const FSolverReal Distance = Difference.Size();
+		if (UNLIKELY(Distance <= (FSolverReal)SMALL_NUMBER))
 		{
-			return FVec3((FReal)0.);
+			return FSolverVec3((FSolverReal)0.);
 		}
-		const FVec3 Direction = Difference / Distance;
-		const FVec3 Delta = (Distance - Dists[ConstraintIndex]) * Direction;
-		const FReal CombinedInvMass = PInvMass + Particles.InvM(i1);
-		checkSlow(CombinedInvMass > (FReal)1e-7);
+		const FSolverVec3 Direction = Difference / Distance;
+		const FSolverVec3 Delta = (Distance - Dists[ConstraintIndex]) * Direction;
+		const FSolverReal CombinedInvMass = PInvMass + Particles.InvM(i1);
+		checkSlow(CombinedInvMass > (FSolverReal)1e-7);
 		return ExpStiffnessValue * Delta / CombinedInvMass;
 	}
 
 private:
-	FReal FindBary(const FPBDParticles& Particles, const int32 i1, const int32 i2, const int32 i3)
+	FSolverReal FindBary(const FSolverParticles& Particles, const int32 i1, const int32 i2, const int32 i3)
 	{
-		const FVec3& P1 = Particles.X(i1);
-		const FVec3& P2 = Particles.X(i2);
-		const FVec3& P3 = Particles.X(i3);
-		const FVec3& P32 = P3 - P2;
-		const FReal Bary = FVec3::DotProduct(P32, P3 - P1) / P32.SizeSquared();
-		return FMath::Clamp(Bary, (FReal)0., (FReal)1.);
+		const FSolverVec3& P1 = Particles.X(i1);
+		const FSolverVec3& P2 = Particles.X(i2);
+		const FSolverVec3& P3 = Particles.X(i3);
+		const FSolverVec3& P32 = P3 - P2;
+		const FSolverReal Bary = FSolverVec3::DotProduct(P32, P3 - P1) / P32.SizeSquared();
+		return FMath::Clamp(Bary, (FSolverReal)0., (FSolverReal)1.);
 	}
 
 	template<typename Predicate>
@@ -103,7 +104,7 @@ private:
 		return TrimmedConstraints.Array();
 	}
 
-	void Init(const FPBDParticles& Particles)
+	void Init(const FSolverParticles& Particles)
 	{
 		Barys.Reset(Constraints.Num());
 		Dists.Reset(Constraints.Num());
@@ -114,13 +115,13 @@ private:
 			int32 i2 = Constraint[1];
 			int32 i3 = Constraint[2];
 			// Find Bary closest to 0.5
-			const FReal Bary1 = FindBary(Particles, i1, i2, i3);
-			const FReal Bary2 = FindBary(Particles, i2, i3, i1);
-			const FReal Bary3 = FindBary(Particles, i3, i1, i2);
-			FReal Bary = Bary1;
-			const FReal Bary1dist = FGenericPlatformMath::Abs(Bary1 - (FReal)0.5);
-			const FReal Bary2dist = FGenericPlatformMath::Abs(Bary2 - (FReal)0.5);
-			const FReal Bary3dist = FGenericPlatformMath::Abs(Bary3 - (FReal)0.5);
+			const FSolverReal Bary1 = FindBary(Particles, i1, i2, i3);
+			const FSolverReal Bary2 = FindBary(Particles, i2, i3, i1);
+			const FSolverReal Bary3 = FindBary(Particles, i3, i1, i2);
+			FSolverReal Bary = Bary1;
+			const FSolverReal Bary1dist = FGenericPlatformMath::Abs(Bary1 - (FSolverReal)0.5);
+			const FSolverReal Bary2dist = FGenericPlatformMath::Abs(Bary2 - (FSolverReal)0.5);
+			const FSolverReal Bary3dist = FGenericPlatformMath::Abs(Bary3 - (FSolverReal)0.5);
 			if (Bary3dist < Bary2dist && Bary3dist < Bary1dist)
 			{
 				Constraint[0] = i3;
@@ -139,10 +140,10 @@ private:
 			i1 = Constraint[0];
 			i2 = Constraint[1];
 			i3 = Constraint[2];
-			const FVec3& P1 = Particles.X(i1);
-			const FVec3& P2 = Particles.X(i2);
-			const FVec3& P3 = Particles.X(i3);
-			const FVec3 P = (P2 - P3) * Bary + P3;
+			const FSolverVec3& P1 = Particles.X(i1);
+			const FSolverVec3& P2 = Particles.X(i2);
+			const FSolverVec3& P3 = Particles.X(i3);
+			const FSolverVec3 P = (P2 - P3) * Bary + P3;
 			Barys.Add(Bary);
 			Dists.Add((P1 - P).Size());
 		}
@@ -150,9 +151,9 @@ private:
 
 protected:
 	TArray<TVec3<int32>> Constraints;
-	TArray<FReal> Barys;
-	TArray<FReal> Dists;
+	TArray<FSolverReal> Barys;
+	TArray<FSolverReal> Dists;
 	FPBDStiffness Stiffness;
 };
 
-}
+}  // End namespace Chaos::Softs
diff --git a/Engine/Source/Runtime/Experimental/Chaos/Public/Chaos/PBDBendingConstraints.h b/Engine/Source/Runtime/Experimental/Chaos/Public/Chaos/PBDBendingConstraints.h
--- a/Engine/Source/Runtime/Experimental/Chaos/Public/Chaos/PBDBendingConstraints.h
+++ b/Engine/Source/Runtime/Experimental/Chaos/Public/Chaos/PBDBendingConstraints.h
@@ -1,34 +1,32 @@
 // Copyright Epic Games, Inc. All Rights Reserved.
 #pragma once
 
-#include "Chaos/Array.h"
 #include "Chaos/PBDBendingConstraintsBase.h"
-#include "Chaos/PBDParticles.h"
-#include "Chaos/ParticleRule.h"
 
-namespace Chaos
+namespace Chaos::Softs
 {
-class FPBDBendingConstraints : public FParticleRule, public FPBDBendingConstraintsBase
+
+class FPBDBendingConstraints : public FPBDBendingConstraintsBase
 {
 	typedef FPBDBendingConstraintsBase Base;
-	using Base::MConstraints;
+	using Base::Constraints;
 
-  public:
-	FPBDBendingConstraints(const FDynamicParticles& InParticles, TArray<TVec4<int32>>&& Constraints, const FReal stiffness = (FReal)1.)
-	    : Base(InParticles, MoveTemp(Constraints), stiffness) {}
-	virtual ~FPBDBendingConstraints() {}
+public:
+	FPBDBendingConstraints(const FSolverParticles& InParticles, TArray<TVec4<int32>>&& InConstraints, const FSolverReal InStiffness = (FSolverReal)1.)
+	    : Base(InParticles, MoveTemp(InConstraints), InStiffness) {}
+	virtual ~FPBDBendingConstraints() override {}
 
-	void Apply(FPBDParticles& InParticles, const FReal Dt) const override //-V762
+	void Apply(FSolverParticles& InParticles, const FSolverReal Dt) const
 	{
-		for (int i = 0; i < MConstraints.Num(); ++i)
+		for (int i = 0; i < Constraints.Num(); ++i)
 		{
-			const auto& Constraint = MConstraints[i];
+			const TVec4<int32>& Constraint = Constraints[i];
 			const int32 i1 = Constraint[0];
 			const int32 i2 = Constraint[1];
 			const int32 i3 = Constraint[2];
 			const int32 i4 = Constraint[3];
-			auto Grads = Base::GetGradients(InParticles, i);
-			auto S = Base::GetScalingFactor(InParticles, i, Grads);
+			const TArray<FSolverVec3> Grads = Base::GetGradients(InParticles, i);
+			const FSolverReal S = Base::GetScalingFactor(InParticles, i, Grads);
 			InParticles.P(i1) -= S * InParticles.InvM(i1) * Grads[0];
 			InParticles.P(i2) -= S * InParticles.InvM(i2) * Grads[1];
 			InParticles.P(i3) -= S * InParticles.InvM(i3) * Grads[2];
@@ -37,6 +35,4 @@ class FPBDBendingConstraints : public FParticleRule, public FPBDBendingConstrain
 	}
 };
 
-template<class T>
-using TPBDBendingConstraints UE_DEPRECATED(4.27, "Deprecated. this class is to be deleted, use FPBDBendingConstraints instead") = FPBDBendingConstraints;
-}
+}  // End namespace Chaos::Softs
diff --git a/Engine/Source/Runtime/Experimental/Chaos/Public/Chaos/PBDBendingConstraintsBase.h b/Engine/Source/Runtime/Experimental/Chaos/Public/Chaos/PBDBendingConstraintsBase.h
--- a/Engine/Source/Runtime/Experimental/Chaos/Public/Chaos/PBDBendingConstraintsBase.h
+++ b/Engine/Source/Runtime/Experimental/Chaos/Public/Chaos/PBDBendingConstraintsBase.h
@@ -1,104 +1,107 @@
 // Copyright Epic Games, Inc. All Rights Reserved.
 #pragma once
 
+#include "Chaos/PBDSoftsEvolutionFwd.h"
 #include "Chaos/Array.h"
 #include "Chaos/PBDParticles.h"
 #include "Chaos/ParticleRule.h"
 #include "GenericPlatform/GenericPlatformMath.h"
 
-namespace Chaos
+namespace Chaos::Softs
 {
+
 class FPBDBendingConstraintsBase
 {
   public:
-	  FPBDBendingConstraintsBase(const FDynamicParticles& InParticles, TArray<TVec4<int32>>&& Constraints, const FReal Stiffness = (FReal)1.)
-	    : MConstraints(MoveTemp(Constraints)), MStiffness(Stiffness)
+	  FPBDBendingConstraintsBase(const FSolverParticles& InParticles, TArray<TVec4<int32>>&& InConstraints, const FSolverReal InStiffness = (FSolverReal)1.)
+	    : Constraints(MoveTemp(InConstraints)), Stiffness(InStiffness)
 	{
-		for (auto Constraint : MConstraints)
+		for (const TVec4<int32>& Constraint : Constraints)
 		{
-			const FVec3& P1 = InParticles.X(Constraint[0]);
-			const FVec3& P2 = InParticles.X(Constraint[1]);
-			const FVec3& P3 = InParticles.X(Constraint[2]);
-			const FVec3& P4 = InParticles.X(Constraint[3]);
+			const FSolverVec3& P1 = InParticles.X(Constraint[0]);
+			const FSolverVec3& P2 = InParticles.X(Constraint[1]);
+			const FSolverVec3& P3 = InParticles.X(Constraint[2]);
+			const FSolverVec3& P4 = InParticles.X(Constraint[3]);
 			MAngles.Add(GetAngle(P1, P2, P3, P4));
 		}
 	}
 	virtual ~FPBDBendingConstraintsBase() {}
 
-	TArray<FVec3> GetGradients(const FPBDParticles& InParticles, const int32 i) const
+	TArray<FSolverVec3> GetGradients(const FSolverParticles& InParticles, const int32 i) const
 	{
-		TArray<FVec3> Grads;
+		TArray<FSolverVec3> Grads;
 		Grads.SetNum(4);
-		const auto& Constraint = MConstraints[i];
-		const FVec3& P1 = InParticles.P(Constraint[0]);
-		const FVec3& P2 = InParticles.P(Constraint[1]);
-		const FVec3& P3 = InParticles.P(Constraint[2]);
-		const FVec3& P4 = InParticles.P(Constraint[3]);
-		auto Edge = P2 - P1;
-		auto Normal1 = FVec3::CrossProduct(P3 - P1, P3 - P2);
+		const auto& Constraint = Constraints[i];
+		const FSolverVec3& P1 = InParticles.P(Constraint[0]);
+		const FSolverVec3& P2 = InParticles.P(Constraint[1]);
+		const FSolverVec3& P3 = InParticles.P(Constraint[2]);
+		const FSolverVec3& P4 = InParticles.P(Constraint[3]);
+		const FSolverVec3 Edge = P2 - P1;
+		auto Normal1 = FSolverVec3::CrossProduct(P3 - P1, P3 - P2);
 		SafeDivide(Normal1, Normal1.SizeSquared());
-		auto Normal2 = FVec3::CrossProduct(P4 - P2, P4 - P1);
+		auto Normal2 = FSolverVec3::CrossProduct(P4 - P2, P4 - P1);
 		SafeDivide(Normal2, Normal2.SizeSquared());
-		FReal EdgeSize = Edge.Size();
-		Grads[0] = SafeDivide(FVec3::DotProduct(Edge, P3 - P2), EdgeSize) * Normal1 + SafeDivide(FVec3::DotProduct(Edge, P4 - P2), EdgeSize) * Normal2;
-		Grads[1] = SafeDivide(FVec3::DotProduct(Edge, P1 - P3), EdgeSize) * Normal1 + SafeDivide(FVec3::DotProduct(Edge, P1 - P4), EdgeSize) * Normal2;
+		FSolverReal EdgeSize = Edge.Size();
+		Grads[0] = SafeDivide(FSolverVec3::DotProduct(Edge, P3 - P2), EdgeSize) * Normal1 + SafeDivide(FSolverVec3::DotProduct(Edge, P4 - P2), EdgeSize) * Normal2;
+		Grads[1] = SafeDivide(FSolverVec3::DotProduct(Edge, P1 - P3), EdgeSize) * Normal1 + SafeDivide(FSolverVec3::DotProduct(Edge, P1 - P4), EdgeSize) * Normal2;
 		Grads[2] = EdgeSize * Normal1;
 		Grads[3] = EdgeSize * Normal2;
 		return Grads;
 	}
 
-	FReal GetScalingFactor(const FPBDParticles& InParticles, const int32 i, const TArray<FVec3>& Grads) const
+	FSolverReal GetScalingFactor(const FSolverParticles& InParticles, const int32 i, const TArray<FSolverVec3>& Grads) const
 	{
-		const auto& Constraint = MConstraints[i];
+		const auto& Constraint = Constraints[i];
 		const int32 i1 = Constraint[0];
 		const int32 i2 = Constraint[1];
 		const int32 i3 = Constraint[2];
 		const int32 i4 = Constraint[3];
-		const FVec3& P1 = InParticles.P(i1);
-		const FVec3& P2 = InParticles.P(i2);
-		const FVec3& P3 = InParticles.P(i3);
-		const FVec3& P4 = InParticles.P(i4);
-		FReal Angle = GetAngle(P1, P2, P3, P4);
-		FReal Denom = (InParticles.InvM(i1) * Grads[0].SizeSquared() + InParticles.InvM(i2) * Grads[1].SizeSquared() + InParticles.InvM(i3) * Grads[2].SizeSquared() + InParticles.InvM(i4) * Grads[3].SizeSquared());
+		const FSolverVec3& P1 = InParticles.P(i1);
+		const FSolverVec3& P2 = InParticles.P(i2);
+		const FSolverVec3& P3 = InParticles.P(i3);
+		const FSolverVec3& P4 = InParticles.P(i4);
+		FSolverReal Angle = GetAngle(P1, P2, P3, P4);
+		FSolverReal Denom = (InParticles.InvM(i1) * Grads[0].SizeSquared() + InParticles.InvM(i2) * Grads[1].SizeSquared() + InParticles.InvM(i3) * Grads[2].SizeSquared() + InParticles.InvM(i4) * Grads[3].SizeSquared());
 		{
 			auto Edge = P2 - P1;
-			auto Normal1 = FVec3::CrossProduct(P3 - P1, P3 - P2).GetSafeNormal();
-			auto Normal2 = FVec3::CrossProduct(P4 - P2, P4 - P1).GetSafeNormal();
-			Denom = FVec3::DotProduct(Edge, FVec3::CrossProduct(Normal1, Normal2)) > (FReal)0. ? -Denom : Denom;
+			auto Normal1 = FSolverVec3::CrossProduct(P3 - P1, P3 - P2).GetSafeNormal();
+			auto Normal2 = FSolverVec3::CrossProduct(P4 - P2, P4 - P1).GetSafeNormal();
+			Denom = FSolverVec3::DotProduct(Edge, FSolverVec3::CrossProduct(Normal1, Normal2)) > (FSolverReal)0. ? -Denom : Denom;
 		}
-		FReal Delta = Angle - MAngles[i];
-		return MStiffness * SafeDivide(Delta, Denom);
+		FSolverReal Delta = Angle - MAngles[i];
+		return Stiffness * SafeDivide(Delta, Denom);
 	}
 
-	void SetStiffness(FReal InStiffness) { MStiffness = FMath::Clamp(InStiffness, (FReal)0., (FReal)1.); }
+	void SetStiffness(FSolverReal InStiffness) { Stiffness = FMath::Clamp(InStiffness, (FSolverReal)0., (FSolverReal)1.); }
 
   private:
 	template<class TNum>
-	inline TNum SafeDivide(const TNum& Numerator, const FReal& Denominator) const
+	inline TNum SafeDivide(const TNum& Numerator, const FSolverReal& Denominator) const
 	{
-		if (Denominator > (FReal)1e-7)
+		if (Denominator > (FSolverReal)1e-7)
 			return Numerator / Denominator;
 		return TNum(0);
 	}
 
-	inline FReal Clamp(const FReal& Value, const FReal& Low, const FReal& High) const
+	inline FSolverReal Clamp(const FSolverReal& Value, const FSolverReal& Low, const FSolverReal& High) const
 	{
 		return Value < Low ? Low : (Value > High ? High : Value);
 	}
 
-	FReal GetAngle(const FVec3& P1, const FVec3& P2, const FVec3& P3, const FVec3& P4) const
+	FSolverReal GetAngle(const FSolverVec3& P1, const FSolverVec3& P2, const FSolverVec3& P3, const FSolverVec3& P4) const
 	{
-		auto Normal1 = FVec3::CrossProduct(P3 - P1, P3 - P2).GetSafeNormal();
-		auto Normal2 = FVec3::CrossProduct(P4 - P2, P4 - P1).GetSafeNormal();
-		auto Dot = FVec3::DotProduct(Normal1, Normal2);
+		auto Normal1 = FSolverVec3::CrossProduct(P3 - P1, P3 - P2).GetSafeNormal();
+		auto Normal2 = FSolverVec3::CrossProduct(P4 - P2, P4 - P1).GetSafeNormal();
+		auto Dot = FSolverVec3::DotProduct(Normal1, Normal2);
 		return FGenericPlatformMath::Acos(Clamp(Dot, 1e-4, 1 - 1e-4));
 	}
 
   protected:
-	TArray<TVec4<int32>> MConstraints;
+	TArray<TVec4<int32>> Constraints;
 
   private:
-	TArray<FReal> MAngles;
-	FReal MStiffness;
+	TArray<FSolverReal> MAngles;
+	FSolverReal Stiffness;
 };
-}
+
+}  // End namespace Chaos::Softs
