Chaos Cloth - Move back the post LWC solver to float for performance reason, since cloth is being simulated in local space.

diff --git a/Engine/Source/Runtime/Experimental/Chaos/Public/Chaos/PBDCollisionSpringConstraints.h b/Engine/Source/Runtime/Experimental/Chaos/Public/Chaos/PBDCollisionSpringConstraints.h
--- a/Engine/Source/Runtime/Experimental/Chaos/Public/Chaos/PBDCollisionSpringConstraints.h
+++ b/Engine/Source/Runtime/Experimental/Chaos/Public/Chaos/PBDCollisionSpringConstraints.h
@@ -2,12 +2,11 @@
 #pragma once
 
 #if !COMPILE_WITHOUT_UNREAL_SUPPORT
-#include "Chaos/Array.h"
 #include "Chaos/PBDCollisionSpringConstraintsBase.h"
-#include "Chaos/PBDParticles.h"
 
-namespace Chaos
+namespace Chaos::Softs
 {
+
 class FPBDCollisionSpringConstraints : public FPBDCollisionSpringConstraintsBase
 {
 	typedef FPBDCollisionSpringConstraintsBase Base;
@@ -20,16 +19,16 @@ public:
 		const int32 InNumParticles,
 		const TArray<TVec3<int32>>& InElements,
 		TSet<TVec2<int32>>&& InDisabledCollisionElements,
-		const FReal InThickness = (FReal)1.,
-		const FReal InStiffness = (FReal)1.)
+		const FSolverReal InThickness = (FSolverReal)1.,
+		const FSolverReal InStiffness = (FSolverReal)1.)
 	    : Base(InOffset, InNumParticles, InElements, MoveTemp(InDisabledCollisionElements), InThickness, InStiffness)
 	{}
 
-	virtual ~FPBDCollisionSpringConstraints() {}
+	virtual ~FPBDCollisionSpringConstraints() override {}
 
 	using Base::Init;
 
-	void Apply(FPBDParticles& Particles, const FReal Dt, const int32 ConstraintIndex) const
+	void Apply(FSolverParticles& Particles, const FSolverReal Dt, const int32 ConstraintIndex) const
 	{
 		const int32 i = ConstraintIndex;
 		const TVector<int32, 4>& Constraint = Constraints[i];
@@ -37,27 +36,27 @@ public:
 		const int32 i2 = Constraint[1];
 		const int32 i3 = Constraint[2];
 		const int32 i4 = Constraint[3];
-		const FVec3 Delta = Base::GetDelta(Particles, i);
-		static const FReal Multiplier = (FReal)1.;  // TODO(mlentine): Figure out what the best multiplier here is
-		if (Particles.InvM(i1) > 0)
+		const FSolverVec3 Delta = Base::GetDelta(Particles, i);
+		static const FSolverReal Multiplier = (FSolverReal)1.;  // TODO(mlentine): Figure out what the best multiplier here is
+		if (Particles.InvM(i1) > (FSolverReal)0.)
 		{
 			Particles.P(i1) += Multiplier * Particles.InvM(i1) * Delta;
 		}
-		if (Particles.InvM(i2) > 0)
+		if (Particles.InvM(i2) > (FSolverReal)0.)
 		{
 			Particles.P(i2) -= Multiplier * Particles.InvM(i2) * Barys[i][0] * Delta;
 		}
-		if (Particles.InvM(i3) > 0)
+		if (Particles.InvM(i3) > (FSolverReal)0.)
 		{
 			Particles.P(i3) -= Multiplier * Particles.InvM(i3) * Barys[i][1] * Delta;
 		}
-		if (Particles.InvM(i4) > 0)
+		if (Particles.InvM(i4) > (FSolverReal)0.)
 		{
 			Particles.P(i4) -= Multiplier * Particles.InvM(i4) * Barys[i][2] * Delta;
 		}
 	}
 
-	void Apply(FPBDParticles& InParticles, const FReal Dt) const
+	void Apply(FSolverParticles& InParticles, const FSolverReal Dt) const
 	{
 		for (int32 i = 0; i < Constraints.Num(); ++i)
 		{
@@ -65,7 +64,7 @@ public:
 		}
 	}
 
-	void Apply(FPBDParticles& InParticles, const FReal Dt, const TArray<int32>& InConstraintIndices) const
+	void Apply(FSolverParticles& InParticles, const FSolverReal Dt, const TArray<int32>& InConstraintIndices) const
 	{
 		for (int32 i : InConstraintIndices)
 		{
@@ -74,5 +73,7 @@ public:
 	}
 
 };
-}
+
+}  // End namespace Chaos::Softs
+
 #endif
diff --git a/Engine/Source/Runtime/Experimental/Chaos/Public/Chaos/PBDCollisionSpringConstraintsBase.h b/Engine/Source/Runtime/Experimental/Chaos/Public/Chaos/PBDCollisionSpringConstraintsBase.h
--- a/Engine/Source/Runtime/Experimental/Chaos/Public/Chaos/PBDCollisionSpringConstraintsBase.h
+++ b/Engine/Source/Runtime/Experimental/Chaos/Public/Chaos/PBDCollisionSpringConstraintsBase.h
@@ -3,12 +3,13 @@
 
 #if !COMPILE_WITHOUT_UNREAL_SUPPORT
 #include "Chaos/Core.h"
-#include "Chaos/Array.h"
+#include "Chaos/PBDSoftsEvolutionFwd.h"
 #include "Chaos/PBDParticles.h"
 #include "Containers/Set.h"
 
-namespace Chaos
+namespace Chaos::Softs
 {
+
 // This is an invertible spring class, typical springs are not invertible aware
 class CHAOS_API FPBDCollisionSpringConstraintsBase
 {
@@ -18,34 +19,36 @@ public:
 		const int32 InNumParticles,
 		const TArray<TVec3<int32>>& InElements,
 		TSet<TVec2<int32>>&& InDisabledCollisionElements,
-		const FReal InThickness = (FReal)1.,
-		const FReal InStiffness = (FReal)1.);
+		const FSolverReal InThickness = (FSolverReal)1.,
+		const FSolverReal InStiffness = (FSolverReal)1.);
 
 	virtual ~FPBDCollisionSpringConstraintsBase() {}
 
-	void Init(const FPBDParticles& Particles);
+	void Init(const FSolverParticles& Particles);
 
-	FVec3 GetDelta(const FPBDParticles& InParticles, const int32 i) const;
+	FSolverVec3 GetDelta(const FSolverParticles& InParticles, const int32 i) const;
 
 	const TArray<TVec4<int32>>& GetConstraints() const { return Constraints;  }
-	const TArray<FVec3>& GetBarys() const { return Barys; }
-	const TArray<FVec3>& GetNormals() const { return Normals; }
-	FReal GetThickness() const { return Thickness; }
+	const TArray<FSolverVec3>& GetBarys() const { return Barys; }
+	const TArray<FSolverVec3>& GetNormals() const { return Normals; }
+	FSolverReal GetThickness() const { return Thickness; }
 
-	void SetThickness(FReal InThickness) { Thickness = FMath::Max(InThickness, (FReal)0.);  }
+	void SetThickness(FSolverReal InThickness) { Thickness = FMath::Max(InThickness, (FSolverReal)0.);  }
 
 protected:
 	TArray<TVec4<int32>> Constraints;
-	TArray<FVec3> Barys;
-	TArray<FVec3> Normals;
+	TArray<FSolverVec3> Barys;
+	TArray<FSolverVec3> Normals;
 
 private:
 	const TArray<TVec3<int32>>& Elements;
 	const TSet<TVec2<int32>> DisabledCollisionElements;  // TODO: Make this a bitarray
 	int32 Offset;
 	int32 NumParticles;
-	FReal Thickness;
-	FReal Stiffness;
+	FSolverReal Thickness;
+	FSolverReal Stiffness;
 };
-}
+
+}  // End namespace Chaos::Softs
+
 #endif
diff --git a/Engine/Source/Runtime/Experimental/Chaos/Public/Chaos/PBDEvolution.h b/Engine/Source/Runtime/Experimental/Chaos/Public/Chaos/PBDEvolution.h
--- a/Engine/Source/Runtime/Experimental/Chaos/Public/Chaos/PBDEvolution.h
+++ b/Engine/Source/Runtime/Experimental/Chaos/Public/Chaos/PBDEvolution.h
@@ -1,27 +1,26 @@
 // Copyright Epic Games, Inc. All Rights Reserved.
 #pragma once
 
+#include "Chaos/Core.h"
+#include "Chaos/ArrayCollection.h"
+#include "Chaos/PBDActiveView.h"
+#include "Chaos/PBDSoftsEvolutionFwd.h"
+#include "Chaos/PBDParticles.h"
 #include "Chaos/KinematicGeometryParticles.h"
-#include "Chaos/PerParticleGravity.h"
 #include "Chaos/VelocityField.h"
-#include "Chaos/PBDParticles.h"
-#include "Chaos/PBDActiveView.h"
-#include "Chaos/Vector.h"
 
-namespace Chaos
+namespace Chaos::Softs
 {
 
 class CHAOS_API FPBDEvolution : public TArrayCollection
 {
  public:
-	using FGravityForces = FPerParticleGravity;
-
 	// TODO(mlentine): Init particles from some type of input
-	FPBDEvolution(FPBDParticles&& InParticles, FKinematicGeometryClothParticles&& InGeometryParticles, TArray<TVec3<int32>>&& CollisionTriangles, int32 NumIterations = 1, FReal CollisionThickness = 0, FReal SelfCollisionsThickness = 0, FReal CoefficientOfFriction = 0, FReal Damping = 0.04f);
+	FPBDEvolution(FSolverParticles&& InParticles, FSolverRigidParticles&& InGeometryParticles, TArray<TVec3<int32>>&& CollisionTriangles, int32 NumIterations = 1, FSolverReal CollisionThickness = 0, FSolverReal SelfCollisionsThickness = 0, FSolverReal CoefficientOfFriction = 0, FSolverReal Damping = 0.04f);
 	~FPBDEvolution() {}
 
 	// Advance one time step. Filter the input time step if specified.
-	void AdvanceOneTimeStep(const FReal Dt, const bool bSmoothDt = true);
+	void AdvanceOneTimeStep(const FSolverReal Dt, const bool bSmoothDt = true);
 
 	// Remove all particles, will also reset all rules
 	void ResetParticles();
@@ -36,9 +35,9 @@ class CHAOS_API FPBDEvolution : public TArrayCollection
 	void ActivateParticleRange(int32 Offset, bool bActivate)  { MParticlesActiveView.ActivateRange(Offset, bActivate); }
 
 	// Particles accessors
-	const FPBDParticles& Particles() const { return MParticles; }
-	FPBDParticles& Particles() { return MParticles; }
-	const TPBDActiveView<FPBDParticles>& ParticlesActiveView() { return MParticlesActiveView; }
+	const FSolverParticles& Particles() const { return MParticles; }
+	FSolverParticles& Particles() { return MParticles; }
+	const TPBDActiveView<FSolverParticles>& ParticlesActiveView() { return MParticlesActiveView; }
 
 	const TArray<uint32>& ParticleGroupIds() const { return MParticleGroupIds; }
 
@@ -56,10 +55,10 @@ class CHAOS_API FPBDEvolution : public TArrayCollection
 	int32 GetCollisionParticleRangeSize(int32 Offset) const { return MCollisionParticlesActiveView.GetRangeSize(Offset); }
 
 	// Collision particles accessors
-	const FKinematicGeometryClothParticles& CollisionParticles() const { return MCollisionParticles; }
-	FKinematicGeometryClothParticles& CollisionParticles() { return MCollisionParticles; }
+	const FSolverRigidParticles& CollisionParticles() const { return MCollisionParticles; }
+	FSolverRigidParticles& CollisionParticles() { return MCollisionParticles; }
 	const TArray<uint32>& CollisionParticleGroupIds() const { return MCollisionParticleGroupIds; }
-	const TPBDActiveView<FKinematicGeometryClothParticles>& CollisionParticlesActiveView() { return MCollisionParticlesActiveView; }
+	const TPBDActiveView<FSolverRigidParticles>& CollisionParticlesActiveView() { return MCollisionParticlesActiveView; }
 
 	// Reset all constraint init and rule functions.
 	void ResetConstraintRules() { MConstraintInits.Reset(); MConstraintRules.Reset(); MConstraintInitsActiveView.Reset(); MConstraintRulesActiveView.Reset();  };
@@ -77,19 +76,19 @@ class CHAOS_API FPBDEvolution : public TArrayCollection
 	void ActivateConstraintRuleRange(int32 Offset, bool bActivate) { MConstraintRulesActiveView.ActivateRange(Offset, bActivate); }
 
 	// Constraint accessors
-	const TArray<TFunction<void(const FPBDParticles&, const FReal)>>& ConstraintInits() const { return MConstraintInits; }
-	TArray<TFunction<void(const FPBDParticles&, const FReal)>>& ConstraintInits() { return MConstraintInits; }
-	const TArray<TFunction<void(FPBDParticles&, const FReal)>>& ConstraintRules() const { return MConstraintRules; }
-	TArray<TFunction<void(FPBDParticles&, const FReal)>>& ConstraintRules() { return MConstraintRules; }
+	const TArray<TFunction<void(const FSolverParticles&, const FSolverReal)>>& ConstraintInits() const { return MConstraintInits; }
+	TArray<TFunction<void(const FSolverParticles&, const FSolverReal)>>& ConstraintInits() { return MConstraintInits; }
+	const TArray<TFunction<void(FSolverParticles&, const FSolverReal)>>& ConstraintRules() const { return MConstraintRules; }
+	TArray<TFunction<void(FSolverParticles&, const FSolverReal)>>& ConstraintRules() { return MConstraintRules; }
 	
-	void SetKinematicUpdateFunction(TFunction<void(FPBDParticles&, const FReal, const FReal, const int32)> KinematicUpdate) { MKinematicUpdate = KinematicUpdate; }
-	void SetCollisionKinematicUpdateFunction(TFunction<void(FKinematicGeometryClothParticles&, const FReal, const FReal, const int32)> KinematicUpdate) { MCollisionKinematicUpdate = KinematicUpdate; }
+	void SetKinematicUpdateFunction(TFunction<void(FSolverParticles&, const FSolverReal, const FSolverReal, const int32)> KinematicUpdate) { MKinematicUpdate = KinematicUpdate; }
+	void SetCollisionKinematicUpdateFunction(TFunction<void(FSolverRigidParticles&, const FSolverReal, const FSolverReal, const int32)> KinematicUpdate) { MCollisionKinematicUpdate = KinematicUpdate; }
 
-	TFunction<void(FPBDParticles&, const FReal, const int32)>& GetForceFunction(const uint32 GroupId = 0) { return MGroupForceRules[GroupId]; }
-	const TFunction<void(FPBDParticles&, const FReal, const int32)>& GetForceFunction(const uint32 GroupId = 0) const { return MGroupForceRules[GroupId]; }
+	TFunction<void(FSolverParticles&, const FSolverReal, const int32)>& GetForceFunction(const uint32 GroupId = 0) { return MGroupForceRules[GroupId]; }
+	const TFunction<void(FSolverParticles&, const FSolverReal, const int32)>& GetForceFunction(const uint32 GroupId = 0) const { return MGroupForceRules[GroupId]; }
 
-	FGravityForces& GetGravityForces(const uint32 GroupId = 0) { check(GroupId < TArrayCollection::Size()); return MGroupGravityForces[GroupId]; }
-	const FGravityForces& GetGravityForces(const uint32 GroupId = 0) const { check(GroupId < TArrayCollection::Size()); return MGroupGravityForces[GroupId]; }
+	const FSolverVec3& GetGravity(const uint32 GroupId = 0) const { check(GroupId < TArrayCollection::Size()); return MGroupGravityAccelerations[GroupId]; }
+	void SetGravity(const FSolverVec3& Acceleration, const uint32 GroupId = 0) { check(GroupId < TArrayCollection::Size()); MGroupGravityAccelerations[GroupId] = Acceleration; }
 
 	FVelocityField& GetVelocityField(const uint32 GroupId = 0) { check(GroupId < TArrayCollection::Size()); return MGroupVelocityFields[GroupId]; }
 	const FVelocityField& GetVelocityField(const uint32 GroupId = 0) const { check(GroupId < TArrayCollection::Size()); return MGroupVelocityFields[GroupId]; }
@@ -101,17 +100,17 @@ class CHAOS_API FPBDEvolution : public TArrayCollection
 	int32 GetIterations() const { return MNumIterations; }
 	void SetIterations(const int32 Iterations) { MNumIterations = Iterations; }
 
-	FReal GetSelfCollisionThickness(const uint32 GroupId = 0) const { check(GroupId < TArrayCollection::Size()); return MGroupSelfCollisionThicknesses[GroupId]; }
-	void SetSelfCollisionThickness(const FReal SelfCollisionThickness, const uint32 GroupId = 0) { check(GroupId < TArrayCollection::Size()); MGroupSelfCollisionThicknesses[GroupId] = SelfCollisionThickness; }
+	FSolverReal GetSelfCollisionThickness(const uint32 GroupId = 0) const { check(GroupId < TArrayCollection::Size()); return MGroupSelfCollisionThicknesses[GroupId]; }
+	void SetSelfCollisionThickness(const FSolverReal SelfCollisionThickness, const uint32 GroupId = 0) { check(GroupId < TArrayCollection::Size()); MGroupSelfCollisionThicknesses[GroupId] = SelfCollisionThickness; }
 
-	FReal GetCollisionThickness(const uint32 GroupId = 0) const { check(GroupId < TArrayCollection::Size()); return MGroupCollisionThicknesses[GroupId]; }
-	void SetCollisionThickness(const FReal CollisionThickness, const uint32 GroupId = 0) { check(GroupId < TArrayCollection::Size()); MGroupCollisionThicknesses[GroupId] = CollisionThickness; }
+	FSolverReal GetCollisionThickness(const uint32 GroupId = 0) const { check(GroupId < TArrayCollection::Size()); return MGroupCollisionThicknesses[GroupId]; }
+	void SetCollisionThickness(const FSolverReal CollisionThickness, const uint32 GroupId = 0) { check(GroupId < TArrayCollection::Size()); MGroupCollisionThicknesses[GroupId] = CollisionThickness; }
 
-	FReal GetCoefficientOfFriction(const uint32 GroupId = 0) const { check(GroupId < TArrayCollection::Size()); return MGroupCoefficientOfFrictions[GroupId]; }
-	void SetCoefficientOfFriction(const FReal CoefficientOfFriction, const uint32 GroupId = 0) { check(GroupId < TArrayCollection::Size()); MGroupCoefficientOfFrictions[GroupId] = CoefficientOfFriction; }
+	FSolverReal GetCoefficientOfFriction(const uint32 GroupId = 0) const { check(GroupId < TArrayCollection::Size()); return MGroupCoefficientOfFrictions[GroupId]; }
+	void SetCoefficientOfFriction(const FSolverReal CoefficientOfFriction, const uint32 GroupId = 0) { check(GroupId < TArrayCollection::Size()); MGroupCoefficientOfFrictions[GroupId] = CoefficientOfFriction; }
 
-	FReal GetDamping(const uint32 GroupId = 0) const { check(GroupId < TArrayCollection::Size()); return MGroupDampings[GroupId]; }
-	void SetDamping(const FReal Damping, const uint32 GroupId = 0) { check(GroupId < TArrayCollection::Size()); MGroupDampings[GroupId] = Damping; }
+	FSolverReal GetDamping(const uint32 GroupId = 0) const { check(GroupId < TArrayCollection::Size()); return MGroupDampings[GroupId]; }
+	void SetDamping(const FSolverReal Damping, const uint32 GroupId = 0) { check(GroupId < TArrayCollection::Size()); MGroupDampings[GroupId] = Damping; }
 
 	bool GetUseCCD(const uint32 GroupId = 0) const { check(GroupId < TArrayCollection::Size()); return MGroupUseCCDs[GroupId]; }
 	void SetUseCCD(const bool bUseCCD, const uint32 GroupId = 0) { check(GroupId < TArrayCollection::Size()); MGroupUseCCDs[GroupId] = bUseCCD; }
@@ -120,10 +119,10 @@ class CHAOS_API FPBDEvolution : public TArrayCollection
 	const bool Collided(int32 index) { return MCollided[index]; }
 
 	const TArray<bool>& GetCollisionStatus() { return MCollided; }
-	const TArray<FVec3>& GetCollisionContacts() const { return MCollisionContacts; }
-	const TArray<FVec3>& GetCollisionNormals() const { return MCollisionNormals; }
+	const TArray<FSolverVec3>& GetCollisionContacts() const { return MCollisionContacts; }
+	const TArray<FSolverVec3>& GetCollisionNormals() const { return MCollisionNormals; }
 
-	FReal GetTime() const { return MTime; }
+	FSolverReal GetTime() const { return MTime; }
 
  private:
 	// Add simulation groups and set default values
@@ -132,51 +131,52 @@ class CHAOS_API FPBDEvolution : public TArrayCollection
 	void ResetGroups();
 	// Selected versions of the pre-iteration updates (euler step, force, velocity field. damping updates)..
 	template<bool bForceRule, bool bVelocityField, bool bDampVelocityRule>
-	void PreIterationUpdate(const FReal Dt, const int32 Offset, const int32 Range, const int32 MinParallelBatchSize);
+	void PreIterationUpdate(const FSolverReal Dt, const int32 Offset, const int32 Range, const int32 MinParallelBatchSize);
 
 private:
-	FPBDParticles MParticles;
-	TPBDActiveView<FPBDParticles> MParticlesActiveView;
-	FKinematicGeometryClothParticles MCollisionParticles;
-	TPBDActiveView<FKinematicGeometryClothParticles> MCollisionParticlesActiveView;
+	FSolverParticles MParticles;
+	TPBDActiveView<FSolverParticles> MParticlesActiveView;
+	FSolverRigidParticles MCollisionParticles;
+	TPBDActiveView<FSolverRigidParticles> MCollisionParticlesActiveView;
 
 	TArray<TVector<int32, 3>> MCollisionTriangles;       // Used for self-collisions
 	TSet<TVector<int32, 2>> MDisabledCollisionElements;  // 
 
-	TArrayCollectionArray<FRigidTransform3> MCollisionTransforms;  // Used for CCD to store the initial state before the kinematic update
+	TArrayCollectionArray<FSolverRigidTransform3> MCollisionTransforms;  // Used for CCD to store the initial state before the kinematic update
 	TArrayCollectionArray<bool> MCollided;
 	TArrayCollectionArray<uint32> MCollisionParticleGroupIds;  // Used for per group parameters for collision particles
 	TArrayCollectionArray<uint32> MParticleGroupIds;  // Used for per group parameters for particles
-	TArray<FVec3> MCollisionContacts;
-	TArray<FVec3> MCollisionNormals;
+	TArray<FSolverVec3> MCollisionContacts;
+	TArray<FSolverVec3> MCollisionNormals;
 
-	TArrayCollectionArray<FGravityForces> MGroupGravityForces;
+	TArrayCollectionArray<FSolverVec3> MGroupGravityAccelerations;
 	TArrayCollectionArray<FVelocityField> MGroupVelocityFields;
-	TArrayCollectionArray<TFunction<void(FPBDParticles&, const FReal, const int32)>> MGroupForceRules;
-	TArrayCollectionArray<FReal> MGroupCollisionThicknesses;
-	TArrayCollectionArray<FReal> MGroupSelfCollisionThicknesses;
-	TArrayCollectionArray<FReal> MGroupCoefficientOfFrictions;
-	TArrayCollectionArray<FReal> MGroupDampings;
+	TArrayCollectionArray<TFunction<void(FSolverParticles&, const FSolverReal, const int32)>> MGroupForceRules;
+	TArrayCollectionArray<FSolverReal> MGroupCollisionThicknesses;
+	TArrayCollectionArray<FSolverReal> MGroupSelfCollisionThicknesses;
+	TArrayCollectionArray<FSolverReal> MGroupCoefficientOfFrictions;
+	TArrayCollectionArray<FSolverReal> MGroupDampings;
 	TArrayCollectionArray<bool> MGroupUseCCDs;
 	
-	TArray<TFunction<void(const FPBDParticles&, const FReal)>> MConstraintInits;
-	TPBDActiveView<TArray<TFunction<void(const FPBDParticles&, const FReal)>>> MConstraintInitsActiveView;
-	TArray<TFunction<void(FPBDParticles&, const FReal)>> MConstraintRules;
-	TPBDActiveView<TArray<TFunction<void(FPBDParticles&, const FReal)>>> MConstraintRulesActiveView;
+	TArray<TFunction<void(const FSolverParticles&, const FSolverReal)>> MConstraintInits;
+	TPBDActiveView<TArray<TFunction<void(const FSolverParticles&, const FSolverReal)>>> MConstraintInitsActiveView;
+	TArray<TFunction<void(FSolverParticles&, const FSolverReal)>> MConstraintRules;
+	TPBDActiveView<TArray<TFunction<void(FSolverParticles&, const FSolverReal)>>> MConstraintRulesActiveView;
 
-	TFunction<void(FPBDParticles&, const FReal, const FReal, const int32)> MKinematicUpdate;
-	TFunction<void(FKinematicGeometryClothParticles&, const FReal, const FReal, const int32)> MCollisionKinematicUpdate;
+	TFunction<void(FSolverParticles&, const FSolverReal, const FSolverReal, const int32)> MKinematicUpdate;
+	TFunction<void(FSolverRigidParticles&, const FSolverReal, const FSolverReal, const int32)> MCollisionKinematicUpdate;
 
 	int32 MNumIterations;
-	FVec3 MGravity;
-	FReal MCollisionThickness;
-	FReal MSelfCollisionThickness;
-	FReal MCoefficientOfFriction;
-	FReal MDamping;
-	FReal MTime;
-	FReal MSmoothDt;
+	FSolverVec3 MGravity;
+	FSolverReal MCollisionThickness;
+	FSolverReal MSelfCollisionThickness;
+	FSolverReal MCoefficientOfFriction;
+	FSolverReal MDamping;
+	FSolverReal MTime;
+	FSolverReal MSmoothDt;
 };
-}
+
+}  // End namespace Chaos::Softs
 
 // Support ISPC enable/disable in non-shipping builds
 #if !INTEL_ISPC
diff --git a/Engine/Source/Runtime/Experimental/Chaos/Public/Chaos/PBDLongRangeConstraints.h b/Engine/Source/Runtime/Experimental/Chaos/Public/Chaos/PBDLongRangeConstraints.h
--- a/Engine/Source/Runtime/Experimental/Chaos/Public/Chaos/PBDLongRangeConstraints.h
+++ b/Engine/Source/Runtime/Experimental/Chaos/Public/Chaos/PBDLongRangeConstraints.h
@@ -1,32 +1,34 @@
 // Copyright Epic Games, Inc. All Rights Reserved.
 #pragma once
 
+#include "Chaos/PBDSoftsEvolutionFwd.h"
 #include "Chaos/PBDLongRangeConstraintsBase.h"
 #include "ChaosStats.h"
 
 DECLARE_CYCLE_STAT(TEXT("Chaos PBD Long Range Constraint"), STAT_PBD_LongRange, STATGROUP_Chaos);
 
-namespace Chaos
+namespace Chaos::Softs
 {
-class CHAOS_API FPBDLongRangeConstraints final : public FPBDLongRangeConstraintsBase
+
+class CHAOS_API FPBDLongRangeConstraints : public FPBDLongRangeConstraintsBase
 {
 public:
 	typedef FPBDLongRangeConstraintsBase Base;
 	typedef typename Base::FTether FTether;
 
 	FPBDLongRangeConstraints(
-		const FPBDParticles& Particles,
+		const FSolverParticles& Particles,
 		const int32 InParticleOffset,
 		const int32 InParticleCount,
 		const TArray<TConstArrayView<TTuple<int32, int32, FRealSingle>>>& InTethers,
 		const TConstArrayView<FRealSingle>& StiffnessMultipliers,
 		const TConstArrayView<FRealSingle>& ScaleMultipliers,
-		const FVec2& InStiffness = FVec2::UnitVector,
-		const FVec2& InScale = FVec2::UnitVector)
+		const FSolverVec2& InStiffness = FSolverVec2::UnitVector,
+		const FSolverVec2& InScale = FSolverVec2::UnitVector)
 		: FPBDLongRangeConstraintsBase(Particles, InParticleOffset, InParticleCount, InTethers, StiffnessMultipliers, ScaleMultipliers, InStiffness, InScale) {}
-	virtual ~FPBDLongRangeConstraints() {}
+	virtual ~FPBDLongRangeConstraints() override {}
 
-	void Apply(FPBDParticles& Particles, const FReal Dt) const;
+	void Apply(FSolverParticles& Particles, const FSolverReal Dt) const;
 
 private:
 	using Base::Tethers;
@@ -35,7 +37,8 @@ private:
 	using Base::ScaleTable;
 	using Base::ScaleIndices;
 };
-}
+
+}  // End namespace Chaos::Softs
 
 // Support ISPC enable/disable in non-shipping builds
 #if !INTEL_ISPC
diff --git a/Engine/Source/Runtime/Experimental/Chaos/Public/Chaos/PBDLongRangeConstraintsBase.h b/Engine/Source/Runtime/Experimental/Chaos/Public/Chaos/PBDLongRangeConstraintsBase.h
--- a/Engine/Source/Runtime/Experimental/Chaos/Public/Chaos/PBDLongRangeConstraintsBase.h
+++ b/Engine/Source/Runtime/Experimental/Chaos/Public/Chaos/PBDLongRangeConstraintsBase.h
@@ -3,13 +3,15 @@
 
 #include "Chaos/Core.h"
 #include "Chaos/PBDParticles.h"
+#include "Chaos/PBDSoftsEvolutionFwd.h"
 #include "Chaos/PBDActiveView.h"
 #include "Chaos/PBDStiffness.h"
 #include "Containers/Map.h"
 #include "Containers/Set.h"
 
-namespace Chaos
+namespace Chaos::Softs
 {
+
 class CHAOS_API FPBDLongRangeConstraintsBase
 {
 public:
@@ -24,31 +26,31 @@ public:
 		AccurateTetherAccurateLength = Geodesic
 	};
 
-	typedef TTuple<int32, int32, float> FTether;
+	typedef TTuple<int32, int32, FRealSingle> FTether;
 
 	FPBDLongRangeConstraintsBase(
-		const FPBDParticles& Particles,
+		const FSolverParticles& Particles,
 		const int32 InParticleOffset,
 		const int32 InParticleCount,
 		const TArray<TConstArrayView<TTuple<int32, int32, FRealSingle>>>& InTethers,
 		const TConstArrayView<FRealSingle>& StiffnessMultipliers,
 		const TConstArrayView<FRealSingle>& ScaleMultipliers,
-		const FVec2& InStiffness = FVec2::UnitVector,
-		const FVec2& InScale = FVec2::UnitVector);
+		const FSolverVec2& InStiffness = FSolverVec2::UnitVector,
+		const FSolverVec2& InScale = FSolverVec2::UnitVector);
 
 	virtual ~FPBDLongRangeConstraintsBase() {}
 
 	// Return the stiffness input values used by the constraint
-	FVec2 GetStiffness() const { return Stiffness.GetWeightedValue(); }
+	FSolverVec2 GetStiffness() const { return Stiffness.GetWeightedValue(); }
 
 	// Set the stiffness input values used by the constraint
-	void SetStiffness(const FVec2& InStiffness) { Stiffness.SetWeightedValue(InStiffness); }
+	void SetStiffness(const FSolverVec2& InStiffness) { Stiffness.SetWeightedValue(InStiffness); }
 
 	// Set the scale low and high value of the scale weight map
-	void SetScale(const FVec2& InScale) { Scale = FVec2(FMath::Clamp((FReal)InScale.X, (FReal)0.01, (FReal)10.), FMath::Clamp((FReal)InScale.Y, (FReal)0.01, (FReal)10.)); }  // TODO: Fix FVec2 double to float conversion error on some platforms as to remove extra cast to FReal
+	void SetScale(const FSolverVec2& InScale) { Scale = FSolverVec2(FMath::Clamp((FSolverReal)InScale.X, (FSolverReal)0.01, (FSolverReal)10.), FMath::Clamp((FSolverReal)InScale.Y, (FSolverReal)0.01, (FSolverReal)10.)); }  // TODO: Fix FSolverVec2 double to float conversion error on some platforms as to remove extra cast to FSolverReal
 
 	// Set stiffness offset and range, as well as the simulation stiffness exponent
-	void ApplyProperties(const FReal Dt, const int32 NumIterations);
+	void ApplyProperties(const FSolverReal Dt, const int32 NumIterations);
 
 	// Return the tethers, organized in concurent friendly batches
 	const TArray<TConstArrayView<FTether>>& GetTethers() const { return Tethers; }
@@ -66,13 +68,13 @@ public:
 	int32 GetEndParticle(const FTether& Tether) const { return GetEndIndex(Tether) + ParticleOffset; }
 
 	// Return the reference length of the specified tether
-	FReal GetRefLength(const FTether& Tether) const { return (FReal)Tether.Get<2>(); }
+	FSolverReal GetRefLength(const FTether& Tether) const { return (FSolverReal)Tether.Get<2>(); }
 
 	// Return the Tether scale for the specified tether
-	FReal GetScale(const FTether& Tether) const { return HasScaleWeightMap() ? ScaleTable[ScaleIndices[GetEndIndex(Tether)]] : ScaleTable[0]; }
+	FSolverReal GetScale(const FTether& Tether) const { return HasScaleWeightMap() ? ScaleTable[ScaleIndices[GetEndIndex(Tether)]] : ScaleTable[0]; }
 
 	// Return the target length of the specified tether (= RefLength * Scale)
-	FReal GetTargetLength(const FTether& Tether) const { return GetRefLength(Tether) * GetScale(Tether); }
+	FSolverReal GetTargetLength(const FTether& Tether) const { return GetRefLength(Tether) * GetScale(Tether); }
 
 protected:
 	// Return the minimum number of long range tethers in a batch to process in parallel
@@ -82,30 +84,30 @@ protected:
 	bool HasScaleWeightMap() const { return ScaleTable.Num() > 1; }
 
 	// Return a vector representing the amount of segment required for the tether to shrink back to its maximum target length constraint, or zero if the constraint is already met
-	inline FVec3 GetDelta(const FPBDParticles& Particles, const FTether& Tether, const FReal InScale) const
+	inline FSolverVec3 GetDelta(const FSolverParticles& Particles, const FTether& Tether, const FSolverReal InScale) const
 	{
 		const int32 Start = GetStartParticle(Tether);
 		const int32 End = GetEndParticle(Tether);
-		const FReal TargetLength = GetRefLength(Tether) * InScale;
-		checkSlow(Particles.InvM(Start) == (FReal)0.);
-		checkSlow(Particles.InvM(End) > (FReal)0.);
-		FVec3 Direction = Particles.P(Start) - Particles.P(End);
-		const FReal Length = Direction.SafeNormalize();
-		const FReal Offset = Length - TargetLength;
-		return Offset < (FReal)0. ? FVec3((FReal)0.) : Offset * Direction;
+		const FSolverReal TargetLength = GetRefLength(Tether) * InScale;
+		checkSlow(Particles.InvM(Start) == (FSolverReal)0.);
+		checkSlow(Particles.InvM(End) > (FSolverReal)0.);
+		FSolverVec3 Direction = Particles.P(Start) - Particles.P(End);
+		const FSolverReal Length = Direction.SafeNormalize();
+		const FSolverReal Offset = Length - TargetLength;
+		return Offset < (FSolverReal)0. ? FSolverVec3((FSolverReal)0.) : Offset * Direction;
 	};
 
 	// Return a direction and length representing the amount of segment required for the tether to shrink back to its maximum target length constraint, or zero if the constraint is already met
-	inline void GetDelta(const FPBDParticles& Particles, const FTether& Tether, const FReal InScale, FVec3& OutDirection, FReal& OutOffset) const
+	inline void GetDelta(const FSolverParticles& Particles, const FTether& Tether, const FSolverReal InScale, FSolverVec3& OutDirection, FSolverReal& OutOffset) const
 	{
 		const int32 Start = GetStartParticle(Tether);
 		const int32 End = GetEndParticle(Tether);
-		const FReal TargetLength = GetRefLength(Tether) * InScale;
-		checkSlow(Particles.InvM(Start) == (FReal)0.);
-		checkSlow(Particles.InvM(End) > (FReal)0.);
+		const FSolverReal TargetLength = GetRefLength(Tether) * InScale;
+		checkSlow(Particles.InvM(Start) == (FSolverReal)0.);
+		checkSlow(Particles.InvM(End) > (FSolverReal)0.);
 		OutDirection = Particles.P(Start) - Particles.P(End);
-		const FReal Length = OutDirection.SafeNormalize();
-		OutOffset = FMath::Max((FReal)0., Length - TargetLength);
+		const FSolverReal Length = OutDirection.SafeNormalize();
+		OutOffset = FMath::Max((FSolverReal)0., Length - TargetLength);
 	};
 
 protected:
@@ -113,9 +115,10 @@ protected:
 	const TArray<TConstArrayView<FTether>>& Tethers;  // Array view on the tether provided to this constraint
 	FPBDStiffness Stiffness;  // Stiffness weightmap lookup table 
 	TArray<uint8> ScaleIndices;  // Per particle array of index to the scale lookup table
-	TArray<FReal> ScaleTable;  // Fixed lookup table of scale values
-	FVec2 Scale;  // High and low values of tether target length scale used by the the lookup tables interpolation
+	TArray<FSolverReal> ScaleTable;  // Fixed lookup table of scale values
+	FSolverVec2 Scale;  // High and low values of tether target length scale used by the the lookup tables interpolation
 	const int32 ParticleOffset;  // Index of the first usable particle
 	const int32 ParticleCount;  // Number of particles available to this constraint
 };
-}
+
+}  // End namespace Chaos::Softs
diff --git a/Engine/Source/Runtime/Experimental/Chaos/Public/Chaos/PBDShapeConstraints.h b/Engine/Source/Runtime/Experimental/Chaos/Public/Chaos/PBDShapeConstraints.h
--- a/Engine/Source/Runtime/Experimental/Chaos/Public/Chaos/PBDShapeConstraints.h
+++ b/Engine/Source/Runtime/Experimental/Chaos/Public/Chaos/PBDShapeConstraints.h
@@ -3,42 +3,44 @@
 
 #include "Chaos/PBDShapeConstraintsBase.h"
 
-namespace Chaos
+namespace Chaos::Softs
 {
-	class FPBDShapeConstraints : public FPBDShapeConstraintsBase
+
+class FPBDShapeConstraints : public FPBDShapeConstraintsBase
+{
+	typedef FPBDShapeConstraintsBase Base;
+	using Base::ParticleOffset;
+	using Base::TargetPositions;
+
+public:
+
+	FPBDShapeConstraints(
+		int32 InParticleOffset,
+		int32 InParticleCount,
+		const TArray<FSolverVec3>& StartPositions,
+		const TArray<FSolverVec3>& TargetPositions,
+		const FSolverReal Stiffness = (FSolverReal)1.
+	)
+		: Base(InParticleOffset, InParticleCount, StartPositions, TargetPositions, Stiffness)
 	{
-		typedef FPBDShapeConstraintsBase Base;
-		using Base::MParticleOffset;
-		using Base::MTargetPositions;
-
-	public:
-
-		FPBDShapeConstraints(
-			int32 InParticleOffset,
-			int32 InParticleCount,
-			const TArray<FVec3>& StartPositions,
-			const TArray<FVec3>& TargetPositions,
-			const FReal Stiffness = (FReal)1.
-		)
-			: Base(InParticleOffset, InParticleCount, StartPositions, TargetPositions, Stiffness)
-		{
-		}
-		virtual ~FPBDShapeConstraints() override {}
+	}
+	virtual ~FPBDShapeConstraints() override {}
 
-		void Apply(FPBDParticles& InParticles, const FReal Dt, const int32 Index) const
+	void Apply(FSolverParticles& InParticles, const FSolverReal Dt, const int32 Index) const
+	{
+		if (InParticles.InvM(Index) > (FSolverReal)0.)
 		{
-			if (InParticles.InvM(Index) > 0)
-			{
-				InParticles.P(Index) -= InParticles.InvM(Index) * Base::GetDelta(InParticles, Index);
-			}
+			InParticles.P(Index) -= InParticles.InvM(Index) * Base::GetDelta(InParticles, Index);
 		}
+	}
 
-		void Apply(FPBDParticles& InParticles, const FReal Dt) const
+	void Apply(FSolverParticles& InParticles, const FSolverReal Dt) const
+	{
+		for (int32 Index = ParticleOffset; Index < ParticleOffset + TargetPositions.Num(); ++Index)
 		{
-			for (int32 Index = MParticleOffset; Index < MParticleOffset + MTargetPositions.Num(); Index++)
-			{
-				Apply(InParticles, Dt, Index);
-			}
+			Apply(InParticles, Dt, Index);
 		}
-	};
-}
+	}
+};
+
+}  // End namespace Chaos::Softs
diff --git a/Engine/Source/Runtime/Experimental/Chaos/Public/Chaos/PBDShapeConstraintsBase.h b/Engine/Source/Runtime/Experimental/Chaos/Public/Chaos/PBDShapeConstraintsBase.h
--- a/Engine/Source/Runtime/Experimental/Chaos/Public/Chaos/PBDShapeConstraintsBase.h
+++ b/Engine/Source/Runtime/Experimental/Chaos/Public/Chaos/PBDShapeConstraintsBase.h
@@ -2,61 +2,64 @@
 #pragma once
 
 #include "Chaos/Core.h"
+#include "Chaos/PBDSoftsEvolutionFwd.h"
 #include "Chaos/Array.h"
 #include "Chaos/PBDParticles.h"
 
 #include <functional>
 
-namespace Chaos
+namespace Chaos::Softs
 {
-	class FPBDShapeConstraintsBase
+
+class FPBDShapeConstraintsBase
+{
+public:
+	FPBDShapeConstraintsBase(
+		int32 InParticleOffset,
+		int32 InParticleCount,
+		const TArray<FSolverVec3>& StartPositions,
+		const TArray<FSolverVec3>& InTargetPositions,
+		const FSolverReal InStiffness
+	)
+		: TargetPositions(InTargetPositions)
+		, ParticleOffset(InParticleOffset)
+		, Stiffness(InStiffness)
 	{
-	public:
-		FPBDShapeConstraintsBase(
-			int32 InParticleOffset,
-			int32 InParticleCount,
-			const TArray<FVec3>& StartPositions,
-			const TArray<FVec3>& TargetPositions,
-			const FReal Stiffness
-		)
-			: MTargetPositions(TargetPositions)
-			, MParticleOffset(InParticleOffset)
-			, MStiffness(Stiffness)
+		const int32 NumConstraints = InParticleCount;
+		Dists.SetNumUninitialized(InParticleCount);
+		for (int32 Index = 0; Index < InParticleCount; ++Index)
 		{
-			const int32 NumConstraints = InParticleCount;
-			MDists.SetNumUninitialized(InParticleCount);
-			for (int32 Index = 0; Index < InParticleCount; ++Index)
-			{
-				const int32 ParticleIndex = MParticleOffset + Index;
-				const FVec3& P1 = StartPositions[ParticleIndex];
-				const FVec3& P2 = MTargetPositions[ParticleIndex];
-				MDists[Index] = (P1 - P2).Size();
-			}
+			const int32 ParticleIndex = ParticleOffset + Index;
+			const FSolverVec3& P1 = StartPositions[ParticleIndex];
+			const FSolverVec3& P2 = TargetPositions[ParticleIndex];
+			Dists[Index] = (P1 - P2).Size();
 		}
-		virtual ~FPBDShapeConstraintsBase() {}
+	}
+	virtual ~FPBDShapeConstraintsBase() {}
 
-		FVec3 GetDelta(const FPBDParticles& InParticles, const int32 Index) const
+	FSolverVec3 GetDelta(const FSolverParticles& InParticles, const int32 Index) const
+	{
+		checkSlow(Index >= ParticleOffset && Index < ParticleOffset + Dists.Num())
+		if (InParticles.InvM(Index) == (FSolverReal)0.)
 		{
-			checkSlow(Index >= MParticleOffset && Index < MParticleOffset + MDists.Num())
-			if (InParticles.InvM(Index) == (FReal)0.)
-			{
-				return FVec3(0.);
-			}
-			const FVec3& P1 = InParticles.P(Index);
-			const FVec3& P2 = MTargetPositions[Index];
-			const FVec3 Difference = P1 - P2;
-			const FReal Distance = Difference.Size();
-			const FVec3 Direction = Difference / Distance;
-			const FVec3 Delta = (Distance - MDists[Index - MParticleOffset]) * Direction;
-			return MStiffness * Delta / InParticles.InvM(Index);
+			return FSolverVec3(0.);
 		}
+		const FSolverVec3& P1 = InParticles.P(Index);
+		const FSolverVec3& P2 = TargetPositions[Index];
+		const FSolverVec3 Difference = P1 - P2;
+		const FSolverReal Distance = Difference.Size();
+		const FSolverVec3 Direction = Difference / Distance;
+		const FSolverVec3 Delta = (Distance - Dists[Index - ParticleOffset]) * Direction;
+		return Stiffness * Delta / InParticles.InvM(Index);
+	}
+
+protected:
+	const TArray<FSolverVec3>& TargetPositions;
+	const int32 ParticleOffset;
 
-	protected:
-		const TArray<FVec3>& MTargetPositions;
-		const int32 MParticleOffset;
+private:
+	TArray<FSolverReal> Dists;
+	FSolverReal Stiffness;
+};
 
-	private:
-		TArray<FReal> MDists;
-		FReal MStiffness;
-	};
-}
+}  // End namespace Chaos::Softs
diff --git a/Engine/Source/Runtime/Experimental/Chaos/Public/Chaos/PBDSphericalConstraint.h b/Engine/Source/Runtime/Experimental/Chaos/Public/Chaos/PBDSphericalConstraint.h
--- a/Engine/Source/Runtime/Experimental/Chaos/Public/Chaos/PBDSphericalConstraint.h
+++ b/Engine/Source/Runtime/Experimental/Chaos/Public/Chaos/PBDSphericalConstraint.h
@@ -1,6 +1,7 @@
 // Copyright Epic Games, Inc. All Rights Reserved.
 #pragma once
 
+#include "Chaos/PBDSoftsEvolutionFwd.h"
 #include "Chaos/PBDParticles.h"
 #include "ChaosStats.h"
 #include "Containers/ArrayView.h"
@@ -19,255 +20,251 @@ const bool bChaos_Spherical_ISPC_Enabled = true;
 extern CHAOS_API bool bChaos_Spherical_ISPC_Enabled;
 #endif
 
-namespace Chaos
+namespace Chaos::Softs
 {
-	class CHAOS_API FPBDSphericalConstraint final
+
+class CHAOS_API FPBDSphericalConstraint final
+{
+public:
+	FPBDSphericalConstraint(
+		const uint32 InParticleOffset,
+		const uint32 InParticleCount,
+		const TArray<FSolverVec3>& InAnimationPositions,  // Use global indexation (will need adding ParticleOffset)
+		const TConstArrayView<FRealSingle>& InSphereRadii  // Use local indexation
+	)
+		: AnimationPositions(InAnimationPositions)
+		, SphereRadii(InSphereRadii)
+		, ParticleOffset(InParticleOffset)
+		, SphereRadiiMultiplier((FSolverReal)1.)
 	{
-	public:
-		FPBDSphericalConstraint(
-			const uint32 InParticleOffset,
-			const uint32 InParticleCount,
-			const TArray<FVec3>& InAnimationPositions,  // Use global indexation (will need adding ParticleOffset)
-			const TConstArrayView<FRealSingle>& InSphereRadii  // Use local indexation
-		)
-			: AnimationPositions(InAnimationPositions)
-			, SphereRadii(InSphereRadii)
-			, ParticleOffset(InParticleOffset)
-			, SphereRadiiMultiplier((FReal)1.)
-		{
-			check(InSphereRadii.Num() == InParticleCount);
-		}
-		~FPBDSphericalConstraint() {}
+		check(InSphereRadii.Num() == InParticleCount);
+	}
+	~FPBDSphericalConstraint() {}
 
-		void Apply(FPBDParticles& Particles, const FReal Dt) const
-		{
-			SCOPE_CYCLE_COUNTER(STAT_PBD_Spherical);
+	void Apply(FSolverParticles& Particles, const FSolverReal Dt) const
+	{
+		SCOPE_CYCLE_COUNTER(STAT_PBD_Spherical);
 
-			if (bRealTypeCompatibleWithISPC && bChaos_Spherical_ISPC_Enabled)
-			{
-				ApplyHelperISPC(Particles, Dt);
-			}
-			else
-			{
-				ApplyHelper(Particles, Dt);
-			}
+		if (bRealTypeCompatibleWithISPC && bChaos_Spherical_ISPC_Enabled)
+		{
+			ApplyHelperISPC(Particles, Dt);
 		}
-
-		void SetSphereRadiiMultiplier(const FReal InSphereRadiiMultiplier)
+		else
 		{
-			SphereRadiiMultiplier = FMath::Max((FReal)0., InSphereRadiiMultiplier);
+			ApplyHelper(Particles, Dt);
 		}
+	}
+
+	void SetSphereRadiiMultiplier(const FSolverReal InSphereRadiiMultiplier)
+	{
+		SphereRadiiMultiplier = FMath::Max((FSolverReal)0., InSphereRadiiMultiplier);
+	}
 
-	private:
-		void ApplyHelper(FPBDParticles& Particles, const FReal Dt) const
+private:
+	void ApplyHelper(FSolverParticles& Particles, const FSolverReal Dt) const
+	{
+		const int32 ParticleCount = SphereRadii.Num();
+
+		PhysicsParallelFor(ParticleCount, [&](int32 Index)  // TODO: profile need for parallel loop based on particle count
 		{
-			const int32 ParticleCount = SphereRadii.Num();
+			const int32 ParticleIndex = ParticleOffset + Index;
 
-			PhysicsParallelFor(ParticleCount, [&](int32 Index)  // TODO: profile need for parallel loop based on particle count
+			if (Particles.InvM(ParticleIndex) == 0)
 			{
-				const int32 ParticleIndex = ParticleOffset + Index;
-
-				if (Particles.InvM(ParticleIndex) == 0)
-				{
-					return;
-				}
+				return;
+			}
 
-				const FReal Radius = SphereRadii[Index] * SphereRadiiMultiplier;
-				const FVec3& Center = AnimationPositions[ParticleIndex];
+			const FSolverReal Radius = SphereRadii[Index] * SphereRadiiMultiplier;
+			const FSolverVec3& Center = AnimationPositions[ParticleIndex];
 
-				const FVec3 CenterToParticle = Particles.P(ParticleIndex) - Center;
-				const FReal DistanceSquared = CenterToParticle.SizeSquared();
+			const FSolverVec3 CenterToParticle = Particles.P(ParticleIndex) - Center;
+			const FSolverReal DistanceSquared = CenterToParticle.SizeSquared();
 
-				static const FReal DeadZoneSquareRadius = SMALL_NUMBER; // We will not push the particle away in the dead zone
-				if (DistanceSquared > FMath::Square(Radius) + DeadZoneSquareRadius)
-				{
-					const FReal Distance = sqrt(DistanceSquared);
-					const FVec3 PositionOnSphere = (Radius / Distance) * CenterToParticle;
-					Particles.P(ParticleIndex) = Center + PositionOnSphere;
-				}
-			});
-		}
+			static const FSolverReal DeadZoneSquareRadius = SMALL_NUMBER; // We will not push the particle away in the dead zone
+			if (DistanceSquared > FMath::Square(Radius) + DeadZoneSquareRadius)
+			{
+				const FSolverReal Distance = sqrt(DistanceSquared);
+				const FSolverVec3 PositionOnSphere = (Radius / Distance) * CenterToParticle;
+				Particles.P(ParticleIndex) = Center + PositionOnSphere;
+			}
+		});
+	}
 
-		void ApplyHelperISPC(FPBDParticles& Particles, const FReal Dt) const;
+	void ApplyHelperISPC(FSolverParticles& Particles, const FSolverReal Dt) const;
 
-	protected:
-		const TArray<FVec3>& AnimationPositions;  // Use global indexation (will need adding ParticleOffset)
-		const TConstArrayView<FRealSingle> SphereRadii;  // Use local indexation
-		const int32 ParticleOffset;
-		FReal SphereRadiiMultiplier;
-	};
+protected:
+	const TArray<FSolverVec3>& AnimationPositions;  // Use global indexation (will need adding ParticleOffset)
+	const TConstArrayView<FRealSingle> SphereRadii;  // Use local indexation
+	const int32 ParticleOffset;
+	FSolverReal SphereRadiiMultiplier;
+};
 
-	class CHAOS_API FPBDSphericalBackstopConstraint final
+class CHAOS_API FPBDSphericalBackstopConstraint final
+{
+public:
+	FPBDSphericalBackstopConstraint(
+		const int32 InParticleOffset,
+		const int32 InParticleCount,
+		const TArray<FSolverVec3>& InAnimationPositions,  // Use global indexation (will need adding ParticleOffset)
+		const TArray<FSolverVec3>& InAnimationNormals,  // Use global indexation (will need adding ParticleOffset)
+		const TConstArrayView<FRealSingle>& InSphereRadii,  // Use local indexation
+		const TConstArrayView<FRealSingle>& InSphereOffsetDistances,  // Use local indexation
+		const bool bInUseLegacyBackstop  // Do not include the sphere radius in the distance calculations when this is true
+	)
+		: AnimationPositions(InAnimationPositions)
+		, AnimationNormals(InAnimationNormals)
+		, SphereRadii(InSphereRadii)
+		, SphereOffsetDistances(InSphereOffsetDistances)
+		, ParticleOffset(InParticleOffset)
+		, SphereRadiiMultiplier((FSolverReal)1.)
+		, bEnabled(true)
+		, bUseLegacyBackstop(bInUseLegacyBackstop)
 	{
-	public:
-		FPBDSphericalBackstopConstraint(
-			const int32 InParticleOffset,
-			const int32 InParticleCount,
-			const TArray<FVec3>& InAnimationPositions,  // Use global indexation (will need adding ParticleOffset)
-			const TArray<FVec3>& InAnimationNormals,  // Use global indexation (will need adding ParticleOffset)
-			const TConstArrayView<FRealSingle>& InSphereRadii,  // Use local indexation
-			const TConstArrayView<FRealSingle>& InSphereOffsetDistances,  // Use local indexation
-			const bool bInUseLegacyBackstop  // Do not include the sphere radius in the distance calculations when this is true
-		)
-			: AnimationPositions(InAnimationPositions)
-			, AnimationNormals(InAnimationNormals)
-			, SphereRadii(InSphereRadii)
-			, SphereOffsetDistances(InSphereOffsetDistances)
-			, ParticleOffset(InParticleOffset)
-			, SphereRadiiMultiplier((FReal)1.)
-			, bEnabled(true)
-			, bUseLegacyBackstop(bInUseLegacyBackstop)
-		{
-			check(InSphereRadii.Num() == InParticleCount);
-			check(InSphereOffsetDistances.Num() == InParticleCount);
-		}
-		~FPBDSphericalBackstopConstraint() {}
+		check(InSphereRadii.Num() == InParticleCount);
+		check(InSphereOffsetDistances.Num() == InParticleCount);
+	}
+	~FPBDSphericalBackstopConstraint() {}
 
-		void SetEnabled(bool bInEnabled) { bEnabled = bInEnabled; }
-		bool IsEnabled() const { return bEnabled; }
+	void SetEnabled(bool bInEnabled) { bEnabled = bInEnabled; }
+	bool IsEnabled() const { return bEnabled; }
 
-		void Apply(FPBDParticles& Particles, const FReal Dt) const
-		{
-			SCOPE_CYCLE_COUNTER(STAT_PBD_SphericalBackstop);
+	void Apply(FSolverParticles& Particles, const FSolverReal Dt) const
+	{
+		SCOPE_CYCLE_COUNTER(STAT_PBD_SphericalBackstop);
 
-			if (bEnabled)
+		if (bEnabled)
+		{
+			if (bUseLegacyBackstop)
+			{
+				// SphereOffsetDistances includes the sphere radius
+				// This is harder to author, and does not follow the NvCloth specs.
+				// However, this is how it's been done in the Unreal Engine PhysX cloth implementation.
+				if (bRealTypeCompatibleWithISPC && bChaos_Spherical_ISPC_Enabled)
+				{
+					ApplyLegacyHelperISPC(Particles, Dt);
+				}
+				else
+				{
+					ApplyLegacyHelper(Particles, Dt);
+				}
+			}
+			else
 			{
-				if (bUseLegacyBackstop)
+				// SphereOffsetDistances doesn't include the sphere radius
+				if (bRealTypeCompatibleWithISPC && bChaos_Spherical_ISPC_Enabled)
 				{
-					// SphereOffsetDistances includes the sphere radius
-					// This is harder to author, and does not follow the NvCloth specs.
-					// However, this is how it's been done in the Unreal Engine PhysX cloth implementation.
-					if (bRealTypeCompatibleWithISPC && bChaos_Spherical_ISPC_Enabled)
-					{
-						ApplyLegacyHelperISPC(Particles, Dt);
-					}
-					else
-					{
-						ApplyLegacyHelper(Particles, Dt);
-					}
+					ApplyHelperISPC(Particles, Dt);
 				}
 				else
 				{
-					// SphereOffsetDistances doesn't include the sphere radius
-					if (bRealTypeCompatibleWithISPC && bChaos_Spherical_ISPC_Enabled)
-					{
-						ApplyHelperISPC(Particles, Dt);
-					}
-					else
-					{
-						ApplyHelper(Particles, Dt);
-					}
+					ApplyHelper(Particles, Dt);
 				}
 			}
 		}
+	}
 
-		void SetSphereRadiiMultiplier(const FReal InSphereRadiiMultiplier)
-		{
-			SphereRadiiMultiplier = FMath::Max((FReal)0., InSphereRadiiMultiplier);
-		}
+	void SetSphereRadiiMultiplier(const FSolverReal InSphereRadiiMultiplier)
+	{
+		SphereRadiiMultiplier = FMath::Max((FSolverReal)0., InSphereRadiiMultiplier);
+	}
 
-		FReal GetSphereRadiiMultiplier() const
-		{
-			return SphereRadiiMultiplier;
-		}
+	FSolverReal GetSphereRadiiMultiplier() const
+	{
+		return SphereRadiiMultiplier;
+	}
 
-		bool UseLegacyBackstop() const
-		{
-			return bUseLegacyBackstop;
-		}
+	bool UseLegacyBackstop() const
+	{
+		return bUseLegacyBackstop;
+	}
 
-	private:
-		void ApplyHelper(FPBDParticles& Particles, const FReal Dt) const
+private:
+	void ApplyHelper(FSolverParticles& Particles, const FSolverReal Dt) const
+	{
+		const int32 ParticleCount = SphereRadii.Num();
+
+		PhysicsParallelFor(ParticleCount, [&](int32 Index)  // TODO: profile need for parallel loop based on particle count
 		{
-			const int32 ParticleCount = SphereRadii.Num();
+			const int32 ParticleIndex = ParticleOffset + Index;
 
-			PhysicsParallelFor(ParticleCount, [&](int32 Index)  // TODO: profile need for parallel loop based on particle count
+			if (Particles.InvM(ParticleIndex) == 0)
 			{
-				const int32 ParticleIndex = ParticleOffset + Index;
+				return;
+			}
 
-				if (Particles.InvM(ParticleIndex) == 0)
-				{
-					return;
-				}
+			const FSolverVec3& AnimationPosition = AnimationPositions[ParticleIndex];
+			const FSolverVec3& AnimationNormal = AnimationNormals[ParticleIndex];
 
-				const FVec3& AnimationPosition = AnimationPositions[ParticleIndex];
-				const FVec3& AnimationNormal = AnimationNormals[ParticleIndex];
+			const FSolverReal SphereOffsetDistance = SphereOffsetDistances[Index];
+			const FSolverReal Radius = SphereRadii[Index] * SphereRadiiMultiplier;
 
-				const FReal SphereOffsetDistance = SphereOffsetDistances[Index];
-				const FReal Radius = SphereRadii[Index] * SphereRadiiMultiplier;
+			const FSolverVec3 Center = AnimationPosition - (Radius + SphereOffsetDistance) * AnimationNormal;  // Non legacy version adds radius to the distance
+			const FSolverVec3 CenterToParticle = Particles.P(ParticleIndex) - Center;
+			const FSolverReal DistanceSquared = CenterToParticle.SizeSquared();
 
-				const FVec3 Center = AnimationPosition - (Radius + SphereOffsetDistance) * AnimationNormal;  // Non legacy version adds radius to the distance
-				const FVec3 CenterToParticle = Particles.P(ParticleIndex) - Center;
-				const FReal DistanceSquared = CenterToParticle.SizeSquared();
+			static const FSolverReal DeadZoneSquareRadius = SMALL_NUMBER;
+			if (DistanceSquared < DeadZoneSquareRadius)
+			{
+				Particles.P(ParticleIndex) = AnimationPosition - SphereOffsetDistance * AnimationNormal;  // Non legacy version adds radius to the distance
+			}
+			else if (DistanceSquared < FMath::Square(Radius))
+			{
+				const FSolverVec3 PositionOnSphere = (Radius / sqrt(DistanceSquared)) * CenterToParticle;
+				Particles.P(ParticleIndex) = Center + PositionOnSphere;
+			}
+			// Else the particle is outside the sphere, and there is nothing to do
+		});
+	}
 
-				static const FReal DeadZoneSquareRadius = SMALL_NUMBER;
-				if (DistanceSquared < DeadZoneSquareRadius)
-				{
-					Particles.P(ParticleIndex) = AnimationPosition - SphereOffsetDistance * AnimationNormal;  // Non legacy version adds radius to the distance
-				}
-				else if (DistanceSquared < FMath::Square(Radius))
-				{
-					const FVec3 PositionOnSphere = (Radius / sqrt(DistanceSquared)) * CenterToParticle;
-					Particles.P(ParticleIndex) = Center + PositionOnSphere;
-				}
-				// Else the particle is outside the sphere, and there is nothing to do
-			});
-		}
+	void ApplyLegacyHelper(FSolverParticles& Particles, const FSolverReal Dt) const
+	{
+		const int32 ParticleCount = SphereRadii.Num();
 
-		void ApplyLegacyHelper(FPBDParticles& Particles, const FReal Dt) const
+		PhysicsParallelFor(ParticleCount, [&](int32 Index)  // TODO: profile need for parallel loop based on particle count
 		{
-			const int32 ParticleCount = SphereRadii.Num();
+			const int32 ParticleIndex = ParticleOffset + Index;
 
-			PhysicsParallelFor(ParticleCount, [&](int32 Index)  // TODO: profile need for parallel loop based on particle count
+			if (Particles.InvM(ParticleIndex) == 0)
 			{
-				const int32 ParticleIndex = ParticleOffset + Index;
-
-				if (Particles.InvM(ParticleIndex) == 0)
-				{
-					return;
-				}
+				return;
+			}
 
-				const FVec3& AnimationPosition = AnimationPositions[ParticleIndex];
-				const FVec3& AnimationNormal = AnimationNormals[ParticleIndex];
+			const FSolverVec3& AnimationPosition = AnimationPositions[ParticleIndex];
+			const FSolverVec3& AnimationNormal = AnimationNormals[ParticleIndex];
 
-				const FReal SphereOffsetDistance = SphereOffsetDistances[Index];
-				const FReal Radius = SphereRadii[Index] * SphereRadiiMultiplier;
+			const FSolverReal SphereOffsetDistance = SphereOffsetDistances[Index];
+			const FSolverReal Radius = SphereRadii[Index] * SphereRadiiMultiplier;
 
-				const FVec3 Center = AnimationPosition - SphereOffsetDistance * AnimationNormal;  // Legacy version already includes the radius within the distance
-				const FVec3 CenterToParticle = Particles.P(ParticleIndex) - Center;
-				const FReal DistanceSquared = CenterToParticle.SizeSquared();
+			const FSolverVec3 Center = AnimationPosition - SphereOffsetDistance * AnimationNormal;  // Legacy version already includes the radius within the distance
+			const FSolverVec3 CenterToParticle = Particles.P(ParticleIndex) - Center;
+			const FSolverReal DistanceSquared = CenterToParticle.SizeSquared();
 
-				static const FReal DeadZoneSquareRadius = SMALL_NUMBER;
-				if (DistanceSquared < DeadZoneSquareRadius)
-				{
-					Particles.P(ParticleIndex) = AnimationPosition - (SphereOffsetDistance - Radius) * AnimationNormal;  // Legacy version already includes the radius to the distance
-				}
-				else if (DistanceSquared < FMath::Square(Radius))
-				{
-					const FVec3 PositionOnSphere = (Radius / sqrt(DistanceSquared)) * CenterToParticle;
-					Particles.P(ParticleIndex) = Center + PositionOnSphere;
-				}
-				// Else the particle is outside the sphere, and there is nothing to do
-			});
-		}
-
-		void ApplyLegacyHelperISPC(FPBDParticles& Particles, const FReal Dt) const;
-		void ApplyHelperISPC(FPBDParticles& Particles, const FReal Dt) const;
-
-	private:
-		const TArray<FVec3>& AnimationPositions;  // Positions of spheres, use global indexation (will need adding ParticleOffset)
-		const TArray<FVec3>& AnimationNormals; // Sphere offset directions, use global indexation (will need adding ParticleOffset)
-		const TConstArrayView<FRealSingle> SphereRadii; // Start at index 0, use local indexation
-		const TConstArrayView<FRealSingle> SphereOffsetDistances;  // Sphere position offsets, use local indexation
-		const int32 ParticleOffset;
-		FReal SphereRadiiMultiplier;
-		bool bEnabled;
-		bool bUseLegacyBackstop;
-	};
-
-	template<typename T, int d>
-	using TPBDSphericalConstraint UE_DEPRECATED(4.27, "Deprecated. this class is to be deleted, use FPBDSphericalConstraint instead") = FPBDSphericalConstraint;
-
-	template<typename T, int d>
-	using TPBDSphericalBackstopConstraint UE_DEPRECATED(4.27, "Deprecated. this class is to be deleted, use FPBDSphericalBackstopConstraint instead") = FPBDSphericalBackstopConstraint;
-}
+			static const FSolverReal DeadZoneSquareRadius = SMALL_NUMBER;
+			if (DistanceSquared < DeadZoneSquareRadius)
+			{
+				Particles.P(ParticleIndex) = AnimationPosition - (SphereOffsetDistance - Radius) * AnimationNormal;  // Legacy version already includes the radius to the distance
+			}
+			else if (DistanceSquared < FMath::Square(Radius))
+			{
+				const FSolverVec3 PositionOnSphere = (Radius / sqrt(DistanceSquared)) * CenterToParticle;
+				Particles.P(ParticleIndex) = Center + PositionOnSphere;
+			}
+			// Else the particle is outside the sphere, and there is nothing to do
+		});
+	}
+
+	void ApplyLegacyHelperISPC(FSolverParticles& Particles, const FSolverReal Dt) const;
+	void ApplyHelperISPC(FSolverParticles& Particles, const FSolverReal Dt) const;
+
+private:
+	const TArray<FSolverVec3>& AnimationPositions;  // Positions of spheres, use global indexation (will need adding ParticleOffset)
+	const TArray<FSolverVec3>& AnimationNormals; // Sphere offset directions, use global indexation (will need adding ParticleOffset)
+	const TConstArrayView<FRealSingle> SphereRadii; // Start at index 0, use local indexation
+	const TConstArrayView<FRealSingle> SphereOffsetDistances;  // Sphere position offsets, use local indexation
+	const int32 ParticleOffset;
+	FSolverReal SphereRadiiMultiplier;
+	bool bEnabled;
+	bool bUseLegacyBackstop;
+};
+
+}  // End namespace Chaos::Softs
diff --git a/Engine/Source/Runtime/Experimental/Chaos/Public/Chaos/PBDSpringConstraints.h b/Engine/Source/Runtime/Experimental/Chaos/Public/Chaos/PBDSpringConstraints.h
--- a/Engine/Source/Runtime/Experimental/Chaos/Public/Chaos/PBDSpringConstraints.h
+++ b/Engine/Source/Runtime/Experimental/Chaos/Public/Chaos/PBDSpringConstraints.h
@@ -3,7 +3,7 @@
 
 #include "Chaos/PBDSpringConstraintsBase.h"
 
-namespace Chaos
+namespace Chaos::Softs
 {
 
 class CHAOS_API FPBDSpringConstraints : public FPBDSpringConstraintsBase
@@ -15,12 +15,12 @@ class CHAOS_API FPBDSpringConstraints : public FPBDSpringConstraintsBase
 public:
 	template<int32 Valence>
 	FPBDSpringConstraints(
-		const FPBDParticles& Particles,
+		const FSolverParticles& Particles,
 		int32 ParticleOffset,
 		int32 ParticleCount,
 		const TArray<TVector<int32, Valence>>& InConstraints,
 		const TConstArrayView<FRealSingle>& StiffnessMultipliers,
-		const FVec2& InStiffness,
+		const FSolverVec2& InStiffness,
 		bool bTrimKinematicConstraints = false,
 		typename TEnableIf<Valence >= 2 && Valence <= 4>::Type* = nullptr)
 		: Base(Particles, ParticleOffset, ParticleCount, InConstraints, StiffnessMultipliers, InStiffness, bTrimKinematicConstraints)
@@ -28,19 +28,19 @@ public:
 		InitColor(Particles);
 	}
 
-	virtual ~FPBDSpringConstraints() {}
+	virtual ~FPBDSpringConstraints() override {}
 
-	void Apply(FPBDParticles& Particles, const FReal Dt) const;
+	void Apply(FSolverParticles& Particles, const FSolverReal Dt) const;
 
 private:
-	void InitColor(const FPBDParticles& InParticles);
-	void ApplyHelper(FPBDParticles& Particles, const FReal Dt, const int32 ConstraintIndex, const FReal ExpStiffnessValue) const;
+	void InitColor(const FSolverParticles& InParticles);
+	void ApplyHelper(FSolverParticles& Particles, const FSolverReal Dt, const int32 ConstraintIndex, const FSolverReal ExpStiffnessValue) const;
 
 private:
 	TArray<TArray<int32>> ConstraintsPerColor;
 };
 
-}
+}  // End namespace Chaos::Softs
 
 // Support ISPC enable/disable in non-shipping builds
 #if !INTEL_ISPC
diff --git a/Engine/Source/Runtime/Experimental/Chaos/Public/Chaos/PBDSpringConstraintsBase.h b/Engine/Source/Runtime/Experimental/Chaos/Public/Chaos/PBDSpringConstraintsBase.h
--- a/Engine/Source/Runtime/Experimental/Chaos/Public/Chaos/PBDSpringConstraintsBase.h
+++ b/Engine/Source/Runtime/Experimental/Chaos/Public/Chaos/PBDSpringConstraintsBase.h
@@ -2,13 +2,14 @@
 #pragma once
 
 #include "Chaos/Core.h"
+#include "Chaos/PBDSoftsEvolutionFwd.h"
 #include "Chaos/Utilities.h"
 #include "Chaos/PBDParticles.h"
 #include "Chaos/PBDStiffness.h"
 #include "Containers/Array.h"
 #include "Templates/EnableIf.h"
 
-namespace Chaos
+namespace Chaos::Softs
 {
 
 class FPBDSpringConstraintsBase
@@ -16,18 +17,18 @@ class FPBDSpringConstraintsBase
 public:
 	template<int32 Valence>
 	FPBDSpringConstraintsBase(
-		const FPBDParticles& Particles,
+		const FSolverParticles& Particles,
 		int32 ParticleOffset,
 		int32 ParticleCount,
 		const TArray<TVector<int32, Valence>>& InConstraints,
 		const TConstArrayView<FRealSingle>& StiffnessMultipliers,
-		const FVec2& InStiffness,
+		const FSolverVec2& InStiffness,
 		bool bTrimKinematicConstraints = false,
 		typename TEnableIf<Valence >= 2 && Valence <= 4>::Type* = nullptr)
 		: Constraints(TrimConstraints(InConstraints, 
 			[&Particles, bTrimKinematicConstraints](int32 Index0, int32 Index1)
 			{
-				return bTrimKinematicConstraints && Particles.InvM(Index0) == (FReal)0. && Particles.InvM(Index1) == (FReal)0.;
+				return bTrimKinematicConstraints && Particles.InvM(Index0) == (FSolverReal)0. && Particles.InvM(Index1) == (FSolverReal)0.;
 			}))
 		, Stiffness(InStiffness, StiffnessMultipliers, TConstArrayView<TVec2<int32>>(Constraints), ParticleOffset, ParticleCount)
 	{
@@ -35,8 +36,8 @@ public:
 		Dists.Reset(Constraints.Num());
 		for (const TVec2<int32>& Constraint : Constraints)
 		{
-			const FVec3& P0 = Particles.X(Constraint[0]);
-			const FVec3& P1 = Particles.X(Constraint[1]);
+			const FSolverVec3& P0 = Particles.X(Constraint[0]);
+			const FSolverVec3& P1 = Particles.X(Constraint[1]);
 			Dists.Add((P1 - P0).Size());
 		}
 	}
@@ -45,32 +46,32 @@ public:
 	{}
 
 	// Update stiffness values
-	void SetProperties(const FVec2& InStiffness) { Stiffness.SetWeightedValue(InStiffness.ClampAxes((FReal)0., (FReal)1.)); }
+	void SetProperties(const FSolverVec2& InStiffness) { Stiffness.SetWeightedValue(InStiffness.ClampAxes((FSolverReal)0., (FSolverReal)1.)); }
 
 	// Update stiffness table, as well as the simulation stiffness exponent
-	inline void ApplyProperties(const FReal Dt, const int32 NumIterations) { Stiffness.ApplyValues(Dt, NumIterations); }
+	inline void ApplyProperties(const FSolverReal Dt, const int32 NumIterations) { Stiffness.ApplyValues(Dt, NumIterations); }
 
 	const TArray<TVec2<int32>>& GetConstraints() const { return Constraints; }
 
 protected:
-	FVec3 GetDelta(const FPBDParticles& Particles, const int32 ConstraintIndex, const FReal ExpStiffnessValue) const
+	FSolverVec3 GetDelta(const FSolverParticles& Particles, const int32 ConstraintIndex, const FSolverReal ExpStiffnessValue) const
 	{
 		const auto& Constraint = Constraints[ConstraintIndex];
 		const int32 i1 = Constraint[0];
 		const int32 i2 = Constraint[1];
 
-		if (Particles.InvM(i2) == (FReal)0. && Particles.InvM(i1) == (FReal)0.)
+		if (Particles.InvM(i2) == (FSolverReal)0. && Particles.InvM(i1) == (FSolverReal)0.)
 		{
-			return FVec3((FReal)0.);
+			return FSolverVec3((FSolverReal)0.);
 		}
-		const FReal CombinedMass = Particles.InvM(i2) + Particles.InvM(i1);
+		const FSolverReal CombinedMass = Particles.InvM(i2) + Particles.InvM(i1);
 
-		const FVec3& P1 = Particles.P(i1);
-		const FVec3& P2 = Particles.P(i2);
-		FVec3 Direction = P1 - P2;
-		const FReal Distance = Direction.SafeNormalize();
+		const FSolverVec3& P1 = Particles.P(i1);
+		const FSolverVec3& P2 = Particles.P(i2);
+		FSolverVec3 Direction = P1 - P2;
+		const FSolverReal Distance = Direction.SafeNormalize();
 
-		const FVec3 Delta = (Distance - Dists[ConstraintIndex]) * Direction;
+		const FSolverVec3 Delta = (Distance - Dists[ConstraintIndex]) * Direction;
 		return ExpStiffnessValue * Delta / CombinedMass;
 	}
 
@@ -103,8 +104,8 @@ private:
 
 protected:
 	TArray<TVec2<int32>> Constraints;
-	TArray<FReal> Dists;
+	TArray<FSolverReal> Dists;
 	FPBDStiffness Stiffness;
 };
 
-}
+}  // End namespace Chaos::Softs
diff --git a/Engine/Source/Runtime/Experimental/Chaos/Public/Chaos/PBDStiffness.h b/Engine/Source/Runtime/Experimental/Chaos/Public/Chaos/PBDStiffness.h
--- a/Engine/Source/Runtime/Experimental/Chaos/Public/Chaos/PBDStiffness.h
+++ b/Engine/Source/Runtime/Experimental/Chaos/Public/Chaos/PBDStiffness.h
@@ -2,14 +2,14 @@
 #pragma once
 
 #include "Chaos/Core.h"
-#include "Chaos/Framework/Parallel.h"
+#include "Chaos/PBDSoftsEvolutionFwd.h"
 #include "Containers/ArrayView.h"
 #include "Containers/ContainersFwd.h"
 #include "ChaosStats.h"
 
 DECLARE_CYCLE_STAT(TEXT("Chaos PBD Stiffness Apply Values"), STAT_PBD_StiffnessApplyValues, STATGROUP_Chaos);
 
-namespace Chaos
+namespace Chaos::Softs
 {
 
 /**
@@ -19,30 +19,30 @@ namespace Chaos
 class FPBDStiffness final
 {
 public:
-	static constexpr FReal ParameterFrequency = (FReal)120.;  // 60Hz @ 2 iterations as a root for all stiffness values TODO: Make this a global solver parameter
+	static constexpr FSolverReal ParameterFrequency = (FSolverReal)120.;  // 60Hz @ 2 iterations as a root for all stiffness values TODO: Make this a global solver parameter
 
 	/**
 	 * Weightmap particle constructor. 
 	 */
 	inline FPBDStiffness(
-		const FVec2& InWeightedValue,
+		const FSolverVec2& InWeightedValue,
 		const TConstArrayView<FRealSingle>& Multipliers = TConstArrayView<FRealSingle>(),
 		int32 ParticleCount = 0,  // A value of 0 also disables the map
 		int32 TableSize = 16,  // Size of the lookup table, can't be more than 256 values, the larger the table the longer it takes to apply changes to the stiffness values
-		FReal InParameterFitBase = (FReal)1.e3);  // Logarithm base to use in the stiffness parameter fit function
+		FSolverReal InParameterFitBase = (FSolverReal)1.e3);  // Logarithm base to use in the stiffness parameter fit function
 
 	/**
 	 * Weightmap constraint constructor. 
 	 */
 	template<int32 Valence>
 	inline FPBDStiffness(
-		const FVec2& InWeightedValue,
+		const FSolverVec2& InWeightedValue,
 		const TConstArrayView<FRealSingle>& Multipliers = TConstArrayView<FRealSingle>(),
 		const TConstArrayView<TVector<int32, Valence>>& Constraints = TConstArrayView<TVector<int32, Valence>>(),
 		int32 ParticleOffset = INDEX_NONE, // Constraints have usually a particle offset added to them compared to the weight maps that always starts at index 0
 		int32 ParticleCount = 0,  // A value of 0 also disables the map
 		int32 TableSize = 16,  // Size of the lookup table, can't be more than 256 values, the larger the table the longer it takes to apply changes to the stiffness values
-		FReal InParameterFitBase = (FReal)1.e3,  // Logarithm base to use in the stiffness parameter fit function
+		FSolverReal InParameterFitBase = (FSolverReal)1.e3,  // Logarithm base to use in the stiffness parameter fit function
 		typename TEnableIf<Valence >= 2 && Valence <= 4>::Type* = nullptr);  // Prevents incorrect valence, the value is actually unused
 
 	~FPBDStiffness() {}
@@ -57,52 +57,52 @@ public:
 	 * Set the low and high values of the weight map.
 	 * The weight map table only gets updated after ApplyValues is called.
 	 */
-	void SetWeightedValue(const FVec2& InWeightedValue) { WeightedValue = InWeightedValue.ClampAxes((FReal)0., (FReal)1.); }
+	void SetWeightedValue(const FSolverVec2& InWeightedValue) { WeightedValue = InWeightedValue.ClampAxes((FSolverReal)0., (FSolverReal)1.); }
 
 	/**
 	 * Return the low and high values set for this weight map.
 	 * Both values will always be between 0 and 1 due to having been clamped in SetWeightedValue.
 	 */
-	const FVec2& GetWeightedValue() const { return WeightedValue; }
+	const FSolverVec2& GetWeightedValue() const { return WeightedValue; }
 
 	/** Update the weight map table with the current simulation parameters. */
-	inline void ApplyValues(const FReal Dt, const int32 NumIterations);
+	inline void ApplyValues(const FSolverReal Dt, const int32 NumIterations);
 
 	/**
 	 * Lookup for the exponential weighted value at the specified weight map index.
 	 * This function will assert if it is called with a non zero index on an empty weight map.
 	*/
-	FReal operator[](int32 Index) const { return Table[Indices[Index]]; }
+	FSolverReal operator[](int32 Index) const { return Table[Indices[Index]]; }
 
 	/** Return the exponential value at the Low weight. */
-	FReal GetLow() const { return Table[0]; }
+	FSolverReal GetLow() const { return Table[0]; }
 
 	/** Return the exponential value at the High weight. */
-	FReal GetHigh() const { return Table.Last(); }
+	FSolverReal GetHigh() const { return Table.Last(); }
 
 	/** Return the exponential stiffness value when the weight map is not used. */
-	explicit operator FReal() const { return GetLow(); }
+	explicit operator FSolverReal() const { return GetLow(); }
 
 	/** Return the particles/constraints indices to the stiffness table as a read only array. */
 	TConstArrayView<uint8> GetIndices() const { return TConstArrayView<uint8>(Indices); }
 
 	/** Return the table of stiffnesses as a read only array. */
-	TConstArrayView<FReal> GetTable() const { return TConstArrayView<FReal>(Table); }
+	TConstArrayView<FSolverReal> GetTable() const { return TConstArrayView<FSolverReal>(Table); }
 
 private:
 	TArray<uint8> Indices; // Per particle/constraints array of index to the stiffness table
-	TArray<FReal> Table;  // Fixed lookup table of stiffness values, use uint8 indexation
-	FVec2 WeightedValue;
-	const FReal ParameterFitLogBase;
+	TArray<FSolverReal> Table;  // Fixed lookup table of stiffness values, use uint8 indexation
+	FSolverVec2 WeightedValue;
+	const FSolverReal ParameterFitLogBase;
 };
 
 FPBDStiffness::FPBDStiffness(
-	const FVec2& InWeightedValue,
+	const FSolverVec2& InWeightedValue,
 	const TConstArrayView<FRealSingle>& Multipliers,
 	int32 ParticleCount,
 	int32 TableSize,
-	FReal InParameterFitBase)
-	: WeightedValue(InWeightedValue.ClampAxes((FReal)0., (FReal)1.))
+	FSolverReal InParameterFitBase)
+	: WeightedValue(InWeightedValue.ClampAxes((FSolverReal)0., (FSolverReal)1.))
 	, ParameterFitLogBase(FMath::Loge(InParameterFitBase))
 {
 	check(TableSize > 0 && TableSize < 256);  // The Stiffness lookup table is restricted to uint8 sized indices
@@ -132,15 +132,15 @@ FPBDStiffness::FPBDStiffness(
 
 template<int32 Valence>
 FPBDStiffness::FPBDStiffness(
-	const FVec2& InWeightedValue,
+	const FSolverVec2& InWeightedValue,
 	const TConstArrayView<FRealSingle>& Multipliers,
 	const TConstArrayView<TVector<int32, Valence>>& Constraints,
 	int32 ParticleOffset,
 	int32 ParticleCount,
 	int32 TableSize,
-	FReal InParameterFitBase,
+	FSolverReal InParameterFitBase,
 	typename TEnableIf<Valence >= 2 && Valence <= 4>::Type*)
-	: WeightedValue(InWeightedValue.ClampAxes((FReal)0., (FReal)1.))
+	: WeightedValue(InWeightedValue.ClampAxes((FSolverReal)0., (FSolverReal)1.))
 	, ParameterFitLogBase(FMath::Loge(InParameterFitBase))
 {
 	check(TableSize > 0 && TableSize < 256);  // The Stiffness lookup table is restricted to uint8 sized indices
@@ -180,37 +180,37 @@ FPBDStiffness::FPBDStiffness(
 }
 
 
-void FPBDStiffness::ApplyValues(const FReal Dt, const int32 NumIterations)
+void FPBDStiffness::ApplyValues(const FSolverReal Dt, const int32 NumIterations)
 {
 	SCOPE_CYCLE_COUNTER(STAT_PBD_StiffnessApplyValues);
 
 	// Calculate the simulation exponent
-	const FReal Exponent = Dt * ParameterFrequency / (FReal)NumIterations;
+	const FSolverReal Exponent = Dt * ParameterFrequency / (FSolverReal)NumIterations;
 
 	// Define the stiffness mapping function
-	auto SimulationValue = [this, Exponent](const FReal InValue)->FReal
+	auto SimulationValue = [this, Exponent](const FSolverReal InValue)->FSolverReal
 	{
 		// Get a very steep exponential curve between the [0, 1] range to make easier to set the parameter
 		// The base has been chosen empirically
 		// ParameterValue = Pow(ParameterFitBase, ParameterValue - 1)
-		const FReal ParameterFit = FMath::Exp(ParameterFitLogBase * (FMath::Clamp(InValue, (FReal)0., (FReal)1.) - (FReal)1.));
+		const FSolverReal ParameterFit = FMath::Exp(ParameterFitLogBase * (FMath::Clamp(InValue, (FSolverReal)0., (FSolverReal)1.) - (FSolverReal)1.));
 
 		// Use simulation dependent stiffness exponent to alleviate the variations in effect when Dt and NumIterations change
 		// This is based on the Position-Based Simulation Methods paper (page 8),
 		// but uses the delta time in addition of the number of iterations in the calculation of the error term.
-		const FReal LogValue = FMath::Loge((FReal)1. - ParameterFit);
-		return (FReal)1. - FMath::Exp(LogValue * Exponent);
+		const FSolverReal LogValue = FMath::Loge((FSolverReal)1. - ParameterFit);
+		return (FSolverReal)1. - FMath::Exp(LogValue * Exponent);
 	};
 
-	const FReal Offset = WeightedValue[0];
-	const FReal Range = WeightedValue[1] - WeightedValue[0];
+	const FSolverReal Offset = WeightedValue[0];
+	const FSolverReal Range = WeightedValue[1] - WeightedValue[0];
 	const int32 TableSize = Table.Num();
-	const FReal WeightIncrement = (TableSize > 1) ? (FReal)1. / (FReal)(TableSize - 1) : (FReal)1.; // Must allow full range from 0 to 1 included
+	const FSolverReal WeightIncrement = (TableSize > 1) ? (FSolverReal)1. / (FSolverReal)(TableSize - 1) : (FSolverReal)1.; // Must allow full range from 0 to 1 included
 	for (int32 Index = 0; Index < TableSize; ++Index)
 	{
-		const FReal Weight = (FReal)Index * WeightIncrement;
+		const FSolverReal Weight = (FSolverReal)Index * WeightIncrement;
 		Table[Index] = SimulationValue(Offset + Weight * Range);
 	}
 }
 
-}
+}  // End namespace Chaos::Softs
diff --git a/Engine/Source/Runtime/Experimental/Chaos/Public/Chaos/PBDTetConstraints.h b/Engine/Source/Runtime/Experimental/Chaos/Public/Chaos/PBDTetConstraints.h
--- a/Engine/Source/Runtime/Experimental/Chaos/Public/Chaos/PBDTetConstraints.h
+++ b/Engine/Source/Runtime/Experimental/Chaos/Public/Chaos/PBDTetConstraints.h
@@ -1,34 +1,32 @@
 // Copyright Epic Games, Inc. All Rights Reserved.
 #pragma once
 
-#include "Chaos/Array.h"
-#include "Chaos/PBDParticles.h"
 #include "Chaos/PBDTetConstraintsBase.h"
-#include "Chaos/ParticleRule.h"
 
-namespace Chaos
+namespace Chaos::Softs
 {
-class FPBDTetConstraints : public FParticleRule, public FPBDTetConstraintsBase
+
+class FPBDTetConstraints : public FPBDTetConstraintsBase
 {
 	typedef FPBDTetConstraintsBase Base;
-	using Base::MConstraints;
+	using Base::Constraints;
 
   public:
-	  FPBDTetConstraints(const FDynamicParticles& InParticles, TArray<TVec4<int32>>&& Constraints, const FReal Stiffness = (FReal)1)
-	    : Base(InParticles, MoveTemp(Constraints), Stiffness) {}
-	virtual ~FPBDTetConstraints() {}
+	  FPBDTetConstraints(const FSolverParticles& InParticles, TArray<TVec4<int32>>&& InConstraints, const FReal InStiffness = (FReal)1)
+	    : Base(InParticles, MoveTemp(InConstraints), InStiffness) {}
+	virtual ~FPBDTetConstraints() override {}
 
-	void Apply(FPBDParticles& InParticles, const FReal dt) const override //-V762
+	void Apply(FSolverParticles& InParticles, const FReal dt) const
 	{
-		for (int i = 0; i < MConstraints.Num(); ++i)
+		for (int i = 0; i < Constraints.Num(); ++i)
 		{
-			const auto& Constraint = MConstraints[i];
+			const TVec4<int32>& Constraint = Constraints[i];
 			const int32 i1 = Constraint[0];
 			const int32 i2 = Constraint[1];
 			const int32 i3 = Constraint[2];
 			const int32 i4 = Constraint[3];
-			auto Grads = Base::GetGradients(InParticles, i);
-			auto S = Base::GetScalingFactor(InParticles, i, Grads);
+			const TVec4<FSolverVec3> Grads = Base::GetGradients(InParticles, i);
+			const FSolverReal S = Base::GetScalingFactor(InParticles, i, Grads);
 			InParticles.P(i1) -= S * InParticles.InvM(i1) * Grads[0];
 			InParticles.P(i2) -= S * InParticles.InvM(i2) * Grads[1];
 			InParticles.P(i3) -= S * InParticles.InvM(i3) * Grads[2];
@@ -37,7 +35,4 @@ class FPBDTetConstraints : public FParticleRule, public FPBDTetConstraintsBase
 	}
 };
 
-template<class T>
-using PBDTetConstraints UE_DEPRECATED(4.27, "Deprecated. this class is to be deleted, use FPBDTetConstraints instead") = FPBDTetConstraints;
-
-}
+}  // End namespace Chaos::Softs
diff --git a/Engine/Source/Runtime/Experimental/Chaos/Public/Chaos/PBDTetConstraintsBase.h b/Engine/Source/Runtime/Experimental/Chaos/Public/Chaos/PBDTetConstraintsBase.h
--- a/Engine/Source/Runtime/Experimental/Chaos/Public/Chaos/PBDTetConstraintsBase.h
+++ b/Engine/Source/Runtime/Experimental/Chaos/Public/Chaos/PBDTetConstraintsBase.h
@@ -2,76 +2,75 @@
 #pragma once
 
 #include "Chaos/Array.h"
+#include "Chaos/PBDSoftsEvolutionFwd.h"
 #include "Chaos/PBDParticles.h"
-#include "Chaos/ParticleRule.h"
 #include "Chaos/DynamicParticles.h"
 #include "Chaos/PBDParticles.h"
 
-namespace Chaos
+namespace Chaos::Softs
 {
+
 class FPBDTetConstraintsBase
 {
-  public:
-	FPBDTetConstraintsBase(const FDynamicParticles& InParticles, TArray<TVec4<int32>>&& Constraints, const FReal Stiffness = (FReal)1.)
-	    : MConstraints(Constraints), MStiffness(Stiffness)
+public:
+	FPBDTetConstraintsBase(const FSolverParticles& InParticles, TArray<TVec4<int32>>&& InConstraints, const FSolverReal InStiffness = (FSolverReal)1.)
+	    : Constraints(InConstraints), Stiffness(InStiffness)
 	{
-		for (auto Constraint : MConstraints)
+		for (const TVec4<int32>& Constraint : Constraints)
 		{
-			const FVec3& P1 = InParticles.X(Constraint[0]);
-			const FVec3& P2 = InParticles.X(Constraint[1]);
-			const FVec3& P3 = InParticles.X(Constraint[2]);
-			const FVec3& P4 = InParticles.X(Constraint[3]);
-			MVolumes.Add(FVec3::DotProduct(FVec3::CrossProduct(P2 - P1, P3 - P1), P4 - P1) / (FReal)6.);
+			const FSolverVec3& P1 = InParticles.X(Constraint[0]);
+			const FSolverVec3& P2 = InParticles.X(Constraint[1]);
+			const FSolverVec3& P3 = InParticles.X(Constraint[2]);
+			const FSolverVec3& P4 = InParticles.X(Constraint[3]);
+			Volumes.Add(FSolverVec3::DotProduct(FSolverVec3::CrossProduct(P2 - P1, P3 - P1), P4 - P1) / (FSolverReal)6.);
 		}
 	}
 	virtual ~FPBDTetConstraintsBase() {}
 
-	TVec4<FVec3> GetGradients(const FPBDParticles& InParticles, const int32 i) const
+	TVec4<FSolverVec3> GetGradients(const FSolverParticles& InParticles, const int32 i) const
 	{
-		TVec4<FVec3> Grads;
-		const auto& Constraint = MConstraints[i];
-		const FVec3& P1 = InParticles.P(Constraint[0]);
-		const FVec3& P2 = InParticles.P(Constraint[1]);
-		const FVec3& P3 = InParticles.P(Constraint[2]);
-		const FVec3& P4 = InParticles.P(Constraint[3]);
-		const FVec3 P2P1 = P2 - P1;
-		const FVec3 P3P1 = P3 - P1;
-		const FVec3 P4P1 = P4 - P1;
-		Grads[1] = FVec3::CrossProduct(P3P1, P4P1) / (FReal)6.;
-		Grads[2] = FVec3::CrossProduct(P4P1, P2P1) / (FReal)6.;
-		Grads[3] = FVec3::CrossProduct(P2P1, P3P1) / (FReal)6.;
-		Grads[0] = -1 * (Grads[1] + Grads[2] + Grads[3]);
+		TVec4<FSolverVec3> Grads;
+		const TVec4<int32>& Constraint = Constraints[i];
+		const FSolverVec3& P1 = InParticles.P(Constraint[0]);
+		const FSolverVec3& P2 = InParticles.P(Constraint[1]);
+		const FSolverVec3& P3 = InParticles.P(Constraint[2]);
+		const FSolverVec3& P4 = InParticles.P(Constraint[3]);
+		const FSolverVec3 P2P1 = P2 - P1;
+		const FSolverVec3 P4P1 = P4 - P1;
+		const FSolverVec3 P3P1 = P3 - P1;
+		Grads[1] = FSolverVec3::CrossProduct(P3P1, P4P1) / (FSolverReal)6.;
+		Grads[2] = FSolverVec3::CrossProduct(P4P1, P2P1) / (FSolverReal)6.;
+		Grads[3] = FSolverVec3::CrossProduct(P2P1, P3P1) / (FSolverReal)6.;
+		Grads[0] = -(Grads[1] + Grads[2] + Grads[3]);
 		return Grads;
 	}
 
-	FReal GetScalingFactor(const FPBDParticles& InParticles, const int32 i, const TVec4<FVec3>& Grads) const
+	FSolverReal GetScalingFactor(const FSolverParticles& InParticles, const int32 i, const TVec4<FSolverVec3>& Grads) const
 	{
-		const auto& Constraint = MConstraints[i];
+		const TVec4<int32>& Constraint = Constraints[i];
 		const int32 i1 = Constraint[0];
 		const int32 i2 = Constraint[1];
 		const int32 i3 = Constraint[2];
 		const int32 i4 = Constraint[3];
-		const FVec3& P1 = InParticles.P(i1);
-		const FVec3& P2 = InParticles.P(i2);
-		const FVec3& P3 = InParticles.P(i3);
-		const FVec3& P4 = InParticles.P(i4);
-		FReal Volume = FVec3::DotProduct(FVec3::CrossProduct(P2 - P1, P3 - P1), P4 - P1) / (FReal)6.;
-		FReal S = (Volume - MVolumes[i]) / (InParticles.InvM(i1) * Grads[0].SizeSquared() +
-			                            InParticles.InvM(i2) * Grads[1].SizeSquared() + 
-			                            InParticles.InvM(i3) * Grads[2].SizeSquared() + 
-			                            InParticles.InvM(i4) * Grads[3].SizeSquared());
-		return MStiffness * S;
+		const FSolverVec3& P1 = InParticles.P(i1);
+		const FSolverVec3& P2 = InParticles.P(i2);
+		const FSolverVec3& P3 = InParticles.P(i3);
+		const FSolverVec3& P4 = InParticles.P(i4);
+		const FSolverReal Volume = FSolverVec3::DotProduct(FSolverVec3::CrossProduct(P2 - P1, P3 - P1), P4 - P1) / (FSolverReal)6.;
+		const FSolverReal S = (Volume - Volumes[i]) / (
+			InParticles.InvM(i1) * Grads[0].SizeSquared() +
+			InParticles.InvM(i2) * Grads[1].SizeSquared() + 
+			InParticles.InvM(i3) * Grads[2].SizeSquared() + 
+			InParticles.InvM(i4) * Grads[3].SizeSquared());
+		return Stiffness * S;
 	}
 
-  protected:
-	TArray<TVec4<int32>> MConstraints;
+protected:
+	TArray<TVec4<int32>> Constraints;
 
-  private:
-	TArray<FReal> MVolumes;
-	FReal MStiffness;
+private:
+	TArray<FSolverReal> Volumes;
+	FSolverReal Stiffness;
 };
 
-template<class T>
-using PBDTetConstraintsBase UE_DEPRECATED(4.27, "Deprecated. this class is to be deleted, use FPBDTetConstraintsBase instead") = FPBDTetConstraintsBase;
-
-}
+}  // End namespace Chaos::Softs
diff --git a/Engine/Source/Runtime/Experimental/Chaos/Public/Chaos/PBDVolumeConstraint.h b/Engine/Source/Runtime/Experimental/Chaos/Public/Chaos/PBDVolumeConstraint.h
--- a/Engine/Source/Runtime/Experimental/Chaos/Public/Chaos/PBDVolumeConstraint.h
+++ b/Engine/Source/Runtime/Experimental/Chaos/Public/Chaos/PBDVolumeConstraint.h
@@ -1,27 +1,25 @@
 // Copyright Epic Games, Inc. All Rights Reserved.
 #pragma once
 
-#include "Chaos/Array.h"
-#include "Chaos/PBDParticles.h"
 #include "Chaos/PBDVolumeConstraintBase.h"
-#include "Chaos/ParticleRule.h"
 
-namespace Chaos
+namespace Chaos::Softs
 {
-class FPBDVolumeConstraint : public FParticleRule, public FPBDVolumeConstraintBase
+
+class FPBDVolumeConstraint : public FPBDVolumeConstraintBase
 {
 	typedef FPBDVolumeConstraintBase Base;
 
   public:
-	  FPBDVolumeConstraint(const FDynamicParticles& InParticles, TArray<TVec3<int32>>&& InConstraints, const FReal InStiffness = (FReal)1.)
+	  FPBDVolumeConstraint(const FSolverParticles& InParticles, TArray<TVec3<int32>>&& InConstraints, const FSolverReal InStiffness = (FSolverReal)1.)
 	    : Base(InParticles, MoveTemp(InConstraints), InStiffness) {}
-	virtual ~FPBDVolumeConstraint() {}
+	virtual ~FPBDVolumeConstraint() override {}
 
-	void Apply(FPBDParticles& InParticles, const FReal dt) const override //-V762
+	void Apply(FSolverParticles& InParticles, const FSolverReal dt) const
 	{
-		auto W = Base::GetWeights(InParticles, (FReal)1.);
-		auto Grads = Base::GetGradients(InParticles);
-		auto S = Base::GetScalingFactor(InParticles, Grads, W);
+		const TArray<FSolverReal> W = Base::GetWeights(InParticles, (FSolverReal)1.);
+		const TArray<FSolverVec3> Grads = Base::GetGradients(InParticles);
+		const FSolverReal S = Base::GetScalingFactor(InParticles, Grads, W);
 		for (uint32 i = 0; i < InParticles.Size(); ++i)
 		{
 			InParticles.P(i) -= S * W[i] * Grads[i];
@@ -29,7 +27,4 @@ class FPBDVolumeConstraint : public FParticleRule, public FPBDVolumeConstraintBa
 	}
 };
 
-template<class T>
-using TPBDVolumeConstraint UE_DEPRECATED(4.27, "Deprecated. this class is to be deleted, use FPBDVolumeConstraint instead") = FPBDVolumeConstraint;
-
-}
+}  // End namespace Chaos::Softs
diff --git a/Engine/Source/Runtime/Experimental/Chaos/Public/Chaos/PBDVolumeConstraintBase.h b/Engine/Source/Runtime/Experimental/Chaos/Public/Chaos/PBDVolumeConstraintBase.h
--- a/Engine/Source/Runtime/Experimental/Chaos/Public/Chaos/PBDVolumeConstraintBase.h
+++ b/Engine/Source/Runtime/Experimental/Chaos/Public/Chaos/PBDVolumeConstraintBase.h
@@ -1,68 +1,70 @@
 // Copyright Epic Games, Inc. All Rights Reserved.
 #pragma once
 
+#include "Chaos/PBDSoftsEvolutionFwd.h"
 #include "Chaos/Array.h"
 #include "Chaos/PBDParticles.h"
 #include "Chaos/ParticleRule.h"
 
-namespace Chaos
+namespace Chaos::Softs
 {
+
 class FPBDVolumeConstraintBase
 {
   public:
-	  FPBDVolumeConstraintBase(const FDynamicParticles& InParticles, TArray<TVec3<int32>>&& InConstraints, const FReal InStiffness = (FReal)1.)
+	  FPBDVolumeConstraintBase(const FSolverParticles& InParticles, TArray<TVec3<int32>>&& InConstraints, const FSolverReal InStiffness = (FSolverReal)1.)
 	    : Constraints(InConstraints), Stiffness(InStiffness)
 	{
-		FVec3 Com = FVec3(0, 0, 0);
+		FSolverVec3 Com = FSolverVec3(0.);
 		for (uint32 i = 0; i < InParticles.Size(); ++i)
 		{
 			Com += InParticles.X(i);
 		}
 		Com /= InParticles.Size();
-		RefVolume = 0;
+		RefVolume = (FSolverReal)0.;
 		for (const TVec3<int32>& Constraint : Constraints)
 		{
-			const FVec3& P1 = InParticles.X(Constraint[0]);
-			const FVec3& P2 = InParticles.X(Constraint[1]);
-			const FVec3& P3 = InParticles.X(Constraint[2]);
+			const FSolverVec3& P1 = InParticles.X(Constraint[0]);
+			const FSolverVec3& P2 = InParticles.X(Constraint[1]);
+			const FSolverVec3& P3 = InParticles.X(Constraint[2]);
 			RefVolume += GetVolume(P1, P2, P3, Com);
 		}
-		RefVolume /= (FReal)9.;
+		RefVolume /= (FSolverReal)9.;
 	}
 	virtual ~FPBDVolumeConstraintBase() {}
 
-	TArray<FReal> GetWeights(const FPBDParticles& InParticles, const FReal Alpha) const
+	TArray<FSolverReal> GetWeights(const FSolverParticles& InParticles, const FSolverReal Alpha) const
 	{
-		TArray<FReal> W;
+		TArray<FSolverReal> W;
 		W.SetNum(InParticles.Size());
-		FReal oneminusAlpha = (FReal)1. - Alpha;
-		FReal Wg = (FReal)1. / (FReal)InParticles.Size();
-		FReal WlDenom = 0;
+		const FSolverReal OneMinusAlpha = (FSolverReal)1. - Alpha;
+		const FSolverReal Wg = (FSolverReal)1. / (FSolverReal)InParticles.Size();
+		FSolverReal WlDenom = (FSolverReal)0.;
 		for (uint32 i = 0; i < InParticles.Size(); ++i)
 		{
 			WlDenom += (InParticles.P(i) - InParticles.X(i)).Size();
 		}
 		for (uint32 i = 0; i < InParticles.Size(); ++i)
 		{
-			FReal Wl = (InParticles.P(i) - InParticles.X(i)).Size() / WlDenom;
-			W[i] = oneminusAlpha * Wl + Alpha * Wg;
+			const FSolverReal Wl = (InParticles.P(i) - InParticles.X(i)).Size() / WlDenom;
+			W[i] = OneMinusAlpha * Wl + Alpha * Wg;
 		}
 		return W;
 	}
 
-	TArray<FVec3> GetGradients(const FPBDParticles& InParticles) const
+	TArray<FSolverVec3> GetGradients(const FSolverParticles& InParticles) const
 	{
-		FVec3 Com = FVec3(0, 0, 0);
+		FSolverVec3 Com = FSolverVec3(0.);
 		for (uint32 i = 0; i < InParticles.Size(); ++i)
 		{
 			Com += InParticles.P(i);
 		}
 		Com /= InParticles.Size();
-		TArray<FVec3> Grads;
+		TArray<FSolverVec3> Grads;
 		Grads.SetNum(InParticles.Size());
-		for (FVec3& Elem : Grads)
+		for (FSolverVec3& Elem : Grads)
 		{
-			Elem = FVec3(0, 0, 0);
+			Elem = FSolverVec3(0.);
 		}
 		for (int32 i = 0; i < Constraints.Num(); ++i)
 		{
@@ -70,76 +72,79 @@ class FPBDVolumeConstraintBase
 			const int32 i1 = Constraint[0];
 			const int32 i2 = Constraint[1];
 			const int32 i3 = Constraint[2];
-			const FVec3& P1 = InParticles.P(i1);
-			const FVec3& P2 = InParticles.P(i2);
-			const FVec3& P3 = InParticles.P(i3);
-			const FReal Area = GetArea(P1, P2, P3);
-			const FVec3 Normal = GetNormal(P1, P2, P3, Com);
+			const FSolverVec3& P1 = InParticles.P(i1);
+			const FSolverVec3& P2 = InParticles.P(i2);
+			const FSolverVec3& P3 = InParticles.P(i3);
+			const FSolverReal Area = GetArea(P1, P2, P3);
+			const FSolverVec3 Normal = GetNormal(P1, P2, P3, Com);
 			Grads[i1] += Area * Normal;
 			Grads[i2] += Area * Normal;
 			Grads[i3] += Area * Normal;
 		}
-		for (FVec3& Elem : Grads)
+		for (FSolverVec3& Elem : Grads)
 		{
-			Elem *= (FReal)1. / (FReal)3.;
+			Elem *= (FSolverReal)1. / (FSolverReal)3.;
 		}
 		return Grads;
 	}
 
-	FReal GetScalingFactor(const FPBDParticles& InParticles, const TArray<FVec3>& Grads, const TArray<FReal>& W) const
+	FSolverReal GetScalingFactor(const FSolverParticles& InParticles, const TArray<FSolverVec3>& Grads, const TArray<FSolverReal>& W) const
 	{
-		FVec3 Com = FVec3(0, 0, 0);
+		FSolverVec3 Com = FSolverVec3(0.);
 		for (uint32 i = 0; i < InParticles.Size(); ++i)
 		{
 			Com += InParticles.P(i);
 		}
 		Com /= InParticles.Size();
-		FReal Volume = 0;
+		FSolverReal Volume = (FSolverReal)0.;
 		for (const TVec3<int32>& Constraint : Constraints)
 		{
-			const FVec3& P1 = InParticles.P(Constraint[0]);
-			const FVec3& P2 = InParticles.P(Constraint[1]);
-			const FVec3& P3 = InParticles.P(Constraint[2]);
+			const FSolverVec3& P1 = InParticles.P(Constraint[0]);
+			const FSolverVec3& P2 = InParticles.P(Constraint[1]);
+			const FSolverVec3& P3 = InParticles.P(Constraint[2]);
 			Volume += GetVolume(P1, P2, P3, Com);
 		}
-		Volume /= (FReal)9.;
-		FReal Denom = 0;
+		Volume /= (FSolverReal)9.;
+		FSolverReal Denom = (FSolverReal)0;
 		for (uint32 i = 0; i < InParticles.Size(); ++i)
 		{
 			Denom += W[i] * Grads[i].SizeSquared();
 		}
-		FReal S = (Volume - RefVolume) / Denom;
+		const FSolverReal S = (Volume - RefVolume) / Denom;
 		return Stiffness * S;
 	}
 
-	void SetStiffness(FReal InStiffness) { Stiffness = FMath::Clamp(InStiffness, (FReal)0., (FReal)1.); }
+	void SetStiffness(FSolverReal InStiffness) { Stiffness = FMath::Clamp(InStiffness, (FSolverReal)0., (FSolverReal)1.); }
 
 protected:
 	TArray<TVec3<int32>> Constraints;
 
 private:
 	// Utility functions for the triangle concept
-	FVec3 GetNormal(const FVec3 P1, const FVec3& P2, const FVec3& P3, const FVec3& Com) const
+	FSolverVec3 GetNormal(const FSolverVec3 P1, const FSolverVec3& P2, const FSolverVec3& P3, const FSolverVec3& Com) const
 	{
-		const FVec3 Normal = FVec3::CrossProduct(P2 - P1, P3 - P1).GetSafeNormal();
-		if (FVec3::DotProduct((P1 + P2 + P3) / (FReal)3. - Com, Normal) < 0)
+		const FSolverVec3 Normal = FSolverVec3::CrossProduct(P2 - P1, P3 - P1).GetSafeNormal();
+		if (FSolverVec3::DotProduct((P1 + P2 + P3) / (FSolverReal)3. - Com, Normal) < (FSolverReal)0.)
+		{
 			return -Normal;
+		}
 		return Normal;
 	}
 
-	FReal GetArea(const FVec3& P1, const FVec3& P2, const FVec3& P3) const
+	FSolverReal GetArea(const FSolverVec3& P1, const FSolverVec3& P2, const FSolverVec3& P3) const
 	{
-		FVec3 B = (P2 - P1).GetSafeNormal();
-		FVec3 H = FVec3::DotProduct(B, P3 - P1) * B + P1;
-		return (FReal)0.5 * (P2 - P1).Size() * (P3 - H).Size();
+		const FSolverVec3 B = (P2 - P1).GetSafeNormal();
+		const FSolverVec3 H = FSolverVec3::DotProduct(B, P3 - P1) * B + P1;
+		return (FSolverReal)0.5 * (P2 - P1).Size() * (P3 - H).Size();
 	}
 
-	FReal GetVolume(const FVec3& P1, const FVec3& P2, const FVec3& P3, const FVec3& Com) const
+	FSolverReal GetVolume(const FSolverVec3& P1, const FSolverVec3& P2, const FSolverVec3& P3, const FSolverVec3& Com) const
 	{
-		return GetArea(P1, P2, P3) * FVec3::DotProduct(P1 + P2 + P3, GetNormal(P1, P2, P3, Com));
+		return GetArea(P1, P2, P3) * FSolverVec3::DotProduct(P1 + P2 + P3, GetNormal(P1, P2, P3, Com));
 	}
 
-	FReal RefVolume;
-	FReal Stiffness;
+	FSolverReal RefVolume;
+	FSolverReal Stiffness;
 };
-}
+
+}  // End namespace Chaos::Softs
diff --git a/Engine/Source/Runtime/Experimental/Chaos/Public/Chaos/PerParticleDampVelocity.h b/Engine/Source/Runtime/Experimental/Chaos/Public/Chaos/PerParticleDampVelocity.h
--- a/Engine/Source/Runtime/Experimental/Chaos/Public/Chaos/PerParticleDampVelocity.h
+++ b/Engine/Source/Runtime/Experimental/Chaos/Public/Chaos/PerParticleDampVelocity.h
@@ -1,37 +1,38 @@
 // Copyright Epic Games, Inc. All Rights Reserved.
 #pragma once
 
-#include "Chaos/Matrix.h"
-#include "Chaos/PerParticleRule.h"
-#include "GenericPlatform/GenericPlatformMath.h"
+#include "Chaos/PBDSoftsEvolutionFwd.h"
+#include "Chaos/Vector.h"
+#include "Chaos/PBDParticles.h"
 
-namespace Chaos
+namespace Chaos::Softs
 {
-	class FPerParticleDampVelocity : public FPerParticleRule
+
+class FPerParticleDampVelocity final
+{
+public:
+	FPerParticleDampVelocity(const FSolverReal InCoefficient = (FSolverReal)0.01)
+		: Coefficient(InCoefficient)
 	{
-	public:
-		FPerParticleDampVelocity(const FReal Coefficient = (FReal)0.01)
-		    : MCoefficient(Coefficient)
-		{
-		}
-		virtual ~FPerParticleDampVelocity() {}
-
-		void UpdatePositionBasedState(const FPBDParticles& Particles, const int32 Offset, const int32 Range);
-
-		// Apply damping without first checking for kinematic particles
-		inline void ApplyFast(FPBDParticles& Particles, const FReal /*Dt*/, const int32 Index) const
-		{
-			const FVec3 R = Particles.X(Index) - MXcm;
-			const FVec3 Dv = MVcm - Particles.V(Index) + FVec3::CrossProduct(R, MOmega);
-			Particles.V(Index) += MCoefficient * Dv;
-		}
-
-	private:
-		FReal MCoefficient;
-		FVec3 MXcm, MVcm, MOmega;
-		TArray<int32> ActiveIndices;
-	};
-}
+	}
+	~FPerParticleDampVelocity() {}
+
+	void UpdatePositionBasedState(const FSolverParticles& Particles, const int32 Offset, const int32 Range);
+
+	// Apply damping without first checking for kinematic particles
+	inline void ApplyFast(FSolverParticles& Particles, const FSolverReal /*Dt*/, const int32 Index) const
+	{
+		const FSolverVec3 R = Particles.X(Index) - Xcm;
+		const FSolverVec3 Dv = Vcm - Particles.V(Index) + FSolverVec3::CrossProduct(R, Omega);
+		Particles.V(Index) += Coefficient * Dv;
+	}
+
+private:
+	FSolverReal Coefficient;
+	FSolverVec3 Xcm, Vcm, Omega;
+};
+
+}  // End namespace Chaos::Softs
 
 // Support ISPC enable/disable in non-shipping builds
 #if !INTEL_ISPC
