Chaos Cloth - Move back the post LWC solver to float for performance reason, since cloth is being simulated in local space.

diff --git a/Engine/Source/Runtime/Experimental/Chaos/Public/Chaos/PerParticlePBDCCDCollisionConstraint.h b/Engine/Source/Runtime/Experimental/Chaos/Public/Chaos/PerParticlePBDCCDCollisionConstraint.h
--- a/Engine/Source/Runtime/Experimental/Chaos/Public/Chaos/PerParticlePBDCCDCollisionConstraint.h
+++ b/Engine/Source/Runtime/Experimental/Chaos/Public/Chaos/PerParticlePBDCCDCollisionConstraint.h
@@ -6,24 +6,22 @@
 #include "Chaos/PBDActiveView.h"
 #include "Misc/ScopeLock.h"
 
-namespace Chaos
+namespace Chaos::Softs
 {
-template<EGeometryParticlesSimType SimType>
-class TPerParticlePBDCCDCollisionConstraint final
+
+class FPerParticlePBDCCDCollisionConstraint final
 {
 public:
-	typedef TKinematicGeometryParticlesImp<FReal, 3, SimType> FCollisionParticles;
-
-	TPerParticlePBDCCDCollisionConstraint(
-		const TPBDActiveView<FCollisionParticles>& InCollisionParticlesActiveView,
-		TArray<FRigidTransform3>& InCollisionTransforms,
+	FPerParticlePBDCCDCollisionConstraint(
+		const TPBDActiveView<FSolverRigidParticles>& InCollisionParticlesActiveView,
+		TArray<FSolverRigidTransform3>& InCollisionTransforms,
 		TArray<bool>& InCollided,
-		TArray<FVec3>& InContacts,
-		TArray<FVec3>& InNormals,
+		TArray<FSolverVec3>& InContacts,
+		TArray<FSolverVec3>& InNormals,
 		TArray<uint32>& InDynamicGroupIds,
 		TArray<uint32>& InKinematicGroupIds,
-		const TArray<FReal>& InPerGroupThicknesses,
-		const TArray<FReal>& InPerGroupFriction,
+		const TArray<FSolverReal>& InPerGroupThicknesses,
+		const TArray<FSolverReal>& InPerGroupFriction,
 		bool bWriteCCDContacts)
 		: CollisionParticlesActiveView(InCollisionParticlesActiveView)
 		, CollisionTransforms(InCollisionTransforms)
@@ -38,12 +36,12 @@ public:
 	{
 	}
 
-	~TPerParticlePBDCCDCollisionConstraint()
+	~FPerParticlePBDCCDCollisionConstraint()
 	{
 		delete Mutex;
 	}
 
-	inline void ApplyRange(FPBDParticles& Particles, const FReal Dt, const int32 Offset, const int32 Range) const
+	inline void ApplyRange(FSolverParticles& Particles, const FSolverReal Dt, const int32 Offset, const int32 Range) const
 	{
 		if (Mutex)
 		{
@@ -57,22 +55,22 @@ public:
 
 private:
 	template<bool bLockAndWriteContacts>
-	inline void ApplyRangeHelper(FPBDParticles& Particles, const FReal Dt, const int32 Offset, const int32 Range) const
+	inline void ApplyRangeHelper(FSolverParticles& Particles, const FSolverReal Dt, const int32 Offset, const int32 Range) const
 	{
 		const uint32 DynamicGroupId = DynamicGroupIds[Offset];  // Particle group Id, must be the same across the entire range
-		const FReal Friction = PerGroupFriction[DynamicGroupId];
-		const FReal Thickness = PerGroupThicknesses[DynamicGroupId];
+		const FSolverReal Friction = PerGroupFriction[DynamicGroupId];
+		const FSolverReal Thickness = PerGroupThicknesses[DynamicGroupId];
 
 		PhysicsParallelFor(Range - Offset, [this, &Particles, Offset, DynamicGroupId, Thickness, Friction, Dt](int32 i)
 		{
 			const int32 Index = Offset + i;
 
-			if (Particles.InvM(Index) == (FReal)0.)
+			if (Particles.InvM(Index) == (FSolverReal)0.)
 			{
 				return;  // Continue
 			}
 
-			CollisionParticlesActiveView.SequentialFor([this, &Particles, &Index, DynamicGroupId, Thickness, Friction, Dt](FCollisionParticles& CollisionParticles, int32 CollisionIndex)
+			CollisionParticlesActiveView.SequentialFor([this, &Particles, &Index, DynamicGroupId, Thickness, Friction, Dt](FSolverRigidParticles& CollisionParticles, int32 CollisionIndex)
 			{
 				const uint32 KinematicGroupId = KinematicGroupIds[CollisionIndex];  // Collision group Id
 
@@ -81,19 +79,19 @@ private:
 					return; // Bail out if the collision groups doesn't match the particle group id, or use INDEX_NONE (= global collision that affects all particle)
 				}
 
-				const FRigidTransform3 Frame(CollisionParticles.X(CollisionIndex), CollisionParticles.R(CollisionIndex));
+				const FSolverRigidTransform3 Frame(CollisionParticles.X(CollisionIndex), CollisionParticles.R(CollisionIndex));
 
-				const Pair<FVec3, bool> PointPair = CollisionParticles.Geometry(CollisionIndex)->FindClosestIntersection(
-					CollisionTransforms[CollisionIndex].InverseTransformPositionNoScale(Particles.X(Index)),
-					Frame.InverseTransformPositionNoScale(Particles.P(Index)), Thickness);
+				const Pair<FVec3, bool> PointPair = CollisionParticles.Geometry(CollisionIndex)->FindClosestIntersection(  // Geometry operates in FReal
+					FVec3(CollisionTransforms[CollisionIndex].InverseTransformPositionNoScale(Particles.X(Index))),        // hence the back and forth
+					FVec3(Frame.InverseTransformPositionNoScale(Particles.P(Index))), (FReal)Thickness);                   // FVec3/FReal conversions
 
 				if (PointPair.Second)
 				{
 					Collided[CollisionIndex] = true;
 
-					const FVec3 Normal = CollisionParticles.Geometry(CollisionIndex)->Normal(PointPair.First);
-					const FVec3 NormalWorld = Frame.TransformVectorNoScale(Normal);
-					const FVec3 ContactWorld = Frame.TransformPositionNoScale(PointPair.First);
+					const FSolverVec3 Normal = FSolverVec3(CollisionParticles.Geometry(CollisionIndex)->Normal(PointPair.First));
+					const FSolverVec3 NormalWorld = Frame.TransformVectorNoScale(Normal);
+					const FSolverVec3 ContactWorld = Frame.TransformPositionNoScale(PointPair.First);
 
 					if (bLockAndWriteContacts)
 					{
@@ -102,20 +100,20 @@ private:
 						Contacts.Emplace(ContactWorld);
 						Normals.Emplace(NormalWorld);
 					}
-					const FVec3 Direction = ContactWorld - Particles.P(Index);
-					const FReal Penetration = FMath::Max((FReal)0., FVec3::DotProduct(NormalWorld, Direction)) + (FReal)THRESH_POINT_ON_PLANE;
+					const FSolverVec3 Direction = ContactWorld - Particles.P(Index);
+					const FSolverReal Penetration = FMath::Max((FSolverReal)0., FSolverVec3::DotProduct(NormalWorld, Direction)) + (FSolverReal)THRESH_POINT_ON_PLANE;
 
 					Particles.P(Index) += Penetration * NormalWorld;
 
 					// Friction
-					const FVec3 VectorToPoint = Particles.P(Index) - CollisionParticles.X(CollisionIndex);
-					const FVec3 RelativeDisplacement = (Particles.P(Index) - Particles.X(Index)) - (CollisionParticles.V(CollisionIndex) + FVec3::CrossProduct(CollisionParticles.W(CollisionIndex), VectorToPoint)) * Dt;  // This corresponds to the tangential velocity multiplied by dt (friction will drive this to zero if it is high enough)
-					const FVec3 RelativeDisplacementTangent = RelativeDisplacement - FVec3::DotProduct(RelativeDisplacement, NormalWorld) * NormalWorld;  // Project displacement into the tangential plane
-					const FReal RelativeDisplacementTangentLength = RelativeDisplacementTangent.Size();
+					const FSolverVec3 VectorToPoint = Particles.P(Index) - CollisionParticles.X(CollisionIndex);
+					const FSolverVec3 RelativeDisplacement = (Particles.P(Index) - Particles.X(Index)) - (CollisionParticles.V(CollisionIndex) + FSolverVec3::CrossProduct(CollisionParticles.W(CollisionIndex), VectorToPoint)) * Dt;  // This corresponds to the tangential velocity multiplied by dt (friction will drive this to zero if it is high enough)
+					const FSolverVec3 RelativeDisplacementTangent = RelativeDisplacement - FSolverVec3::DotProduct(RelativeDisplacement, NormalWorld) * NormalWorld;  // Project displacement into the tangential plane
+					const FSolverReal RelativeDisplacementTangentLength = RelativeDisplacementTangent.Size();
 					if (RelativeDisplacementTangentLength >= SMALL_NUMBER)
 					{
-						const FReal PositionCorrection = FMath::Min<FReal>(Penetration * Friction, RelativeDisplacementTangentLength);
-						const FReal CorrectionRatio = PositionCorrection / RelativeDisplacementTangentLength;
+						const FSolverReal PositionCorrection = FMath::Min<FSolverReal>(Penetration * Friction, RelativeDisplacementTangentLength);
+						const FSolverReal CorrectionRatio = PositionCorrection / RelativeDisplacementTangentLength;
 						Particles.P(Index) -= CorrectionRatio * RelativeDisplacementTangent;
 					}
 				}
@@ -125,15 +123,16 @@ private:
 
 private:
 	// TODO(mlentine): Need a bb hierarchy
-	const TPBDActiveView<FCollisionParticles>& CollisionParticlesActiveView;
-	const TArray<FRigidTransform3>& CollisionTransforms;
+	const TPBDActiveView<FSolverRigidParticles>& CollisionParticlesActiveView;
+	const TArray<FSolverRigidTransform3>& CollisionTransforms;
 	TArray<bool>& Collided;
-	TArray<FVec3>& Contacts;
-	TArray<FVec3>& Normals;
+	TArray<FSolverVec3>& Contacts;
+	TArray<FSolverVec3>& Normals;
 	const TArray<uint32>& DynamicGroupIds;
 	const TArray<uint32>& KinematicGroupIds;
-	const TArray<FReal>& PerGroupThicknesses;
-	const TArray<FReal>& PerGroupFriction;
+	const TArray<FSolverReal>& PerGroupThicknesses;
+	const TArray<FSolverReal>& PerGroupFriction;
 	FCriticalSection* const Mutex;
 };
-}
+
+}  // End namespace Chaos::Softs
diff --git a/Engine/Source/Runtime/Experimental/Chaos/Public/Chaos/PerParticlePBDCollisionConstraint.h b/Engine/Source/Runtime/Experimental/Chaos/Public/Chaos/PerParticlePBDCollisionConstraint.h
--- a/Engine/Source/Runtime/Experimental/Chaos/Public/Chaos/PerParticlePBDCollisionConstraint.h
+++ b/Engine/Source/Runtime/Experimental/Chaos/Public/Chaos/PerParticlePBDCollisionConstraint.h
@@ -1,12 +1,11 @@
 // Copyright Epic Games, Inc. All Rights Reserved.
 #pragma once
 
-#include "Chaos/Real.h"
-#include "Chaos/PerParticleRule.h"
-#include "Chaos/Transform.h"
+#include "Chaos/Core.h"
 #include "Chaos/PBDActiveView.h"
-#include "Chaos/GeometryParticlesfwd.h"
-
+#include "Chaos/PBDSoftsEvolutionFwd.h"
+#include "Chaos/PBDParticles.h"
+#include "Chaos/KinematicGeometryParticles.h"
 #include "HAL/PlatformMath.h"
 
 // Support ISPC enable/disable in non-shipping builds
@@ -18,26 +17,24 @@ const bool bChaos_PerParticleCollision_ISPC_Enabled = true;
 extern CHAOS_API bool bChaos_PerParticleCollision_ISPC_Enabled;
 #endif
 
-namespace Chaos
+namespace Chaos::Softs
 {
-template<EGeometryParticlesSimType SimType>
-class CHAOS_API TPerParticlePBDCollisionConstraint final
+
+class CHAOS_API FPerParticlePBDCollisionConstraint final
 {
 	struct FVelocityConstraint
 	{
-		FVec3 Velocity;
-		FVec3 Normal;
+		FSolverVec3 Velocity;
+		FSolverVec3 Normal;
 	};
 
 public:
-	typedef TKinematicGeometryParticlesImp<FReal, 3, SimType> FCollisionParticles;
-
-	TPerParticlePBDCollisionConstraint(const TPBDActiveView<FCollisionParticles>& InParticlesActiveView, TArray<bool>& Collided, TArray<uint32>& DynamicGroupIds, TArray<uint32>& KinematicGroupIds, const TArray<FReal>& PerGroupThickness, const TArray<FReal>& PerGroupFriction)
+	FPerParticlePBDCollisionConstraint(const TPBDActiveView<FSolverRigidParticles>& InParticlesActiveView, TArray<bool>& Collided, TArray<uint32>& DynamicGroupIds, TArray<uint32>& KinematicGroupIds, const TArray<FSolverReal>& PerGroupThickness, const TArray<FSolverReal>& PerGroupFriction)
 	    : bFastPositionBasedFriction(true), MCollisionParticlesActiveView(InParticlesActiveView), MCollided(Collided), MDynamicGroupIds(DynamicGroupIds), MKinematicGroupIds(KinematicGroupIds), MPerGroupThickness(PerGroupThickness), MPerGroupFriction(PerGroupFriction) {}
 
-	~TPerParticlePBDCollisionConstraint() {}
+	~FPerParticlePBDCollisionConstraint() {}
 
-	inline void ApplyRange(FPBDParticles& Particles, const FReal Dt, const int32 Offset, const int32 Range) const
+	inline void ApplyRange(FSolverParticles& Particles, const FSolverReal Dt, const int32 Offset, const int32 Range) const
 	{
 		if (bRealTypeCompatibleWithISPC && bChaos_PerParticleCollision_ISPC_Enabled && bFastPositionBasedFriction )
 		{
@@ -49,7 +46,7 @@ public:
 		}
 	}
 
-	void ApplyFriction(FPBDParticles& Particles, const FReal Dt, const int32 Index) const
+	void ApplyFriction(FSolverParticles& Particles, const FSolverReal Dt, const int32 Index) const
 	{
 		check(!bFastPositionBasedFriction);  // Do not call this function if this is setup to run with fast PB friction
 
@@ -57,38 +54,38 @@ public:
 		{
 			return;
 		}
-		const FReal VN = FVec3::DotProduct(Particles.V(Index), MVelocityConstraints[Index].Normal);
-		const FReal VNBody = FVec3::DotProduct(MVelocityConstraints[Index].Velocity, MVelocityConstraints[Index].Normal);
-		const FVec3 VTBody = MVelocityConstraints[Index].Velocity - VNBody * MVelocityConstraints[Index].Normal;
-		const FVec3 VTRelative = Particles.V(Index) - VN * MVelocityConstraints[Index].Normal - VTBody;
-		const FReal VTRelativeSize = VTRelative.Size();
-		const FReal VNMax = FMath::Max(VN, VNBody);
-		const FReal VNDelta = VNMax - VN;
-		const FReal CoefficientOfFriction = MPerGroupFriction[MDynamicGroupIds[Index]];
+		const FSolverReal VN = FSolverVec3::DotProduct(Particles.V(Index), MVelocityConstraints[Index].Normal);
+		const FSolverReal VNBody = FSolverVec3::DotProduct(MVelocityConstraints[Index].Velocity, MVelocityConstraints[Index].Normal);
+		const FSolverVec3 VTBody = MVelocityConstraints[Index].Velocity - VNBody * MVelocityConstraints[Index].Normal;
+		const FSolverVec3 VTRelative = Particles.V(Index) - VN * MVelocityConstraints[Index].Normal - VTBody;
+		const FSolverReal VTRelativeSize = VTRelative.Size();
+		const FSolverReal VNMax = FMath::Max(VN, VNBody);
+		const FSolverReal VNDelta = VNMax - VN;
+		const FSolverReal CoefficientOfFriction = MPerGroupFriction[MDynamicGroupIds[Index]];
 		check(CoefficientOfFriction > 0);
-		const FReal Friction = CoefficientOfFriction * VNDelta < VTRelativeSize ? CoefficientOfFriction * VNDelta / VTRelativeSize : 1;
+		const FSolverReal Friction = CoefficientOfFriction * VNDelta < VTRelativeSize ? CoefficientOfFriction * VNDelta / VTRelativeSize : 1;
 		Particles.V(Index) = VNMax * MVelocityConstraints[Index].Normal + VTBody + VTRelative * (1 - Friction);
 	}
 
 private:
-	inline void ApplyHelper(FPBDParticles& Particles, const FReal Dt, const int32 Offset, const int32 Range) const
+	inline void ApplyHelper(FSolverParticles& Particles, const FSolverReal Dt, const int32 Offset, const int32 Range) const
 	{
 		const uint32 DynamicGroupId = MDynamicGroupIds[Offset];  // Particle group Id, must be the same across the entire range
-		const FReal PerGroupFriction = MPerGroupFriction[DynamicGroupId];
-		const FReal PerGroupThickness = MPerGroupThickness[DynamicGroupId];
+		const FSolverReal PerGroupFriction = MPerGroupFriction[DynamicGroupId];
+		const FSolverReal PerGroupThickness = MPerGroupThickness[DynamicGroupId];
 
-		if (PerGroupFriction > (FReal)KINDA_SMALL_NUMBER)
+		if (PerGroupFriction > (FSolverReal)KINDA_SMALL_NUMBER)
 		{
 			PhysicsParallelFor(Range - Offset, [this, &Particles, Dt, Offset, DynamicGroupId, PerGroupFriction, PerGroupThickness](int32 i)
 			{
 				const int32 Index = Offset + i;
 
-				if (Particles.InvM(Index) == (FReal)0.)
+				if (Particles.InvM(Index) == (FSolverReal)0.)
 				{
 					return;  // Continue
 				}
 
-				MCollisionParticlesActiveView.SequentialFor([this, &Particles, &Dt, &Index, DynamicGroupId, PerGroupFriction, PerGroupThickness](FCollisionParticles& CollisionParticles, int32 i)
+				MCollisionParticlesActiveView.SequentialFor([this, &Particles, &Dt, &Index, DynamicGroupId, PerGroupFriction, PerGroupThickness](FSolverRigidParticles& CollisionParticles, int32 i)
 				{
 					const uint32 KinematicGroupId = MKinematicGroupIds[i];  // Collision group Id
 
@@ -96,26 +93,28 @@ private:
 					{
 						return; // Bail out if the collision groups doesn't match the particle group id, or use INDEX_NONE (= global collision that affects all particle)
 					}
-					FVec3 Normal;
-					const FRigidTransform3 Frame(CollisionParticles.X(i), CollisionParticles.R(i));
-					const FReal Phi = CollisionParticles.Geometry(i)->PhiWithNormal(Frame.InverseTransformPosition(Particles.P(Index)), Normal);
-
-					const FReal Penetration = PerGroupThickness - Phi; // This is related to the Normal impulse
-					if (Penetration > (FReal)0.)
+					const FSolverRigidTransform3 Frame(CollisionParticles.X(i), CollisionParticles.R(i));
+					const FVec3 RigidSpacePosition(Frame.InverseTransformPosition(Particles.P(Index)));  // PhiWithNormal requires FReal based arguments
+					FVec3 ImplicitNormal;                                                                // since implicits don't use FSolverReal
+					const FSolverReal Phi = (FSolverReal)CollisionParticles.Geometry(i)->PhiWithNormal(RigidSpacePosition, ImplicitNormal);
+					const FSolverVec3 Normal(ImplicitNormal);
+
+					const FSolverReal Penetration = PerGroupThickness - Phi; // This is related to the Normal impulse
+					if (Penetration > (FSolverReal)0.)
 					{
-						const FVec3 NormalWorld = Frame.TransformVector(Normal);
+						const FSolverVec3 NormalWorld = Frame.TransformVector(Normal);
 						Particles.P(Index) += Penetration * NormalWorld;
 
 						if (bFastPositionBasedFriction)
 						{
-							FVec3 VectorToPoint = Particles.P(Index) - CollisionParticles.X(i);
-							const FVec3 RelativeDisplacement = (Particles.P(Index) - Particles.X(Index)) - (CollisionParticles.V(i) + FVec3::CrossProduct(CollisionParticles.W(i), VectorToPoint)) * Dt; // This corresponds to the tangential velocity multiplied by dt (friction will drive this to zero if it is high enough)
-							const FVec3 RelativeDisplacementTangent = RelativeDisplacement - FVec3::DotProduct(RelativeDisplacement, NormalWorld) * NormalWorld; // Project displacement into the tangential plane
-							const FReal RelativeDisplacementTangentLength = RelativeDisplacementTangent.Size();
+							FSolverVec3 VectorToPoint = Particles.P(Index) - CollisionParticles.X(i);
+							const FSolverVec3 RelativeDisplacement = (Particles.P(Index) - Particles.X(Index)) - (CollisionParticles.V(i) + FSolverVec3::CrossProduct(CollisionParticles.W(i), VectorToPoint)) * Dt; // This corresponds to the tangential velocity multiplied by dt (friction will drive this to zero if it is high enough)
+							const FSolverVec3 RelativeDisplacementTangent = RelativeDisplacement - FSolverVec3::DotProduct(RelativeDisplacement, NormalWorld) * NormalWorld; // Project displacement into the tangential plane
+							const FSolverReal RelativeDisplacementTangentLength = RelativeDisplacementTangent.Size();
 							if (RelativeDisplacementTangentLength >= SMALL_NUMBER)
 							{
-								const FReal PositionCorrection = FMath::Min<FReal>(Penetration * PerGroupFriction, RelativeDisplacementTangentLength);
-								const FReal CorrectionRatio = PositionCorrection / RelativeDisplacementTangentLength;
+								const FSolverReal PositionCorrection = FMath::Min<FSolverReal>(Penetration * PerGroupFriction, RelativeDisplacementTangentLength);
+								const FSolverReal CorrectionRatio = PositionCorrection / RelativeDisplacementTangentLength;
 								Particles.P(Index) -= CorrectionRatio * RelativeDisplacementTangent;
 							}
 						}
@@ -123,8 +122,8 @@ private:
 						{
 							// Note, to fix: Only use fast position based friction for now, since adding to TMaps here is not thread safe when calling Apply on multiple threads (will cause crash)
 							FVelocityConstraint Constraint;
-							FVec3 VectorToPoint = Particles.P(Index) - CollisionParticles.X(i);
-							Constraint.Velocity = CollisionParticles.V(i) + FVec3::CrossProduct(CollisionParticles.W(i), VectorToPoint);
+							FSolverVec3 VectorToPoint = Particles.P(Index) - CollisionParticles.X(i);
+							Constraint.Velocity = CollisionParticles.V(i) + FSolverVec3::CrossProduct(CollisionParticles.W(i), VectorToPoint);
 							Constraint.Normal = Frame.TransformVector(Normal);
 						
 							MVelocityConstraints.Add(Index, Constraint);
@@ -144,7 +143,7 @@ private:
 					return;  // Continue
 				}
 
-				MCollisionParticlesActiveView.SequentialFor([this, &Particles, &Dt, &Index, DynamicGroupId, PerGroupFriction, PerGroupThickness](FCollisionParticles& CollisionParticles, int32 i)
+				MCollisionParticlesActiveView.SequentialFor([this, &Particles, &Dt, &Index, DynamicGroupId, PerGroupFriction, PerGroupThickness](FSolverRigidParticles& CollisionParticles, int32 i)
 				{
 					const uint32 KinematicGroupId = MKinematicGroupIds[i];  // Collision group Id
 
@@ -152,14 +151,16 @@ private:
 					{
 						return; // Bail out if the collision groups doesn't match the particle group id, or use INDEX_NONE (= global collision that affects all particle)
 					}
-					FVec3 Normal;
-					const FRigidTransform3 Frame(CollisionParticles.X(i), CollisionParticles.R(i));
-					const FReal Phi = CollisionParticles.Geometry(i)->PhiWithNormal(Frame.InverseTransformPosition(Particles.P(Index)), Normal);
-
-					const FReal Penetration = PerGroupThickness - Phi; // This is related to the Normal impulse
-					if (Penetration > (FReal)0.)
+					const FSolverRigidTransform3 Frame(CollisionParticles.X(i), CollisionParticles.R(i));
+					const FVec3 RigidSpacePosition(Frame.InverseTransformPosition(Particles.P(Index)));  // PhiWithNormal requires FReal based arguments
+					FVec3 ImplicitNormal;                                                                // since implicits don't use FSolverReal
+					const FSolverReal Phi = (FSolverReal)CollisionParticles.Geometry(i)->PhiWithNormal(RigidSpacePosition, ImplicitNormal);
+					const FSolverVec3 Normal(ImplicitNormal);
+
+					const FSolverReal Penetration = PerGroupThickness - Phi; // This is related to the Normal impulse
+					if (Penetration > (FSolverReal)0.)
 					{
-						const FVec3 NormalWorld = Frame.TransformVector(Normal);
+						const FSolverVec3 NormalWorld = Frame.TransformVector(Normal);
 						Particles.P(Index) += Penetration * NormalWorld;
 					}
 				});
@@ -167,17 +168,18 @@ private:
 		}
 	}
 
-	void ApplyHelperISPC(FPBDParticles& Particles, const FReal Dt, int32 Offset, int32 Range) const;
+	void ApplyHelperISPC(FSolverParticles& Particles, const FSolverReal Dt, const int32 Offset, const int32 Range) const;
 
 private:
 	bool bFastPositionBasedFriction;
 	// TODO(mlentine): Need a bb hierarchy
-	const TPBDActiveView<FCollisionParticles>& MCollisionParticlesActiveView;
+	const TPBDActiveView<FSolverRigidParticles>& MCollisionParticlesActiveView;
 	TArray<bool>& MCollided;
 	const TArray<uint32>& MDynamicGroupIds;
 	const TArray<uint32>& MKinematicGroupIds;
 	mutable TMap<int32, FVelocityConstraint> MVelocityConstraints;
-	const TArray<FReal>& MPerGroupThickness;
-	const TArray<FReal>& MPerGroupFriction;
+	const TArray<FSolverReal>& MPerGroupThickness;
+	const TArray<FSolverReal>& MPerGroupFriction;
 };
-}
+
+}  // End namespace Chaos::Softs
diff --git a/Engine/Source/Runtime/Experimental/Chaos/Public/Chaos/PerParticlePBDShapeConstraints.h b/Engine/Source/Runtime/Experimental/Chaos/Public/Chaos/PerParticlePBDShapeConstraints.h
--- a/Engine/Source/Runtime/Experimental/Chaos/Public/Chaos/PerParticlePBDShapeConstraints.h
+++ b/Engine/Source/Runtime/Experimental/Chaos/Public/Chaos/PerParticlePBDShapeConstraints.h
@@ -1,31 +1,29 @@
 // Copyright Epic Games, Inc. All Rights Reserved.
 #pragma once
 
-#include "Chaos/Array.h"
-#include "Chaos/Map.h"
-#include "Chaos/PBDParticles.h"
 #include "Chaos/PBDShapeConstraintsBase.h"
-#include "Chaos/PerParticleRule.h"
+#include "Chaos/Framework/Parallel.h"
 
-namespace Chaos
+namespace Chaos::Softs
 {
-class FPerParticlePBDShapeConstraints : public FPerParticleRule, public FPBDShapeConstraintsBase
+
+class UE_DEPRECATED(5.0, "Per particle constraint rules are no longer used by the cloth solver. Use FPBDShapeConstraints instead.") FPerParticlePBDShapeConstraints final : public FPBDShapeConstraintsBase
 {
 	typedef FPBDShapeConstraintsBase Base;
 
   public:
-	FPerParticlePBDShapeConstraints(const FReal Stiffness = (FReal)1.)
+	FPerParticlePBDShapeConstraints(const FSolverReal Stiffness = (FSolverReal)1.)
 	    : Base(Stiffness)
 	{
 	}
-	FPerParticlePBDShapeConstraints(const FDynamicParticles& InParticles, const TArray<FVec3>& TargetPositions, const FReal Stiffness = (FReal)1.)
+	FPerParticlePBDShapeConstraints(const FSolverParticles& InParticles, const TArray<FSolverVec3>& TargetPositions, const FSolverReal Stiffness = (FSolverReal)1.)
 	    : Base(InParticles, TargetPositions, Stiffness)
 	{
 	}
-	virtual ~FPerParticlePBDShapeConstraints() {}
+	virtual ~FPerParticlePBDShapeConstraints() override {}
 
 	// TODO(mlentine): We likely need to use time n positions here
-	void Apply(FPBDParticles& InParticles, const FReal Dt, const int32 Index) const override //-V762
+	void Apply(FSolverParticles& InParticles, const FSolverReal Dt, const int32 Index) const
 	{
 		if (InParticles.InvM(Index) > 0)
 		{
@@ -33,11 +31,12 @@ class FPerParticlePBDShapeConstraints : public FPerParticleRule, public FPBDShap
 		}
 	}
 
-	void Apply(FPBDParticles& InParticles, const FReal Dt) const override //-V762
+	void Apply(FSolverParticles& InParticles, const FSolverReal Dt) const
 	{
 		PhysicsParallelFor(InParticles.Size(), [&](int32 Index) {
 			Apply(InParticles, Dt, Index);
 		});
 	}
 };
-}
+
+}  // End namespace Chaos::Softs
diff --git a/Engine/Source/Runtime/Experimental/Chaos/Public/Chaos/PerParticlePBDSpringConstraints.h b/Engine/Source/Runtime/Experimental/Chaos/Public/Chaos/PerParticlePBDSpringConstraints.h
--- a/Engine/Source/Runtime/Experimental/Chaos/Public/Chaos/PerParticlePBDSpringConstraints.h
+++ b/Engine/Source/Runtime/Experimental/Chaos/Public/Chaos/PerParticlePBDSpringConstraints.h
@@ -1,22 +1,20 @@
 // Copyright Epic Games, Inc. All Rights Reserved.
 #pragma once
 
-#include "Chaos/Array.h"
-#include "Chaos/PBDParticles.h"
 #include "Chaos/PBDSpringConstraintsBase.h"
-#include "Chaos/PerParticleRule.h"
 
-namespace Chaos
+namespace Chaos::Softs
 {
-class FPerParticlePBDSpringConstraints : public FPerParticleRule, public FPBDSpringConstraintsBase
+
+class UE_DEPRECATED(5.0, "Per particle constraint rules are no longer used by the cloth solver. Use FPBDSpringConstraints instead.") FPerParticlePBDSpringConstraints : public FPBDSpringConstraintsBase
 {
 	typedef FPBDSpringConstraintsBase Base;
 	using Base::Constraints;
 	using Base::Stiffness;
 
-  public:
-	FPerParticlePBDSpringConstraints(const FPBDParticles& InParticles, const TArray<TVec2<int32>>& InConstraints, const FReal InStiffness = (FReal)1.)
-	    : Base(InParticles, 0, 0, InConstraints, TConstArrayView<FRealSingle>(), FVec2(InStiffness))
+public:
+	FPerParticlePBDSpringConstraints(const FSolverParticles& InParticles, const TArray<TVec2<int32>>& InConstraints, const FSolverReal InStiffness = (FSolverReal)1.)
+	    : Base(InParticles, 0, 0, InConstraints, TConstArrayView<FRealSingle>(), FSolverVec2(InStiffness))
 	{
 		for (int32 i = 0; i < Constraints.Num(); ++i)
 		{
@@ -35,10 +33,10 @@ class FPerParticlePBDSpringConstraints : public FPerParticleRule, public FPBDSpr
 			ParticleToConstraints[i2].Add(i);
 		}
 	}
-	virtual ~FPerParticlePBDSpringConstraints() {}
+	virtual ~FPerParticlePBDSpringConstraints() override {}
 
 	// TODO(mlentine): We likely need to use time n positions here
-	void Apply(FPBDParticles& InParticles, const FReal Dt, const int32 Index) const override //-V762
+	void Apply(FSolverParticles& InParticles, const FSolverReal Dt, const int32 Index) const
 	{
 		for (int32 i = 0; i < ParticleToConstraints[Index].Num(); ++i)
 		{
@@ -48,17 +46,18 @@ class FPerParticlePBDSpringConstraints : public FPerParticleRule, public FPBDSpr
 			int32 i2 = Constraint[1];
 			if (Index == i1 && InParticles.InvM(i1) > 0)
 			{
-				InParticles.P(i1) -= InParticles.InvM(i1) * Base::GetDelta(InParticles, CIndex, (FReal)Stiffness);
+				InParticles.P(i1) -= InParticles.InvM(i1) * Base::GetDelta(InParticles, CIndex, (FSolverReal)Stiffness);
 			}
 			else if (InParticles.InvM(i2) > 0)
 			{
 				check(Index == i2);
-				InParticles.P(i2) += InParticles.InvM(i2) * Base::GetDelta(InParticles, CIndex, (FReal)Stiffness);
+				InParticles.P(i2) += InParticles.InvM(i2) * Base::GetDelta(InParticles, CIndex, (FSolverReal)Stiffness);
 			}
 		}
 	}
 
-  private:
+private:
 	TArray<TArray<int32>> ParticleToConstraints;
 };
-}
+
+}  // End namespace Chaos::Softs
diff --git a/Engine/Source/Runtime/Experimental/Chaos/Public/Chaos/PerParticlePBDTetConstraints.h b/Engine/Source/Runtime/Experimental/Chaos/Public/Chaos/PerParticlePBDTetConstraints.h
--- a/Engine/Source/Runtime/Experimental/Chaos/Public/Chaos/PerParticlePBDTetConstraints.h
+++ b/Engine/Source/Runtime/Experimental/Chaos/Public/Chaos/PerParticlePBDTetConstraints.h
@@ -1,92 +1,87 @@
 // Copyright Epic Games, Inc. All Rights Reserved.
 #pragma once
 
-#include "Chaos/Array.h"
-#include "Chaos/PBDParticles.h"
 #include "Chaos/PBDTetConstraintsBase.h"
-#include "Chaos/PerParticleRule.h"
 
-namespace Chaos
+namespace Chaos::Softs
 {
-class FPerParticlePBDTetConstraints : public FPerParticleRule, public FPBDTetConstraintsBase
+
+class UE_DEPRECATED(5.0, "Per particle constraint rules are no longer used by the cloth solver. Use Softs::FPBDTetConstraints instead.") FPerParticlePBDTetConstraints : public Softs::FPBDTetConstraintsBase
 {
-	typedef FPBDTetConstraintsBase Base;
-	using Base::MConstraints;
+	typedef Softs::FPBDTetConstraintsBase Base;
+	using Base::Constraints;
 
   public:
-	  FPerParticlePBDTetConstraints(const FDynamicParticles& InParticles, TArray<TVec4<int32>>&& Constraints, const FReal Stiffness = (FReal)1)
-	    : Base(InParticles, MoveTemp(Constraints), Stiffness)
+	  FPerParticlePBDTetConstraints(const FSolverParticles& InParticles, TArray<TVec4<int32>>&& InConstraints, const FSolverReal InStiffness = (FSolverReal)1.)
+	    : Base(InParticles, MoveTemp(InConstraints), InStiffness)
 	{
-		for (int32 i = 0; i < MConstraints.Num(); ++i)
+		for (int32 i = 0; i < Constraints.Num(); ++i)
 		{
-			const auto& Constraint = MConstraints[i];
+			const TVec4<int32>& Constraint = Constraints[i];
 			int32 i1 = Constraint[0];
 			int32 i2 = Constraint[1];
 			int32 i3 = Constraint[2];
 			int32 i4 = Constraint[3];
-			if (i1 >= MParticleToConstraints.Num())
+			if (i1 >= ParticleToConstraints.Num())
 			{
-				MParticleToConstraints.SetNum(i1 + 1);
+				ParticleToConstraints.SetNum(i1 + 1);
 			}
-			if (i2 >= MParticleToConstraints.Num())
+			if (i2 >= ParticleToConstraints.Num())
 			{
-				MParticleToConstraints.SetNum(i2 + 1);
+				ParticleToConstraints.SetNum(i2 + 1);
 			}
-			if (i3 >= MParticleToConstraints.Num())
+			if (i3 >= ParticleToConstraints.Num())
 			{
-				MParticleToConstraints.SetNum(i3 + 1);
+				ParticleToConstraints.SetNum(i3 + 1);
 			}
-			if (i4 >= MParticleToConstraints.Num())
+			if (i4 >= ParticleToConstraints.Num())
 			{
-				MParticleToConstraints.SetNum(i4 + 1);
+				ParticleToConstraints.SetNum(i4 + 1);
 			}
-			MParticleToConstraints[i1].Add(i);
-			MParticleToConstraints[i2].Add(i);
-			MParticleToConstraints[i2].Add(i);
-			MParticleToConstraints[i2].Add(i);
+			ParticleToConstraints[i1].Add(i);
+			ParticleToConstraints[i2].Add(i);
+			ParticleToConstraints[i2].Add(i);
+			ParticleToConstraints[i2].Add(i);
 		}
 	}
-	virtual ~FPerParticlePBDTetConstraints()
+	virtual ~FPerParticlePBDTetConstraints() override
 	{
 	}
 
-	void Apply(FPBDParticles& InParticles, const FReal Dt, const int32 Index) const override //-V762
+	void Apply(FSolverParticles& InParticles, const FSolverReal Dt, const int32 Index) const
 	{
-		for (int i = 0; i < MParticleToConstraints[Index].Num(); ++i)
+		for (int i = 0; i < ParticleToConstraints[Index].Num(); ++i)
 		{
-			int32 cIndex = MParticleToConstraints[Index][i];
-			const auto& Constraint = MConstraints[cIndex];
+			int32 CIndex = ParticleToConstraints[Index][i];
+			const TVec4<int32>& Constraint = Constraints[CIndex];
 			int32 i1 = Constraint[0];
 			int32 i2 = Constraint[1];
 			int32 i3 = Constraint[2];
 			int32 i4 = Constraint[3];
-			auto grads = Base::GetGradients(InParticles, cIndex);
-			auto s = Base::GetScalingFactor(InParticles, cIndex, grads);
+			const TVec4<FSolverVec3> Grads = Base::GetGradients(InParticles, CIndex);
+			const FSolverReal S = Base::GetScalingFactor(InParticles, CIndex, Grads);
 			if (Index == i1)
 			{
-				InParticles.P(i1) -= s * InParticles.InvM(i1) * grads[0];
+				InParticles.P(i1) -= S * InParticles.InvM(i1) * Grads[0];
 			}
 			else if (Index == i2)
 			{
-				InParticles.P(i2) -= s * InParticles.InvM(i2) * grads[1];
+				InParticles.P(i2) -= S * InParticles.InvM(i2) * Grads[1];
 			}
 			else if (Index == i3)
 			{
-				InParticles.P(i3) -= s * InParticles.InvM(i3) * grads[2];
+				InParticles.P(i3) -= S * InParticles.InvM(i3) * Grads[2];
 			}
 			else
 			{
 				check(Index == i4);
-				InParticles.P(i4) -= s * InParticles.InvM(i4) * grads[3];
+				InParticles.P(i4) -= S * InParticles.InvM(i4) * Grads[3];
 			}
 		}
 	}
 
   private:
-	TArray<TArray<int32>> MParticleToConstraints;
+	TArray<TArray<int32>> ParticleToConstraints;
 };
 
-template<class T>
-using PerParticlePBDTetConstraints UE_DEPRECATED(4.27, "Deprecated. this class is to be deleted, use FPerParticlePBDTetConstraints instead") = FPerParticlePBDTetConstraints;
-
-}
+}  // End namespace Chaos::Softs
diff --git a/Engine/Source/Runtime/Experimental/Chaos/Public/Chaos/Rotation.h b/Engine/Source/Runtime/Experimental/Chaos/Public/Chaos/Rotation.h
--- a/Engine/Source/Runtime/Experimental/Chaos/Public/Chaos/Rotation.h
+++ b/Engine/Source/Runtime/Experimental/Chaos/Public/Chaos/Rotation.h
@@ -44,20 +44,25 @@ namespace Chaos
 	};
 
 	template<>
-	class TRotation<FReal, 3> : public UE::Math::TQuat<FReal>
+	class TRotation<FRealSingle, 3> : public UE::Math::TQuat<FRealSingle>
 	{
-		using BaseQuat = UE::Math::TQuat<FReal>;
+		using BaseQuat = UE::Math::TQuat<FRealSingle>;
 	public:
 		TRotation()
 		    : BaseQuat() {}
 		TRotation(const BaseQuat& Quat)
 		    : BaseQuat(Quat) {}
-		TRotation(const FMatrix& Matrix)
+		TRotation(const FMatrix44f& Matrix)
 		    : BaseQuat(Matrix) {}
+		TRotation(const FMatrix44d& Matrix)
+		    : BaseQuat(FMatrix44f(Matrix)) {}
+		template<typename OtherType>
+		TRotation(const UE::Math::TQuat<OtherType>& Other)
+		    : BaseQuat((FRealSingle)Other.X, (FRealSingle)Other.Y, (FRealSingle)Other.Z, (FRealSingle)Other.W) {}
 
-		inline PMatrix<FReal, 3, 3> ToMatrix() const
+		inline PMatrix<FRealSingle, 3, 3> ToMatrix() const
 		{
-			PMatrix<FReal, 3, 3> R;
+			PMatrix<FRealSingle, 3, 3> R;
 			BaseQuat::ToMatrix(R);
 			return R;
 		}
@@ -73,7 +78,7 @@ namespace Chaos
 		 * @note EpsilonSq is approximately the square of the angle below which we cannot calculate the axis. It needs to be "much greater" than square of the error in the
 		 * quaternion values which is usually ~1e-4, so values around 1e-3^2 = 1e-6 or greater are about right.
 		 */
-		inline bool ToAxisAndAngleSafe(TVector<FReal, 3>& OutAxis, FReal& OutAngle, const TVector<FReal, 3>& DefaultAxis, FReal EpsilionSq = 1e-6f) const
+		inline bool ToAxisAndAngleSafe(TVector<FRealSingle, 3>& OutAxis, FRealSingle& OutAngle, const TVector<FRealSingle, 3>& DefaultAxis, FRealSingle EpsilionSq = 1e-6f) const
 		{
 			OutAngle = GetAngle();
 			return GetRotationAxisSafe(OutAxis, DefaultAxis, EpsilionSq);
@@ -89,16 +94,16 @@ namespace Chaos
 		 * @note EpsilonSq is approximately the square of the angle below which we cannot calculate the axis. It needs to be "much greater" than square of the error in the
 		 * quaternion values which is usually ~1e-4, so values around 1e-3^2 = 1e-6 or greater are about right.
 		 */
-		inline bool GetRotationAxisSafe(TVector<FReal, 3>& OutAxis, const TVector<FReal, 3>& DefaultAxis, FReal EpsilionSq = 1e-6f) const
+		inline bool GetRotationAxisSafe(TVector<FRealSingle, 3>& OutAxis, const TVector<FRealSingle, 3>& DefaultAxis, FRealSingle EpsilionSq = 1e-6f) const
 		{
 			// Tolerance must be much larger than error in normalized vector (usually ~1e-4) for the 
 			// axis calculation to succeed for small angles. For small angles, W ~= 1, and
 			// X, Y, Z ~= 0. If the values of X, Y, Z are around 1e-4 we are just normalizing error.
-			const FReal LenSq = X * X + Y * Y + Z * Z;
+			const FRealSingle LenSq = X * X + Y * Y + Z * Z;
 			if (LenSq > EpsilionSq)
 			{
-				FReal InvLen = FMath::InvSqrt(LenSq);
-				OutAxis = TVector<FReal, 3>(X * InvLen, Y * InvLen, Z * InvLen);
+				FRealSingle InvLen = FMath::InvSqrt(LenSq);
+				OutAxis = TVector<FRealSingle, 3>(X * InvLen, Y * InvLen, Z * InvLen);
 				return true;
 			}
 
@@ -119,7 +124,7 @@ namespace Chaos
 		/**
 		 * @brief Return the large absolute element value
 		*/
-		FORCEINLINE FReal GetAbsMax() const
+		FORCEINLINE FRealSingle GetAbsMax() const
 		{
 			return FMath::Max(FMath::Max(FMath::Abs(X), FMath::Abs(Y)), FMath::Max(FMath::Abs(Z), FMath::Abs(W)));
 		}
@@ -127,7 +132,7 @@ namespace Chaos
 		/**
 		 * @brief Return the dot product of two quaternions
 		*/
-		static inline FReal DotProduct(const TRotation<FReal, 3>& L, const TRotation<FReal, 3>& R)
+		static inline FRealSingle DotProduct(const TRotation<FRealSingle, 3>& L, const TRotation<FRealSingle, 3>& R)
 		{
 			return (L | R);
 		}
@@ -135,17 +140,17 @@ namespace Chaos
 		/**
 		 * Return the complex conjugate of the rotation
 		 */
-		CHAOS_API static TRotation<FReal, 3> Conjugate(const ::Chaos::TRotation<FReal, 3>& InR);
+		CHAOS_API static TRotation<FRealSingle, 3> Conjugate(const ::Chaos::TRotation<FRealSingle, 3>& InR);
 
 		/**
 		 * Negate all values of the quaternion (note: not the inverse rotation. See Conjugate)
 		 */
-		CHAOS_API static TRotation<FReal, 3> Negate(const ::Chaos::TRotation<FReal, 3>& InR);
+		CHAOS_API static TRotation<FRealSingle, 3> Negate(const ::Chaos::TRotation<FRealSingle, 3>& InR);
 
 		/**
 		 * Create an identity rotation
 		 */
-		static inline TRotation<FReal, 3> FromIdentity()
+		static inline TRotation<FRealSingle, 3> FromIdentity()
 		{
 			return BaseQuat(0, 0, 0, 1);
 		}
@@ -153,7 +158,7 @@ namespace Chaos
 		/**
 		 * Create a rotation by explicitly specifying all elements
 		 */
-		static inline TRotation<FReal, 3> FromElements(const FReal X, const FReal Y, const FReal Z, const FReal W)
+		static inline TRotation<FRealSingle, 3> FromElements(const FRealSingle X, const FRealSingle Y, const FRealSingle Z, const FRealSingle W)
 		{
 			using FQuatReal = BaseQuat::FReal;
 			return BaseQuat((FQuatReal)X, (FQuatReal)Y, (FQuatReal)Z, (FQuatReal)W);
@@ -162,7 +167,7 @@ namespace Chaos
 		/**
 		 * Create a rotation by explicitly specifying all elements
 		 */
-		static inline TRotation<FReal, 3> FromElements(const ::Chaos::TVector<FReal, 3>& V, const FReal W)
+		static inline TRotation<FRealSingle, 3> FromElements(const ::Chaos::TVector<FRealSingle, 3>& V, const FRealSingle W)
 		{
 			return FromElements(V.X, V.Y, V.Z, W);
 		}
@@ -170,43 +175,43 @@ namespace Chaos
 		/**
 		 * Create a rotation about an axis by an angle specified in radians
 		 */
-		static inline TRotation<FReal, 3> FromAxisAngle(const ::Chaos::TVector<FReal, 3>& Axis, const FReal AngleRad)
+		static inline TRotation<FRealSingle, 3> FromAxisAngle(const ::Chaos::TVector<FRealSingle, 3>& Axis, const FRealSingle AngleRad)
 		{
-			return BaseQuat(UE::Math::TVector<FReal>(Axis.X, Axis.Y, Axis.Z), (BaseQuat::FReal)AngleRad);
+			return BaseQuat(UE::Math::TVector<FRealSingle>(Axis.X, Axis.Y, Axis.Z), (BaseQuat::FReal)AngleRad);
 		}
 
 		/**
 		 * Create a rotation about an axis V/|V| by angle |V| in radians
 		 */
-		CHAOS_API static TRotation<FReal, 3> FromVector(const ::Chaos::TVector<FReal, 3>& V);
+		CHAOS_API static TRotation<FRealSingle, 3> FromVector(const ::Chaos::TVector<FRealSingle, 3>& V);
 
 		/**
 		 * Generate a Rotation that would rotate vector InitialVector to FinalVector
 		 */
-		CHAOS_API static TRotation<FReal, 3> FromRotatedVector(
-			const ::Chaos::TVector<FReal, 3>& InitialVector,
-			const ::Chaos::TVector<FReal, 3>& FinalVector);
+		CHAOS_API static TRotation<FRealSingle, 3> FromRotatedVector(
+			const ::Chaos::TVector<FRealSingle, 3>& InitialVector,
+			const ::Chaos::TVector<FRealSingle, 3>& FinalVector);
 
 		/**
 		 * Calculate the angular velocity required to take an object with orientation R0 to orientation R1 in time Dt.
 		 *
 		 * Uses the relation: DQ/DT = (W * Q)/2
 		 */
-		CHAOS_API static TVector<FReal, 3> CalculateAngularVelocity1(const TRotation<FReal, 3>& R0, const TRotation<FReal, 3>& InR1, const FReal InDt);
+		CHAOS_API static TVector<FRealSingle, 3> CalculateAngularVelocity1(const TRotation<FRealSingle, 3>& R0, const TRotation<FRealSingle, 3>& InR1, const FRealSingle InDt);
 
 		/**
 		 * Calculate the angular velocity required to take an object with orientation R0 to orientation R1 in time Dt.
 		 *
 		 * Uses the Quaternion to Axis/Angle method.
 		 */
-		CHAOS_API static TVector<FReal, 3> CalculateAngularVelocity2(const TRotation<FReal, 3>& R0, const TRotation<FReal, 3>& InR1, const FReal InDt);
+		CHAOS_API static TVector<FRealSingle, 3> CalculateAngularVelocity2(const TRotation<FRealSingle, 3>& R0, const TRotation<FRealSingle, 3>& InR1, const FRealSingle InDt);
 
 		/**
 		 * Calculate the angular velocity required to take an object with orientation R0 to orientation R1 in time Dt.
 		 *
 		 * This should match the algorithm used in PerParticleUpdateFromDeltaPosition rule.
 		 */
-		static inline TVector<FReal, 3> CalculateAngularVelocity(const TRotation<FReal, 3>& InR0, const TRotation<FReal, 3>& InR1, const FReal InDt)
+		static inline TVector<FRealSingle, 3> CalculateAngularVelocity(const TRotation<FRealSingle, 3>& InR0, const TRotation<FRealSingle, 3>& InR1, const FRealSingle InDt)
 		{
 			return CalculateAngularVelocity1(InR0, InR1, InDt);
 		}
@@ -216,7 +221,7 @@ namespace Chaos
 		 *
 		 * This should match the algorithm used in PerParticleUpdateFromDeltaPosition rule.
 		 */
-		static inline TVector<FReal, 3> CalculateAngularDelta(const TRotation<FReal, 3>& InR0, const TRotation<FReal, 3>& InR1)
+		static inline TVector<FRealSingle, 3> CalculateAngularDelta(const TRotation<FRealSingle, 3>& InR0, const TRotation<FRealSingle, 3>& InR1)
 		{
 			return CalculateAngularVelocity(InR0, InR1, 1.0f);
 		}
@@ -226,19 +231,223 @@ namespace Chaos
 		 *
 		 * Uses the relation: DQ/DT = (W * Q)/2
 		 */
-		CHAOS_API static TRotation<FReal, 3> IntegrateRotationWithAngularVelocity(const TRotation<FReal, 3>& InR0, const TVector<FReal, 3>& InW, const FReal InDt);
+		CHAOS_API static TRotation<FRealSingle, 3> IntegrateRotationWithAngularVelocity(const TRotation<FRealSingle, 3>& InR0, const TVector<FRealSingle, 3>& InW, const FRealSingle InDt);
 
 		/**
 		 * Check that two rotations are approximately equal. Assumes the quaternions are normalized and in the same hemisphere.
 		 * For small values of Epsilon, this is approximately equivalent to checking that the rotations are within 2*Epsilon
 		 * radians of each other.
 		 */
-		static inline bool IsNearlyEqual(const TRotation<FReal, 3>& A, const TRotation<FReal, 3>& B, const FReal Epsilon)
+		static inline bool IsNearlyEqual(const TRotation<FRealSingle, 3>& A, const TRotation<FRealSingle, 3>& B, const FRealSingle Epsilon)
 		{
 			// Only check imaginary part. This is comparing Epsilon to 2*AngleDelta for small angle deltas
-			return FMath::IsNearlyEqual((FReal)A.X, (FReal)B.X, Epsilon)
-				&& FMath::IsNearlyEqual((FReal)A.Y, (FReal)B.Y, Epsilon)
-				&& FMath::IsNearlyEqual((FReal)A.Z, (FReal)B.Z, Epsilon);
+			return FMath::IsNearlyEqual(A.X, B.X, Epsilon)
+				&& FMath::IsNearlyEqual(A.Y, B.Y, Epsilon)
+				&& FMath::IsNearlyEqual(A.Z, B.Z, Epsilon);
+		}
+	};
+
+	template<>
+	class TRotation<FRealDouble, 3> : public UE::Math::TQuat<FRealDouble>
+	{
+		using BaseQuat = UE::Math::TQuat<FRealDouble>;
+	public:
+		TRotation()
+		    : BaseQuat() {}
+		TRotation(const BaseQuat& Quat)
+		    : BaseQuat(Quat) {}
+		TRotation(const FMatrix44d& Matrix)
+		    : BaseQuat(Matrix) {}
+		TRotation(const FMatrix44f& Matrix)
+		    : BaseQuat(FMatrix44d(Matrix)) {}
+		template<typename OtherType>
+		TRotation(const UE::Math::TQuat<OtherType>& Other)
+		    : BaseQuat((FRealDouble)Other.X, (FRealDouble)Other.Y, (FRealDouble)Other.Z, (FRealDouble)Other.W) {}
+
+		inline PMatrix<FRealDouble, 3, 3> ToMatrix() const
+		{
+			PMatrix<FRealDouble, 3, 3> R;
+			BaseQuat::ToMatrix(R);
+			return R;
+		}
+
+		/**
+		 * Extract the axis and angle from the Quaternion.
+		 * @param OutAxis The axis of rotation.
+		 * @param OutAngle The angle of rotation about the axis (radians).
+		 * @param DefaultAxis The axis to set when the angle is too small to accurately calculate the axis.
+		 * @param EpsilonSq The squared tolerance used to check for small angles.
+		 * @return Whether the axis was successfully calculated (true except for very small angles around or less than Epsilon).
+		 * @warning The axis calculation cannot succeed for small angles due to numerical error. In this case, the function will return false, but set the axis to DefaultAxis.
+		 * @note EpsilonSq is approximately the square of the angle below which we cannot calculate the axis. It needs to be "much greater" than square of the error in the
+		 * quaternion values which is usually ~1e-4, so values around 1e-3^2 = 1e-6 or greater are about right.
+		 */
+		inline bool ToAxisAndAngleSafe(TVector<FRealDouble, 3>& OutAxis, FRealDouble& OutAngle, const TVector<FRealDouble, 3>& DefaultAxis, FRealDouble EpsilionSq = 1e-6f) const
+		{
+			OutAngle = GetAngle();
+			return GetRotationAxisSafe(OutAxis, DefaultAxis, EpsilionSq);
+		}
+
+		/**
+		 * Extract the axis from the Quaternion.
+		 * @param OutAxis The axis of rotation.
+		 * @param DefaultAxis The axis to set when the angle is too small to accurately calculate the axis.
+		 * @param EpsilonSq The squared tolerance used to check for small angles.
+		 * @return Whether the axis was successfully calculated (true except for very small angles around or less than Epsilon).
+		 * @warning The axis calculation cannot succeed for small angles due to numerical error. In this case, the function will return false, but set the axis to DefaultAxis.
+		 * @note EpsilonSq is approximately the square of the angle below which we cannot calculate the axis. It needs to be "much greater" than square of the error in the
+		 * quaternion values which is usually ~1e-4, so values around 1e-3^2 = 1e-6 or greater are about right.
+		 */
+		inline bool GetRotationAxisSafe(TVector<FRealDouble, 3>& OutAxis, const TVector<FRealDouble, 3>& DefaultAxis, FRealDouble EpsilionSq = 1e-6f) const
+		{
+			// Tolerance must be much larger than error in normalized vector (usually ~1e-4) for the 
+			// axis calculation to succeed for small angles. For small angles, W ~= 1, and
+			// X, Y, Z ~= 0. If the values of X, Y, Z are around 1e-4 we are just normalizing error.
+			const FRealDouble LenSq = X * X + Y * Y + Z * Z;
+			if (LenSq > EpsilionSq)
+			{
+				FRealDouble InvLen = FMath::InvSqrt(LenSq);
+				OutAxis = TVector<FRealDouble, 3>(X * InvLen, Y * InvLen, Z * InvLen);
+				return true;
+			}
+
+			OutAxis = DefaultAxis;
+			return false;
+		}
+
+		/**
+		 * Extract the Swing and Twist rotations, assuming that the Twist Axis is (1,0,0).
+		 * /see ToSwingTwist
+		 */
+		inline void ToSwingTwistX(BaseQuat& OutSwing, BaseQuat& OutTwist) const
+		{
+			OutTwist = (X != 0.0f)? BaseQuat(X, 0, 0, W).GetNormalized() : BaseQuat::Identity;
+			OutSwing = *this * OutTwist.Inverse();
+		}
+
+		/**
+		 * @brief Return the large absolute element value
+		*/
+		FORCEINLINE FRealDouble GetAbsMax() const
+		{
+			return FMath::Max(FMath::Max(FMath::Abs(X), FMath::Abs(Y)), FMath::Max(FMath::Abs(Z), FMath::Abs(W)));
+		}
+
+		/**
+		 * @brief Return the dot product of two quaternions
+		*/
+		static inline FRealDouble DotProduct(const TRotation<FRealDouble, 3>& L, const TRotation<FRealDouble, 3>& R)
+		{
+			return (L | R);
+		}
+
+		/**
+		 * Return the complex conjugate of the rotation
+		 */
+		CHAOS_API static TRotation<FRealDouble, 3> Conjugate(const ::Chaos::TRotation<FRealDouble, 3>& InR);
+
+		/**
+		 * Negate all values of the quaternion (note: not the inverse rotation. See Conjugate)
+		 */
+		CHAOS_API static TRotation<FRealDouble, 3> Negate(const ::Chaos::TRotation<FRealDouble, 3>& InR);
+
+		/**
+		 * Create an identity rotation
+		 */
+		static inline TRotation<FRealDouble, 3> FromIdentity()
+		{
+			return BaseQuat(0, 0, 0, 1);
+		}
+
+		/**
+		 * Create a rotation by explicitly specifying all elements
+		 */
+		static inline TRotation<FRealDouble, 3> FromElements(const FRealDouble X, const FRealDouble Y, const FRealDouble Z, const FRealDouble W)
+		{
+			using FQuatReal = BaseQuat::FReal;
+			return BaseQuat((FQuatReal)X, (FQuatReal)Y, (FQuatReal)Z, (FQuatReal)W);
+		}
+
+		/**
+		 * Create a rotation by explicitly specifying all elements
+		 */
+		static inline TRotation<FRealDouble, 3> FromElements(const ::Chaos::TVector<FRealDouble, 3>& V, const FRealDouble W)
+		{
+			return FromElements(V.X, V.Y, V.Z, W);
+		}
+
+		/**
+		 * Create a rotation about an axis by an angle specified in radians
+		 */
+		static inline TRotation<FRealDouble, 3> FromAxisAngle(const ::Chaos::TVector<FRealDouble, 3>& Axis, const FRealDouble AngleRad)
+		{
+			return BaseQuat(UE::Math::TVector<FRealDouble>(Axis.X, Axis.Y, Axis.Z), (BaseQuat::FReal)AngleRad);
+		}
+
+		/**
+		 * Create a rotation about an axis V/|V| by angle |V| in radians
+		 */
+		CHAOS_API static TRotation<FRealDouble, 3> FromVector(const ::Chaos::TVector<FRealDouble, 3>& V);
+
+		/**
+		 * Generate a Rotation that would rotate vector InitialVector to FinalVector
+		 */
+		CHAOS_API static TRotation<FRealDouble, 3> FromRotatedVector(
+			const ::Chaos::TVector<FRealDouble, 3>& InitialVector,
+			const ::Chaos::TVector<FRealDouble, 3>& FinalVector);
+
+		/**
+		 * Calculate the angular velocity required to take an object with orientation R0 to orientation R1 in time Dt.
+		 *
+		 * Uses the relation: DQ/DT = (W * Q)/2
+		 */
+		CHAOS_API static TVector<FRealDouble, 3> CalculateAngularVelocity1(const TRotation<FRealDouble, 3>& R0, const TRotation<FRealDouble, 3>& InR1, const FRealDouble InDt);
+
+		/**
+		 * Calculate the angular velocity required to take an object with orientation R0 to orientation R1 in time Dt.
+		 *
+		 * Uses the Quaternion to Axis/Angle method.
+		 */
+		CHAOS_API static TVector<FRealDouble, 3> CalculateAngularVelocity2(const TRotation<FRealDouble, 3>& R0, const TRotation<FRealDouble, 3>& InR1, const FRealDouble InDt);
+
+		/**
+		 * Calculate the angular velocity required to take an object with orientation R0 to orientation R1 in time Dt.
+		 *
+		 * This should match the algorithm used in PerParticleUpdateFromDeltaPosition rule.
+		 */
+		static inline TVector<FRealDouble, 3> CalculateAngularVelocity(const TRotation<FRealDouble, 3>& InR0, const TRotation<FRealDouble, 3>& InR1, const FRealDouble InDt)
+		{
+			return CalculateAngularVelocity1(InR0, InR1, InDt);
+		}
+
+		/**
+		 * Calculate the axis-angle delta (angular velocity * dt) required to take an object with orientation R0 to orientation R1.
+		 *
+		 * This should match the algorithm used in PerParticleUpdateFromDeltaPosition rule.
+		 */
+		static inline TVector<FRealDouble, 3> CalculateAngularDelta(const TRotation<FRealDouble, 3>& InR0, const TRotation<FRealDouble, 3>& InR1)
+		{
+			return CalculateAngularVelocity(InR0, InR1, 1.0f);
+		}
+
+		/**
+		 * Return a new rotation equal to the input rotation with angular velocity W applied over time Dt.
+		 *
+		 * Uses the relation: DQ/DT = (W * Q)/2
+		 */
+		CHAOS_API static TRotation<FRealDouble, 3> IntegrateRotationWithAngularVelocity(const TRotation<FRealDouble, 3>& InR0, const TVector<FRealDouble, 3>& InW, const FRealDouble InDt);
+
+		/**
+		 * Check that two rotations are approximately equal. Assumes the quaternions are normalized and in the same hemisphere.
+		 * For small values of Epsilon, this is approximately equivalent to checking that the rotations are within 2*Epsilon
+		 * radians of each other.
+		 */
+		static inline bool IsNearlyEqual(const TRotation<FRealDouble, 3>& A, const TRotation<FRealDouble, 3>& B, const FRealDouble Epsilon)
+		{
+			// Only check imaginary part. This is comparing Epsilon to 2*AngleDelta for small angle deltas
+			return FMath::IsNearlyEqual(A.X, B.X, Epsilon)
+				&& FMath::IsNearlyEqual(A.Y, B.Y, Epsilon)
+				&& FMath::IsNearlyEqual(A.Z, B.Z, Epsilon);
 		}
 	};
 }
diff --git a/Engine/Source/Runtime/Experimental/Chaos/Public/Chaos/Transform.h b/Engine/Source/Runtime/Experimental/Chaos/Public/Chaos/Transform.h
--- a/Engine/Source/Runtime/Experimental/Chaos/Public/Chaos/Transform.h
+++ b/Engine/Source/Runtime/Experimental/Chaos/Public/Chaos/Transform.h
@@ -145,72 +145,185 @@ namespace Chaos
 	};
 
 	template<>
-	class TRigidTransform<FReal, 3> : public UE::Math::TTransform<FReal>
+	class TRigidTransform<FRealSingle, 3> : public UE::Math::TTransform<FRealSingle>
 	{
-		using BaseTransform = UE::Math::TTransform<FReal>;
+		using BaseTransform = UE::Math::TTransform<FRealSingle>;
 	public:
 		TRigidTransform()
 			: BaseTransform() {}
-		TRigidTransform(const TVector<FReal, 3>& Translation, const TRotation<FReal, 3>& Rotation)
+		TRigidTransform(const TVector<FRealSingle, 3>& Translation, const TRotation<FRealSingle, 3>& Rotation)
+			: BaseTransform(Rotation, Translation) {}
+		TRigidTransform(const TVector<FRealSingle, 3>& Translation, const TRotation<FRealSingle, 3>& Rotation, const TVector<FRealSingle, 3>& Scale)
+			: BaseTransform(Rotation, Translation, Scale) {}
+		TRigidTransform(const FMatrix44f& Matrix)
+			: BaseTransform(Matrix) {}
+		TRigidTransform(const FMatrix44d& Matrix)
+			: BaseTransform(FMatrix44f(Matrix)) {}
+		TRigidTransform(const BaseTransform& Transform)
+			: BaseTransform(Transform) {}
+		template<typename OtherType>
+		explicit TRigidTransform(const OtherType& Other)
+			: BaseTransform(TRotation<FRealSingle, 3>(Other.GetRotation()), TVector<FRealSingle, 3>(Other.GetTranslation())) {}
+
+		TRigidTransform<FRealSingle, 3> Inverse() const
+		{
+			return BaseTransform::Inverse();
+		}
+
+		PMatrix<FRealSingle, 4, 4> ToMatrixWithScale() const
+		{
+			return BaseTransform::ToMatrixWithScale();
+		}
+
+		PMatrix<FRealSingle, 4, 4> ToMatrixNoScale() const
+		{
+			return BaseTransform::ToMatrixNoScale();
+		}
+
+		// For low-level VectorRegister programming
+		const TPersistentVectorRegisterType<FRealSingle>& GetTranslationRegister() const { return Translation; }
+		const TPersistentVectorRegisterType<FRealSingle>& GetRotationRegister() const { return Rotation; }
+		void SetTranslationRegister(TransformVectorRegister InTranslation) { Translation = InTranslation; }
+		void SetRotationRegister(TransformVectorRegister InRotation) { Rotation = InRotation; }
+
+		CHAOS_API PMatrix<FRealSingle, 4, 4> operator*(const PMatrix<FRealSingle, 4, 4>& Matrix) const;
+		
+		inline TRigidTransform<FRealSingle, 3> operator*(const TRigidTransform<FRealSingle, 3>& Other) const
+		{
+			return BaseTransform::operator*(Other);
+		}
+
+		// Get the transform which maps from Other to This, ignoring the scale on both.
+		TRigidTransform<FRealSingle, 3> GetRelativeTransformNoScale(const TRigidTransform<FRealSingle, 3>& Other) const
+		{
+			// @todo(chaos): optimize
+			TRotation<FRealSingle, 3> OtherInverse = Other.GetRotation().Inverse();
+			return TRigidTransform<FRealSingle, 3>(
+				(OtherInverse * (GetTranslation() - Other.GetTranslation())),
+				OtherInverse * GetRotation());
+		}
+
+		TVector<FRealSingle, 3> TransformNormalNoScale(const TVector<FRealSingle, 3>& Normal) const
+		{
+			return TransformVectorNoScale(Normal);
+		}
+
+		// Transform the normal when scale may be non-unitary. Assumes no scale components are zero.
+		TVector<FRealSingle, 3> TransformNormalUnsafe(const TVector<FRealSingle, 3>& Normal) const
+		{
+			const TVector<FRealSingle, 3> RotatedNormal = TransformNormalNoScale(Normal);
+			const TVector<FRealSingle, 3> ScaledNormal = RotatedNormal / GetScale3D();
+			const FRealSingle ScaledNormal2 = ScaledNormal.SizeSquared();
+			if (ScaledNormal2 > SMALL_NUMBER)
+			{
+				return ScaledNormal * FMath::InvSqrt(ScaledNormal2);
+			}
+			else
+			{
+				return RotatedNormal;
+			}
+		}
+
+		/**
+		 * @brief Equivalent to (A * B) but assuming both have unit scale
+		*/
+		static TRigidTransform<FRealSingle, 3> MultiplyNoScale(const TRigidTransform<FRealSingle, 3>& A, const TRigidTransform<FRealSingle, 3>& B)
+		{
+			TRigidTransform<FRealSingle, 3> Result;
+
+#if ENABLE_VECTORIZED_TRANSFORM
+			const TransformVectorRegister QuatA = A.Rotation;
+			const TransformVectorRegister QuatB = B.Rotation;
+			const TransformVectorRegister TranslateA = A.Translation;
+			const TransformVectorRegister TranslateB = B.Translation;
+
+			const TransformVectorRegister Rotation = VectorQuaternionMultiply2(QuatB, QuatA);
+			const TransformVectorRegister RotatedTranslate = VectorQuaternionRotateVector(QuatB, TranslateA);
+			const TransformVectorRegister Translation = VectorAdd(RotatedTranslate, TranslateB);
+
+			Result.Rotation = Rotation;
+			Result.Translation = Translation;
+			Result.Scale3D = VectorOne();
+#else
+			Result.Rotation = B.Rotation * A.Rotation;
+			Result.Translation = B.Rotation * A.Translation + B.Translation;
+			Result.Scale3D = FVector::OneVector;
+#endif
+
+			return Result;
+		}
+	};
+
+	template<>
+	class TRigidTransform<FRealDouble, 3> : public UE::Math::TTransform<FRealDouble>
+	{
+		using BaseTransform = UE::Math::TTransform<FRealDouble>;
+	public:
+		TRigidTransform()
+			: BaseTransform() {}
+		TRigidTransform(const TVector<FRealDouble, 3>& Translation, const TRotation<FRealDouble, 3>& Rotation)
 			: BaseTransform(Rotation, Translation) {}
-		TRigidTransform(const TVector<FReal, 3>& Translation, const TRotation<FReal, 3>& Rotation, const TVector<FReal, 3>& Scale)
+		TRigidTransform(const TVector<FRealDouble, 3>& Translation, const TRotation<FRealDouble, 3>& Rotation, const TVector<FRealDouble, 3>& Scale)
 			: BaseTransform(Rotation, Translation, Scale) {}
 		TRigidTransform(const FMatrix44d& Matrix)
 			: BaseTransform(Matrix) {}
 		TRigidTransform(const FMatrix44f& Matrix)
-			: BaseTransform(FMatrix(Matrix)) {}
+			: BaseTransform(FMatrix44d(Matrix)) {}
 		TRigidTransform(const BaseTransform& Transform)
 			: BaseTransform(Transform) {}
-		TRigidTransform<FReal, 3> Inverse() const
+		template<typename OtherType>
+		explicit TRigidTransform(const OtherType& Other)
+			: BaseTransform(TRotation<FRealDouble, 3>(Other.GetRotation()), TVector<FRealDouble, 3>(Other.GetTranslation())) {}
+
+		TRigidTransform<FRealDouble, 3> Inverse() const
 		{
 			return BaseTransform::Inverse();
 		}
 
-		PMatrix<FReal, 4, 4> ToMatrixWithScale() const
+		PMatrix<FRealDouble, 4, 4> ToMatrixWithScale() const
 		{
 			return BaseTransform::ToMatrixWithScale();
 		}
 
-		PMatrix<FReal, 4, 4> ToMatrixNoScale() const
+		PMatrix<FRealDouble, 4, 4> ToMatrixNoScale() const
 		{
 			return BaseTransform::ToMatrixNoScale();
 		}
 
 		// For low-level VectorRegister programming
-		const TPersistentVectorRegisterType<FReal>& GetTranslationRegister() const { return Translation; }
-		const TPersistentVectorRegisterType<FReal>& GetRotationRegister() const { return Rotation; }
+		const TPersistentVectorRegisterType<FRealDouble>& GetTranslationRegister() const { return Translation; }
+		const TPersistentVectorRegisterType<FRealDouble>& GetRotationRegister() const { return Rotation; }
 		void SetTranslationRegister(TransformVectorRegister InTranslation) { Translation = InTranslation; }
 		void SetRotationRegister(TransformVectorRegister InRotation) { Rotation = InRotation; }
 
-
-		CHAOS_API Chaos::PMatrix<Chaos::FReal, 4, 4> operator*(const Chaos::PMatrix<Chaos::FReal, 4, 4>& Matrix) const;
+		CHAOS_API PMatrix<FRealDouble, 4, 4> operator*(const Chaos::PMatrix<FRealDouble, 4, 4>& Matrix) const;
 		
-		inline TRigidTransform<FReal, 3> operator*(const TRigidTransform<FReal, 3>& Other) const
+		inline TRigidTransform<FRealDouble, 3> operator*(const TRigidTransform<FRealDouble, 3>& Other) const
 		{
 			return BaseTransform::operator*(Other);
 		}
 
 		// Get the transform which maps from Other to This, ignoring the scale on both.
-		TRigidTransform<FReal, 3> GetRelativeTransformNoScale(const TRigidTransform<FReal, 3>& Other) const
+		TRigidTransform<FRealDouble, 3> GetRelativeTransformNoScale(const TRigidTransform<FRealDouble, 3>& Other) const
 		{
 			// @todo(chaos): optimize
-			TRotation<FReal, 3> OtherInverse = Other.GetRotation().Inverse();
-			return TRigidTransform<FReal, 3>(
+			TRotation<FRealDouble, 3> OtherInverse = Other.GetRotation().Inverse();
+			return TRigidTransform<FRealDouble, 3>(
 				(OtherInverse * (GetTranslation() - Other.GetTranslation())),
 				OtherInverse * GetRotation());
 		}
 
-		TVector<FReal, 3> TransformNormalNoScale(const TVector<FReal, 3>& Normal) const
+		TVector<FRealDouble, 3> TransformNormalNoScale(const TVector<FRealDouble, 3>& Normal) const
 		{
 			return TransformVectorNoScale(Normal);
 		}
 
 		// Transform the normal when scale may be non-unitary. Assumes no scale components are zero.
-		TVector<FReal, 3> TransformNormalUnsafe(const TVector<FReal, 3>& Normal) const
+		TVector<FRealDouble, 3> TransformNormalUnsafe(const TVector<FRealDouble, 3>& Normal) const
 		{
-			const TVector<FReal, 3> RotatedNormal = TransformNormalNoScale(Normal);
-			const TVector<FReal, 3> ScaledNormal = RotatedNormal / GetScale3D();
-			const FReal ScaledNormal2 = ScaledNormal.SizeSquared();
+			const TVector<FRealDouble, 3> RotatedNormal = TransformNormalNoScale(Normal);
+			const TVector<FRealDouble, 3> ScaledNormal = RotatedNormal / GetScale3D();
+			const FRealDouble ScaledNormal2 = ScaledNormal.SizeSquared();
 			if (ScaledNormal2 > SMALL_NUMBER)
 			{
 				return ScaledNormal * FMath::InvSqrt(ScaledNormal2);
@@ -224,9 +337,9 @@ namespace Chaos
 		/**
 		 * @brief Equivalent to (A * B) but assuming both have unit scale
 		*/
-		static TRigidTransform<FReal, 3> MultiplyNoScale(const TRigidTransform<FReal, 3>& A, const TRigidTransform<FReal, 3>& B)
+		static TRigidTransform<FRealDouble, 3> MultiplyNoScale(const TRigidTransform<FRealDouble, 3>& A, const TRigidTransform<FRealDouble, 3>& B)
 		{
-			TRigidTransform<FReal, 3> Result;
+			TRigidTransform<FRealDouble, 3> Result;
 
 #if ENABLE_VECTORIZED_TRANSFORM
 			const TransformVectorRegister QuatA = A.Rotation;
@@ -252,7 +365,12 @@ namespace Chaos
 	};
 }
 
-inline uint32 GetTypeHash(const Chaos::TRigidTransform<Chaos::FReal, 3>& InTransform)
+inline uint32 GetTypeHash(const Chaos::TRigidTransform<Chaos::FRealSingle, 3>& InTransform)
+{
+	return HashCombine(GetTypeHash(InTransform.GetTranslation()), HashCombine(GetTypeHash(InTransform.GetRotation().Euler()), GetTypeHash(InTransform.GetScale3D())));
+}
+
+inline uint32 GetTypeHash(const Chaos::TRigidTransform<Chaos::FRealDouble, 3>& InTransform)
 {
 	return HashCombine(GetTypeHash(InTransform.GetTranslation()), HashCombine(GetTypeHash(InTransform.GetRotation().Euler()), GetTypeHash(InTransform.GetScale3D())));
 }
diff --git a/Engine/Source/Runtime/Experimental/Chaos/Public/Chaos/TriangleMesh.h b/Engine/Source/Runtime/Experimental/Chaos/Public/Chaos/TriangleMesh.h
--- a/Engine/Source/Runtime/Experimental/Chaos/Public/Chaos/TriangleMesh.h
+++ b/Engine/Source/Runtime/Experimental/Chaos/Public/Chaos/TriangleMesh.h
@@ -137,22 +137,25 @@ namespace Chaos
 		/** The GetFaceNormals functions assume Counter Clockwise triangle windings in a Left Handed coordinate system
 			If this is not the case the returned face normals may be inverted
 		*/
-		CHAOS_API TArray<FVec3> GetFaceNormals(const TConstArrayView<FVec3>& Points, const bool ReturnEmptyOnError = true) const;
-		CHAOS_API void GetFaceNormals(TArray<FVec3>& Normals, const TConstArrayView<FVec3>& Points, const bool ReturnEmptyOnError = true) const;
+		template <typename T>
+		TArray<TVec3<T>> GetFaceNormals(const TConstArrayView<TVec3<T>>& Points, const bool ReturnEmptyOnError = true) const;
+		template <typename T>
+		void GetFaceNormals(TArray<TVec3<T>>& Normals, const TConstArrayView<TVec3<T>>& Points, const bool ReturnEmptyOnError = true) const;
 		FORCEINLINE TArray<FVec3> GetFaceNormals(const FParticles& InParticles, const bool ReturnEmptyOnError = true) const
-		{ return GetFaceNormals(InParticles.X(), ReturnEmptyOnError); }
+		{ return GetFaceNormals(TConstArrayView<FVec3>(InParticles.X()), ReturnEmptyOnError); }
 
 		CHAOS_API TArray<FVec3> GetPointNormals(const TConstArrayView<FVec3>& points, const bool ReturnEmptyOnError = true);
 		FORCEINLINE TArray<FVec3> GetPointNormals(const FParticles& InParticles, const bool ReturnEmptyOnError = true)
-		{ return GetPointNormals(InParticles.X(), ReturnEmptyOnError); }
-
+		{ return GetPointNormals(TConstArrayView<FVec3>(InParticles.X()), ReturnEmptyOnError); }
+		
 		CHAOS_API void GetPointNormals(TArrayView<FVec3> PointNormals, const TConstArrayView<FVec3>& FaceNormals, const bool bUseGlobalArray);
 		/** \brief Get per-point normals. 
 		 * This const version of this function requires \c GetPointToTriangleMap() 
 		 * to be called prior to invoking this function. 
 		 * @param bUseGlobalArray When true, fill the array from the StartIdx to StartIdx + NumIndices - 1 positions, otherwise fill the array from the 0 to NumIndices - 1 positions.
 		 */
-		CHAOS_API void GetPointNormals(TArrayView<FVec3> PointNormals, const TConstArrayView<FVec3>& FaceNormals, const bool bUseGlobalArray) const;
+		template <typename T>
+		void GetPointNormals(TArrayView<TVec3<T>> PointNormals, const TConstArrayView<TVec3<T>>& FaceNormals, const bool bUseGlobalArray) const;
 
 		static CHAOS_API FTriangleMesh GetConvexHullFromParticles(const TConstArrayView<FVec3>& points);
 		/** Deprecated. Use TArrayView version. */
diff --git a/Engine/Source/Runtime/Experimental/Chaos/Public/Chaos/VelocityField.h b/Engine/Source/Runtime/Experimental/Chaos/Public/Chaos/VelocityField.h
--- a/Engine/Source/Runtime/Experimental/Chaos/Public/Chaos/VelocityField.h
+++ b/Engine/Source/Runtime/Experimental/Chaos/Public/Chaos/VelocityField.h
@@ -1,161 +1,112 @@
 // Copyright Epic Games, Inc. All Rights Reserved.
 #pragma once
 
+#include "Chaos/PBDSoftsEvolutionFwd.h"
 #include "Chaos/PBDParticles.h"
 #include "Chaos/TriangleMesh.h"
 
-namespace Chaos
+namespace Chaos::Softs
 {
-	// Velocity field basic implementation
-	// TODO: Add lift
-	class CHAOS_API FVelocityField final
-	{
-	public:
-		static constexpr FReal DefaultDragCoefficient = (FReal)0.5;
-		static constexpr FReal DefaultLiftCoefficient = (FReal)0.1;
-		static constexpr FReal DefaultFluidDensity = (FReal)1.225e-6;
-
-		// Construct an uninitialized field. Mesh, properties, and velocity will have to be set for this field to be valid.
-		FVelocityField()
-			: Offset(INDEX_NONE)
-			, NumParticles(0)
-		{
-			SetProperties(FVec2::ZeroVector, FVec2::ZeroVector, (FReal)0.);
-		}
-
-		// Construct a uniform field.
-		UE_DEPRECATED(4.27, "Use Chaos fields instead.")
-		FVelocityField(
-			const FTriangleMesh& TriangleMesh,
-			const FVec3& InVelocity,
-			const FReal InDragCoefficient = DefaultDragCoefficient,
-			const FReal InLiftCoefficient = DefaultLiftCoefficient,
-			const FReal InFluidDensity = DefaultFluidDensity)
-			: PointToTriangleMap(TriangleMesh.GetPointToTriangleMap())
-			, Elements(TriangleMesh.GetElements())
-			, Velocity(InVelocity)
-			, Offset(TriangleMesh.GetVertexRange()[0])
-			, NumParticles(TriangleMesh.GetVertexRange()[1] - TriangleMesh.GetVertexRange()[0] + 1)
-		{
-			Forces.SetNumUninitialized(Elements.Num());
-			SetProperties(InDragCoefficient, InLiftCoefficient, InFluidDensity);
-		}
 
-		// Construct a vector field.
-		UE_DEPRECATED(4.27, "Use Chaos fields instead.")
-		FVelocityField(
-			const FTriangleMesh& TriangleMesh,
-			TFunction<FVec3(const FVec3&)> GetVelocity,
-			const FReal InDragCoefficient = DefaultDragCoefficient,
-			const FReal InLiftCoefficient = DefaultLiftCoefficient,
-			const FReal InFluidDensity = DefaultFluidDensity)
-			: PointToTriangleMap(TriangleMesh.GetPointToTriangleMap())
-			, Elements(TriangleMesh.GetElements())
-			, Velocity(GetVelocity(FVec3::ZeroVector))
-			, Offset(TriangleMesh.GetVertexRange()[0])
-			, NumParticles(TriangleMesh.GetVertexRange()[1] - TriangleMesh.GetVertexRange()[0] + 1)
-		{
-			Forces.SetNumUninitialized(Elements.Num());
-			SetProperties(InDragCoefficient, InLiftCoefficient, InFluidDensity);
-		}
-
-		~FVelocityField() {}
-
-		void UpdateForces(const FPBDParticles& InParticles, const FReal /*Dt*/);
-
-		inline void Apply(FPBDParticles& InParticles, const FReal Dt, const int32 Index) const
-		{
-			checkSlow(Index >= Offset && Index < Offset + NumParticles);  // The index should always match the original triangle mesh range
-
-			const TArray<int32>& ElementIndices = PointToTriangleMap[Index];
-			for (const int32 ElementIndex : ElementIndices)
-			{
-				InParticles.Acceleration(Index) += InParticles.InvM(Index) * Forces[ElementIndex];
-			}
-		}
+// Velocity field used solely for aerodynamics effects, use Chaos Fields for other types of fields.
+class CHAOS_API FVelocityField final
+{
+public:
+	static constexpr FSolverReal DefaultDragCoefficient = (FSolverReal)0.5;
+	static constexpr FSolverReal DefaultLiftCoefficient = (FSolverReal)0.1;
+	static constexpr FSolverReal DefaultFluidDensity = (FSolverReal)1.225e-6;
+
+	// Construct an uninitialized field. Mesh, properties, and velocity will have to be set for this field to be valid.
+	FVelocityField()
+		: Offset(INDEX_NONE)
+		, NumParticles(0)
+	{
+		SetProperties(FSolverVec2(0.), FSolverVec2(0.), (FSolverReal)0.);
+	}
 
-		UE_DEPRECATED(4.27, "Use SetProperties instead.")
-		void SetFluidDensity(const FReal InFluidDensity)
-		{
-			QuarterRho = (FReal)0.25 * InFluidDensity;
-		}
+	~FVelocityField() {}
 
-		UE_DEPRECATED(4.27, "Use SetProperties instead.")
-		void SetCoefficients(const FReal InDragCoefficient, const FReal InLiftCoefficient)
-		{
-			SetProperties(FVec2(InDragCoefficient), FVec2(InLiftCoefficient), (FReal)4. * QuarterRho);
-		}
+	void UpdateForces(const FSolverParticles& InParticles, const FSolverReal /*Dt*/);
 
-		void SetProperties(const FVec2& Drag, const FVec2& Lift, const FReal FluidDensity)
-		{
-			constexpr FReal OneQuarter = (FReal)0.25;
-			QuarterRho = FluidDensity * OneQuarter;
-
-			constexpr FReal MinCoefficient = (FReal)0.;
-			constexpr FReal MaxCoefficient = (FReal)10.;
-			DragBase = FMath::Clamp(Drag[0], MinCoefficient, MaxCoefficient);
-			DragRange = FMath::Clamp(Drag[1], MinCoefficient, MaxCoefficient) - DragBase;
-			LiftBase = FMath::Clamp(Lift[0], MinCoefficient, MaxCoefficient);
-			LiftRange = FMath::Clamp(Lift[1], MinCoefficient, MaxCoefficient) - LiftBase;
-		}
+	inline void Apply(FSolverParticles& InParticles, const FSolverReal Dt, const int32 Index) const
+	{
+		checkSlow(Index >= Offset && Index < Offset + NumParticles);  // The index should always match the original triangle mesh range
 
-		bool IsActive() const 
+		const TArray<int32>& ElementIndices = PointToTriangleMap[Index];
+		for (const int32 ElementIndex : ElementIndices)
 		{
-			return (DragBase > (FReal)0. || DragRange != (FReal)0.) || (LiftBase > (FReal)0. || LiftRange != (FReal)0.);  // Note: range can be a negative value (although not when base is zero)
+			InParticles.Acceleration(Index) += InParticles.InvM(Index) * Forces[ElementIndex];
 		}
+	}
 
-		void SetGeometry(const FTriangleMesh* TriangleMesh, const TConstArrayView<FRealSingle>& DragMultipliers, const TConstArrayView<FRealSingle>& LiftMultipliers);
+	void SetProperties(const FSolverVec2& Drag, const FSolverVec2& Lift, const FSolverReal FluidDensity)
+	{
+		constexpr FSolverReal OneQuarter = (FSolverReal)0.25;
+		QuarterRho = FluidDensity * OneQuarter;
+
+		constexpr FSolverReal MinCoefficient = (FSolverReal)0.;
+		constexpr FSolverReal MaxCoefficient = (FSolverReal)10.;
+		DragBase = FMath::Clamp(Drag[0], MinCoefficient, MaxCoefficient);
+		DragRange = FMath::Clamp(Drag[1], MinCoefficient, MaxCoefficient) - DragBase;
+		LiftBase = FMath::Clamp(Lift[0], MinCoefficient, MaxCoefficient);
+		LiftRange = FMath::Clamp(Lift[1], MinCoefficient, MaxCoefficient) - LiftBase;
+	}
+
+	bool IsActive() const 
+	{
+		return (DragBase > (FSolverReal)0. || DragRange != (FSolverReal)0.) || (LiftBase > (FSolverReal)0. || LiftRange != (FSolverReal)0.);  // Note: range can be a negative value (although not when base is zero)
+	}
 
-		void SetVelocity(const FVec3& InVelocity) { Velocity = InVelocity; }
+	void SetGeometry(const FTriangleMesh* TriangleMesh, const TConstArrayView<FRealSingle>& DragMultipliers, const TConstArrayView<FRealSingle>& LiftMultipliers);
 
-		UE_DEPRECATED(4.27, "Use SetVeloccity(const FVec3&) instead.")
-		void SetVelocity(TFunction<FVec3(const FVec3&)> GetVelocity)  { Velocity = GetVelocity(FVec3::ZeroVector); }
+	void SetVelocity(const FSolverVec3& InVelocity) { Velocity = InVelocity; }
 
-		const TConstArrayView<TVector<int32, 3>>& GetElements() const { return Elements; }
-		TConstArrayView<FVec3> GetForces() const { return TConstArrayView<FVec3>(Forces); }
+	const TConstArrayView<TVector<int32, 3>>& GetElements() const { return Elements; }
+	TConstArrayView<FSolverVec3> GetForces() const { return TConstArrayView<FSolverVec3>(Forces); }
 
-	private:
-		void UpdateField(const FPBDParticles& InParticles, int32 ElementIndex, const FVec3& InVelocity, const FReal Cd, const FReal Cl)
-		{
-			const TVec3<int32>& Element = Elements[ElementIndex];
-
-			// Calculate the normal and the area of the surface exposed to the flow
-			FVec3 N = FVec3::CrossProduct(
-				InParticles.X(Element[1]) - InParticles.X(Element[0]),
-				InParticles.X(Element[2]) - InParticles.X(Element[0]));
-			const FReal DoubleArea = N.SafeNormalize();
-
-			// Calculate the direction and the relative velocity of the triangle to the flow
-			const FVec3& SurfaceVelocity = (FReal)(1. / 3.) * (
-				InParticles.V(Element[0]) +
-				InParticles.V(Element[1]) +
-				InParticles.V(Element[2]));
-			const FVec3 V = InVelocity - SurfaceVelocity;
-
-			// Set the aerodynamic forces
-			const FReal VDotN = FVec3::DotProduct(V, N);
-			const FReal VSquare = FVec3::DotProduct(V, V);
-
-			Forces[ElementIndex] = QuarterRho * DoubleArea * (VDotN >= (FReal)0. ?  // The flow can hit either side of the triangle, so the normal might need to be reversed
-				(Cd - Cl) * VDotN * V + Cl * VSquare * N :
-				(Cl - Cd) * VDotN * V - Cl * VSquare * N);
-		}
-
-	private:
-		TConstArrayView<TArray<int32>> PointToTriangleMap;
-		TConstArrayView<TVec3<int32>> Elements;
-		TArray<FVec3> Forces;
-		TArray<FVec2> Multipliers;
-		FVec3 Velocity;
-		FReal DragBase;
-		FReal DragRange;
-		FReal LiftBase;
-		FReal LiftRange;
-		FReal QuarterRho;
-		int32 Offset;
-		int32 NumParticles;
-	};
-}
+private:
+	void UpdateField(const FSolverParticles& InParticles, int32 ElementIndex, const FSolverVec3& InVelocity, const FSolverReal Cd, const FSolverReal Cl)
+	{
+		const TVec3<int32>& Element = Elements[ElementIndex];
+
+		// Calculate the normal and the area of the surface exposed to the flow
+		FSolverVec3 N = FSolverVec3::CrossProduct(
+			InParticles.X(Element[1]) - InParticles.X(Element[0]),
+			InParticles.X(Element[2]) - InParticles.X(Element[0]));
+		const FSolverReal DoubleArea = N.SafeNormalize();
+
+		// Calculate the direction and the relative velocity of the triangle to the flow
+		const FSolverVec3& SurfaceVelocity = (FSolverReal)(1. / 3.) * (
+			InParticles.V(Element[0]) +
+			InParticles.V(Element[1]) +
+			InParticles.V(Element[2]));
+		const FSolverVec3 V = InVelocity - SurfaceVelocity;
+
+		// Set the aerodynamic forces
+		const FSolverReal VDotN = FSolverVec3::DotProduct(V, N);
+		const FSolverReal VSquare = FSolverVec3::DotProduct(V, V);
+
+		Forces[ElementIndex] = QuarterRho * DoubleArea * (VDotN >= (FSolverReal)0. ?  // The flow can hit either side of the triangle, so the normal might need to be reversed
+			(Cd - Cl) * VDotN * V + Cl * VSquare * N :
+			(Cl - Cd) * VDotN * V - Cl * VSquare * N);
+	}
+
+private:
+	TConstArrayView<TArray<int32>> PointToTriangleMap;
+	TConstArrayView<TVec3<int32>> Elements;
+	TArray<FSolverVec3> Forces;
+	TArray<FSolverVec2> Multipliers;
+	FSolverVec3 Velocity;
+	FSolverReal DragBase;
+	FSolverReal DragRange;
+	FSolverReal LiftBase;
+	FSolverReal LiftRange;
+	FSolverReal QuarterRho;
+	int32 Offset;
+	int32 NumParticles;
+};
+
+}  // End namespace Chaos::Softs
 
 // Support ISPC enable/disable in non-shipping builds
 #if !INTEL_ISPC
diff --git a/Engine/Source/Runtime/Experimental/Chaos/Public/Chaos/XPBDAxialSpringConstraints.h b/Engine/Source/Runtime/Experimental/Chaos/Public/Chaos/XPBDAxialSpringConstraints.h
--- a/Engine/Source/Runtime/Experimental/Chaos/Public/Chaos/XPBDAxialSpringConstraints.h
+++ b/Engine/Source/Runtime/Experimental/Chaos/Public/Chaos/XPBDAxialSpringConstraints.h
@@ -6,13 +6,13 @@
 
 DECLARE_CYCLE_STAT(TEXT("Chaos XPBD Axial Spring Constraint"), STAT_XPBD_AxialSpring, STATGROUP_Chaos);
 
-namespace Chaos
+namespace Chaos::Softs
 {
 
 // Stiffness is in N/CM^2, so it needs to be adjusted from the PBD stiffness ranging between [0,1]
 static const double XPBDAxialSpringMaxCompliance = 1e-7;  // Max stiffness: 1e+11 N/M^2 = 1e+7 N/CM^2 -> Max compliance: 1e-7 CM^2/N
 
-class FXPBDAxialSpringConstraints : public FPBDAxialSpringConstraintsBase
+class FXPBDAxialSpringConstraints final : public FPBDAxialSpringConstraintsBase
 {
 	typedef FPBDAxialSpringConstraintsBase Base;
 	using Base::Barys;
@@ -22,36 +22,36 @@ class FXPBDAxialSpringConstraints : public FPBDAxialSpringConstraintsBase
 
 public:
 	FXPBDAxialSpringConstraints(
-		const FPBDParticles& Particles,
+		const FSolverParticles& Particles,
 		int32 ParticleOffset,
 		int32 ParticleCount,
 		const TArray<TVec3<int32>>& InConstraints,
 		const TConstArrayView<FRealSingle>& StiffnessMultipliers,
-		const FVec2& InStiffness,
+		const FSolverVec2& InStiffness,
 		bool bTrimKinematicConstraints)
 		: Base(Particles, ParticleOffset, ParticleCount, InConstraints, StiffnessMultipliers, InStiffness, bTrimKinematicConstraints)
 	{
 		Lambdas.Init(0.f, Constraints.Num());
 	}
 
-	virtual ~FXPBDAxialSpringConstraints() {}
+	virtual ~FXPBDAxialSpringConstraints() override {}
 
-	void Init() const { for (FReal& Lambda : Lambdas) { Lambda = (FReal)0.; } }
+	void Init() const { for (FSolverReal& Lambda : Lambdas) { Lambda = (FSolverReal)0.; } }
 
-	void Apply(FPBDParticles& Particles, const FReal Dt) const
+	void Apply(FSolverParticles& Particles, const FSolverReal Dt) const
 	{
 		SCOPE_CYCLE_COUNTER(STAT_XPBD_AxialSpring);
 		if (!Stiffness.HasWeightMap())
 		{
-			const FReal ExpStiffnessValue = (FReal)Stiffness;
+			const FSolverReal ExpStiffnessValue = (FSolverReal)Stiffness;
 			for (int32 ConstraintIndex = 0; ConstraintIndex < Constraints.Num(); ++ConstraintIndex)
 			{
 				const TVector<int32, 3>& constraint = Constraints[ConstraintIndex];
 				const int32 i1 = constraint[0];
 				const int32 i2 = constraint[1];
 				const int32 i3 = constraint[2];
-				const FVec3 Delta = GetDelta(Particles, Dt, ConstraintIndex, ExpStiffnessValue);
-				const FReal Multiplier = (FReal)2. / (FMath::Max(Barys[ConstraintIndex], (FReal)1. - Barys[ConstraintIndex]) + (FReal)1.);
+				const FSolverVec3 Delta = GetDelta(Particles, Dt, ConstraintIndex, ExpStiffnessValue);
+				const FSolverReal Multiplier = (FSolverReal)2. / (FMath::Max(Barys[ConstraintIndex], (FSolverReal)1. - Barys[ConstraintIndex]) + (FSolverReal)1.);
 				if (Particles.InvM(i1) > 0)
 				{
 					Particles.P(i1) -= Multiplier * Particles.InvM(i1) * Delta;
@@ -62,7 +62,7 @@ public:
 				}
 				if (Particles.InvM(i3))
 				{
-					Particles.P(i3) += Multiplier * Particles.InvM(i3) * ((FReal)1. - Barys[ConstraintIndex]) * Delta;
+					Particles.P(i3) += Multiplier * Particles.InvM(i3) * ((FSolverReal)1. - Barys[ConstraintIndex]) * Delta;
 				}
 			}
 		}
@@ -70,13 +70,13 @@ public:
 		{
 			for (int32 ConstraintIndex = 0; ConstraintIndex < Constraints.Num(); ++ConstraintIndex)
 			{
-				const FReal ExpStiffnessValue = Stiffness[ConstraintIndex];
+				const FSolverReal ExpStiffnessValue = Stiffness[ConstraintIndex];
 				const TVector<int32, 3>& constraint = Constraints[ConstraintIndex];
 				const int32 i1 = constraint[0];
 				const int32 i2 = constraint[1];
 				const int32 i3 = constraint[2];
-				const FVec3 Delta = GetDelta(Particles, Dt, ConstraintIndex, ExpStiffnessValue);
-				const FReal Multiplier = (FReal)2. / (FMath::Max(Barys[ConstraintIndex], (FReal)1. - Barys[ConstraintIndex]) + (FReal)1.);
+				const FSolverVec3 Delta = GetDelta(Particles, Dt, ConstraintIndex, ExpStiffnessValue);
+				const FSolverReal Multiplier = (FSolverReal)2. / (FMath::Max(Barys[ConstraintIndex], (FSolverReal)1. - Barys[ConstraintIndex]) + (FSolverReal)1.);
 				if (Particles.InvM(i1) > 0)
 				{
 					Particles.P(i1) -= Multiplier * Particles.InvM(i1) * Delta;
@@ -87,57 +87,55 @@ public:
 				}
 				if (Particles.InvM(i3))
 				{
-					Particles.P(i3) += Multiplier * Particles.InvM(i3) * ((FReal)1. - Barys[ConstraintIndex]) * Delta;
+					Particles.P(i3) += Multiplier * Particles.InvM(i3) * ((FSolverReal)1. - Barys[ConstraintIndex]) * Delta;
 				}
 			}
 		}
 	}
 
 private:
-	FVec3 GetDelta(const FPBDParticles& Particles, const FReal Dt, const int32 InConstraintIndex, const FReal ExpStiffnessValue) const
+	FSolverVec3 GetDelta(const FSolverParticles& Particles, const FSolverReal Dt, const int32 InConstraintIndex, const FSolverReal ExpStiffnessValue) const
 	{
 		const TVector<int32, 3>& Constraint = Constraints[InConstraintIndex];
 		const int32 i1 = Constraint[0];
 		const int32 i2 = Constraint[1];
 		const int32 i3 = Constraint[2];
 
-		const FReal Bary = Barys[InConstraintIndex];
-		const FReal PInvMass = Particles.InvM(i3) * ((FReal)1. - Bary) + Particles.InvM(i2) * Bary;
-		if (Particles.InvM(i1) == (FReal)0. && PInvMass == (FReal)0.)
+		const FSolverReal Bary = Barys[InConstraintIndex];
+		const FSolverReal PInvMass = Particles.InvM(i3) * ((FSolverReal)1. - Bary) + Particles.InvM(i2) * Bary;
+		if (Particles.InvM(i1) == (FSolverReal)0. && PInvMass == (FSolverReal)0.)
 		{
-			return FVec3((FReal)0.);
+			return FSolverVec3((FSolverReal)0.);
 		}
-		const FReal CombinedInvMass = PInvMass + Particles.InvM(i1);
-		ensure(CombinedInvMass > (FReal)SMALL_NUMBER);
+		const FSolverReal CombinedInvMass = PInvMass + Particles.InvM(i1);
+		ensure(CombinedInvMass > (FSolverReal)SMALL_NUMBER);
 
-		const FVec3& P1 = Particles.P(i1);
-		const FVec3& P2 = Particles.P(i2);
-		const FVec3& P3 = Particles.P(i3);
-		const FVec3 P = (P2 - P3) * Bary + P3;
+		const FSolverVec3& P1 = Particles.P(i1);
+		const FSolverVec3& P2 = Particles.P(i2);
+		const FSolverVec3& P3 = Particles.P(i3);
+		const FSolverVec3 P = (P2 - P3) * Bary + P3;
 
-		const FVec3 Difference = P1 - P;
-		const FReal Distance = Difference.Size();
+		const FSolverVec3 Difference = P1 - P;
+		const FSolverReal Distance = Difference.Size();
 		if (UNLIKELY(Distance <= SMALL_NUMBER))
 		{
-			return FVec3((FReal)0.);
+			return FSolverVec3((FSolverReal)0.);
 		}
-		const FVec3 Direction = Difference / Distance;
-		const FReal Offset = (Distance - Dists[InConstraintIndex]);
+		const FSolverVec3 Direction = Difference / Distance;
+		const FSolverReal Offset = (Distance - Dists[InConstraintIndex]);
 
-		FReal& Lambda = Lambdas[InConstraintIndex];
-		const FReal Alpha = (FReal)XPBDAxialSpringMaxCompliance / (ExpStiffnessValue * Dt * Dt);
+		FSolverReal& Lambda = Lambdas[InConstraintIndex];
+		const FSolverReal Alpha = (FSolverReal)XPBDAxialSpringMaxCompliance / (ExpStiffnessValue * Dt * Dt);
 
-		const FReal DLambda = (Offset - Alpha * Lambda) / (CombinedInvMass + Alpha);
-		const FVec3 Delta = DLambda * Direction;
+		const FSolverReal DLambda = (Offset - Alpha * Lambda) / (CombinedInvMass + Alpha);
+		const FSolverVec3 Delta = DLambda * Direction;
 		Lambda += DLambda;
 
 		return Delta;
 	}
 
 private:
-	mutable TArray<FReal> Lambdas;
+	mutable TArray<FSolverReal> Lambdas;
 };
 
-template<class T, int d>
-using TXPBDAxialSpringConstraints UE_DEPRECATED(4.27, "Deprecated. this class is to be deleted, use FXPBDAxialSpringConstraints instead") = FXPBDAxialSpringConstraints;
-}
+}  // End namespace Chaos::Softs
diff --git a/Engine/Source/Runtime/Experimental/Chaos/Public/Chaos/XPBDLongRangeConstraints.h b/Engine/Source/Runtime/Experimental/Chaos/Public/Chaos/XPBDLongRangeConstraints.h
--- a/Engine/Source/Runtime/Experimental/Chaos/Public/Chaos/XPBDLongRangeConstraints.h
+++ b/Engine/Source/Runtime/Experimental/Chaos/Public/Chaos/XPBDLongRangeConstraints.h
@@ -7,8 +7,9 @@
 
 DECLARE_CYCLE_STAT(TEXT("Chaos XPBD Long Range Constraint"), STAT_XPBD_LongRange, STATGROUP_Chaos);
 
-namespace Chaos
+namespace Chaos::Softs
 {
+
 // Stiffness is in N/CM^2, so it needs to be adjusted from the PBD stiffness ranging between [0,1]
 static const double XPBDLongRangeMaxCompliance = 1.e-3;
 
@@ -19,14 +20,14 @@ public:
 	typedef typename Base::FTether FTether;
 
 	FXPBDLongRangeConstraints(
-		const FPBDParticles& Particles,
+		const FSolverParticles& Particles,
 		const int32 InParticleOffset,
 		const int32 InParticleCount,
 		const TArray<TConstArrayView<TTuple<int32, int32, FRealSingle>>>& InTethers,
 		const TConstArrayView<FRealSingle>& StiffnessMultipliers,
 		const TConstArrayView<FRealSingle>& ScaleMultipliers,
-		const FVec2& InStiffness = FVec2::UnitVector,
-		const FVec2& InScale = FVec2::UnitVector)
+		const FSolverVec2& InStiffness = FSolverVec2::UnitVector,
+		const FSolverVec2& InScale = FSolverVec2::UnitVector)
 	    : FPBDLongRangeConstraintsBase(Particles, InParticleOffset, InParticleCount, InTethers, StiffnessMultipliers, ScaleMultipliers, InStiffness, InScale)
 	{
 		int32 NumTethers = 0;
@@ -37,7 +38,7 @@ public:
 		Lambdas.Reserve(NumTethers);
 	}
 
-	virtual ~FXPBDLongRangeConstraints() {}
+	virtual ~FXPBDLongRangeConstraints() override {}
 
 	void Init() const
 	{
@@ -45,7 +46,7 @@ public:
 		Lambdas.AddZeroed(Tethers.Num());
 	}
 
-	void Apply(FPBDParticles& Particles, const FReal Dt) const 
+	void Apply(FSolverParticles& Particles, const FSolverReal Dt) const 
 	{
 		SCOPE_CYCLE_COUNTER(STAT_XPBD_LongRange);
 		// Run particles in parallel, and ranges in sequence to avoid a race condition when updating the same particle from different tethers
@@ -62,8 +63,8 @@ public:
 						{
 							const FTether& Tether = TetherBatch[Index];
 							const int32 LocalParticleIndex = GetEndIndex(Tether);
-							const FReal Scale = ScaleTable[ScaleIndices[LocalParticleIndex]];
-							const FReal ExpStiffnessValue = Stiffness[LocalParticleIndex];
+							const FSolverReal Scale = ScaleTable[ScaleIndices[LocalParticleIndex]];
+							const FSolverReal ExpStiffnessValue = Stiffness[LocalParticleIndex];
 							Apply(Particles, Dt, Tether, ConstraintOffset + Index, ExpStiffnessValue, Scale);
 						}, TetherBatch.Num() < MinParallelSize);
 					ConstraintOffset += TetherBatch.Num();
@@ -71,7 +72,7 @@ public:
 			}
 			else
 			{
-				const FReal ScaleValue = ScaleTable[0];
+				const FSolverReal ScaleValue = ScaleTable[0];
 				int32 ConstraintOffset = 0;
 				for (const TConstArrayView<FTether>& TetherBatch : Tethers)
 				{
@@ -79,7 +80,7 @@ public:
 						{
 							const FTether& Tether = TetherBatch[Index];
 							const int32 LocalParticleIndex = GetEndIndex(Tether);
-							const FReal ExpStiffnessValue = Stiffness[LocalParticleIndex];
+							const FSolverReal ExpStiffnessValue = Stiffness[LocalParticleIndex];
 							Apply(Particles, Dt, Tether, ConstraintOffset + Index, ExpStiffnessValue, ScaleValue);
 						}, TetherBatch.Num() < MinParallelSize);
 					ConstraintOffset += TetherBatch.Num();
@@ -88,7 +89,7 @@ public:
 		}
 		else
 		{
-			const FReal ExpStiffnessValue = (FReal)Stiffness;
+			const FSolverReal ExpStiffnessValue = (FSolverReal)Stiffness;
 
 			if (HasScaleWeightMap())
 			{
@@ -99,7 +100,7 @@ public:
 						{
 							const FTether& Tether = TetherBatch[Index];
 							const int32 LocalParticleIndex = GetEndIndex(Tether);
-							const FReal Scale = ScaleTable[ScaleIndices[LocalParticleIndex]];
+							const FSolverReal Scale = ScaleTable[ScaleIndices[LocalParticleIndex]];
 							Apply(Particles, Dt, Tether, ConstraintOffset + Index, ExpStiffnessValue, Scale);
 						}, TetherBatch.Num() < MinParallelSize);
 					ConstraintOffset += TetherBatch.Num();
@@ -107,7 +108,7 @@ public:
 			}
 			else
 			{
-				const FReal ScaleValue = ScaleTable[0];
+				const FSolverReal ScaleValue = ScaleTable[0];
 				int32 ConstraintOffset = 0;
 				for (const TConstArrayView<FTether>& TetherBatch : Tethers)
 				{
@@ -123,16 +124,16 @@ public:
 	}
 
 private:
-	void Apply(FPBDParticles& Particles, const FReal Dt, const FTether& Tether, int32 ConstraintIndex, const FReal InStiffness, const FReal InScale) const
+	void Apply(FSolverParticles& Particles, const FSolverReal Dt, const FTether& Tether, int32 ConstraintIndex, const FSolverReal InStiffness, const FSolverReal InScale) const
 	{
-		FVec3 Direction;
-		FReal Offset;
+		FSolverVec3 Direction;
+		FSolverReal Offset;
 		GetDelta(Particles, Tether, InScale, Direction, Offset);
 
-		FReal& Lambda = Lambdas[ConstraintIndex];
-		const FReal Alpha = (FReal)XPBDLongRangeMaxCompliance / (InStiffness * Dt * Dt);
+		FSolverReal& Lambda = Lambdas[ConstraintIndex];
+		const FSolverReal Alpha = (FSolverReal)XPBDLongRangeMaxCompliance / (InStiffness * Dt * Dt);
 
-		const FReal DLambda = (Offset - Alpha * Lambda) / ((FReal)1. + Alpha);
+		const FSolverReal DLambda = (Offset - Alpha * Lambda) / ((FSolverReal)1. + Alpha);
 		Particles.P(GetEndParticle(Tether)) += DLambda * Direction;
 		Lambda += DLambda;
 	}
@@ -144,6 +145,7 @@ private:
 	using Base::ScaleTable;
 	using Base::ScaleIndices;
 
-	mutable TArray<FReal> Lambdas;
+	mutable TArray<FSolverReal> Lambdas;
 };
-}
+
+}  // End namespace Chaos::Softs
diff --git a/Engine/Source/Runtime/Experimental/Chaos/Public/Chaos/XPBDSpringConstraints.h b/Engine/Source/Runtime/Experimental/Chaos/Public/Chaos/XPBDSpringConstraints.h
--- a/Engine/Source/Runtime/Experimental/Chaos/Public/Chaos/XPBDSpringConstraints.h
+++ b/Engine/Source/Runtime/Experimental/Chaos/Public/Chaos/XPBDSpringConstraints.h
@@ -6,13 +6,13 @@
 
 DECLARE_CYCLE_STAT(TEXT("Chaos XPBD Spring Constraint"), STAT_XPBD_Spring, STATGROUP_Chaos);
 
-namespace Chaos
+namespace Chaos::Softs
 {
 
 // Stiffness is in N/CM^2, so it needs to be adjusted from the PBD stiffness ranging between [0,1]
 static const double XPBDSpringMaxCompliance = 1e-7;  // Max stiffness: 1e+11 N/M^2 = 1e+7 N/CM^2 -> Max compliance: 1e-7 CM^2/N
 
-class FXPBDSpringConstraints : public FPBDSpringConstraintsBase
+class FXPBDSpringConstraints final : public FPBDSpringConstraintsBase
 {
 	typedef FPBDSpringConstraintsBase Base;
 	using Base::Constraints;
@@ -22,46 +22,46 @@ class FXPBDSpringConstraints : public FPBDSpringConstraintsBase
 public:
 	template<int32 Valence>
 	FXPBDSpringConstraints(
-		const FPBDParticles& Particles,
+		const FSolverParticles& Particles,
 		int32 ParticleOffset,
 		int32 ParticleCount,
 		const TArray<TVector<int32, Valence>>& InConstraints,
 		const TConstArrayView<FRealSingle>& StiffnessMultipliers,
-		const FVec2& InStiffness,
+		const FSolverVec2& InStiffness,
 		bool bTrimKinematicConstraints = false,
 		typename TEnableIf<Valence >= 2 && Valence <= 4>::Type* = nullptr)
 		: Base(Particles, ParticleOffset, ParticleCount, InConstraints, StiffnessMultipliers, InStiffness, bTrimKinematicConstraints)
 	{
-		Lambdas.Init((FReal)0., Constraints.Num());
+		Lambdas.Init((FSolverReal)0., Constraints.Num());
 	}
 
-	virtual ~FXPBDSpringConstraints() {}
+	virtual ~FXPBDSpringConstraints() override {}
 
-	void Init() const { for (FReal& Lambda : Lambdas) { Lambda = (FReal)0.; } }
+	void Init() const { for (FSolverReal& Lambda : Lambdas) { Lambda = (FSolverReal)0.; } }
 
-	void Apply(FPBDParticles& Particles, const FReal Dt, const int32 ConstraintIndex, const FReal ExpStiffnessValue) const
+	void Apply(FSolverParticles& Particles, const FSolverReal Dt, const int32 ConstraintIndex, const FSolverReal ExpStiffnessValue) const
 	{
 		const TVec2<int32>& Constraint = Constraints[ConstraintIndex];
 		const int32 i1 = Constraint[0];
 		const int32 i2 = Constraint[1];
-		const FVec3 Delta = GetDelta(Particles, Dt, ConstraintIndex, ExpStiffnessValue);
-		if (Particles.InvM(i1) > (FReal)0.)
+		const FSolverVec3 Delta = GetDelta(Particles, Dt, ConstraintIndex, ExpStiffnessValue);
+		if (Particles.InvM(i1) > (FSolverReal)0.)
 		{
 			Particles.P(i1) -= Particles.InvM(i1) * Delta;
 		}
-		if (Particles.InvM(i2) > (FReal)0.)
+		if (Particles.InvM(i2) > (FSolverReal)0.)
 		{
 			Particles.P(i2) += Particles.InvM(i2) * Delta;
 		}
 	}
 
-	void Apply(FPBDParticles& Particles, const FReal Dt) const
+	void Apply(FSolverParticles& Particles, const FSolverReal Dt) const
 	{
 		SCOPE_CYCLE_COUNTER(STAT_XPBD_Spring);
 
 		if (!Stiffness.HasWeightMap())
 		{
-			const FReal ExpStiffnessValue = (FReal)Stiffness;
+			const FSolverReal ExpStiffnessValue = (FSolverReal)Stiffness;
 			for (int32 ConstraintIndex = 0; ConstraintIndex < Constraints.Num(); ++ConstraintIndex)
 			{
 				Apply(Particles, Dt, ConstraintIndex, ExpStiffnessValue);
@@ -71,44 +71,44 @@ public:
 		{
 			for (int32 ConstraintIndex = 0; ConstraintIndex < Constraints.Num(); ++ConstraintIndex)
 			{
-				const FReal ExpStiffnessValue = Stiffness[ConstraintIndex];
+				const FSolverReal ExpStiffnessValue = Stiffness[ConstraintIndex];
 				Apply(Particles, Dt, ConstraintIndex, ExpStiffnessValue);
 			}
 		}
 	}
 
 private:
-	FVec3 GetDelta(const FPBDParticles& Particles, const FReal Dt, const int32 ConstraintIndex, const FReal ExpStiffnessValue) const
+	FSolverVec3 GetDelta(const FSolverParticles& Particles, const FSolverReal Dt, const int32 ConstraintIndex, const FSolverReal ExpStiffnessValue) const
 	{
 		const TVec2<int32>& Constraint = Constraints[ConstraintIndex];
 
 		const int32 i1 = Constraint[0];
 		const int32 i2 = Constraint[1];
 
-		if (Particles.InvM(i2) == (FReal)0. && Particles.InvM(i1) == (FReal)0.)
+		if (Particles.InvM(i2) == (FSolverReal)0. && Particles.InvM(i1) == (FSolverReal)0.)
 		{
-			return FVec3((FReal)0.);
+			return FSolverVec3((FSolverReal)0.);
 		}
-		const FReal CombinedInvMass = Particles.InvM(i2) + Particles.InvM(i1);
+		const FSolverReal CombinedInvMass = Particles.InvM(i2) + Particles.InvM(i1);
 
-		const FVec3& P1 = Particles.P(i1);
-		const FVec3& P2 = Particles.P(i2);
-		FVec3 Direction = P1 - P2;
-		const FReal Distance = Direction.SafeNormalize();
-		const FReal Offset = Distance - Dists[ConstraintIndex];
+		const FSolverVec3& P1 = Particles.P(i1);
+		const FSolverVec3& P2 = Particles.P(i2);
+		FSolverVec3 Direction = P1 - P2;
+		const FSolverReal Distance = Direction.SafeNormalize();
+		const FSolverReal Offset = Distance - Dists[ConstraintIndex];
 
-		FReal& Lambda = Lambdas[ConstraintIndex];
-		const FReal Alpha = (FReal)XPBDSpringMaxCompliance / (ExpStiffnessValue * Dt * Dt);
+		FSolverReal& Lambda = Lambdas[ConstraintIndex];
+		const FSolverReal Alpha = (FSolverReal)XPBDSpringMaxCompliance / (ExpStiffnessValue * Dt * Dt);
 
-		const FReal DLambda = (Offset - Alpha * Lambda) / (CombinedInvMass + Alpha);
-		const FVec3 Delta = DLambda * Direction;
+		const FSolverReal DLambda = (Offset - Alpha * Lambda) / (CombinedInvMass + Alpha);
+		const FSolverVec3 Delta = DLambda * Direction;
 		Lambda += DLambda;
 
 		return Delta;
 	}
 
 private:
-	mutable TArray<FReal> Lambdas;
+	mutable TArray<FSolverReal> Lambdas;
 };
 
-}
+}  // End namespace Chaos::Softs
diff --git a/Engine/Source/Runtime/Experimental/ChaosCore/Public/Chaos/Matrix.h b/Engine/Source/Runtime/Experimental/ChaosCore/Public/Chaos/Matrix.h
--- a/Engine/Source/Runtime/Experimental/ChaosCore/Public/Chaos/Matrix.h
+++ b/Engine/Source/Runtime/Experimental/ChaosCore/Public/Chaos/Matrix.h
@@ -144,13 +144,13 @@ namespace Chaos
 	};
 
 	template<>
-	class PMatrix<FReal, 4, 4> : public UE::Math::TMatrix<FReal>
+	class PMatrix<FRealSingle, 4, 4> : public UE::Math::TMatrix<FRealSingle>
 	{
 	public:
 		PMatrix()
-		    : UE::Math::TMatrix<FReal>() {}
-		PMatrix(const FReal x00, const FReal x10, const FReal x20, const FReal x30, const FReal x01, const FReal x11, const FReal x21, const FReal x31, const FReal x02, const FReal x12, const FReal x22, const FReal x32, const FReal x03, const FReal x13, const FReal x23, const FReal x33)
-		    : UE::Math::TMatrix<FReal>()
+		    : UE::Math::TMatrix<FRealSingle>() {}
+		PMatrix(const FRealSingle x00, const FRealSingle x10, const FRealSingle x20, const FRealSingle x30, const FRealSingle x01, const FRealSingle x11, const FRealSingle x21, const FRealSingle x31, const FRealSingle x02, const FRealSingle x12, const FRealSingle x22, const FRealSingle x32, const FRealSingle x03, const FRealSingle x13, const FRealSingle x23, const FRealSingle x33)
+		    : UE::Math::TMatrix<FRealSingle>()
 		{
 			M[0][0] = x00;
 			M[1][0] = x10;
@@ -169,14 +169,14 @@ namespace Chaos
 			M[2][3] = x23;
 			M[3][3] = x33;
 		}
-		PMatrix(const UE::Math::TMatrix<FReal>& Matrix)
-		    : UE::Math::TMatrix<FReal>(Matrix)
+		PMatrix(const UE::Math::TMatrix<FRealSingle>& Matrix)
+		    : UE::Math::TMatrix<FRealSingle>(Matrix)
 		{
 		}
 #if COMPILE_WITHOUT_UNREAL_SUPPORT
-		Vector<FReal, 4> operator*(const Vector<Chaos::FReal, 4>& Other)
+		Vector<FRealSingle, 4> operator*(const Vector<Chaos::FRealSingle, 4>& Other)
 		{
-			return Vector<Chaos::FReal, 4>(
+			return Vector<Chaos::FRealSingle, 4>(
 			    M[0][0] * Other[0] + M[0][1] * Other[1] + M[0][2] * Other[2] + M[0][3] * Other[3],
 			    M[1][0] * Other[0] + M[1][1] * Other[1] + M[1][2] * Other[2] + M[1][3] * Other[3],
 			    M[2][0] * Other[0] + M[2][1] * Other[1] + M[2][2] * Other[2] + M[2][3] * Other[3],
@@ -184,9 +184,56 @@ namespace Chaos
 		}
 #endif
 
-		PMatrix<FReal, 4, 4> operator*(const PMatrix<FReal, 4, 4>& Other) const
+		PMatrix<FRealSingle, 4, 4> operator*(const PMatrix<FRealSingle, 4, 4>& Other) const
 		{
-			return static_cast<const UE::Math::TMatrix<FReal>*>(this)->operator*(static_cast<const UE::Math::TMatrix<FReal>&>(Other));
+			return static_cast<const UE::Math::TMatrix<FRealSingle>*>(this)->operator*(static_cast<const UE::Math::TMatrix<FRealSingle>&>(Other));
+		}
+	};
+
+	template<>
+	class PMatrix<FRealDouble, 4, 4> : public UE::Math::TMatrix<FRealDouble>
+	{
+	public:
+		PMatrix()
+		    : UE::Math::TMatrix<FRealDouble>() {}
+		PMatrix(const FRealDouble x00, const FRealDouble x10, const FRealDouble x20, const FRealDouble x30, const FRealDouble x01, const FRealDouble x11, const FRealDouble x21, const FRealDouble x31, const FRealDouble x02, const FRealDouble x12, const FRealDouble x22, const FRealDouble x32, const FRealDouble x03, const FRealDouble x13, const FRealDouble x23, const FRealDouble x33)
+		    : UE::Math::TMatrix<FRealDouble>()
+		{
+			M[0][0] = x00;
+			M[1][0] = x10;
+			M[2][0] = x20;
+			M[3][0] = x30;
+			M[0][1] = x01;
+			M[1][1] = x11;
+			M[2][1] = x21;
+			M[3][1] = x31;
+			M[0][2] = x02;
+			M[1][2] = x12;
+			M[2][2] = x22;
+			M[3][2] = x32;
+			M[0][3] = x03;
+			M[1][3] = x13;
+			M[2][3] = x23;
+			M[3][3] = x33;
+		}
+		PMatrix(const UE::Math::TMatrix<FRealDouble>& Matrix)
+		    : UE::Math::TMatrix<FRealDouble>(Matrix)
+		{
+		}
+#if COMPILE_WITHOUT_UNREAL_SUPPORT
+		Vector<FRealDouble, 4> operator*(const Vector<Chaos::FRealDouble, 4>& Other)
+		{
+			return Vector<Chaos::FRealDouble, 4>(
+			    M[0][0] * Other[0] + M[0][1] * Other[1] + M[0][2] * Other[2] + M[0][3] * Other[3],
+			    M[1][0] * Other[0] + M[1][1] * Other[1] + M[1][2] * Other[2] + M[1][3] * Other[3],
+			    M[2][0] * Other[0] + M[2][1] * Other[1] + M[2][2] * Other[2] + M[2][3] * Other[3],
+			    M[3][0] * Other[0] + M[3][1] * Other[1] + M[3][2] * Other[2] + M[3][3] * Other[3]);
+		}
+#endif
+
+		PMatrix<FRealDouble, 4, 4> operator*(const PMatrix<FRealDouble, 4, 4>& Other) const
+		{
+			return static_cast<const UE::Math::TMatrix<FRealDouble>*>(this)->operator*(static_cast<const UE::Math::TMatrix<FRealDouble>&>(Other));
 		}
 	};
 
