Chaos Cloth - Move back the post LWC solver to float for performance reason, since cloth is being simulated in local space.

diff --git a/Engine/Source/Runtime/Experimental/ChaosCore/Public/Chaos/Matrix33.isph b/Engine/Source/Runtime/Experimental/ChaosCore/Public/Chaos/Matrix33.isph
--- a/Engine/Source/Runtime/Experimental/ChaosCore/Public/Chaos/Matrix33.isph
+++ b/Engine/Source/Runtime/Experimental/ChaosCore/Public/Chaos/Matrix33.isph
@@ -6,9 +6,14 @@
 #include "Math/Vector.isph"
 #include "Math/Matrix.isph"
 
-typedef FMatrix FMatrix33;
+typedef FMatrix44d FMatrix33d;
+typedef FMatrix44f FMatrix33f;
+typedef FMatrix44d FMatrix33;
 
-static const uniform struct FMatrix33 Matrix33Zero = {{ 0, 0, 0, 0, 0, 0, 0, 0, 0, 0, 0, 0, 0, 0, 0, 1}};
+static const uniform struct FMatrix33d Matrix33DoubleZero = {{ 0, 0, 0, 0, 0, 0, 0, 0, 0, 0, 0, 0, 0, 0, 0, 1}};
+static const uniform struct FMatrix33f Matrix33FloatZero = {{ 0, 0, 0, 0, 0, 0, 0, 0, 0, 0, 0, 0, 0, 0, 0, 1}};
+
+#define FMatrix33Zero Matrix33DoubleZero
 
 inline uniform FMatrix33 SetMatrix33(
 	const uniform FReal x00, const uniform FReal x10, const uniform FReal x20,
@@ -18,12 +23,20 @@ inline uniform FMatrix33 SetMatrix33(
 	return SetMatrix(SetVector4(x00, x01, x02, ZERO), SetVector4(x10, x11, x12, ZERO), SetVector4(x20, x21, x22, ZERO), Vector0001);
 }
 
-inline FMatrix33 SetMatrix33(
-	const FReal x00, const FReal x10, const FReal x20,
-	const FReal x01, const FReal x11, const FReal x21,
-	const FReal x02, const FReal x12, const FReal x22)
+inline FMatrix33d SetMatrix33(
+	const double x00, const double x10, const double x20,
+	const double x01, const double x11, const double x21,
+	const double x02, const double x12, const double x22)
 {
-	return SetMatrix(SetVector4(x00, x01, x02, ZERO), SetVector4(x10, x11, x12, ZERO), SetVector4(x20, x21, x22, ZERO), Vector0001);
+	return SetMatrix(SetVector4(x00, x01, x02, DOUBLE_ZERO), SetVector4(x10, x11, x12, DOUBLE_ZERO), SetVector4(x20, x21, x22, DOUBLE_ZERO), Double0001);
+}
+
+inline FMatrix33f SetMatrix33(
+	const float x00, const float x10, const float x20,
+	const float x01, const float x11, const float x21,
+	const float x02, const float x12, const float x22)
+{
+	return SetMatrix(SetVector4(x00, x01, x02, FLOAT_ZERO), SetVector4(x10, x11, x12, FLOAT_ZERO), SetVector4(x20, x21, x22, FLOAT_ZERO), Float0001);
 }
 
 inline uniform FMatrix33 SetMatrix33(
@@ -39,14 +52,24 @@ inline uniform FMatrix33 SetMatrix33(
 	return SetMatrix(SetVector4(x00, ZERO, ZERO, ZERO), SetVector4(ZERO, x11, ZERO, ZERO), SetVector4(ZERO, ZERO, x22, ZERO), Vector0001);
 }
 
-inline uniform FMatrix33 SetMatrix33(const uniform FVector4& x0, const uniform FVector4& x1, const uniform FVector4& x2)
+inline uniform FMatrix33d SetMatrix33(const uniform FVector4d& x0, const uniform FVector4d& x1, const uniform FVector4d& x2)
+{
+	return SetMatrix(x0, x1, x2, Double0001);
+}
+
+inline uniform FMatrix33f SetMatrix33(const uniform FVector4f& x0, const uniform FVector4f& x1, const uniform FVector4f& x2)
 {
-	return SetMatrix(x0, x1, x2, Vector0001);
+	return SetMatrix(x0, x1, x2, Float0001);
 }
 
-inline FMatrix33 SetMatrix33(const FVector4& x0, const FVector4& x1, const FVector4& x2)
+inline FMatrix33d SetMatrix33(const FVector4d& x0, const FVector4d& x1, const FVector4d& x2)
 {
-	return SetMatrix(x0, x1, x2, Vector0001);
+	return SetMatrix(x0, x1, x2, Double0001);
+}
+
+inline FMatrix33f SetMatrix33(const FVector4f& x0, const FVector4f& x1, const FVector4f& x2)
+{
+	return SetMatrix(x0, x1, x2, Float0001);
 }
 
 inline uniform FMatrix33 MultiplyAB(const uniform FMatrix33& L, const uniform FMatrix33& R)
@@ -83,15 +106,15 @@ inline FMatrix33 MultiplyAB(const FMatrix33& L, const FMatrix33& R)
 	return SetMatrix33(x0, x1, x2);
 }
 
-inline uniform FMatrix33 MultiplyABt(const uniform FMatrix33& L, const uniform FMatrix33& R)
+inline uniform FMatrix33d MultiplyABt(const uniform FMatrix33d& L, const uniform FMatrix33d& R)
 {
-	const uniform FVector4 L0 = SetVector4(L.M[0], L.M[1], L.M[2], L.M[3]);
-	const uniform FVector4 L1 = SetVector4(L.M[4], L.M[5], L.M[6], L.M[7]);
-	const uniform FVector4 L2 = SetVector4(L.M[8], L.M[9], L.M[10], L.M[11]);
+	const uniform FVector4d L0 = SetVector4(L.M[0], L.M[1], L.M[2], L.M[3]);
+	const uniform FVector4d L1 = SetVector4(L.M[4], L.M[5], L.M[6], L.M[7]);
+	const uniform FVector4d L2 = SetVector4(L.M[8], L.M[9], L.M[10], L.M[11]);
 
-	uniform FVector4 x0 = SetVector4(R.M[0]) * L0 + SetVector4(R.M[4]) * L1 + SetVector4(R.M[8]) * L2;
-	uniform FVector4 x1 = SetVector4(R.M[1]) * L0 + SetVector4(R.M[5]) * L1 + SetVector4(R.M[9]) * L2;
-	uniform FVector4 x2 = SetVector4(R.M[2]) * L0 + SetVector4(R.M[6]) * L1 + SetVector4(R.M[10]) * L2;
+	uniform FVector4d x0 = SetVector4(R.M[0]) * L0 + SetVector4(R.M[4]) * L1 + SetVector4(R.M[8]) * L2;
+	uniform FVector4d x1 = SetVector4(R.M[1]) * L0 + SetVector4(R.M[5]) * L1 + SetVector4(R.M[9]) * L2;
+	uniform FVector4d x2 = SetVector4(R.M[2]) * L0 + SetVector4(R.M[6]) * L1 + SetVector4(R.M[10]) * L2;
 
 	x0.V[3] = 0;
 	x1.V[3] = 0;
@@ -100,15 +123,49 @@ inline uniform FMatrix33 MultiplyABt(const uniform FMatrix33& L, const uniform F
 	return SetMatrix33(x0, x1, x2);
 }
 
-inline FMatrix33 MultiplyABt(const FMatrix33& L, const FMatrix33& R)
+inline uniform FMatrix33f MultiplyABt(const uniform FMatrix33f& L, const uniform FMatrix33f& R)
 {
-	const FVector4 L0 = SetVector4(L.M[0], L.M[1], L.M[2], L.M[3]);
-	const FVector4 L1 = SetVector4(L.M[4], L.M[5], L.M[6], L.M[7]);
-	const FVector4 L2 = SetVector4(L.M[8], L.M[9], L.M[10], L.M[11]);
+	const uniform FVector4f L0 = SetVector4(L.M[0], L.M[1], L.M[2], L.M[3]);
+	const uniform FVector4f L1 = SetVector4(L.M[4], L.M[5], L.M[6], L.M[7]);
+	const uniform FVector4f L2 = SetVector4(L.M[8], L.M[9], L.M[10], L.M[11]);
 
-	FVector4 x0 = SetVector4(R.M[0]) * L0 + SetVector4(R.M[4]) * L1 + SetVector4(R.M[8]) * L2;
-	FVector4 x1 = SetVector4(R.M[1]) * L0 + SetVector4(R.M[5]) * L1 + SetVector4(R.M[9]) * L2;
-	FVector4 x2 = SetVector4(R.M[2]) * L0 + SetVector4(R.M[6]) * L1 + SetVector4(R.M[10]) * L2;
+	uniform FVector4f x0 = SetVector4(R.M[0]) * L0 + SetVector4(R.M[4]) * L1 + SetVector4(R.M[8]) * L2;
+	uniform FVector4f x1 = SetVector4(R.M[1]) * L0 + SetVector4(R.M[5]) * L1 + SetVector4(R.M[9]) * L2;
+	uniform FVector4f x2 = SetVector4(R.M[2]) * L0 + SetVector4(R.M[6]) * L1 + SetVector4(R.M[10]) * L2;
+
+	x0.V[3] = 0;
+	x1.V[3] = 0;
+	x2.V[3] = 0;
+
+	return SetMatrix33(x0, x1, x2);
+}
+
+inline FMatrix33d MultiplyABt(const FMatrix33d& L, const FMatrix33d& R)
+{
+	const FVector4d L0 = SetVector4(L.M[0], L.M[1], L.M[2], L.M[3]);
+	const FVector4d L1 = SetVector4(L.M[4], L.M[5], L.M[6], L.M[7]);
+	const FVector4d L2 = SetVector4(L.M[8], L.M[9], L.M[10], L.M[11]);
+
+	FVector4d x0 = SetVector4(R.M[0]) * L0 + SetVector4(R.M[4]) * L1 + SetVector4(R.M[8]) * L2;
+	FVector4d x1 = SetVector4(R.M[1]) * L0 + SetVector4(R.M[5]) * L1 + SetVector4(R.M[9]) * L2;
+	FVector4d x2 = SetVector4(R.M[2]) * L0 + SetVector4(R.M[6]) * L1 + SetVector4(R.M[10]) * L2;
+
+	x0.V[3] = 0;
+	x1.V[3] = 0;
+	x2.V[3] = 0;
+
+	return SetMatrix33(x0, x1, x2);
+}
+
+inline FMatrix33f MultiplyABt(const FMatrix33f& L, const FMatrix33f& R)
+{
+	const FVector4f L0 = SetVector4(L.M[0], L.M[1], L.M[2], L.M[3]);
+	const FVector4f L1 = SetVector4(L.M[4], L.M[5], L.M[6], L.M[7]);
+	const FVector4f L2 = SetVector4(L.M[8], L.M[9], L.M[10], L.M[11]);
+
+	FVector4f x0 = SetVector4(R.M[0]) * L0 + SetVector4(R.M[4]) * L1 + SetVector4(R.M[8]) * L2;
+	FVector4f x1 = SetVector4(R.M[1]) * L0 + SetVector4(R.M[5]) * L1 + SetVector4(R.M[9]) * L2;
+	FVector4f x2 = SetVector4(R.M[2]) * L0 + SetVector4(R.M[6]) * L1 + SetVector4(R.M[10]) * L2;
 
 	x0.V[3] = 0;
 	x1.V[3] = 0;
@@ -148,7 +205,16 @@ inline uniform FMatrix33 Multiply(const uniform FMatrix33& L, const uniform FRea
 	);
 }
 
-inline FMatrix33 Multiply(const FMatrix33& L, const FReal R)
+inline FMatrix33d Multiply(const FMatrix33d& L, const double R)
+{
+	return SetMatrix33(
+		L.M[0] * R, L.M[1] * R, L.M[2] * R,
+		L.M[4] * R, L.M[5] * R, L.M[6] * R,
+		L.M[8] * R, L.M[9] * R, L.M[10] * R
+	);
+}
+
+inline FMatrix33f Multiply(const FMatrix33f& L, const float R)
 {
 	return SetMatrix33(
 		L.M[0] * R, L.M[1] * R, L.M[2] * R,
@@ -178,19 +244,40 @@ inline uniform FMatrix33 AddAB(const uniform FMatrix33& L, const uniform FMatrix
 	return SetMatrix33(x0, x1, x2);
 }
 
-inline FMatrix33 AddAB(const FMatrix33& L, const FMatrix33& R)
+inline FMatrix33d AddAB(const FMatrix33d& L, const FMatrix33d& R)
 {
-	const FVector4 L0 = SetVector4(L.M[0], L.M[1], L.M[2], L.M[3]);
-	const FVector4 L1 = SetVector4(L.M[4], L.M[5], L.M[6], L.M[7]);
-	const FVector4 L2 = SetVector4(L.M[8], L.M[9], L.M[10], L.M[11]);
+	const FVector4d L0 = SetVector4(L.M[0], L.M[1], L.M[2], L.M[3]);
+	const FVector4d L1 = SetVector4(L.M[4], L.M[5], L.M[6], L.M[7]);
+	const FVector4d L2 = SetVector4(L.M[8], L.M[9], L.M[10], L.M[11]);
+
+	const FVector4d R0 = SetVector4(R.M[0], R.M[1], R.M[2], R.M[3]);
+	const FVector4d R1 = SetVector4(R.M[4], R.M[5], R.M[6], R.M[7]);
+	const FVector4d R2 = SetVector4(R.M[8], R.M[9], R.M[10], R.M[11]);
+
+	FVector4d x0 = R0 + L0;
+	FVector4d x1 = R1 + L1;
+	FVector4d x2 = R2 + L2;
+
+	x0.V[3] = 0;
+	x1.V[3] = 0;
+	x2.V[3] = 0;
+
+	return SetMatrix33(x0, x1, x2);
+}
+
+inline FMatrix33f AddAB(const FMatrix33f& L, const FMatrix33f& R)
+{
+	const FVector4f L0 = SetVector4(L.M[0], L.M[1], L.M[2], L.M[3]);
+	const FVector4f L1 = SetVector4(L.M[4], L.M[5], L.M[6], L.M[7]);
+	const FVector4f L2 = SetVector4(L.M[8], L.M[9], L.M[10], L.M[11]);
 
-	const FVector4 R0 = SetVector4(R.M[0], R.M[1], R.M[2], R.M[3]);
-	const FVector4 R1 = SetVector4(R.M[4], R.M[5], R.M[6], R.M[7]);
-	const FVector4 R2 = SetVector4(R.M[8], R.M[9], R.M[10], R.M[11]);
+	const FVector4f R0 = SetVector4(R.M[0], R.M[1], R.M[2], R.M[3]);
+	const FVector4f R1 = SetVector4(R.M[4], R.M[5], R.M[6], R.M[7]);
+	const FVector4f R2 = SetVector4(R.M[8], R.M[9], R.M[10], R.M[11]);
 
-	FVector4 x0 = R0 + L0;
-	FVector4 x1 = R1 + L1;
-	FVector4 x2 = R2 + L2;
+	FVector4f x0 = R0 + L0;
+	FVector4f x1 = R1 + L1;
+	FVector4f x2 = R2 + L2;
 
 	x0.V[3] = 0;
 	x1.V[3] = 0;
diff --git a/Engine/Source/Runtime/Experimental/ChaosCore/Public/Chaos/Vector.h b/Engine/Source/Runtime/Experimental/ChaosCore/Public/Chaos/Vector.h
--- a/Engine/Source/Runtime/Experimental/ChaosCore/Public/Chaos/Vector.h
+++ b/Engine/Source/Runtime/Experimental/ChaosCore/Public/Chaos/Vector.h
@@ -791,23 +791,107 @@ namespace Chaos
 		}
 	};
 
+	template<>
+	class TVector<FRealSingle, 2> : public FVector2f
+	{
+	public:
+		using FElement = decltype(FVector2f::X);
+		using FVector2f::X;
+		using FVector2f::Y;
+
+		TVector()
+		    : FVector2f() {}
+		TVector(const FRealSingle x)
+		    : FVector2f((decltype(FVector2f::X))x, (decltype(FVector2f::X))x) {}	// LWC_TODO: Remove casts once FVector2f supports variants
+		TVector(const FRealSingle x, const FRealSingle y)
+		    : FVector2f((decltype(FVector2f::X))x, (decltype(FVector2f::X))y) {}
+		TVector(const FVector2f& vec)
+		    : FVector2f(vec) {}
+		TVector(std::istream& Stream)
+		{
+			Stream.read(reinterpret_cast<char*>(&X), sizeof(X));
+			Stream.read(reinterpret_cast<char*>(&Y), sizeof(Y));
+		}
+		template <typename OtherT>
+		TVector(const TVector<OtherT, 2>& InVector)
+		{
+			X = ((decltype(X))InVector[0]);	// LWC_TODO: Remove casts once FVector2f supports variants
+			Y = ((decltype(Y))InVector[1]);
+		}
+		void Write(std::ostream& Stream) const
+		{
+			Stream.write(reinterpret_cast<const char*>(&X), sizeof(X));
+			Stream.write(reinterpret_cast<const char*>(&Y), sizeof(Y));
+		}
+
+		static TVector<FRealSingle, 2> AxisVector(const int32 Axis)
+		{
+			check(Axis >= 0 && Axis <= 1);
+			return Axis == 0 ? TVector<FRealSingle, 2>(1.f, 0.f) : TVector<FRealSingle, 2>(0.f, 1.f);
+		}
+		FRealSingle Product() const
+		{
+			return X * Y;
+		}
+		FRealSingle Max() const
+		{
+			return X > Y ? X : Y;
+		}
+		FRealSingle Min() const
+		{
+			return X < Y ? X : Y;
+		}
+		static TVector<FRealSingle, 2> Max(const TVector<FRealSingle, 2>& V1, const TVector<FRealSingle, 2>& V2)
+		{
+			return TVector<FRealSingle, 2>(V1.X > V2.X ? V1.X : V2.X, V1.Y > V2.Y ? V1.Y : V2.Y);
+		}
+		static Pair<FRealSingle, int32> MaxAndAxis(const TVector<FRealSingle, 2>& V1, const TVector<FRealSingle, 2>& V2)
+		{
+			const TVector<FRealSingle, 2> max = Max(V1, V2);
+			if (max.X > max.Y)
+			{
+				return MakePair((FRealSingle)max.X, 0);
+			}
+			else
+			{
+				return MakePair((FRealSingle)max.Y, 1);
+			}
+		}
+		template<class T2>
+		TVector<FRealSingle, 2> operator/(const TVector<T2, 2>& Other) const
+		{
+			return TVector<FRealSingle, 2>(X / static_cast<FRealSingle>(Other[0]), Y / static_cast<FRealSingle>(Other[1]));
+		}
+		TVector<FRealSingle, 2> operator/(const FRealSingle Other) const
+		{
+			return TVector<FRealSingle, 2>(X / Other, Y / Other);
+		}
+		TVector<FRealSingle, 2> operator*(const FRealSingle Other) const
+		{
+			return TVector<FRealSingle, 2>(X * Other, Y * Other);
+		}
+		TVector<FRealSingle, 2> operator*(const TVector<FRealSingle, 2>& Other) const
+		{
+			return TVector<FRealSingle, 2>(X * Other[0], Y * Other[1]);
+		}
+	};
 
 	template<>
-	class TVector<FReal, 2> : public FVector2D
+	class TVector<FRealDouble, 2> : public FVector2d
 	{
 	public:
-		using FElement = decltype(FVector2D::X);
-		using FVector2D::X;
-		using FVector2D::Y;
+		using FElement = decltype(FVector2d::X);
+		using FVector2d::X;
+		using FVector2d::Y;
 
 		TVector()
-		    : FVector2D() {}
-		TVector(const FReal x)
-		    : FVector2D((decltype(FVector2D::X))x, (decltype(FVector2D::X))x) {}	// LWC_TODO: Remove casts once FVector2D supports variants
-		TVector(const FReal x, const FReal y)
-		    : FVector2D((decltype(FVector2D::X))x, (decltype(FVector2D::X))y) {}
-		TVector(const FVector2D& vec)
-		    : FVector2D(vec) {}
+		    : FVector2d() {}
+		TVector(const FRealDouble x)
+		    : FVector2d((decltype(FVector2d::X))x, (decltype(FVector2d::X))x) {}	// LWC_TODO: Remove casts once FVector2d supports variants
+		TVector(const FRealDouble x, const FRealDouble y)
+		    : FVector2d((decltype(FVector2d::X))x, (decltype(FVector2d::X))y) {}
+		TVector(const FVector2d& vec)
+		    : FVector2d(vec) {}
 		TVector(std::istream& Stream)
 		{
 			Stream.read(reinterpret_cast<char*>(&X), sizeof(X));
@@ -816,7 +900,7 @@ namespace Chaos
 		template <typename OtherT>
 		TVector(const TVector<OtherT, 2>& InVector)
 		{
-			X = ((decltype(X))InVector[0]);	// LWC_TODO: Remove casts once FVector2D supports variants
+			X = ((decltype(X))InVector[0]);	// LWC_TODO: Remove casts once FVector2d supports variants
 			Y = ((decltype(Y))InVector[1]);
 		}
 		void Write(std::ostream& Stream) const
@@ -825,55 +909,55 @@ namespace Chaos
 			Stream.write(reinterpret_cast<const char*>(&Y), sizeof(Y));
 		}
 
-		static TVector<FReal, 2> AxisVector(const int32 Axis)
+		static TVector<FRealDouble, 2> AxisVector(const int32 Axis)
 		{
 			check(Axis >= 0 && Axis <= 1);
-			return Axis == 0 ? TVector<FReal, 2>(1.f, 0.f) : TVector<FReal, 2>(0.f, 1.f);
+			return Axis == 0 ? TVector<FRealDouble, 2>(1.f, 0.f) : TVector<FRealDouble, 2>(0.f, 1.f);
 		}
-		FReal Product() const
+		FRealDouble Product() const
 		{
 			return X * Y;
 		}
-		FReal Max() const
+		FRealDouble Max() const
 		{
 			return X > Y ? X : Y;
 		}
-		FReal Min() const
+		FRealDouble Min() const
 		{
 			return X < Y ? X : Y;
 		}
-		static TVector<FReal, 2> Max(const TVector<FReal, 2>& V1, const TVector<FReal, 2>& V2)
+		static TVector<FRealDouble, 2> Max(const TVector<FRealDouble, 2>& V1, const TVector<FRealDouble, 2>& V2)
 		{
-			return TVector<FReal, 2>(V1.X > V2.X ? V1.X : V2.X, V1.Y > V2.Y ? V1.Y : V2.Y);
+			return TVector<FRealDouble, 2>(V1.X > V2.X ? V1.X : V2.X, V1.Y > V2.Y ? V1.Y : V2.Y);
 		}
-		static Pair<FReal, int32> MaxAndAxis(const TVector<FReal, 2>& V1, const TVector<FReal, 2>& V2)
+		static Pair<FRealDouble, int32> MaxAndAxis(const TVector<FRealDouble, 2>& V1, const TVector<FRealDouble, 2>& V2)
 		{
-			const TVector<FReal, 2> max = Max(V1, V2);
+			const TVector<FRealDouble, 2> max = Max(V1, V2);
 			if (max.X > max.Y)
 			{
-				return MakePair((FReal)max.X, 0);
+				return MakePair((FRealDouble)max.X, 0);
 			}
 			else
 			{
-				return MakePair((FReal)max.Y, 1);
+				return MakePair((FRealDouble)max.Y, 1);
 			}
 		}
 		template<class T2>
-		TVector<FReal, 2> operator/(const TVector<T2, 2>& Other) const
+		TVector<FRealDouble, 2> operator/(const TVector<T2, 2>& Other) const
 		{
-			return TVector<FReal, 2>(X / static_cast<FReal>(Other[0]), Y / static_cast<FReal>(Other[1]));
+			return TVector<FRealDouble, 2>(X / static_cast<FRealDouble>(Other[0]), Y / static_cast<FRealDouble>(Other[1]));
 		}
-		TVector<FReal, 2> operator/(const FReal Other) const
+		TVector<FRealDouble, 2> operator/(const FRealDouble Other) const
 		{
-			return TVector<FReal, 2>(X / Other, Y / Other);
+			return TVector<FRealDouble, 2>(X / Other, Y / Other);
 		}
-		TVector<FReal, 2> operator*(const FReal Other) const
+		TVector<FRealDouble, 2> operator*(const FRealDouble Other) const
 		{
-			return TVector<FReal, 2>(X * Other, Y * Other);
+			return TVector<FRealDouble, 2>(X * Other, Y * Other);
 		}
-		TVector<FReal, 2> operator*(const TVector<FReal, 2>& Other) const
+		TVector<FRealDouble, 2> operator*(const TVector<FRealDouble, 2>& Other) const
 		{
-			return TVector<FReal, 2>(X * Other[0], Y * Other[1]);
+			return TVector<FRealDouble, 2>(X * Other[0], Y * Other[1]);
 		}
 	};
 #endif // !COMPILE_WITHOUT_UNREAL_SUPPORT
diff --git a/Engine/Source/Runtime/PhysicsCore/Public/Chaos/ChaosEngineInterface.h b/Engine/Source/Runtime/PhysicsCore/Public/Chaos/ChaosEngineInterface.h
--- a/Engine/Source/Runtime/PhysicsCore/Public/Chaos/ChaosEngineInterface.h
+++ b/Engine/Source/Runtime/PhysicsCore/Public/Chaos/ChaosEngineInterface.h
@@ -208,7 +208,6 @@ namespace Chaos
 
 	class FPerParticleGravity;
 
-	class FPBDSpringConstraints;
 	class FConvex;
 	class FCapsule;
 
