GeometryCollection : Convex
- Auto populate convex geometry in older files.

diff --git a/Engine/Source/Runtime/Core/Public/UObject/UE5ReleaseStreamObjectVersion.h b/Engine/Source/Runtime/Core/Public/UObject/UE5ReleaseStreamObjectVersion.h
--- a/Engine/Source/Runtime/Core/Public/UObject/UE5ReleaseStreamObjectVersion.h
+++ b/Engine/Source/Runtime/Core/Public/UObject/UE5ReleaseStreamObjectVersion.h
@@ -104,6 +104,9 @@ struct CORE_API FUE5ReleaseStreamObjectVersion
 		// Changed shadow defaults for directional light components, version needed to not affect old things
 		UpdatedDirectionalLightShadowDefaults,
 
+		// Refresh geometry collections that had not already generated convex bodies.
+		GeometryCollectionConvexDefaults,
+
 		// -----<new versions can be added above this line>-------------------------------------------------
 		VersionPlusOne,
 		LatestVersion = VersionPlusOne - 1
diff --git a/Engine/Source/Runtime/Experimental/ChaosCore/Public/Chaos/Core.h b/Engine/Source/Runtime/Experimental/ChaosCore/Public/Chaos/Core.h
--- a/Engine/Source/Runtime/Experimental/ChaosCore/Public/Chaos/Core.h
+++ b/Engine/Source/Runtime/Experimental/ChaosCore/Public/Chaos/Core.h
@@ -41,5 +41,5 @@ namespace Chaos
 	// NOTE: if you get a merge conflict on the GUID, you must replace it with a new GUID - do not accept the source or target
 	// or you will likely get DDC version conflicts resulting in crashes during load.
 	// Core version string for Chaos data. Any DDC builder dependent on Chaos for serialization should depend on this version
-	static const FString ChaosVersionString = TEXT("D5C2C576-5318-416A-99DC-EB6D363B390E");
+	static const FString ChaosVersionString = TEXT("727DAD8C-96F4-4726-98C5-56AC8B1F5158");
 }
diff --git a/Engine/Source/Runtime/Experimental/GeometryCollectionEngine/Private/GeometryCollection/GeometryCollectionObject.cpp b/Engine/Source/Runtime/Experimental/GeometryCollectionEngine/Private/GeometryCollection/GeometryCollectionObject.cpp
--- a/Engine/Source/Runtime/Experimental/GeometryCollectionEngine/Private/GeometryCollection/GeometryCollectionObject.cpp
+++ b/Engine/Source/Runtime/Experimental/GeometryCollectionEngine/Private/GeometryCollection/GeometryCollectionObject.cpp
@@ -637,6 +637,8 @@ void UGeometryCollection::Serialize(FArchive& Ar)
 	bool bCreateSimulationData = false;
 	Ar.UsingCustomVersion(FDestructionObjectVersion::GUID);
 	Ar.UsingCustomVersion(FUE5MainStreamObjectVersion::GUID);
+	Ar.UsingCustomVersion(FUE5ReleaseStreamObjectVersion::GUID);
+	
 	Chaos::FChaosArchive ChaosAr(Ar);
 
 	// The Geometry Collection we will be archiving. This may be replaced with a transient, stripped back Geometry Collection if we are cooking.
@@ -781,6 +783,22 @@ void UGeometryCollection::Serialize(FArchive& Ar)
 		}
 	}
 
+
+	// will generate convex bodies when they dont exist. 
+	if (Ar.CustomVer(FUE5ReleaseStreamObjectVersion::GUID) < FUE5ReleaseStreamObjectVersion::GeometryCollectionConvexDefaults)
+	{
+#if WITH_EDITOR
+		if (!FGeometryCollectionConvexUtility::HasConvexHullData(GeometryCollection.Get()))
+		{
+			UE_LOG(LogGeometryCollectionInternal, Warning, TEXT("Regenerated GeoemtryCollections convex geometry on asset (%s). Please resave package %s."), *GetName(), *GetOutermost()->GetPathName());
+			FGeometryCollectionConvexPropertiesInterface::FConvexCreationProperties ConvexProperties = GeometryCollection->GetConvexProperties();
+			FGeometryCollectionConvexUtility::CreateNonOverlappingConvexHullData(GeometryCollection.Get(), ConvexProperties.FractionRemove, ConvexProperties.SimplificationThreshold, ConvexProperties.CanExceedFraction);
+			bCreateSimulationData = true;
+			InvalidateCollection();
+		}
+#endif
+	}
+
 #if WITH_EDITORONLY_DATA
 	if (bCreateSimulationData)
 	{
