Always call UpdateConstraintSwept even constraint can be restored.
Fix race condition in CCDUtilities.

diff --git a/Engine/Source/Runtime/Experimental/Chaos/Private/Chaos/CCDUtilities.cpp b/Engine/Source/Runtime/Experimental/Chaos/Private/Chaos/CCDUtilities.cpp
--- a/Engine/Source/Runtime/Experimental/Chaos/Private/Chaos/CCDUtilities.cpp
+++ b/Engine/Source/Runtime/Experimental/Chaos/Private/Chaos/CCDUtilities.cpp
@@ -262,9 +262,9 @@ namespace Chaos
 		}
 	}
 
-	bool CCDConstraintSortPredicate(const FCCDConstraint &Constraint0, const FCCDConstraint &Constraint1)
+	bool CCDConstraintSortPredicate(const FCCDConstraint* Constraint0, const FCCDConstraint* Constraint1)
 	{
-		return Constraint0.SweptConstraint->TimeOfImpact < Constraint1.SweptConstraint->TimeOfImpact;
+		return Constraint0->SweptConstraint->TimeOfImpact < Constraint1->SweptConstraint->TimeOfImpact;
 	}
 
 	void FCCDManager::ApplyIslandSweptConstraints(const int32 Island, const FReal Dt)
@@ -275,7 +275,7 @@ namespace Chaos
 		check(ConstraintNum > 0);
 
 		// Sort constraints based on TOI
-		::Sort(SortedCCDConstraints.GetData() + ConstraintStart, ConstraintNum, CCDConstraintSortPredicate);
+		std::sort(SortedCCDConstraints.GetData() + ConstraintStart, SortedCCDConstraints.GetData() + ConstraintStart + ConstraintNum, CCDConstraintSortPredicate);
 		FReal IslandTOI = 0.f;
 		ResetIslandParticles(Island);
 		ResetIslandConstraints(Island);
@@ -450,7 +450,7 @@ namespace Chaos
 				if (HasResweptConstraint)
 				{
 					// This could be optimized by using bubble sort if there are only a few updated constraints.
-					::Sort(SortedCCDConstraints.GetData() + ConstraintIndex, ConstraintNum - ConstraintIndex, CCDConstraintSortPredicate);
+					std::sort(SortedCCDConstraints.GetData() + ConstraintIndex, SortedCCDConstraints.GetData() + ConstraintStart + ConstraintNum, CCDConstraintSortPredicate);
 				}
 			}
 		}
diff --git a/Engine/Source/Runtime/Experimental/Chaos/Private/Chaos/Collision/ParticlePairMidPhase.cpp b/Engine/Source/Runtime/Experimental/Chaos/Private/Chaos/Collision/ParticlePairMidPhase.cpp
--- a/Engine/Source/Runtime/Experimental/Chaos/Private/Chaos/Collision/ParticlePairMidPhase.cpp
+++ b/Engine/Source/Runtime/Experimental/Chaos/Private/Chaos/Collision/ParticlePairMidPhase.cpp
@@ -219,16 +219,22 @@ namespace Chaos
 
 	int32 FSingleShapePairCollisionDetector::GenerateCollision(
 		const FReal CullDistance,
-		const bool bUseCCD,
 		const FReal Dt)
 	{
-		if (bUseCCD || DoBoundsOverlap(CullDistance))
+		if (DoBoundsOverlap(CullDistance))
 		{
-			return GenerateCollisionImpl(CullDistance, bUseCCD, Dt);
+			return GenerateCollisionImpl(CullDistance, Dt);
 		}
 		return 0;
 	}
 
+	int32 FSingleShapePairCollisionDetector::GenerateCollisionCCD(
+		const FReal CullDistance,
+		const FReal Dt)
+	{
+		return GenerateCollisionCCDImpl(CullDistance, Dt);
+	}
+
 	void FSingleShapePairCollisionDetector::CreateConstraint(const FReal CullDistance)
 	{
 		PHYSICS_CSV_SCOPED_EXPENSIVE(PhysicsVerbose, NarrowPhase_CreateConstraint);
@@ -242,7 +248,7 @@ namespace Chaos
 		LastUsedEpoch = -1;
 	}
 
-	int32 FSingleShapePairCollisionDetector::GenerateCollisionImpl(const FReal CullDistance, const bool bUseCCD, const FReal Dt)
+	int32 FSingleShapePairCollisionDetector::GenerateCollisionImpl(const FReal CullDistance, const FReal Dt)
 	{
 		if (Constraint == nullptr)
 		{
@@ -270,11 +276,11 @@ namespace Chaos
 			// NOTE: these are not used by CCD which continuously moves the particles
 			Constraint->SetShapeWorldTransforms(ShapeWorldTransform0, ShapeWorldTransform1);
 
-			Constraint->SetCCDEnabled(bUseCCD);
+			Constraint->SetCCDEnabled(false);
 
 			// If the constraint was not used last frame, it needs to be reset. 
 			// Otherwise we will try to reuse it below
-			if (!bWasUpdatedLastTick  || (Constraint->GetManifoldPoints().Num() == 0) || bUseCCD)
+			if (!bWasUpdatedLastTick  || (Constraint->GetManifoldPoints().Num() == 0))
 			{
 				// Clear all manifold data including saved contact data
 				Constraint->ResetManifold();
@@ -298,26 +304,8 @@ namespace Chaos
 			{
 				// We will be updating the manifold, if only partially, so update the restore comparison transforms
 				Constraint->SetLastShapeWorldTransforms(ShapeWorldTransform0, ShapeWorldTransform1);
-
 				// Run the narrow phase
-				if (!bUseCCD)
-				{
-					Collisions::UpdateConstraint(*Constraint.Get(), ShapeWorldTransform0, ShapeWorldTransform1, Dt);
-				}
-				else
-				{
-					// Note: This is unusual but we are using a mix of the previous and current transform
-					// This is due to how CCD rewinds the position (not rotation) and then sweeps to find the first contact at the current orientation
-					// NOTE: These are actor transforms, not CoM transforms
-					// @todo(chaos): this is broken if both objects are CCD
-					FConstGenericParticleHandle P0 = Particle0;
-					FConstGenericParticleHandle P1 = Particle1;
-					const FRigidTransform3 CCDParticleWorldTransform0 = FRigidTransform3(P0->CCDEnabled() ? P0->X() : P0->P(), P0->Q());
-					const FRigidTransform3 CCDParticleWorldTransform1 = FRigidTransform3(P1->CCDEnabled() ? P1->X() : P1->P(), P1->Q());
-					const FRigidTransform3 CCDShapeWorldTransform0 = Constraint->ImplicitTransform[0] * CCDParticleWorldTransform0;
-					const FRigidTransform3 CCDShapeWorldTransform1 = Constraint->ImplicitTransform[1] * CCDParticleWorldTransform1;
-					Collisions::UpdateConstraintSwept(*Constraint.Get(), CCDShapeWorldTransform0, CCDShapeWorldTransform1, Dt);
-				}
+				Collisions::UpdateConstraint(*Constraint.Get(), ShapeWorldTransform0, ShapeWorldTransform1, Dt);
 			}
 
 			// If we have a valid contact, add it to the active list
@@ -334,6 +322,51 @@ namespace Chaos
 		return 0;
 	}
 
+	int32 FSingleShapePairCollisionDetector::GenerateCollisionCCDImpl(const FReal CullDistance, const FReal Dt)
+	{
+		if (Constraint == nullptr)
+		{
+			// Lazy creation of the constraint. 
+			CreateConstraint(CullDistance);
+		}
+
+		if (Constraint != nullptr)
+		{
+			PHYSICS_CSV_SCOPED_EXPENSIVE(PhysicsVerbose, NarrowPhase_UpdateConstraintCCD);
+
+			const FImplicitObject* Implicit0 = Shape0->GetLeafGeometry();
+			const FImplicitObject* Implicit1 = Shape1->GetLeafGeometry();
+			const FRigidTransform3& ShapeWorldTransform0 = Shape0->GetLeafWorldTransform();
+			const FRigidTransform3& ShapeWorldTransform1 = Shape1->GetLeafWorldTransform();
+
+			// Update the world shape transforms on the constraint (we cannot just give it the PerShapeData 
+			// pointer because of Unions - see FMultiShapePairCollisionDetector)
+			// NOTE: these are not used by CCD which continuously moves the particles
+			Constraint->SetShapeWorldTransforms(ShapeWorldTransform0, ShapeWorldTransform1);
+
+			Constraint->SetCCDEnabled(true);
+			Constraint->ResetManifold();
+			Constraint->ResetActiveManifoldContacts();
+
+			// Note: It is unusual that we are mixing X and Q. 
+			// This is due to how CCD rewinds the position (not rotation) and then sweeps to find the first contact at the most recent orientation Q
+			// NOTE: These are actor transforms, not CoM transforms
+			FConstGenericParticleHandle P0 = Particle0;
+			FConstGenericParticleHandle P1 = Particle1;
+			const FRigidTransform3 CCDParticleWorldTransform0 = FRigidTransform3(P0->CCDEnabled() ? P0->X() : P0->P(), P0->Q());
+			const FRigidTransform3 CCDParticleWorldTransform1 = FRigidTransform3(P1->CCDEnabled() ? P1->X() : P1->P(), P1->Q());
+			const FRigidTransform3 CCDShapeWorldTransform0 = Constraint->ImplicitTransform[0] * CCDParticleWorldTransform0;
+			const FRigidTransform3 CCDShapeWorldTransform1 = Constraint->ImplicitTransform[1] * CCDParticleWorldTransform1;
+			Collisions::UpdateConstraintSwept(*Constraint.Get(), CCDShapeWorldTransform0, CCDShapeWorldTransform1, Dt);
+			MidPhase.GetCollisionAllocator().ActivateConstraint(Constraint.Get());
+			LastUsedEpoch = MidPhase.GetCollisionAllocator().GetCurrentEpoch();
+
+			return 1;
+		}
+
+		return 0;
+	}
+
 	void FSingleShapePairCollisionDetector::WakeCollision(const int32 SleepEpoch)
 	{
 		if ((Constraint != nullptr) && IsUsedSince(SleepEpoch))
@@ -413,7 +446,6 @@ namespace Chaos
 
 	int32 FMultiShapePairCollisionDetector::GenerateCollisions(
 		const FReal CullDistance,
-		const bool bUseCCD,
 		const FReal Dt,
 		FCollisionContext& Context)
 	{
@@ -813,16 +845,25 @@ namespace Chaos
 
 		// Enable CCD?
 		const bool bUseCCD = Flags.bIsCCD && ShouldEnableCCD(Dt);
-
 		// Run collision detection on all potentially colliding shape pairs
 		NumActiveConstraints = 0;
-		for (FSingleShapePairCollisionDetector& ShapePair : ShapePairDetectors)
+		if (bUseCCD) // Generate CCD constraints as long as AABBs overlap
 		{
-			NumActiveConstraints += ShapePair.GenerateCollision(CullDistance, bUseCCD, Dt);
+			for (FSingleShapePairCollisionDetector& ShapePair : ShapePairDetectors)
+			{
+				NumActiveConstraints += ShapePair.GenerateCollisionCCD(CullDistance, Dt);
+			}
+		}
+		else
+		{
+			for (FSingleShapePairCollisionDetector& ShapePair : ShapePairDetectors)
+			{
+				NumActiveConstraints += ShapePair.GenerateCollision(CullDistance, Dt);
+			}
 		}
 		for (FMultiShapePairCollisionDetector& MultiShapePair : MultiShapePairDetectors)
 		{
-			NumActiveConstraints += MultiShapePair.GenerateCollisions(CullDistance, bUseCCD, Dt, Context);
+			NumActiveConstraints += MultiShapePair.GenerateCollisions(CullDistance, Dt, Context);
 		}
 
 		LastUsedEpoch = CollisionAllocator->GetCurrentEpoch();
diff --git a/Engine/Source/Runtime/Experimental/Chaos/Public/Chaos/Collision/ParticlePairMidPhase.h b/Engine/Source/Runtime/Experimental/Chaos/Public/Chaos/Collision/ParticlePairMidPhase.h
--- a/Engine/Source/Runtime/Experimental/Chaos/Public/Chaos/Collision/ParticlePairMidPhase.h
+++ b/Engine/Source/Runtime/Experimental/Chaos/Public/Chaos/Collision/ParticlePairMidPhase.h
@@ -67,7 +67,14 @@ namespace Chaos
 		*/
 		int32 GenerateCollision(
 			const FReal CullDistance,
-			const bool bUseCCD,
+			const FReal Dt);
+
+		/**
+		 * @brief Generate a SweptConstraint as long as AABBs overlap
+		 * @return The number of collisions constraints that were activated
+		*/
+		int32 GenerateCollisionCCD(
+			const FReal CullDistance,
 			const FReal Dt);
 
 		/**
@@ -84,7 +91,8 @@ namespace Chaos
 		void SetCollision(const FPBDCollisionConstraint& Constraint);
 
 	private:
-		int32 GenerateCollisionImpl(const FReal CullDistance, const bool bUseCCD, const FReal Dt);
+		int32 GenerateCollisionImpl(const FReal CullDistance, const FReal Dt);
+		int32 GenerateCollisionCCDImpl(const FReal CullDistance, const FReal Dt);
 
 		/**
 		 * @brief Whether the two shapes are separated by less than CullDistance (i.e., we should run the narrow phase).
@@ -146,7 +154,6 @@ namespace Chaos
 		*/
 		int32 GenerateCollisions(
 			const FReal CullDistance,
-			const bool bUseCCD,
 			const FReal Dt,
 			FCollisionContext& Context);
 
