Geometry Collection : Prevent removal of hidden faces on load.

diff --git a/Engine/Source/Runtime/Experimental/Chaos/Private/GeometryCollection/GeometryCollection.cpp b/Engine/Source/Runtime/Experimental/Chaos/Private/GeometryCollection/GeometryCollection.cpp
--- a/Engine/Source/Runtime/Experimental/Chaos/Private/GeometryCollection/GeometryCollection.cpp
+++ b/Engine/Source/Runtime/Experimental/Chaos/Private/GeometryCollection/GeometryCollection.cpp
@@ -1107,10 +1107,6 @@ void FGeometryCollection::Serialize(Chaos::FChaosArchive& Ar)
 			{
 				bool bHasExemplar = ExemplarIndex[TransformGroupIndex] > INDEX_NONE;
 				bool bHasGeometry = HasVisibleFaces(TransformGroupIndex);
-				if (!bHasGeometry && TransformToGeometryIndex[TransformGroupIndex] != INDEX_NONE)
-				{
-					InvalidGeometry.Add(TransformToGeometryIndex[TransformGroupIndex]);
-				}
 				bool bHasRigidChildren = RigidChildren[TransformGroupIndex];
 				bool bKeep = true;
 				if (SimulationType[TransformGroupIndex] == ESimulationTypes::FST_None && !bHasExemplar) // handle exemplars with no exemplar (remove or convert to rigid)
@@ -1125,17 +1121,12 @@ void FGeometryCollection::Serialize(Chaos::FChaosArchive& Ar)
 						bKeep = false;
 					}
 				}
-				else if (SimulationType[TransformGroupIndex] == ESimulationTypes::FST_Rigid && !bHasGeometry) // handle rigids with no geometry (remove or convert to clustered)
+				else if (SimulationType[TransformGroupIndex] == ESimulationTypes::FST_Rigid ) // handle internal rigids
 				{
 					if (bHasRigidChildren)
 					{
 						SimulationType[TransformGroupIndex] = ESimulationTypes::FST_Clustered;
 					}
-					else
-					{
-						InvalidTransforms.Add(TransformGroupIndex);
-						bKeep = false;
-					}
 				}
 				if (bKeep && SimulationType[TransformGroupIndex] != ESimulationTypes::FST_None && Parent[TransformGroupIndex] != INDEX_NONE)
 				{
