Chaos - fix non-manifold path (for RBAN jitter)

diff --git a/Engine/Source/Runtime/Experimental/Chaos/Private/Chaos/Collision/ParticlePairMidPhase.cpp b/Engine/Source/Runtime/Experimental/Chaos/Private/Chaos/Collision/ParticlePairMidPhase.cpp
--- a/Engine/Source/Runtime/Experimental/Chaos/Private/Chaos/Collision/ParticlePairMidPhase.cpp
+++ b/Engine/Source/Runtime/Experimental/Chaos/Private/Chaos/Collision/ParticlePairMidPhase.cpp
@@ -55,7 +55,8 @@ namespace Chaos
 		FGeometryParticleHandle* Particle1,
 		const FPerShapeData* InShape1,
 		const FReal CullDistance,
-		const EContactShapesType ShapePairType)
+		const EContactShapesType ShapePairType,
+		const bool bUseManifold)
 	{
 		const FImplicitObject* Implicit0 = InShape0->GetLeafGeometry();
 		const FBVHParticles* BVHParticles0 = FConstGenericParticleHandle(Particle0)->CollisionParticles().Get();
@@ -63,9 +64,8 @@ namespace Chaos
 		const FImplicitObject* Implicit1 = InShape1->GetLeafGeometry();
 		const FBVHParticles* BVHParticles1 = FConstGenericParticleHandle(Particle1)->CollisionParticles().Get();
 		const FRigidTransform3& ShapeRelativeTransform1 = InShape1->GetLeafRelativeTransform();
-		const bool bUseManifolds = true;
 
-		return FPBDCollisionConstraint::Make(Particle0, Implicit0, BVHParticles0, ShapeRelativeTransform0, Particle1, Implicit1, BVHParticles1, ShapeRelativeTransform1, CullDistance, bUseManifolds, ShapePairType);
+		return FPBDCollisionConstraint::Make(Particle0, Implicit0, BVHParticles0, ShapeRelativeTransform0, Particle1, Implicit1, BVHParticles1, ShapeRelativeTransform1, CullDistance, bUseManifold, ShapePairType);
 	}
 
 	TUniquePtr<FPBDCollisionConstraint> CreateImplicitPairConstraint(
@@ -79,10 +79,9 @@ namespace Chaos
 		const FRigidTransform3& ShapeRelativeTransform1,
 		const FReal CullDistance,
 		const EContactShapesType ShapePairType,
-		const bool bInUseManifold)
+		const bool bUseManifold)
 	{
-		const bool bUseManifolds = true;
-		return FPBDCollisionConstraint::Make(Particle0, Implicit0, BVHParticles0, ShapeRelativeTransform0, Particle1, Implicit1, BVHParticles1, ShapeRelativeTransform1, CullDistance, bUseManifolds, ShapePairType);
+		return FPBDCollisionConstraint::Make(Particle0, Implicit0, BVHParticles0, ShapeRelativeTransform0, Particle1, Implicit1, BVHParticles1, ShapeRelativeTransform1, CullDistance, bUseManifold, ShapePairType);
 	}
 
 
@@ -219,28 +218,30 @@ namespace Chaos
 
 	int32 FSingleShapePairCollisionDetector::GenerateCollision(
 		const FReal CullDistance,
+		const bool bUseManifold,
 		const FReal Dt)
 	{
 		if (DoBoundsOverlap(CullDistance))
 		{
-			return GenerateCollisionImpl(CullDistance, Dt);
+			return GenerateCollisionImpl(CullDistance, bUseManifold, Dt);
 		}
 		return 0;
 	}
 
 	int32 FSingleShapePairCollisionDetector::GenerateCollisionCCD(
 		const FReal CullDistance,
+		const bool bUseManifold,
 		const FReal Dt)
 	{
-		return GenerateCollisionCCDImpl(CullDistance, Dt);
+		return GenerateCollisionCCDImpl(CullDistance, bUseManifold, Dt);
 	}
 
-	void FSingleShapePairCollisionDetector::CreateConstraint(const FReal CullDistance)
+	void FSingleShapePairCollisionDetector::CreateConstraint(const FReal CullDistance, const bool bUseManifold)
 	{
 		PHYSICS_CSV_SCOPED_EXPENSIVE(PhysicsVerbose, NarrowPhase_CreateConstraint);
 		check(Constraint == nullptr);
 
-		Constraint = CreateShapePairConstraint(Particle0, Shape0, Particle1, Shape1, CullDistance, ShapePairType);
+		Constraint = CreateShapePairConstraint(Particle0, Shape0, Particle1, Shape1, CullDistance, ShapePairType, bUseManifold);
 
 		Constraint->GetContainerCookie().MidPhase = &MidPhase;
 		Constraint->GetContainerCookie().bIsMultiShapePair = false;
@@ -248,7 +249,10 @@ namespace Chaos
 		LastUsedEpoch = -1;
 	}
 
-	int32 FSingleShapePairCollisionDetector::GenerateCollisionImpl(const FReal CullDistance, const FReal Dt)
+	int32 FSingleShapePairCollisionDetector::GenerateCollisionImpl(
+		const FReal CullDistance, 
+		const bool bUseManifold,
+		const FReal Dt)
 	{
 		if (Constraint == nullptr)
 		{
@@ -256,7 +260,7 @@ namespace Chaos
 			// other, we never allocate a constraint for them. Once they overlap, we reuse the constraint
 			// until the owing particles are not overlapping. i.e., we keep the constraint even if
 			// the shape pairs stop overlapping, reusing it if they start overlapping again.
-			CreateConstraint(CullDistance);
+			CreateConstraint(CullDistance, bUseManifold);
 		}
 
 		if (Constraint != nullptr)
@@ -322,12 +326,15 @@ namespace Chaos
 		return 0;
 	}
 
-	int32 FSingleShapePairCollisionDetector::GenerateCollisionCCDImpl(const FReal CullDistance, const FReal Dt)
+	int32 FSingleShapePairCollisionDetector::GenerateCollisionCCDImpl(
+		const FReal CullDistance, 
+		const bool bUseManifold,
+		const FReal Dt)
 	{
 		if (Constraint == nullptr)
 		{
 			// Lazy creation of the constraint. 
-			CreateConstraint(CullDistance);
+			CreateConstraint(CullDistance, bUseManifold);
 		}
 
 		if (Constraint != nullptr)
@@ -446,6 +453,7 @@ namespace Chaos
 
 	int32 FMultiShapePairCollisionDetector::GenerateCollisions(
 		const FReal CullDistance,
+		const bool bUseManifold,
 		const FReal Dt,
 		FCollisionContext& Context)
 	{
@@ -588,11 +596,11 @@ namespace Chaos
 		const FRigidTransform3& ShapeRelativeTransform1,
 		const FReal CullDistance,
 		const EContactShapesType ShapePairType,
-		const bool bInUseManifold,
+		const bool bUseManifold,
 		const FCollisionParticlePairConstraintKey& Key)
 	{
 		PHYSICS_CSV_SCOPED_EXPENSIVE(PhysicsVerbose, NarrowPhase_CreateConstraint);
-		TUniquePtr<FPBDCollisionConstraint> Constraint = CreateImplicitPairConstraint(InParticle0, Implicit0, BVHParticles0, ShapeRelativeTransform0, InParticle1, Implicit1, BVHParticles1, ShapeRelativeTransform1, CullDistance, ShapePairType, bInUseManifold);
+		TUniquePtr<FPBDCollisionConstraint> Constraint = CreateImplicitPairConstraint(InParticle0, Implicit0, BVHParticles0, ShapeRelativeTransform0, InParticle1, Implicit1, BVHParticles1, ShapeRelativeTransform1, CullDistance, ShapePairType, bUseManifold);
 		
 		Constraint->GetContainerCookie().MidPhase = &MidPhase;
 		Constraint->GetContainerCookie().bIsMultiShapePair = true;
@@ -845,25 +853,29 @@ namespace Chaos
 
 		// Enable CCD?
 		const bool bUseCCD = Flags.bIsCCD && ShouldEnableCCD(Dt);
+		
+		// Allow Manifolds?
+		const bool bUseManifold = Context.bAllowManifolds;
+
 		// Run collision detection on all potentially colliding shape pairs
 		NumActiveConstraints = 0;
 		if (bUseCCD) // Generate CCD constraints as long as AABBs overlap
 		{
 			for (FSingleShapePairCollisionDetector& ShapePair : ShapePairDetectors)
 			{
-				NumActiveConstraints += ShapePair.GenerateCollisionCCD(CullDistance, Dt);
+				NumActiveConstraints += ShapePair.GenerateCollisionCCD(CullDistance, bUseManifold, Dt);
 			}
 		}
 		else
 		{
 			for (FSingleShapePairCollisionDetector& ShapePair : ShapePairDetectors)
 			{
-				NumActiveConstraints += ShapePair.GenerateCollision(CullDistance, Dt);
+				NumActiveConstraints += ShapePair.GenerateCollision(CullDistance, bUseManifold, Dt);
 			}
 		}
 		for (FMultiShapePairCollisionDetector& MultiShapePair : MultiShapePairDetectors)
 		{
-			NumActiveConstraints += MultiShapePair.GenerateCollisions(CullDistance, Dt, Context);
+			NumActiveConstraints += MultiShapePair.GenerateCollisions(CullDistance, bUseManifold, Dt, Context);
 		}
 
 		LastUsedEpoch = CollisionAllocator->GetCurrentEpoch();
diff --git a/Engine/Source/Runtime/Experimental/Chaos/Public/Chaos/Collision/ParticlePairMidPhase.h b/Engine/Source/Runtime/Experimental/Chaos/Public/Chaos/Collision/ParticlePairMidPhase.h
--- a/Engine/Source/Runtime/Experimental/Chaos/Public/Chaos/Collision/ParticlePairMidPhase.h
+++ b/Engine/Source/Runtime/Experimental/Chaos/Public/Chaos/Collision/ParticlePairMidPhase.h
@@ -67,6 +67,7 @@ namespace Chaos
 		*/
 		int32 GenerateCollision(
 			const FReal CullDistance,
+			const bool bUseManifold,
 			const FReal Dt);
 
 		/**
@@ -75,6 +76,7 @@ namespace Chaos
 		*/
 		int32 GenerateCollisionCCD(
 			const FReal CullDistance,
+			const bool bUseManifold,
 			const FReal Dt);
 
 		/**
@@ -91,8 +93,15 @@ namespace Chaos
 		void SetCollision(const FPBDCollisionConstraint& Constraint);
 
 	private:
-		int32 GenerateCollisionImpl(const FReal CullDistance, const FReal Dt);
-		int32 GenerateCollisionCCDImpl(const FReal CullDistance, const FReal Dt);
+		int32 GenerateCollisionImpl(
+			const FReal CullDistance, 
+			const bool bUseManifold,
+			const FReal Dt);
+			
+		int32 GenerateCollisionCCDImpl(
+			const FReal CullDistance, 
+			const bool bUseManifold,
+			const FReal Dt);
 
 		/**
 		 * @brief Whether the two shapes are separated by less than CullDistance (i.e., we should run the narrow phase).
@@ -103,7 +112,7 @@ namespace Chaos
 		/**
 		 * @brief Create a constraint
 		*/
-		void CreateConstraint(const FReal CullDistance);
+		void CreateConstraint(const FReal CullDistance, const bool bUseManifold);
 
 		FParticlePairMidPhase& MidPhase;
 		TUniquePtr<FPBDCollisionConstraint> Constraint;
@@ -154,6 +163,7 @@ namespace Chaos
 		*/
 		int32 GenerateCollisions(
 			const FReal CullDistance,
+			const bool bUseManifold,
 			const FReal Dt,
 			FCollisionContext& Context);
 
diff --git a/Engine/Source/Runtime/Experimental/Chaos/Public/Chaos/Framework/UncheckedArray.h b/Engine/Source/Runtime/Experimental/Chaos/Public/Chaos/Framework/UncheckedArray.h
--- a/Engine/Source/Runtime/Experimental/Chaos/Public/Chaos/Framework/UncheckedArray.h
+++ b/Engine/Source/Runtime/Experimental/Chaos/Public/Chaos/Framework/UncheckedArray.h
@@ -58,7 +58,7 @@ namespace Chaos
 		*/
 		inline void SetNum(const int32 InNum)
 		{
-			check(InNum < NumElements);
+			check(InNum < MaxElements);
 			NumElements = InNum;
 		}
 
