Chaos - debug draw settings for manifold point info

diff --git a/Engine/Source/Runtime/Experimental/Chaos/Private/Chaos/ChaosDebugDraw.cpp b/Engine/Source/Runtime/Experimental/Chaos/Private/Chaos/ChaosDebugDraw.cpp
--- a/Engine/Source/Runtime/Experimental/Chaos/Private/Chaos/ChaosDebugDraw.cpp
+++ b/Engine/Source/Runtime/Experimental/Chaos/Private/Chaos/ChaosDebugDraw.cpp
@@ -825,7 +825,7 @@ namespace Chaos
 
 				for (const FManifoldPoint& ManifoldPoint : Contact.GetManifoldPoints())
 				{
-					const bool bIsActive = !ManifoldPoint.NetPushOut.IsNearlyZero() || !ManifoldPoint.NetImpulse.IsNearlyZero();
+					const bool bIsActive = !ManifoldPoint.NetPushOut.IsNearlyZero() || !ManifoldPoint.NetImpulse.IsNearlyZero() || (!Contact.GetUseManifold() && !Contact.AccumulatedImpulse.IsNearlyZero());
 					if (!bIsActive && !bChaosDebugDebugDrawInactiveContacts)
 					{
 						continue;
@@ -901,7 +901,7 @@ namespace Chaos
 					FDebugDrawQueue::GetInstance().DrawDebugCircle(WorldPointLocation, 0.5f * Settings.DrawScale * Settings.ContactWidth, 12, DiscColor, false, KINDA_SMALL_NUMBER, Settings.DrawPriority, Settings.LineThickness, Axes.GetUnitAxis(EAxis::Y), Axes.GetUnitAxis(EAxis::Z), false);
 
 					// Previous points
-					if (bIsActive)
+					if (bIsActive && ManifoldPoint.Flags.bWasFrictionRestored)
 					{
 						const FVec3 WorldPrevPointLocation = SpaceTransform.TransformPosition(PointTransform.TransformPosition(ManifoldPoint.ShapeAnchorPoints[ContactPointOwner]));
 						const FVec3 WorldPrevPlaneLocation = SpaceTransform.TransformPosition(PlaneTransform.TransformPosition(ManifoldPoint.ShapeAnchorPoints[ContactPlaneOwner]));
