Chaos - Convex-TriMesh optimizations

diff --git a/Engine/Source/Programs/HeadlessChaos/Private/HeadlessChaosTestDetectCollision.cpp b/Engine/Source/Programs/HeadlessChaos/Private/HeadlessChaosTestDetectCollision.cpp
--- a/Engine/Source/Programs/HeadlessChaos/Private/HeadlessChaosTestDetectCollision.cpp
+++ b/Engine/Source/Programs/HeadlessChaos/Private/HeadlessChaosTestDetectCollision.cpp
@@ -11,11 +11,6 @@
 #include "Chaos/Utilities.h"
 #include "Modules/ModuleManager.h"
 
-namespace Chaos
-{
-	extern bool bChaos_Collision_Manifold_FixNormalsInWorldSpace;
-}
-
 namespace ChaosTest 
 {
 	using namespace Chaos;
diff --git a/Engine/Source/Programs/HeadlessChaos/Private/HeadlessChaosTestOneShotManifolds.cpp b/Engine/Source/Programs/HeadlessChaos/Private/HeadlessChaosTestOneShotManifolds.cpp
--- a/Engine/Source/Programs/HeadlessChaos/Private/HeadlessChaosTestOneShotManifolds.cpp
+++ b/Engine/Source/Programs/HeadlessChaos/Private/HeadlessChaosTestOneShotManifolds.cpp
@@ -31,7 +31,6 @@ namespace Chaos
 			const FReal Dt,
 			FPBDCollisionConstraint& Constraint);
 	}
-	extern bool bChaos_Collision_Manifold_FixNormalsInWorldSpace;
 }
 
 
diff --git a/Engine/Source/Runtime/Experimental/Chaos/Private/Chaos/ChaosDebugDraw.cpp b/Engine/Source/Runtime/Experimental/Chaos/Private/Chaos/ChaosDebugDraw.cpp
--- a/Engine/Source/Runtime/Experimental/Chaos/Private/Chaos/ChaosDebugDraw.cpp
+++ b/Engine/Source/Runtime/Experimental/Chaos/Private/Chaos/ChaosDebugDraw.cpp
@@ -31,8 +31,6 @@
 namespace Chaos
 {
 
-	extern bool bChaos_Collision_Manifold_FixNormalsInWorldSpace;
-
 	namespace DebugDraw
 	{
 #if CHAOS_DEBUG_DRAW
diff --git a/Engine/Source/Runtime/Experimental/Chaos/Private/Chaos/Collision/CollisionSolver.cpp b/Engine/Source/Runtime/Experimental/Chaos/Private/Chaos/Collision/CollisionSolver.cpp
--- a/Engine/Source/Runtime/Experimental/Chaos/Private/Chaos/Collision/CollisionSolver.cpp
+++ b/Engine/Source/Runtime/Experimental/Chaos/Private/Chaos/Collision/CollisionSolver.cpp
@@ -60,8 +60,6 @@ namespace Chaos
 	}
 	using namespace CVars;
 
-	extern bool bChaos_Collision_Manifold_FixNormalsInWorldSpace;
-
 	namespace Collisions
 	{
 		inline FReal GetConstraintStiffness(const FReal InStiffness, const FContactIterationParameters& IterationParameters)
@@ -304,16 +302,7 @@ namespace Chaos
 			// Calculate the position error we need to correct, including static friction and restitution
 			// Position correction uses the deepest point on each body (see velocity correction which uses average contact)
 
-			FVec3 ContactNormal;
-			if (bChaos_Collision_Manifold_FixNormalsInWorldSpace)
-			{
-				ContactNormal = ManifoldPoint.ManifoldContactNormal;
-			}
-			else
-			{
-				const FRotation3 PlaneQ = (ManifoldPoint.ContactPoint.ContactNormalOwnerIndex == 0) ? Body0.Q() : Body1.Q();
-				ContactNormal = PlaneQ * ManifoldPoint.ManifoldContactNormal;
-			}
+			const FVec3 ContactNormal = ManifoldPoint.ManifoldContactNormal;
 
 			const bool bApplyStaticFriction = (ManifoldPoint.bActive && ManifoldPoint.bInsideStaticFrictionCone && ManifoldPoint.bPotentialRestingContact && Chaos_Manifold_PushOut_StaticFriction);
 			FVec3 LocalContactPoint0 = ManifoldPoint.CoMContactPoints[0];
diff --git a/Engine/Source/Runtime/Experimental/Chaos/Private/Chaos/Collision/ContactPointsMiscShapes.cpp b/Engine/Source/Runtime/Experimental/Chaos/Private/Chaos/Collision/ContactPointsMiscShapes.cpp
--- a/Engine/Source/Runtime/Experimental/Chaos/Private/Chaos/Collision/ContactPointsMiscShapes.cpp
+++ b/Engine/Source/Runtime/Experimental/Chaos/Private/Chaos/Collision/ContactPointsMiscShapes.cpp
@@ -348,11 +348,17 @@ namespace Chaos
 		return GJKImplicitContactPoint<FConvex>(A, ATransform, B, BTransform, CullDistance, ShapePadding);
 	}
 
-	template <typename TriMeshType>
-	FContactPoint ConvexTriangleMeshContactPoint(const FImplicitObject& A, const FRigidTransform3& ATransform, const TriMeshType& B, const FRigidTransform3& BTransform, const FReal CullDistance, const FReal ShapePadding)
+	FContactPoint ConvexTriangleMeshContactPoint(const FImplicitObject& A, const FRigidTransform3& ATransform, const FImplicitObject& B, const FRigidTransform3& BTransform, const FReal CullDistance, const FReal ShapePadding)
 	{
 		CONDITIONAL_SCOPE_CYCLE_COUNTER(STAT_Collisions_ConvexTriangleMeshContactPoint, ConstraintsDetailedStats);
-		return GJKImplicitContactPoint<FConvex>(A, ATransform, B, BTransform, CullDistance, ShapePadding);
+
+		// Call GJK with the concrete trimesh type (scaled, instanced, raw)
+		return Utilities::CastWrapped<FTriangleMeshImplicitObject>(B, 
+			[&](auto BConcretePtr)
+			{
+				check(BConcretePtr != nullptr);
+				return GJKImplicitContactPoint<FConvex>(A, ATransform, *BConcretePtr, BTransform, CullDistance, ShapePadding);
+			});
 	}
 
 	template <typename TriMeshType>
@@ -423,8 +429,6 @@ namespace Chaos
 	template FContactPoint SphereTriangleMeshContactPoint<FTriangleMeshImplicitObject>(const TSphere<FReal, 3>& A, const FRigidTransform3& ATransform, const FTriangleMeshImplicitObject& B, const FRigidTransform3& BTransform, const FReal CullDistance, const FReal ShapePadding);
 	template FContactPoint CapsuleTriangleMeshContactPoint<TImplicitObjectScaled<class Chaos::FTriangleMeshImplicitObject, 1>>(const FCapsule& A, const FRigidTransform3& ATransform, const TImplicitObjectScaled<class Chaos::FTriangleMeshImplicitObject, 1>& B, const FRigidTransform3& BTransform, const FReal CullDistance, const FReal ShapePadding);
 	template FContactPoint CapsuleTriangleMeshContactPoint<FTriangleMeshImplicitObject>(const FCapsule& A, const FRigidTransform3& ATransform, const FTriangleMeshImplicitObject& B, const FRigidTransform3& BTransform, const FReal CullDistance, const FReal ShapePadding);
-	template FContactPoint ConvexTriangleMeshContactPoint<TImplicitObjectScaled<class Chaos::FTriangleMeshImplicitObject, 1>>(const FImplicitObject& A, const FRigidTransform3& ATransform, const TImplicitObjectScaled<class Chaos::FTriangleMeshImplicitObject, 1>& B, const FRigidTransform3& BTransform, const FReal CullDistance, const FReal ShapePadding);
-	template FContactPoint ConvexTriangleMeshContactPoint<FTriangleMeshImplicitObject>(const FImplicitObject& A, const FRigidTransform3& ATransform, const FTriangleMeshImplicitObject& B, const FRigidTransform3& BTransform, const FReal CullDistance, const FReal ShapePadding);
 
 	template void GJKImplicitManifold<FConvex, FHeightField>(const FImplicitObject& A, const FRigidTransform3& ATransform, const FHeightField& B, const FRigidTransform3& BTransform, const FReal CullDistance, const FReal ShapePadding, TArray<FContactPoint>& ContactPoints);
 	template void GJKImplicitManifold<FCapsule, FHeightField>(const FImplicitObject& A, const FRigidTransform3& ATransform, const FHeightField& B, const FRigidTransform3& BTransform, const FReal CullDistance, const FReal ShapePadding, TArray<FContactPoint>& ContactPoints);
diff --git a/Engine/Source/Runtime/Experimental/Chaos/Private/Chaos/Collision/PBDCollisionConstraint.cpp b/Engine/Source/Runtime/Experimental/Chaos/Private/Chaos/Collision/PBDCollisionConstraint.cpp
--- a/Engine/Source/Runtime/Experimental/Chaos/Private/Chaos/Collision/PBDCollisionConstraint.cpp
+++ b/Engine/Source/Runtime/Experimental/Chaos/Private/Chaos/Collision/PBDCollisionConstraint.cpp
@@ -32,9 +32,6 @@ namespace Chaos
 	FRealSingle Chaos_GBFCharacteristicTimeRatio = 1.0f;
 	FAutoConsoleVariableRef CVarChaos_GBFCharacteristicTimeRatio(TEXT("p.Chaos.Collision.GBFCharacteristicTimeRatio"), Chaos_GBFCharacteristicTimeRatio, TEXT("The ratio between characteristic time and Dt"));
 
-	bool bChaos_Manifold_EnabledWithJoints = true;
-	FAutoConsoleVariableRef CVarChaos_Manifold_EnabledWithJoints(TEXT("p.Chaos.Collision.Manifold.EnabledWithJoints"), bChaos_Manifold_EnabledWithJoints, TEXT(""));
-
 	bool bChaos_Manifold_EnableGjkWarmStart = true;
 	FAutoConsoleVariableRef CVarChaos_Manifold_EnableGjkWarmStart(TEXT("p.Chaos.Collision.Manifold.EnableGjkWarmStart"), bChaos_Manifold_EnableGjkWarmStart, TEXT(""));
 
@@ -70,8 +67,6 @@ namespace Chaos
 	FCollisionTolerances Chaos_Manifold_Tolerances;
 
 
-	extern bool bChaos_Collision_Manifold_FixNormalsInWorldSpace;
-
 	FString FPBDCollisionConstraint::ToString() const
 	{
 		return FString::Printf(TEXT("Particle:%s, Levelset:%s, AccumulatedImpulse:%s"), *Particle[0]->ToString(), *Particle[1]->ToString(), *AccumulatedImpulse.ToString());
@@ -250,7 +245,7 @@ namespace Chaos
 
 		CullDistance = InCullDistance;
 
-		Flags.bUseManifold = bInUseManifold && CanUseManifold(Particle[0], Particle[1]);
+		Flags.bUseManifold = bInUseManifold;
 		Flags.bUseIncrementalManifold = true;	// This will get changed later if we call AddOneShotManifoldContact
 
 		const FReal Margin0 = GetImplicit0()->GetMargin();
@@ -553,13 +548,6 @@ namespace Chaos
 		Flags.bUseIncrementalManifold = true;
 	}
 
-	bool FPBDCollisionConstraint::CanUseManifold(FGeometryParticleHandle* Particle0, FGeometryParticleHandle* Particle1) const
-	{
-		// Do not use manifolds when a body is connected by a joint to another. Manifolds do not work when the bodies may be moved
-		// and rotated by significant amounts and joints can do that.
-		return bChaos_Manifold_EnabledWithJoints || ((Particle0->ParticleConstraints().Num() == 0) && (Particle1->ParticleConstraints().Num() == 0));
-	}
-
 	void FPBDCollisionConstraint::ResetManifold()
 	{
 		ResetSavedManifoldPoints();
diff --git a/Engine/Source/Runtime/Experimental/Chaos/Private/Chaos/Collision/SolverCollisionContainer.cpp b/Engine/Source/Runtime/Experimental/Chaos/Private/Chaos/Collision/SolverCollisionContainer.cpp
--- a/Engine/Source/Runtime/Experimental/Chaos/Private/Chaos/Collision/SolverCollisionContainer.cpp
+++ b/Engine/Source/Runtime/Experimental/Chaos/Private/Chaos/Collision/SolverCollisionContainer.cpp
@@ -31,8 +31,6 @@ namespace Chaos
 	}
 	using namespace CVars;
 
-	extern bool bChaos_Collision_Manifold_FixNormalsInWorldSpace;
-
 
 	//////////////////////////////////////////////////////////////////////////////////////////////////
 	//////////////////////////////////////////////////////////////////////////////////////////////////
diff --git a/Engine/Source/Runtime/Experimental/Chaos/Private/Chaos/CollisionOneShotManifolds.cpp b/Engine/Source/Runtime/Experimental/Chaos/Private/Chaos/CollisionOneShotManifolds.cpp
--- a/Engine/Source/Runtime/Experimental/Chaos/Private/Chaos/CollisionOneShotManifolds.cpp
+++ b/Engine/Source/Runtime/Experimental/Chaos/Private/Chaos/CollisionOneShotManifolds.cpp
@@ -33,9 +33,6 @@ namespace Chaos
 	FRealSingle Chaos_Collision_Manifold_MinFaceSearchDistance = 1.0f;
 	FAutoConsoleVariableRef CVarChaosCollisioConvexManifoldMinFaceSearchDistance(TEXT("p.Chaos.Collision.Manifold.MinFaceSearchDistance"), Chaos_Collision_Manifold_MinFaceSearchDistance, TEXT(""));
 
-	bool bChaos_Collision_Manifold_FixNormalsInWorldSpace = true;
-	FAutoConsoleVariableRef CVarChaosCollisioFixNormalsInWorldSpace(TEXT("p.Chaos.Collision.Manifold.FixNormalsInWorldSpace"), bChaos_Collision_Manifold_FixNormalsInWorldSpace, TEXT("Fix normals in world space at beginning of frame"));	
-
 	bool ForceOneShotManifoldEdgeEdgeCaseZeroCullDistance = false;
 	FAutoConsoleVariableRef CVarForceOneShotManifoldEdgeEdgeCaseZeroCullDistance(TEXT("p.Chaos.Collision.Manifold.ForceOneShotManifoldEdgeEdgeCaseZeroCullDistance"), ForceOneShotManifoldEdgeEdgeCaseZeroCullDistance,
 	TEXT("If enabled, if one shot manifold hits edge/edge case, we will force a cull distance of zero. That means edge/edge contacts will be thrown out if separated at all. Only applies to Convex/Convex oneshot impl."));
@@ -248,7 +245,8 @@ namespace Chaos
 		// Remove spurious edge contact by removing edge-edge contact that are in the plane of another vertex-plane contact.
 		// Note: this assumes that the first shape in the contacts is the same convex shape (e.g., a box), and the second
 		// can be from any other shape(s), e.g., triangles of a non-convex shape.
-		void PruneEdgeContactPoints(TArray<FContactPoint>& ContactPoints, const FReal MaxPlaneDistance)
+		template<bool bMaintainOrder>
+		void PruneEdgeContactPointsImp(TArray<FContactPoint>& ContactPoints, const FReal MaxPlaneDistance)
 		{
 			if (!bChaos_Collision_EnableEdgePrune)
 			{
@@ -281,7 +279,14 @@ namespace Chaos
 
 								// Note: Cannot use RemoveAtSwap unless we re-sort at the end. See ReduceManifoldContactPointsTriangeMesh
 								const bool bAllowShrinking = false;
-								ContactPoints.RemoveAt(EdgeContactIndex, 1, bAllowShrinking);
+								if (bMaintainOrder)
+								{
+									ContactPoints.RemoveAt(EdgeContactIndex, 1, bAllowShrinking);
+								}
+								else
+								{
+									ContactPoints.RemoveAtSwap(EdgeContactIndex, 1, bAllowShrinking);
+								}
 								--EdgeContactIndex;
 								--NumContactPoints;
 								break;
@@ -292,6 +297,16 @@ namespace Chaos
 			}
 		}
 
+		void PruneEdgeContactPointsOrdered(TArray<FContactPoint>& ContactPoints, const FReal MaxPlaneDistance)
+		{
+			PruneEdgeContactPointsImp<true>(ContactPoints, MaxPlaneDistance);
+		}
+
+		void PruneEdgeContactPointsUnordered(TArray<FContactPoint>& ContactPoints, const FReal MaxPlaneDistance)
+		{
+			PruneEdgeContactPointsImp<false>(ContactPoints, MaxPlaneDistance);
+		}
+
 		// Reduce the number of contact points (in place)
 		// Prerequisites to calling this function:
 		// ContactPoints are sorted on phi (ascending)
@@ -305,6 +320,20 @@ namespace Chaos
 			// Point 1) is the deepest contact point
 			// It is already in position
 			
+			//
+			// @todo(chaos): actually this isn't really good enough. We want to take the outer 4 points for the most stable manifold, but also
+			// must take the deepest for correct behaviour. We also have a problem right now where if we have a number of contacts at the same 
+			// depth, the "deepest" one is arbitrary and may be a poor choice as the anchor of the manifold building. E.g., Consider 5 contact 
+			// points in a X config where the center point is first in the list - this leads to a triangular/unstable manifold.
+			// 
+			// Maybe we should do this:
+			//	- build the manifold using the current process
+			//	- revisit point 0 at the end and replace with the best choice for stability if the depth is similar
+			// Or this:
+			//	- build manifold for best stabilty (largest convex area)
+			//	- if manifold does not contain deepest point, replace one point so minimize area reduction
+			//
+			
 			// Point 2) Find the point with the largest distance to the deepest contact point
 			{
 				uint32 FarthestPointIndex = 1;
@@ -365,6 +394,8 @@ namespace Chaos
 				Swap(ContactPoints[3], ContactPoints[LargestTrianglePointIndex]);
 			}
 
+			// Revisit Point 0 - if we have a choice that increases area at similar depth, use it
+
 			ContactPoints.SetNum(4); // Will end up with 4 points
 		}
 
@@ -516,6 +547,36 @@ namespace Chaos
 			return Contact;
 		}
 
+		// GJK contact point between a convex and a triangle
+		// This assumes that the triangle has been transformed into convex space
+		template <typename GeometryA>
+		FContactPoint GJKContactPointConvexTriangle(const GeometryA& A, const FTriangle& B, FReal MarginA, FReal& OutMaxMarginDelta, int32& VertexIndexA, int32& VertexIndexB)
+		{
+			SCOPE_CYCLE_COUNTER_MANIFOLD_GJK();
+
+			FContactPoint Contact;
+
+			FReal Penetration;
+			FVec3 ClosestA, ClosestB, Normal;
+
+			// GJK and EPA tolerances. See comments in GJKContactPointMargin
+			const FReal GJKEpsilon = Chaos_Collision_GJKEpsilon;
+			const FReal EPAEpsilon = Chaos_Collision_EPAEpsilon;
+
+			const TGJKCoreShape<GeometryA> AWithMargin(A, MarginA);
+			const TGJKShape<FTriangle> BNoMargin(B);
+
+			if (GJKPenetrationSameSpace(AWithMargin, BNoMargin, Penetration, ClosestA, ClosestB, Normal, VertexIndexA, VertexIndexB, OutMaxMarginDelta, GJKEpsilon, EPAEpsilon))
+			{
+				Contact.ShapeContactPoints[0] = ClosestA;
+				Contact.ShapeContactPoints[1] = ClosestB;
+				Contact.ShapeContactNormal = -Normal;	// We want normal pointing from B to A
+				Contact.Phi = -Penetration;
+			}
+
+			return Contact;
+		}
+
 		// Find the the most opposing plane given a position and a direction
 		template <typename ConvexImplicitType>
 		void FindBestPlane(
@@ -956,6 +1017,206 @@ namespace Chaos
 		}
 
 
+		template <typename ConvexType>
+		void ConstructPlanarConvexTriangleOneShotManifold(const ConvexType& Convex, const FTriangle& Triangle, const FReal CullDistance, TCArray<FContactPoint, 4>& OutContactPoints)
+		{
+			SCOPE_CYCLE_COUNTER_MANIFOLD();
+			check(OutContactPoints.Num() == 0);
+
+			// Find the deepest penetration. This is used to determine the planes and points to use for the manifold
+			// MaxMarginDelta is an upper bound on the distance from the contact on the rounded core shape to the actual shape surface. 
+			const FReal ConvexMargin = Convex.GetMargin();
+			FReal MaxMarginDelta = FReal(0);
+			int32 VertexIndexA = INDEX_NONE, VertexIndexB = INDEX_NONE;
+			FContactPoint GJKContactPoint = GJKContactPointConvexTriangle(Convex, Triangle, ConvexMargin, MaxMarginDelta, VertexIndexA, VertexIndexB);
+
+			PHYSICS_CSV_CUSTOM_EXPENSIVE(PhysicsCounters, NumManifoldsGJKCalled, 1, ECsvCustomStatOp::Accumulate);
+
+			// GJK is using margins and rounded corner, so if we have a corner-to-corner contact it will under-report the actual distance by an amount that depends on how
+			// "pointy" the edge/corner is - this error is bounded by MaxMarginDelta.
+			const FReal GJKCullDistance = CullDistance + MaxMarginDelta;
+			if (GJKContactPoint.Phi > GJKCullDistance)
+			{
+				PHYSICS_CSV_CUSTOM_EXPENSIVE(PhysicsCounters, NumManifoldsGJKCulled, 1, ECsvCustomStatOp::Accumulate);
+				return;
+			}
+
+			PHYSICS_CSV_CUSTOM_EXPENSIVE(PhysicsCounters, NumManifoldsCreated, 1, ECsvCustomStatOp::Accumulate);
+
+			// @todo(chaos): handle zero margins...
+			const FReal PlaneSearchDistance = ConvexMargin;
+
+			const FVec3 SeparationDirectionLocalConvex1 = GJKContactPoint.ShapeContactNormal;
+			const int32 MostOpposingPlaneIndexConvex1 = SelectContactPlane(Convex, GJKContactPoint.ShapeContactPoints[0], SeparationDirectionLocalConvex1, PlaneSearchDistance, VertexIndexA);
+			FVec3 BestPlanePosition1, BestPlaneNormal1;
+			Convex.GetPlaneNX(MostOpposingPlaneIndexConvex1, BestPlaneNormal1, BestPlanePosition1);
+			const FReal BestPlaneDotNormalConvex1 = FMath::Abs(FVec3::DotProduct(-SeparationDirectionLocalConvex1, BestPlaneNormal1));
+
+			// @todo(chaos): can be simpler for FTriangle
+			const FVec3 SeparationDirectionLocalConvex2 = GJKContactPoint.ShapeContactNormal;
+			const int32 MostOpposingPlaneIndexConvex2 = SelectContactPlane(Triangle, GJKContactPoint.ShapeContactPoints[1], -SeparationDirectionLocalConvex2, PlaneSearchDistance, VertexIndexB);
+			FVec3 BestPlanePosition2, BestPlaneNormal2;
+			Triangle.GetPlaneNX(MostOpposingPlaneIndexConvex2, BestPlaneNormal2, BestPlanePosition2);
+			const FReal BestPlaneDotNormalConvex2 = FMath::Abs(FVec3::DotProduct(SeparationDirectionLocalConvex2, BestPlaneNormal2));
+
+			const FReal SmallBiasToPreventFeatureFlipping = 0.002f; // This improves frame coherence by penalizing convex 1 in favour of convex 2
+			bool ReferenceFaceConvex1 = true; // Is the reference face on convex1 or convex2?
+			if (BestPlaneDotNormalConvex2 + SmallBiasToPreventFeatureFlipping > BestPlaneDotNormalConvex1)
+			{
+				ReferenceFaceConvex1 = false;
+			}
+
+			// Is this a vertex-plane or edge-edge contact? 
+			const FReal PlaneContactNormalEpsilon = Chaos_Collision_Manifold_PlaneContactNormalEpsilon;
+			const bool bIsPlaneContact = FMath::IsNearlyEqual(BestPlaneDotNormalConvex1, (FReal)1., PlaneContactNormalEpsilon) || FMath::IsNearlyEqual(BestPlaneDotNormalConvex2, (FReal)1., PlaneContactNormalEpsilon);
+
+			// For edge-edge contacts, we find the edges involved and project the contact onto the edges (if necessary - see below)
+			if (!bIsPlaneContact)
+			{
+				SCOPE_CYCLE_COUNTER_MANIFOLD_ADDEDGEEDGE();
+
+				if (ForceOneShotManifoldEdgeEdgeCaseZeroCullDistance && GJKContactPoint.Phi > 0)
+				{
+					return;
+				}
+
+				// This is an edge-edge contact
+				GJKContactPoint.ContactType = EContactPointType::EdgeEdge;
+
+				// @todo(chaos): remove this if we ditch convex margins
+				// We only need to find the best edges when we have convex margins enabled because then the GJK result
+				// is a contact on the rounded margin-reduced shape and not the actual shape. If we are using zero
+				// margins, then the GJK result is already on the real shape surface. Likewise, if we have a
+				// quadratic-versus-polygonal contact, the GJK results are already on the real surfaces.
+				if (ConvexMargin > 0)
+				{
+					// @todo(chaos): this does not work well when the edges are parallel. We should always have points with zero
+					// position delta perpendicular to the normal, but that is not the case for parallel edges
+					FVec3 ShapeEdgePos1 = Convex.GetClosestEdgePosition(MostOpposingPlaneIndexConvex1, GJKContactPoint.ShapeContactPoints[0]);
+					FVec3 ShapeEdgePos2 = Triangle.GetClosestEdgePosition(MostOpposingPlaneIndexConvex2, GJKContactPoint.ShapeContactPoints[1]);
+
+					const FVec3 EdgePos1In2 = ShapeEdgePos1;
+					const FVec3 EdgePos2In2 = ShapeEdgePos2;
+					const FReal EdgePhi = FVec3::DotProduct(EdgePos1In2 - EdgePos2In2, GJKContactPoint.ShapeContactNormal);
+
+					// We now have an accurate separation, so reject points beyond cull distance.
+					// This is quite important because if we use a larger cull distance for edge-edge contacts than for vertex-plane
+					// contacts, the pruning of unnecessary edge-edge contacts does not work (Tri Mesh and heighfield)
+					if (EdgePhi > CullDistance)
+					{
+						return;
+					}
+
+					// @todo(chaos): we leave the contact point on the second shape because this is better for triangle meshes (see PruneContactEdges). For other shapes it should probably be the average
+					GJKContactPoint.ShapeContactPoints[0] = EdgePos2In2 + EdgePhi * GJKContactPoint.ShapeContactNormal;
+					GJKContactPoint.ShapeContactPoints[1] = EdgePos2In2;
+					GJKContactPoint.Phi = EdgePhi;
+					// Normal unchanged from GJK result
+				}
+
+				OutContactPoints.Add(GJKContactPoint);
+				return;
+			}
+
+			// For vertex-plane contacts, we use a convex face as the manifold plane
+			const FVec3 RefSeparationDirection = ReferenceFaceConvex1 ? SeparationDirectionLocalConvex1 : SeparationDirectionLocalConvex2;
+			const FVec3 RefPlaneNormal = ReferenceFaceConvex1 ? BestPlaneNormal1 : BestPlaneNormal2;
+			const FVec3 RefPlanePosition = ReferenceFaceConvex1 ? BestPlanePosition1 : BestPlanePosition2;
+
+			// @todo(chaos): fix use of hard-coded max array size
+			// We will use a double buffer as an optimization
+			const int32 MaxContactPointCount = 32; // This should be tuned
+			int32 ContactPointCount = 0;
+			FVec3* ClippedVertices = nullptr;
+			FVec3 ClippedVertices1[MaxContactPointCount];
+			FVec3 ClippedVertices2[MaxContactPointCount];
+
+			{
+				SCOPE_CYCLE_COUNTER_MANIFOLD_CLIP();
+				if (ReferenceFaceConvex1)
+				{
+					// @todo(chaos): non-transforming version for triangles
+					ClippedVertices = GenerateConvexManifoldClippedVertices(
+						Convex,
+						Triangle,
+						FRigidTransform3(),
+						MostOpposingPlaneIndexConvex1,
+						MostOpposingPlaneIndexConvex2,
+						RefPlaneNormal,
+						RefPlanePosition,
+						ClippedVertices1,
+						ClippedVertices2,
+						(uint32&)ContactPointCount,
+						(uint32)MaxContactPointCount);
+				}
+				else
+				{
+					ClippedVertices = GenerateConvexManifoldClippedVertices(
+						Triangle,
+						Convex,
+						FRigidTransform3(),
+						MostOpposingPlaneIndexConvex2,
+						MostOpposingPlaneIndexConvex1,
+						RefPlaneNormal,
+						RefPlanePosition,
+						ClippedVertices1,
+						ClippedVertices2,
+						(uint32&)ContactPointCount,
+						(uint32)MaxContactPointCount);
+				}
+			}
+
+			// If we have the max number of contact points already, they will be in cyclic order. Stability is better
+			// if we solve points non-sequentially (e.g., on a box, solve one point, then it's opposite corner).
+			// If we have more than 4 contacts, the contact reduction step already effectively does something similar.
+			if (ContactPointCount == 4)
+			{
+				Swap(ClippedVertices[1], ClippedVertices[2]);
+			}
+
+			// Reduce number of contacts to the maximum allowed
+			if (ContactPointCount > 4)
+			{
+				SCOPE_CYCLE_COUNTER_MANIFOLD_REDUCE();
+
+				FRotation3 RotateSeperationToZ = FRotation3::FromRotatedVector(RefPlaneNormal, FVec3(0.0f, 0.0f, 1.0f));
+				for (int32 ContactPointIndex = 0; ContactPointIndex < ContactPointCount; ++ContactPointIndex)
+				{
+					ClippedVertices[ContactPointIndex] = RotateSeperationToZ * ClippedVertices[ContactPointIndex];
+				}
+
+				ContactPointCount = ReduceManifoldContactPoints(ClippedVertices, ContactPointCount);
+
+				for (int32 ContactPointIndex = 0; ContactPointIndex < ContactPointCount; ++ContactPointIndex)
+				{
+					ClippedVertices[ContactPointIndex] = RotateSeperationToZ.Inverse() * ClippedVertices[ContactPointIndex];
+				}
+			}
+
+			// Generate the contact points from the clipped vertices
+			{
+				SCOPE_CYCLE_COUNTER_MANIFOLD_ADDFACEVERTEX();
+				check(ContactPointCount + OutContactPoints.Num() <= OutContactPoints.Max());
+
+				EContactPointType ContactType = ReferenceFaceConvex1 ? EContactPointType::PlaneVertex : EContactPointType::VertexPlane;
+
+				for (int32 ContactPointIndex = 0; ContactPointIndex < ContactPointCount; ++ContactPointIndex)
+				{
+					FContactPoint& ContactPoint = OutContactPoints[OutContactPoints.Add()];
+					FVec3 VertexInReferenceCoordinates = ClippedVertices[ContactPointIndex];
+					FVec3 PointProjectedOntoReferenceFace = VertexInReferenceCoordinates - FVec3::DotProduct(VertexInReferenceCoordinates - RefPlanePosition, RefPlaneNormal) * RefPlaneNormal;
+					FVec3 ClippedPointInOtherCoordinates = VertexInReferenceCoordinates;
+
+					ContactPoint.ShapeContactPoints[0] = ReferenceFaceConvex1 ? PointProjectedOntoReferenceFace : ClippedPointInOtherCoordinates;
+					ContactPoint.ShapeContactPoints[1] = ReferenceFaceConvex1 ? ClippedPointInOtherCoordinates : PointProjectedOntoReferenceFace;
+					ContactPoint.ShapeContactNormal = SeparationDirectionLocalConvex2;
+					ContactPoint.Phi = FVec3::DotProduct(PointProjectedOntoReferenceFace - VertexInReferenceCoordinates, ReferenceFaceConvex1 ? SeparationDirectionLocalConvex1 : -SeparationDirectionLocalConvex2);
+					ContactPoint.ContactType = ContactType;
+				}
+			}
+		}
+
+
 
 		//
 		// Explicit instantiations of all convex-convex manifold combinations we support
@@ -1213,7 +1474,35 @@ namespace Chaos
 				const FRigidTransform3& Convex2Transform, //world
 				const FReal Dt,
 				FPBDCollisionConstraint& Constraint);
-		
+	
+	
+		template
+			void ConstructPlanarConvexTriangleOneShotManifold(
+				const FImplicitConvex3& Convex, 
+				const FTriangle& Triangle, 
+				const FReal CullDistance, 
+				TCArray<FContactPoint, 4>& OutContactPoints);
+
+		template
+			void ConstructPlanarConvexTriangleOneShotManifold(
+				const TImplicitObjectInstanced<FImplicitConvex3>& Convex,
+				const FTriangle& Triangle,
+				const FReal CullDistance,
+				TCArray<FContactPoint, 4>& OutContactPoints);
+
+		template
+			void ConstructPlanarConvexTriangleOneShotManifold(
+				const TImplicitObjectScaled<FImplicitConvex3>& Convex,
+				const FTriangle& Triangle,
+				const FReal CullDistance,
+				TCArray<FContactPoint, 4>& OutContactPoints);
+
+		template
+			void ConstructPlanarConvexTriangleOneShotManifold(
+				const FImplicitBox3& Convex,
+				const FTriangle& Triangle,
+				const FReal CullDistance,
+				TCArray<FContactPoint, 4>& OutContactPoints);
 	}
 }
 
diff --git a/Engine/Source/Runtime/Experimental/Chaos/Private/Chaos/CollisionOneShotManifolds.h b/Engine/Source/Runtime/Experimental/Chaos/Private/Chaos/CollisionOneShotManifolds.h
--- a/Engine/Source/Runtime/Experimental/Chaos/Private/Chaos/CollisionOneShotManifolds.h
+++ b/Engine/Source/Runtime/Experimental/Chaos/Private/Chaos/CollisionOneShotManifolds.h
@@ -10,12 +10,12 @@
 
 namespace Chaos
 {
-	extern bool bChaos_Collision_Manifold_FixNormalsInWorldSpace;
 	namespace Collisions
 	{
 		uint32 BoxBoxClipVerticesAgainstPlane(const FVec3* InputVertexBuffer, FVec3* outputVertexBuffer, uint32 ClipPointCount, int32 ClippingAxis, FReal Distance);
 		uint32 ReduceManifoldContactPoints(FVec3* Points, uint32 PointCount);
-		void PruneEdgeContactPoints(TArray<FContactPoint>& ContactPoints, const FReal MaxPlaneDistance);
+		void PruneEdgeContactPointsOrdered(TArray<FContactPoint>& ContactPoints, const FReal MaxPlaneDistance);
+		void PruneEdgeContactPointsUnordered(TArray<FContactPoint>& ContactPoints, const FReal MaxPlaneDistance);
 		void ReduceManifoldContactPointsTriangeMesh(TArray<FContactPoint>& ContactPoints);
 
 		void ConstructBoxBoxOneShotManifold(
@@ -34,6 +34,13 @@ namespace Chaos
 			const FRigidTransform3& Convex2Transform, //world
 			const FReal Dt,
 			FPBDCollisionConstraint& Constraint);
+
+		template <typename ConvexType>
+		void ConstructPlanarConvexTriangleOneShotManifold(
+			const ConvexType& Convex, 
+			const FTriangle& Triangle, 
+			const FReal CullDistance, 
+			TCArray<FContactPoint, 4>& OutContactPoints);
 	}
 }
 
diff --git a/Engine/Source/Runtime/Experimental/Chaos/Private/Chaos/CollisionOneShotManifoldsMiscShapes.cpp b/Engine/Source/Runtime/Experimental/Chaos/Private/Chaos/CollisionOneShotManifoldsMiscShapes.cpp
--- a/Engine/Source/Runtime/Experimental/Chaos/Private/Chaos/CollisionOneShotManifoldsMiscShapes.cpp
+++ b/Engine/Source/Runtime/Experimental/Chaos/Private/Chaos/CollisionOneShotManifoldsMiscShapes.cpp
@@ -10,9 +10,11 @@
 #include "Chaos/Collision/PBDCollisionConstraint.h"
 #include "Chaos/Convex.h"
 #include "Chaos/Defines.h"
+#include "Chaos/Framework/UncheckedArray.h"
 #include "Chaos/GJK.h"
 #include "Chaos/ImplicitObjectScaled.h"
 #include "Chaos/Transform.h"
+#include "Chaos/Triangle.h"
 #include "Chaos/Utilities.h"
 #include "ChaosStats.h"
 
@@ -22,6 +24,7 @@
 
 namespace Chaos
 {
+	extern FRealSingle Chaos_Collision_EdgePrunePlaneDistance;
 
 	namespace Collisions
 	{
@@ -328,21 +331,96 @@ namespace Chaos
 			}
 		}
 
-		template <typename ConvexType, typename TriMeshType>
-		void ConstructConvexTriMeshOneShotManifold(const FImplicitObject& Convex, const FRigidTransform3& ConvexTransform, const TriMeshType& TriangleMesh, const FRigidTransform3& TriMeshTransform, const FReal Dt, FPBDCollisionConstraint& Constraint)
+		template<typename ConvexType>
+		void ConstructPlanarConvexTriMeshOneShotManifoldImp(const ConvexType& Convex, const FRigidTransform3& ConvexTransform, const FImplicitObject& TriMesh, const FRigidTransform3& TriMeshTransform, FPBDCollisionConstraint& Constraint)
 		{
-			// We only build one shot manifolds once
 			// All convexes are pre-scaled, or wrapped in TImplicitObjectScaled
 			ensure(ConvexTransform.GetScale3D() == FVec3(1.0f, 1.0f, 1.0f));
 			ensure(TriMeshTransform.GetScale3D() == FVec3(1.0f, 1.0f, 1.0f));
 
-			// @todo(chaos): support manifold maintenance
-			Constraint.ResetActiveManifoldContacts();
+			// Unwrap the tri mesh (remove Scaled or Instanced) and get the scale
+			FVec3 TriMeshScale;
+			FReal TriMeshMargin;	// Not used - will be zero
+			const FTriangleMeshImplicitObject* UnscaledTriMesh = UnwrapImplicit<FTriangleMeshImplicitObject>(TriMesh, TriMeshScale, TriMeshMargin);
+			check(UnscaledTriMesh != nullptr);
+
+			const FRigidTransform3 TriangleMeshToConvexTransform = TriMeshTransform.GetRelativeTransformNoScale(ConvexTransform);
 
+			// Calculate the query bounds in trimesh space
+			// NOTE: to handle negative scales, we need to include it in the AABB transform (cannot use FAABB3::Scale)
+			const FRigidTransform3 QueryTransform = FRigidTransform3(TriangleMeshToConvexTransform.GetTranslation(), TriangleMeshToConvexTransform.GetRotation(), TriMeshScale);
+			const FAABB3 TriMeshQueryBounds = Convex.BoundingBox().InverseTransformedAABB(QueryTransform);
+			const FReal CullDistance = Constraint.GetCullDistance();
+
+			// A set of contact points which contains points from all triangle-convex manifolds
 			TArray<FContactPoint> ContactPoints;
-			GJKImplicitManifold<ConvexType, TriMeshType>(Convex, ConvexTransform, TriangleMesh, TriMeshTransform, Constraint.GetCullDistance(), 0.0f, ContactPoints);
+
+			// Prime the triangle producer with overlappig indices
+			FTriangleMeshTriangleProducer TriangleProducer;
+			int32 TriangleIndex;
+			FTriangle Triangle;
+			TCArray<FContactPoint, 4> TriangleManifoldPoints;
+			TriangleProducer.Reset(*UnscaledTriMesh, TriMeshQueryBounds);
+
+			// Loop over all the triangles, build a manifold and add the points to the total manifold
+			while (TriangleProducer.NextTriangle(*UnscaledTriMesh, QueryTransform, Triangle, TriangleIndex))
+			{
+				TriangleManifoldPoints.Reset();
+				ConstructPlanarConvexTriangleOneShotManifold(Convex, Triangle, CullDistance, TriangleManifoldPoints);
+					
+				for (int32 TriangleContactIndex = 0; TriangleContactIndex < TriangleManifoldPoints.Num(); ++TriangleContactIndex)
+				{
+					FContactPoint& ContactPoint = ContactPoints[ContactPoints.AddUninitialized()];
+					ContactPoint.ShapeContactPoints[0] = TriangleManifoldPoints[TriangleContactIndex].ShapeContactPoints[0];
+					ContactPoint.ShapeContactPoints[1] = TriangleMeshToConvexTransform.InverseTransformPositionNoScale(TriangleManifoldPoints[TriangleContactIndex].ShapeContactPoints[1]);
+					ContactPoint.ShapeContactNormal = TriangleMeshToConvexTransform.InverseTransformVectorNoScale(TriangleManifoldPoints[TriangleContactIndex].ShapeContactNormal);
+					ContactPoint.Phi = TriangleManifoldPoints[TriangleContactIndex].Phi;
+					ContactPoint.FaceIndex = TriangleIndex;
+					ContactPoint.ContactType = TriangleManifoldPoints[TriangleContactIndex].ContactType;
+				}
+			}
+
+			// Remove edge contacts that are "hidden" by face contacts
+			// @todo(chaos): EdgePruneDistance should be some fraction of the convex margin...
+			if (ContactPoints.Num() > 0)
+			{
+				const FReal EdgePruneDistance = Chaos_Collision_EdgePrunePlaneDistance;
+				Collisions::PruneEdgeContactPointsUnordered(ContactPoints, EdgePruneDistance);
+			}
+
+			// Whittle the manifold down to 4 points
+			if (ContactPoints.Num() > 4)
+			{
+				std::sort(&ContactPoints[0], &ContactPoints[0] + ContactPoints.Num(), [](const FContactPoint& L, const FContactPoint& R) { return L.Phi < R.Phi; });
+
+				// Remove all points (except for the deepest one, and ones with phis similar to it)
+				// NOTE: relies on the sort above
+				const FReal CullMargin = 0.1f;
+				int32 NewContactPointCount = ContactPoints.Num() > 0 ? 1 : 0;
+				for (int32 Index = 1; Index < ContactPoints.Num(); Index++)
+				{
+					if (ContactPoints[Index].Phi < 0 || ContactPoints[Index].Phi - ContactPoints[0].Phi < CullMargin)
+					{
+						NewContactPointCount++;
+					}
+					else
+					{
+						break;
+					}
+				}
+				ContactPoints.SetNum(NewContactPointCount, false);
+
+				// Reduce to only 4 contact points from here
+				// NOTE: relies on the sort above
+				Collisions::ReduceManifoldContactPointsTriangeMesh(ContactPoints);
+			}
+
+			// Add the manifold points to the constraint
+			Constraint.ResetActiveManifoldContacts();
 			for (FContactPoint& ContactPoint : ContactPoints)
 			{
+				// NOTE: We don't reuse manifolds between frames for Convex-TriMesh so it's not too bad to
+				// skip manifold points that are beyond the cull distance
 				if (ContactPoint.Phi < Constraint.GetCullDistance())
 				{
 					Constraint.AddOneshotManifoldContact(ContactPoint);
@@ -350,17 +428,37 @@ namespace Chaos
 			}
 		}
 
+		void ConstructPlanarConvexTriMeshOneShotManifold(const FImplicitObject& Convex, const FRigidTransform3& ConvexTransform, const FImplicitObject& TriangleMesh, const FRigidTransform3& TriangleMeshTransform, FPBDCollisionConstraint& Constraint)
+		{
+			if (const FImplicitBox3* RawBox = Convex.template GetObject<FImplicitBox3>())
+			{
+				ConstructPlanarConvexTriMeshOneShotManifoldImp(*RawBox, ConvexTransform, TriangleMesh, TriangleMeshTransform, Constraint);
+			}
+			else if (const TImplicitObjectScaled<FImplicitConvex3>* ScaledConvex = Convex.template GetObject<TImplicitObjectScaled<FImplicitConvex3>>())
+			{
+				ConstructPlanarConvexTriMeshOneShotManifoldImp(*ScaledConvex, ConvexTransform, TriangleMesh, TriangleMeshTransform, Constraint);
+			}
+			else if (const TImplicitObjectInstanced<FImplicitConvex3>* InstancedConvex = Convex.template GetObject<TImplicitObjectInstanced<FImplicitConvex3>>())
+			{
+				ConstructPlanarConvexTriMeshOneShotManifoldImp(*InstancedConvex, ConvexTransform, TriangleMesh, TriangleMeshTransform, Constraint);
+			}
+			else if (const FImplicitConvex3* RawConvex = Convex.template GetObject<FImplicitConvex3>())
+			{
+				ConstructPlanarConvexTriMeshOneShotManifoldImp(*RawConvex, ConvexTransform, TriangleMesh, TriangleMeshTransform, Constraint);
+			}
+			else
+			{
+				check(false);
+			}
+		}
+
+
 		template void ConstructSphereTriangleMeshOneShotManifold<TImplicitObjectScaled<class Chaos::FTriangleMeshImplicitObject, 1>>(const TSphere<FReal, 3>& Sphere, const FRigidTransform3& SphereWorldTransform, const TImplicitObjectScaled<class Chaos::FTriangleMeshImplicitObject, 1>& TriangleMesh, const FRigidTransform3& TriMeshWorldTransform, const FReal Dt, FPBDCollisionConstraint& Constraint);
 		template void ConstructSphereTriangleMeshOneShotManifold<FTriangleMeshImplicitObject>(const TSphere<FReal, 3>& Sphere, const FRigidTransform3& SphereWorldTransform, const FTriangleMeshImplicitObject& TriangleMesh, const FRigidTransform3& TriMeshWorldTransform, const FReal Dt, FPBDCollisionConstraint& Constraint);
 
 		template void ConstructCapsuleTriMeshOneShotManifold<TImplicitObjectScaled<class Chaos::FTriangleMeshImplicitObject, 1>>(const FCapsule& Capsule, const FRigidTransform3& CapsuleWorldTransform, const TImplicitObjectScaled<class Chaos::FTriangleMeshImplicitObject, 1>& TriangleMesh, const FRigidTransform3& TriMeshWorldTransform, const FReal Dt, FPBDCollisionConstraint& Constraint);
 		template void ConstructCapsuleTriMeshOneShotManifold<FTriangleMeshImplicitObject>(const FCapsule& Capsule, const FRigidTransform3& CapsuleWorldTransform, const FTriangleMeshImplicitObject& TriangleMesh, const FRigidTransform3& TriMeshWorldTransform, const FReal Dt, FPBDCollisionConstraint& Constraint);
 
-		template void ConstructConvexTriMeshOneShotManifold<FConvex, TImplicitObjectScaled<class Chaos::FTriangleMeshImplicitObject, 1>>(const FImplicitObject& Convex, const FRigidTransform3& ConvexTransform, const TImplicitObjectScaled<class Chaos::FTriangleMeshImplicitObject, 1>& TriangleMesh, const FRigidTransform3& TriMeshTransform, const FReal Dt, FPBDCollisionConstraint& Constraint);
-		template void ConstructConvexTriMeshOneShotManifold<FConvex, FTriangleMeshImplicitObject>(const FImplicitObject& Convex, const FRigidTransform3& ConvexTransform, const FTriangleMeshImplicitObject& TriangleMesh, const FRigidTransform3& TriMeshTransform, const FReal Dt, FPBDCollisionConstraint& Constraint);
-		template void ConstructConvexTriMeshOneShotManifold<FImplicitBox3, TImplicitObjectScaled<class Chaos::FTriangleMeshImplicitObject, 1>>(const FImplicitObject& Convex, const FRigidTransform3& ConvexTransform, const TImplicitObjectScaled<class Chaos::FTriangleMeshImplicitObject, 1>& TriangleMesh, const FRigidTransform3& TriMeshTransform, const FReal Dt, FPBDCollisionConstraint& Constraint);
-		template void ConstructConvexTriMeshOneShotManifold<FImplicitBox3, FTriangleMeshImplicitObject>(const FImplicitObject& Convex, const FRigidTransform3& ConvexTransform, const FTriangleMeshImplicitObject& TriangleMesh, const FRigidTransform3& TriMeshTransform, const FReal Dt, FPBDCollisionConstraint& Constraint);
-
 		template void ConstructConvexHeightFieldOneShotManifold<FConvex>(const FImplicitObject& Convex, const FRigidTransform3& ConvexTransform, const FHeightField& HeightField, const FRigidTransform3& HeightFieldTransform, const FReal Dt, FPBDCollisionConstraint& Constraint);
 		template void ConstructConvexHeightFieldOneShotManifold<FImplicitBox3>(const FImplicitObject& Convex, const FRigidTransform3& ConvexTransform, const FHeightField& HeightField, const FRigidTransform3& HeightFieldTransform, const FReal Dt, FPBDCollisionConstraint& Constraint);
 	}
diff --git a/Engine/Source/Runtime/Experimental/Chaos/Private/Chaos/CollisionOneShotManifoldsMiscShapes.h b/Engine/Source/Runtime/Experimental/Chaos/Private/Chaos/CollisionOneShotManifoldsMiscShapes.h
--- a/Engine/Source/Runtime/Experimental/Chaos/Private/Chaos/CollisionOneShotManifoldsMiscShapes.h
+++ b/Engine/Source/Runtime/Experimental/Chaos/Private/Chaos/CollisionOneShotManifoldsMiscShapes.h
@@ -9,7 +9,6 @@
 
 namespace Chaos
 {
-	extern bool bChaos_Collision_Manifold_FixNormalsInWorldSpace;
 	class FHeightField;
 	namespace Collisions
 	{
@@ -103,13 +102,17 @@ namespace Chaos
 			const FReal Dt,
 			FPBDCollisionConstraint& Constraint);
 
-		template <typename ConvexType = FConvex, typename TriMeshType>
-		void ConstructConvexTriMeshOneShotManifold(
+		/**
+		 * @brief Build the contact manifold between a Planar Convex shape and a TriMesh
+		 * @param Convex an ImplicitObject that must be a box, convex, or a wrapper around it. Asserts if not
+		 * @param TriMesh an ImplicitObject that must be FTriangleMeshImplicitObject or a wrapper around it. Asserts if not
+		*/
+		void ConstructPlanarConvexTriMeshOneShotManifold(
 			const FImplicitObject& Convex, 
 			const FRigidTransform3& ConvexTransform, 
-			const TriMeshType& TriangleMesh, 
-			const FRigidTransform3& TriMeshTransform, 
-			const FReal Dt, 
+			const FImplicitObject& TriangleMesh, 
+			const FRigidTransform3& TriangleMeshTransform, 
 			FPBDCollisionConstraint& Constraint);
+
 	}
 }
\ No newline at end of file
diff --git a/Engine/Source/Runtime/Experimental/Chaos/Private/Chaos/CollisionResolution.cpp b/Engine/Source/Runtime/Experimental/Chaos/Private/Chaos/CollisionResolution.cpp
--- a/Engine/Source/Runtime/Experimental/Chaos/Private/Chaos/CollisionResolution.cpp
+++ b/Engine/Source/Runtime/Experimental/Chaos/Private/Chaos/CollisionResolution.cpp
@@ -668,7 +668,7 @@ namespace Chaos
 		{
 			if(Constraint.GetUseManifold())
 			{
-				ConstructConvexTriMeshOneShotManifold<FImplicitBox3, TriMeshType>(Box0, WorldTransform0, TriangleMesh1, WorldTransform1, Dt, Constraint);
+				ConstructPlanarConvexTriMeshOneShotManifold(Box0, WorldTransform0, TriangleMesh1, WorldTransform1, Constraint);
 			}
 			else
 			{
@@ -1440,14 +1440,12 @@ namespace Chaos
 		// Convex-TriangleMesh
 		//
 
-		template <typename TriMeshType>
-		void UpdateConvexTriangleMeshConstraint(const FImplicitObject& Convex0, const FRigidTransform3& WorldTransform0, const TriMeshType& TriangleMesh1, const FRigidTransform3& WorldTransform1, const FReal Dt, FPBDCollisionConstraint& Constraint)
+		void UpdateConvexTriangleMeshConstraint(const FImplicitObject& Convex0, const FRigidTransform3& WorldTransform0, const FImplicitObject& TriangleMesh1, const FRigidTransform3& WorldTransform1, const FReal Dt, FPBDCollisionConstraint& Constraint)
 		{
-			// @todo(chaos): restitutionpadding
 			CONDITIONAL_SCOPE_CYCLE_COUNTER(STAT_Collisions_UpdateConvexTriangleMeshConstraint, ConstraintsDetailedStats);
 			if (Constraint.GetUseManifold())
 			{
-				ConstructConvexTriMeshOneShotManifold(Convex0, WorldTransform0, TriangleMesh1, WorldTransform1, Dt, Constraint);
+				ConstructPlanarConvexTriMeshOneShotManifold(Convex0, WorldTransform0, TriangleMesh1, WorldTransform1, Constraint);
 			}
 			else
 			{
@@ -1471,29 +1469,12 @@ namespace Chaos
 			CONDITIONAL_SCOPE_CYCLE_COUNTER(STAT_Collisions_ConstructConvexTriangleMeshConstraints, ConstraintsDetailedStats);
 			if (ensure(Implicit0->IsConvex()))
 			{
-				if (const TImplicitObjectScaled<FTriangleMeshImplicitObject>* ScaledTriangleMesh = Implicit1->template GetObject<const TImplicitObjectScaled<FTriangleMeshImplicitObject>>())
-				{
-					FPBDCollisionConstraint* Constraint = Context.CollisionAllocator->FindOrCreateConstraint(Particle0, Implicit0, nullptr, LocalTransform0, Particle1, Implicit1, nullptr, LocalTransform1, CullDistance, EContactShapesType::ConvexTriMesh, Context.bAllowManifolds);
-					if (T_TRAITS::bImmediateUpdate && (Constraint != nullptr))
-					{
-						FRigidTransform3 WorldTransform0 = LocalTransform0 * ParticleWorldTransform0;
-						FRigidTransform3 WorldTransform1 = LocalTransform1 * ParticleWorldTransform1;
-						UpdateConvexTriangleMeshConstraint(*Implicit0, WorldTransform0, *ScaledTriangleMesh, WorldTransform1, Dt, *Constraint);
-					}
-				}
-				else if (const FTriangleMeshImplicitObject* TriangleMesh = Implicit1->template GetObject<const FTriangleMeshImplicitObject>())
-				{
-					FPBDCollisionConstraint* Constraint = Context.CollisionAllocator->FindOrCreateConstraint(Particle0, Implicit0, nullptr, LocalTransform0, Particle1, Implicit1, nullptr, LocalTransform1, CullDistance, EContactShapesType::ConvexTriMesh, Context.bAllowManifolds);
-					if (T_TRAITS::bImmediateUpdate && (Constraint != nullptr))
-					{
-						FRigidTransform3 WorldTransform0 = LocalTransform0 * ParticleWorldTransform0;
-						FRigidTransform3 WorldTransform1 = LocalTransform1 * ParticleWorldTransform1;
-						UpdateConvexTriangleMeshConstraint(*Implicit0, WorldTransform0, *TriangleMesh, WorldTransform1, Dt, *Constraint);
-					}
-				}
-				else
+				FPBDCollisionConstraint* Constraint = Context.CollisionAllocator->FindOrCreateConstraint(Particle0, Implicit0, nullptr, LocalTransform0, Particle1, Implicit1, nullptr, LocalTransform1, CullDistance, EContactShapesType::ConvexTriMesh, Context.bAllowManifolds);
+				if (T_TRAITS::bImmediateUpdate && (Constraint != nullptr))
 				{
-					ensure(false);
+					FRigidTransform3 WorldTransform0 = LocalTransform0 * ParticleWorldTransform0;
+					FRigidTransform3 WorldTransform1 = LocalTransform1 * ParticleWorldTransform1;
+					UpdateConvexTriangleMeshConstraint(*Implicit0, WorldTransform0, *Implicit1, WorldTransform1, Dt, *Constraint);
 				}
 			}
 		}
@@ -1733,17 +1714,7 @@ namespace Chaos
 				UpdateGenericConvexConvexConstraint(Implicit0, WorldTransform0, Implicit1, WorldTransform1, Dt, Constraint);
 				break;
 			case EContactShapesType::ConvexTriMesh:
-				if (const TImplicitObjectScaled<FTriangleMeshImplicitObject>* ScaledTriMesh = Implicit1.template GetObject<const TImplicitObjectScaled<FTriangleMeshImplicitObject> >())
-				{
-					UpdateConvexTriangleMeshConstraint(Implicit0, WorldTransform0, *ScaledTriMesh, WorldTransform1, Dt, Constraint);
-					break;
-				}
-				else if (const FTriangleMeshImplicitObject* TriangleMeshImplicit = Implicit1.template GetObject<const FTriangleMeshImplicitObject>())
-				{
-					UpdateConvexTriangleMeshConstraint(Implicit0, WorldTransform0, *TriangleMeshImplicit, WorldTransform1, Dt, Constraint);
-					break;
-				}
-				ensure(false);
+				UpdateConvexTriangleMeshConstraint(Implicit0, WorldTransform0, Implicit1, WorldTransform1, Dt, Constraint);
 				break;
 			case EContactShapesType::ConvexHeightField:
 				UpdateConvexHeightFieldConstraint(Implicit0, WorldTransform0, *Implicit1.template GetObject< FHeightField >(), WorldTransform1, Dt, Constraint);
diff --git a/Engine/Source/Runtime/Experimental/Chaos/Private/Chaos/HeightField.cpp b/Engine/Source/Runtime/Experimental/Chaos/Private/Chaos/HeightField.cpp
--- a/Engine/Source/Runtime/Experimental/Chaos/Private/Chaos/HeightField.cpp
+++ b/Engine/Source/Runtime/Experimental/Chaos/Private/Chaos/HeightField.cpp
@@ -1494,7 +1494,7 @@ namespace Chaos
 		// Remove edge contacts that are "hidden" by face contacts
 		// EdgePruneDistance should be some fraction of the convex margin...
 		const FReal EdgePruneDistance = Chaos_Collision_EdgePrunePlaneDistance;
-		Collisions::PruneEdgeContactPoints(ContactPoints, EdgePruneDistance);
+		Collisions::PruneEdgeContactPointsOrdered(ContactPoints, EdgePruneDistance);
 
 		// Remove all points (except for the deepest one, and ones with phis similar to it)
 		const FReal CullMargin = 0.1f;
@@ -1719,7 +1719,7 @@ namespace Chaos
 	bool FHeightField::ContactManifold(const TImplicitObjectScaled<TBox<FReal, 3>>& QueryGeom, const FRigidTransform3& QueryTM, const FReal Thickness, TArray<FContactPoint>& ContactPoints) const
 	{
 		return ContactManifoldImp(QueryGeom, QueryTM, Thickness, ContactPoints);
-	}	
+	}
 
 	bool FHeightField::ContactManifold(const TImplicitObjectScaled<FCapsule>& QueryGeom, const FRigidTransform3& QueryTM, const FReal Thickness, TArray<FContactPoint>& ContactPoints) const
 	{
diff --git a/Engine/Source/Runtime/Experimental/Chaos/Private/Chaos/TriangleMeshImplicitObject.cpp b/Engine/Source/Runtime/Experimental/Chaos/Private/Chaos/TriangleMeshImplicitObject.cpp
--- a/Engine/Source/Runtime/Experimental/Chaos/Private/Chaos/TriangleMeshImplicitObject.cpp
+++ b/Engine/Source/Runtime/Experimental/Chaos/Private/Chaos/TriangleMeshImplicitObject.cpp
@@ -44,14 +44,6 @@ void ScaleTransformHelper(const FVec3& TriMeshScale, const FRigidTransform3& Que
 	OutScaledQueryTM = TRigidTransform<FReal, 3>(QueryTM.GetLocation() * TriMeshScale, QueryTM.GetRotation());
 }
 
-template <typename IdxType>
-void TransformVertsHelper(const FVec3& TriMeshScale, int32 TriIdx, const FTriangleMeshImplicitObject::ParticlesType& Particles,
-	const TArray<TVector<IdxType, 3>>& Elements, FVec3& OutA, FVec3& OutB, FVec3& OutC)
-{
-	OutA = Particles.X(Elements[TriIdx][0]) * TriMeshScale;
-	OutB = Particles.X(Elements[TriIdx][1]) * TriMeshScale;
-	OutC = Particles.X(Elements[TriIdx][2]) * TriMeshScale;
-}
 
 template <typename QueryGeomType>
 const QueryGeomType& ScaleGeomIntoWorldHelper(const QueryGeomType& QueryGeom, const FVec3& TriMeshScale)
@@ -428,7 +420,7 @@ bool FTriangleMeshImplicitObject::ContactManifoldImp(const GeomType& QueryGeom,
 		for (int32 TriIdx : PotentialIntersections)
 		{
 			FVec3 A, B, C;
-			TransformVertsHelper(TriMeshScale, TriIdx, MParticles, Elements, A, B, C);
+			TriangleMeshTransformVertsHelper(TriMeshScale, TriIdx, MParticles, Elements, A, B, C);
 			OverlapTriangle(A, B, C, Constraint);
 			for(FManifoldPoint& ManifoldPoint : Constraint.GetManifoldPoints())
 			{
@@ -441,7 +433,7 @@ bool FTriangleMeshImplicitObject::ContactManifoldImp(const GeomType& QueryGeom,
 		// Remove edge contacts that are "hidden" by face contacts
 		// EdgePruneDistance should be some fraction of the convex margin...
 		const FReal EdgePruneDistance = Chaos_Collision_EdgePrunePlaneDistance;
-		Collisions::PruneEdgeContactPoints(ContactPoints, EdgePruneDistance);
+		Collisions::PruneEdgeContactPointsOrdered(ContactPoints, EdgePruneDistance);
 
 		// Remove all points (except for the deepest one, and ones with phis similar to it)
 		const FReal CullMargin = 0.1f;
@@ -524,7 +516,7 @@ bool FTriangleMeshImplicitObject::GJKContactPointImp(const QueryGeomType& QueryG
 		for (int32 TriIdx : PotentialIntersections)
 		{
 			FVec3 A, B, C;
-			TransformVertsHelper(TriMeshScale, TriIdx, MParticles, Elements, A, B, C);
+			TriangleMeshTransformVertsHelper(TriMeshScale, TriIdx, MParticles, Elements, A, B, C);
 
 			if (CalculateTriangleContact(A, B, C, LocalContactLocation, LocalContactNormal, LocalContactPhi))
 			{
@@ -704,7 +696,7 @@ void FTriangleMeshImplicitObject::VisitTriangles(const FAABB3& QueryBounds, cons
 		for (int32 TriIdx : PotentialIntersections)
 		{
 			TVec3<FReal> A, B, C;
-			TransformVertsHelper(FVec3(1), TriIdx, MParticles, Elements, A, B, C);
+			TriangleMeshTransformVertsHelper(FVec3(1), TriIdx, MParticles, Elements, A, B, C);
 
 			Visitor(FTriangle(A, B, C));
 		}
@@ -720,12 +712,12 @@ void FTriangleMeshImplicitObject::VisitTriangles(const FAABB3& QueryBounds, cons
 	}
 }
 
-void FTriangleMeshImplicitObject::VisitTriangle(int32 TriangleIndex, const TFunction<void(const FTriangle& Triangle)>& Visitor) const
+void FTriangleMeshImplicitObject::VisitTriangle(const int32 TriangleIndex, const TFunction<void(const FTriangle& Triangle)>& Visitor) const
 {
 	const auto TriangleProducer = [&](int32 TriIdx, const auto& Elements)
 	{
 		TVec3<FReal> A, B, C;
-		TransformVertsHelper(FVec3(1), TriIdx, MParticles, Elements, A, B, C);
+		TriangleMeshTransformVertsHelper(FVec3(1), TriIdx, MParticles, Elements, A, B, C);
 
 		Visitor(FTriangle(A, B, C));
 	};
@@ -777,7 +769,7 @@ bool FTriangleMeshImplicitObject::OverlapGeomImp(const QueryGeomType& QueryGeom,
 			for (int32 TriIdx : PotentialIntersections)
 			{
 				FVec3 A, B, C;
-				TransformVertsHelper(TriMeshScale, TriIdx, MParticles, Elements, A, B, C);
+				TriangleMeshTransformVertsHelper(TriMeshScale, TriIdx, MParticles, Elements, A, B, C);
 
 				FVec3 TriangleNormal(0.0);
 				FReal Penetration = 0.0;
@@ -804,7 +796,7 @@ bool FTriangleMeshImplicitObject::OverlapGeomImp(const QueryGeomType& QueryGeom,
 			for (int32 TriIdx : PotentialIntersections)
 			{
 				FVec3 A, B, C;
-				TransformVertsHelper(TriMeshScale, TriIdx, MParticles, Elements, A, B, C);
+				TriangleMeshTransformVertsHelper(TriMeshScale, TriIdx, MParticles, Elements, A, B, C);
 
 				const FVec3 AB = B - A;
 				const FVec3 AC = C - A;
@@ -924,7 +916,7 @@ struct FTriangleMeshSweepVisitor
 		FVec3 HitNormal;
 
 		FVec3 A, B, C;
-		TransformVertsHelper(TriMeshScale,TriIdx,TriMesh.MParticles,Elements,A,B,C);
+		TriangleMeshTransformVertsHelper(TriMeshScale,TriIdx,TriMesh.MParticles,Elements,A,B,C);
 		FTriangle Tri(A, B, C);
 
 		if(CullsBackFaceSweepsCode != 0)
diff --git a/Engine/Source/Runtime/Experimental/Chaos/Public/Chaos/CastingUtilities.h b/Engine/Source/Runtime/Experimental/Chaos/Public/Chaos/CastingUtilities.h
--- a/Engine/Source/Runtime/Experimental/Chaos/Public/Chaos/CastingUtilities.h
+++ b/Engine/Source/Runtime/Experimental/Chaos/Public/Chaos/CastingUtilities.h
@@ -235,6 +235,31 @@ namespace Chaos
 			return nullptr;
 		}
 
+		/**
+		 * @brief Given an ImplicitObject known to be of type T or a wrapper around T, call the Lambda with the concrete type
+		 * which will be const T*, const TImplicitObjectInstanced<T>*, or const TImplicitObjectScaled<T>*
+		*/
+		template <typename T, typename Lambda>
+		FORCEINLINE_DEBUGGABLE auto CastWrapped(const FImplicitObject& A, const Lambda& Func)
+		{
+			if (const TImplicitObjectScaled<T>* ScaledImplicit = A.template GetObject<TImplicitObjectScaled<T>>())
+			{
+				return Func(ScaledImplicit);
+			}
+			else if (const TImplicitObjectInstanced<T>* InstancedImplicit = A.template GetObject<TImplicitObjectInstanced<T>>())
+			{
+				return Func(InstancedImplicit);
+			}
+			else if(const T* RawImplicit = A.template GetObject<T>())
+			{
+				return Func(RawImplicit);
+			}
+			else
+			{
+				return Func((T*)nullptr);
+			}
+		}
+
 	} // namespace Utilities
 
 } // namespace Chaos
diff --git a/Engine/Source/Runtime/Experimental/Chaos/Public/Chaos/Collision/ContactPointsMiscShapes.h b/Engine/Source/Runtime/Experimental/Chaos/Public/Chaos/Collision/ContactPointsMiscShapes.h
--- a/Engine/Source/Runtime/Experimental/Chaos/Public/Chaos/Collision/ContactPointsMiscShapes.h
+++ b/Engine/Source/Runtime/Experimental/Chaos/Public/Chaos/Collision/ContactPointsMiscShapes.h
@@ -50,8 +50,7 @@ namespace Chaos
 	FContactPoint CapsuleTriangleMeshSweptContactPoint(const FCapsule& A, const FRigidTransform3& ATransform, const TriMeshType& B, const FRigidTransform3& BStartTransform, const FVec3& Dir, const FReal Length, FReal& TOI);
 
 	FContactPoint ConvexHeightFieldContactPoint(const FImplicitObject& A, const FRigidTransform3& ATransform, const FHeightField& B, const FRigidTransform3& BTransform, const FReal CullDistance, const FReal ShapePadding);
-	template <typename TriMeshType>
-	FContactPoint ConvexTriangleMeshContactPoint(const FImplicitObject& A, const FRigidTransform3& ATransform, const TriMeshType& B, const FRigidTransform3& BTransform, const FReal CullDistance, const FReal ShapePadding);
+	FContactPoint ConvexTriangleMeshContactPoint(const FImplicitObject& A, const FRigidTransform3& ATransform, const FImplicitObject& B, const FRigidTransform3& BTransform, const FReal CullDistance, const FReal ShapePadding);
 	template <typename TriMeshType>
 	FContactPoint ConvexTriangleMeshSweptContactPoint(const FImplicitObject& A, const FRigidTransform3& ATransform, const TriMeshType& B, const FRigidTransform3& BStartTransform, const FVec3& Dir, const FReal Length, FReal& TOI);
 
diff --git a/Engine/Source/Runtime/Experimental/Chaos/Public/Chaos/Collision/PBDCollisionConstraint.h b/Engine/Source/Runtime/Experimental/Chaos/Public/Chaos/Collision/PBDCollisionConstraint.h
--- a/Engine/Source/Runtime/Experimental/Chaos/Public/Chaos/Collision/PBDCollisionConstraint.h
+++ b/Engine/Source/Runtime/Experimental/Chaos/Public/Chaos/Collision/PBDCollisionConstraint.h
@@ -553,10 +553,6 @@ namespace Chaos
 		const FPBDCollisionConstraintContainerCookie& GetContainerCookie() const { return ContainerCookie; }
 		FPBDCollisionConstraintContainerCookie& GetContainerCookie() { return ContainerCookie; }
 
-		// Whether we can use manifolds for the given partices. Manifolds do not work well with Joints and PBD
-		// because the bodies may be moved (and especially rotated) a lot in the solver and this can make the manifold extremely inaccurate
-		bool CanUseManifold(FGeometryParticleHandle* Particle0, FGeometryParticleHandle* Particle1) const;
-
 		bool AreMatchingContactPoints(const FContactPoint& A, const FContactPoint& B, FReal& OutScore) const;
 		int32 FindManifoldPoint(const FContactPoint& ContactPoint) const;
 
diff --git a/Engine/Source/Runtime/Experimental/Chaos/Public/Chaos/GJK.h b/Engine/Source/Runtime/Experimental/Chaos/Public/Chaos/GJK.h
--- a/Engine/Source/Runtime/Experimental/Chaos/Public/Chaos/GJK.h
+++ b/Engine/Source/Runtime/Experimental/Chaos/Public/Chaos/GJK.h
@@ -446,6 +446,178 @@ namespace Chaos
 		}
 	}
 
+	/**
+	 * @brief Calculate the penetration data for two shapes using GJK, assuming both shapes are already in the same space.
+	 * This is intended for use with triangles which have been transformed into the space of the convex shape.
+	 * 
+	 * @tparam T
+	 * @tparam TGeometryA First geometry type
+	 * @tparam TGeometryB Second geometry type. Usually FTriangle
+	 * @param A First geometry
+	 * @param B Second geometry (usually triangle)
+	 * @param OutPenetration penetration depth (negative for separation)
+	 * @param OutClosestA Closest point on A
+	 * @param OutClosestB Closest point on B
+	 * @param OutNormal Contact normal (points from A to B)
+	 * @param OutVertexA The closest vertex on A
+	 * @param OutVertexB The closest vertex on B
+	 * @param OutMaxSupportDelta When the convex has a margin, an upper bounds on the distance from the contact point to the vertex on the outer hull
+	 * @param Epsilon GJK tolerance
+	 * @param EPAEpsilon EPA tolerance
+	 * @return true always
+	*/
+	template <typename T, typename TGeometryA, typename TGeometryB>
+	bool GJKPenetrationSameSpace(const TGeometryA& A, const TGeometryB& B, T& OutPenetration, TVec3<T>& OutClosestA, TVec3<T>& OutClosestB, TVec3<T>& OutNormal, int32& OutVertexA, int32& OutVertexB, T& OutMaxSupportDelta, const T Epsilon = T(1.e-3), const T EPAEpsilon = T(1.e-2))
+	{
+		TGJKSimplexData<T> SimplexData;
+		T SupportDeltaA = 0;
+		T SupportDeltaB = 0;
+		T MaxSupportDelta = 0;
+
+		int32& VertexIndexA = OutVertexA;
+		int32& VertexIndexB = OutVertexB;
+
+		// Return the support vertex in A-local space for a vector V in A-local space
+		auto SupportAFunc = [&A, &SupportDeltaA, &VertexIndexA](const TVec3<T>& V)
+		{
+			return A.SupportCore(V, A.GetMargin(), &SupportDeltaA, VertexIndexA);
+		};
+
+		// Return the support point on B, in B-local space, for a vector V in A-local space
+		auto SupportBFunc = [&B, &SupportDeltaB, &VertexIndexB](const TVec3<T>& V)
+		{
+			return B.SupportCore(V, B.GetMargin(), &SupportDeltaB, VertexIndexB);
+		};
+
+		// V and Simplex are in A-local space
+		TVec3<T> V = TVec3<T>(-1, 0, 0);
+		TVec3<T> Simplex[4];
+		FSimplex SimplexIDs;
+		T Distance = FLT_MAX;
+
+		TVec3<T> Normal = -V;					// Remember the last good normal (i.e. don't update it if separation goes less than Epsilon and we can no longer normalize)
+		bool bIsDegenerate = false;				// True if GJK cannot make any more progress
+		bool bIsContact = false;				// True if shapes are within Epsilon or overlapping - GJK cannot provide a solution
+		int32 NumIterations = 0;
+		const T ThicknessA = A.GetMargin();
+		const T ThicknessB = B.GetMargin();
+		const T SeparatedDistance = ThicknessA + ThicknessB + Epsilon;
+		while (!bIsContact && !bIsDegenerate)
+		{
+			if (!ensure(NumIterations++ < 32))	//todo: take this out
+			{
+				break;	//if taking too long just stop. This should never happen
+			}
+
+			const TVec3<T> NegV = -V;
+			const TVec3<T> SupportA = SupportAFunc(NegV);
+			const TVec3<T> SupportB = SupportBFunc(V);
+			const TVec3<T> W = SupportA - SupportB;
+			MaxSupportDelta = FMath::Max(SupportDeltaA, SupportDeltaB);
+
+			// If we didn't move to at least ConvergedDistance or closer, assume we have reached a minimum
+			const T ConvergenceTolerance = 1.e-4f;
+			const T ConvergedDistance = (1.0f - ConvergenceTolerance) * Distance;
+			const T VW = TVec3<T>::DotProduct(V, W);
+
+			SimplexData.As[SimplexIDs.NumVerts] = SupportA;
+			SimplexData.Bs[SimplexIDs.NumVerts] = SupportB;
+			SimplexIDs[SimplexIDs.NumVerts] = SimplexIDs.NumVerts;
+			Simplex[SimplexIDs.NumVerts++] = W;
+
+			V = SimplexFindClosestToOrigin(Simplex, SimplexIDs, SimplexData.Barycentric, SimplexData.As, SimplexData.Bs);
+			T NewDistance = V.Size();
+
+			// Are we overlapping or too close for GJK to get a good result?
+			bIsContact = (NewDistance < Epsilon);
+
+			// If we did not get closer in this iteration, we are in a degenerate situation
+			bIsDegenerate = (NewDistance >= Distance);
+
+			// If we are not too close, update the separating vector and keep iterating
+			// If we are too close we drop through to EPA, but if EPA determines the shapes are actually separated 
+			// after all, we will wend up using the normal from the previous loop
+			if (!bIsContact)
+			{
+				V /= NewDistance;
+				Normal = -V;
+			}
+			Distance = NewDistance;
+		}
+
+		// Save the warm-start data (not much to store since we updated most of it in the loop above)
+		SimplexData.Save(SimplexIDs);
+
+		if (bIsContact)
+		{
+			// We did not converge or we detected an overlap situation, so run EPA to get contact data
+			// Rebuild the simplex into a variable sized array - EPA can generate any mumber of vertices
+			TArray<TVec3<T>> VertsA;
+			TArray<TVec3<T>> VertsB;
+			VertsA.Reserve(8);
+			VertsB.Reserve(8);
+			for (int i = 0; i < SimplexIDs.NumVerts; ++i)
+			{
+				VertsA.Add(SimplexData.As[i]);
+				VertsB.Add(SimplexData.Bs[i]);
+			}
+
+			T Penetration;
+			TVec3<T> MTD, ClosestA, ClosestB;
+			const EEPAResult EPAResult = EPA(VertsA, VertsB, SupportAFunc, SupportBFunc, Penetration, MTD, ClosestA, ClosestB, EPAEpsilon);
+
+			switch (EPAResult)
+			{
+			case EEPAResult::MaxIterations:
+				// Possibly a solution but with unknown error. Just return the last EPA state. 
+				// Fall through...
+			case EEPAResult::Ok:
+				// EPA has a solution - return it now
+				OutNormal = MTD;
+				OutPenetration = Penetration + ThicknessA + ThicknessB;
+				OutClosestA = ClosestA + MTD * ThicknessA;
+				OutClosestB = ClosestB - MTD * ThicknessB;
+				OutMaxSupportDelta = MaxSupportDelta;
+				return true;
+			case EEPAResult::BadInitialSimplex:
+				// The origin is outside the simplex. Must be a touching contact and EPA setup will have calculated the normal
+				// and penetration but we keep the position generated by GJK.
+				Normal = MTD;
+				Distance = -Penetration;
+				//UE_LOG(LogChaos, Warning, TEXT("EPA Touching Case"));
+				break;
+			case EEPAResult::Degenerate:
+				// We hit a degenerate simplex condition and could not reach a solution so use whatever near touching point GJK came up with.
+				// The result from EPA under these circumstances is not usable.
+				//UE_LOG(LogChaos, Warning, TEXT("EPA Degenerate Case"));
+				break;
+			}
+		}
+
+		// @todo(chaos): handle the case where EPA hits a degenerate triangle in the simplex.
+		// Currently we return a touching contact with the last position and normal from GJK. We could run SAT instead.
+
+		// GJK converged or we have a touching contact
+		{
+			TVec3<T> ClosestA(0);
+			TVec3<T> ClosestB(0);
+			for (int i = 0; i < SimplexIDs.NumVerts; ++i)
+			{
+				ClosestA += SimplexData.As[i] * SimplexData.Barycentric[i];
+				ClosestB += SimplexData.Bs[i] * SimplexData.Barycentric[i];
+			}
+
+			OutPenetration = ThicknessA + ThicknessB - Distance;
+			OutClosestA = ClosestA + Normal * ThicknessA;
+			OutClosestB = ClosestB - Normal * ThicknessB;
+			OutNormal = Normal;
+			OutMaxSupportDelta = MaxSupportDelta;
+
+			// @todo(chaos): we should pass back failure for the degenerate case so we can decide how to handle it externally.
+			return true;
+		}
+	}
+
 
 	template <typename TGeometryA, typename TGeometryB, bool bNegativePenetrationAllowed = false, typename T>
 	bool GJKPenetrationImpl(const TGeometryA& A, const TGeometryB& B, const TRigidTransform<T, 3>& BToATM, T& OutPenetration, TVec3<T>& OutClosestA, TVec3<T>& OutClosestB, TVec3<T>& OutNormal, int32& OutClosestVertexIndexA, int32& OutClosestVertexIndexB, const T InThicknessA = 0.0f, const T InThicknessB = 0.0f, const TVector<T, 3>& InitialDir = TVector<T, 3>(1, 0, 0), const T Epsilon = 1.e-3f)
diff --git a/Engine/Source/Runtime/Experimental/Chaos/Public/Chaos/ImplicitFwd.h b/Engine/Source/Runtime/Experimental/Chaos/Public/Chaos/ImplicitFwd.h
--- a/Engine/Source/Runtime/Experimental/Chaos/Public/Chaos/ImplicitFwd.h
+++ b/Engine/Source/Runtime/Experimental/Chaos/Public/Chaos/ImplicitFwd.h
@@ -7,6 +7,7 @@ namespace Chaos
 	class FImplicitObject;
 	class FConvex;
 	class FCapsule;
+	class FTriangle;
 
 	template<class T, int d> class TBox;
 	template<class T, int d> class TPlane;
diff --git a/Engine/Source/Runtime/Experimental/Chaos/Public/Chaos/ImplicitObjectScaled.h b/Engine/Source/Runtime/Experimental/Chaos/Public/Chaos/ImplicitObjectScaled.h
--- a/Engine/Source/Runtime/Experimental/Chaos/Public/Chaos/ImplicitObjectScaled.h
+++ b/Engine/Source/Runtime/Experimental/Chaos/Public/Chaos/ImplicitObjectScaled.h
@@ -221,6 +221,18 @@ public:
 		return MObject->OverlapGeom(B,BToATM,Thickness, OutMTD);
 	}
 
+	template <typename QueryGeomType>
+	bool GJKContactPoint(const QueryGeomType& A, const FRigidTransform3& AToBTM, const FReal Thickness, FVec3& Location, FVec3& Normal, FReal& Penetration) const
+	{
+		return MObject->GJKContactPoint(A, AToBTM, Thickness, Location, Normal, Penetration);
+	}
+
+	template <typename QueryGeomType>
+	bool ContactManifold(const QueryGeomType& A, const FRigidTransform3& AToBTM, const FReal Thickness, TArray<FContactPoint>& ContactPoints) const
+	{
+		return MObject->ContactManifold(A, AToBTM, Thickness, ContactPoints);
+	}
+
 	virtual uint16 GetMaterialIndex(uint32 HintIndex) const
 	{
 		return MObject->GetMaterialIndex(HintIndex);
@@ -993,4 +1005,37 @@ using TImplicitObjectScaledNonSerializable = TImplicitObjectScaled<TConcrete, fa
 template <typename T, int d>
 using TImplicitObjectScaledGeneric = TImplicitObjectScaled<FImplicitObject>;
 
+
+/**
+ * @brief Remove the Instanced or Scaled wrapper from an ImplicitObject of a known inner type and extract the instance properties
+ * @return the inner implicit object or null if the wrong type
+*/
+template<typename T>
+const T* UnwrapImplicit(const FImplicitObject& Implicit, FVec3& OutScale, FReal &OutMargin)
+{
+	OutScale = FVec3(1);
+	OutMargin = FReal(0);
+
+	if (const TImplicitObjectScaled<T>* ScaledImplicit = Implicit.template GetObject<TImplicitObjectScaled<T>>())
+	{
+		OutScale = ScaledImplicit->GetScale();
+		OutMargin = ScaledImplicit->GetMargin();
+		return ScaledImplicit->GetUnscaledObject();
+	}
+	else if (const TImplicitObjectInstanced<T>* InstancedImplicit = Implicit.template GetObject<TImplicitObjectInstanced<T>>())
+	{
+		OutMargin = InstancedImplicit->GetMargin();
+		return InstancedImplicit->GetInstancedObject();
+	}
+	else if (const T* RawImplicit = Implicit.template GetObject<T>())
+	{
+		OutMargin = RawImplicit->GetMargin();
+		return RawImplicit;
+	}
+	else
+	{
+		return nullptr;
+	}
+}
+
 }
diff --git a/Engine/Source/Runtime/Experimental/Chaos/Public/Chaos/Triangle.h b/Engine/Source/Runtime/Experimental/Chaos/Public/Chaos/Triangle.h
--- a/Engine/Source/Runtime/Experimental/Chaos/Public/Chaos/Triangle.h
+++ b/Engine/Source/Runtime/Experimental/Chaos/Public/Chaos/Triangle.h
@@ -13,6 +13,10 @@ namespace Chaos
 	class FTriangle
 	{
 	public:
+		FTriangle()
+		{
+		}
+
 		FTriangle(const FVec3& InA, const FVec3& InB, const FVec3& InC)
 			: A(InA)
 			, B(InB)
diff --git a/Engine/Source/Runtime/Experimental/Chaos/Public/Chaos/TriangleMeshImplicitObject.h b/Engine/Source/Runtime/Experimental/Chaos/Public/Chaos/TriangleMeshImplicitObject.h
--- a/Engine/Source/Runtime/Experimental/Chaos/Public/Chaos/TriangleMeshImplicitObject.h
+++ b/Engine/Source/Runtime/Experimental/Chaos/Public/Chaos/TriangleMeshImplicitObject.h
@@ -5,6 +5,7 @@
 #include "Chaos/GeometryParticles.h"
 #include "Chaos/ImplicitObjectScaled.h"
 #include "Chaos/SegmentMesh.h"
+#include "Chaos/Triangle.h"
 
 #include "AABBTree.h"
 #include "BoundingVolume.h"
@@ -101,6 +102,26 @@ namespace Chaos
 		return Ar;
 	}
 
+	template <typename IdxType, typename ParticlesType>
+	inline void TriangleMeshTransformVertsHelper(const FVec3& TriMeshScale, int32 TriIdx, const ParticlesType& Particles,
+		const TArray<TVector<IdxType, 3>>& Elements, FVec3& OutA, FVec3& OutB, FVec3& OutC)
+	{
+		OutA = Particles.X(Elements[TriIdx][0]) * TriMeshScale;
+		OutB = Particles.X(Elements[TriIdx][1]) * TriMeshScale;
+		OutC = Particles.X(Elements[TriIdx][2]) * TriMeshScale;
+	}
+
+	template <typename IdxType, typename ParticlesType>
+	inline void TriangleMeshTransformVertsHelper(const FRigidTransform3& Transform, int32 TriIdx, const ParticlesType& Particles,
+		const TArray<TVector<IdxType, 3>>& Elements, FVec3& OutA, FVec3& OutB, FVec3& OutC)
+	{
+		// Note: deliberately using scaled transform
+		OutA = Transform.TransformPosition(Particles.X(Elements[TriIdx][0]));
+		OutB = Transform.TransformPosition(Particles.X(Elements[TriIdx][1]));
+		OutC = Transform.TransformPosition(Particles.X(Elements[TriIdx][2]));
+	}
+
+
 	class CHAOS_API FTriangleMeshImplicitObject final : public FImplicitObject
 	{
 	public:
@@ -305,9 +326,34 @@ namespace Chaos
 
 		void VisitTriangles(const FAABB3& InQueryBounds, const TFunction<void(const FTriangle& Triangle)>& Visitor) const;
 
-		void VisitTriangle(int32 TriangleIndex, const TFunction<void(const FTriangle& Triangle)>& Visitor) const;
+		void VisitTriangle(const int32 TriangleIndex, const TFunction<void(const FTriangle& Triangle)>& Visitor) const;
+
+		/**
+		 * @brief Generate the triangle at the specified index with the specified transform (including scale)
+		*/
+		void GetTransformedTriangle(const int32 TriangleIndex, const FRigidTransform3& Transform, FTriangle& OutTriangle) const
+		{
+			if (MElements.RequiresLargeIndices())
+			{
+				TriangleMeshTransformVertsHelper(Transform, TriangleIndex, MParticles, MElements.GetLargeIndexBuffer(), OutTriangle[0], OutTriangle[1], OutTriangle[2]);
+			}
+			else
+			{
+				TriangleMeshTransformVertsHelper(Transform, TriangleIndex, MParticles, MElements.GetSmallIndexBuffer(), OutTriangle[0], OutTriangle[1], OutTriangle[2]);
+			}
+		}
+
+		/**
+		 * @brief Get a list of triangle indices that overlap the query bounds
+		 * @param QueryBounds query bounds in trimesh space
+		*/
+		void FindOverlappingTriangles(const FAABB3& QueryBounds, TArray<int32>& OutTriangleIndices) const
+		{
+			OutTriangleIndices = BVH.FindAllIntersections(QueryBounds);
+		}
 
 	private:
+
 		void RebuildBV();
 
 		ParticlesType MParticles;
@@ -422,4 +468,61 @@ namespace Chaos
 		template <typename IdxType>
 		TUniquePtr<FTriangleMeshImplicitObject> CopySlowImpl(const TArray < TVector<IdxType, 3>>& InElements) const;
 	};
+
+
+	/**
+	 * @brief A helper for iterating over FTriangleMeshImplicitObject triangles in a bounding box
+	*/
+	class FTriangleMeshTriangleProducer
+	{
+	public:
+		FTriangleMeshTriangleProducer()
+			: NextIndex(0)
+			, TriangleIndices()
+		{
+		}
+
+		/**
+		 * @brief Find the set of triangle indices that overlap the query bounds
+		 * @param InQueryBounds The query bounds in triangle mesh space
+		*/
+		inline void Reset(const FTriangleMeshImplicitObject& InTriMesh, const FAABB3& InQueryBounds)
+		{
+			InTriMesh.FindOverlappingTriangles(InQueryBounds, TriangleIndices);
+			NextIndex = 0;
+		}
+
+		/**
+		 * @brief Whether we are done producing triangles
+		 * @return
+		*/
+		inline const bool IsDone() const
+		{
+			return (NextIndex >= TriangleIndices.Num());
+		}
+
+		/**
+		 * @brief Get the next triangle with its index and transform the vertices
+		 * @return true if we produced a triangle, false if we are done with the triangle list
+		*/
+		inline bool NextTriangle(const FTriangleMeshImplicitObject& InTriMesh, const FRigidTransform3& Transform, FTriangle& OutTriangle, int32& OutTriangleIndex)
+		{
+			if (IsDone())
+			{
+				return false;
+			}
+
+			const int32 TriIndex = TriangleIndices[NextIndex++];
+
+			InTriMesh.GetTransformedTriangle(TriIndex, Transform, OutTriangle);
+			OutTriangleIndex = TriIndex;
+
+			return true;
+		}
+
+	private:
+		int32 NextIndex;
+		TArray<int32> TriangleIndices;
+	};
+
 }
