Make the attribute storage tests better, verify and fix anything that was not running correctly.
#rnx

diff --git a/Engine/Source/Runtime/Experimental/Interchange/Core/Private/Tests/StorageTest.cpp b/Engine/Source/Runtime/Experimental/Interchange/Core/Private/Tests/StorageTest.cpp
--- a/Engine/Source/Runtime/Experimental/Interchange/Core/Private/Tests/StorageTest.cpp
+++ b/Engine/Source/Runtime/Experimental/Interchange/Core/Private/Tests/StorageTest.cpp
@@ -13,7 +13,7 @@
 IMPLEMENT_SIMPLE_AUTOMATION_TEST(FAttributeStorageTest, "System.Runtime.Interchange.AttributeStorage", EAutomationTestFlags::EditorContext | EAutomationTestFlags::EngineFilter)
 
 
-#define ADD_ATTRIBUTE_STORAGE(ATTRIBUTE_TYPE, ATTRIBUTE_VALUE, COMPARE_CAST)																			\
+#define ADD_ATTRIBUTE_STORAGE(ATTRIBUTE_TYPE, ATTRIBUTE_VALUE)																							\
 {																																						\
 	const ATTRIBUTE_TYPE RefValue = ATTRIBUTE_VALUE;																									\
 	FAttributeKey Key = CreateUniqueKey();																												\
@@ -23,7 +23,7 @@ IMPLEMENT_SIMPLE_AUTOMATION_TEST(FAttributeStorageTest, "System.Runtime.Intercha
 	}																																					\
 	ATTRIBUTE_TYPE StoredValue;																															\
 	TestStorage.GetAttributeHandle<ATTRIBUTE_TYPE>(Key).Get(StoredValue);																				\
-	TestEqual(TEXT("`AttributeStorage` must handle add and retrieve" #ATTRIBUTE_TYPE "attribute"), COMPARE_CAST(StoredValue), COMPARE_CAST(RefValue));		\
+	TestEqual<ATTRIBUTE_TYPE>(TEXT("`AttributeStorage` must handle add and retrieve" #ATTRIBUTE_TYPE "attribute"), StoredValue, RefValue);				\
 }
 
 #define ADD_ATTRIBUTE_STORAGE_NOCOMPARE(ATTRIBUTE_TYPE, ATTRIBUTE_VALUE)																				\
@@ -59,24 +59,24 @@ bool FAttributeStorageTest::RunTest(const FString& Parameters)
 	int32 TestStoredSeed = 0;
 	const int32 RandHelperUInt8 = 255;
 
-	ADD_ATTRIBUTE_STORAGE(bool, false, );
-	ADD_ATTRIBUTE_STORAGE(bool, true, static_cast<int32>);
-	ADD_ATTRIBUTE_STORAGE(FColor, FColor(0xdf,0xcf, 0x00), );
-	ADD_ATTRIBUTE_STORAGE(FDateTime, FDateTime(2022,01,25,10, 24), );
-	ADD_ATTRIBUTE_STORAGE(double, 0.1234567890123456789, );
-	ADD_ATTRIBUTE_STORAGE(float, 0.1234567f, );
-	ADD_ATTRIBUTE_STORAGE(FGuid, FGuid(0xA4F0E4CD, 0xF97C4375, 0x9C25C211, 0xFD17B8BF), );
-	ADD_ATTRIBUTE_STORAGE(int8, -0x7f, );
-	ADD_ATTRIBUTE_STORAGE(int8, 0x7f, );
-	ADD_ATTRIBUTE_STORAGE(int16, -0x75cc, );
-	ADD_ATTRIBUTE_STORAGE(int16, 0x7fcc, );
-	ADD_ATTRIBUTE_STORAGE(int32, -0x75ccabcd, );
-	ADD_ATTRIBUTE_STORAGE(int32, 0x7fccabcd, );
-	ADD_ATTRIBUTE_STORAGE(int64, -0x75ccabcd12345678, );
-	ADD_ATTRIBUTE_STORAGE(int64, 0x7fccabcd12345678, );
-	ADD_ATTRIBUTE_STORAGE(FIntRect, FIntRect(-1,-2, 1, 2), );
-	ADD_ATTRIBUTE_STORAGE(FLinearColor, FLinearColor(0.59f, 0.49f, 0.0f), );
-	ADD_ATTRIBUTE_STORAGE(FName, FName(TEXT("Testing FName storage")), );
+	ADD_ATTRIBUTE_STORAGE(bool, false);
+	ADD_ATTRIBUTE_STORAGE(bool, true);
+	ADD_ATTRIBUTE_STORAGE(FColor, FColor(0xdf,0xcf, 0x00));
+	ADD_ATTRIBUTE_STORAGE(FDateTime, FDateTime(2022,01,25,10, 24));
+	ADD_ATTRIBUTE_STORAGE(double, 0.1234567890123456789);
+	ADD_ATTRIBUTE_STORAGE(float, 0.1234567f);
+	ADD_ATTRIBUTE_STORAGE(FGuid, FGuid(0xA4F0E4CD, 0xF97C4375, 0x9C25C211, 0xFD17B8BF));
+	ADD_ATTRIBUTE_STORAGE(int8, -0x7f);
+	ADD_ATTRIBUTE_STORAGE(int8, 0x7f);
+	ADD_ATTRIBUTE_STORAGE(int16, -0x75cc);
+	ADD_ATTRIBUTE_STORAGE(int16, 0x7fcc);
+	ADD_ATTRIBUTE_STORAGE(int32, -0x75ccabcd);
+	ADD_ATTRIBUTE_STORAGE(int32, 0x7fccabcd);
+	ADD_ATTRIBUTE_STORAGE(int64, -0x75ccabcd12345678);
+	ADD_ATTRIBUTE_STORAGE(int64, 0x7fccabcd12345678);
+	ADD_ATTRIBUTE_STORAGE(FIntRect, FIntRect(-1,-2, 1, 2));
+	ADD_ATTRIBUTE_STORAGE(FLinearColor, FLinearColor(0.59f, 0.49f, 0.0f));
+	ADD_ATTRIBUTE_STORAGE(FName, FName(TEXT("Testing FName storage")));
 
 	{
 		ADD_ATTRIBUTE_STORAGE_NOCOMPARE(FRandomStream, FRandomStream(34))
@@ -85,21 +85,21 @@ bool FAttributeStorageTest::RunTest(const FString& Parameters)
 		TestEqual(TEXT("`AttributeStorage` must handle add and retrieve FRandomStream attribute"), StoredValue.ToString(), RefValue.ToString());
 	}
 	
-	ADD_ATTRIBUTE_STORAGE(FString, FString(TEXT("Testing FString storage")), );
-	ADD_ATTRIBUTE_STORAGE(FString, FString(TEXT("")), );
-	ADD_ATTRIBUTE_STORAGE(FString, FString(TEXT("A")), );
-	ADD_ATTRIBUTE_STORAGE(FTimespan, FTimespan(10, 44, 38), );
-	ADD_ATTRIBUTE_STORAGE(FTwoVectors, FTwoVectors(FVector(150.0, -203.0, 4500.7), FVector(1.0, 1.0, 1.0)), );
-	ADD_ATTRIBUTE_STORAGE(uint8, 0x85, );
-	ADD_ATTRIBUTE_STORAGE(uint8, 0x7f, );
-	ADD_ATTRIBUTE_STORAGE(uint16, 0x85cc, );
-	ADD_ATTRIBUTE_STORAGE(uint16, 0x7fcc, );
-	ADD_ATTRIBUTE_STORAGE(uint32, 0x85ccabcd, );
-	ADD_ATTRIBUTE_STORAGE(uint32, 0x7fccabcd, );
-	ADD_ATTRIBUTE_STORAGE(uint64, 0x85ccabcd12345678, );
-	ADD_ATTRIBUTE_STORAGE(uint64, 0x7fccabcd12345678, );
-	ADD_ATTRIBUTE_STORAGE(FIntPoint, FIntPoint(1, 4), );
-	ADD_ATTRIBUTE_STORAGE(FIntVector, FIntVector(1, 4, 3), );
+	ADD_ATTRIBUTE_STORAGE(FString, FString(TEXT("Testing FString storage")));
+	ADD_ATTRIBUTE_STORAGE(FString, FString(TEXT("")));
+	ADD_ATTRIBUTE_STORAGE(FString, FString(TEXT("A")));
+	ADD_ATTRIBUTE_STORAGE(FTimespan, FTimespan(10, 44, 38));
+	ADD_ATTRIBUTE_STORAGE(FTwoVectors, FTwoVectors(FVector(150.0, -203.0, 4500.7), FVector(1.0, 1.0, 1.0)));
+	ADD_ATTRIBUTE_STORAGE(uint8, 0x85);
+	ADD_ATTRIBUTE_STORAGE(uint8, 0x7f);
+	ADD_ATTRIBUTE_STORAGE(uint16, 0x85cc);
+	ADD_ATTRIBUTE_STORAGE(uint16, 0x7fcc);
+	ADD_ATTRIBUTE_STORAGE(uint32, 0x85ccabcd);
+	ADD_ATTRIBUTE_STORAGE(uint32, 0x7fccabcd);
+	ADD_ATTRIBUTE_STORAGE(uint64, 0x85ccabcd12345678);
+	ADD_ATTRIBUTE_STORAGE(uint64, 0x7fccabcd12345678);
+	ADD_ATTRIBUTE_STORAGE(FIntPoint, FIntPoint(1, 4));
+	ADD_ATTRIBUTE_STORAGE(FIntVector, FIntVector(1, 4, 3));
 
 	{
 		ADD_ATTRIBUTE_STORAGE_NOCOMPARE(FVector2DHalf, FVector2DHalf(150.0, -203.0))
@@ -111,7 +111,7 @@ bool FAttributeStorageTest::RunTest(const FString& Parameters)
 		}
 	}
 
-	ADD_ATTRIBUTE_STORAGE(FFloat16, FFloat16(0.1234567f), );
+	ADD_ATTRIBUTE_STORAGE(FFloat16, FFloat16(0.1234567f));
 
 	{
 		ADD_ATTRIBUTE_STORAGE_NOCOMPARE(FOrientedBox, FOrientedBox())
@@ -129,18 +129,18 @@ bool FAttributeStorageTest::RunTest(const FString& Parameters)
 		}
 	}
 
-	ADD_ATTRIBUTE_STORAGE(FFrameNumber, FFrameNumber(30), );
-	ADD_ATTRIBUTE_STORAGE(FFrameRate, FFrameRate(1,60), );
-	ADD_ATTRIBUTE_STORAGE(FFrameTime, FFrameTime(29,0.5f), );
-	ADD_ATTRIBUTE_STORAGE(FSoftObjectPath, FSoftObjectPath(UClass::StaticClass()), );
-	ADD_ATTRIBUTE_STORAGE(FMatrix44f, FMatrix44f(FTransform3f(FRotator3f(25.0f, 2.0f, 3.14159f), FVector3f(150.0f, -203.0f, 4500.7f), FVector3f(1.0f, 1.0f, 1.0f)).ToMatrixWithScale()), );
-	ADD_ATTRIBUTE_STORAGE(FMatrix44d, FMatrix44d(FTransform3d(FRotator3d(25.0, 2.0, 3.14159), FVector3d(150.0, -203.0, 4500.7), FVector3d(1.0, 1.0, 1.0)).ToMatrixWithScale()), );
-	ADD_ATTRIBUTE_STORAGE(FPlane4f, FPlane4f(FVector3f(1.0f, 1.0f, 0.0f), FVector3f(0.0f, 0.0f, 1.0f)), );
-	ADD_ATTRIBUTE_STORAGE(FPlane4d, FPlane4d(FVector3d(1.0, 1.0, 0), FVector3d(0, 0, 1.0)), );
-	ADD_ATTRIBUTE_STORAGE(FQuat4f, FQuat4f(1.0f, 1.0f, 0.0f, 1.0f), );
-	ADD_ATTRIBUTE_STORAGE(FQuat4d, FQuat4d(1.0, 1.0, 0, 1.0), );
-	ADD_ATTRIBUTE_STORAGE(FRotator3f, FRotator3f(25.0f, 2.0f, 3.14159f), );
-	ADD_ATTRIBUTE_STORAGE(FRotator3d, FRotator3d(25.0, 2.0, 3.14159), );
+	ADD_ATTRIBUTE_STORAGE(FFrameNumber, FFrameNumber(30));
+	ADD_ATTRIBUTE_STORAGE(FFrameRate, FFrameRate(1,60));
+	ADD_ATTRIBUTE_STORAGE(FFrameTime, FFrameTime(29,0.5f));
+	ADD_ATTRIBUTE_STORAGE(FSoftObjectPath, FSoftObjectPath(UClass::StaticClass()));
+	ADD_ATTRIBUTE_STORAGE(FMatrix44f, FMatrix44f(FTransform3f(FRotator3f(25.0f, 2.0f, 3.14159f), FVector3f(150.0f, -203.0f, 4500.7f), FVector3f(1.0f, 1.0f, 1.0f)).ToMatrixWithScale()));
+	ADD_ATTRIBUTE_STORAGE(FMatrix44d, FMatrix44d(FTransform3d(FRotator3d(25.0, 2.0, 3.14159), FVector3d(150.0, -203.0, 4500.7), FVector3d(1.0, 1.0, 1.0)).ToMatrixWithScale()));
+	ADD_ATTRIBUTE_STORAGE(FPlane4f, FPlane4f(FVector3f(1.0f, 1.0f, 0.0f), FVector3f(0.0f, 0.0f, 1.0f)));
+	ADD_ATTRIBUTE_STORAGE(FPlane4d, FPlane4d(FVector3d(1.0, 1.0, 0), FVector3d(0, 0, 1.0)));
+	ADD_ATTRIBUTE_STORAGE(FQuat4f, FQuat4f(1.0f, 1.0f, 0.0f, 1.0f));
+	ADD_ATTRIBUTE_STORAGE(FQuat4d, FQuat4d(1.0, 1.0, 0, 1.0));
+	ADD_ATTRIBUTE_STORAGE(FRotator3f, FRotator3f(25.0f, 2.0f, 3.14159f));
+	ADD_ATTRIBUTE_STORAGE(FRotator3d, FRotator3d(25.0, 2.0, 3.14159));
 	{
 		ADD_ATTRIBUTE_STORAGE_NOCOMPARE(FTransform3f, FTransform3f(FRotator3f(25.0f, 2.0f, 3.14159f), FVector3f(150.0f, -203.0f, 4500.7f), FVector3f(1.0f, 1.0f, 1.0f)))
 		FTransform3f StoredValue;
@@ -159,17 +159,17 @@ bool FAttributeStorageTest::RunTest(const FString& Parameters)
 			AddError(FString(TEXT("`AttributeStorage` must handle adding FTransform3d attribute")));
 		}
 	}
-	ADD_ATTRIBUTE_STORAGE(FVector3f, FVector3f(150.0f, -203.0f, 4500.7f), );
-	ADD_ATTRIBUTE_STORAGE(FVector3d, FVector3d(150.0, -203.0, 4500.7), );
-	ADD_ATTRIBUTE_STORAGE(FVector2f, FVector2f(150.0f, -203.0f), );
-	ADD_ATTRIBUTE_STORAGE(FVector4f, FVector4f(150.0f, -203.0f, 4500.7f, 2.0f), );
-	ADD_ATTRIBUTE_STORAGE(FVector4d, FVector4d(150.0, -203.0, 4500.7, 2.0), );
-	ADD_ATTRIBUTE_STORAGE(FBox2f, FBox2f(FVector2f(-1.0f), FVector2f(1.0f)), );
-	ADD_ATTRIBUTE_STORAGE(FBox2D, FBox2D(FVector2D(-1.0), FVector2D(1.0)), );
-	ADD_ATTRIBUTE_STORAGE(FBox3f, FBox3f(FVector3f(-1.0f), FVector3f(1.0f)), );
-	ADD_ATTRIBUTE_STORAGE(FBox3d, FBox3d(FVector3d(-1.0), FVector3d(1.0)), );
-	ADD_ATTRIBUTE_STORAGE(FBoxSphereBounds3f, FBoxSphereBounds3f(FVector3f(0.0f), FVector3f(2.0f), 2.2f), );
-	ADD_ATTRIBUTE_STORAGE(FBoxSphereBounds3d, FBoxSphereBounds3d(FVector3d(0.0), FVector3d(2.0), 2.2), );
+	ADD_ATTRIBUTE_STORAGE(FVector3f, FVector3f(150.0f, -203.0f, 4500.7f));
+	ADD_ATTRIBUTE_STORAGE(FVector3d, FVector3d(150.0, -203.0, 4500.7));
+	ADD_ATTRIBUTE_STORAGE(FVector2f, FVector2f(150.0f, -203.0f));
+	ADD_ATTRIBUTE_STORAGE(FVector4f, FVector4f(150.0f, -203.0f, 4500.7f, 2.0f));
+	ADD_ATTRIBUTE_STORAGE(FVector4d, FVector4d(150.0, -203.0, 4500.7, 2.0));
+	ADD_ATTRIBUTE_STORAGE(FBox2f, FBox2f(FVector2f(-1.0f), FVector2f(1.0f)));
+	ADD_ATTRIBUTE_STORAGE(FBox2D, FBox2D(FVector2D(-1.0), FVector2D(1.0)));
+	ADD_ATTRIBUTE_STORAGE(FBox3f, FBox3f(FVector3f(-1.0f), FVector3f(1.0f)));
+	ADD_ATTRIBUTE_STORAGE(FBox3d, FBox3d(FVector3d(-1.0), FVector3d(1.0)));
+	ADD_ATTRIBUTE_STORAGE(FBoxSphereBounds3f, FBoxSphereBounds3f(FVector3f(0.0f), FVector3f(2.0f), 2.2f));
+	ADD_ATTRIBUTE_STORAGE(FBoxSphereBounds3d, FBoxSphereBounds3d(FVector3d(0.0), FVector3d(2.0), 2.2));
 	{
 		FSphere3f RefValue;
 		RefValue.Center = FVector3f(2.0f);
@@ -180,7 +180,6 @@ bool FAttributeStorageTest::RunTest(const FString& Parameters)
 			AddError(FString(TEXT("`AttributeStorage` must handle adding FSphere3f attribute")));
 		}
 		FSphere3f StoredValue;
-		uint32 bobafet = -45;
 		TestStorage.GetAttributeHandle<FSphere3f>(Key).Get(StoredValue);
 		if (!StoredValue.Equals(RefValue))
 		{
@@ -197,7 +196,6 @@ bool FAttributeStorageTest::RunTest(const FString& Parameters)
 			AddError(FString(TEXT("`AttributeStorage` must handle adding FSphere3d attribute")));
 		}
 		FSphere3d StoredValue;
-		uint32 bobafet = -45;
 		TestStorage.GetAttributeHandle<FSphere3d>(Key).Get(StoredValue);
 		if (!StoredValue.Equals(RefValue))
 		{
@@ -205,33 +203,13 @@ bool FAttributeStorageTest::RunTest(const FString& Parameters)
 		}
 	}
 
-	//Test Adding/Reading uint8 attribute with 0 value
+	//Test Adding/Reading uint8 attribute with different value
 	{
-		const uint8 RefValue = 0;
-		FAttributeKey Key = CreateUniqueKey();
-		if (!IsAttributeStorageResultSuccess(TestStorage.RegisterAttribute(Key, RefValue)))
-		{
-			AddError(FString(TEXT("`AttributeStorage` must handle adding uint8 attribute")));
-		}
-		uint8 StoredValue = 0;
-		TestStorage.GetAttributeHandle<uint8>(Key).Get(StoredValue);
-		TestEqual(TEXT("`AttributeStorage` must handle add and retrieve uint8 attribute"), static_cast<int32>(StoredValue), static_cast<int32>(RefValue));
+		ADD_ATTRIBUTE_STORAGE(uint8, 0);
+		ADD_ATTRIBUTE_STORAGE(uint8, 43);
 	}
 
-	//Test Adding/Reading uint8 attribute with a value different then the default value
-	{
-		const uint8 RefValue = 43;
-		FAttributeKey Key = CreateUniqueKey();
-		if (!IsAttributeStorageResultSuccess(TestStorage.RegisterAttribute(Key, RefValue)))
-		{
-			AddError(FString(TEXT("`AttributeStorage` must handle adding uint8 attribute")));
-		}
-		uint8 StoredValue = 0;
-		TestStorage.GetAttributeHandle<uint8>(Key).Get(StoredValue);
-		TestEqual(TEXT("`AttributeStorage` must handle add and retrieve uint8 attribute"), static_cast<int32>(StoredValue), static_cast<int32>(RefValue));
-	}
-
-	//Test Adding/Reading simple attribute with a value different then the default value
+	//Test Adding/Reading int32 attribute with a negative value
 	{
 		if (!IsAttributeStorageResultSuccess(TestStorage.RegisterAttribute(TestInt32KeyName, NegativeValueRef)))
 		{
@@ -246,57 +224,26 @@ bool FAttributeStorageTest::RunTest(const FString& Parameters)
 	for (int32 AddedIndex = 0; AddedIndex < 2; ++AddedIndex)
 	{
 		const FVector ValueRef = RandomStream.VRand();
-		FAttributeKey Key = CreateUniqueKey();
-		if (!IsAttributeStorageResultSuccess(TestStorage.RegisterAttribute(Key, ValueRef)))
-		{
-			AddError(FString(TEXT("`AttributeStorage` must handle adding FVector attribute")));
-			break;
-		}
-		FVector StoredValue(0);
-		TestStorage.GetAttributeHandle<FVector>(Key).Get(StoredValue);
-		if (!TestEqual(TEXT("`AttributeStorage` must handle add and retrieve FVector attribute"), StoredValue.X, ValueRef.X))
-		{
-			break;
-		}
+		ADD_ATTRIBUTE_STORAGE(FVector, ValueRef);
 	}
 
 	//Adding FName
 	{
-		const FName RefValue = TEXT("The magic carpet ride!");
-		FAttributeKey Key = CreateUniqueKey();
-		if (!IsAttributeStorageResultSuccess(TestStorage.RegisterAttribute(Key, RefValue)))
-		{
-			AddError(FString(TEXT("`AttributeStorage` must handle adding FName attribute")));
-		}
-		FName StoredValue;
-		TestStorage.GetAttributeHandle<FName>(Key).Get(StoredValue);
-		if(StoredValue != RefValue)
-		{
-			TestEqual(TEXT("`AttributeStorage` must handle add and retrieve FName attribute"), StoredValue.ToString(), RefValue.ToString());
-		}
+		const FName ValueRef = TEXT("The magic carpet ride!");
+		ADD_ATTRIBUTE_STORAGE(FName, ValueRef);
 	}
 
 	//Adding FSoftObjectPath
 	{
 
-		const FSoftObjectPath RefValue(UClass::StaticClass());
-		FAttributeKey Key = CreateUniqueKey();
-		if (!IsAttributeStorageResultSuccess(TestStorage.RegisterAttribute(Key, RefValue)))
-		{
-			AddError(FString(TEXT("`AttributeStorage` must handle adding FSoftObjectPath attribute")));
-		}
-		FSoftObjectPath StoredValue;
-		TestStorage.GetAttributeHandle<FSoftObjectPath>(Key).Get(StoredValue);
-		if (StoredValue != RefValue)
-		{
-			TestEqual(TEXT("`AttributeStorage` must handle add and retrieve FSoftObjectPath attribute"), StoredValue.ToString(), RefValue.ToString());
-		}
+		const FSoftObjectPath ValueRef(UClass::StaticClass());
+		ADD_ATTRIBUTE_STORAGE(FSoftObjectPath, ValueRef);
 	}
 
 	//Adding one big TArray64<uint8> outside of the Hash
 	{
 		TArray64<uint8> ValueRef;
-		uint64 ArrayNum = 50;
+		uint64 ArrayNum = 500;
 		ValueRef.Reserve(ArrayNum);
 		for (uint64 ArrayIndex = 0; ArrayIndex < ArrayNum; ++ArrayIndex)
 		{
@@ -340,18 +287,7 @@ bool FAttributeStorageTest::RunTest(const FString& Parameters)
 	//Store a FString
 	{
 		const FString ValueRef = TEXT("The quick brown fox jumped over the lazy dogs");
-		const FAttributeKey Key = CreateUniqueKey();
-		if (!IsAttributeStorageResultSuccess(TestStorage.RegisterAttribute(Key, ValueRef)))
-		{
-			AddError(FString(TEXT("`AttributeStorage` must handle adding FString attribute")));
-		}
-		FString StoredValue;
-		if (!IsAttributeStorageResultSuccess(TestStorage.GetAttributeHandle<FString >(Key).Get(StoredValue)))
-		{
-			AddError(FString(TEXT("`AttributeStorage` must handle retrieve FString attribute")));
-		}
-		//Equal random stream should give the same result when asking fraction
-		TestEqual(TEXT("`AttributeStorage` must handle add and retrieve FString attribute"), StoredValue, ValueRef);
+		ADD_ATTRIBUTE_STORAGE(FString, ValueRef);
 	}
 
 	//Store the random stream
@@ -637,8 +573,7 @@ bool FAttributeStorageTest::RunTest(const FString& Parameters)
 				AddError(FString(TEXT("`AttributeStorage` must handle adding int32 attribute")));
 				break;
 			}
-
-			const FVector VectorRef(AddedIndex/TestCount);
+			const FVector3f VectorRef((float)AddedIndex/(float)TestCount);
 			Key = CreateUniqueKey();
 			AttributeKeys[(AddedIndex*2)+1] = Key;
 			if (!IsAttributeStorageResultSuccess(TestStorage.RegisterAttribute(Key, VectorRef)))
@@ -657,10 +592,10 @@ bool FAttributeStorageTest::RunTest(const FString& Parameters)
 				int32 StoredValue(0);
 				TestStorage.GetAttributeHandle<int32>(Key).Get(StoredValue);
 			}
-			else if(TestStorage.GetAttributeType(Key) == EAttributeTypes::Vector)
+			else if(TestStorage.GetAttributeType(Key) == EAttributeTypes::Vector3f)
 			{
-				FVector StoredVector(0.0f);
-				TestStorage.GetAttributeHandle<FVector>(Key).Get(StoredVector);
+				FVector3f StoredVector(0.0f);
+				TestStorage.GetAttributeHandle<FVector3f>(Key).Get(StoredVector);
 			}
 		}
 		const double SmokeTestReadTime = FPlatformTime::Seconds();
