Move FPerShapeData RelativeTransform and LeafGeometryPointer to derived class, only allocate when needed to save memory.
Change RelativeTransform to floats.

diff --git a/Engine/Source/Runtime/Engine/Private/PhysicsEngine/Experimental/ChaosInterfaceUtils.cpp b/Engine/Source/Runtime/Engine/Private/PhysicsEngine/Experimental/ChaosInterfaceUtils.cpp
--- a/Engine/Source/Runtime/Engine/Private/PhysicsEngine/Experimental/ChaosInterfaceUtils.cpp
+++ b/Engine/Source/Runtime/Engine/Private/PhysicsEngine/Experimental/ChaosInterfaceUtils.cpp
@@ -201,8 +201,7 @@ namespace ChaosInterface
 
 		auto NewShapeHelper = [&InParams, &CollisionTraceType](Chaos::TSerializablePtr<Chaos::FImplicitObject> InGeom, int32 ShapeIdx, void* UserData, ECollisionEnabled::Type ShapeCollisionEnabled, bool bComplexShape = false)
 		{
-			TUniquePtr<Chaos::FPerShapeData> NewShape = Chaos::FPerShapeData::CreatePerShapeData(ShapeIdx);
-			NewShape->SetGeometry(InGeom);
+			TUniquePtr<Chaos::FPerShapeData> NewShape = Chaos::FPerShapeData::CreatePerShapeData(ShapeIdx, InGeom);
 			NewShape->SetQueryData(bComplexShape ? InParams.CollisionData.CollisionFilterData.QueryComplexFilter : InParams.CollisionData.CollisionFilterData.QuerySimpleFilter);
 			NewShape->SetSimData(InParams.CollisionData.CollisionFilterData.SimFilter);
 			NewShape->SetCollisionTraceType(ConvertCollisionTraceFlag(CollisionTraceType));
diff --git a/Engine/Source/Runtime/Engine/Private/PhysicsEngine/ImmediatePhysics/ImmediatePhysicsChaos/ImmediatePhysicsActorHandle_Chaos.cpp b/Engine/Source/Runtime/Engine/Private/PhysicsEngine/ImmediatePhysics/ImmediatePhysicsChaos/ImmediatePhysicsActorHandle_Chaos.cpp
--- a/Engine/Source/Runtime/Engine/Private/PhysicsEngine/ImmediatePhysics/ImmediatePhysicsChaos/ImmediatePhysicsActorHandle_Chaos.cpp
+++ b/Engine/Source/Runtime/Engine/Private/PhysicsEngine/ImmediatePhysics/ImmediatePhysicsChaos/ImmediatePhysicsActorHandle_Chaos.cpp
@@ -38,8 +38,7 @@ namespace ImmediatePhysics_Chaos
 		const FReal Radius = 1.0f * Scale.GetMax();
 
 		auto ImplicitSphere = MakeUnique<Chaos::TSphere<FReal, 3>>(FVec3(0), Radius);
-		auto NewShape = Chaos::FPerShapeData::CreatePerShapeData(OutShapes.Num());
-		NewShape->SetGeometry(MakeSerializable(ImplicitSphere));
+		auto NewShape = Chaos::FPerShapeData::CreatePerShapeData(OutShapes.Num(), MakeSerializable(ImplicitSphere));
 		NewShape->UpdateShapeBounds(FTransform::Identity);
 		NewShape->SetUserData(nullptr);
 		NewShape->SetQueryEnabled(false);
diff --git a/Engine/Source/Runtime/Experimental/Chaos/Private/Chaos/Collision/ParticlePairMidPhase.cpp b/Engine/Source/Runtime/Experimental/Chaos/Private/Chaos/Collision/ParticlePairMidPhase.cpp
--- a/Engine/Source/Runtime/Experimental/Chaos/Private/Chaos/Collision/ParticlePairMidPhase.cpp
+++ b/Engine/Source/Runtime/Experimental/Chaos/Private/Chaos/Collision/ParticlePairMidPhase.cpp
@@ -60,10 +60,10 @@ namespace Chaos
 	{
 		const FImplicitObject* Implicit0 = InShape0->GetLeafGeometry();
 		const FBVHParticles* BVHParticles0 = FConstGenericParticleHandle(Particle0)->CollisionParticles().Get();
-		const FRigidTransform3& ShapeRelativeTransform0 = InShape0->GetLeafRelativeTransform();
+		const FRigidTransform3& ShapeRelativeTransform0 = (FRigidTransform3)InShape0->GetLeafRelativeTransform();
 		const FImplicitObject* Implicit1 = InShape1->GetLeafGeometry();
 		const FBVHParticles* BVHParticles1 = FConstGenericParticleHandle(Particle1)->CollisionParticles().Get();
-		const FRigidTransform3& ShapeRelativeTransform1 = InShape1->GetLeafRelativeTransform();
+		const FRigidTransform3& ShapeRelativeTransform1 = (FRigidTransform3)InShape1->GetLeafRelativeTransform();
 
 		return FPBDCollisionConstraint::Make(Particle0, Implicit0, BVHParticles0, ShapeRelativeTransform0, Particle1, Implicit1, BVHParticles1, ShapeRelativeTransform1, CullDistance, bUseManifold, ShapePairType);
 	}
@@ -462,11 +462,11 @@ namespace Chaos
 
 		const FImplicitObject* Implicit0 = Shape0->GetLeafGeometry();
 		const FBVHParticles* BVHParticles0 = P0->CollisionParticles().Get();
-		const FRigidTransform3& ShapeRelativeTransform0 = Shape0->GetLeafRelativeTransform();
+		const FRigidTransform3& ShapeRelativeTransform0 = (FRigidTransform3)Shape0->GetLeafRelativeTransform();
 		const FRigidTransform3 ParticleWorldTransform0 = FParticleUtilities::GetActorWorldTransform(P0);
 		const FImplicitObject* Implicit1 = Shape1->GetLeafGeometry();
 		const FBVHParticles* BVHParticles1 = P1->CollisionParticles().Get();
-		const FRigidTransform3& ShapeRelativeTransform1 = Shape1->GetLeafRelativeTransform();
+		const FRigidTransform3& ShapeRelativeTransform1 = (FRigidTransform3)Shape1->GetLeafRelativeTransform();
 		const FRigidTransform3 ParticleWorldTransform1 = FParticleUtilities::GetActorWorldTransform(P1);
 
 		FCollisionContext LocalContext = Context;
diff --git a/Engine/Source/Runtime/Experimental/Chaos/Private/Chaos/GeometryParticles.cpp b/Engine/Source/Runtime/Experimental/Chaos/Private/Chaos/GeometryParticles.cpp
--- a/Engine/Source/Runtime/Experimental/Chaos/Private/Chaos/GeometryParticles.cpp
+++ b/Engine/Source/Runtime/Experimental/Chaos/Private/Chaos/GeometryParticles.cpp
@@ -21,13 +21,18 @@ namespace Chaos
 
 				for (int32 ShapeIndex = 0; ShapeIndex < ShapesArray.Num(); ++ShapeIndex)
 				{
+					TSerializablePtr<FImplicitObject> ShapeGeometry = MakeSerializable(Union->GetObjects()[ShapeIndex]);
+
 					if (ShapeIndex >= OldShapeNum)
 					{
 						// If newly allocated shape, initialize it.
-						ShapesArray[ShapeIndex] = FPerShapeData::CreatePerShapeData(ShapeIndex);
+						ShapesArray[ShapeIndex] = FPerShapeData::CreatePerShapeData(ShapeIndex, ShapeGeometry);
+					}
+					else if (ShapeGeometry != ShapesArray[ShapeIndex]->GetGeometry())
+					{
+						// Update geometry pointer if it changed
+						FPerShapeData::UpdateGeometry(ShapesArray[ShapeIndex], ShapeGeometry);
 					}
-
-					ShapesArray[ShapeIndex]->SetGeometry(MakeSerializable(Union->GetObjects()[ShapeIndex]));
 				}
 			}
 			else
@@ -35,9 +40,12 @@ namespace Chaos
 				ShapesArray.SetNum(1);
 				if (OldShapeNum == 0)
 				{
-					ShapesArray[0] = FPerShapeData::CreatePerShapeData(0);
+					ShapesArray[0] = FPerShapeData::CreatePerShapeData(0, Geometry);
+				}
+				else
+				{
+					FPerShapeData::UpdateGeometry(ShapesArray[0], Geometry);
 				}
-				ShapesArray[0]->SetGeometry(Geometry);
 			}
 
 			if (Geometry->HasBoundingBox())
@@ -67,7 +75,7 @@ namespace Chaos
 	// Scaled so the Scale and Margin is stored on the wrapper (the convex itself is shared).
 	// - support for Margin as per-shape data passed through the collision functions
 	// - support for Scale as per-shape data passed through the collision functions (or ideally in the Transforms)
-	const FImplicitObject* GetInnerGeometryInstanceData(const FImplicitObject* Implicit, FRigidTransform3& OutTransform, FReal& OutMargin)
+	const FImplicitObject* GetInnerGeometryInstanceData(const FImplicitObject* Implicit, FRigidTransformRealSingle3& OutTransform, FReal& OutMargin)
 	{
 		if (Implicit != nullptr)
 		{
@@ -76,7 +84,7 @@ namespace Chaos
 			{
 				// Transformed Implicit
 				const TImplicitObjectTransformed<FReal, 3>* TransformedImplicit = Implicit->template GetObject<const TImplicitObjectTransformed<FReal, 3>>();
-				OutTransform = TransformedImplicit->GetTransform() * OutTransform;
+				OutTransform = (FRigidTransformRealSingle3)TransformedImplicit->GetTransform() * OutTransform;
 				//OutMargin += TransformedImplicit->GetMargin();
 				return GetInnerGeometryInstanceData(TransformedImplicit->GetTransformedObject(), OutTransform, OutMargin);
 			}
@@ -107,20 +115,81 @@ namespace Chaos
 
 
 	FPerShapeData::FPerShapeData(int32 InShapeIdx)
-		: Proxy(nullptr)
+		: bHasCachedLeafInfo(false)
+		, Proxy(nullptr)
 		, ShapeIdx(InShapeIdx)
 		, Geometry()
 		, WorldSpaceInflatedShapeBounds(FAABB3(FVec3(0), FVec3(0)))
 	{
 	}
 
+	FPerShapeData::FPerShapeData(int32 InShapeIdx, TSerializablePtr<FImplicitObject> InGeometry, bool bInHasCachedLeafInfo)
+		: bHasCachedLeafInfo(bInHasCachedLeafInfo)
+		, Proxy(nullptr)
+		, ShapeIdx(InShapeIdx)
+		, Geometry(InGeometry)
+		, WorldSpaceInflatedShapeBounds(FAABB3(FVec3(0), FVec3(0)))
+	{
+	}
+
+	FPerShapeData::FPerShapeData(FPerShapeData&& Other)
+		: bHasCachedLeafInfo(Other.bHasCachedLeafInfo)
+		, Proxy(MoveTemp(Other.Proxy))
+		, DirtyFlags(MoveTemp(Other.DirtyFlags))
+		, ShapeIdx(MoveTemp(Other.ShapeIdx))
+		, CollisionData(MoveTemp(Other.CollisionData))
+		, Materials(MoveTemp(Other.Materials))
+		, Geometry(MoveTemp(Other.Geometry))
+		, WorldSpaceInflatedShapeBounds(MoveTemp(Other.WorldSpaceInflatedShapeBounds))
+	{
+	}
+
 	FPerShapeData::~FPerShapeData()
 	{
 	}
 
-	TUniquePtr<FPerShapeData> FPerShapeData::CreatePerShapeData(int32 ShapeIdx)
+	TUniquePtr<FPerShapeData> FPerShapeData::CreatePerShapeData(int32 ShapeIdx, TSerializablePtr<FImplicitObject> InGeometry)
+	{
+		FReal LeafMargin = 0.0f;
+		FRigidTransformRealSingle3 LeafRelativeTransform = FRigidTransformRealSingle3::Identity;
+		const FImplicitObject* LeafGeometry = GetInnerGeometryInstanceData(InGeometry.Get(), LeafRelativeTransform, LeafMargin);
+
+		if (RequiresCachedLeafInfo(LeafRelativeTransform, LeafGeometry, InGeometry.Get()))
+		{
+			return TUniquePtr<FPerShapeData>(new FPerShapeDataCachedLeafInfo(ShapeIdx, InGeometry, LeafGeometry, LeafRelativeTransform));
+		}
+		else
+		{
+			return TUniquePtr<FPerShapeData>(new FPerShapeData(ShapeIdx, InGeometry));
+		}
+	}
+
+	void FPerShapeData::UpdateGeometry(TUniquePtr<FPerShapeData>& ShapePtr, TSerializablePtr<FImplicitObject> InGeometry)
+	{
+		ShapePtr->Geometry = InGeometry;
+
+		FReal LeafMargin = 0.0f;
+		FRigidTransformRealSingle3 LeafRelativeTransform = FRigidTransformRealSingle3::Identity;
+		const FImplicitObject* LeafGeometry = GetInnerGeometryInstanceData(InGeometry.Get(), LeafRelativeTransform, LeafMargin);
+
+		if (ShapePtr->bHasCachedLeafInfo)
+		{
+			// set relative info
+			ShapePtr->SetLeafRelativeTransform(LeafRelativeTransform);
+			ShapePtr->SetLeafGeometry(LeafGeometry);
+			return;
+		}
+
+		if (RequiresCachedLeafInfo(LeafRelativeTransform, LeafGeometry, InGeometry.Get()))
+		{
+			// We need to move to FPerShapeDataCachedLeafInfo to cache this info.
+			ShapePtr = TUniquePtr<FPerShapeData>(new FPerShapeDataCachedLeafInfo(MoveTemp(*ShapePtr.Get()), LeafGeometry, LeafRelativeTransform));
+		}
+	}
+
+	bool FPerShapeData::RequiresCachedLeafInfo(const FRigidTransformRealSingle3& RelativeTransform, const FImplicitObject* LeafGeometry, const FImplicitObject* Geometry)
 	{
-		return TUniquePtr<FPerShapeData>(new FPerShapeData(ShapeIdx));
+		return (LeafGeometry != Geometry || !RelativeTransform.Equals(FRigidTransformRealSingle3::Identity));
 	}
 
 	void FPerShapeData::UpdateShapeBounds(const FRigidTransform3& WorldTM, const FVec3& BoundsExpansion)
@@ -131,11 +200,44 @@ namespace Chaos
 		}
 	}
 
-	void FPerShapeData::UpdateLeafGeometry()
+	const FImplicitObject* FPerShapeData::GetLeafGeometry() const
+	{
+		if (bHasCachedLeafInfo)
+		{
+			const FPerShapeDataCachedLeafInfo& LeafInfo = *static_cast<const FPerShapeDataCachedLeafInfo*>(this);
+			return LeafInfo.LeafGeometry;
+		}
+
+		return Geometry.Get();
+	}
+
+	void FPerShapeData::SetLeafGeometry(const FImplicitObject* LeafGeometry)
+	{
+		if (ensure(bHasCachedLeafInfo))
+		{
+			FPerShapeDataCachedLeafInfo& LeafInfo = *static_cast<FPerShapeDataCachedLeafInfo*>(this);
+			LeafInfo.LeafGeometry = LeafGeometry;
+		}
+	}
+
+	FRigidTransformRealSingle3 FPerShapeData::GetLeafRelativeTransform() const
 	{
-		LeafRelativeTransform = FRigidTransform3::Identity;
-		LeafMargin = 0.0f;
-		LeafGeometry = GetInnerGeometryInstanceData(Geometry.Get(), LeafRelativeTransform, LeafMargin);
+		if (bHasCachedLeafInfo)
+		{
+			const FPerShapeDataCachedLeafInfo& LeafInfo = *static_cast<const FPerShapeDataCachedLeafInfo*>(this);
+			return LeafInfo.LeafRelativeTransform;
+		}
+
+		return FRigidTransformRealSingle3::Identity;
+	}
+
+	void FPerShapeData::SetLeafRelativeTransform(const FRigidTransformRealSingle3& RelativeTransform)
+	{
+		if (ensure(bHasCachedLeafInfo))
+		{
+			FPerShapeDataCachedLeafInfo& LeafInfo = *static_cast<FPerShapeDataCachedLeafInfo*>(this);
+			LeafInfo.LeafRelativeTransform = RelativeTransform;
+		}
 	}
 
 	FPerShapeData* FPerShapeData::SerializationFactory(FChaosArchive& Ar, FPerShapeData*)
diff --git a/Engine/Source/Runtime/Experimental/Chaos/Public/Chaos/GeometryParticles.h b/Engine/Source/Runtime/Experimental/Chaos/Public/Chaos/GeometryParticles.h
--- a/Engine/Source/Runtime/Experimental/Chaos/Public/Chaos/GeometryParticles.h
+++ b/Engine/Source/Runtime/Experimental/Chaos/Public/Chaos/GeometryParticles.h
@@ -33,10 +33,11 @@ namespace Chaos
 	public:
 
 		static constexpr bool AlwaysSerializable = true;
-		static TUniquePtr<FPerShapeData> CreatePerShapeData(int32 InShapeIdx);
+		static TUniquePtr<FPerShapeData> CreatePerShapeData(int32 InShapeIdx, TSerializablePtr<FImplicitObject> InGeometry);
+		static void UpdateGeometry(TUniquePtr<FPerShapeData>& ShapePtr, TSerializablePtr<FImplicitObject> InGeometry);
+		static bool RequiresCachedLeafInfo(const FRigidTransformRealSingle3& RelativeTransform, const FImplicitObject* LeafGeometry, const FImplicitObject* Geometry);
 
 		~FPerShapeData();
-		FPerShapeData(const FPerShapeData& Other) = delete;
 		
 		void UpdateShapeBounds(const FRigidTransform3& WorldTM, const FVec3& BoundsExpansion = FVec3(0));
 
@@ -68,35 +69,35 @@ namespace Chaos
 		}
 
 		TSerializablePtr<FImplicitObject> GetGeometry() const { return Geometry; }
-		void SetGeometry(TSerializablePtr<FImplicitObject> InGeometry)
-		{
-			Geometry = InGeometry;
-			UpdateLeafGeometry();
-		}
 
 		const TAABB<FReal,3>& GetWorldSpaceInflatedShapeBounds() const { return WorldSpaceInflatedShapeBounds; }
 
 		void UpdateWorldSpaceState(const FRigidTransform3& WorldTransform, const FVec3& BoundsExpansion)
 		{
-			LeafWorldTransform = FRigidTransform3::MultiplyNoScale(LeafRelativeTransform, WorldTransform);
+			const FImplicitObject* LeafGeometry = GetLeafGeometry();
+			if (bHasCachedLeafInfo)
+			{
+				LeafWorldTransform = FRigidTransform3::MultiplyNoScale(FRigidTransform3(GetLeafRelativeTransform()), WorldTransform);
+			}
+			else
+			{
+				LeafWorldTransform = WorldTransform;
+			}
+			
 			if ((LeafGeometry != nullptr) && LeafGeometry->HasBoundingBox())
 			{
 				WorldSpaceInflatedShapeBounds = LeafGeometry->CalculateTransformedBounds(LeafWorldTransform).ThickenSymmetrically(BoundsExpansion);;
 			}
 		}
 
-		// The leaf shape (with transformed and implicit wrapper removed)
-		const FImplicitObject* GetLeafGeometry() const { return LeafGeometry; }
+		// The leaf shape (with transformed and implicit wrapper removed).
+		const FImplicitObject* GetLeafGeometry() const;
 
-		// The actor-relative transform of the leaf geometry
-		const FRigidTransform3& GetLeafRelativeTransform() const { return LeafRelativeTransform; }
+		// The actor-relative transform of the leaf geometry.
+		FRigidTransformRealSingle3 GetLeafRelativeTransform() const;
 
 		// The world-space transform of the leaf geometry (from the last call to UpdateWorldSpaceState)
 		const FRigidTransform3& GetLeafWorldTransform() const { return LeafWorldTransform; }
-		void SetLeafWorldTransform(const FRigidTransform3& InTransform) { LeafWorldTransform = InTransform; }
-
-		// The margin to use with the leaf shape (which could be from its Instanced wrapper if it had one)
-		const FReal GetLeafMargin() const { return LeafMargin; }
 
 		const TArray<FMaterialHandle>& GetMaterials() const { return Materials.Read().Materials; }
 		const TArray<FMaterialMaskHandle>& GetMaterialMasks() const { return Materials.Read().MaterialMasks; }
@@ -249,7 +250,21 @@ namespace Chaos
 
 	private:
 
-		void UpdateLeafGeometry();
+		void SetLeafRelativeTransform(const FRigidTransformRealSingle3& RelativeTransform);
+		void SetLeafGeometry(const FImplicitObject* LeafGeometry);
+
+	protected:
+		// use CreatePerShapeData
+		FPerShapeData(int32 InShapeIdx, TSerializablePtr<FImplicitObject> InGeometry, bool bInHasCachedLeafInfo = false);
+
+		explicit FPerShapeData(FPerShapeData&& Other);
+
+		// Can we be downcasted to FPerShapeDataCachedLeafInfo?
+		uint8 bHasCachedLeafInfo : 1;
+
+	private:
+		// Should only be used by SerializationFactory.
+		FPerShapeData(int32 InShapeIdx);
 
 		class IPhysicsProxyBase* Proxy;
 		FShapeDirtyFlags DirtyFlags;
@@ -261,13 +276,35 @@ namespace Chaos
 		TSerializablePtr<FImplicitObject> Geometry;
 		TAABB<FReal,3> WorldSpaceInflatedShapeBounds;
 
-		const FImplicitObject* LeafGeometry;
-		FRigidTransform3 LeafRelativeTransform;
 		FRigidTransform3 LeafWorldTransform;
-		FReal LeafMargin;
-		
-		// use CreatePerShapeData
-		FPerShapeData(int32 InShapeIdx);
+	};
+
+	class CHAOS_API FPerShapeDataCachedLeafInfo final : public FPerShapeData 
+	{
+		friend class FPerShapeData;
+
+	private:
+		FPerShapeDataCachedLeafInfo(int32 InShapeIdx, TSerializablePtr<FImplicitObject> InGeometry, const FImplicitObject* InLeafGeometry, const FRigidTransformRealSingle3& RelativeTransform)
+			: FPerShapeData(InShapeIdx, InGeometry, /*bHasCachedLeafInfo=*/true)
+			, LeafRelativeTransform(RelativeTransform)
+			, LeafGeometry(InLeafGeometry)
+		{}
+
+		// Move old shape into shape with cached leaf info
+		FPerShapeDataCachedLeafInfo(FPerShapeData&& PerShapeData, const FImplicitObject* InLeafGeometry, const FRigidTransformRealSingle3& RelativeTransform)
+			: FPerShapeData(MoveTemp(PerShapeData))
+			, LeafRelativeTransform(RelativeTransform)
+			, LeafGeometry(InLeafGeometry)
+		{
+			bHasCachedLeafInfo = true;
+		}
+
+		// WARNING: //////////////////////////////////////////////////////
+		// We don't have a virtual destructor on FPerShapeData (trying to save memory)
+		// FPerShapeDataCachedLeafInfo destructor will NOT be called, do not add non-PoD members.
+
+		FRigidTransformRealSingle3 LeafRelativeTransform;
+		const FImplicitObject* LeafGeometry;
 	};
 
 	inline FChaosArchive& operator<<(FChaosArchive& Ar, FPerShapeData& Shape)
diff --git a/Engine/Source/Runtime/Experimental/Chaos/Public/Chaos/Properties.h b/Engine/Source/Runtime/Experimental/Chaos/Public/Chaos/Properties.h
--- a/Engine/Source/Runtime/Experimental/Chaos/Public/Chaos/Properties.h
+++ b/Engine/Source/Runtime/Experimental/Chaos/Public/Chaos/Properties.h
@@ -105,6 +105,9 @@ public:
 		: Property(Val)
 	{
 	}
+	explicit TShapeProperty(TShapeProperty<T, PropName>&& Other)
+		: Property(MoveTemp(Other.Property))
+	{}
 
 	//we don't support this because it could lead to bugs with values not being properly written to remote
 	TShapeProperty(const TShapeProperty<T,PropName>& Rhs) = delete;
diff --git a/Engine/Source/Runtime/Experimental/ChaosCore/Public/Chaos/Core.h b/Engine/Source/Runtime/Experimental/ChaosCore/Public/Chaos/Core.h
--- a/Engine/Source/Runtime/Experimental/ChaosCore/Public/Chaos/Core.h
+++ b/Engine/Source/Runtime/Experimental/ChaosCore/Public/Chaos/Core.h
@@ -20,6 +20,7 @@ namespace Chaos
 	using FMatrix33 = PMatrix<FReal, 3, 3>;
 	using FMatrix44 = PMatrix<FReal, 4, 4>;
 	using FRigidTransform3 = TRigidTransform<FReal, 3>;
+	using FRigidTransformRealSingle3 = TRigidTransform<FRealSingle, 3>;
 
 	using FAABB3 = TAABB<FReal, 3>;
 
@@ -41,5 +42,5 @@ namespace Chaos
 	// NOTE: if you get a merge conflict on the GUID, you must replace it with a new GUID - do not accept the source or target
 	// or you will likely get DDC version conflicts resulting in crashes during load.
 	// Core version string for Chaos data. Any DDC builder dependent on Chaos for serialization should depend on this version
-	static const FString ChaosVersionString = TEXT("727DAD8C-96F4-4726-98C5-56AC8B1F5158");
+	static const FString ChaosVersionString = TEXT("F8EF37A7-2C6C-4269-BE5F-6B8B7C6B5295");
 }
diff --git a/Engine/Source/Runtime/Landscape/Private/LandscapeCollision.cpp b/Engine/Source/Runtime/Landscape/Private/LandscapeCollision.cpp
--- a/Engine/Source/Runtime/Landscape/Private/LandscapeCollision.cpp
+++ b/Engine/Source/Runtime/Landscape/Private/LandscapeCollision.cpp
@@ -436,11 +436,11 @@ void ULandscapeHeightfieldCollisionComponent::OnCreatePhysicsState()
 				TArray<TUniquePtr<Chaos::FImplicitObject>> Geoms;
 
 				// First add complex geometry
-				TUniquePtr<Chaos::FPerShapeData> NewShape = Chaos::FPerShapeData::CreatePerShapeData(ShapeArray.Num());
-
 				HeightfieldRef->Heightfield->SetScale(FinalScale * LandscapeComponentTransform.GetScale3D().GetSignVector());
 				TUniquePtr<Chaos::TImplicitObjectTransformed<Chaos::FReal, 3>> ChaosHeightFieldFromCooked = MakeUnique<Chaos::TImplicitObjectTransformed<Chaos::FReal, 3>>(MakeSerializable(HeightfieldRef->Heightfield), Chaos::FRigidTransform3(FTransform::Identity));
 
+				TUniquePtr<Chaos::FPerShapeData> NewShape = Chaos::FPerShapeData::CreatePerShapeData(ShapeArray.Num(), MakeSerializable(ChaosHeightFieldFromCooked));
+
 				// Setup filtering
 				FCollisionFilterData QueryFilterData, SimFilterData;
 				CreateShapeFilterData(GetCollisionObjectType(), FMaskFilter(0), GetOwner()->GetUniqueID(), GetCollisionResponseToChannels(), GetUniqueID(), 0, QueryFilterData, SimFilterData, true, false, true);
@@ -449,7 +449,6 @@ void ULandscapeHeightfieldCollisionComponent::OnCreatePhysicsState()
 				QueryFilterData.Word3 |= bCreateSimpleCollision ? EPDF_ComplexCollision : (EPDF_SimpleCollision | EPDF_ComplexCollision);
 				SimFilterData.Word3 |= bCreateSimpleCollision ? EPDF_ComplexCollision : (EPDF_SimpleCollision | EPDF_ComplexCollision);
 
-				NewShape->SetGeometry(MakeSerializable(ChaosHeightFieldFromCooked));
 				NewShape->SetQueryData(QueryFilterData);
 				NewShape->SetSimData(SimFilterData);
 				NewShape->SetMaterials(HeightfieldRef->UsedChaosMaterials);
@@ -460,18 +459,17 @@ void ULandscapeHeightfieldCollisionComponent::OnCreatePhysicsState()
 				// Add simple geometry if necessary
 				if(bCreateSimpleCollision)
 				{
-					TUniquePtr<Chaos::FPerShapeData> NewSimpleShape = Chaos::FPerShapeData::CreatePerShapeData(ShapeArray.Num());
-					
 					FVector FinalSimpleCollisionScale(LandscapeScale.X* SimpleCollisionScale, LandscapeScale.Y* SimpleCollisionScale, LandscapeScale.Z* LANDSCAPE_ZSCALE);
 					HeightfieldRef->HeightfieldSimple->SetScale(FinalSimpleCollisionScale);
 					TUniquePtr<Chaos::TImplicitObjectTransformed<Chaos::FReal, 3>> ChaosSimpleHeightFieldFromCooked = MakeUnique<Chaos::TImplicitObjectTransformed<Chaos::FReal, 3>>(MakeSerializable(HeightfieldRef->HeightfieldSimple), Chaos::FRigidTransform3(FTransform::Identity));
 
+					TUniquePtr<Chaos::FPerShapeData> NewSimpleShape = Chaos::FPerShapeData::CreatePerShapeData(ShapeArray.Num(), MakeSerializable(ChaosSimpleHeightFieldFromCooked));
+
 					FCollisionFilterData QueryFilterDataSimple = QueryFilterData;
 					FCollisionFilterData SimFilterDataSimple = SimFilterData;
 					QueryFilterDataSimple.Word3 = (QueryFilterDataSimple.Word3 & ~EPDF_ComplexCollision) | EPDF_SimpleCollision;
 					SimFilterDataSimple.Word3 = (SimFilterDataSimple.Word3 & ~EPDF_ComplexCollision) | EPDF_SimpleCollision;
 
-					NewSimpleShape->SetGeometry(MakeSerializable(ChaosSimpleHeightFieldFromCooked));
 					NewSimpleShape->SetQueryData(QueryFilterDataSimple);
 					NewSimpleShape->SetSimData(SimFilterDataSimple);
 
@@ -483,11 +481,11 @@ void ULandscapeHeightfieldCollisionComponent::OnCreatePhysicsState()
 				// Create a shape for a heightfield which is used only by the landscape editor
 				if(!GetWorld()->IsGameWorld() && !GetOutermost()->bIsCookedForEditor)
 				{
-					TUniquePtr<Chaos::FPerShapeData> NewEditorShape = Chaos::FPerShapeData::CreatePerShapeData(ShapeArray.Num());
-
 					HeightfieldRef->EditorHeightfield->SetScale(FinalScale);
 					TUniquePtr<Chaos::TImplicitObjectTransformed<Chaos::FReal, 3>> ChaosEditorHeightFieldFromCooked = MakeUnique<Chaos::TImplicitObjectTransformed<Chaos::FReal, 3>>(MakeSerializable(HeightfieldRef->EditorHeightfield), Chaos::FRigidTransform3(FTransform::Identity));
 
+					TUniquePtr<Chaos::FPerShapeData> NewEditorShape = Chaos::FPerShapeData::CreatePerShapeData(ShapeArray.Num(), MakeSerializable(ChaosEditorHeightFieldFromCooked));
+
 					FCollisionResponseContainer CollisionResponse;
 					CollisionResponse.SetAllChannels(ECollisionResponse::ECR_Ignore);
 					CollisionResponse.SetResponse(ECollisionChannel::ECC_Visibility, ECR_Block);
@@ -496,7 +494,6 @@ void ULandscapeHeightfieldCollisionComponent::OnCreatePhysicsState()
 
 					QueryFilterDataEd.Word3 |= (EPDF_SimpleCollision | EPDF_ComplexCollision);
 
-					NewEditorShape->SetGeometry(MakeSerializable(ChaosEditorHeightFieldFromCooked));
 					NewEditorShape->SetQueryData(QueryFilterDataEd);
 					NewEditorShape->SetSimData(SimFilterDataEd);
 
@@ -1637,10 +1634,9 @@ struct FMeshCollisionInitHelper
 		FVector Scale = FVector(ComponentScale.X * CollisionScale, ComponentScale.Y * CollisionScale, ComponentScale.Z);
 
 		{
-			TUniquePtr<Chaos::FPerShapeData> NewShape = Chaos::FPerShapeData::CreatePerShapeData(ShapeArray.Num());
 			TUniquePtr<Chaos::TImplicitObjectScaled<Chaos::FTriangleMeshImplicitObject>> ScaledTrimesh = MakeUnique<Chaos::TImplicitObjectScaled<Chaos::FTriangleMeshImplicitObject>>(MakeSerializable(MeshRef->Trimesh), Scale);
+			TUniquePtr<Chaos::FPerShapeData> NewShape = Chaos::FPerShapeData::CreatePerShapeData(ShapeArray.Num(), MakeSerializable(ScaledTrimesh));
 
-			NewShape->SetGeometry(MakeSerializable(ScaledTrimesh));
 			NewShape->SetQueryData(QueryFilter);
 			NewShape->SetSimData(SimulationFilter);
 			NewShape->SetCollisionTraceType(Chaos::EChaosCollisionTraceFlag::Chaos_CTF_UseComplexAsSimple);
@@ -1653,10 +1649,9 @@ struct FMeshCollisionInitHelper
 #if WITH_EDITOR
 		if(!World->IsGameWorld())
 		{
-			TUniquePtr<Chaos::FPerShapeData> NewEdShape = Chaos::FPerShapeData::CreatePerShapeData(ShapeArray.Num());
 			TUniquePtr<Chaos::TImplicitObjectScaled<Chaos::FTriangleMeshImplicitObject>> ScaledTrimeshEd = MakeUnique<Chaos::TImplicitObjectScaled<Chaos::FTriangleMeshImplicitObject>>(MakeSerializable(MeshRef->EditorTrimesh), Scale);
+			TUniquePtr<Chaos::FPerShapeData> NewEdShape = Chaos::FPerShapeData::CreatePerShapeData(ShapeArray.Num(), MakeSerializable(ScaledTrimeshEd));
 
-			NewEdShape->SetGeometry(MakeSerializable(ScaledTrimeshEd));
 			NewEdShape->SetQueryData(QueryFilterEd);
 			NewEdShape->SetSimEnabled(false);
 			NewEdShape->SetCollisionTraceType(Chaos::EChaosCollisionTraceFlag::Chaos_CTF_UseComplexAsSimple);
