Fix CIS IslandManager.cpp

diff --git a/Engine/Source/Runtime/Experimental/Chaos/Private/Chaos/Island/IslandManager.cpp b/Engine/Source/Runtime/Experimental/Chaos/Private/Chaos/Island/IslandManager.cpp
--- a/Engine/Source/Runtime/Experimental/Chaos/Private/Chaos/Island/IslandManager.cpp
+++ b/Engine/Source/Runtime/Experimental/Chaos/Private/Chaos/Island/IslandManager.cpp
@@ -128,7 +128,7 @@ inline void UpdateSleepState(FPBDIslandSolver* IslandSolver, FPBDRigidsSOAs& Par
 				const FRigidTransform3 ParticleWorldTransform = FParticleUtilities::GetActorWorldTransform(FConstGenericParticleHandle(IslandParticle));
 				for (const TUniquePtr<FPerShapeData>& Shape : IslandParticle->ShapesArray())
 				{
-					const FRigidTransform3 ShapeWorldTransform = Shape->GetLeafRelativeTransform() * ParticleWorldTransform;
+					const FRigidTransform3 ShapeWorldTransform = (FRigidTransform3)Shape->GetLeafRelativeTransform() * ParticleWorldTransform;
 					Shape->SetLeafWorldTransform(ShapeWorldTransform);
 				}
 
diff --git a/Engine/Source/Runtime/Experimental/Chaos/Public/Chaos/GeometryParticles.h b/Engine/Source/Runtime/Experimental/Chaos/Public/Chaos/GeometryParticles.h
--- a/Engine/Source/Runtime/Experimental/Chaos/Public/Chaos/GeometryParticles.h
+++ b/Engine/Source/Runtime/Experimental/Chaos/Public/Chaos/GeometryParticles.h
@@ -98,6 +98,7 @@ namespace Chaos
 
 		// The world-space transform of the leaf geometry (from the last call to UpdateWorldSpaceState)
 		const FRigidTransform3& GetLeafWorldTransform() const { return LeafWorldTransform; }
+		void SetLeafWorldTransform(const FRigidTransform3& InTransform) { LeafWorldTransform = InTransform; }
 
 		const TArray<FMaterialHandle>& GetMaterials() const { return Materials.Read().Materials; }
 		const TArray<FMaterialMaskHandle>& GetMaterialMasks() const { return Materials.Read().MaterialMasks; }
