Chaos - improve perofrmance in GJK using vectorization for sweep test
- Add simplex fully vectorized
- Add a vectorized utility functions to support additional intrinics on Neon and SEE

diff --git a/Engine/Source/Runtime/Experimental/Chaos/Public/Chaos/AABB.h b/Engine/Source/Runtime/Experimental/Chaos/Public/Chaos/AABB.h
--- a/Engine/Source/Runtime/Experimental/Chaos/Public/Chaos/AABB.h
+++ b/Engine/Source/Runtime/Experimental/Chaos/Public/Chaos/AABB.h
@@ -371,6 +371,16 @@ namespace Chaos
 			return ChosenPt;
 		}
 
+		FORCEINLINE_DEBUGGABLE VectorRegister4Float SupportCoreSimd(const VectorRegister4Float& Direction, const FReal InMargin) const
+		{
+			FVec3 DirectionVec3;
+			VectorStoreFloat3(Direction, &DirectionVec3);
+			int32 VertexIndex = INDEX_NONE;
+			FVec3 SupportVert = SupportCore(DirectionVec3, InMargin, nullptr, VertexIndex);
+			alignas(16) FRealSingle SupportVertFloat[4] = { static_cast<FRealSingle>(SupportVert.X), static_cast<FRealSingle>(SupportVert.Y), static_cast<FRealSingle>(SupportVert.Z), 0.0f };
+			return VectorLoadAligned(SupportVertFloat);
+		}
+
 		FORCEINLINE_DEBUGGABLE TVector<T, d> SupportCoreScaled(const TVector<T, d>& Direction, const T InMargin, const TVector<T, d>& Scale, T* OutSupportDelta, int32& VertexIndex) const
 		{
 			const TVector<T, d> ScaledDirection = Direction * Scale;
diff --git a/Engine/Source/Runtime/Experimental/Chaos/Public/Chaos/Box.h b/Engine/Source/Runtime/Experimental/Chaos/Public/Chaos/Box.h
--- a/Engine/Source/Runtime/Experimental/Chaos/Public/Chaos/Box.h
+++ b/Engine/Source/Runtime/Experimental/Chaos/Public/Chaos/Box.h
@@ -374,6 +374,17 @@ namespace Chaos
 			return AABB.SupportCore(Direction, InMargin, OutSupportDelta, VertexIndex);
 		}
 
+		// Returns a position on the core shape excluding the margin
+		FORCEINLINE_DEBUGGABLE VectorRegister4Float SupportCoreSimd(const VectorRegister4Float& Direction, const FReal InMargin) const
+		{
+			FVec3 DirectionVec3;
+			VectorStoreFloat3(Direction, &DirectionVec3);
+			int32 VertexIndex = INDEX_NONE;
+			FVec3 SupportVert = AABB.SupportCore(DirectionVec3, InMargin, nullptr, VertexIndex);
+			alignas(16) FRealSingle SupportVertFloat[4] = { static_cast<FRealSingle>(SupportVert.X), static_cast<FRealSingle>(SupportVert.Y), static_cast<FRealSingle>(SupportVert.Z), 0.0f };
+			return VectorLoadAligned(SupportVertFloat);
+		}
+
 		FORCEINLINE_DEBUGGABLE TVector<T, d> SupportCoreScaled(const TVector<T, d>& Direction, const FReal InMargin, const TVector<T, d>& Scale, T* OutSupportDelta, int32& VertexIndex) const
 		{
 			return AABB.SupportCoreScaled(Direction, InMargin, Scale, OutSupportDelta, VertexIndex);
diff --git a/Engine/Source/Runtime/Experimental/Chaos/Public/Chaos/Capsule.h b/Engine/Source/Runtime/Experimental/Chaos/Public/Chaos/Capsule.h
--- a/Engine/Source/Runtime/Experimental/Chaos/Public/Chaos/Capsule.h
+++ b/Engine/Source/Runtime/Experimental/Chaos/Public/Chaos/Capsule.h
@@ -10,6 +10,8 @@
 #include "Chaos/Segment.h"
 #include "ChaosArchive.h"
 
+#include "Math/VectorRegister.h"
+
 #include "UObject/ReleaseObjectVersion.h"
 
 namespace Chaos
@@ -361,6 +363,18 @@ namespace Chaos
 			return MSegment.SupportCore(Direction, VertexIndex);
 		}
 
+		FORCEINLINE VectorRegister4Float SupportCoreSimd(const VectorRegister4Float& Direction, const FReal InMargin) const
+		{
+			// NOTE: Ignores InMargin, assumes Radius
+			FVec3 DirectionVec3;
+			VectorStoreFloat3(Direction, &DirectionVec3);
+			int32 VertexIndex = INDEX_NONE;
+			FVec3 SupportVert =  MSegment.SupportCore(DirectionVec3, VertexIndex);
+			alignas(16) FRealSingle SupportVertFloat[4] = { static_cast<FRealSingle>(SupportVert.X), static_cast<FRealSingle>(SupportVert.Y), static_cast<FRealSingle>(SupportVert.Z), 0.0f };
+			return VectorLoadAligned(SupportVertFloat);
+		}
+
+
 		FORCEINLINE FVec3 SupportCoreScaled(const FVec3& Direction, const FReal InMargin, const FVec3& Scale, FReal* OutSupportDelta, int32& VertexIndex) const
 		{
 			// NOTE: Ignores InMargin, assumes Radius
diff --git a/Engine/Source/Runtime/Experimental/Chaos/Public/Chaos/Convex.h b/Engine/Source/Runtime/Experimental/Chaos/Public/Chaos/Convex.h
--- a/Engine/Source/Runtime/Experimental/Chaos/Public/Chaos/Convex.h
+++ b/Engine/Source/Runtime/Experimental/Chaos/Public/Chaos/Convex.h
@@ -6,6 +6,8 @@
 #include "Chaos/ConvexStructureData.h"
 #include "Chaos/MassProperties.h"
 
+#include "Math/UnrealMathVectorCommon.h"
+
 #include "CollisionConvexMesh.h"
 #include "ChaosArchive.h"
 #include "ChaosCheck.h"
@@ -636,6 +638,22 @@ namespace Chaos
 			return FVec3(0);
 		}
 
+		// Return support point on the core shape (the convex shape with all planes moved inwards by margin).
+		VectorRegister4Float SupportCoreSimd(const VectorRegister4Float& Direction, const FReal InMargin) const
+		{
+			FVec3 DirectionVec3;
+			VectorStoreFloat3(Direction, &DirectionVec3);
+			const int32 SupportVertexIndex = GetSupportVertex(DirectionVec3);
+			if (SupportVertexIndex != INDEX_NONE)
+			{
+				FVec3 SupportVert =  GetMarginAdjustedVertex(SupportVertexIndex, InMargin, nullptr);
+				alignas(16) FRealSingle SupportVertFloat[4] = { static_cast<FRealSingle>(SupportVert.X), static_cast<FRealSingle>(SupportVert.Y), static_cast<FRealSingle>(SupportVert.Z), 0.0f };
+				return VectorLoadAligned(SupportVertFloat);
+			}
+			return VectorZeroFloat();
+		}
+
+
 		// SupportCore with non-uniform scale support. This is required for the margin in scaled
 		// space to by uniform. Note in this version all the inputs are in outer container's (scaled shape) space
 		FVec3 SupportCoreScaled(const FVec3& Direction, FReal InMargin, const FVec3& Scale, FReal* OutSupportDelta, int32& VertexIndex) const
diff --git a/Engine/Source/Runtime/Experimental/Chaos/Public/Chaos/EPA.h b/Engine/Source/Runtime/Experimental/Chaos/Public/Chaos/EPA.h
--- a/Engine/Source/Runtime/Experimental/Chaos/Public/Chaos/EPA.h
+++ b/Engine/Source/Runtime/Experimental/Chaos/Public/Chaos/EPA.h
@@ -48,7 +48,7 @@ struct TEPAEntry
 		const TVec3<T> V0V1 = V1 - V0;
 		const TVec3<T> V0V2 = V2 - V0;
 		const TVec3<T> Norm = TVec3<T>::CrossProduct(V0V1, V0V2);
-		const FReal NormLenSq = Norm.SizeSquared();
+		const T NormLenSq = Norm.SizeSquared();
 		constexpr T Eps = 1e-4f;
 		if (NormLenSq < Eps)
 		{
diff --git a/Engine/Source/Runtime/Experimental/Chaos/Public/Chaos/GJK.h b/Engine/Source/Runtime/Experimental/Chaos/Public/Chaos/GJK.h
--- a/Engine/Source/Runtime/Experimental/Chaos/Public/Chaos/GJK.h
+++ b/Engine/Source/Runtime/Experimental/Chaos/Public/Chaos/GJK.h
@@ -12,6 +12,27 @@
 #include "ChaosCheck.h"
 #include "ChaosLog.h"
 
+
+
+// Platform specific vector intrinsics include.
+#if WITH_DIRECTXMATH
+#define GJK_VECTORIZED 0
+#elif PLATFORM_ENABLE_VECTORINTRINSICS_NEON
+#include "Math/VectorRegister.h"
+#include "Chaos/VectorUtility.h"
+#include "Chaos/SimplexVectorized.h"
+#define GJK_VECTORIZED 1
+#elif defined(__cplusplus_cli) && !PLATFORM_HOLOLENS
+#define GJK_VECTORIZED 0
+#elif PLATFORM_ENABLE_VECTORINTRINSICS
+#include "Math/VectorRegister.h"
+#include "Chaos/VectorUtility.h"
+#include "Chaos/SimplexVectorized.h"
+#define GJK_VECTORIZED 1
+#else
+#define GJK_VECTORIZED 0
+#endif
+
 namespace Chaos
 {
 	
@@ -352,7 +373,7 @@ namespace Chaos
 
 			// If we are not too close, update the separating vector and keep iterating
 			// If we are too close we drop through to EPA, but if EPA determines the shapes are actually separated 
-			// after all, we will wend up using the normal from the previous loop
+			// after all, we will end up using the normal from the previous loop
 			if (!bIsContact)
 			{
 				V /= NewDistance;
@@ -360,14 +381,12 @@ namespace Chaos
 			}
 			Distance = NewDistance;
 		}
-
 		// Save the warm-start data (not much to store since we updated most of it in the loop above)
 		InOutSimplexData.Save(SimplexIDs);
-
 		if (bIsContact)
 		{
 			// We did not converge or we detected an overlap situation, so run EPA to get contact data
-			// Rebuild the simplex into a variable sized array - EPA can generate any mumber of vertices
+			// Rebuild the simplex into a variable sized array - EPA can generate any number of vertices
 			TArray<TVec3<T>> VertsA;
 			TArray<TVec3<T>> VertsB;
 			VertsA.Reserve(8);
@@ -408,7 +427,6 @@ namespace Chaos
 				// and penetration but we keep the position generated by GJK.
 				Normal = MTD;
 				Distance = -Penetration;
-				//UE_LOG(LogChaos, Warning, TEXT("EPA Touching Case"));
 				break;
 			case EEPAResult::Degenerate:
 				// We hit a degenerate simplex condition and could not reach a solution so use whatever near touching point GJK came up with.
@@ -863,7 +881,7 @@ namespace Chaos
 		bool bNearZero = false;
 		bool bDegenerate = false;
 		int NumIterations = 0;
-		T InGJKPreDist2 = TNumericLimits<T>::Max();
+		T GJKPreDist2 = TNumericLimits<T>::Max();
 		do
 		{
 			//if (!ensure(NumIterations++ < 32))	//todo: take this out
@@ -912,7 +930,7 @@ namespace Chaos
 					Simplex[2] += XMinusOldX;
 					Simplex[SimplexIDs.NumVerts++] = X - P;
 
-					InGJKPreDist2 = TNumericLimits<T>::Max();	//translated origin so restart gjk search
+					GJKPreDist2 = TNumericLimits<T>::Max();	//translated origin so restart gjk search
 				}
 			}
 			else
@@ -924,8 +942,8 @@ namespace Chaos
 
 			T NewDist2 = V.SizeSquared();	//todo: relative error
 			bNearZero = NewDist2 < 1e-6;
-			bDegenerate = NewDist2 >= InGJKPreDist2;
-			InGJKPreDist2 = NewDist2;
+			bDegenerate = NewDist2 >= GJKPreDist2;
+			GJKPreDist2 = NewDist2;
 			bTerminate = bNearZero || bDegenerate;
 
 		} while (!bTerminate);
@@ -1032,7 +1050,7 @@ namespace Chaos
 		bool bTerminate = bCloseEnough;
 		bool bInflatedCloseEnough = bCloseEnough;
 		int NumIterations = 0;
-		T InGJKPreDist2 = TNumericLimits<T>::Max();
+		T GJKPreDist2 = TNumericLimits<T>::Max();
 		while (!bTerminate)
 		{
 			//if (!ensure(NumIterations++ < 32))	//todo: take this out
@@ -1082,7 +1100,7 @@ namespace Chaos
 					Simplex[2] += XMinusOldX;
 					Simplex[SimplexIDs.NumVerts++] = X - P;
 
-					InGJKPreDist2 = TNumericLimits<T>::Max();	//translated origin so restart gjk search
+					GJKPreDist2 = TNumericLimits<T>::Max();	//translated origin so restart gjk search
 					bInflatedCloseEnough = false;
 				}
 			}
@@ -1095,7 +1113,7 @@ namespace Chaos
 			{
 				//Inflated shapes are close enough, but we want MTD so we need to find closest point on core shape
 				const T VDotW2 = VDotW * VDotW;
-				bCloseEnough = InGJKPreDist2 <= Eps2 + VDotW2;	//todo: relative error
+				bCloseEnough = GJKPreDist2 <= Eps2 + VDotW2;	//todo: relative error
 			}
 
 			if (!bCloseEnough)
@@ -1103,11 +1121,11 @@ namespace Chaos
 				V = SimplexFindClosestToOrigin(Simplex, SimplexIDs, Barycentric, As, Bs);
 				T NewDist2 = V.SizeSquared();	//todo: relative error
 				bCloseEnough = NewDist2 < Inflation2;
-				bDegenerate = NewDist2 >= InGJKPreDist2;
-				InGJKPreDist2 = NewDist2;
+				bDegenerate = NewDist2 >= GJKPreDist2;
+				GJKPreDist2 = NewDist2;
 
 
-				if (bComputeMTD && bCloseEnough && Lambda == 0 && InGJKPreDist2 > 1e-6 && Inflation2 > 1e-6 && SimplexIDs.NumVerts < 4)
+				if (bComputeMTD && bCloseEnough && Lambda == 0 && GJKPreDist2 > 1e-6 && Inflation2 > 1e-6 && SimplexIDs.NumVerts < 4)
 				{
 					//For mtd of inflated shapes we have to find the closest point, so we have to keep going
 					bCloseEnough = false;
@@ -1142,7 +1160,7 @@ namespace Chaos
 			// If Inflation == 0 we would expect GJKPreDist2 to be 0
 			// However, due to precision we can still end up with GJK failing.
 			// When that happens fall back on EPA
-			if (Inflation > 0 && InGJKPreDist2 > 1e-6 && InGJKPreDist2 < TNumericLimits<T>::Max())
+			if (Inflation > 0 && GJKPreDist2 > 1e-6 && GJKPreDist2 < TNumericLimits<T>::Max())
 			{
 				OutNormal = Normal;
 				TVector<T, 3> ClosestA(0);
@@ -1164,7 +1182,7 @@ namespace Chaos
 				}
 				
 				const TVec3<T> ClosestBInA = StartPoint + ClosestB;
-				const T InGJKPreDist = FMath::Sqrt(InGJKPreDist2);
+				const T InGJKPreDist = FMath::Sqrt(GJKPreDist2);
 				OutNormal = V.GetUnsafeNormal();
 
 				const T Penetration = FMath::Clamp<T>(MarginA + MarginB - InGJKPreDist, 0, TNumericLimits<T>::Max());
@@ -1231,6 +1249,431 @@ namespace Chaos
 		return true;
 	}
 
+	template <typename TGeometryA, typename TGeometryB, typename T>
+	bool GJKRaycast2ImplSimd(const TGeometryA& A, const TGeometryB& B, const VectorRegister4Float& BToARotation, const VectorRegister4Float& StartPoint, const VectorRegister4Float& RayDir, const T RayLength,
+		T& OutTime, VectorRegister4Float& OutPosition, VectorRegister4Float& OutNormal, bool bComputeMTD, const VectorRegister4Float& InitialDir, const TRigidTransform<double, 3>& StartTM)
+	{
+		ensure(RayLength > 0);
+
+		// Margin selection logic: we only need a small margin for sweeps since we only move the sweeping object
+		// to the point where it just touches.
+		// Spheres and Capsules: always use the core shape and full "margin" because it represents the radius
+		// Sphere/Capsule versus OtherShape: no margin on other
+		// OtherShape versus OtherShape: use margin of the smaller shape, zero margin on the other
+		const T RadiusA = static_cast<T>(A.GetRadius());
+		const T RadiusB = static_cast<T>(B.GetRadius());
+		const bool bHasRadiusA = RadiusA > 0;
+		const bool bHasRadiusB = RadiusB > 0;
+
+		// The sweep margins if required. Only one can be non-zero (we keep the smaller one)
+		const T SweepMarginScale = 0.05f;
+		const bool bAIsSmallest = A.GetMargin() < B.GetMargin();
+		const T SweepMarginA = (bHasRadiusA || bHasRadiusB) ? 0.0f : (bAIsSmallest ? SweepMarginScale * static_cast<T>(A.GetMargin()) : 0.0f);
+		const T SweepMarginB = (bHasRadiusA || bHasRadiusB) ? 0.0f : (bAIsSmallest ? 0.0f : SweepMarginScale * static_cast<T>(B.GetMargin()));
+
+		// Net margin (note: both SweepMargins are zero if either Radius is non-zero, and only one SweepMargin can be non-zero)
+		const T MarginA = RadiusA + SweepMarginA;
+		const T MarginB = RadiusB + SweepMarginB;
+
+		const VectorRegister4Float MarginASimd = VectorLoadFloat1(&MarginA);
+		const VectorRegister4Float MarginBSimd = VectorLoadFloat1(&MarginB);
+
+		VectorRegister4Float Simplex[4] = { VectorZeroFloat(), VectorZeroFloat(), VectorZeroFloat(), VectorZeroFloat() };
+		VectorRegister4Float As[4] = { VectorZeroFloat(), VectorZeroFloat(), VectorZeroFloat(), VectorZeroFloat() };
+		VectorRegister4Float Bs[4] = { VectorZeroFloat(), VectorZeroFloat(), VectorZeroFloat(), VectorZeroFloat() };
+
+		VectorRegister4Float Barycentric;
+
+		VectorRegister4Float Inflation = VectorAdd(MarginASimd, MarginBSimd);
+		constexpr VectorRegister4Float Eps2Simd = MakeVectorRegisterFloatConstant(1e-6f, 1e-6f, 1e-6f, 1e-6f);
+		const VectorRegister4Float Inflation2Simd = VectorMultiplyAdd(Inflation, Inflation, Eps2Simd);
+
+		const VectorRegister4Float RayLengthSimd = MakeVectorRegisterFloat(RayLength, RayLength, RayLength, RayLength);
+
+		NumVertsUnion NumVerts;
+		NumVerts.NumVerts = GlobalVectorConstants::IntZero;
+		VectorRegister4Float AToBRotation = VectorQuaternionInverse(BToARotation);
+
+		auto SupportAFunc = [&A, MarginA](const VectorRegister4Float& V)
+		{
+			return A.SupportCoreSimd(V, MarginA);
+		};
+
+		auto SupportBFunc = [&B, MarginB, &AToBRotation, &BToARotation](const VectorRegister4Float& V)
+		{
+			const VectorRegister4Float VInB = VectorQuaternionRotateVector(AToBRotation, V);
+			const VectorRegister4Float SupportBLocal = B.SupportCoreSimd(VInB, MarginB);
+
+			return VectorQuaternionRotateVector(BToARotation, SupportBLocal);
+		};
+
+
+		VectorRegister4Float SupportA = SupportAFunc(InitialDir);
+		As[0] = SupportA;
+
+		VectorRegister4Float SupportB = SupportBFunc(VectorNegate(InitialDir));
+		Bs[0] = SupportB;
+
+		VectorRegister4Float Lambda = VectorZeroFloat();
+		VectorRegister4Float X = StartPoint;
+		VectorRegister4Float V = VectorSubtract(X, VectorSubtract(SupportA, SupportB));
+		VectorRegister4Float Normal = { 0.f, 0.f, 1.f, 0.f };
+
+		const VectorRegister4Float InitialPreDist2Simd = VectorDot3(V, V);
+
+		FRealSingle InitialPreDist2;
+		VectorStoreFloat1(InitialPreDist2Simd, &InitialPreDist2);
+
+		FRealSingle Inflation2;
+		VectorStoreFloat1(Inflation2Simd, &Inflation2);
+		
+		constexpr FRealSingle Eps2 = 1e-6f;
+
+		//mtd needs to find closest point even in inflation region, so can only skip if we found the closest points
+		bool bCloseEnough = InitialPreDist2 < Inflation2 && (!bComputeMTD || InitialPreDist2 < Eps2);
+		bool bDegenerate = false;
+		bool bTerminate = bCloseEnough;
+		bool bInflatedCloseEnough = bCloseEnough;
+		int NumIterations = 0;
+		constexpr VectorRegister4Float LimitMax = MakeVectorRegisterFloatConstant(TNumericLimits<T>::Max(), TNumericLimits<T>::Max(), TNumericLimits<T>::Max(), TNumericLimits<T>::Max());
+		VectorRegister4Float GJKPreDist2 = LimitMax;
+
+		while (!bTerminate)
+		{
+			if (!(NumIterations++ < 32))	//todo: take this out
+			{
+				break;	//if taking too long just stop. This should never happen
+			}
+
+			V = VectorNormalizeAccurate(V);
+
+			SupportA = SupportAFunc(V);
+			SupportB = SupportBFunc(VectorNegate(V));
+			const VectorRegister4Float P = VectorSubtract(SupportA, SupportB);
+			const VectorRegister4Float W = VectorSubtract(X, P);
+
+			// NumVerts store here at the beginning of the loop, it should be safe to reuse it further in the loop 
+			const int32 NumVertsInt = NumVerts.NumVertsInts[0];
+
+			As[NumVertsInt] = SupportA;
+			Bs[NumVertsInt] = SupportB;
+
+			const VectorRegister4Float VDotW = VectorDot3(V, W);
+
+			VectorRegister4Float VDotWGTInflationSimd = VectorCompareGT(VDotW, Inflation);
+
+			if (VectorMaskBits(VDotWGTInflationSimd))
+			{
+				const VectorRegister4Float VDotRayDir = VectorDot3(V, RayDir);
+				VectorRegister4Float VDotRayDirGEZero = VectorCompareGE(VDotRayDir, VectorZeroFloat());
+
+				if (VectorMaskBits(VDotRayDirGEZero))
+				{
+					return false;
+				}
+
+				const VectorRegister4Float PreLambda = Lambda;	//use to check for no progress
+				// @todo(ccaulfield): this can still overflow - the comparisons against zero above should be changed (though not sure to what yet)
+				Lambda = VectorSubtract(Lambda, VectorDivide(VectorSubtract(VDotW, Inflation), VDotRayDir));
+				VectorRegister4Float LambdaGTPreLambda = VectorCompareGT(Lambda, PreLambda);
+				if (VectorMaskBits(LambdaGTPreLambda))
+				{
+					VectorRegister4Float LambdaGTRayLength = VectorCompareGT(Lambda, RayLengthSimd);
+					if (VectorMaskBits(LambdaGTRayLength))
+					{
+						return false;
+					}
+
+					const VectorRegister4Float OldX = X;
+					X = VectorMultiplyAdd(Lambda, RayDir, StartPoint);
+					Normal = V;
+
+					//Update simplex from (OldX - P) to (X - P)
+					VectorRegister4Float XMinusOldX = VectorSubtract(X, OldX);
+					Simplex[0] = VectorAdd(Simplex[0], XMinusOldX);
+					Simplex[1] = VectorAdd(Simplex[1], XMinusOldX);
+					Simplex[2] = VectorAdd(Simplex[2], XMinusOldX);
+					Simplex[NumVertsInt] = VectorSubtract(X, P);
+					NumVerts.NumVerts = VectorIntAdd(NumVerts.NumVerts, GlobalVectorConstants::IntOne);
+
+					GJKPreDist2 = LimitMax; //translated origin so restart gjk search
+					bInflatedCloseEnough = false;
+				}
+			}
+			else
+			{
+				Simplex[NumVertsInt] = W;	//this is really X - P which is what we need for simplex computation
+				NumVerts.NumVerts = VectorIntAdd(NumVerts.NumVerts, GlobalVectorConstants::IntOne);
+			}
+
+			if (bInflatedCloseEnough && VectorMaskBits(VectorCompareGE(VDotW, VectorZeroFloat())))
+			{
+				//Inflated shapes are close enough, but we want MTD so we need to find closest point on core shape
+				const VectorRegister4Float VDotW2 = VectorDot3(VDotW, VDotW);
+				bCloseEnough = static_cast<bool>(VectorMaskBits(VectorCompareGE(VectorAdd(Eps2Simd, VDotW), GJKPreDist2)));
+			}
+
+			if (!bCloseEnough)
+			{
+				V = VectorSimplexFindClosestToOrigin(Simplex, NumVerts, Barycentric, As, Bs);
+
+				VectorRegister4Float NewDist2 = VectorDot3(V, V);	//todo: relative error
+				bCloseEnough = static_cast<bool>(VectorMaskBits(VectorCompareGT(Inflation2Simd, NewDist2)));
+				bDegenerate = static_cast<bool>(VectorMaskBits(VectorCompareGE(NewDist2, GJKPreDist2)));
+				GJKPreDist2 = NewDist2;
+
+				if (bComputeMTD && bCloseEnough)
+				{
+					const VectorRegister4Float LambdaEqZero = VectorCompareEQ(Lambda, VectorZeroFloat());
+					const VectorRegister4Float InGJKPreDist2GTEps2 = VectorCompareGT(GJKPreDist2, Eps2Simd);
+					const VectorRegister4Float Inflation22GTEps2 = VectorCompareGT(Inflation2Simd, Eps2Simd);
+					constexpr VectorRegister4Int fourInt = MakeVectorRegisterIntConstant(4, 4, 4, 4);
+					const VectorRegister4Int Is4GTNumVerts = VectorIntCompareGT(fourInt, NumVerts.NumVerts);
+
+					const VectorRegister4Float IsInflatCloseEnough = VectorBitwiseAnd(LambdaEqZero, VectorBitwiseAnd(InGJKPreDist2GTEps2, VectorBitwiseAnd(Inflation22GTEps2, VectorCast4IntTo4Float(Is4GTNumVerts))));
+
+					// Leaving the original code, to explain the logic there
+					//if (bComputeMTD && bCloseEnough && Lambda == 0 && GJKPreDist2 > 1e-6 && Inflation2 > 1e-6 && NumVerts < 4)
+					//{
+					//	bCloseEnough = false;
+					//	bInflatedCloseEnough = true;
+					//}
+					bInflatedCloseEnough = static_cast<bool>(VectorMaskBits(IsInflatCloseEnough));
+					bCloseEnough = !bInflatedCloseEnough;
+				}
+			}
+			else
+			{
+				//It must be that we want MTD and we can terminate. However, we must make one final call to fixup the simplex
+				V = VectorSimplexFindClosestToOrigin(Simplex, NumVerts, Barycentric, As, Bs);
+			}
+			bTerminate = bCloseEnough || bDegenerate;
+		}
+		VectorStoreFloat1(Lambda, &OutTime);
+
+		if (OutTime > 0)
+		{
+			OutNormal = Normal;
+			VectorRegister4Float ClosestB = VectorZeroFloat();
+
+			VectorRegister4Float Barycentrics[4];
+			Barycentrics[0] = VectorSwizzle(Barycentric, 0, 0, 0, 0);
+			Barycentrics[1] = VectorSwizzle(Barycentric, 1, 1, 1, 1);
+			Barycentrics[2] = VectorSwizzle(Barycentric, 2, 2, 2, 2);
+			Barycentrics[3] = VectorSwizzle(Barycentric, 3, 3, 3, 3);
+
+
+			const VectorRegister4Float  ClosestB1 = VectorMultiplyAdd(Bs[0], Barycentrics[0], ClosestB);
+			const VectorRegister4Float  ClosestB2 = VectorMultiplyAdd(Bs[1], Barycentrics[1], ClosestB1);
+			const VectorRegister4Float  ClosestB3 = VectorMultiplyAdd(Bs[2], Barycentrics[2], ClosestB2);
+			const VectorRegister4Float  ClosestB4 = VectorMultiplyAdd(Bs[3], Barycentrics[3], ClosestB3);
+
+			constexpr VectorRegister4Int TwoInt = MakeVectorRegisterIntConstant(2, 2, 2, 2);
+			constexpr VectorRegister4Int ThreeInt = MakeVectorRegisterIntConstant(3, 3, 3, 3);
+
+			const VectorRegister4Float IsB0 = VectorCast4IntTo4Float(VectorIntCompareEQ(NumVerts.NumVerts, GlobalVectorConstants::IntZero));
+			const VectorRegister4Float IsB1 = VectorCast4IntTo4Float(VectorIntCompareEQ(NumVerts.NumVerts, GlobalVectorConstants::IntOne));
+			const VectorRegister4Float IsB2 = VectorCast4IntTo4Float(VectorIntCompareEQ(NumVerts.NumVerts, TwoInt));
+			const VectorRegister4Float IsB3 = VectorCast4IntTo4Float(VectorIntCompareEQ(NumVerts.NumVerts, ThreeInt));
+
+			ClosestB = VectorSelect(IsB0, ClosestB, ClosestB4);
+			ClosestB = VectorSelect(IsB1, ClosestB1, ClosestB);
+			ClosestB = VectorSelect(IsB2, ClosestB2, ClosestB);
+			ClosestB = VectorSelect(IsB3, ClosestB3, ClosestB);
+
+			const VectorRegister4Float ClosestLocal = VectorNegateMultiplyAdd(OutNormal, MarginBSimd, ClosestB);
+
+			OutPosition = VectorAdd(VectorMultiplyAdd(RayDir, Lambda, StartPoint), ClosestLocal);
+
+		}
+		else if (bComputeMTD)
+		{
+			// If Inflation == 0 we would expect GJKPreDist2 to be 0
+			// However, due to precision we can still end up with GJK failing.
+			// When that happens fall back on EPA
+			VectorRegister4Float InflationGTZero = VectorCompareGT(Inflation, VectorZeroFloat());
+			VectorRegister4Float InGJKPreDist2GTEps2 = VectorCompareGT(GJKPreDist2, Eps2Simd);
+			VectorRegister4Float LimitMaxGTInGJKPreDist2 = VectorCompareGT(LimitMax, GJKPreDist2);
+			VectorRegister4Float IsDone = VectorBitwiseAnd(InflationGTZero, VectorBitwiseAnd(InGJKPreDist2GTEps2, LimitMaxGTInGJKPreDist2));
+			if (VectorMaskBits(IsDone))
+			{
+				OutNormal = Normal;
+				VectorRegister4Float ClosestA = VectorZeroFloat();
+				VectorRegister4Float ClosestB = VectorZeroFloat();
+
+				if (NumIterations)
+				{
+					VectorRegister4Float Barycentrics[4];
+					Barycentrics[0] = VectorSwizzle(Barycentric, 0, 0, 0, 0);
+					Barycentrics[1] = VectorSwizzle(Barycentric, 1, 1, 1, 1);
+					Barycentrics[2] = VectorSwizzle(Barycentric, 2, 2, 2, 2);
+					Barycentrics[3] = VectorSwizzle(Barycentric, 3, 3, 3, 3);
+
+					const int NumVertsInt = NumVerts.NumVertsInts[0];
+					for (int i = 0; i < NumVertsInt; ++i)
+					{
+						ClosestA = VectorMultiplyAdd(As[i], Barycentrics[i], ClosestA);
+						ClosestB = VectorMultiplyAdd(Bs[i], Barycentrics[i], ClosestB);
+					}
+				}
+				else
+				{
+					//didn't even go into gjk loop
+					ClosestA = As[0];
+					ClosestB = Bs[0];
+				}
+
+				const VectorRegister4Float ClosestBInA = VectorAdd(StartPoint, ClosestB);
+				const VectorRegister4Float InGJKPreDist = VectorSqrt(GJKPreDist2);
+				OutNormal = VectorNormalizeAccurate(V);
+
+				VectorRegister4Float Penetration = VectorSubtract(VectorAdd(MarginASimd, MarginBSimd), InGJKPreDist);
+				Penetration = VectorMin(Penetration, LimitMax);
+				Penetration = VectorMax(Penetration, VectorZeroFloat());
+
+				const VectorRegister4Float ClosestLocal = VectorNegateMultiplyAdd(OutNormal, MarginBSimd, ClosestB);
+
+				OutPosition = VectorAdd(VectorMultiplyAdd(OutNormal, Penetration, StartPoint), ClosestLocal);
+				Penetration = VectorNegate(Penetration);
+				VectorStoreFloat1(Penetration, &OutTime);
+			}
+			else
+			{
+				if (NumIterations)
+				{
+
+					//use EPA
+					TArray<TVec3<T>> VertsA;
+					TArray<TVec3<T>> VertsB;
+
+					VertsA.Reserve(8);
+					VertsB.Reserve(8);
+
+					const int32 NumVertsInt = NumVerts.NumVertsInts[0];
+
+					alignas(16) FRealSingle XFloat[4];
+					VectorStoreAligned(X, XFloat);
+					TVec3<T> XVec3 = { XFloat[0], XFloat[1], XFloat[2] };
+
+					for (int i = 0; i < NumVertsInt; ++i)
+					{
+						alignas(16) FRealSingle AsFloat[4];
+						VectorStoreAligned(As[i], AsFloat);
+						TVec3<T> AsVec3 = { AsFloat[0], AsFloat[1], AsFloat[2] };
+						VertsA.Add(AsVec3);
+
+						alignas(16) FRealSingle BsFloat[4];
+						VectorStoreAligned(Bs[i], BsFloat);
+						TVec3<T> BsVec3 = { BsFloat[0], BsFloat[1], BsFloat[2] };
+
+						const TVec3<T> BAtOrigin = BsVec3 + XVec3;
+						VertsB.Add(BAtOrigin);
+					}
+
+
+					auto SupportAFuncFloat = [&A, MarginA](const TVec3<T>& V)
+					{
+						int32 VertexIndex = INDEX_NONE;
+						return A.SupportCore(V, MarginA, nullptr, VertexIndex);
+					};
+
+					const TRotation<double, 3> BToARotationEPA = StartTM.GetRotation();
+					const TRotation<double, 3> AToBRotationEPA = BToARotationEPA.Inverse();
+
+					auto SupportBAtOriginFunc = [&B, MarginB, &StartTM, &AToBRotationEPA](const TVec3<T>& Dir)
+					{
+						int32 VertexIndex = INDEX_NONE;
+						const TVector<T, 3> DirInB = AToBRotationEPA * Dir;
+						const TVector<T, 3> SupportBLocal = B.SupportCore(DirInB, MarginB, nullptr, VertexIndex);
+						return StartTM.TransformPositionNoScale(SupportBLocal);
+					};
+
+
+					T Penetration;
+					TVec3<T> MTD, ClosestA, ClosestBInA, OutNormalVec3, OutPositionVec3;
+					const EEPAResult EPAResult = EPA(VertsA, VertsB, SupportAFuncFloat, SupportBAtOriginFunc, Penetration, MTD, ClosestA, ClosestBInA);
+					if (IsEPASuccess(EPAResult))
+					{
+						OutNormalVec3 = MTD;
+						OutTime = -Penetration - (MarginA + MarginB);
+						OutPositionVec3 = ClosestA;
+						OutPosition = VectorLoadFloat3_W0(&OutPositionVec3);
+						OutNormal = VectorLoadFloat3_W0(&OutNormalVec3);
+					}
+					else
+					{
+						//assume touching hit
+						OutTime = -(MarginA + MarginB);
+						OutNormal = VectorLoadFloat3_W0(&MTD);
+						OutPosition = VectorMultiplyAdd(OutNormal, MarginASimd, As[0]);
+					}
+				}
+				else
+				{
+					//didn't even go into gjk loop, touching hit
+					OutTime = -(MarginA + MarginB);
+					OutNormal = { 0,0,1,0 };
+					OutPosition = VectorMultiplyAdd(OutNormal, MarginASimd, As[0]);
+				}
+			}
+		}
+		else
+		{
+			// Initial overlap without MTD. These properties are not valid, but assigning them anyway so they don't contain NaNs and cause issues in invoking code.
+			OutNormal = { 0,0,1,0 };
+			OutPosition = { 0,0,0,0 };
+		}
+
+		return true;
+	}
+
+#if GJK_VECTORIZED
+	template <typename T, typename TGeometryA, typename TGeometryB>
+	bool GJKRaycast2(const TGeometryA& A, const TGeometryB& B, const TRigidTransform<T, 3>& StartTM, const TVector<T, 3>& RayDir, const T RayLength,
+		T& OutTime, TVector<T, 3>& OutPosition, TVector<T, 3>& OutNormal, const T GivenThicknessA = 0, bool bComputeMTD = false, const TVector<T, 3>& InitialDir = TVector<T, 3>(1, 0, 0), const T GivenThicknessB = 0)
+	{
+		UE::Math::TQuat<T> InRotationDouble = StartTM.GetRotation();
+		alignas(16) FRealSingle InRotationFloat[4] = { static_cast<FRealSingle>(InRotationDouble.X), static_cast<FRealSingle>(InRotationDouble.Y), static_cast<FRealSingle>(InRotationDouble.Z), static_cast<FRealSingle>(InRotationDouble.W) };
+
+		VectorRegister4Float Rotation = VectorLoadAligned(InRotationFloat);
+
+		UE::Math::TVector<T> InTranslationDouble = StartTM.GetTranslation();
+		alignas(16) FRealSingle InTranslationFloat[4] = { static_cast<FRealSingle>(InTranslationDouble.X), static_cast<FRealSingle>(InTranslationDouble.Y), static_cast<FRealSingle>(InTranslationDouble.Z) };
+
+		VectorRegister4Float Translation = VectorLoadAligned(InTranslationFloat);
+
+		// Normalize rotation
+		Rotation = VectorNormalizeSafe(Rotation, GlobalVectorConstants::Float0001);
+		VectorRegister4Float OutPositionSimd, OutNormalSimd, InitialDirSimd, RayDirSimd;
+
+		alignas(16) FRealSingle InitialDirFloat[4] = {static_cast<FRealSingle>(InitialDir[0]), static_cast<FRealSingle>(InitialDir[1]), static_cast<FRealSingle>(InitialDir[2]), 0.0f};
+		alignas(16) FRealSingle RayDirFloat[4] = {static_cast<FRealSingle>(RayDir[0]), static_cast<FRealSingle>(RayDir[1]), static_cast<FRealSingle>(RayDir[2]), 0.0f };
+
+		InitialDirSimd = VectorLoadAligned(InitialDirFloat);
+		RayDirSimd = VectorLoadAligned(RayDirFloat);
+
+		FRealSingle OutTimeFloat;
+		bool result = GJKRaycast2ImplSimd(A, B, Rotation, Translation, RayDirSimd, static_cast<FRealSingle>(RayLength), OutTimeFloat, OutPositionSimd, OutNormalSimd, bComputeMTD, InitialDirSimd, StartTM);
+
+		OutTime = static_cast<double>(OutTimeFloat);
+
+
+		alignas(16) FRealSingle OutFloat[4];
+		VectorStoreAligned(OutNormalSimd, OutFloat);
+		OutNormal.X = OutFloat[0];
+		OutNormal.Y = OutFloat[1];
+		OutNormal.Z = OutFloat[2];
+
+		VectorStoreAligned(OutPositionSimd, OutFloat);
+		OutPosition.X = OutFloat[0];
+		OutPosition.Y = OutFloat[1];
+		OutPosition.Z = OutFloat[2];
+
+		return result;
+	}
+
+#else // GJK_VECTORIZED
+
 	/** Sweeps one geometry against the other
 	 @A The first geometry
 	 @B The second geometry
@@ -1253,6 +1696,8 @@ namespace Chaos
 	{
 		return GJKRaycast2Impl(A, B, StartTM, RayDir, RayLength, OutTime, OutPosition, OutNormal, GivenThicknessA, bComputeMTD, InitialDir, GivenThicknessB);
 	}
+#endif // GJK_VECTORIZED
+
 	/**
 	 * Used by GJKDistance. It must return a vector in the Minkowski sum A - B. In principle this can be the vector of any point
 	 * in A to any point in B, but some choices will cause GJK to minimize faster (e.g., for two spheres, we can easily calculate
diff --git a/Engine/Source/Runtime/Experimental/Chaos/Public/Chaos/ImplicitObjectScaled.h b/Engine/Source/Runtime/Experimental/Chaos/Public/Chaos/ImplicitObjectScaled.h
--- a/Engine/Source/Runtime/Experimental/Chaos/Public/Chaos/ImplicitObjectScaled.h
+++ b/Engine/Source/Runtime/Experimental/Chaos/Public/Chaos/ImplicitObjectScaled.h
@@ -177,6 +177,16 @@ public:
 		return MObject->SupportCore(Direction, InMargin, OutSupportDelta, VertexIndex);
 	}
 
+	FORCEINLINE VectorRegister4Float SupportCoreSimd(const VectorRegister4Float& Direction, const FReal InMargin) const
+	{
+		FVec3 DirectionVec3;
+		VectorStoreFloat3(Direction, &DirectionVec3);
+		int32 VertexIndex = INDEX_NONE;
+		FVec3 SupportVert = MObject->SupportCore(DirectionVec3, InMargin, nullptr, VertexIndex);
+		alignas(16) FRealSingle SupportVertFloat[4] = { static_cast<float>(SupportVert.X), static_cast<float>(SupportVert.Y), static_cast<float>(SupportVert.Z), 0.0f };
+		return VectorLoadAligned(SupportVertFloat);
+	}
+
 	virtual const FAABB3 BoundingBox() const override 
 	{ 
 		return MObject->BoundingBox();
@@ -875,6 +885,16 @@ public:
 		return MObject->SupportCoreScaled(Direction, InMargin, MScale, OutSupportDelta, VertexIndex);
 	}
 
+	FORCEINLINE_DEBUGGABLE VectorRegister4Float SupportCoreSimd(const VectorRegister4Float& Direction, const FReal InMargin) const
+	{
+		FVec3 DirectionVec3;
+		VectorStoreFloat3(Direction, &DirectionVec3);
+		int32 VertexIndex = INDEX_NONE;
+		FVec3 SupportVert = MObject->SupportCoreScaled(DirectionVec3, InMargin, MScale, nullptr, VertexIndex);
+		TVector<float, 3> SupportVertFloat(static_cast<float>(SupportVert.X), static_cast<float>(SupportVert.Y), static_cast<float>(SupportVert.Z));
+		return VectorLoadFloat3_W0(&SupportVertFloat);
+	}
+
 	void SetScale(const FVec3& Scale)
 	{
 		constexpr FReal MinMagnitude = 1e-6f;
diff --git a/Engine/Source/Runtime/Experimental/Chaos/Public/Chaos/SimplexVectorized.h b/Engine/Source/Runtime/Experimental/Chaos/Public/Chaos/SimplexVectorized.h
new file mode 100644
--- /dev/null
+++ b/Engine/Source/Runtime/Experimental/Chaos/Public/Chaos/SimplexVectorized.h
@@ -0,0 +1,377 @@
+// Copyright Epic Games, Inc. All Rights Reserved.
+#pragma once
+
+#if PLATFORM_ENABLE_VECTORINTRINSICS_NEON || PLATFORM_ENABLE_VECTORINTRINSICS
+#include "Math/VectorRegister.h"
+#include "Chaos/VectorUtility.h"
+
+namespace Chaos
+{
+
+	FORCEINLINE VectorRegister4Float VectorLineSimplexFindOrigin(VectorRegister4Float* Simplex, VectorRegister4Int& NumVerts, VectorRegister4Float& OutBarycentric, VectorRegister4Float* A, VectorRegister4Float* B)
+	{
+		const VectorRegister4Float& X0 = Simplex[0];
+		const VectorRegister4Float& X1 = Simplex[1];
+		const VectorRegister4Float X0ToX1 = VectorSubtract(X1, X0);
+
+		//Closest Point = (-X0 dot X1-X0) / ||(X1-X0)||^2 * (X1-X0)
+
+		const VectorRegister4Float X0ToOrigin = VectorNegate(X0);
+		const VectorRegister4Float Dot = VectorDot3(X0ToOrigin, X0ToX1);
+
+		const VectorRegister4Float IsX0 = VectorCompareGE(VectorZeroFloat(), Dot);
+
+		constexpr VectorRegister4Float OutBarycentricIfX0OrX1 = GlobalVectorConstants::Float1000;
+		const VectorRegister4Float X0ToX1Squared = VectorDot3(X0ToX1, X0ToX1);
+		const VectorRegister4Float DotBigger = VectorCompareGE(Dot, X0ToX1Squared);
+
+		const VectorRegister4Float MinLimt = { std::numeric_limits<FRealSingle>::min() , std::numeric_limits<FRealSingle>::min(), std::numeric_limits<FRealSingle>::min(), std::numeric_limits<FRealSingle>::min() };
+		const VectorRegister4Float X0ToX1SquaredSmall = VectorCompareGE(MinLimt, X0ToX1Squared);
+		const VectorRegister4Float IsX1 = VectorBitwiseOr(DotBigger, X0ToX1SquaredSmall);
+
+		Simplex[0] = VectorSelect(IsX1, Simplex[1], Simplex[0]);
+		A[0] = VectorSelect(IsX1, A[1], A[0]);
+		B[0] = VectorSelect(IsX1, B[1], B[0]);
+
+		VectorRegister4Float Ratio = VectorDivide(Dot, X0ToX1Squared);
+
+		Ratio = VectorMax(Ratio, VectorZeroFloat());
+		Ratio = VectorMin(Ratio, GlobalVectorConstants::FloatOne);
+
+		VectorRegister4Float Closest = VectorMultiplyAdd(Ratio, X0ToX1, X0);
+
+		const VectorRegister4Float OneMinusRatio = VectorSubtract(GlobalVectorConstants::FloatOne, Ratio);
+		const VectorRegister4Float OutBarycentricOtherwise = VectorUnpackLo(OneMinusRatio, Ratio);
+
+		Closest = VectorSelect(IsX0, X0, VectorSelect(IsX1, X1, Closest));
+
+		const VectorRegister4Float IsX0OrX1 = VectorBitwiseOr(IsX0, IsX1);
+		const VectorRegister4Int IsX0OrX1Int = VectorCast4FloatTo4Int(IsX0OrX1);
+		NumVerts = VectorIntSelect(IsX0OrX1Int, GlobalVectorConstants::IntOne, NumVerts);
+
+		OutBarycentric = VectorSelect(IsX0OrX1, OutBarycentricIfX0OrX1, OutBarycentricOtherwise);
+
+		return Closest;
+	}
+
+	union alignas(16) NumVertsUnion 
+	{ 
+		VectorRegister4Int NumVerts;
+		int32 NumVertsInts[4]; 
+	};
+
+	// Based on an algorithm in Real Time Collision Detection - Ericson (very close to that)
+	// Using the same variable name conventions for easy reference
+	FORCEINLINE VectorRegister4Float  TriangleSimplexFindOriginFast(VectorRegister4Float* Simplex, VectorRegister4Int& NumVerts, VectorRegister4Float& OutBarycentric, VectorRegister4Float* As, VectorRegister4Float* Bs)
+	{
+		const VectorRegister4Float& A = Simplex[0];
+		const VectorRegister4Float& B = Simplex[1];
+		const VectorRegister4Float& C = Simplex[2];
+
+		// Vertex region A
+		const VectorRegister4Float AB = VectorSubtract(B, A);
+		const VectorRegister4Float AC = VectorSubtract(C, A);
+		const VectorRegister4Float AO = VectorNegate(A);
+
+		const VectorRegister4Float d1 = VectorDot3(AB, AO);
+		const VectorRegister4Float d2 = VectorDot3(AC, AO);
+
+		const VectorRegister4Float IsD1SEZero = VectorCompareGE(VectorZeroFloat(), d1);
+		const VectorRegister4Float IsD2SEZero = VectorCompareGE(VectorZeroFloat(), d2);
+		const VectorRegister4Float IsA = VectorBitwiseAnd(IsD1SEZero, IsD2SEZero);
+
+		if (VectorMaskBits(IsA))
+		{
+			NumVerts = GlobalVectorConstants::IntOne;
+			OutBarycentric = GlobalVectorConstants::Float1000;
+			return A;
+		}
+
+		//Vertex region B
+		const VectorRegister4Float BO = VectorNegate(B);
+		const VectorRegister4Float d3 = VectorDot3(AB, BO);
+		const VectorRegister4Float d4 = VectorDot3(AC, BO);
+
+		const VectorRegister4Float IsD3GEZero = VectorCompareGE(d3, VectorZeroFloat());
+		const VectorRegister4Float IsD3GED4 = VectorCompareGE(d3, d4);
+		const VectorRegister4Float IsB = VectorBitwiseAnd(IsD3GEZero, IsD3GED4);
+
+		if (VectorMaskBits(IsB))
+		{
+			NumVerts = GlobalVectorConstants::IntOne;
+			OutBarycentric = GlobalVectorConstants::Float1000;
+			Simplex[0] = B;
+			As[0] = As[1];
+			Bs[0] = Bs[1];
+			return B;
+		}
+
+		// Edge AB
+		const VectorRegister4Float d1d4 = VectorMultiply(d1, d4);
+		const VectorRegister4Float vc = VectorNegateMultiplyAdd(d3, d2, d1d4);
+		const VectorRegister4Float NormalizationDenominatorAB = VectorSubtract(d1, d3);
+
+		const VectorRegister4Float IsZeroGEvc = VectorCompareGE(VectorZeroFloat(), vc);
+		const VectorRegister4Float IsD1GEZero = VectorCompareGE(d1, VectorZeroFloat());
+		const VectorRegister4Float IsZeroGED3= VectorCompareGE(VectorZeroFloat(), d3);
+		const VectorRegister4Float IsNDABGTZero = VectorCompareGT(NormalizationDenominatorAB, VectorZeroFloat());
+		const VectorRegister4Float IsAB = VectorBitwiseAnd(VectorBitwiseAnd(IsZeroGEvc, IsD1GEZero), VectorBitwiseAnd(IsZeroGED3, IsNDABGTZero));
+
+		if (VectorMaskBits(IsAB))
+		{
+			constexpr VectorRegister4Int two = MakeVectorRegisterIntConstant(2, 2, 2, 2);
+			NumVerts = two;
+
+			const VectorRegister4Float v = VectorDivide(d1, NormalizationDenominatorAB);
+			const VectorRegister4Float OneMinusV = VectorSubtract(GlobalVectorConstants::FloatOne, v);
+			// b0	a1	a2	a3
+			OutBarycentric = VectorUnpackLo(OneMinusV, v);
+			return VectorMultiplyAdd(v, AB, A);
+		}
+
+		// Vertex C
+		const VectorRegister4Float CO = VectorNegate(C);
+		const VectorRegister4Float d5 = VectorDot3(AB, CO);
+		const VectorRegister4Float d6 = VectorDot3(AC, CO);
+		const VectorRegister4Float IsD6GEZero = VectorCompareGE(d6, VectorZeroFloat());
+		const VectorRegister4Float IsD6GED5 = VectorCompareGE(d6, d5);
+		const VectorRegister4Float IsC = VectorBitwiseAnd(IsD6GEZero, IsD6GED5);
+
+		if (VectorMaskBits(IsC))
+		{
+			NumVerts = GlobalVectorConstants::IntOne;
+			OutBarycentric = GlobalVectorConstants::Float1000;
+			Simplex[0] = C;
+			As[0] = As[2];
+			Bs[0] = Bs[2];
+			return C;
+		}
+
+		// Edge AC
+		const VectorRegister4Float d5d2 = VectorMultiply(d5, d2);
+		const VectorRegister4Float vb = VectorNegateMultiplyAdd(d1, d6, d5d2);
+		const VectorRegister4Float NormalizationDenominatorAC = VectorSubtract(d2, d6);
+
+		const VectorRegister4Float IsZeroGEvb = VectorCompareGE(VectorZeroFloat(), vb);
+		const VectorRegister4Float IsD2GEZero = VectorCompareGE(d2, VectorZeroFloat());
+		const VectorRegister4Float IsZeroGED6 = VectorCompareGE(VectorZeroFloat(), d6);
+		const VectorRegister4Float IsNDACGTZero = VectorCompareGT(NormalizationDenominatorAC, VectorZeroFloat());
+		const VectorRegister4Float IsAC = VectorBitwiseAnd(VectorBitwiseAnd(IsZeroGEvb, IsD2GEZero), VectorBitwiseAnd(IsZeroGED6, IsNDACGTZero));
+
+		if (VectorMaskBits(IsAC))
+		{
+			const VectorRegister4Float w = VectorDivide(d2, NormalizationDenominatorAC);
+			constexpr VectorRegister4Int two = MakeVectorRegisterIntConstant(2, 2, 2, 2);
+			NumVerts = two;
+			const VectorRegister4Float OneMinusW = VectorSubtract(GlobalVectorConstants::FloatOne, w);
+			// b0	a1	a2	a3
+			OutBarycentric = VectorUnpackLo(OneMinusW, w);
+			Simplex[1] = C;
+			As[1] = As[2];
+			Bs[1] = Bs[2];
+			return VectorMultiplyAdd(w, AC, A);
+		}
+
+		// Edge BC
+		const VectorRegister4Float d3d6 = VectorMultiply(d3, d6);
+		const VectorRegister4Float va = VectorNegateMultiplyAdd(d5, d4, d3d6);
+		const VectorRegister4Float d4MinusD3 = VectorSubtract(d4, d3);
+		const VectorRegister4Float d5MinusD6 = VectorSubtract(d5, d6);
+		const VectorRegister4Float NormalizationDenominatorBC = VectorAdd(d4MinusD3, d5MinusD6);
+
+		const VectorRegister4Float IsZeroGEva = VectorCompareGE(VectorZeroFloat(), va);
+		const VectorRegister4Float IsD4MinusD3GEZero = VectorCompareGE(d4MinusD3, VectorZeroFloat());
+		const VectorRegister4Float IsD5MinusD6GEZero = VectorCompareGE(d5MinusD6, VectorZeroFloat());
+		const VectorRegister4Float IsNDBCGTZero = VectorCompareGT(NormalizationDenominatorBC, VectorZeroFloat());
+		const VectorRegister4Float IsBC = VectorBitwiseAnd(VectorBitwiseAnd(IsZeroGEva, IsD4MinusD3GEZero), VectorBitwiseAnd(IsD5MinusD6GEZero, IsNDBCGTZero));
+
+		if (VectorMaskBits(IsBC))
+		{
+			const VectorRegister4Float w = VectorDivide(d4MinusD3, NormalizationDenominatorBC);
+			constexpr VectorRegister4Int two = MakeVectorRegisterIntConstant(2, 2, 2, 2);
+			NumVerts = two;
+			const VectorRegister4Float OneMinusW = VectorSubtract(GlobalVectorConstants::FloatOne, w);
+			// b0	a1	a2	a3
+			OutBarycentric = VectorUnpackLo(OneMinusW, w);
+			const VectorRegister4Float CMinusB = VectorSubtract(C, B);
+			const VectorRegister4Float Result = VectorMultiplyAdd(w, CMinusB, B);
+			Simplex[0] = B;
+			Simplex[1] = C;
+			As[0] = As[1];
+			Bs[0] = Bs[1];
+			As[1] = As[2];
+			Bs[1] = Bs[2];
+			return Result;
+		}
+
+		// Inside triangle
+		const VectorRegister4Float denom = VectorDivide(VectorOne(), VectorAdd(va, VectorAdd(vb, vc)));
+		VectorRegister4Float v = VectorMultiply(vb, denom);
+		VectorRegister4Float w = VectorMultiply(vc, denom);
+		constexpr VectorRegister4Int three = MakeVectorRegisterIntConstant(3, 3, 3, 3);
+		NumVerts = three;
+
+		const VectorRegister4Float OneMinusVMinusW = VectorSubtract(VectorSubtract(GlobalVectorConstants::FloatOne, v), w);
+		// b0	a1	a2	a3
+		const VectorRegister4Float OneMinusVMinusW_W = VectorUnpackLo(OneMinusVMinusW, w);
+		// a0	b0	a1	b1
+		OutBarycentric = VectorUnpackLo(OneMinusVMinusW_W, v);
+
+		// We know that we are inside the triangle so we can use the projected point we calculated above. 
+		// The closest point can also be derived from the barycentric coordinates, but it will contain 
+		// numerical error from the determinant calculation  and can cause GJK to terminate with a poor solution.
+		// (E.g., this caused jittering when walking on box with dimensions of 100000cm or more).
+		// This fix the unit test TestSmallCapsuleLargeBoxGJKRaycast_Vertical
+		// Previously was return VectorMultiplyAdd(AC, w, VectorMultiplyAdd(AB, v, A));
+		const VectorRegister4Float TriNormal = VectorCross(AB, AC);
+		const VectorRegister4Float TriNormal2 = VectorDot3(TriNormal, TriNormal);
+		const VectorRegister4Float TriNormalOverSize2 = VectorDivide(TriNormal, TriNormal2);
+		const VectorRegister4Float SignedDistance = VectorDot3(A, TriNormalOverSize2);
+		return VectorMultiply(TriNormal, SignedDistance);
+	}
+
+	FORCEINLINE bool VectorSignMatch(VectorRegister4Float A, VectorRegister4Float B)
+	{
+		VectorRegister4Float OneIsZero = VectorMultiply(A, B);
+		OneIsZero = VectorCompareEQ(OneIsZero, VectorZeroFloat());
+
+		const bool IsZero = static_cast<bool>(VectorMaskBits(OneIsZero));
+		const int32 MaskA = VectorMaskBits(A);
+		const int32 MaskB = VectorMaskBits(B);
+		return (MaskA == MaskB) && !IsZero;
+	}
+
+
+	FORCEINLINE VectorRegister4Float VectorTetrahedronSimplexFindOrigin(VectorRegister4Float* Simplex, VectorRegister4Int& NumVerts, VectorRegister4Float& OutBarycentric, VectorRegister4Float* A, VectorRegister4Float* B)
+	{
+		const VectorRegister4Float& X0 = Simplex[0];
+		const VectorRegister4Float& X1 = Simplex[1];
+		const VectorRegister4Float& X2 = Simplex[2];
+		const VectorRegister4Float& X3 = Simplex[3];
+
+		//Use signed volumes to determine if origin is inside or outside
+		/*
+			M = [X0x X1x X2x X3x;
+				 X0y X1y X2y X3y;
+				 X0z X1z X2z X3z;
+				 1   1   1   1]
+		*/
+
+		VectorRegister4Float Cofactors[4];
+		Cofactors[0] = VectorNegate(VectorDot3(X1, VectorCross(X2, X3)));
+		Cofactors[1] = VectorDot3(X0, VectorCross(X2, X3));
+		Cofactors[2] = VectorNegate(VectorDot3(X0, VectorCross(X1, X3)));
+		Cofactors[3] = VectorDot3(X0, VectorCross(X1, X2));
+		VectorRegister4Float DetM = VectorAdd(VectorAdd(Cofactors[0], Cofactors[1]), VectorAdd(Cofactors[2], Cofactors[3]));
+
+		bool bSignMatch[4];
+		// constexpr int32 SubIdxs[4][3] = { {1,2,3}, {0,2,3}, {0,1,3}, {0,1,2} };
+		constexpr VectorRegister4Int ThreeInt = MakeVectorRegisterIntConstant(3, 3, 3, 3);
+		VectorRegister4Int SubNumVerts[4] = { ThreeInt, ThreeInt, ThreeInt, ThreeInt };
+		VectorRegister4Float SubSimplices[4][3] = { {Simplex[1], Simplex[2], Simplex[3]}, {Simplex[0], Simplex[2], Simplex[3]}, {Simplex[0], Simplex[1], Simplex[3]}, {Simplex[0], Simplex[1], Simplex[2]} };
+		VectorRegister4Float SubAs[4][3] = { {A[1], A[2], A[3]}, {A[0], A[2], A[3]}, {A[0], A[1], A[3]}, {A[0], A[1], A[2]} };
+		VectorRegister4Float SubBs[4][3] = { {B[1], B[2], B[3]}, {B[0], B[2], B[3]}, {B[0], B[1], B[3]}, {B[0], B[1], B[2]} };
+		VectorRegister4Float ClosestPointSub[4];
+		VectorRegister4Float SubBarycentric[4];
+		constexpr VectorRegister4Int IndexNone = MakeVectorRegisterIntConstant(INDEX_NONE, INDEX_NONE, INDEX_NONE, INDEX_NONE);
+		VectorRegister4Int ClosestTriangleIdx = IndexNone;
+		VectorRegister4Float MinTriangleDist2 = VectorZeroFloat();
+
+
+		constexpr VectorRegister4Int IdxSimd[4] = { GlobalVectorConstants::IntZero,
+													GlobalVectorConstants::IntOne,
+													MakeVectorRegisterIntConstant(2, 2, 2, 2),
+													MakeVectorRegisterIntConstant(3, 3, 3, 3)
+		};
+
+		bool bInside = true;
+		for (int Idx = 0; Idx < 4; ++Idx)
+		{
+			bSignMatch[Idx] = VectorSignMatch(DetM, Cofactors[Idx]);
+			if (!bSignMatch[Idx])
+			{
+				bInside = false;
+				ClosestPointSub[Idx] = TriangleSimplexFindOriginFast(SubSimplices[Idx], SubNumVerts[Idx], SubBarycentric[Idx], SubAs[Idx], SubBs[Idx]);
+
+				const VectorRegister4Float Dist2 = VectorDot3(ClosestPointSub[Idx], ClosestPointSub[Idx]);
+
+				const VectorRegister4Float MinGTDist = VectorCompareGT(MinTriangleDist2, Dist2);
+				const VectorRegister4Int ClsoestEq = VectorIntCompareEQ(ClosestTriangleIdx, IndexNone);
+
+				const VectorRegister4Float FindClosest = VectorBitwiseOr(MinGTDist, VectorCast4IntTo4Float(ClsoestEq));
+				const VectorRegister4Int FindClosestInt = VectorCast4FloatTo4Int(FindClosest);
+
+				MinTriangleDist2 = VectorSelect(FindClosest, Dist2, MinTriangleDist2);
+				ClosestTriangleIdx = VectorIntSelect(FindClosestInt, IdxSimd[Idx], ClosestTriangleIdx);
+			}
+		}
+
+		if (bInside)
+		{
+			VectorRegister4Float OutBarycentricVectors[4];
+			const VectorRegister4Float  InvDetM = VectorDivide(GlobalVectorConstants::FloatOne, DetM);
+			OutBarycentricVectors[0] = VectorDivide(Cofactors[0], InvDetM);
+			OutBarycentricVectors[1] = VectorDivide(Cofactors[1], InvDetM);
+			OutBarycentricVectors[2] = VectorDivide(Cofactors[2], InvDetM);
+			OutBarycentricVectors[3] = VectorDivide(Cofactors[3], InvDetM);
+
+			// a0	b0	a1	b1
+			const VectorRegister4Float OutBarycentric0101 = VectorUnpackLo(OutBarycentricVectors[0], OutBarycentricVectors[1]);
+			const VectorRegister4Float OutBarycentric2323 = VectorUnpackLo(OutBarycentricVectors[2], OutBarycentricVectors[3]);
+			// a0	a1	b0	b1
+			OutBarycentric = VectorMoveLh(OutBarycentric0101, OutBarycentric2323);
+
+			return VectorZeroFloat();
+		}
+
+		alignas(16) int32 ClosestTriangleIdxInts[4];
+		VectorIntStoreAligned(ClosestTriangleIdx, ClosestTriangleIdxInts);
+		int32 ClosestTriangleIdxInt = ClosestTriangleIdxInts[0];
+		NumVerts = SubNumVerts[ClosestTriangleIdxInt];
+		OutBarycentric = SubBarycentric[ClosestTriangleIdxInt];
+
+		for (int i = 0; i < 3; i++)
+		{
+			Simplex[i] = SubSimplices[ClosestTriangleIdxInt][i];
+			A[i] = SubAs[ClosestTriangleIdxInt][i];
+			B[i] = SubBs[ClosestTriangleIdxInt][i];
+		}
+
+		return ClosestPointSub[ClosestTriangleIdxInt];
+	}
+
+
+	FORCEINLINE VectorRegister4Float VectorSimplexFindClosestToOrigin(VectorRegister4Float* Simplex, NumVertsUnion& NumVerts, VectorRegister4Float& OutBarycentric, VectorRegister4Float* A, VectorRegister4Float* B)
+	{
+		VectorRegister4Float ClosestPoint;
+		switch (NumVerts.NumVertsInts[0])
+		{
+		case 1:
+			OutBarycentric = GlobalVectorConstants::Float1000;
+			ClosestPoint = Simplex[0]; 
+			break;
+		case 2:
+		{
+			ClosestPoint = VectorLineSimplexFindOrigin(Simplex, NumVerts.NumVerts, OutBarycentric, A, B);
+			break;
+		}
+		case 3:
+		{
+			ClosestPoint = TriangleSimplexFindOriginFast(Simplex, NumVerts.NumVerts, OutBarycentric, A, B);
+			break;
+		}
+		case 4:
+		{
+			ClosestPoint = VectorTetrahedronSimplexFindOrigin(Simplex, NumVerts.NumVerts, OutBarycentric, A, B);
+			break;
+		}
+		default:
+			ensure(false);
+			ClosestPoint = VectorZeroFloat();
+		}
+
+		return ClosestPoint;
+	}
+}
+
+#endif
\ No newline at end of file
diff --git a/Engine/Source/Runtime/Experimental/Chaos/Public/Chaos/Sphere.h b/Engine/Source/Runtime/Experimental/Chaos/Public/Chaos/Sphere.h
--- a/Engine/Source/Runtime/Experimental/Chaos/Public/Chaos/Sphere.h
+++ b/Engine/Source/Runtime/Experimental/Chaos/Public/Chaos/Sphere.h
@@ -7,6 +7,8 @@
 #include "Chaos/ImplicitObject.h"
 #include "ChaosArchive.h"
 
+#include "Math/VectorRegister.h"
+
 #include "UObject/ReleaseObjectVersion.h"
 
 namespace Chaos
@@ -209,6 +211,11 @@ namespace Chaos
 			return Center;
 		}
 
+		FORCEINLINE VectorRegister4Float SupportCoreSimd(const VectorRegister4Float& Direction, const FReal InMargin) const
+		{
+			alignas(16) FRealSingle CenterFloat[4] = { static_cast<FRealSingle>(Center[0]), static_cast<FRealSingle>(Center[1]), static_cast<FRealSingle>(Center[2]), 0.0f };
+			return VectorLoadAligned(CenterFloat);
+		}
 		FORCEINLINE TVector<T, d> SupportCoreScaled(const TVector<T, d>& Direction, const FReal InMargin, const TVector<T, d>& Scale, FReal* OutSupportDelta, int32& VertexIndex) const
 		{
 			VertexIndex = 0;
diff --git a/Engine/Source/Runtime/Experimental/Chaos/Public/Chaos/Triangle.h b/Engine/Source/Runtime/Experimental/Chaos/Public/Chaos/Triangle.h
--- a/Engine/Source/Runtime/Experimental/Chaos/Public/Chaos/Triangle.h
+++ b/Engine/Source/Runtime/Experimental/Chaos/Public/Chaos/Triangle.h
@@ -8,6 +8,8 @@
 #include "Plane.h"
 #include "Chaos/Plane.h"
 
+#include "Math/UnrealMathVectorConstants.h"
+
 namespace Chaos
 {
 	class FTriangle
@@ -21,7 +23,11 @@ namespace Chaos
 			: A(InA)
 			, B(InB)
 			, C(InC)
-		{}
+		{
+			VectorA = { static_cast<FRealSingle>(A[0]), static_cast<FRealSingle>(A[1]), static_cast<FRealSingle>(A[2]), static_cast<FRealSingle>(0.0f) };
+			VectorB = { static_cast<FRealSingle>(B[0]), static_cast<FRealSingle>(B[1]), static_cast<FRealSingle>(B[2]), static_cast<FRealSingle>(0.0f) };
+			VectorC = { static_cast<FRealSingle>(C[0]), static_cast<FRealSingle>(C[1]), static_cast<FRealSingle>(C[2]), static_cast<FRealSingle>(0.0f) };
+		}
 
 		FVec3& operator[](uint32 InIndex)
 		{
@@ -225,6 +231,27 @@ namespace Chaos
 			return C;
 		}
 
+		FORCEINLINE VectorRegister4Float SupportCoreSimd(const VectorRegister4Float& Direction, const FReal InMargin) const
+		{
+			// Note: assumes margin == 0
+			VectorRegister4Float DotA = VectorDot4(VectorA, Direction);
+			VectorRegister4Float DotB = VectorDot4(VectorB, Direction);
+			VectorRegister4Float DotC = VectorDot4(VectorC, Direction);
+
+			VectorRegister4Float DotAGEDotB = VectorCompareGE(DotA, DotB);
+			VectorRegister4Float DotAGEDotC = VectorCompareGE(DotA, DotC);
+			VectorRegister4Float DotAG = VectorBitwiseAnd(DotAGEDotB, DotAGEDotC);
+
+			VectorRegister4Float DotBGEDotA = VectorCompareGE(DotB, DotA);
+			VectorRegister4Float DotBGEDotC = VectorCompareGE(DotB, DotC);
+			VectorRegister4Float DotBG = VectorBitwiseAnd(DotBGEDotA, DotBGEDotC);
+
+			VectorRegister4Float Result = VectorSelect(DotAG, VectorA, VectorC);
+			Result = VectorSelect(DotBG, VectorB, Result);
+
+			return Result;
+		}
+
 		FORCEINLINE FVec3 SupportCoreScaled(const FVec3& Direction, FReal InMargin, const FVec3& Scale, FReal* OutSupportDelta, int32& VertexIndex) const
 		{
 			// Note: ignores InMargin, assumed 0 (triangles cannot have a margin as they are zero thickness)
@@ -264,6 +291,10 @@ namespace Chaos
 		FVec3 A;
 		FVec3 B;
 		FVec3 C;
+		VectorRegister4Float VectorA;
+		VectorRegister4Float VectorB;
+		VectorRegister4Float VectorC;
+
 	};
 
 	inline FChaosArchive& operator<<(FChaosArchive& Ar, FTriangle& Value)
diff --git a/Engine/Source/Runtime/Experimental/Chaos/Public/Chaos/VectorUtility.h b/Engine/Source/Runtime/Experimental/Chaos/Public/Chaos/VectorUtility.h
new file mode 100644
--- /dev/null
+++ b/Engine/Source/Runtime/Experimental/Chaos/Public/Chaos/VectorUtility.h
@@ -0,0 +1,120 @@
+// Copyright Epic Games, Inc. All Rights Reserved.
+#pragma once
+
+
+/**
+ * Cast VectorRegister4Int in VectorRegister4Float
+ *
+ * @param V	vector
+ * @return		VectorRegister4Float( B.x, A.y, A.z, A.w)
+ */
+#if !defined(_MSC_VER)
+FORCEINLINE VectorRegister4Float VectorCast4IntTo4Float(const VectorRegister4Int& V)
+{
+	return VectorRegister4Float(V);
+}
+#else
+FORCEINLINE VectorRegister4Float VectorCast4IntTo4Float(const VectorRegister4Int& V)
+{
+	return _mm_castsi128_ps(V);
+}
+
+#endif
+
+/**
+ * Cast VectorRegister4Float in VectorRegister4Int
+ *
+ * @param V	vector
+ * @return		VectorCast4FloatTo4Int( B.x, A.y, A.z, A.w)
+ */
+#if !defined(_MSC_VER)
+FORCEINLINE VectorRegister4Int VectorCast4FloatTo4Int(const VectorRegister4Float& V)
+{
+	return VectorRegister4Int(V);
+}
+#else
+FORCEINLINE VectorRegister4Int VectorCast4FloatTo4Int(const VectorRegister4Float& V)
+{
+	return _mm_castps_si128(V);
+}
+#endif
+
+
+#if PLATFORM_ENABLE_VECTORINTRINSICS_NEON
+
+/**
+ * Selects and interleaves the lower two SP FP values from A and B.
+ *
+ * @param A	1st vector
+ * @param B	2nd vector
+ * @return		VectorRegister4Float( A.x, B.x, A.y, B.y)
+ */
+FORCEINLINE VectorRegister4Float VectorUnpackLo(const VectorRegister4Float& A, const VectorRegister4Float& B)
+{
+	return vzip1q_f32(A, B);
+}
+
+/**
+ * Moves the lower 2 SP FP values of b to the upper 2 SP FP values of the result. The lower 2 SP FP values of a are passed through to the result.
+ *
+ * @param A	1st vector
+ * @param B	2nd vector
+ * @return		VectorRegister4Float( A.x, A.y, B.x, B.y)
+  */
+FORCEINLINE VectorRegister4Float VectorMoveLh(const VectorRegister4Float& A, const VectorRegister4Float& B)
+{
+	return vzip1q_f64(A, B);
+}
+
+/**
+ * Return square root.
+ *
+ * @param A	1st vector
+  * @return		VectorRegister4Float( sqrt(A.x), sqrt(A.y), sqrt(A.z), sqrt(A.w))
+  */
+FORCEINLINE VectorRegister4Float VectorSqrt(const VectorRegister4Float& A)
+{
+	return vsqrtq_f32(A);
+}
+
+
+#else // PLATFORM_ENABLE_VECTORINTRINSICS_NEON
+
+/**
+ * Selects and interleaves the lower two SP FP values from A and B.
+ *
+ * @param A	1st vector
+ * @param B	2nd vector
+ * @return		VectorRegister4Float( A.x, B.x, A.y, B.y)
+ */
+FORCEINLINE VectorRegister4Float VectorUnpackLo(const VectorRegister4Float& A, const VectorRegister4Float& B)
+{
+	return _mm_unpacklo_ps(A, B);
+}
+
+/**
+ * Moves the lower 2 SP FP values of b to the upper 2 SP FP values of the result. The lower 2 SP FP values of a are passed through to the result.
+ *
+ * @param A	1st vector
+ * @param B	2nd vector
+ * @return		VectorRegister4Float( A.x, A.y, B.x, B.y)
+  */
+FORCEINLINE VectorRegister4Float VectorMoveLh(const VectorRegister4Float& A, const VectorRegister4Float& B)
+{
+	return _mm_movelh_ps(A, B);
+}
+
+/**
+ * Return square root.
+ *
+ * @param A	1st vector
+  * @return		VectorRegister4Float( sqrt(A.x), sqrt(A.y), sqrt(A.z), sqrt(A.w))
+  */
+FORCEINLINE VectorRegister4Float VectorSqrt(const VectorRegister4Float& A)
+{
+	return _mm_sqrt_ps(A);
+}
+
+#endif // PLATFORM_ENABLE_VECTORINTRINSICS_NEON
+
+
