Chaos - Support for zero margins (still disabled in settings)
- Lower tolerances for GJK and EPA
- Various zero-margin related optimizations

Chaos
- support for Target Separation in collision solver
- VectorRegister use for CollisionConstraint Gather(disbaled)

Chaos - fix TargetPhi
- fix headless chaos

Chaos - fix island group count when using "onethread" or when physics parallelism is disabled
- fix constraint sets for island groups when not using colours

Chaos
- add smooth edge collision flag
- remove unused collision flag
- clean up collision flag setting API

Fix initialization of new FBodyInstance properties

Chaos - prune invalid edge collisions in triangle meshes and heighfields
- add option to prune invalid edge collisions between different shapes (disabled by default)

Chaos - fix edge pruning

diff --git a/Engine/Source/Runtime/Experimental/Chaos/Private/Chaos/Collision/SolverCollisionContainer.cpp b/Engine/Source/Runtime/Experimental/Chaos/Private/Chaos/Collision/SolverCollisionContainer.cpp
--- a/Engine/Source/Runtime/Experimental/Chaos/Private/Chaos/Collision/SolverCollisionContainer.cpp
+++ b/Engine/Source/Runtime/Experimental/Chaos/Private/Chaos/Collision/SolverCollisionContainer.cpp
@@ -258,14 +258,14 @@ namespace Chaos
 				}
 				else
 				{
-					const VectorRegister BodyV0 = VectorLoadFloat3(&Body0->V());
-					const VectorRegister BodyW0 = VectorLoadFloat3(&Body0->W());
-					const VectorRegister BodyV1 = VectorLoadFloat3(&Body1->V());
-					const VectorRegister BodyW1 = VectorLoadFloat3(&Body1->W());
-					const VectorRegister ContactVel0 = VectorAdd(BodyV0, VectorCross(BodyW0, RelativeContactPosition0));
-					const VectorRegister ContactVel1 = VectorAdd(BodyV1, VectorCross(BodyW1, RelativeContactPosition1));
-					const SolverVectorRegister ContactVel = MakeVectorRegisterFloatFromDouble(VectorSubtract(ContactVel0, ContactVel1));
-					WorldFrictionDelta = VectorMultiply(ContactVel, VectorSetFloat1(Dt));
+					//const VectorRegister BodyV0 = VectorLoadFloat3(&Body0->V());
+					//const VectorRegister BodyW0 = VectorLoadFloat3(&Body0->W());
+					//const VectorRegister BodyV1 = VectorLoadFloat3(&Body1->V());
+					//const VectorRegister BodyW1 = VectorLoadFloat3(&Body1->W());
+					//const VectorRegister ContactVel0 = VectorAdd(BodyV0, VectorCross(BodyW0, RelativeContactPosition0));
+					//const VectorRegister ContactVel1 = VectorAdd(BodyV1, VectorCross(BodyW1, RelativeContactPosition1));
+					//const SolverVectorRegister ContactVel = MakeVectorRegisterFloatFromDouble(VectorSubtract(ContactVel0, ContactVel1));
+					//WorldFrictionDelta = VectorMultiply(ContactVel, VectorSetFloat1(Dt));
 
 					ManifoldPoint.ShapeAnchorPoints[0] = ManifoldPoint.ContactPoint.ShapeContactPoints[0];
 					ManifoldPoint.ShapeAnchorPoints[1] = ManifoldPoint.ContactPoint.ShapeContactPoints[1];
diff --git a/Engine/Source/Runtime/Experimental/Chaos/Public/Chaos/GJK.h b/Engine/Source/Runtime/Experimental/Chaos/Public/Chaos/GJK.h
--- a/Engine/Source/Runtime/Experimental/Chaos/Public/Chaos/GJK.h
+++ b/Engine/Source/Runtime/Experimental/Chaos/Public/Chaos/GJK.h
@@ -373,7 +373,7 @@ namespace Chaos
 
 			// If we are not too close, update the separating vector and keep iterating
 			// If we are too close we drop through to EPA, but if EPA determines the shapes are actually separated 
-			// after all, we will end up using the normal from the previous loop
+			// after all, we will wend up using the normal from the previous loop
 			if (!bIsContact)
 			{
 				V /= NewDistance;
@@ -427,6 +427,7 @@ namespace Chaos
 				// and penetration but we keep the position generated by GJK.
 				Normal = MTD;
 				Distance = -Penetration;
+				//UE_LOG(LogChaos, Warning, TEXT("EPA Touching Case"));
 				break;
 			case EEPAResult::Degenerate:
 				// We hit a degenerate simplex condition and could not reach a solution so use whatever near touching point GJK came up with.
