Chaos : Serialization Parity with 4.27-Chaos
- Updated PhysicsObjectVersion to reflect changes in 4.27-Chaos and ensure backwards combability.

diff --git a/Engine/Source/Runtime/Core/Public/UObject/PhysicsObjectVersion.h b/Engine/Source/Runtime/Core/Public/UObject/PhysicsObjectVersion.h
--- a/Engine/Source/Runtime/Core/Public/UObject/PhysicsObjectVersion.h
+++ b/Engine/Source/Runtime/Core/Public/UObject/PhysicsObjectVersion.h
@@ -66,6 +66,9 @@ struct CORE_API FPhysicsObjectVersion
 		// Added max linear and angular speed to Chaos bodies
 		AddChaosMaxLinearAngularSpeed,
 
+		// add convex geometry to older collections that did not have any
+		GeometryCollectionConvexDefaults,
+
 		// -----<new versions can be added above this line>-------------------------------------------------
 		VersionPlusOne,
 		LatestVersion = VersionPlusOne - 1
diff --git a/Engine/Source/Runtime/Experimental/GeometryCollectionEngine/Private/GeometryCollection/GeometryCollectionObject.cpp b/Engine/Source/Runtime/Experimental/GeometryCollectionEngine/Private/GeometryCollection/GeometryCollectionObject.cpp
--- a/Engine/Source/Runtime/Experimental/GeometryCollectionEngine/Private/GeometryCollection/GeometryCollectionObject.cpp
+++ b/Engine/Source/Runtime/Experimental/GeometryCollectionEngine/Private/GeometryCollection/GeometryCollectionObject.cpp
@@ -638,6 +638,7 @@ void UGeometryCollection::Serialize(FArchive& Ar)
 	Ar.UsingCustomVersion(FDestructionObjectVersion::GUID);
 	Ar.UsingCustomVersion(FUE5MainStreamObjectVersion::GUID);
 	Ar.UsingCustomVersion(FUE5ReleaseStreamObjectVersion::GUID);
+	Ar.UsingCustomVersion(FPhysicsObjectVersion::GUID);
 	
 	Chaos::FChaosArchive ChaosAr(Ar);
 
@@ -785,7 +786,8 @@ void UGeometryCollection::Serialize(FArchive& Ar)
 
 
 	// will generate convex bodies when they dont exist. 
-	if (Ar.CustomVer(FUE5ReleaseStreamObjectVersion::GUID) < FUE5ReleaseStreamObjectVersion::GeometryCollectionConvexDefaults)
+	if (Ar.CustomVer(FUE5ReleaseStreamObjectVersion::GUID) < FUE5ReleaseStreamObjectVersion::GeometryCollectionConvexDefaults
+		&& Ar.CustomVer(FPhysicsObjectVersion::GUID) < FPhysicsObjectVersion::GeometryCollectionConvexDefaults)
 	{
 #if WITH_EDITOR
 		if (!FGeometryCollectionConvexUtility::HasConvexHullData(GeometryCollection.Get()))
