Fix ARM64 build

diff --git a/Engine/Source/Runtime/Experimental/Chaos/Public/Chaos/GJK.h b/Engine/Source/Runtime/Experimental/Chaos/Public/Chaos/GJK.h
--- a/Engine/Source/Runtime/Experimental/Chaos/Public/Chaos/GJK.h
+++ b/Engine/Source/Runtime/Experimental/Chaos/Public/Chaos/GJK.h
@@ -1291,8 +1291,8 @@ namespace Chaos
 
 		const VectorRegister4Float RayLengthSimd = MakeVectorRegisterFloat(RayLength, RayLength, RayLength, RayLength);
 
-		NumVertsUnion NumVerts;
-		NumVerts.NumVerts = GlobalVectorConstants::IntZero;
+		VectorRegister4Int NumVerts = GlobalVectorConstants::IntZero;
+
 		VectorRegister4Float AToBRotation = VectorQuaternionInverse(BToARotation);
 
 		auto SupportAFunc = [&A, MarginA](const VectorRegister4Float& V)
@@ -1318,7 +1318,7 @@ namespace Chaos
 		VectorRegister4Float Lambda = VectorZeroFloat();
 		VectorRegister4Float X = StartPoint;
 		VectorRegister4Float V = VectorSubtract(X, VectorSubtract(SupportA, SupportB));
-		VectorRegister4Float Normal = { 0.f, 0.f, 1.f, 0.f };
+		VectorRegister4Float Normal = MakeVectorRegisterFloat(0.f, 0.f, 1.f, 0.f);
 
 		const VectorRegister4Float InitialPreDist2Simd = VectorDot3(V, V);
 
@@ -1354,7 +1354,9 @@ namespace Chaos
 			const VectorRegister4Float W = VectorSubtract(X, P);
 
 			// NumVerts store here at the beginning of the loop, it should be safe to reuse it further in the loop 
-			const int32 NumVertsInt = NumVerts.NumVertsInts[0];
+			alignas(16) int32 NumVertsInts[4];
+			VectorIntStoreAligned(NumVerts, NumVertsInts);
+			const int32 NumVertsInt = NumVertsInts[0];
 
 			As[NumVertsInt] = SupportA;
 			Bs[NumVertsInt] = SupportB;
@@ -1395,7 +1397,7 @@ namespace Chaos
 					Simplex[1] = VectorAdd(Simplex[1], XMinusOldX);
 					Simplex[2] = VectorAdd(Simplex[2], XMinusOldX);
 					Simplex[NumVertsInt] = VectorSubtract(X, P);
-					NumVerts.NumVerts = VectorIntAdd(NumVerts.NumVerts, GlobalVectorConstants::IntOne);
+					NumVerts = VectorIntAdd(NumVerts, GlobalVectorConstants::IntOne);
 
 					GJKPreDist2 = LimitMax; //translated origin so restart gjk search
 					bInflatedCloseEnough = false;
@@ -1404,7 +1406,7 @@ namespace Chaos
 			else
 			{
 				Simplex[NumVertsInt] = W;	//this is really X - P which is what we need for simplex computation
-				NumVerts.NumVerts = VectorIntAdd(NumVerts.NumVerts, GlobalVectorConstants::IntOne);
+				NumVerts = VectorIntAdd(NumVerts, GlobalVectorConstants::IntOne);
 			}
 
 			if (bInflatedCloseEnough && VectorMaskBits(VectorCompareGE(VDotW, VectorZeroFloat())))
@@ -1429,7 +1431,7 @@ namespace Chaos
 					const VectorRegister4Float InGJKPreDist2GTEps2 = VectorCompareGT(GJKPreDist2, Eps2Simd);
 					const VectorRegister4Float Inflation22GTEps2 = VectorCompareGT(Inflation2Simd, Eps2Simd);
 					constexpr VectorRegister4Int fourInt = MakeVectorRegisterIntConstant(4, 4, 4, 4);
-					const VectorRegister4Int Is4GTNumVerts = VectorIntCompareGT(fourInt, NumVerts.NumVerts);
+					const VectorRegister4Int Is4GTNumVerts = VectorIntCompareGT(fourInt, NumVerts);
 
 					const VectorRegister4Float IsInflatCloseEnough = VectorBitwiseAnd(LambdaEqZero, VectorBitwiseAnd(InGJKPreDist2GTEps2, VectorBitwiseAnd(Inflation22GTEps2, VectorCast4IntTo4Float(Is4GTNumVerts))));
 
@@ -1472,10 +1474,10 @@ namespace Chaos
 			constexpr VectorRegister4Int TwoInt = MakeVectorRegisterIntConstant(2, 2, 2, 2);
 			constexpr VectorRegister4Int ThreeInt = MakeVectorRegisterIntConstant(3, 3, 3, 3);
 
-			const VectorRegister4Float IsB0 = VectorCast4IntTo4Float(VectorIntCompareEQ(NumVerts.NumVerts, GlobalVectorConstants::IntZero));
-			const VectorRegister4Float IsB1 = VectorCast4IntTo4Float(VectorIntCompareEQ(NumVerts.NumVerts, GlobalVectorConstants::IntOne));
-			const VectorRegister4Float IsB2 = VectorCast4IntTo4Float(VectorIntCompareEQ(NumVerts.NumVerts, TwoInt));
-			const VectorRegister4Float IsB3 = VectorCast4IntTo4Float(VectorIntCompareEQ(NumVerts.NumVerts, ThreeInt));
+			const VectorRegister4Float IsB0 = VectorCast4IntTo4Float(VectorIntCompareEQ(NumVerts, GlobalVectorConstants::IntZero));
+			const VectorRegister4Float IsB1 = VectorCast4IntTo4Float(VectorIntCompareEQ(NumVerts, GlobalVectorConstants::IntOne));
+			const VectorRegister4Float IsB2 = VectorCast4IntTo4Float(VectorIntCompareEQ(NumVerts, TwoInt));
+			const VectorRegister4Float IsB3 = VectorCast4IntTo4Float(VectorIntCompareEQ(NumVerts, ThreeInt));
 
 			ClosestB = VectorSelect(IsB0, ClosestB, ClosestB4);
 			ClosestB = VectorSelect(IsB1, ClosestB1, ClosestB);
@@ -1510,7 +1512,9 @@ namespace Chaos
 					Barycentrics[2] = VectorSwizzle(Barycentric, 2, 2, 2, 2);
 					Barycentrics[3] = VectorSwizzle(Barycentric, 3, 3, 3, 3);
 
-					const int NumVertsInt = NumVerts.NumVertsInts[0];
+					alignas(16) int32 NumVertsInts[4];
+					VectorIntStoreAligned(NumVerts, NumVertsInts);
+					const int NumVertsInt = NumVertsInts[0];
 					for (int i = 0; i < NumVertsInt; ++i)
 					{
 						ClosestA = VectorMultiplyAdd(As[i], Barycentrics[i], ClosestA);
@@ -1550,7 +1554,9 @@ namespace Chaos
 					VertsA.Reserve(8);
 					VertsB.Reserve(8);
 
-					const int32 NumVertsInt = NumVerts.NumVertsInts[0];
+					alignas(16) int32 NumVertsInts[4];
+					VectorIntStoreAligned(NumVerts, NumVertsInts);
+					const int32 NumVertsInt = NumVertsInts[0];
 
 					alignas(16) FRealSingle XFloat[4];
 					VectorStoreAligned(X, XFloat);
@@ -1613,7 +1619,7 @@ namespace Chaos
 				{
 					//didn't even go into gjk loop, touching hit
 					OutTime = -(MarginA + MarginB);
-					OutNormal = { 0,0,1,0 };
+					OutNormal = MakeVectorRegisterFloat(0.0f, 0.0f, 1.0f, 0.0f);
 					OutPosition = VectorMultiplyAdd(OutNormal, MarginASimd, As[0]);
 				}
 			}
@@ -1621,8 +1627,8 @@ namespace Chaos
 		else
 		{
 			// Initial overlap without MTD. These properties are not valid, but assigning them anyway so they don't contain NaNs and cause issues in invoking code.
-			OutNormal = { 0,0,1,0 };
-			OutPosition = { 0,0,0,0 };
+			OutNormal = MakeVectorRegisterFloat(0.0f, 0.0f, 1.0f, 0.0f);
+			OutPosition = MakeVectorRegisterFloat(0.0f, 0.0f, 0.0f, 0.0f);
 		}
 
 		return true;
diff --git a/Engine/Source/Runtime/Experimental/Chaos/Public/Chaos/SimplexVectorized.h b/Engine/Source/Runtime/Experimental/Chaos/Public/Chaos/SimplexVectorized.h
--- a/Engine/Source/Runtime/Experimental/Chaos/Public/Chaos/SimplexVectorized.h
+++ b/Engine/Source/Runtime/Experimental/Chaos/Public/Chaos/SimplexVectorized.h
@@ -25,7 +25,7 @@ namespace Chaos
 		const VectorRegister4Float X0ToX1Squared = VectorDot3(X0ToX1, X0ToX1);
 		const VectorRegister4Float DotBigger = VectorCompareGE(Dot, X0ToX1Squared);
 
-		const VectorRegister4Float MinLimt = { std::numeric_limits<FRealSingle>::min() , std::numeric_limits<FRealSingle>::min(), std::numeric_limits<FRealSingle>::min(), std::numeric_limits<FRealSingle>::min() };
+		const VectorRegister4Float MinLimt = MakeVectorRegisterFloat(std::numeric_limits<FRealSingle>::min() , std::numeric_limits<FRealSingle>::min(), std::numeric_limits<FRealSingle>::min(), std::numeric_limits<FRealSingle>::min());
 		const VectorRegister4Float X0ToX1SquaredSmall = VectorCompareGE(MinLimt, X0ToX1Squared);
 		const VectorRegister4Float IsX1 = VectorBitwiseOr(DotBigger, X0ToX1SquaredSmall);
 
@@ -54,12 +54,6 @@ namespace Chaos
 		return Closest;
 	}
 
-	union alignas(16) NumVertsUnion 
-	{ 
-		VectorRegister4Int NumVerts;
-		int32 NumVertsInts[4]; 
-	};
-
 	// Based on an algorithm in Real Time Collision Detection - Ericson (very close to that)
 	// Using the same variable name conventions for easy reference
 	FORCEINLINE VectorRegister4Float  TriangleSimplexFindOriginFast(VectorRegister4Float* Simplex, VectorRegister4Int& NumVerts, VectorRegister4Float& OutBarycentric, VectorRegister4Float* As, VectorRegister4Float* Bs)
@@ -341,10 +335,12 @@ namespace Chaos
 	}
 
 
-	FORCEINLINE VectorRegister4Float VectorSimplexFindClosestToOrigin(VectorRegister4Float* Simplex, NumVertsUnion& NumVerts, VectorRegister4Float& OutBarycentric, VectorRegister4Float* A, VectorRegister4Float* B)
+	FORCEINLINE VectorRegister4Float VectorSimplexFindClosestToOrigin(VectorRegister4Float* Simplex, VectorRegister4Int& NumVerts, VectorRegister4Float& OutBarycentric, VectorRegister4Float* A, VectorRegister4Float* B)
 	{
 		VectorRegister4Float ClosestPoint;
-		switch (NumVerts.NumVertsInts[0])
+		alignas(16) int32 NumVertsInt[4];
+		VectorIntStoreAligned(NumVerts, NumVertsInt);
+		switch (NumVertsInt[0])
 		{
 		case 1:
 			OutBarycentric = GlobalVectorConstants::Float1000;
@@ -352,17 +348,17 @@ namespace Chaos
 			break;
 		case 2:
 		{
-			ClosestPoint = VectorLineSimplexFindOrigin(Simplex, NumVerts.NumVerts, OutBarycentric, A, B);
+			ClosestPoint = VectorLineSimplexFindOrigin(Simplex, NumVerts, OutBarycentric, A, B);
 			break;
 		}
 		case 3:
 		{
-			ClosestPoint = TriangleSimplexFindOriginFast(Simplex, NumVerts.NumVerts, OutBarycentric, A, B);
+			ClosestPoint = TriangleSimplexFindOriginFast(Simplex, NumVerts, OutBarycentric, A, B);
 			break;
 		}
 		case 4:
 		{
-			ClosestPoint = VectorTetrahedronSimplexFindOrigin(Simplex, NumVerts.NumVerts, OutBarycentric, A, B);
+			ClosestPoint = VectorTetrahedronSimplexFindOrigin(Simplex, NumVerts, OutBarycentric, A, B);
 			break;
 		}
 		default:
diff --git a/Engine/Source/Runtime/Experimental/Chaos/Public/Chaos/Triangle.h b/Engine/Source/Runtime/Experimental/Chaos/Public/Chaos/Triangle.h
--- a/Engine/Source/Runtime/Experimental/Chaos/Public/Chaos/Triangle.h
+++ b/Engine/Source/Runtime/Experimental/Chaos/Public/Chaos/Triangle.h
@@ -24,9 +24,9 @@ namespace Chaos
 			, B(InB)
 			, C(InC)
 		{
-			VectorA = { static_cast<FRealSingle>(A[0]), static_cast<FRealSingle>(A[1]), static_cast<FRealSingle>(A[2]), static_cast<FRealSingle>(0.0f) };
-			VectorB = { static_cast<FRealSingle>(B[0]), static_cast<FRealSingle>(B[1]), static_cast<FRealSingle>(B[2]), static_cast<FRealSingle>(0.0f) };
-			VectorC = { static_cast<FRealSingle>(C[0]), static_cast<FRealSingle>(C[1]), static_cast<FRealSingle>(C[2]), static_cast<FRealSingle>(0.0f) };
+			VectorA = MakeVectorRegisterFloat(static_cast<FRealSingle>(A[0]), static_cast<FRealSingle>(A[1]), static_cast<FRealSingle>(A[2]), static_cast<FRealSingle>(0.0f));
+			VectorB = MakeVectorRegisterFloat(static_cast<FRealSingle>(B[0]), static_cast<FRealSingle>(B[1]), static_cast<FRealSingle>(B[2]), static_cast<FRealSingle>(0.0f));
+			VectorC = MakeVectorRegisterFloat(static_cast<FRealSingle>(C[0]), static_cast<FRealSingle>(C[1]), static_cast<FRealSingle>(C[2]), static_cast<FRealSingle>(0.0f));
 		}
 
 		FVec3& operator[](uint32 InIndex)
diff --git a/Engine/Source/Runtime/Experimental/Chaos/Public/Chaos/VectorUtility.h b/Engine/Source/Runtime/Experimental/Chaos/Public/Chaos/VectorUtility.h
--- a/Engine/Source/Runtime/Experimental/Chaos/Public/Chaos/VectorUtility.h
+++ b/Engine/Source/Runtime/Experimental/Chaos/Public/Chaos/VectorUtility.h
@@ -8,7 +8,7 @@
  * @param V	vector
  * @return		VectorRegister4Float( B.x, A.y, A.z, A.w)
  */
-#if !defined(_MSC_VER)
+#if !defined(_MSC_VER)|| PLATFORM_ENABLE_VECTORINTRINSICS_NEON
 FORCEINLINE VectorRegister4Float VectorCast4IntTo4Float(const VectorRegister4Int& V)
 {
 	return VectorRegister4Float(V);
@@ -27,7 +27,7 @@ FORCEINLINE VectorRegister4Float VectorCast4IntTo4Float(const VectorRegister4Int
  * @param V	vector
  * @return		VectorCast4FloatTo4Int( B.x, A.y, A.z, A.w)
  */
-#if !defined(_MSC_VER)
+#if !defined(_MSC_VER) || PLATFORM_ENABLE_VECTORINTRINSICS_NEON
 FORCEINLINE VectorRegister4Int VectorCast4FloatTo4Int(const VectorRegister4Float& V)
 {
 	return VectorRegister4Int(V);
