Chaos : Optimization : Convex raycast now using a dedicated method instead of GJKRaycast

diff --git a/Engine/Source/Runtime/Experimental/Chaos/Private/Chaos/Convex.cpp b/Engine/Source/Runtime/Experimental/Chaos/Private/Chaos/Convex.cpp
--- a/Engine/Source/Runtime/Experimental/Chaos/Private/Chaos/Convex.cpp
+++ b/Engine/Source/Runtime/Experimental/Chaos/Private/Chaos/Convex.cpp
@@ -4,7 +4,7 @@
 #include "Chaos/Sphere.h"
 #include "Chaos/ImplicitObjectScaled.h"
 
-//PRAGMA_DISABLE_OPTIMIZATION
+#define CHAOS_CONVEX_USE_FAST_RAYCAST 1
 
 namespace Chaos
 {
@@ -16,12 +16,69 @@ namespace Chaos
 
 	bool FConvex::Raycast(const FVec3& StartPoint, const FVec3& Dir, const FReal Length, const FReal Thickness, FReal& OutTime, FVec3& OutPosition, FVec3& OutNormal, int32& OutFaceIndex) const
 	{
+#if CHAOS_CONVEX_USE_FAST_RAYCAST
+		return RaycastFast(StartPoint, Dir, Length, Thickness, OutTime, OutPosition, OutNormal, OutFaceIndex);
+#else
 		OutFaceIndex = INDEX_NONE;	//finding face is expensive, should be called directly by user
 		const FRigidTransform3 StartTM(StartPoint, FRotation3::FromIdentity());
 		const TSphere<FReal, 3> Sphere(FVec3(0), Thickness);
 		return GJKRaycast(*this, Sphere, StartTM, Dir, Length, OutTime, OutPosition, OutNormal);
+#endif
 	}
 
+	bool FConvex::RaycastFast(const FVec3& StartPoint, const FVec3& Dir, const FReal Length, const FReal Thickness, FReal& OutTime, FVec3& OutPosition, FVec3& OutNormal, int32& OutFaceIndex) const
+	{
+		FReal EntryTime = 0.;
+		FReal ExitTime = Length;
+		int32 PlaneIndex = -1;
+		
+		for (int32 Idx = 0; Idx < Planes.Num(); ++Idx)
+		{
+			const FPlaneType& Plane = Planes[Idx];
+			
+			const FReal NDotDir = FVec3::DotProduct(Plane.Normal(), Dir);
+			const FReal Distance = FVec3::DotProduct((StartPoint - FVec3{Plane.X()}), Plane.Normal()); //Plane.SignedDistance(StartPoint);
+
+			if (NDotDir == 0.)
+			{
+				if (Distance > 0.)
+				{
+					// ray is perpendicular to one plane and outside of the convex, no intersection 
+					return false;
+				}
+			}
+			else
+			{
+				const FReal Time = -Distance /  NDotDir;
+				if (NDotDir < 0.)
+				{
+					// ray opposite to the plane normal, it is an entry
+					if (Time > EntryTime)
+					{
+						EntryTime = Time; 
+						PlaneIndex = Idx;
+					}
+				}
+				else
+				{
+					// ray has same orientation as the plane normal, it is an exit
+					ExitTime = FMath::Min(ExitTime, Time);
+				}
+			}
+			// early exit when possible
+			if (EntryTime > ExitTime)
+			{
+				return false;
+			};
+		} 
+
+		// successful intersection
+		OutTime = EntryTime;
+		OutPosition = StartPoint + (Dir * OutTime);
+		OutNormal = (PlaneIndex >= 0)? FVec3{Planes[PlaneIndex].Normal()}: -Dir;
+		OutFaceIndex = INDEX_NONE; // @todo(chaos) we could return PlaneIndex ( but would differ from previous implementation )
+		return true;
+	}
 
 	int32 FConvex::FindMostOpposingFace(const FVec3& Position, const FVec3& UnitDir, int32 HintFaceIndex, FReal SearchDist) const
 	{
diff --git a/Engine/Source/Runtime/Experimental/Chaos/Public/Chaos/Convex.h b/Engine/Source/Runtime/Experimental/Chaos/Public/Chaos/Convex.h
--- a/Engine/Source/Runtime/Experimental/Chaos/Public/Chaos/Convex.h
+++ b/Engine/Source/Runtime/Experimental/Chaos/Public/Chaos/Convex.h
@@ -343,6 +343,7 @@ namespace Chaos
 		 * caller. 
 		 */
 		virtual bool Raycast(const FVec3& StartPoint, const FVec3& Dir, const FReal Length, const FReal Thickness, FReal& OutTime, FVec3& OutPosition, FVec3& OutNormal, int32& OutFaceIndex) const override;
+		bool RaycastFast(const FVec3& StartPoint, const FVec3& Dir, const FReal Length, const FReal Thickness, FReal& OutTime, FVec3& OutPosition, FVec3& OutNormal, int32& OutFaceIndex) const;
 
 		virtual Pair<FVec3, bool> FindClosestIntersectionImp(const FVec3& StartPoint, const FVec3& EndPoint, const FReal Thickness) const override
 		{
