Chaos Cloth - updates to PBDCollisionSpringConstraintsBase (cloth self collisions) to improve stability by using simple proximity detection and point-based repulsion.

This required a proximity detection method for cloth. Currently proximities are only detected as point-face. This was done by adding a TAABBTree to the TriangleMesh, and adding a PointProximityQuery. I'm not sure if TriangleMesh is the best place for the BVH to live since the TriangleMesh does not own the point positions, but it seemed like a reasonable place to start. I opted to put it here rather than inside the collision springs class directly because I believe the BVH will be useful for other cloth collision operations in the future (e.g., intersection detection).

Additionally, I removed some data for self collisions in PBDEvolution that is unused.

==CHANGES==
PBDCollisionSpringConstraintsBase.{h,cpp}
This now holds a reference to the TriangleMesh and optional RestPositions. The TriangleMesh is used for its topology and the BVH. The RestPositions are used to exclude collisions between points that are closer than thickness away in the rest state. Normals was removed since it's not longer used.
Init (the method that creates the constraints) now uses the BVH's PointProximityQuery and creates up to 3 constraints per point. This leads to a modest reduction in jitter as connections tend to be more stable compared with the previous more prediction-based detection method.
GetDelta has been updated to be a simple point repulsion rather than the triangle normal-based repulsion from before.  This was the most significant change when it comes to jitter. There was a bug in the previous method that caused too large of deltas in some cases, but even fixing that bug did not resolve the jitter problems.

TriangleMesh.{h,cpp}
Added the BVH here and added methods to update the BVH and do PointProximityQuery

Plane.h
Added new method FindClosestPointAndBaryOnTriangle. This should produce the same closest point as FindClosestPointOnTriangle (up to floating point differences), but it has less duplicate work, and returns the calculated barycentric coordinates. It also handles degenerate triangles like the old code in PBDCollisionSpringConstraintsBase.

PBDCollisionSpringConstraints.h, ChaosClothConstraints.{h,cpp}, ChaosClothingSimulationCloth.cpp
Updated with new constructor for PBDCollisionSpringConstraintsBase
Modified the magic number multiplier in Apply to .5. This worked better with the new constraint delta.

ChaosClothingSimulation.cpp
Updated debug draw for collision springs.

ChaosClothingSimulationCloth.cpp
Updated FLODData to manage the BVH in the triangle mesh, including a "PreTimestep" which rebuilds the BVH.

ChaosClothingSimulationSolver.cpp
Added calling PreTimestep on the cloth.

PBDEvolution.{h,cpp}
Removed unused self collision data.

TriangleCollisionPoint.h
New struct for holding triangle based collision data.

==TESTING==
I tested by enabling self collisions in the example in UE-104875 as well as in the Cloth map in ContentExamples. Jitter is almost completely gone in these examples. Cloth collisions can still occur for fast moving cloth. The topline simulation time for cloth was not significantly different between this and the previous self collision method.

diff --git a/Engine/Plugins/Experimental/ChaosCloth/Source/Chaos/Private/ChaosCloth/ChaosClothConstraints.cpp b/Engine/Plugins/Experimental/ChaosCloth/Source/Chaos/Private/ChaosCloth/ChaosClothConstraints.cpp
--- a/Engine/Plugins/Experimental/ChaosCloth/Source/Chaos/Private/ChaosCloth/ChaosClothConstraints.cpp
+++ b/Engine/Plugins/Experimental/ChaosCloth/Source/Chaos/Private/ChaosCloth/ChaosClothConstraints.cpp
@@ -476,12 +476,13 @@ void FClothConstraints::SetShapeTargetConstraints(Softs::FSolverReal ShapeTarget
 	++NumConstraintRules;
 }
 
-void FClothConstraints::SetSelfCollisionConstraints(const TArray<TVec3<int32>>& SurfaceElements, TSet<TVec2<int32>>&& DisabledCollisionElements, Softs::FSolverReal SelfCollisionThickness)
+void FClothConstraints::SetSelfCollisionConstraints(const FTriangleMesh& TriangleMesh, TSet<TVec2<int32>>&& DisabledCollisionElements, Softs::FSolverReal SelfCollisionThickness)
 {
 	SelfCollisionConstraints = MakeShared<Softs::FPBDCollisionSpringConstraints>(
 		ParticleOffset,
 		NumParticles,
-		SurfaceElements,
+		TriangleMesh,
+		AnimationPositions,
 		MoveTemp(DisabledCollisionElements),
 		SelfCollisionThickness,
 		/*Stiffness =*/ (Softs::FSolverReal)1.);
diff --git a/Engine/Plugins/Experimental/ChaosCloth/Source/Chaos/Private/ChaosCloth/ChaosClothingSimulation.cpp b/Engine/Plugins/Experimental/ChaosCloth/Source/Chaos/Private/ChaosCloth/ChaosClothingSimulation.cpp
--- a/Engine/Plugins/Experimental/ChaosCloth/Source/Chaos/Private/ChaosCloth/ChaosClothingSimulation.cpp
+++ b/Engine/Plugins/Experimental/ChaosCloth/Source/Chaos/Private/ChaosCloth/ChaosClothingSimulation.cpp
@@ -1976,7 +1976,6 @@ void FClothingSimulation::DebugDrawSelfCollision(FPrimitiveDrawInterface* PDI) c
 			const TConstArrayView<Softs::FSolverVec3> Positions = Cloth->GetParticlePositions(Solver.Get());
 			const TArray<TVec4<int32>>& Constraints = SelfCollisionConstraints->GetConstraints();
 			const TArray<Softs::FSolverVec3>& Barys = SelfCollisionConstraints->GetBarys();
-			const TArray<Softs::FSolverVec3>& Normals = SelfCollisionConstraints->GetNormals();
 			const FReal Thickness = (FReal)SelfCollisionConstraints->GetThickness();
 			const FReal Height = Thickness + Thickness;
 
@@ -1984,37 +1983,25 @@ void FClothingSimulation::DebugDrawSelfCollision(FPrimitiveDrawInterface* PDI) c
 			{
 				const TVec4<int32>& Constraint = Constraints[Index];
 				const FVec3 Bary(Barys[Index]);
-				const FVec3 Normal(Normals[Index]);
 
-				const FVector P = LocalSpaceLocation + FVector(Positions[Constraint[0]]);
-				const FVector P0 = LocalSpaceLocation + FVector(Positions[Constraint[1]]);
-				const FVector P1 = LocalSpaceLocation + FVector(Positions[Constraint[2]]);
-				const FVector P2 = LocalSpaceLocation + FVector(Positions[Constraint[3]]);
+				// Constraint index includes Offset, but so does Positions.
+				const FVector P = LocalSpaceLocation + FVector(Positions[Constraint[0] - Offset]);
+				const FVector P0 = LocalSpaceLocation + FVector(Positions[Constraint[1] - Offset]);
+				const FVector P1 = LocalSpaceLocation + FVector(Positions[Constraint[2] - Offset]);
+				const FVector P2 = LocalSpaceLocation + FVector(Positions[Constraint[3] - Offset]);
 
 				const FVector Pos0 = P0 * Bary[0] + P1 * Bary[1] + P2 * Bary[2];
-				const FVector Pos1 = Pos0 + Height * Normal;
 
-				// Draw point to surface line (=normal)
 				static const FLinearColor Brown(0.1f, 0.05f, 0.f);
-				DrawLine(PDI, Pos0, Pos1, Brown);
 
 				// Draw barycentric coordinate
 				DrawLine(PDI, Pos0, P0, Brown);
 				DrawLine(PDI, Pos0, P1, Brown);
 				DrawLine(PDI, Pos0, P2, Brown);
 
-				// Draw pushup to point
+				// Draw connection to point
 				static const FLinearColor Orange(0.3f, 0.15f, 0.f);
-				DrawLine(PDI, Pos1, P, Orange);
-
-				// Draw contact
-				FVec3 TangentU, TangentV;
-				Normal.FindBestAxisVectors(TangentU, TangentV);
-
-				DrawLine(PDI, Pos0 + TangentU, Pos0 + TangentV, Brown);
-				DrawLine(PDI, Pos0 + TangentU, Pos0 - TangentV, Brown);
-				DrawLine(PDI, Pos0 - TangentU, Pos0 - TangentV, Brown);
-				DrawLine(PDI, Pos0 - TangentU, Pos0 + TangentV, Brown);
+				DrawLine(PDI, Pos0, P, Orange);
 
 				// Draw hit triangle (with thickness and added colliding particle's thickness)
 				static const FLinearColor Red(0.3f, 0.f, 0.f);
diff --git a/Engine/Plugins/Experimental/ChaosCloth/Source/Chaos/Private/ChaosCloth/ChaosClothingSimulationCloth.cpp b/Engine/Plugins/Experimental/ChaosCloth/Source/Chaos/Private/ChaosCloth/ChaosClothingSimulationCloth.cpp
--- a/Engine/Plugins/Experimental/ChaosCloth/Source/Chaos/Private/ChaosCloth/ChaosClothingSimulationCloth.cpp
+++ b/Engine/Plugins/Experimental/ChaosCloth/Source/Chaos/Private/ChaosCloth/ChaosClothingSimulationCloth.cpp
@@ -121,7 +121,7 @@ void FClothingSimulationCloth::FLODData::Add(FClothingSimulationSolver* Solver,
 				DisabledCollisionElements.Emplace(TVec2<int32>(Element, Index));
 			}
 		}
-		ClothConstraints.SetSelfCollisionConstraints(SurfaceElements, MoveTemp(DisabledCollisionElements), (FReal)Cloth->SelfCollisionThickness);
+		ClothConstraints.SetSelfCollisionConstraints(TriangleMesh, MoveTemp(DisabledCollisionElements), (FReal)Cloth->SelfCollisionThickness);
 	}
 
 	// Edge constraints
diff --git a/Engine/Plugins/Experimental/ChaosCloth/Source/Chaos/Private/ChaosCloth/ChaosClothingSimulationSolver.cpp b/Engine/Plugins/Experimental/ChaosCloth/Source/Chaos/Private/ChaosCloth/ChaosClothingSimulationSolver.cpp
--- a/Engine/Plugins/Experimental/ChaosCloth/Source/Chaos/Private/ChaosCloth/ChaosClothingSimulationSolver.cpp
+++ b/Engine/Plugins/Experimental/ChaosCloth/Source/Chaos/Private/ChaosCloth/ChaosClothingSimulationSolver.cpp
@@ -304,7 +304,6 @@ void FClothingSimulationSolver::ResetParticles()
 {
 	Evolution->ResetParticles();
 	Evolution->ResetConstraintRules();
-	Evolution->ResetSelfCollision();
 	ClothsConstraints.Reset();
 }
 
diff --git a/Engine/Plugins/Experimental/ChaosCloth/Source/Chaos/Public/ChaosCloth/ChaosClothConstraints.h b/Engine/Plugins/Experimental/ChaosCloth/Source/Chaos/Public/ChaosCloth/ChaosClothConstraints.h
--- a/Engine/Plugins/Experimental/ChaosCloth/Source/Chaos/Public/ChaosCloth/ChaosClothConstraints.h
+++ b/Engine/Plugins/Experimental/ChaosCloth/Source/Chaos/Public/ChaosCloth/ChaosClothConstraints.h
@@ -41,7 +41,7 @@ namespace Chaos
 		void SetBackstopConstraints(const TConstArrayView<FRealSingle>& BackstopDistances, const TConstArrayView<FRealSingle>& BackstopRadiuses, bool bUseLegacyBackstop);
 		void SetAnimDriveConstraints(const TConstArrayView<FRealSingle>& AnimDriveStiffnessMultipliers, const TConstArrayView<FRealSingle>& AnimDriveDampingMultipliers);
 		void SetShapeTargetConstraints(Softs::FSolverReal ShapeTargetStiffness);
-		void SetSelfCollisionConstraints(const TArray<TVec3<int32>>& SurfaceElements, TSet<TVec2<int32>>&& DisabledCollisionElements, Softs::FSolverReal SelfCollisionThickness);
+		void SetSelfCollisionConstraints(const class FTriangleMesh& TriangleMesh, TSet<TVec2<int32>>&& DisabledCollisionElements, Softs::FSolverReal SelfCollisionThickness);
 
 		void CreateRules();
 		void Enable(bool bEnable);
diff --git a/Engine/Source/Runtime/Experimental/Chaos/Private/Chaos/PBDCollisionSpringConstraintsBase.cpp b/Engine/Source/Runtime/Experimental/Chaos/Private/Chaos/PBDCollisionSpringConstraintsBase.cpp
--- a/Engine/Source/Runtime/Experimental/Chaos/Private/Chaos/PBDCollisionSpringConstraintsBase.cpp
+++ b/Engine/Source/Runtime/Experimental/Chaos/Private/Chaos/PBDCollisionSpringConstraintsBase.cpp
@@ -1,63 +1,25 @@
 // Copyright Epic Games, Inc. All Rights Reserved.
 #include "Chaos/PBDCollisionSpringConstraintsBase.h"
+#include "Chaos/Plane.h"
+#include "Chaos/TriangleCollisionPoint.h"
+#include "Chaos/TriangleMesh.h"
 
 #if !COMPILE_WITHOUT_UNREAL_SUPPORT
 #include "Chaos/Framework/Parallel.h"
-#if PLATFORM_DESKTOP && PLATFORM_64BITS
-#include "kDOP.h"
-#endif
 
 namespace Chaos::Softs {
 
-#if PLATFORM_DESKTOP && PLATFORM_64BITS
-class FMeshBuildDataProvider
-{
-public:
-	/** Initialization constructor. */
-	FMeshBuildDataProvider(const TkDOPTree<const FMeshBuildDataProvider, uint32>& InkDopTree)
-	    : kDopTree(InkDopTree)
-	{}
-
-	// kDOP data provider interface.
-
-	FORCEINLINE const TkDOPTree<const FMeshBuildDataProvider, uint32>& GetkDOPTree() const
-	{
-		return kDopTree;
-	}
-
-	FORCEINLINE const FMatrix& GetLocalToWorld() const
-	{
-		return FMatrix::Identity;
-	}
-
-	FORCEINLINE const FMatrix& GetWorldToLocal() const
-	{
-		return FMatrix::Identity;
-	}
-
-	FORCEINLINE FMatrix GetLocalToWorldTransposeAdjoint() const
-	{
-		return FMatrix::Identity;
-	}
-
-	FORCEINLINE FSolverReal GetDeterminant() const
-	{
-		return 1.0f;
-	}
-
-private:
-	const TkDOPTree<const FMeshBuildDataProvider, uint32>& kDopTree;
-};
-#endif
-
 FPBDCollisionSpringConstraintsBase::FPBDCollisionSpringConstraintsBase(
 	const int32 InOffset,
 	const int32 InNumParticles,
-	const TArray<TVec3<int32>>& InElements,
+	const FTriangleMesh& InTriangleMesh,
+	const TArray<FSolverVec3>* InReferencePositions,
 	TSet<TVec2<int32>>&& InDisabledCollisionElements,
 	const FSolverReal InThickness,
 	const FSolverReal InStiffness)
-	: Elements(InElements)
+	: TriangleMesh(InTriangleMesh)
+	, Elements(InTriangleMesh.GetSurfaceElements())
+	, ReferencePositions(InReferencePositions)
 	, DisabledCollisionElements(InDisabledCollisionElements)
 	, Offset(InOffset)
 	, NumParticles(InNumParticles)
@@ -68,7 +30,6 @@ FPBDCollisionSpringConstraintsBase::FPBDCollisionSpringConstraintsBase(
 
 void FPBDCollisionSpringConstraintsBase::Init(const FSolverParticles& Particles)
 {
-#if PLATFORM_DESKTOP && PLATFORM_64BITS
 	if (!Elements.Num())
 	{
 		return;
@@ -76,122 +37,69 @@ void FPBDCollisionSpringConstraintsBase::Init(const FSolverParticles& Particles)
 
 	Constraints.Reset();
 	Barys.Reset();
-	Normals.Reset();
 
-	TkDOPTree<const FMeshBuildDataProvider, uint32> DopTree;
-	TArray<FkDOPBuildCollisionTriangle<uint32>> BuildTriangleArray;
-	BuildTriangleArray.Reserve(Elements.Num());
-	for (int32 i = 0; i < Elements.Num(); ++i)
-	{
-		const TVector<int32, 3>& Elem = Elements[i];
-		const FSolverVec3& P0 = Particles.X(Elem[0]);
-		const FSolverVec3& P1 = Particles.X(Elem[1]);
-		const FSolverVec3& P2 = Particles.X(Elem[2]);
-
-		BuildTriangleArray.Add(FkDOPBuildCollisionTriangle<uint32>(i, FVec3(P0), FVec3(P1), FVec3(P2)));
-	}
-	DopTree.Build(BuildTriangleArray);
+	FTriangleMesh::TBVHType<FSolverReal> BVH;
+	TriangleMesh.BuildBVH(static_cast<const TArrayView<const FSolverVec3>&>(Particles.XArray()), BVH);
 
-	FMeshBuildDataProvider DopDataProvider(DopTree);
 	FCriticalSection CriticalSection;
-
-	const FSolverReal Height = Thickness + Thickness;
-
+	const FSolverReal HeightSq = FMath::Square(Thickness + Thickness);
 	PhysicsParallelFor(NumParticles,
-		[this, &Particles, &CriticalSection, Height, &DopDataProvider, &DopTree](int32 i)
+		[this, &BVH, &Particles, &CriticalSection, HeightSq](int32 i)
 		{
 			const int32 Index = i + Offset;
 
-			FkHitResult Result;
-			const FSolverVec3& Start = Particles.X(Index);
-			const FSolverVec3 Direction = Particles.V(Index).GetSafeNormal();
-			const FSolverVec3 End = Particles.P(Index) + Direction * Height;
-			using FVec4Real = decltype(FVector4::X);
-			const FVector4 Start4((FVec4Real)Start.X, (FVec4Real)Start.Y, (FVec4Real)Start.Z, 0);
-			const FVector4 End4((FVec4Real)End.X, (FVec4Real)End.Y, (FVec4Real)End.Z, 0);
-			TkDOPLineCollisionCheck<const FMeshBuildDataProvider, uint32> Ray(Start4, End4, true, DopDataProvider, &Result);
-			if (DopTree.LineCheck(Ray))
-			{
-				const TVector<int32, 3>& Elem = Elements[Result.Item];
-				if (DisabledCollisionElements.Contains({Index, Elem[0]}) ||
-					DisabledCollisionElements.Contains({Index, Elem[1]}) ||
-					DisabledCollisionElements.Contains({Index, Elem[2]}))
-				{
-					return;
-				}
-
-				const FSolverVec3& P = Particles.X(Index);
-				const FSolverVec3& P0 = Particles.X(Elem[0]);
-				const FSolverVec3& P1 = Particles.X(Elem[1]);
-				const FSolverVec3& P2 = Particles.X(Elem[2]);
-				const FSolverVec3 P10 = P1 - P0;
-				const FSolverVec3 P20 = P2 - P0;
-				const FSolverVec3 PP0 = P - P0;
-				const FSolverReal Size10 = P10.SizeSquared();
-				const FSolverReal Size20 = P20.SizeSquared();
-				FSolverVec3 Bary;
-				if (Size10 < SMALL_NUMBER)
+			TArray< TTriangleCollisionPoint<FSolverReal> > Result;
+			if (TriangleMesh.PointProximityQuery(BVH, static_cast<const TArrayView<const FSolverVec3>&>(Particles.XArray()), Index, Particles.X(Index), Thickness, Thickness,
+				[this](const int32 PointIndex, const int32 TriangleIndex)->bool
 				{
-					if (Size20 < SMALL_NUMBER)
+					const TVector<int32, 3>& Elem = Elements[TriangleIndex];
+					if (DisabledCollisionElements.Contains({ PointIndex, Elem[0] }) ||
+						DisabledCollisionElements.Contains({ PointIndex, Elem[1] }) ||
+						DisabledCollisionElements.Contains({ PointIndex, Elem[2] }))
 					{
-						Bary.Y = Bary.Z = 0.f;
-						Bary.X = 1.f;
+						return false;
 					}
-					else
-					{
-						const FSolverReal Size = PP0.SizeSquared();
-						if (Size < SMALL_NUMBER)
-						{
-							Bary.Y = Bary.Z = 0.f;
-							Bary.X = 1.f;
-						}
-						else
+					return true;
+				},
+				Result))
+			{
+
+				constexpr int32 MaxConnectionsPerPoint = 3;
+				if (Result.Num() > MaxConnectionsPerPoint)
+				{
+					// TODO: once we have a PartialSort, use that instead here.
+					Result.Sort(
+						[](const TTriangleCollisionPoint<FSolverReal>& First, const TTriangleCollisionPoint<FSolverReal>& Second)->bool
 						{
-							const FSolverReal ProjP2 = FSolverVec3::DotProduct(PP0, P20);
-							Bary.Y = 0.f;
-							Bary.Z = ProjP2 * FMath::InvSqrt(Size20 * Size);
-							Bary.X = 1.0f - Bary.Z;
+							return First.Phi < Second.Phi;
 						}
-					}
+					);
+					Result.SetNum(MaxConnectionsPerPoint, false /*bAllowShrinking*/);
 				}
-				else if (Size20 < SMALL_NUMBER)
+
+				for (const TTriangleCollisionPoint<FSolverReal>& CollisionPoint : Result)
 				{
-					const FSolverReal Size = PP0.SizeSquared();
-					if (Size < SMALL_NUMBER)
+					const TVector<int32, 3>& Elem = Elements[CollisionPoint.Indices[1]];
+					if (ReferencePositions)
 					{
-						Bary.Y = Bary.Z = 0.f;
-						Bary.X = 1.f;
-					}
-					else
-					{
-						const FSolverReal ProjP1 = FSolverVec3::DotProduct(PP0, P10);
-						Bary.Y = ProjP1 * FMath::InvSqrt(Size10 * PP0.SizeSquared());
-						Bary.Z = 0.f;
-						Bary.X = 1.0f - Bary.Y;
+						const FSolverVec3& RefP = (*ReferencePositions)[Index];
+						const FSolverVec3& RefP0 = (*ReferencePositions)[Elem[0]];
+						const FSolverVec3& RefP1 = (*ReferencePositions)[Elem[1]];
+						const FSolverVec3& RefP2 = (*ReferencePositions)[Elem[2]];
+						const FSolverVec3 RefDiff = RefP - CollisionPoint.Bary[1] * RefP0 - CollisionPoint.Bary[2] * RefP1 - CollisionPoint.Bary[3] * RefP2;
+						if (RefDiff.SizeSquared() < HeightSq)
+						{
+							continue;
+						}
 					}
+					CriticalSection.Lock();
+					Constraints.Add({ Index, Elem[0], Elem[1], Elem[2] });
+					Barys.Add({ CollisionPoint.Bary[1], CollisionPoint.Bary[2], CollisionPoint.Bary[3] });
+					CriticalSection.Unlock();
 				}
-				else
-				{
-					const FSolverReal ProjSides = FSolverVec3::DotProduct(P10, P20);
-					const FSolverReal ProjP1 = FSolverVec3::DotProduct(PP0, P10);
-					const FSolverReal ProjP2 = FSolverVec3::DotProduct(PP0, P20);
-					const FSolverReal Denom = Size10 * Size20 - ProjSides * ProjSides;
-					Bary.Y = (Size20 * ProjP1 - ProjSides * ProjP2) / Denom;
-					Bary.Z = (Size10 * ProjP2 - ProjSides * ProjP1) / Denom;
-					Bary.X = 1.0f - Bary.Z - Bary.Y;
-				}
-
-				FSolverVec3 Normal = FSolverVec3::CrossProduct(P10, P20).GetSafeNormal();
-				Normal = (FSolverVec3::DotProduct(Normal, PP0) > 0) ? Normal : -Normal;
-
-				CriticalSection.Lock();
-				Constraints.Add({Index, Elem[0], Elem[1], Elem[2]});
-				Barys.Add(Bary);
-				Normals.Add(Normal);
-				CriticalSection.Unlock();
 			}
-		});
-#endif  // #if PLATFORM_DESKTOP && PLATFORM_64BITS
+		}
+	);
 }
 
 FSolverVec3 FPBDCollisionSpringConstraintsBase::GetDelta(const FSolverParticles& Particles, const int32 i) const
@@ -217,15 +125,14 @@ FSolverVec3 FPBDCollisionSpringConstraintsBase::GetDelta(const FSolverParticles&
 	const FSolverVec3& P4 = Particles.P(i4);
 
 	const FSolverReal Height = Thickness + Thickness;
-	const FSolverVec3 P = Barys[i][0] * P2 + Barys[i][1] * P3 + Barys[i][2] * P4 + Height * Normals[i];
-
+	const FSolverVec3 P = Barys[i][0] * P2 + Barys[i][1] * P3 + Barys[i][2] * P4;
 	const FSolverVec3 Difference = P1 - P;
-	if (FSolverVec3::DotProduct(Difference, Normals[i]) > 0)
+	const FSolverReal DistSq = Difference.SizeSquared();
+	if (DistSq > Height * Height)
 	{
 		return FSolverVec3(0);
 	}
-	const FSolverReal Distance = Difference.Size();
-	const FSolverVec3 Delta = Distance * Normals[i];
+	const FSolverVec3 Delta = Difference * Height * FMath::InvSqrt(DistSq) - Difference;
 	return Stiffness * Delta / CombinedMass;
 }
 
diff --git a/Engine/Source/Runtime/Experimental/Chaos/Private/Chaos/PBDEvolution.cpp b/Engine/Source/Runtime/Experimental/Chaos/Private/Chaos/PBDEvolution.cpp
--- a/Engine/Source/Runtime/Experimental/Chaos/Private/Chaos/PBDEvolution.cpp
+++ b/Engine/Source/Runtime/Experimental/Chaos/Private/Chaos/PBDEvolution.cpp
@@ -75,7 +75,6 @@ FPBDEvolution::FPBDEvolution(FSolverParticles&& InParticles, FSolverRigidParticl
 	, MParticlesActiveView(MParticles)
 	, MCollisionParticles(MoveTemp(InGeometryParticles))
 	, MCollisionParticlesActiveView(MCollisionParticles)
-	, MCollisionTriangles(MoveTemp(CollisionTriangles))
 	, MConstraintInitsActiveView(MConstraintInits)
 	, MConstraintRulesActiveView(MConstraintRules)
 	, MNumIterations(NumIterations)
diff --git a/Engine/Source/Runtime/Experimental/Chaos/Private/Chaos/TriangleMesh.cpp b/Engine/Source/Runtime/Experimental/Chaos/Private/Chaos/TriangleMesh.cpp
--- a/Engine/Source/Runtime/Experimental/Chaos/Private/Chaos/TriangleMesh.cpp
+++ b/Engine/Source/Runtime/Experimental/Chaos/Private/Chaos/TriangleMesh.cpp
@@ -4,6 +4,7 @@
 #include "Chaos/Box.h"
 #include "Chaos/Defines.h"
 #include "Chaos/Plane.h"
+#include "Chaos/TriangleCollisionPoint.h"
 #include "HAL/IConsoleManager.h"
 #include "Math/NumericLimits.h"
 #include "Math/RandomStream.h"
@@ -27,6 +28,36 @@ FAutoConsoleVariableRef CVarChaosTriangleMeshISPCEnabled(TEXT("p.Chaos.TriangleM
 
 namespace Chaos
 {
+template<typename T>
+struct TTriangleMeshBvEntry
+{
+	const FTriangleMesh* TmData;
+	const TConstArrayView<TVec3<T>>* Points;
+	int32 Index;
+
+	bool HasBoundingBox() const 
+	{
+		return true;
+	}
+
+	Chaos::TAABB<T,3> BoundingBox() const
+	{
+		const TVec3<int32>& Tri = TmData->GetElements()[Index];
+		const TVec3<T>& A = (*Points)[Tri[0]];
+		const TVec3<T>& B = (*Points)[Tri[1]];
+		const TVec3<T>& C = (*Points)[Tri[2]];
+
+		TAABB<T,3> Bounds(A, A);
+
+		Bounds.GrowToInclude(B);
+		Bounds.GrowToInclude(C);
+
+		return Bounds;
+	}
+
+	template<typename TPayloadType>
+	int32 GetPayload(int32 Idx) const { return Idx; }
+};
 
 FTriangleMesh::FTriangleMesh()
     : MStartIdx(0)
@@ -1426,4 +1457,70 @@ void FTriangleMesh::RemoveDegenerateElements()
 	}
 }
 
+template<typename T>
+void FTriangleMesh::BuildBVH(const TConstArrayView<TVec3<T>>& Points, TBVHType<T>& BVH) const
+{
+	TArray<TTriangleMeshBvEntry<T>> BVEntries;
+	const int32 NumTris = MElements.Num();
+	BVEntries.Reset(NumTris);
+	for (int32 Tri = 0; Tri < NumTris; ++Tri)
+	{
+		BVEntries.Add({ this, &Points, Tri });
+	}
+	BVH.Reinitialize(BVEntries);
+}
+template CHAOS_API void FTriangleMesh::BuildBVH<FRealSingle>(const TConstArrayView<TVec3<FRealSingle>>& Points, TBVHType<FRealSingle>& BVH) const;
+template CHAOS_API void FTriangleMesh::BuildBVH<FRealDouble>(const TConstArrayView<TVec3<FRealDouble>>& Points, TBVHType<FRealDouble>& BVH) const;
+
+template<typename T>
+bool FTriangleMesh::PointProximityQuery(const TBVHType<T>& BVH, const TConstArrayView<TVec3<T>>& Points, const int32 PointIndex, const TVec3<T>& PointPosition, const T PointThickness, const T ThisThickness, 
+	TFunctionRef<bool(const int32 PointIndex, const int32 TriangleIndex)> BroadphaseTest, TArray<TTriangleCollisionPoint<T>>& Result) const
+{
+	const T TotalThickness = ThisThickness + PointThickness;
+	const T TotalThicknessSq = TotalThickness * TotalThickness;
+	FAABB3 QueryBounds(PointPosition, PointPosition);
+	QueryBounds.Thicken(TotalThickness);
+
+	const TArray<int32> PotentialIntersections = BVH.FindAllIntersections(QueryBounds);
+
+	Result.Reset(PotentialIntersections.Num());
+
+	for (int32 TriIdx : PotentialIntersections)
+	{
+		if (!BroadphaseTest(PointIndex, TriIdx))
+		{
+			continue;
+		}
+
+		const TVec3<T>& A = Points[MElements[TriIdx][0]];
+		const TVec3<T>& B = Points[MElements[TriIdx][1]];
+		const TVec3<T>& C = Points[MElements[TriIdx][2]];
+		TVec3<T> Bary;
+		const TVec3<T> ClosestPoint = FindClosestPointAndBaryOnTriangle(A, B, C, PointPosition, Bary);
+
+		const T DistSq = (PointPosition - ClosestPoint).SizeSquared();
+		if (DistSq > TotalThicknessSq)
+		{
+			// Failed narrow test.
+			continue;
+		}
+
+		TVec3<T> Normal = TVec3<T>::CrossProduct(B-A, C-A).GetSafeNormal();
+		Normal = (TVec3<T>::DotProduct(Normal, PointPosition-A) > 0) ? Normal : -Normal;
+
+		TTriangleCollisionPoint<T> CollisionPoint;
+		CollisionPoint.ContactType = TTriangleCollisionPoint<T>::EContactType::PointFace;
+		CollisionPoint.Indices[0] = PointIndex;
+		CollisionPoint.Indices[1] = TriIdx;
+		CollisionPoint.Bary = TVec4<T>((T)1., Bary.X, Bary.Y, Bary.Z);
+		CollisionPoint.Location = ClosestPoint;
+		CollisionPoint.Normal = Normal;
+		CollisionPoint.Phi = FMath::Sqrt(DistSq);
+		Result.Add(CollisionPoint);
+	}
+	return Result.Num() > 0;
+}
+template CHAOS_API bool FTriangleMesh::PointProximityQuery<FRealSingle>(const TBVHType<FRealSingle>& BVH, const TConstArrayView<TVector<FRealSingle, 3>>& Points, const int32 PointIndex, const TVector<FRealSingle, 3>& PointPosition, const FRealSingle PointThickness, const FRealSingle ThisThickness, TFunctionRef<bool(const int32 PointIndex, const int32 TriangleIndex)> BroadphaseTest, TArray<TTriangleCollisionPoint<FRealSingle>>& Result) const;
+template CHAOS_API bool FTriangleMesh::PointProximityQuery<FRealDouble>(const TBVHType<FRealDouble>& BVH, const TConstArrayView<TVector<FRealDouble, 3>>& Points, const int32 PointIndex, const TVector<FRealDouble, 3>& PointPosition, const FRealDouble PointThickness, const FRealDouble ThisThickness, TFunctionRef<bool(const int32 PointIndex, const int32 TriangleIndex)> BroadphaseTest, TArray<TTriangleCollisionPoint<FRealDouble>>& Result) const;
+
 }  // End namespace Chaos
diff --git a/Engine/Source/Runtime/Experimental/Chaos/Public/Chaos/PBDCollisionSpringConstraints.h b/Engine/Source/Runtime/Experimental/Chaos/Public/Chaos/PBDCollisionSpringConstraints.h
--- a/Engine/Source/Runtime/Experimental/Chaos/Public/Chaos/PBDCollisionSpringConstraints.h
+++ b/Engine/Source/Runtime/Experimental/Chaos/Public/Chaos/PBDCollisionSpringConstraints.h
@@ -17,11 +17,12 @@ public:
 	FPBDCollisionSpringConstraints(
 		const int32 InOffset,
 		const int32 InNumParticles,
-		const TArray<TVec3<int32>>& InElements,
+		const FTriangleMesh& InTriangleMesh,
+		const TArray<FSolverVec3>* InRestPositions,
 		TSet<TVec2<int32>>&& InDisabledCollisionElements,
 		const FSolverReal InThickness = (FSolverReal)1.,
 		const FSolverReal InStiffness = (FSolverReal)1.)
-	    : Base(InOffset, InNumParticles, InElements, MoveTemp(InDisabledCollisionElements), InThickness, InStiffness)
+		: Base(InOffset, InNumParticles, InTriangleMesh, InRestPositions, MoveTemp(InDisabledCollisionElements), InThickness, InStiffness)
 	{}
 
 	virtual ~FPBDCollisionSpringConstraints() override {}
@@ -37,8 +38,8 @@ public:
 		const int32 i3 = Constraint[2];
 		const int32 i4 = Constraint[3];
 		const FSolverVec3 Delta = Base::GetDelta(Particles, i);
-		static const FSolverReal Multiplier = (FSolverReal)1.;  // TODO(mlentine): Figure out what the best multiplier here is
-		if (Particles.InvM(i1) > (FSolverReal)0.)
+		static const FSolverReal Multiplier = (FSolverReal)0.5;  // TODO(mlentine): Figure out what the best multiplier here is
+		if (Particles.InvM(i1) > 0)
 		{
 			Particles.P(i1) += Multiplier * Particles.InvM(i1) * Delta;
 		}
diff --git a/Engine/Source/Runtime/Experimental/Chaos/Public/Chaos/PBDCollisionSpringConstraintsBase.h b/Engine/Source/Runtime/Experimental/Chaos/Public/Chaos/PBDCollisionSpringConstraintsBase.h
--- a/Engine/Source/Runtime/Experimental/Chaos/Public/Chaos/PBDCollisionSpringConstraintsBase.h
+++ b/Engine/Source/Runtime/Experimental/Chaos/Public/Chaos/PBDCollisionSpringConstraintsBase.h
@@ -7,6 +7,12 @@
 #include "Chaos/PBDParticles.h"
 #include "Containers/Set.h"
 
+
+namespace Chaos
+{
+	class FTriangleMesh;
+}
+
 namespace Chaos::Softs
 {
 
@@ -14,10 +20,12 @@ namespace Chaos::Softs
 class CHAOS_API FPBDCollisionSpringConstraintsBase
 {
 public:
+
 	FPBDCollisionSpringConstraintsBase(
 		const int32 InOffset,
 		const int32 InNumParticles,
-		const TArray<TVec3<int32>>& InElements,
+		const FTriangleMesh& InTriangleMesh,
+		const TArray<FSolverVec3>* InReferencePositions,
 		TSet<TVec2<int32>>&& InDisabledCollisionElements,
 		const FSolverReal InThickness = (FSolverReal)1.,
 		const FSolverReal InStiffness = (FSolverReal)1.);
@@ -30,7 +38,6 @@ public:
 
 	const TArray<TVec4<int32>>& GetConstraints() const { return Constraints;  }
 	const TArray<FSolverVec3>& GetBarys() const { return Barys; }
-	const TArray<FSolverVec3>& GetNormals() const { return Normals; }
 	FSolverReal GetThickness() const { return Thickness; }
 
 	void SetThickness(FSolverReal InThickness) { Thickness = FMath::Max(InThickness, (FSolverReal)0.);  }
@@ -38,11 +45,13 @@ public:
 protected:
 	TArray<TVec4<int32>> Constraints;
 	TArray<FSolverVec3> Barys;
-	TArray<FSolverVec3> Normals;
 
 private:
+	const FTriangleMesh& TriangleMesh;
 	const TArray<TVec3<int32>>& Elements;
+	const TArray<FSolverVec3>* ReferencePositions;
 	const TSet<TVec2<int32>> DisabledCollisionElements;  // TODO: Make this a bitarray
+
 	int32 Offset;
 	int32 NumParticles;
 	FSolverReal Thickness;
diff --git a/Engine/Source/Runtime/Experimental/Chaos/Public/Chaos/PBDEvolution.h b/Engine/Source/Runtime/Experimental/Chaos/Public/Chaos/PBDEvolution.h
--- a/Engine/Source/Runtime/Experimental/Chaos/Public/Chaos/PBDEvolution.h
+++ b/Engine/Source/Runtime/Experimental/Chaos/Public/Chaos/PBDEvolution.h
@@ -93,10 +93,6 @@ class CHAOS_API FPBDEvolution : public TArrayCollection
 	FVelocityField& GetVelocityField(const uint32 GroupId = 0) { check(GroupId < TArrayCollection::Size()); return MGroupVelocityFields[GroupId]; }
 	const FVelocityField& GetVelocityField(const uint32 GroupId = 0) const { check(GroupId < TArrayCollection::Size()); return MGroupVelocityFields[GroupId]; }
 
-	void ResetSelfCollision() { MCollisionTriangles.Reset(); MDisabledCollisionElements.Reset(); };
-	TArray<TVector<int32, 3>>& CollisionTriangles() { return MCollisionTriangles; }
-	TSet<TVector<int32, 2>>& DisabledCollisionElements() { return MDisabledCollisionElements; }
-
 	int32 GetIterations() const { return MNumIterations; }
 	void SetIterations(const int32 Iterations) { MNumIterations = Iterations; }
 
@@ -139,9 +135,6 @@ private:
 	FSolverRigidParticles MCollisionParticles;
 	TPBDActiveView<FSolverRigidParticles> MCollisionParticlesActiveView;
 
-	TArray<TVector<int32, 3>> MCollisionTriangles;       // Used for self-collisions
-	TSet<TVector<int32, 2>> MDisabledCollisionElements;  // 
-
 	TArrayCollectionArray<FSolverRigidTransform3> MCollisionTransforms;  // Used for CCD to store the initial state before the kinematic update
 	TArrayCollectionArray<bool> MCollided;
 	TArrayCollectionArray<uint32> MCollisionParticleGroupIds;  // Used for per group parameters for collision particles
diff --git a/Engine/Source/Runtime/Experimental/Chaos/Public/Chaos/Plane.h b/Engine/Source/Runtime/Experimental/Chaos/Public/Chaos/Plane.h
--- a/Engine/Source/Runtime/Experimental/Chaos/Public/Chaos/Plane.h
+++ b/Engine/Source/Runtime/Experimental/Chaos/Public/Chaos/Plane.h
@@ -440,6 +440,176 @@ TVector<T, d> FindClosestPointOnTriangle(const TPlane<T, d>& TrianglePlane, cons
 	return FindClosestPointOnTriangle(PointOnPlane, P0, P1, P2, P);
 }
 
+// This method follows FindClosestPointOnTriangle but does less duplicate work and can handle degenerate triangles. It also returns the barycentric coordinates for the returned point.
+template<typename T, int d>
+TVector<T, d> FindClosestPointAndBaryOnTriangle(const TVector<T, d>& P0, const TVector<T, d>& P1, const TVector<T, d>& P2, const TVector<T, d>& P, TVector<T, 3>& Bary)
+{
+	const TVector<T, d> P10 = P1 - P0;
+	const TVector<T, d> P20 = P2 - P0;
+	const TVector<T, d> P21 = P2 - P1;
+	const TVector<T, d> PP0 = P - P0;
+	
+	const T Size10 = P10.SizeSquared();
+	const T Size20 = P20.SizeSquared();
+	const T Size21 = P21.SizeSquared();
+
+	const T ProjP1 = PP0.Dot(P10);
+	const T ProjP2 = PP0.Dot(P20);
+
+
+	auto ProjectToP01 = [&Bary, &P0, &P1, Size10, ProjP1]() -> TVector<T,d>
+	{
+		Bary.Z = 0.f;
+		Bary.Y = FMath::Clamp(ProjP1 / Size10, (T)0., (T)1.);
+		Bary.X = 1.f - Bary.Y;
+		const TVector<T,d> Result = Bary.X * P0 + Bary.Y * P1;
+		return Result;
+	};
+
+	auto ProjectToP02 = [&Bary, &P0, &P2, Size20, ProjP2]() -> TVector<T, d>
+	{
+		Bary.Y = 0.f;
+		Bary.Z = FMath::Clamp(ProjP2 / Size20, (T)0., (T)1.);
+		Bary.X = 1.f - Bary.Z;
+		const TVector<T, d> Result = Bary.X * P0 + Bary.Z * P2;
+		return Result;
+	};
+
+	auto ProjectToP12 = [&Bary, &P1, &P2, &P]() -> TVector<T, d>
+	{
+		const TVector<T,d> P2P1 = P2 - P1;
+
+		Bary.X = 0.f;
+		Bary.Z = FMath::Clamp(P2P1.Dot(P - P1) / P2P1.SizeSquared(), (T)0., (T)1.);
+		Bary.Y = 1.f - Bary.Z;
+		const TVector<T, d> Result = Bary.Y * P1 + Bary.Z * P2;
+		return Result;
+	};
+
+
+	// Degenerate triangles
+	if (Size10 < (T)DOUBLE_SMALL_NUMBER)
+	{
+		if (Size20 < (T)DOUBLE_SMALL_NUMBER)
+		{
+			// Triangle is (nearly) a single point.
+			Bary.X = (T)1.;
+			Bary.Y = Bary.Z = (T)0.;
+			return P0;
+		}
+
+		// Triangle is a line segment from P0(=P1) to P2. Project to that line segment.
+		return ProjectToP02();
+	}
+	if (Size20 < (T)DOUBLE_SMALL_NUMBER)
+	{
+		// Triangle is a line segment from P0(=P2) to P1. Project to that line segment.
+		return ProjectToP01();
+	}
+
+
+	const T ProjSides = P10.Dot(P20);
+	const T Denom = Size10 * Size20 - ProjSides * ProjSides;
+
+	if (Denom < (T)DOUBLE_SMALL_NUMBER)
+	{
+		// Triangle is a line segment from P0 to P1(=P2), or otherwise the 3 points are (nearly) colinear. Project to the longest edge.
+		if (Size21 > Size20)
+		{
+			return ProjectToP12();
+		}
+		else if (Size20 > Size10)
+		{
+			return ProjectToP02();
+		}
+		else
+		{
+			return ProjectToP01();
+		}
+	}
+
+	// Non-degenerate triangle
+
+	// Numerators of barycentric coordinates if P is inside triangle
+	const T BaryYNum = (Size20 * ProjP1 - ProjSides * ProjP2);
+	const T BaryZNum = (Size10 * ProjP2 - ProjSides * ProjP1);
+
+	// Using this specific epsilon to have parity with FindClosestPointOnTriangle
+	constexpr T Epsilon = (T)1e-4;
+	const T EpsilonScaledByDenom = Epsilon * Denom;
+
+	if (BaryYNum + BaryZNum > Denom * (1 + Epsilon)) // i.e., Bary.X < -Epsilon.
+	{
+		if (BaryYNum < -EpsilonScaledByDenom)
+		{
+			// Bary.X < 0, Bary.Y < 0, (Bary.Z > 1)
+
+			// We're outside the triangle in the region opposite angle P2. (if angle at P2 is > 90 degrees, closest point could be on one of these segments)
+			// Should project onto line segment P02 or P12
+			if (ProjP2 < Size20)
+			{
+				// Closer to P02 
+				// Project to line segment from P02
+				return ProjectToP02();
+			}
+		} 
+		else if (BaryZNum < -EpsilonScaledByDenom)
+		{
+			// Bary.X < 0, Bary.Y > 1, Bary.Z < 0
+			
+			// We're outside the triangle in the region opposite the angle P1.
+			// Should project onto line segment P01 or P12
+			if (ProjP1 < Size10)
+			{
+				// Closer to P01.
+				// Project to line segment from P01
+				return ProjectToP01();
+			}
+		}
+
+		// Bary.X < 0, Bary.Y >=0, Bary.Z >= 0
+
+		// We're outside the triangle in the region closest to P12
+		// Project to line segment from P12
+		return ProjectToP12();
+	}
+	else if (BaryYNum < -EpsilonScaledByDenom)
+	{
+		if (BaryZNum < -EpsilonScaledByDenom)
+		{
+			// Bary.X > 1, Bary.Y < 0, Bary.Z < 0
+
+			// We're outside the triangle in the region opposite P0.
+			// Should project onto line segment P01 or P02 
+			if (ProjP1 > 0.f)
+			{
+				// Closer to P01.
+				// Project to line segment from P01
+				return ProjectToP01();
+			}
+		}
+
+		// Bary.X >= 0, Bary.Y < 0, Bary.Z >= 0
+		// We're outside the triangle in region closest to P02.
+		// Project to line segment from P02
+		return ProjectToP02();
+	}
+	else if (BaryZNum < -EpsilonScaledByDenom)
+	{
+		// Bary.X >=0, Bary.Y >=0, Bary.Z < 0
+		// We're outside the triangle in the region closest to P01
+		// Project to line segment from P0 to P1
+		return ProjectToP01();
+	}
+	
+	// Bary.X >= 0, Bary.Y >= 0, Bary.Z >= 0
+	// Interior of triangle!
+	Bary.Y = BaryYNum / Denom;
+	Bary.Z = BaryZNum / Denom;
+	Bary.X = (T)1. - Bary.Y - Bary.Z;
+	const TVector<T,d> Result = Bary.X * P0 + Bary.Y * P1 + Bary.Z * P2;
+	return Result;
+}
 
 template<typename T, int d>
 bool IntersectPlanes2(TVector<T,d>& I, TVector<T,d>& D, const TPlane<T,d>& P1, const TPlane<T,d>& P2)
diff --git a/Engine/Source/Runtime/Experimental/Chaos/Public/Chaos/TriangleCollisionPoint.h b/Engine/Source/Runtime/Experimental/Chaos/Public/Chaos/TriangleCollisionPoint.h
new file mode 100644
--- /dev/null
+++ b/Engine/Source/Runtime/Experimental/Chaos/Public/Chaos/TriangleCollisionPoint.h
@@ -0,0 +1,52 @@
+// Copyright Epic Games, Inc. All Rights Reserved.
+#pragma once
+
+#include "CoreMinimal.h"
+#include "Templates/UnrealTemplate.h"
+
+#include "Chaos/Core.h"
+
+
+namespace Chaos
+{
+	/**
+	 * @brief Data returned by TriangleMesh point-triangle and edge-edge queries.
+	 */
+	template<typename T>
+	struct CHAOS_API TTriangleCollisionPoint
+	{
+		// Type of contact (determines meaning of some properties below such as Indices and BarycentricCoordinates)
+		enum struct EContactType : uint8
+		{
+			Invalid = 0,
+			PointFace,
+			EdgeEdge
+		} ContactType;
+
+		// Index into Point and Face arrays or Edge arrays
+		int32 Indices[2];
+
+		// Barycentric coordinates for contacts
+		// Point-Face: Bary[0] = 1.0 (Point), Bary[1:3] correspond with coordinates on Face
+		// Edge-Edge: Bary[0:1] = coordinates on first edge, [2:3] coordinates on second edge
+		TVec4<T> Bary;
+
+		// World space location (which point depends on query type)
+		TVec3<T> Location;
+
+		// World space normal (which normal depends on query type)
+		TVec3<T> Normal;
+
+		// Contact separation (negative for overlap)
+		T Phi;
+
+		TTriangleCollisionPoint()
+			: ContactType(EContactType::Invalid)
+			, Phi(TNumericLimits<T>::Max())
+		{
+		}
+
+		// Whether the contact point has been set up with contact data
+		bool IsSet() const { return ContactType != EContactType::Invalid; }
+	};
+}
diff --git a/Engine/Source/Runtime/Experimental/Chaos/Public/Chaos/TriangleMesh.h b/Engine/Source/Runtime/Experimental/Chaos/Public/Chaos/TriangleMesh.h
--- a/Engine/Source/Runtime/Experimental/Chaos/Public/Chaos/TriangleMesh.h
+++ b/Engine/Source/Runtime/Experimental/Chaos/Public/Chaos/TriangleMesh.h
@@ -7,8 +7,12 @@
 #include "Chaos/SegmentMesh.h"
 #include "Containers/ContainersFwd.h"
 
+#include "AABBTree.h"
+
 namespace Chaos
 {
+	template<typename T> struct TTriangleCollisionPoint;
+
 	class FTriangleMesh
 	{
 	public:
@@ -262,6 +266,17 @@ namespace Chaos
 
 			TriMesh.Init(MoveTemp(Elements));
 		}
+
+		template<typename T>
+		using TBVHType = TAABBTree<int32, TAABBTreeLeafArray<int32, /*bComputeBounds=*/false, T>, /*bMutable=*/true, T>;
+
+		template<typename T>
+		CHAOS_API void BuildBVH(const TConstArrayView<TVec3<T>>& Points, TBVHType<T>& BVH) const;
+
+		// NOTE: This method assumes the BVH has already been built/fitted to Points.
+		template<typename T>
+		CHAOS_API bool PointProximityQuery(const TBVHType<T>& BVH, const TConstArrayView<TVec3<T>>& Points, const int32 PointIndex, const TVec3<T>& PointPosition, const T PointThickness, const T ThisThickness, 
+			TFunctionRef<bool (const int32 PointIndex, const int32 TriangleIndex)> BroadphaseTest, TArray<TTriangleCollisionPoint<T>>& Result) const;
 		
 	private:
 		CHAOS_API void InitHelper(const int32 StartIdx, const int32 EndIdx, const bool CullDegenerateElements=true);
