Engine : Default Destruciton Physics Material
- Added a default material for destructable objects.

diff --git a/Engine/Plugins/Experimental/ChaosEditor/Source/FractureEditor/Private/FractureToolGenerators.cpp b/Engine/Plugins/Experimental/ChaosEditor/Source/FractureEditor/Private/FractureToolGenerators.cpp
--- a/Engine/Plugins/Experimental/ChaosEditor/Source/FractureEditor/Private/FractureToolGenerators.cpp
+++ b/Engine/Plugins/Experimental/ChaosEditor/Source/FractureEditor/Private/FractureToolGenerators.cpp
@@ -280,6 +280,7 @@ class AGeometryCollectionActor* UFractureToolGenerateAsset::CreateNewGeometryAct
 
 	// Set the Geometry Collection asset in the new actor
 	NewActor->GetGeometryCollectionComponent()->SetRestCollection(InGeometryCollection);
+	NewActor->GetGeometryCollectionComponent()->SetPhysMaterialOverride(GEngine->DefaultDestructiblePhysMaterial);
 
 	// copy transform of original static mesh actor to this new actor
 	NewActor->SetActorLabel(UniqueAssetName);
diff --git a/Engine/Plugins/Experimental/GeometryCollectionPlugin/Source/GeometryCollectionEditor/Private/GeometryCollection/ActorFactoryGeometryCollection.cpp b/Engine/Plugins/Experimental/GeometryCollectionPlugin/Source/GeometryCollectionEditor/Private/GeometryCollection/ActorFactoryGeometryCollection.cpp
--- a/Engine/Plugins/Experimental/GeometryCollectionPlugin/Source/GeometryCollectionEditor/Private/GeometryCollection/ActorFactoryGeometryCollection.cpp
+++ b/Engine/Plugins/Experimental/GeometryCollectionPlugin/Source/GeometryCollectionEditor/Private/GeometryCollection/ActorFactoryGeometryCollection.cpp
@@ -53,6 +53,7 @@ void UActorFactoryGeometryCollection::PostSpawnActor(UObject* Asset, AActor* New
 	NewGeometryCollectionActor->GetGeometryCollectionComponent()->EnableClustering = GeometryCollection->EnableClustering;
 	NewGeometryCollectionActor->GetGeometryCollectionComponent()->ClusterGroupIndex = GeometryCollection->ClusterGroupIndex;
 	NewGeometryCollectionActor->GetGeometryCollectionComponent()->MaxClusterLevel = GeometryCollection->MaxClusterLevel;
+	NewGeometryCollectionActor->GetGeometryCollectionComponent()->SetPhysMaterialOverride(GEngine->DefaultDestructiblePhysMaterial);
 
 	// Init Component
 	NewGeometryCollectionActor->GetGeometryCollectionComponent()->RegisterComponent();
diff --git a/Engine/Source/Runtime/Engine/Classes/Engine/Engine.h b/Engine/Source/Runtime/Engine/Classes/Engine/Engine.h
--- a/Engine/Source/Runtime/Engine/Classes/Engine/Engine.h
+++ b/Engine/Source/Runtime/Engine/Classes/Engine/Engine.h
@@ -1186,6 +1186,14 @@ public:
 	UPROPERTY(globalconfig)
 	FSoftObjectPath DefaultPhysMaterialName;
 
+	/** PhysicalMaterial to use if none is defined for a Destructible object. */
+	UPROPERTY()
+	TObjectPtr<class UPhysicalMaterial> DefaultDestructiblePhysMaterial;
+
+	/** Path of the PhysicalMaterial to use if none is defined for a particular object. */
+	UPROPERTY(globalconfig, EditAnywhere, Category = DefaultMaterials, meta = (AllowedClasses = "Physics", DisplayName = "Destructible Physics Material"))
+	FSoftObjectPath DefaultDestructiblePhysMaterialName;
+
 	/** Deprecated rules for redirecting renamed objects, replaced by the CoreRedirects system*/
 	UPROPERTY(config)
 	TArray<FGameNameRedirect> ActiveGameNameRedirects;
diff --git a/Engine/Source/Runtime/Engine/Private/UnrealEngine.cpp b/Engine/Source/Runtime/Engine/Private/UnrealEngine.cpp
--- a/Engine/Source/Runtime/Engine/Private/UnrealEngine.cpp
+++ b/Engine/Source/Runtime/Engine/Private/UnrealEngine.cpp
@@ -3025,9 +3025,28 @@ void UEngine::InitializeObjectReferences()
 			DefaultPhysMaterial = NewObject<UPhysicalMaterial>();
 		}
 	}
-
+	
 	UPhysicalMaterial::SetEngineDefaultPhysMaterial(DefaultPhysMaterial);
 
+	if (DefaultDestructiblePhysMaterial == NULL)
+	{
+		FString ValidDefaultDestructiblePhysMaterialName = DefaultDestructiblePhysMaterialName.ToString();
+		if (ValidDefaultDestructiblePhysMaterialName.IsEmpty())
+		{
+			UE_LOG(LogEngine, Warning, TEXT("The default destructible physics material was not found. Using default destructible physics material instead. Please make sure to have the default destructible physics material is set up correctly."));
+			ValidDefaultDestructiblePhysMaterialName = TEXT("/Engine/EngineMaterials/DefaultDestructiblePhysicalMaterial.DefaultDestructiblePhysicalMaterial");
+		}
+
+		DefaultDestructiblePhysMaterial = LoadObject<UPhysicalMaterial>(NULL, *ValidDefaultDestructiblePhysMaterialName, NULL, LOAD_None, NULL);
+		if (!DefaultDestructiblePhysMaterial)
+		{
+			UE_LOG(LogEngine, Error, TEXT("The default physical material for destruction (%s) was not found. Please make sure you have a default physical material fo destruction set up correctly."), *ValidDefaultDestructiblePhysMaterialName);
+			DefaultDestructiblePhysMaterial = NewObject<UPhysicalMaterial>();
+		}
+	}
+
+	UPhysicalMaterial::SetEngineDefaultDestructiblePhysMaterial(DefaultPhysMaterial);
+
 	LoadEngineClass<UConsole>(ConsoleClassName, ConsoleClass);
 	LoadEngineClass<UGameViewportClient>(GameViewportClientClassName, GameViewportClientClass);
 	LoadEngineClass<ULocalPlayer>(LocalPlayerClassName, LocalPlayerClass);
diff --git a/Engine/Source/Runtime/Experimental/GeometryCollectionEngine/Private/GeometryCollection/GeometryCollectionComponent.cpp b/Engine/Source/Runtime/Experimental/GeometryCollectionEngine/Private/GeometryCollection/GeometryCollectionComponent.cpp
--- a/Engine/Source/Runtime/Experimental/GeometryCollectionEngine/Private/GeometryCollection/GeometryCollectionComponent.cpp
+++ b/Engine/Source/Runtime/Experimental/GeometryCollectionEngine/Private/GeometryCollection/GeometryCollectionComponent.cpp
@@ -844,7 +844,7 @@ bool UGeometryCollectionComponent::DoCustomNavigableGeometryExport(FNavigableGeo
 UPhysicalMaterial* UGeometryCollectionComponent::GetPhysicalMaterial() const
 {
 	// Pull material from first mesh element to grab physical material. Prefer an override if one exists
-	UPhysicalMaterial* PhysMatToUse = PhysicalMaterialOverride;
+	UPhysicalMaterial* PhysMatToUse = BodyInstance.GetSimplePhysicalMaterial();
 
 	if(!PhysMatToUse)
 	{
@@ -1090,7 +1090,7 @@ void UGeometryCollectionComponent::PostEditChangeChainProperty(FPropertyChangedC
 	if (PropertyChangedEvent.Property && PropertyChangedEvent.Property->GetFName() == GET_MEMBER_NAME_CHECKED(UGeometryCollectionComponent, bShowBoneColors))
 	{
 		FScopedColorEdit EditBoneColor(this, true /*bForceUpdate*/); // the property has already changed; this will trigger the color update + render state updates
-	}
+	}					
 }
 #endif
 
@@ -3249,4 +3249,21 @@ bool UGeometryCollectionComponent::CalculateInnerSphere(int32 TransformIndex, FS
 		// Likely an embedded geometry, which doesn't count towards volume.
 		return false;
 	}
+}
+
+
+void UGeometryCollectionComponent::PostLoad()
+{
+	Super::PostLoad();
+
+	//
+	// The UGeometryCollectionComponent::PhysicalMaterial_DEPRECATED needs
+	// to be transferred to the BodyInstance simple material. Going forward
+	// the deprecated value will not be saved.
+	//
+	if (PhysicalMaterialOverride_DEPRECATED)
+	{
+		BodyInstance.SetPhysMaterialOverride(PhysicalMaterialOverride_DEPRECATED.Get());
+		PhysicalMaterialOverride_DEPRECATED = nullptr;
+	}
 }
\ No newline at end of file
diff --git a/Engine/Source/Runtime/Experimental/GeometryCollectionEngine/Public/GeometryCollection/GeometryCollectionComponent.h b/Engine/Source/Runtime/Experimental/GeometryCollectionEngine/Public/GeometryCollection/GeometryCollectionComponent.h
--- a/Engine/Source/Runtime/Experimental/GeometryCollectionEngine/Public/GeometryCollection/GeometryCollectionComponent.h
+++ b/Engine/Source/Runtime/Experimental/GeometryCollectionEngine/Public/GeometryCollection/GeometryCollectionComponent.h
@@ -375,6 +375,7 @@ public:
 	virtual void SetNotifyRigidBodyCollision(bool bNewNotifyRigidBodyCollision) override;
 	virtual bool CanEditSimulatePhysics() override;
 	virtual void SetSimulatePhysics(bool bEnabled) override;
+	virtual void PostLoad() override;
 	//~ End UPrimitiveComponent Interface.
 
 
@@ -518,8 +519,8 @@ public:
 	UPROPERTY(EditAnywhere, BlueprintReadWrite, Category = "ChaosPhysics|Initial Velocity")
 	FVector InitialAngularVelocity;
 
-	UPROPERTY(EditAnywhere, BlueprintReadOnly, Category="ChaosPhysics|Collisions")
-	TObjectPtr<UPhysicalMaterial> PhysicalMaterialOverride;
+	UPROPERTY(meta = (DeprecatedProperty, DeprecationMessage = "Physical material now derived from render materials, for instance overrides use Colliisons PhysicalMaterialOverride."))
+	TObjectPtr<UPhysicalMaterial> PhysicalMaterialOverride_DEPRECATED;
 
 	UPROPERTY()
 	FGeomComponentCacheParameters CacheParameters;
diff --git a/Engine/Source/Runtime/PhysicsCore/Private/PhysicalMaterial.cpp b/Engine/Source/Runtime/PhysicsCore/Private/PhysicalMaterial.cpp
--- a/Engine/Source/Runtime/PhysicsCore/Private/PhysicalMaterial.cpp
+++ b/Engine/Source/Runtime/PhysicsCore/Private/PhysicalMaterial.cpp
@@ -128,6 +128,13 @@ void UPhysicalMaterial::SetEngineDefaultPhysMaterial(UPhysicalMaterial* Material
 	GEngineDefaultPhysMaterial = Material;
 }
 
+static UPhysicalMaterial* GEngineDefaultDestructiblePhysMaterial = nullptr;
+
+void UPhysicalMaterial::SetEngineDefaultDestructiblePhysMaterial(UPhysicalMaterial* Material)
+{
+	GEngineDefaultDestructiblePhysMaterial = Material;
+}
+
 EPhysicalSurface UPhysicalMaterial::DetermineSurfaceType(UPhysicalMaterial const* PhysicalMaterial)
 {
 	if (PhysicalMaterial == NULL)
diff --git a/Engine/Source/Runtime/PhysicsCore/Public/PhysicalMaterials/PhysicalMaterial.h b/Engine/Source/Runtime/PhysicsCore/Public/PhysicalMaterials/PhysicalMaterial.h
--- a/Engine/Source/Runtime/PhysicsCore/Public/PhysicalMaterials/PhysicalMaterial.h
+++ b/Engine/Source/Runtime/PhysicsCore/Public/PhysicalMaterials/PhysicalMaterial.h
@@ -121,6 +121,8 @@ public:
 
 	static void SetEngineDefaultPhysMaterial(UPhysicalMaterial* Material);
 
+	static void SetEngineDefaultDestructiblePhysMaterial(UPhysicalMaterial* Material);
+
 	/** Determine Material Type from input PhysicalMaterial **/
 	static EPhysicalSurface DetermineSurfaceType(UPhysicalMaterial const* PhysicalMaterial);
 };
