Chaos - fix spurious vertex collisions in convex-trimesh collision detection
- reduce likelihood of producing an edge collision when a face collision will do

diff --git a/Engine/Source/Runtime/Experimental/Chaos/Private/Chaos/Collision/PBDCollisionSolver.h b/Engine/Source/Runtime/Experimental/Chaos/Private/Chaos/Collision/PBDCollisionSolver.h
--- a/Engine/Source/Runtime/Experimental/Chaos/Private/Chaos/Collision/PBDCollisionSolver.h
+++ b/Engine/Source/Runtime/Experimental/Chaos/Private/Chaos/Collision/PBDCollisionSolver.h
@@ -962,7 +962,7 @@ namespace Chaos
 	FORCEINLINE_DEBUGGABLE bool FPBDCollisionSolverManifoldPoint::ShouldSolveVelocity() const
 	{
 		// We ensure positive separating velocity for close contacts even if they didn't receive a pushout
-		return (NetImpulseNormal > FSolverReal(0)) || (WorldContactDeltaNormal < FSolverReal(0));
+		return (NetPushOutNormal > FSolverReal(0)) || (WorldContactDeltaNormal < FSolverReal(0));
 	}
 
 	//////////////////////////////////////////////////////////////////////////////////////////////////
diff --git a/Engine/Source/Runtime/Experimental/Chaos/Private/Chaos/CollisionOneShotManifolds.cpp b/Engine/Source/Runtime/Experimental/Chaos/Private/Chaos/CollisionOneShotManifolds.cpp
--- a/Engine/Source/Runtime/Experimental/Chaos/Private/Chaos/CollisionOneShotManifolds.cpp
+++ b/Engine/Source/Runtime/Experimental/Chaos/Private/Chaos/CollisionOneShotManifolds.cpp
@@ -23,6 +23,12 @@ namespace Chaos
 	FRealSingle Chaos_Collision_Manifold_PlaneContactNormalEpsilon = 0.001f;
 	FAutoConsoleVariableRef CVarChaos_Manifold_PlaneContactNormalEpsilon(TEXT("p.Chaos.Collision.Manifold.PlaneContactNormalEpsilon"), Chaos_Collision_Manifold_PlaneContactNormalEpsilon, TEXT("Normal tolerance used to distinguish face contacts from edge-edge contacts"));
 
+	FRealSingle Chaos_Collision_Manifold_TriangleContactNormalThreshold = 0.75f;	// About 40deg
+	FAutoConsoleVariableRef CVarChaos_Manifold_TriangleContactNormalThreshold(TEXT("p.Chaos.Collision.Manifold.TriangleNormalThreshold"), Chaos_Collision_Manifold_TriangleContactNormalThreshold, TEXT(""));
+
+	FRealSingle Chaos_Collision_Manifold_EdgeContactNormalThreshold = 0.9f;	// About 25deg
+	FAutoConsoleVariableRef CVarChaos_Manifold_EdgeContactNormalThreshold(TEXT("p.Chaos.Collision.Manifold.EdgeNormalThreshold"), Chaos_Collision_Manifold_EdgeContactNormalThreshold, TEXT(""));
+
 	// @todo(chaos): TEMP - use convex-convex collisio detection for box-box until TBox::GetClosestEdgePosition is implemented for that path (without plane hint)
 	bool bChaos_Collision_Manifold_BoxAsConvex = true;
 	FAutoConsoleVariableRef CVarChaosCollisioConvexManifoldBoxAsConvex(TEXT("p.Chaos.Collision.Manifold.BoxAsConvex"), bChaos_Collision_Manifold_BoxAsConvex, TEXT(""));
@@ -245,6 +251,8 @@ namespace Chaos
 		// Remove spurious edge contact by removing edge-edge contact that are in the plane of another vertex-plane contact.
 		// Note: this assumes that the first shape in the contacts is the same convex shape (e.g., a box), and the second
 		// can be from any other shape(s), e.g., triangles of a non-convex shape.
+		// We treat PlaneVertex contacts the same because these are contacts on the convex face with a vertex of the triangle
+		// which require the same culling rule be applied.
 		template<bool bMaintainOrder>
 		void PruneEdgeContactPointsImp(TArray<FContactPoint>& ContactPoints, const FReal MaxPlaneDistance)
 		{
@@ -257,17 +265,17 @@ namespace Chaos
 			for (int32 EdgeContactIndex = 0; EdgeContactIndex < NumContactPoints; ++EdgeContactIndex)
 			{
 				const FContactPoint& EdgeContactPoint = ContactPoints[EdgeContactIndex];
+				const bool bIsTriangleEdgeOrVertex = (EdgeContactPoint.ContactType == EContactPointType::EdgeEdge) || (EdgeContactPoint.ContactType == EContactPointType::PlaneVertex);
 
-				if (EdgeContactPoint.ContactType == EContactPointType::EdgeEdge)
+				if (bIsTriangleEdgeOrVertex)
 				{
 					for (int32 PlaneContactIndex = 0; PlaneContactIndex < NumContactPoints; ++PlaneContactIndex)
 					{
 						const FContactPoint PlaneContactPoint = ContactPoints[PlaneContactIndex];
+						const bool bIsTriangleFace = (PlaneContactPoint.ContactType == EContactPointType::VertexPlane);
 
-						// Does the edge-edge contact lie in the plane of the vertex-plane contact...?
-						// We only consider Vertex-Plane contacts. We ignore Plane-Vertex because the first
-						// shape is always the convex object.
-						if (PlaneContactPoint.ContactType == EContactPointType::VertexPlane)
+						// If this is a triangle face contact, does the edge/vertex contact lie in the plane?
+						if (bIsTriangleFace)
 						{
 							const int32 PlaneShapeIndex = 1;
 							const FVec3 DeltaPos = PlaneContactPoint.ShapeContactPoints[PlaneShapeIndex] - EdgeContactPoint.ShapeContactPoints[PlaneShapeIndex];
@@ -1043,46 +1051,32 @@ namespace Chaos
 
 			PHYSICS_CSV_CUSTOM_EXPENSIVE(PhysicsCounters, NumManifoldsCreated, 1, ECsvCustomStatOp::Accumulate);
 
+			// See if the triangle face is a good choice for the manifold plane. We preferentially use the triangle face with
+			// a high threshold to avoid spurious edge/vertex collisions.
+			const FVec3 TrianglePlaneNormal = Triangle.GetNormal();
+			const FVec3 TrianglePlanePosition = Triangle[0];
+			const FReal TrianglePlaneNormalDotContactNormal = FVec3::DotProduct(TrianglePlaneNormal, GJKContactPoint.ShapeContactNormal);
+
+			// Find the best opposing plane on the convex
 			// @todo(chaos): handle zero margins...
 			const FReal PlaneSearchDistance = ConvexMargin;
+			const int32 ConvexPlaneIndex = SelectContactPlane(Convex, GJKContactPoint.ShapeContactPoints[0], GJKContactPoint.ShapeContactNormal, PlaneSearchDistance, VertexIndexA);
+			FVec3 ConvexPlanePosition, ConvexPlaneNormal;
+			Convex.GetPlaneNX(ConvexPlaneIndex, ConvexPlaneNormal, ConvexPlanePosition);
+			const FReal ConvexPlaneNormalDotContactNormal = FVec3::DotProduct(ConvexPlaneNormal, GJKContactPoint.ShapeContactNormal);
 
-			const FVec3 SeparationDirectionLocalConvex1 = GJKContactPoint.ShapeContactNormal;
-			const int32 MostOpposingPlaneIndexConvex1 = SelectContactPlane(Convex, GJKContactPoint.ShapeContactPoints[0], SeparationDirectionLocalConvex1, PlaneSearchDistance, VertexIndexA);
-			FVec3 BestPlanePosition1, BestPlaneNormal1;
-			Convex.GetPlaneNX(MostOpposingPlaneIndexConvex1, BestPlaneNormal1, BestPlanePosition1);
-			const FReal BestPlaneDotNormalConvex1 = FMath::Abs(FVec3::DotProduct(-SeparationDirectionLocalConvex1, BestPlaneNormal1));
-
-			// @todo(chaos): can be simpler for FTriangle
-			const FVec3 SeparationDirectionLocalConvex2 = GJKContactPoint.ShapeContactNormal;
-			const int32 MostOpposingPlaneIndexConvex2 = SelectContactPlane(Triangle, GJKContactPoint.ShapeContactPoints[1], -SeparationDirectionLocalConvex2, PlaneSearchDistance, VertexIndexB);
-			FVec3 BestPlanePosition2, BestPlaneNormal2;
-			Triangle.GetPlaneNX(MostOpposingPlaneIndexConvex2, BestPlaneNormal2, BestPlanePosition2);
-			const FReal BestPlaneDotNormalConvex2 = FMath::Abs(FVec3::DotProduct(SeparationDirectionLocalConvex2, BestPlaneNormal2));
-
-			const FReal SmallBiasToPreventFeatureFlipping = 0.002f; // This improves frame coherence by penalizing convex 1 in favour of convex 2
-			bool ReferenceFaceConvex1 = true; // Is the reference face on convex1 or convex2?
-			if (BestPlaneDotNormalConvex2 + SmallBiasToPreventFeatureFlipping > BestPlaneDotNormalConvex1)
-			{
-				ReferenceFaceConvex1 = false;
-			}
-
-			// Is this a vertex-plane or edge-edge contact? 
-			const FReal PlaneContactNormalEpsilon = Chaos_Collision_Manifold_PlaneContactNormalEpsilon;
-			const bool bIsPlaneContact = FMath::IsNearlyEqual(BestPlaneDotNormalConvex1, (FReal)1., PlaneContactNormalEpsilon) || FMath::IsNearlyEqual(BestPlaneDotNormalConvex2, (FReal)1., PlaneContactNormalEpsilon);
-
+			// Is either plane a good choice, or do we have and edge-edge collision
 			// For edge-edge contacts, we find the edges involved and project the contact onto the edges (if necessary - see below)
+			const FReal PlaneContactNormalEpsilon = Chaos_Collision_Manifold_PlaneContactNormalEpsilon;
+			bool bIsPlaneContact = FMath::IsNearlyEqual(ConvexPlaneNormalDotContactNormal, FReal(-1), PlaneContactNormalEpsilon) || FMath::IsNearlyEqual(TrianglePlaneNormalDotContactNormal, FReal(1), PlaneContactNormalEpsilon);
 			if (!bIsPlaneContact)
 			{
 				SCOPE_CYCLE_COUNTER_MANIFOLD_ADDEDGEEDGE();
-
 				if (ForceOneShotManifoldEdgeEdgeCaseZeroCullDistance && GJKContactPoint.Phi > 0)
 				{
 					return;
 				}
 
-				// This is an edge-edge contact
-				GJKContactPoint.ContactType = EContactPointType::EdgeEdge;
-
 				// @todo(chaos): remove this if we ditch convex margins
 				// We only need to find the best edges when we have convex margins enabled because then the GJK result
 				// is a contact on the rounded margin-reduced shape and not the actual shape. If we are using zero
@@ -1090,80 +1084,95 @@ namespace Chaos
 				// quadratic-versus-polygonal contact, the GJK results are already on the real surfaces.
 				if (ConvexMargin > 0)
 				{
-					// @todo(chaos): this does not work well when the edges are parallel. We should always have points with zero
-					// position delta perpendicular to the normal, but that is not the case for parallel edges
-					FVec3 ShapeEdgePos1 = Convex.GetClosestEdgePosition(MostOpposingPlaneIndexConvex1, GJKContactPoint.ShapeContactPoints[0]);
-					FVec3 ShapeEdgePos2 = Triangle.GetClosestEdgePosition(MostOpposingPlaneIndexConvex2, GJKContactPoint.ShapeContactPoints[1]);
+					FVec3 ConvexEdgePos0, ConvexEdgePos1;
+					const FVec3 ConvexEdgeNearPos = Convex.GetClosestEdge(ConvexPlaneIndex, GJKContactPoint.ShapeContactPoints[0], ConvexEdgePos0, ConvexEdgePos1);
 
-					const FVec3 EdgePos1In2 = ShapeEdgePos1;
-					const FVec3 EdgePos2In2 = ShapeEdgePos2;
-					const FReal EdgePhi = FVec3::DotProduct(EdgePos1In2 - EdgePos2In2, GJKContactPoint.ShapeContactNormal);
+					FVec3 TriangleEdgePos0, TriangleEdgePos1;
+					const FVec3 TriangleEdgeNearPos = Triangle.GetClosestEdge(0, GJKContactPoint.ShapeContactPoints[1], TriangleEdgePos0, TriangleEdgePos1);
 
-					// We now have an accurate separation, so reject points beyond cull distance.
-					// This is quite important because if we use a larger cull distance for edge-edge contacts than for vertex-plane
-					// contacts, the pruning of unnecessary edge-edge contacts does not work (Tri Mesh and heighfield)
-					if (EdgePhi > CullDistance)
+					// Does the normal we found from the edges match the GJK normal? If not, we have a false negative in the check for
+					// bIsPlaneContact above caused by margins. If this happens, fall back to a plane contact.
+					const FVec3 EdgeNormal = FVec3::CrossProduct(ConvexEdgePos1 - ConvexEdgePos0, TriangleEdgePos1 - TriangleEdgePos0).GetSafeNormal();
+					const FReal EdgeNormalDotContactNormal = FVec3::DotProduct(EdgeNormal, GJKContactPoint.ShapeContactNormal);
+					bIsPlaneContact = (FMath::Abs(EdgeNormalDotContactNormal) < Chaos_Collision_Manifold_EdgeContactNormalThreshold);
+
+					if (!bIsPlaneContact)
 					{
-						return;
-					}
+						const FReal EdgePhi = FVec3::DotProduct(ConvexEdgeNearPos - TriangleEdgeNearPos, GJKContactPoint.ShapeContactNormal);
 
-					// @todo(chaos): we leave the contact point on the second shape because this is better for triangle meshes (see PruneContactEdges). For other shapes it should probably be the average
-					GJKContactPoint.ShapeContactPoints[0] = EdgePos2In2 + EdgePhi * GJKContactPoint.ShapeContactNormal;
-					GJKContactPoint.ShapeContactPoints[1] = EdgePos2In2;
-					GJKContactPoint.Phi = EdgePhi;
-					// Normal unchanged from GJK result
+						// We now have an accurate separation, so reject points beyond cull distance.
+						// This is quite important because if we use a larger cull distance for edge-edge contacts than for vertex-plane
+						// contacts, the pruning of unnecessary edge-edge contacts does not work (Tri Mesh and heighfield)
+						if (EdgePhi > CullDistance)
+						{
+							return;
+						}
+
+						// Remove anmy lateral delta from the contact on point on the convex
+						GJKContactPoint.ShapeContactPoints[0] = TriangleEdgeNearPos + EdgePhi * GJKContactPoint.ShapeContactNormal;
+						GJKContactPoint.ShapeContactPoints[1] = TriangleEdgeNearPos;
+						// Normal unchanged from GJK result
+						GJKContactPoint.Phi = EdgePhi;
+					}
 				}
 
-				OutContactPoints.Add(GJKContactPoint);
-				return;
+				if (!bIsPlaneContact)
+				{
+					GJKContactPoint.ContactType = EContactPointType::EdgeEdge;
+					OutContactPoints.Add(GJKContactPoint);
+					return;
+				}
 			}
 
-			// For vertex-plane contacts, we use a convex face as the manifold plane
-			const FVec3 RefSeparationDirection = ReferenceFaceConvex1 ? SeparationDirectionLocalConvex1 : SeparationDirectionLocalConvex2;
-			const FVec3 RefPlaneNormal = ReferenceFaceConvex1 ? BestPlaneNormal1 : BestPlaneNormal2;
-			const FVec3 RefPlanePosition = ReferenceFaceConvex1 ? BestPlanePosition1 : BestPlanePosition2;
+			// We prefer to use the triangle face for the contact to avoid collisions with internal edges and vertices
+			const EContactPointType ContactType = (TrianglePlaneNormalDotContactNormal > Chaos_Collision_Manifold_TriangleContactNormalThreshold) ? EContactPointType::VertexPlane : EContactPointType::PlaneVertex;
 
 			// @todo(chaos): fix use of hard-coded max array size
 			// We will use a double buffer as an optimization
 			const int32 MaxContactPointCount = 32; // This should be tuned
 			int32 ContactPointCount = 0;
+			FVec3 RefPlaneNormal;
 			FVec3* ClippedVertices = nullptr;
 			FVec3 ClippedVertices1[MaxContactPointCount];
 			FVec3 ClippedVertices2[MaxContactPointCount];
 
+			// Generate the set of vertices by clipping the convex face against the triangle face in the selected face plane
+			if (ContactType == EContactPointType::VertexPlane)
 			{
-				SCOPE_CYCLE_COUNTER_MANIFOLD_CLIP();
-				if (ReferenceFaceConvex1)
-				{
-					// @todo(chaos): non-transforming version for triangles
-					ClippedVertices = GenerateConvexManifoldClippedVertices(
-						Convex,
-						Triangle,
-						FRigidTransform3(),
-						MostOpposingPlaneIndexConvex1,
-						MostOpposingPlaneIndexConvex2,
-						RefPlaneNormal,
-						RefPlanePosition,
-						ClippedVertices1,
-						ClippedVertices2,
-						(uint32&)ContactPointCount,
-						(uint32)MaxContactPointCount);
-				}
-				else
-				{
-					ClippedVertices = GenerateConvexManifoldClippedVertices(
-						Triangle,
-						Convex,
-						FRigidTransform3(),
-						MostOpposingPlaneIndexConvex2,
-						MostOpposingPlaneIndexConvex1,
-						RefPlaneNormal,
-						RefPlanePosition,
-						ClippedVertices1,
-						ClippedVertices2,
-						(uint32&)ContactPointCount,
-						(uint32)MaxContactPointCount);
-				}
+				// We have a triangle face contact
+				RefPlaneNormal = TrianglePlaneNormal;
+
+				ClippedVertices = GenerateConvexManifoldClippedVertices(
+					Triangle,
+					Convex,
+					FRigidTransform3(),
+					0,
+					ConvexPlaneIndex,
+					TrianglePlaneNormal,
+					TrianglePlanePosition,
+					ClippedVertices1,
+					ClippedVertices2,
+					(uint32&)ContactPointCount,
+					(uint32)MaxContactPointCount);
+
+			}
+			else
+			{
+				// We have a convex face contact
+				RefPlaneNormal = ConvexPlaneNormal;
+
+				ClippedVertices = GenerateConvexManifoldClippedVertices(
+					Convex,
+					Triangle,
+					FRigidTransform3(),
+					ConvexPlaneIndex,
+					0,
+					ConvexPlaneNormal,
+					ConvexPlanePosition,
+					ClippedVertices1,
+					ClippedVertices2,
+					(uint32&)ContactPointCount,
+					(uint32)MaxContactPointCount);
 			}
 
 			// If we have the max number of contact points already, they will be in cyclic order. Stability is better
@@ -1194,30 +1203,45 @@ namespace Chaos
 			}
 
 			// Generate the contact points from the clipped vertices
+			check(ContactPointCount + OutContactPoints.Num() <= OutContactPoints.Max());
+
+			if (ContactType == EContactPointType::VertexPlane)
 			{
+				// Triangle face contact (clipped vertices are convex vertices)
 				SCOPE_CYCLE_COUNTER_MANIFOLD_ADDFACEVERTEX();
-				check(ContactPointCount + OutContactPoints.Num() <= OutContactPoints.Max());
-
-				EContactPointType ContactType = ReferenceFaceConvex1 ? EContactPointType::PlaneVertex : EContactPointType::VertexPlane;
+				for (int32 ContactPointIndex = 0; ContactPointIndex < ContactPointCount; ++ContactPointIndex)
+				{
+					FContactPoint& ContactPoint = OutContactPoints[OutContactPoints.Add()];
+					const FVec3& ConvexContactPoint = ClippedVertices[ContactPointIndex];
+					const FVec3 TriangleContactPoint = ConvexContactPoint - FVec3::DotProduct(ConvexContactPoint - TrianglePlanePosition, TrianglePlaneNormal) * TrianglePlaneNormal;
 
+					ContactPoint.ShapeContactPoints[0] = ConvexContactPoint;
+					ContactPoint.ShapeContactPoints[1] = TriangleContactPoint;
+					ContactPoint.ShapeContactNormal = TrianglePlaneNormal;
+					ContactPoint.Phi = FVec3::DotProduct(ConvexContactPoint - TriangleContactPoint, TrianglePlaneNormal);
+					ContactPoint.ContactType = ContactType;
+				}
+			}
+			else
+			{
+				// Convex face contact (clipped vertices are triangle vertices)
+				SCOPE_CYCLE_COUNTER_MANIFOLD_ADDFACEVERTEX();
 				for (int32 ContactPointIndex = 0; ContactPointIndex < ContactPointCount; ++ContactPointIndex)
 				{
 					FContactPoint& ContactPoint = OutContactPoints[OutContactPoints.Add()];
-					FVec3 VertexInReferenceCoordinates = ClippedVertices[ContactPointIndex];
-					FVec3 PointProjectedOntoReferenceFace = VertexInReferenceCoordinates - FVec3::DotProduct(VertexInReferenceCoordinates - RefPlanePosition, RefPlaneNormal) * RefPlaneNormal;
-					FVec3 ClippedPointInOtherCoordinates = VertexInReferenceCoordinates;
+					const FVec3& TriangleContactPoint = ClippedVertices[ContactPointIndex];
+					const FVec3 ConvexContactPoint = TriangleContactPoint - FVec3::DotProduct(TriangleContactPoint - ConvexPlanePosition, ConvexPlaneNormal) * ConvexPlaneNormal;
 
-					ContactPoint.ShapeContactPoints[0] = ReferenceFaceConvex1 ? PointProjectedOntoReferenceFace : ClippedPointInOtherCoordinates;
-					ContactPoint.ShapeContactPoints[1] = ReferenceFaceConvex1 ? ClippedPointInOtherCoordinates : PointProjectedOntoReferenceFace;
-					ContactPoint.ShapeContactNormal = SeparationDirectionLocalConvex2;
-					ContactPoint.Phi = FVec3::DotProduct(PointProjectedOntoReferenceFace - VertexInReferenceCoordinates, ReferenceFaceConvex1 ? SeparationDirectionLocalConvex1 : -SeparationDirectionLocalConvex2);
+					ContactPoint.ShapeContactPoints[0] = ConvexContactPoint;
+					ContactPoint.ShapeContactPoints[1] = TriangleContactPoint;
+					ContactPoint.ShapeContactNormal = -ConvexPlaneNormal;
+					ContactPoint.Phi = FVec3::DotProduct(ConvexContactPoint - TriangleContactPoint, -ConvexPlaneNormal);
 					ContactPoint.ContactType = ContactType;
 				}
 			}
 		}
 
 
-
 		//
 		// Explicit instantiations of all convex-convex manifold combinations we support
 		// Box, Convex, Scaled-Convex
diff --git a/Engine/Source/Runtime/Experimental/Chaos/Private/Chaos/CollisionOneShotManifoldsMiscShapes.cpp b/Engine/Source/Runtime/Experimental/Chaos/Private/Chaos/CollisionOneShotManifoldsMiscShapes.cpp
--- a/Engine/Source/Runtime/Experimental/Chaos/Private/Chaos/CollisionOneShotManifoldsMiscShapes.cpp
+++ b/Engine/Source/Runtime/Experimental/Chaos/Private/Chaos/CollisionOneShotManifoldsMiscShapes.cpp
@@ -409,7 +409,10 @@ namespace Chaos
 					}
 				}
 				ContactPoints.SetNum(NewContactPointCount, false);
+			}
 
+			if (ContactPoints.Num() > 4)
+			{
 				// Reduce to only 4 contact points from here
 				// NOTE: relies on the sort above
 				Collisions::ReduceManifoldContactPointsTriangeMesh(ContactPoints);
diff --git a/Engine/Source/Runtime/Experimental/Chaos/Private/Chaos/Convex.cpp b/Engine/Source/Runtime/Experimental/Chaos/Private/Chaos/Convex.cpp
--- a/Engine/Source/Runtime/Experimental/Chaos/Private/Chaos/Convex.cpp
+++ b/Engine/Source/Runtime/Experimental/Chaos/Private/Chaos/Convex.cpp
@@ -209,7 +209,7 @@ namespace Chaos
 		return MostOpposingIdx;
 	}
 
-	FVec3 FConvex::GetClosestEdgePosition(int32 PlaneIndex, const FVec3& InPosition) const
+	FVec3 FConvex::GetClosestEdge(int32 PlaneIndex, const FVec3& InPosition, FVec3& OutEdgePos0, FVec3& OutEdgePos1) const
 	{
 		FVec3Type ClosestEdgePosition = FVec3Type(0);
 		FRealType ClosestDistanceSq = FLT_MAX;
@@ -238,6 +238,8 @@ namespace Chaos
 				{
 					ClosestDistanceSq = EdgeDistanceSq;
 					ClosestEdgePosition = EdgePosition;
+					OutEdgePos0 = P0;
+					OutEdgePos1 = P1;
 				}
 
 				P0 = P1;
diff --git a/Engine/Source/Runtime/Experimental/Chaos/Private/Chaos/PBDCollisionConstraints.cpp b/Engine/Source/Runtime/Experimental/Chaos/Private/Chaos/PBDCollisionConstraints.cpp
--- a/Engine/Source/Runtime/Experimental/Chaos/Private/Chaos/PBDCollisionConstraints.cpp
+++ b/Engine/Source/Runtime/Experimental/Chaos/Private/Chaos/PBDCollisionConstraints.cpp
@@ -691,10 +691,13 @@ namespace Chaos
 				if (EdgeCollision.IsEnabled())
 				{
 					const int32 EdgeOtherShapeIndex = (EdgeCollision.GetParticle0() == Particle) ? 1 : 0;
+					const EContactPointType VertexContactType = (EdgeOtherShapeIndex == 0) ? EContactPointType::VertexPlane : EContactPointType::PlaneVertex;
 
 					for (const FManifoldPoint& EdgeManifoldPoint : EdgeCollision.GetManifoldPoints())
 					{
-						if (EdgeManifoldPoint.ContactPoint.ContactType == EContactPointType::EdgeEdge)
+						const bool bIsEdgeContact = (EdgeManifoldPoint.ContactPoint.ContactType == EContactPointType::EdgeEdge);
+
+						if (bIsEdgeContact)
 						{
 							const FRigidTransform3& EdgeTransform = (EdgeOtherShapeIndex == 0) ? EdgeCollision.GetShapeWorldTransform0() : EdgeCollision.GetShapeWorldTransform1();
 							const FVec3 EdgePos = EdgeTransform.TransformPositionNoScale(EdgeManifoldPoint.ContactPoint.ShapeContactPoints[EdgeOtherShapeIndex]);
@@ -706,27 +709,17 @@ namespace Chaos
 									if ((&PlaneCollision != &EdgeCollision) && PlaneCollision.IsEnabled())
 									{
 										const int32 PlaneOtherShapeIndex = (PlaneCollision.GetParticle0() == Particle) ? 1 : 0;
+										const EContactPointType PlaneContactType = (PlaneOtherShapeIndex == 0) ? EContactPointType::PlaneVertex : EContactPointType::VertexPlane;
+										const FRigidTransform3& PlaneTransform = (PlaneOtherShapeIndex == 0) ? PlaneCollision.GetShapeWorldTransform0() : PlaneCollision.GetShapeWorldTransform1();
 
 										for (const FManifoldPoint& PlaneManifoldPoint : PlaneCollision.GetManifoldPoints())
 										{
-											// If the edge position is in the plane, disable it
-											FVec3 PlanePos;
-											FVec3 PlaneNormal;
-											bool bHavePlane = false;
-											if ((PlaneOtherShapeIndex == 0) && (PlaneManifoldPoint.ContactPoint.ContactType == EContactPointType::PlaneVertex))
-											{
-												PlanePos = PlaneCollision.GetShapeWorldTransform0().TransformPositionNoScale(PlaneManifoldPoint.ContactPoint.ShapeContactPoints[0]);
-												bHavePlane = true;
-											}
-											else if ((PlaneOtherShapeIndex == 1) && (PlaneManifoldPoint.ContactPoint.ContactType == EContactPointType::VertexPlane))
+											if (PlaneManifoldPoint.ContactPoint.ContactType == PlaneContactType)
 											{
-												PlanePos = PlaneCollision.GetShapeWorldTransform1().TransformPositionNoScale(PlaneManifoldPoint.ContactPoint.ShapeContactPoints[1]);
-												bHavePlane = true;
-											}
-											PlaneNormal = PlaneCollision.GetShapeWorldTransform1().TransformVectorNoScale(PlaneManifoldPoint.ContactPoint.ShapeContactNormal);
+												// If the edge position is in the plane, disable it
+												const FVec3 PlanePos = PlaneTransform.TransformPositionNoScale(PlaneManifoldPoint.ContactPoint.ShapeContactPoints[PlaneOtherShapeIndex]);
+												const FVec3 PlaneNormal = PlaneCollision.GetShapeWorldTransform1().TransformVectorNoScale(PlaneManifoldPoint.ContactPoint.ShapeContactNormal);
 
-											if (bHavePlane)
-											{
 												const FVec3 EdgePlaneDelta = EdgePos - PlanePos;
 												const FReal EdgePlaneDistance = FVec3::DotProduct(EdgePlaneDelta, PlaneNormal);
 												if (FMath::Abs(EdgePlaneDistance) < EdgePlaneTolerance)
diff --git a/Engine/Source/Runtime/Experimental/Chaos/Public/Chaos/Box.h b/Engine/Source/Runtime/Experimental/Chaos/Public/Chaos/Box.h
--- a/Engine/Source/Runtime/Experimental/Chaos/Public/Chaos/Box.h
+++ b/Engine/Source/Runtime/Experimental/Chaos/Public/Chaos/Box.h
@@ -220,8 +220,9 @@ namespace Chaos
 			return GetMostOpposingPlane(Normal);
 		}
 
-		// Get the nearest point on an edge
-		TVector<T, d> GetClosestEdgePosition(int32 PlaneIndexHint, const TVector<T, d>& Position) const
+		// Get the nearest point on an edge and the edge vertices
+		// Used for manifold generation
+		FVec3 GetClosestEdge(int32 PlaneIndexHint, const FVec3& Position, FVec3& OutEdgePos0, FVec3& OutEdgePos1) const
 		{
 			TVector<T, d> ClosestEdgePosition = FVec3(0);
 			if (PlaneIndexHint >= 0)
@@ -244,6 +245,8 @@ namespace Chaos
 						{
 							ClosestDistanceSq = EdgeDistanceSq;
 							ClosestEdgePosition = EdgePosition;
+							OutEdgePos0 = P0;
+							OutEdgePos1 = P1;
 						}
 
 						P0 = P1;
@@ -258,6 +261,13 @@ namespace Chaos
 			return ClosestEdgePosition;
 		}
 
+		// Get the nearest point on an edge
+		TVector<T, d> GetClosestEdgePosition(int32 PlaneIndexHint, const TVector<T, d>& Position) const
+		{
+			TVector<T, d> Unused0, Unused1;
+			return GetClosestEdge(PlaneIndexHint, Position, Unused0, Unused1);
+		}
+
 		bool GetClosestEdgeVertices(int32 PlaneIndexHint, const FVec3& Position, int32& OutVertexIndex0, int32& OutVertexIndex1) const
 		{
 			if (PlaneIndexHint >= 0)
diff --git a/Engine/Source/Runtime/Experimental/Chaos/Public/Chaos/Convex.h b/Engine/Source/Runtime/Experimental/Chaos/Public/Chaos/Convex.h
--- a/Engine/Source/Runtime/Experimental/Chaos/Public/Chaos/Convex.h
+++ b/Engine/Source/Runtime/Experimental/Chaos/Public/Chaos/Convex.h
@@ -381,8 +381,16 @@ namespace Chaos
 		// Get the index of the plane that most opposes the normal
 		int32 GetMostOpposingPlaneScaled(const FVec3& Normal, const FVec3& Scale) const;
 
+		// Get the nearest point on an edge and the edge vertices
+		// Used for manifold generation
+		FVec3 GetClosestEdge(int32 PlaneIndexHint, const FVec3& Position, FVec3& OutEdgePos0, FVec3& OutEdgePos1) const;
+
 		// Get the nearest point on an edge of the specified face
-		FVec3 GetClosestEdgePosition(int32 PlaneIndex, const FVec3& Position) const;
+		FVec3 GetClosestEdgePosition(int32 PlaneIndex, const FVec3& Position) const
+		{
+			FVec3 Unused0, Unused1;
+			return GetClosestEdge(PlaneIndex, Position, Unused0, Unused1);
+		}
 
 		bool GetClosestEdgeVertices(int32 PlaneIndex, const FVec3& Position, int32& OutVertexIndex0, int32& OutVertexIndex1) const;
 
diff --git a/Engine/Source/Runtime/Experimental/Chaos/Public/Chaos/ImplicitObjectScaled.h b/Engine/Source/Runtime/Experimental/Chaos/Public/Chaos/ImplicitObjectScaled.h
--- a/Engine/Source/Runtime/Experimental/Chaos/Public/Chaos/ImplicitObjectScaled.h
+++ b/Engine/Source/Runtime/Experimental/Chaos/Public/Chaos/ImplicitObjectScaled.h
@@ -254,6 +254,12 @@ public:
 		return MObject->GetMostOpposingPlane(Normal);
 	}
 
+	// Get the nearest point on an edge of the specified face
+	FVec3 GetClosestEdge(int32 PlaneIndex, const FVec3& Position, FVec3& OutEdgePos0, FVec3& OutEdgePos1) const
+	{
+		return MObject->GetClosestEdge(PlaneIndex, Position, OutEdgePos0, OutEdgePos1);
+	}
+
 	// Get the nearest point on an edge of the specified face
 	FVec3 GetClosestEdgePosition(int32 PlaneIndex, const FVec3& Position) const
 	{
@@ -711,6 +717,15 @@ public:
 		return MObject->GetMostOpposingPlaneScaled(Normal, MScale);
 	}
 
+	// Get the nearest point on an edge of the specified face
+	FVec3 GetClosestEdge(int32 PlaneIndex, const FVec3& Position, FVec3& OutEdgePos0, FVec3& OutEdgePos1) const
+	{
+		FVec3 EdgePos = MObject->GetClosestEdge(PlaneIndex, MInvScale * Position, OutEdgePos0, OutEdgePos1);
+		OutEdgePos0 = OutEdgePos0 * MScale;
+		OutEdgePos1 = OutEdgePos1 * MScale;
+		return EdgePos * MScale;
+	}
+
 	// Get the nearest point on an edge of the specified face
 	FVec3 GetClosestEdgePosition(int32 PlaneIndex, const FVec3& Position) const
 	{
diff --git a/Engine/Source/Runtime/Experimental/Chaos/Public/Chaos/Triangle.h b/Engine/Source/Runtime/Experimental/Chaos/Public/Chaos/Triangle.h
--- a/Engine/Source/Runtime/Experimental/Chaos/Public/Chaos/Triangle.h
+++ b/Engine/Source/Runtime/Experimental/Chaos/Public/Chaos/Triangle.h
@@ -29,15 +29,21 @@ namespace Chaos
 			VectorC = MakeVectorRegisterFloat(static_cast<FRealSingle>(C[0]), static_cast<FRealSingle>(C[1]), static_cast<FRealSingle>(C[2]), static_cast<FRealSingle>(0.0f));
 		}
 
-		FVec3& operator[](uint32 InIndex)
+		FORCEINLINE FVec3& operator[](uint32 InIndex)
 		{
-			check(InIndex < 3);
+			checkSlow(InIndex < 3);
 			return (&A)[InIndex];
 		}
 
-		const FVec3& operator[](uint32 InIndex) const
+		FORCEINLINE const FVec3& operator[](uint32 InIndex) const
 		{
-			check(InIndex < 3);
+			checkSlow(InIndex < 3);
+			return (&A)[InIndex];
+		}
+
+		FORCEINLINE const FVec3& GetVertex(const int32 InIndex) const
+		{
+			checkSlow(InIndex < 3);
 			return (&A)[InIndex];
 		}
 
@@ -64,29 +70,20 @@ namespace Chaos
 			OutX = A;
 		}
 
-		// Get the nearest point on an edge
+		// Get the nearest point on an edge and the edge vertices
 		// Used for manifold generation
-		FVec3 GetClosestEdgePosition(int32 PlaneIndexHint, const FVec3& Position) const
+		FVec3 GetClosestEdge(int32 PlaneIndexHint, const FVec3& Position, FVec3& OutEdgePos0, FVec3& OutEdgePos1) const
 		{
 			FVec3 ClosestEdgePosition = FVec3(0);
-			FReal ClosestDistanceSq = FLT_MAX;
+			FReal ClosestDistanceSq = TNumericLimits<FReal>::Max();
 
 			int32 PlaneVerticesNum = 3;
-			
+
 			FVec3 P0 = C;
 			for (int32 PlaneVertexIndex = 0; PlaneVertexIndex < PlaneVerticesNum; ++PlaneVertexIndex)
 			{
-				TVector<FReal, 3> P1;
-				switch (PlaneVertexIndex)
-				{
-				case 0:
-					P1 = A; break;
-				case 1:
-					P1 = B; break;
-				case 2:
-					P1 = C; break;
-				}
-				
+				const TVector<FReal, 3>& P1 = GetVertex(PlaneVertexIndex);
+
 				const FVec3 EdgePosition = FMath::ClosestPointOnLine(P0, P1, Position);
 				const FReal EdgeDistanceSq = (EdgePosition - Position).SizeSquared();
 
@@ -94,6 +91,8 @@ namespace Chaos
 				{
 					ClosestDistanceSq = EdgeDistanceSq;
 					ClosestEdgePosition = EdgePosition;
+					OutEdgePos0 = P0;
+					OutEdgePos1 = P1;
 				}
 
 				P0 = P1;
@@ -102,6 +101,15 @@ namespace Chaos
 			return ClosestEdgePosition;
 		}
 
+		// Get the nearest point on an edge
+		// Used for manifold generation
+		FVec3 GetClosestEdgePosition(int32 PlaneIndexHint, const FVec3& Position) const
+		{
+			FVec3 Unused0, Unused1;
+			return GetClosestEdge(PlaneIndexHint, Position, Unused0, Unused1);
+		}
+
+
 		// The number of vertices that make up the corners of the specified face
 		// Used for manifold generation
 		int32 NumPlaneVertices(int32 PlaneIndex) const
@@ -116,28 +124,9 @@ namespace Chaos
 			return 1.0f;
 		}
 
-		// Get the vertex at the specified index (e.g., indices from GetPlaneVertexs)
-		// Used for manifold generation
-		const FVec3 GetVertex(int32 VertexIndex) const
-		{
-			FVec3 Result;
-
-			switch (VertexIndex)
-			{
-			case 0:
-				Result = A; break;
-			case 1:
-				Result = B; break;
-			case 2:
-				Result = C; break;
-			}
-
-			return Result;
-		}
-
 		// Get an array of all the plane indices that belong to a vertex (up to MaxVertexPlanes).
 		// Returns the number of planes found.
-		int32 FindVertexPlanes(int32 VertexIndex, int32* OutVertexPlanes, int32 MaxVertexPlanes) const
+		FORCEINLINE int32 FindVertexPlanes(int32 VertexIndex, int32* OutVertexPlanes, int32 MaxVertexPlanes) const
 		{
 			if(MaxVertexPlanes > 0)
 			{
