GeometryCollection : Property updates
- Convex threshold defaulted to 0
- MassAsDensity is now true in uasset.

diff --git a/Engine/Source/Runtime/Experimental/Chaos/Public/GeometryCollection/GeometryCollectionSimulationCoreTypes.h b/Engine/Source/Runtime/Experimental/Chaos/Public/GeometryCollection/GeometryCollectionSimulationCoreTypes.h
--- a/Engine/Source/Runtime/Experimental/Chaos/Public/GeometryCollection/GeometryCollectionSimulationCoreTypes.h
+++ b/Engine/Source/Runtime/Experimental/Chaos/Public/GeometryCollection/GeometryCollectionSimulationCoreTypes.h
@@ -50,7 +50,7 @@ struct FCollectionCollisionTypeData
 		, LevelSetData()
 		, CollisionParticleData()
 		, CollisionObjectReductionPercentage(0.f)
-		, CollisionMarginFraction(0.1f)
+		, CollisionMarginFraction(0.f)
 	{
 	}
 
@@ -90,7 +90,7 @@ enum ESimulationInitializationState { Unintialized = 0, Activated, Created, Init
 struct FSharedSimulationParameters
 {
 	FSharedSimulationParameters()
-	: bMassAsDensity(false)
+	: bMassAsDensity(true)
 	, Mass(1.0f)
 	, MinimumMassClamp(0.1f)							// todo : Expose to users with better initial values
 	, MaximumMassClamp(1e5f)							// todo : Expose to users with better initial values
diff --git a/Engine/Source/Runtime/Experimental/GeometryCollectionEngine/Private/GeometryCollection/GeometryCollectionObject.cpp b/Engine/Source/Runtime/Experimental/GeometryCollectionEngine/Private/GeometryCollection/GeometryCollectionObject.cpp
--- a/Engine/Source/Runtime/Experimental/GeometryCollectionEngine/Private/GeometryCollection/GeometryCollectionObject.cpp
+++ b/Engine/Source/Runtime/Experimental/GeometryCollectionEngine/Private/GeometryCollection/GeometryCollectionObject.cpp
@@ -116,7 +116,7 @@ FGeometryCollectionCollisionTypeData::FGeometryCollectionCollisionTypeData()
 	, LevelSet()
 	, CollisionParticles()
 	, CollisionObjectReductionPercentage(0.0f)
-	, CollisionMarginFraction(0.1)
+	, CollisionMarginFraction(0.f)
 {
 }
 
