Chaos Cloth Optimization - Reorder spring constraints so batches are contiguous. This reduces the number of gather operations when applying constraint updates.

When running with the Owen asset on my PC, running for ~20s

FPBDSpringConstraints_Apply: 42us -> 36us per call
FPBDAxialSpringConstraints_Apply: 40us -> 34.5us per call

We actually have to store less data (we don't need to store the lists of indices for each color).

Tested edge stiffness maps to verify the map reordering is correct.

diff --git a/Engine/Source/Runtime/Experimental/Chaos/Private/Chaos/PBDAxialSpringConstraints.cpp b/Engine/Source/Runtime/Experimental/Chaos/Private/Chaos/PBDAxialSpringConstraints.cpp
--- a/Engine/Source/Runtime/Experimental/Chaos/Private/Chaos/PBDAxialSpringConstraints.cpp
+++ b/Engine/Source/Runtime/Experimental/Chaos/Private/Chaos/PBDAxialSpringConstraints.cpp
@@ -30,7 +30,40 @@ void FPBDAxialSpringConstraints::InitColor(const FSolverParticles& InParticles)
 	if (Constraints.Num() > Chaos_AxialSpring_ParallelConstraintCount)
 #endif
 	{
-		ConstraintsPerColor = FGraphColoring::ComputeGraphColoring(Constraints, InParticles);
+		const TArray<TArray<int32>> ConstraintsPerColor = FGraphColoring::ComputeGraphColoring(Constraints, InParticles);
+
+		// Reorder constraints based on color so each array in ConstraintsPerColor contains contiguous elements.
+		TArray<TVec3<int32>> ReorderedConstraints;
+		TArray<FSolverReal> ReorderedBarys;
+		TArray<FSolverReal> ReorderedDists;
+		TArray<int32> OrigToReorderedIndices; // used to reorder stiffness indices
+		ReorderedConstraints.SetNumUninitialized(Constraints.Num());
+		ReorderedBarys.SetNumUninitialized(Constraints.Num());
+		ReorderedDists.SetNumUninitialized(Constraints.Num());
+		OrigToReorderedIndices.SetNumUninitialized(Constraints.Num());
+
+		ConstraintsPerColorStartIndex.Reset(ConstraintsPerColor.Num() + 1);
+
+		int32 ReorderedIndex = 0;
+		for (const TArray<int32>& ConstraintsBatch : ConstraintsPerColor)
+		{
+			ConstraintsPerColorStartIndex.Add(ReorderedIndex);
+			for (const int32& BatchConstraint : ConstraintsBatch)
+			{
+				const int32 OrigIndex = BatchConstraint;
+				ReorderedConstraints[ReorderedIndex] = Constraints[OrigIndex];
+				ReorderedBarys[ReorderedIndex] = Barys[OrigIndex];
+				ReorderedDists[ReorderedIndex] = Dists[OrigIndex];
+				OrigToReorderedIndices[OrigIndex] = ReorderedIndex;
+				++ReorderedIndex;
+			}
+		}
+		ConstraintsPerColorStartIndex.Add(ReorderedIndex);
+
+		Constraints = MoveTemp(ReorderedConstraints);
+		Barys = MoveTemp(ReorderedBarys);
+		Dists = MoveTemp(ReorderedDists);
+		Stiffness.ReorderIndices(OrigToReorderedIndices);
 	}
 }
 
@@ -60,8 +93,9 @@ void FPBDAxialSpringConstraints::Apply(FSolverParticles& Particles, const FSolve
 {
 	TRACE_CPUPROFILER_EVENT_SCOPE(FPBDAxialSpringConstraints_Apply);
 	SCOPE_CYCLE_COUNTER(STAT_PBD_AxialSpring);
-	if (ConstraintsPerColor.Num() > 0 && Constraints.Num() > Chaos_AxialSpring_ParallelConstraintCount)
+	if (ConstraintsPerColorStartIndex.Num() > 0 && Constraints.Num() > Chaos_AxialSpring_ParallelConstraintCount)
 	{
+		const int32 ConstraintColorNum = ConstraintsPerColorStartIndex.Num() - 1;
 		if (!Stiffness.HasWeightMap())
 		{
 			const FSolverReal ExpStiffnessValue = (FSolverReal)Stiffness;
@@ -69,27 +103,30 @@ void FPBDAxialSpringConstraints::Apply(FSolverParticles& Particles, const FSolve
 #if INTEL_ISPC
 			if (bRealTypeCompatibleWithISPC && bChaos_AxialSpring_ISPC_Enabled)
 			{
-				for (const TArray<int32>& ConstraintBatch : ConstraintsPerColor)
+				for (int32 ConstraintColorIndex = 0; ConstraintColorIndex < ConstraintColorNum; ++ConstraintColorIndex)
 				{
+					const int32 ColorStart = ConstraintsPerColorStartIndex[ConstraintColorIndex];
+					const int32 ColorSize = ConstraintsPerColorStartIndex[ConstraintColorIndex + 1] - ColorStart;
 					ispc::ApplyAxialSpringConstraints(
 						(ispc::FVector3f*)&Particles.GetP()[0],
-						(ispc::FIntVector*)&Constraints.GetData()[0],
-						&ConstraintBatch.GetData()[0],
+						(ispc::FIntVector*)&Constraints.GetData()[ColorStart],
 						&Particles.GetInvM().GetData()[0],
-						&Barys.GetData()[0],
-						&Dists.GetData()[0],
+						&Barys.GetData()[ColorStart],
+						&Dists.GetData()[ColorStart],
 						ExpStiffnessValue,
-						ConstraintBatch.Num());
+						ColorSize);
 				}
 			}
 			else
 #endif
 			{
-				for (const TArray<int32>& ConstraintBatch : ConstraintsPerColor)
+				for (int32 ConstraintColorIndex = 0; ConstraintColorIndex < ConstraintColorNum; ++ConstraintColorIndex)
 				{
-					PhysicsParallelFor(ConstraintBatch.Num(), [&](const int32 Index)
+					const int32 ColorStart = ConstraintsPerColorStartIndex[ConstraintColorIndex];
+					const int32 ColorSize = ConstraintsPerColorStartIndex[ConstraintColorIndex + 1] - ColorStart;
+					PhysicsParallelFor(ColorSize, [&](const int32 Index)
 					{
-						const int32 ConstraintIndex = ConstraintBatch[Index];
+						const int32 ConstraintIndex = ColorStart + Index;
 						ApplyHelper(Particles, Dt, ConstraintIndex, ExpStiffnessValue);
 					});
 				}
@@ -100,28 +137,31 @@ void FPBDAxialSpringConstraints::Apply(FSolverParticles& Particles, const FSolve
 #if INTEL_ISPC
 			if (bRealTypeCompatibleWithISPC && bChaos_AxialSpring_ISPC_Enabled)
 			{
-				for (const TArray<int32>& ConstraintBatch : ConstraintsPerColor)
+				for (int32 ConstraintColorIndex = 0; ConstraintColorIndex < ConstraintColorNum; ++ConstraintColorIndex)
 				{
+					const int32 ColorStart = ConstraintsPerColorStartIndex[ConstraintColorIndex];
+					const int32 ColorSize = ConstraintsPerColorStartIndex[ConstraintColorIndex + 1] - ColorStart;
 					ispc::ApplyAxialSpringConstraintsWithWeightMaps(
 						(ispc::FVector3f*)&Particles.GetP()[0],
-						(ispc::FIntVector*)&Constraints.GetData()[0],
-						&ConstraintBatch.GetData()[0],
+						(ispc::FIntVector*)&Constraints.GetData()[ColorStart],
 						&Particles.GetInvM().GetData()[0],
-						&Barys.GetData()[0],
-						&Dists.GetData()[0],
-						&Stiffness.GetIndices().GetData()[0],
+						&Barys.GetData()[ColorStart],
+						&Dists.GetData()[ColorStart],
+						&Stiffness.GetIndices().GetData()[ColorStart],
 						&Stiffness.GetTable().GetData()[0],
-						ConstraintBatch.Num());
+						ColorSize);
 				}
 			}
 			else
 #endif
 			{
-				for (const TArray<int32>& ConstraintBatch : ConstraintsPerColor)
+				for (int32 ConstraintColorIndex = 0; ConstraintColorIndex < ConstraintColorNum; ++ConstraintColorIndex)
 				{
-					PhysicsParallelFor(ConstraintBatch.Num(), [&](const int32 Index)
+					const int32 ColorStart = ConstraintsPerColorStartIndex[ConstraintColorIndex];
+					const int32 ColorSize = ConstraintsPerColorStartIndex[ConstraintColorIndex + 1] - ColorStart;
+					PhysicsParallelFor(ColorSize, [&](const int32 Index)
 					{
-						const int32 ConstraintIndex = ConstraintBatch[Index];
+						const int32 ConstraintIndex = ColorStart + Index;
 						const FSolverReal ExpStiffnessValue = Stiffness[ConstraintIndex];
 						ApplyHelper(Particles, Dt, ConstraintIndex, ExpStiffnessValue);
 					});
diff --git a/Engine/Source/Runtime/Experimental/Chaos/Private/Chaos/PBDAxialSpringConstraints.ispc b/Engine/Source/Runtime/Experimental/Chaos/Private/Chaos/PBDAxialSpringConstraints.ispc
--- a/Engine/Source/Runtime/Experimental/Chaos/Private/Chaos/PBDAxialSpringConstraints.ispc
+++ b/Engine/Source/Runtime/Experimental/Chaos/Private/Chaos/PBDAxialSpringConstraints.ispc
@@ -5,38 +5,20 @@
 #include "Math/Vector.isph"
 
 export void ApplyAxialSpringConstraints(uniform FVector3f P[],
-										const uniform FIntVector AllConstraints[],
-										const uniform int32 ActiveConstraints[],
+										const uniform FIntVector Constraints[],
 										const uniform float InvM[],
 										const uniform float Barys[],
 										const uniform float Dists[],
 										const uniform float Stiffness,
 										const uniform int32 NumConstraints)
 {
-	varying int32 ActiveConstraint;
-	varying FIntVector Constraint;
-
-#if HW_GATHER_SUPPORTED == 1
-	if(programIndex < NumConstraints)
-	{
-		ActiveConstraint = ActiveConstraints[programIndex];
-		#pragma ignore warning(perf)
-		Constraint = AllConstraints[ActiveConstraint];
-	}
-#endif
 
 	foreach(i = 0 ... NumConstraints)
 	{
-#if HW_GATHER_SUPPORTED == 0
-		ActiveConstraint = ActiveConstraints[i];
-		#pragma ignore warning(perf)
-		Constraint = AllConstraints[ActiveConstraint];
-#endif
+		const varying FIntVector Constraint = VectorLoad(&Constraints[extract(i,0)]);
 
-		#pragma ignore warning(perf)
-		const varying float Bary = Barys[ActiveConstraint];
-		#pragma ignore warning(perf)
-		const varying float Dist = Dists[ActiveConstraint];
+		const varying float Bary = Barys[i];
+		const varying float Dist = Dists[i];
 
 		const varying int32 i1 = Constraint.V[0];
 		const varying int32 i2 = Constraint.V[1];
@@ -45,15 +27,6 @@ export void ApplyAxialSpringConstraints(uniform FVector3f P[],
 		const varying FVector3f P2 = VectorGather(&P[i2]);
 		const varying FVector3f P3 = VectorGather(&P[i3]);
 
-#if HW_GATHER_SUPPORTED == 1
-		if(i + programCount < NumConstraints)
-		{
-			ActiveConstraint = ActiveConstraints[i+programCount];
-			#pragma ignore warning(perf)
-			Constraint = AllConstraints[ActiveConstraint];
-		}
-#endif
-
 		const varying FVector3f P1 = VectorGather(&P[i1]);
 
 		const FVector3f P0 = (P2 - P3) * Bary + P3;
@@ -94,8 +67,7 @@ export void ApplyAxialSpringConstraints(uniform FVector3f P[],
 }
 
 export void ApplyAxialSpringConstraintsWithWeightMaps(uniform FVector3f P[],
-										const uniform FIntVector AllConstraints[],
-										const uniform int32 ActiveConstraints[],
+										const uniform FIntVector Constraints[],
 										const uniform float InvM[],
 										const uniform float Barys[],
 										const uniform float Dists[],
@@ -103,32 +75,14 @@ export void ApplyAxialSpringConstraintsWithWeightMaps(uniform FVector3f P[],
 										const uniform float StiffnessTable[],
 										const uniform int32 NumConstraints)
 {
-	varying int32 ActiveConstraint;
-	varying FIntVector Constraint;
-
-#if HW_GATHER_SUPPORTED == 1
-	if(programIndex < NumConstraints)
-	{
-		ActiveConstraint = ActiveConstraints[programIndex];
-		#pragma ignore warning(perf)
-		Constraint = AllConstraints[ActiveConstraint];
-	}
-#endif
 
 	foreach(i = 0 ... NumConstraints)
 	{
-#if HW_GATHER_SUPPORTED == 0
-		ActiveConstraint = ActiveConstraints[i];
-		#pragma ignore warning(perf)
-		Constraint = AllConstraints[ActiveConstraint];
-#endif
+		const varying FIntVector Constraint = VectorLoad(&Constraints[extract(i,0)]);
 
-		#pragma ignore warning(perf)
-		const varying float Bary = Barys[ActiveConstraint];
-		#pragma ignore warning(perf)
-		const varying float Dist = Dists[ActiveConstraint];
-		#pragma ignore warning(perf)
-		const varying uint8 StiffnessIndex = StiffnessIndices[ActiveConstraint];
+		const varying float Bary = Barys[i];
+		const varying float Dist = Dists[i];
+		const varying uint8 StiffnessIndex = StiffnessIndices[i];
 
 		const varying int32 i1 = Constraint.V[0];
 		const varying int32 i2 = Constraint.V[1];
@@ -137,15 +91,6 @@ export void ApplyAxialSpringConstraintsWithWeightMaps(uniform FVector3f P[],
 		const varying FVector3f P2 = VectorGather(&P[i2]);
 		const varying FVector3f P3 = VectorGather(&P[i3]);
 
-#if HW_GATHER_SUPPORTED == 1
-		if(i + programCount < NumConstraints)
-		{
-			ActiveConstraint = ActiveConstraints[i+programCount];
-			#pragma ignore warning(perf)
-			Constraint = AllConstraints[ActiveConstraint];
-		}
-#endif
-
 		const varying FVector3f P1 = VectorGather(&P[i1]);
 
 		const FVector3f P0 = (P2 - P3) * Bary + P3;
diff --git a/Engine/Source/Runtime/Experimental/Chaos/Private/Chaos/PBDSpringConstraints.cpp b/Engine/Source/Runtime/Experimental/Chaos/Private/Chaos/PBDSpringConstraints.cpp
--- a/Engine/Source/Runtime/Experimental/Chaos/Private/Chaos/PBDSpringConstraints.cpp
+++ b/Engine/Source/Runtime/Experimental/Chaos/Private/Chaos/PBDSpringConstraints.cpp
@@ -37,7 +37,37 @@ void FPBDSpringConstraints::InitColor(const FSolverParticles& Particles)
 	if (Constraints.Num() > Chaos_Spring_ParallelConstraintCount)
 #endif
 	{
-		ConstraintsPerColor = FGraphColoring::ComputeGraphColoring(Constraints, Particles);
+		const TArray<TArray<int32>> ConstraintsPerColor = FGraphColoring::ComputeGraphColoring(Constraints, Particles);
+		
+		// Reorder constraints based on color so each array in ConstraintsPerColor contains contiguous elements.
+		TArray<TVec2<int32>> ReorderedConstraints;
+		TArray<FSolverReal> ReorderedDists;
+		TArray<int32> OrigToReorderedIndices; // used to reorder stiffness indices
+		ReorderedConstraints.SetNumUninitialized(Constraints.Num());
+		ReorderedDists.SetNumUninitialized(Dists.Num());
+		OrigToReorderedIndices.SetNumUninitialized(Constraints.Num());
+
+		ConstraintsPerColorStartIndex.Reset(ConstraintsPerColor.Num() + 1);
+
+		int32 ReorderedIndex = 0;
+		for (const TArray<int32>& ConstraintsBatch : ConstraintsPerColor)
+		{
+			ConstraintsPerColorStartIndex.Add(ReorderedIndex);
+			for (const int32& BatchConstraint : ConstraintsBatch)
+			{
+				const int32 OrigIndex = BatchConstraint;
+				ReorderedConstraints[ReorderedIndex] = Constraints[OrigIndex];
+				ReorderedDists[ReorderedIndex] = Dists[OrigIndex];
+				OrigToReorderedIndices[OrigIndex] = ReorderedIndex;
+
+				++ReorderedIndex;
+			}
+		}
+		ConstraintsPerColorStartIndex.Add(ReorderedIndex);
+
+		Constraints = MoveTemp(ReorderedConstraints);
+		Dists = MoveTemp(ReorderedDists);
+		Stiffness.ReorderIndices(OrigToReorderedIndices);
 	}
 }
 
@@ -59,9 +89,11 @@ void FPBDSpringConstraints::ApplyHelper(FSolverParticles& Particles, const FSolv
 
 void FPBDSpringConstraints::Apply(FSolverParticles& Particles, const FSolverReal Dt) const
 {
+	TRACE_CPUPROFILER_EVENT_SCOPE(FPBDSpringConstraints_Apply);
 	SCOPE_CYCLE_COUNTER(STAT_PBD_Spring);
-	if ((ConstraintsPerColor.Num() > 0) && (Constraints.Num() > Chaos_Spring_ParallelConstraintCount))
+	if ((ConstraintsPerColorStartIndex.Num() > 1) && (Constraints.Num() > Chaos_Spring_ParallelConstraintCount))
 	{
+		const int32 ConstraintColorNum = ConstraintsPerColorStartIndex.Num() - 1;
 		if (!Stiffness.HasWeightMap())
 		{
 			const FSolverReal ExpStiffnessValue = (FSolverReal)Stiffness;
@@ -69,26 +101,29 @@ void FPBDSpringConstraints::Apply(FSolverParticles& Particles, const FSolverReal
 #if INTEL_ISPC
 			if (bRealTypeCompatibleWithISPC && bChaos_Spring_ISPC_Enabled)
 			{
-				for (const TArray<int32>& ConstraintBatch : ConstraintsPerColor)
+				for (int32 ConstraintColorIndex = 0; ConstraintColorIndex < ConstraintColorNum; ++ConstraintColorIndex)
 				{
+					const int32 ColorStart = ConstraintsPerColorStartIndex[ConstraintColorIndex];
+					const int32 ColorSize = ConstraintsPerColorStartIndex[ConstraintColorIndex + 1] - ColorStart;
 					ispc::ApplySpringConstraints(
-						(ispc::FVector3f*) & Particles.GetP()[0],
-						(ispc::FIntVector2*) & Constraints.GetData()[0],
-						&ConstraintBatch.GetData()[0],
+						(ispc::FVector3f*)&Particles.GetP()[0],
+						(ispc::FIntVector2*)&Constraints.GetData()[ColorStart],
 						&Particles.GetInvM().GetData()[0],
-						&Dists.GetData()[0],
+						&Dists.GetData()[ColorStart],
 						ExpStiffnessValue,
-						ConstraintBatch.Num());
+						ColorSize);
 				}
 			}
 			else
 #endif
 			{
-				for (const TArray<int32>& ConstraintBatch : ConstraintsPerColor)
+				for (int32 ConstraintColorIndex = 0; ConstraintColorIndex < ConstraintColorNum; ++ConstraintColorIndex)
 				{
-					PhysicsParallelFor(ConstraintBatch.Num(), [&](const int32 Index)
+					const int32 ColorStart = ConstraintsPerColorStartIndex[ConstraintColorIndex];
+					const int32 ColorSize = ConstraintsPerColorStartIndex[ConstraintColorIndex + 1] - ColorStart;
+					PhysicsParallelFor(ColorSize, [&](const int32 Index)
 					{
-						const int32 ConstraintIndex = ConstraintBatch[Index];
+						const int32 ConstraintIndex = ColorStart + Index;
 						ApplyHelper(Particles, Dt, ConstraintIndex, ExpStiffnessValue);
 					});
 				}
@@ -99,27 +134,30 @@ void FPBDSpringConstraints::Apply(FSolverParticles& Particles, const FSolverReal
 #if INTEL_ISPC
 			if (bRealTypeCompatibleWithISPC && bChaos_Spring_ISPC_Enabled)
 			{
-				for (const TArray<int32>& ConstraintBatch : ConstraintsPerColor)
+				for (int32 ConstraintColorIndex = 0; ConstraintColorIndex < ConstraintColorNum; ++ConstraintColorIndex)
 				{
+					const int32 ColorStart = ConstraintsPerColorStartIndex[ConstraintColorIndex];
+					const int32 ColorSize = ConstraintsPerColorStartIndex[ConstraintColorIndex + 1] - ColorStart;
 					ispc::ApplySpringConstraintsWithWeightMaps(
 						(ispc::FVector3f*) & Particles.GetP()[0],
-						(ispc::FIntVector2*) & Constraints.GetData()[0],
-						&ConstraintBatch.GetData()[0],
+						(ispc::FIntVector2*) & Constraints.GetData()[ColorStart],
 						&Particles.GetInvM().GetData()[0],
-						&Dists.GetData()[0],
-						&Stiffness.GetIndices().GetData()[0],
+						&Dists.GetData()[ColorStart],
+						&Stiffness.GetIndices().GetData()[ColorStart],
 						&Stiffness.GetTable().GetData()[0],
-						ConstraintBatch.Num());
+						ColorSize);
 				}
 			}
 			else
 #endif
 			{
-				for (const TArray<int32>& ConstraintBatch : ConstraintsPerColor)
+				for (int32 ConstraintColorIndex = 0; ConstraintColorIndex < ConstraintColorNum; ++ConstraintColorIndex)
 				{
-					PhysicsParallelFor(ConstraintBatch.Num(), [&](const int32 Index)
+					const int32 ColorStart = ConstraintsPerColorStartIndex[ConstraintColorIndex];
+					const int32 ColorSize = ConstraintsPerColorStartIndex[ConstraintColorIndex + 1] - ColorStart;
+					PhysicsParallelFor(ColorSize, [&](const int32 Index)
 					{
-						const int32 ConstraintIndex = ConstraintBatch[Index];
+						const int32 ConstraintIndex = ColorStart + Index;
 						const FSolverReal ExpStiffnessValue = Stiffness[ConstraintIndex];
 						ApplyHelper(Particles, Dt, ConstraintIndex, ExpStiffnessValue);
 					});
diff --git a/Engine/Source/Runtime/Experimental/Chaos/Private/Chaos/PBDSpringConstraints.ispc b/Engine/Source/Runtime/Experimental/Chaos/Private/Chaos/PBDSpringConstraints.ispc
--- a/Engine/Source/Runtime/Experimental/Chaos/Private/Chaos/PBDSpringConstraints.ispc
+++ b/Engine/Source/Runtime/Experimental/Chaos/Private/Chaos/PBDSpringConstraints.ispc
@@ -13,33 +13,15 @@ static inline float SafeNormalize(FVector3f &Direction)
 }
 
 export void ApplySpringConstraints(uniform FVector3f P[],
-									const uniform FIntVector2 AllConstraints[],
-									const uniform int32 ActiveConstraints[],
+									const uniform FIntVector2 Constraints[],
 									const uniform float InvM[],
 									const uniform float Dists[],
 									const uniform float Stiffness,
 									const uniform int32 NumConstraints)
 {
-	varying int32 ActiveConstraint;
-	varying FIntVector2 Constraint;
-
-#if HW_GATHER_SUPPORTED == 1
-	if(programIndex < NumConstraints)
-	{
-		ActiveConstraint = ActiveConstraints[programIndex];
-		#pragma ignore warning(perf)
-		Constraint = AllConstraints[ActiveConstraint];
-	}
-#endif
-
 	foreach(i = 0 ... NumConstraints)
 	{
-#if HW_GATHER_SUPPORTED == 0
-		ActiveConstraint = ActiveConstraints[i];
-		#pragma ignore warning(perf)
-		Constraint = AllConstraints[ActiveConstraint];
-#endif
-
+		const varying FIntVector2 Constraint = VectorLoad(&Constraints[extract(i,0)]);
 		const varying int32 i1 = Constraint.V[0];
 		const varying int32 i2 = Constraint.V[1];
 
@@ -49,17 +31,7 @@ export void ApplySpringConstraints(uniform FVector3f P[],
 		FVector3f Direction = P1 - P2;
 		const float Distance = SafeNormalize(Direction);
 
-		#pragma ignore warning(perf)
-		const varying float Dist = Dists[ActiveConstraint];
-
-#if HW_GATHER_SUPPORTED == 1
-		if(i + programCount < NumConstraints)
-		{
-			ActiveConstraint = ActiveConstraints[i+programCount];
-			#pragma ignore warning(perf)
-			Constraint = AllConstraints[ActiveConstraint];
-		}
-#endif
+		const varying float Dist = Dists[i];
 
 		const FVector3f InnerDelta = (Distance - Dist) * Direction;
 
@@ -85,34 +57,17 @@ export void ApplySpringConstraints(uniform FVector3f P[],
 }
 
 export void ApplySpringConstraintsWithWeightMaps(uniform FVector3f P[],
-									const uniform FIntVector2 AllConstraints[],
-									const uniform int32 ActiveConstraints[],
+									const uniform FIntVector2 Constraints[],
 									const uniform float InvM[],
 									const uniform float Dists[],
 									const uniform uint8 StiffnessIndices[],
 									const uniform float StiffnessTable[],
 									const uniform int32 NumConstraints)
 {
-	varying int32 ActiveConstraint;
-	varying FIntVector2 Constraint;
-
-#if HW_GATHER_SUPPORTED == 1
-	if(programIndex < NumConstraints)
-	{
-		ActiveConstraint = ActiveConstraints[programIndex];
-		#pragma ignore warning(perf)
-		Constraint = AllConstraints[ActiveConstraint];
-	}
-#endif
 
 	foreach(i = 0 ... NumConstraints)
 	{
-#if HW_GATHER_SUPPORTED == 0
-		ActiveConstraint = ActiveConstraints[i];
-		#pragma ignore warning(perf)
-		Constraint = AllConstraints[ActiveConstraint];
-#endif
-
+		const varying FIntVector2 Constraint = VectorLoad(&Constraints[extract(i,0)]);
 		const varying int32 i1 = Constraint.V[0];
 		const varying int32 i2 = Constraint.V[1];
 
@@ -122,20 +77,8 @@ export void ApplySpringConstraintsWithWeightMaps(uniform FVector3f P[],
 		FVector3f Direction = P1 - P2;
 		const float Distance = SafeNormalize(Direction);
 
-		#pragma ignore warning(perf)
-		const varying float Dist = Dists[ActiveConstraint];
-
-		#pragma ignore warning(perf)
-		const varying uint8 StiffnessIndex = StiffnessIndices[ActiveConstraint];
-
-#if HW_GATHER_SUPPORTED == 1
-		if(i + programCount < NumConstraints)
-		{
-			ActiveConstraint = ActiveConstraints[i+programCount];
-			#pragma ignore warning(perf)
-			Constraint = AllConstraints[ActiveConstraint];
-		}
-#endif
+		const varying float Dist = Dists[i];
+		const varying uint8 StiffnessIndex = StiffnessIndices[i];
 
 		const FVector3f InnerDelta = (Distance - Dist) * Direction;
 
diff --git a/Engine/Source/Runtime/Experimental/Chaos/Public/Chaos/PBDAxialSpringConstraints.h b/Engine/Source/Runtime/Experimental/Chaos/Public/Chaos/PBDAxialSpringConstraints.h
--- a/Engine/Source/Runtime/Experimental/Chaos/Public/Chaos/PBDAxialSpringConstraints.h
+++ b/Engine/Source/Runtime/Experimental/Chaos/Public/Chaos/PBDAxialSpringConstraints.h
@@ -35,7 +35,7 @@ private:
 	void ApplyHelper(FSolverParticles& InParticles, const FSolverReal Dt, const int32 ConstraintIndex, const FSolverReal ExpStiffnessValue) const;
 
 private:
-	TArray<TArray<int32>> ConstraintsPerColor;
+	TArray<int32> ConstraintsPerColorStartIndex; // Constraints are ordered so each batch is contiguous. This is ColorNum + 1 length so it can be used as start and end.
 };
 
 }  // End namespace Chaos::Softs
diff --git a/Engine/Source/Runtime/Experimental/Chaos/Public/Chaos/PBDSpringConstraints.h b/Engine/Source/Runtime/Experimental/Chaos/Public/Chaos/PBDSpringConstraints.h
--- a/Engine/Source/Runtime/Experimental/Chaos/Public/Chaos/PBDSpringConstraints.h
+++ b/Engine/Source/Runtime/Experimental/Chaos/Public/Chaos/PBDSpringConstraints.h
@@ -37,7 +37,7 @@ private:
 	void ApplyHelper(FSolverParticles& Particles, const FSolverReal Dt, const int32 ConstraintIndex, const FSolverReal ExpStiffnessValue) const;
 
 private:
-	TArray<TArray<int32>> ConstraintsPerColor;
+	TArray<int32> ConstraintsPerColorStartIndex; // Constraints are ordered so each batch is contiguous. This is ColorNum + 1 length so it can be used as start and end.
 };
 
 }  // End namespace Chaos::Softs
diff --git a/Engine/Source/Runtime/Experimental/Chaos/Public/Chaos/PBDStiffness.h b/Engine/Source/Runtime/Experimental/Chaos/Public/Chaos/PBDStiffness.h
--- a/Engine/Source/Runtime/Experimental/Chaos/Public/Chaos/PBDStiffness.h
+++ b/Engine/Source/Runtime/Experimental/Chaos/Public/Chaos/PBDStiffness.h
@@ -89,6 +89,9 @@ public:
 	/** Return the table of stiffnesses as a read only array. */
 	TConstArrayView<FSolverReal> GetTable() const { return TConstArrayView<FSolverReal>(Table); }
 
+	/** Reorder Indices based on Constraint reordering. */
+	inline void ReorderIndices(const TArray<int32>& OrigToReorderedIndices);
+
 private:
 	TArray<uint8> Indices; // Per particle/constraints array of index to the stiffness table
 	TArray<FSolverReal> Table;  // Fixed lookup table of stiffness values, use uint8 indexation
@@ -213,4 +216,23 @@ void FPBDStiffness::ApplyValues(const FSolverReal Dt, const int32 NumIterations)
 	}
 }
 
+void FPBDStiffness::ReorderIndices(const TArray<int32>& OrigToReorderedConstraintIndices)
+{
+	if (Indices.Num() == OrigToReorderedConstraintIndices.Num())
+	{
+		TArray<uint8> ReorderedIndices;
+		ReorderedIndices.SetNumUninitialized(Indices.Num());
+		for (int32 OrigConstraintIndex = 0; OrigConstraintIndex < OrigToReorderedConstraintIndices.Num(); ++OrigConstraintIndex)
+		{
+			const int32 ReorderedConstraintIndex = OrigToReorderedConstraintIndices[OrigConstraintIndex];
+			ReorderedIndices[ReorderedConstraintIndex] = Indices[OrigConstraintIndex];
+		}
+		Indices = MoveTemp(ReorderedIndices);
+	}
+	else
+	{
+		check(Indices.Num() == 1);
+	}
+}
+
 }  // End namespace Chaos::Softs
