Fix island with moving kinematics + Update Sleep State issue

diff --git a/Engine/Source/Runtime/Experimental/Chaos/Private/Chaos/Island/IslandGraph.cpp b/Engine/Source/Runtime/Experimental/Chaos/Private/Chaos/Island/IslandGraph.cpp
--- a/Engine/Source/Runtime/Experimental/Chaos/Private/Chaos/Island/IslandGraph.cpp
+++ b/Engine/Source/Runtime/Experimental/Chaos/Private/Chaos/Island/IslandGraph.cpp
@@ -138,7 +138,7 @@ void FIslandGraph<NodeType, EdgeType, IslandType>::RemoveNode(const NodeType& No
 			FGraphNode& GraphNode = GraphNodes[NodeIndex];
 			
 			// If only one node and zero edges, we invalidate the node island
-			if(GraphNode.NodeEdges.Num() == 0  && GraphIslands.IsValidIndex(GraphNode.IslandIndex))
+			if(GraphNode.bValidNode  && GraphIslands.IsValidIndex(GraphNode.IslandIndex))
 			{
 				if(!GraphNode.bStationaryNode)
 				{
@@ -647,23 +647,30 @@ void FIslandGraph<NodeType, EdgeType, IslandType>::SplitIslands()
 		if(GraphNodes.IsValidIndex(RootIndex))
 		{
 			FGraphNode& RootNode = GraphNodes[RootIndex];
-			if(RootNode.NodeCounter != GraphCounter && RootNode.bValidNode)
+			if(RootNode.NodeCounter != GraphCounter)
 			{
-				int32 CurrentIsland = RootNode.IslandIndex;
-				
-				if (GraphIslands.IsValidIndex(CurrentIsland) && GraphIslands[CurrentIsland].bIsPersistent && !GraphIslands[CurrentIsland].bIsSleeping)
+				if(RootNode.bValidNode)
 				{
-					// We don't want to rebuild a new island if this one can't be splitted
-					// It is why by default the first one is the main one
-					if(GraphIslands[CurrentIsland].IslandCounter == GraphCounter)
-					{
-						FGraphIsland GraphIsland = { 0, 1, 0, false, false };
-						CurrentIsland = GraphIslands.Emplace(GraphIsland);
-					}
+					int32 CurrentIsland = RootNode.IslandIndex;
 					
-					GraphIslands[CurrentIsland].IslandCounter = GraphCounter;
+					if (GraphIslands.IsValidIndex(CurrentIsland) && GraphIslands[CurrentIsland].bIsPersistent && !GraphIslands[CurrentIsland].bIsSleeping)
+					{
+						// We don't want to rebuild a new island if this one can't be splitted
+						// It is why by default the first one is the main one
+						if(GraphIslands[CurrentIsland].IslandCounter == GraphCounter)
+						{
+							FGraphIsland GraphIsland = { 0, 1, 0, false, false };
+							CurrentIsland = GraphIslands.Emplace(GraphIsland);
+						}
+						
+						GraphIslands[CurrentIsland].IslandCounter = GraphCounter;
 
-					SplitIsland(NodeQueue, RootIndex, CurrentIsland);
+						SplitIsland(NodeQueue, RootIndex, CurrentIsland);
+					}
+				}
+				else
+				{
+					RootNode.IslandIndex = INDEX_NONE;
 				}
 			}
 		}
@@ -735,7 +742,7 @@ void FIslandGraph<NodeType, EdgeType, IslandType>::UpdateGraph()
 			// Update of the sleeping flag on the island
 			if(GraphIslands.IsValidIndex(GraphNode.IslandIndex) && !GraphNode.bStationaryNode)
 			{
-				if(GraphNode.NodeEdges.Num() == 0)
+				if(GraphNode.bValidNode)
 				{
 					GraphIslands[GraphNode.IslandIndex].bIsSleeping = false;
 				}
diff --git a/Engine/Source/Runtime/Experimental/Chaos/Private/Chaos/Island/IslandManager.cpp b/Engine/Source/Runtime/Experimental/Chaos/Private/Chaos/Island/IslandManager.cpp
--- a/Engine/Source/Runtime/Experimental/Chaos/Private/Chaos/Island/IslandManager.cpp
+++ b/Engine/Source/Runtime/Experimental/Chaos/Private/Chaos/Island/IslandManager.cpp
@@ -114,31 +114,34 @@ inline void UpdateSleepState(FPBDIslandSolver* IslandSolver, FPBDRigidsSOAs& Par
 		bool bNeedRebuild = false;
 		for (auto& IslandParticle : IslandSolver->GetParticles())
 		{
-			// If not sleeping we activate the sleeping particles
-			if (!bIsSleeping && IslandParticle->Sleeping())
+			if(IslandParticle->CastToRigidParticle() && !IslandParticle->CastToRigidParticle()->Disabled())
 			{
-				Particles.ActivateParticle(IslandParticle, true);
-
-				// When we wake particles, we have skipped their integrate step which causes some issues:
-				//	- we have zero velocity (no gravity or external forces applied)
-				//	- the world transforms cached in the ShapesArray will be at the last post-integrate positions
-				//	  which doesn't match what the velocity is telling us
-				// This causes problems for the solver - essentially we have an "initial overlap" situation.
-				// @todo(chaos): We could just run (partial) integrate here for this particle, but we don't know about the Evolution - fix this
-				const FRigidTransform3 ParticleWorldTransform = FParticleUtilities::GetActorWorldTransform(FConstGenericParticleHandle(IslandParticle));
-				for (const TUniquePtr<FPerShapeData>& Shape : IslandParticle->ShapesArray())
+				// If not sleeping we activate the sleeping particles
+				if (!bIsSleeping && IslandParticle->Sleeping())
 				{
-					const FRigidTransform3 ShapeWorldTransform = (FRigidTransform3)Shape->GetLeafRelativeTransform() * ParticleWorldTransform;
-					Shape->SetLeafWorldTransform(ShapeWorldTransform);
-				}
+					Particles.ActivateParticle(IslandParticle, true);
+
+					// When we wake particles, we have skipped their integrate step which causes some issues:
+					//	- we have zero velocity (no gravity or external forces applied)
+					//	- the world transforms cached in the ShapesArray will be at the last post-integrate positions
+					//	  which doesn't match what the velocity is telling us
+					// This causes problems for the solver - essentially we have an "initial overlap" situation.
+					// @todo(chaos): We could just run (partial) integrate here for this particle, but we don't know about the Evolution - fix this
+					const FRigidTransform3 ParticleWorldTransform = FParticleUtilities::GetActorWorldTransform(FConstGenericParticleHandle(IslandParticle));
+					for (const TUniquePtr<FPerShapeData>& Shape : IslandParticle->ShapesArray())
+					{
+						const FRigidTransform3 ShapeWorldTransform = (FRigidTransform3)Shape->GetLeafRelativeTransform() * ParticleWorldTransform;
+						Shape->SetLeafWorldTransform(ShapeWorldTransform);
+					}
 
-				bNeedRebuild = true;
-			}
-			// If sleeping we deactivate the dynamic particles
-			else if (bIsSleeping && !IslandParticle->Sleeping())
-			{
-				Particles.DeactivateParticle(IslandParticle, true);
-				bNeedRebuild = true;
+					bNeedRebuild = true;
+				}
+				// If sleeping we deactivate the dynamic particles
+				else if (bIsSleeping && !IslandParticle->Sleeping())
+				{
+					Particles.DeactivateParticle(IslandParticle, true);
+					bNeedRebuild = true;
+				}
 			}
 		}
 		if(bNeedRebuild)
@@ -190,8 +193,8 @@ inline void PopulateIslands(FPBDIslandManager::GraphType* IslandGraph)
 	TSet<int32> NodeIslands;
 	for(auto& GraphNode : IslandGraph->GraphNodes)
 	{
-		// If the node has one edge or if the node is valid : only one island
-		if((GraphNode.NodeEdges.Num() == 0) || GraphNode.bValidNode)
+		// If the node is valid : only one island
+		if( GraphNode.bValidNode)
 		{
 			AddNodeToIsland(GraphNode.IslandIndex, GraphNode);
 		}
@@ -346,6 +349,11 @@ void FPBDIslandManager::InitializeGraph(const TParticleView<FPBDRigidParticles>&
 			FGeometryParticleHandle* ParticleHandle = GraphNode.NodeItem;
 			IslandGraph->UpdateNode(ParticleHandle, IsDynamicParticle(ParticleHandle), GraphNode.IslandIndex,
 													IsStationaryParticle(ParticleHandle), NodeIndex);
+
+			if (ParticleHandle->CastToRigidParticle() && ParticleHandle->CastToRigidParticle()->Disabled())
+			{
+				RemoveParticle(ParticleHandle);
+			}
 		}
 	}
 	// For now we are resetting all the constraints but we should keep the
@@ -454,15 +462,22 @@ int32 FPBDIslandManager::AddConstraint(const uint32 ContainerId, FConstraintHand
 			const int32 EdgeIndex = IslandGraph->AddEdge(ConstraintHandle, ContainerId, NodeIndex0, NodeIndex1);
 			ConstraintHandle->SetConstraintGraphIndex(EdgeIndex);
 
-			// Make sure to sync the state of the constraint with its owning island, otherwise the constraint may be flag as destroyable and leave a dangling pointer in the islands
-			const int32 IslandIndex = IslandGraph->GraphEdges[EdgeIndex].IslandIndex;
-			if (IslandGraph->GraphIslands.IsValidIndex(IslandIndex) && IslandGraph->GraphIslands[IslandIndex].bIsSleeping)
+			if(IslandGraph->GraphEdges.IsValidIndex(EdgeIndex))
 			{
-				ConstraintHandle->SetIsSleeping(true);
+				// Make sure to sync the state of the constraint with its owning island, otherwise the constraint may be flag as destroyable and leave a dangling pointer in the islands
+				const int32 IslandIndex = IslandGraph->GraphEdges[EdgeIndex].IslandIndex;
+				if (IslandGraph->GraphIslands.IsValidIndex(IslandIndex) && IslandGraph->GraphIslands[IslandIndex].bIsSleeping)
+				{
+					ConstraintHandle->SetIsSleeping(true);
+				}
 			}
 
 			return EdgeIndex;
 		}
+		else
+		{
+			ConstraintHandle->SetConstraintGraphIndex(INDEX_NONE);
+		}
 		return INDEX_NONE;
 	}
 	return INDEX_NONE;
@@ -486,8 +501,8 @@ void FPBDIslandManager::RemoveParticle(FGeometryParticleHandle* ParticleHandle)
 			if (IslandGraph->GraphNodes.IsValidIndex(*NodeIndex))
 			{
 				const FGraphNode& GraphNode = IslandGraph->GraphNodes[*NodeIndex];
-				// the list of edges could be empty or the node is valid : we need to also remove it from its own island 
-				if(GraphNode.NodeEdges.Num() == 0 && IslandSolvers.IsValidIndex(GraphNode.IslandIndex))
+				// If the node is valid : we need to also remove it from its own island 
+				if(GraphNode.bValidNode && IslandSolvers.IsValidIndex(GraphNode.IslandIndex))
 				{
 					IslandSolvers[GraphNode.IslandIndex]->RemoveParticle(ParticleHandle);
 				}
