remove ensures on empty convex hulls, and update too-small-hull fallback code to make a larger hull so we're less likely to create empty hulls in the first place

#rnx

diff --git a/Engine/Source/Runtime/Experimental/Chaos/Private/GeometryCollection/GeometryCollectionConvexUtility.cpp b/Engine/Source/Runtime/Experimental/Chaos/Private/GeometryCollection/GeometryCollectionConvexUtility.cpp
--- a/Engine/Source/Runtime/Experimental/Chaos/Private/GeometryCollection/GeometryCollectionConvexUtility.cpp
+++ b/Engine/Source/Runtime/Experimental/Chaos/Private/GeometryCollection/GeometryCollectionConvexUtility.cpp
@@ -350,10 +350,9 @@ void CreateNonoverlappingConvexHulls(
 
 	auto IsColliding = [&Convexes](int32 ConvexA, int32 ConvexB)
 	{
-		if (ConvexA == -1 || ConvexB == -1 || !ensure(Convexes[ConvexA]->NumVertices() > 0 && Convexes[ConvexB]->NumVertices() > 0))
+		if (ConvexA == -1 || ConvexB == -1 || Convexes[ConvexA]->NumVertices() == 0 || Convexes[ConvexB]->NumVertices() == 0)
 		{
 			// at least one of the convex hulls was empty, so cannot be colliding
-			// TODO: eventually remove this ensure; and add more handling for empty hulls
 			return false;
 		}
 		const Chaos::TRigidTransform<Chaos::FReal, 3> IdentityTransform = Chaos::TRigidTransform<Chaos::FReal, 3>::Identity;
@@ -378,7 +377,7 @@ void CreateNonoverlappingConvexHulls(
 	auto GetConvexSpan = [](const Chaos::FConvex& Convex, const Chaos::FVec3& Center, const Chaos::FVec3& Normal) -> Chaos::FVec2
 	{
 		int32 NumVertices = Convex.NumVertices();
-		if (!ensure(NumVertices > 0))
+		if (NumVertices == 0)
 		{
 			return Chaos::FVec2(0, 0);
 		}
@@ -497,10 +496,9 @@ void CreateNonoverlappingConvexHulls(
 
 	auto FixCollisionWithCut = [&Convexes, &FindCutPlane, &SimplificationDistanceThreshold](int32 ConvexA, int32 ConvexB)
 	{
-		if (!ensure(Convexes[ConvexA]->NumVertices() > 0 && Convexes[ConvexB]->NumVertices() > 0))
+		if (Convexes[ConvexA]->NumVertices() == 0 || Convexes[ConvexB]->NumVertices() == 0)
 		{
 			// at least one of the convex hulls was empty, so cannot be colliding
-			// TODO: eventually remove this ensure; and add more handling for empty hulls
 			return false;
 		}
 		Chaos::FReal Depth;
@@ -979,23 +977,36 @@ void HullsFromGeometry(
 			int32 ConvexIdx = Convexes.Add(MakeUnique<Chaos::FConvex>(HullPts, KINDA_SMALL_NUMBER));
 			if (Convexes[ConvexIdx]->NumVertices() == 0 && HullPts.Num() > 0)
 			{
-				// if we've failed to make a convex hull, add a tiny bounding box just to ensure every geometry has a hull
+				// if we've failed to make a convex hull, add a tiny bounding box to try to give it a hull anyway
+				// (note it may still end up with an empty hull after cutting the hull to remove collisions with neighboring hulls; this just reduces the chance/number of empty hulls)
 				Chaos::FConvex::FAABB3Type AABB = Convexes[ConvexIdx]->GetLocalBoundingBox();
-				AABB.Thicken(.001f);
-				Chaos::FConvex::FVec3Type Min = AABB.Min();
-				Chaos::FConvex::FVec3Type Max = AABB.Max();
-				HullPts.Add(Min);
-				HullPts.Add(Max);
-				HullPts.Add({Min.X, Min.Y, Max.Z});
-				HullPts.Add({Min.X, Max.Y, Max.Z});
-				HullPts.Add({Max.X, Min.Y, Max.Z});
-				HullPts.Add({Max.X, Max.Y, Min.Z});
-				HullPts.Add({Max.X, Min.Y, Min.Z});
-				HullPts.Add({Min.X, Max.Y, Min.Z});
-				// note: Do not use SimplificationDistanceThreshold for this fixed tiny hull
-				*Convexes[ConvexIdx] = Chaos::FConvex(HullPts, KINDA_SMALL_NUMBER);
+				Chaos::FVec3 Extents = AABB.Extents();
+				constexpr Chaos::FReal SmallestConvexWidth = (Chaos::FReal).11; // Convexes that span less distance than this will become empty hulls due to the size thresholds in FConvex
+				if (Extents.GetMax() < SmallestConvexWidth * 2) // only apply this logic if the bounding box was actually small (in case FConvex fails to make a hull even for larger hulls)
+				{
+					// expand the hull to cover a .11 ^ 3 volume, so it won't disappear
+					constexpr Chaos::FReal MinExpand = (Chaos::FReal).01;
+					Chaos::FVec3 GrowVec;
+					for (int Axis = 0; Axis < 3; Axis++)
+					{
+						GrowVec[Axis] = FMath::Max((Chaos::FReal).5 * (SmallestConvexWidth - Extents[Axis]), MinExpand);
+					}
+					AABB.ThickenSymmetrically(GrowVec);
+					Chaos::FConvex::FVec3Type Min = AABB.Min();
+					Chaos::FConvex::FVec3Type Max = AABB.Max();
+					HullPts.Reset(); // we're adding the bounding box, so no need to keep the points inside
+					HullPts.Add(Min);
+					HullPts.Add(Max);
+					HullPts.Add({ Min.X, Min.Y, Max.Z });
+					HullPts.Add({ Min.X, Max.Y, Max.Z });
+					HullPts.Add({ Max.X, Min.Y, Max.Z });
+					HullPts.Add({ Max.X, Max.Y, Min.Z });
+					HullPts.Add({ Max.X, Min.Y, Min.Z });
+					HullPts.Add({ Min.X, Max.Y, Min.Z });
+					// note: Do not use SimplificationDistanceThreshold for this fixed tiny hull
+					*Convexes[ConvexIdx] = Chaos::FConvex(HullPts, KINDA_SMALL_NUMBER);
+				}
 			}
-			ensure(Convexes[ConvexIdx]->NumVertices() > 0);
 			TransformToConvexIndices[Idx].Add(ConvexIdx);
 		}
 	}
@@ -1406,6 +1417,7 @@ void FGeometryCollectionConvexUtility::SetVolumeAttributes(FGeometryCollection*
 	GeometryCollectionAlgo::GlobalMatrices(Collection->Transform, Collection->Parent, Transforms);
 	TArray<int32> RecursiveOrder = GeometryCollectionAlgo::ComputeRecursiveOrder(*Collection);
 	float MaxGeoVolume = 0.0f;
+	Volumes.Fill(0.0f);
 	for (int32 Bone : RecursiveOrder)
 	{
 		if (SimulationType[Bone] == FGeometryCollection::ESimulationTypes::FST_Rigid)
