Add fbx user defined node attribute. Create an extra attribute API in UInterchangeBaseNode to enable add and query of the attribute.
User defined attribute are normally attribute user add to a node in the DCC application.

#rnx

diff --git a/Engine/Plugins/Experimental/Interchange/Runtime/Source/Dispatcher/Public/InterchangeCommands.h b/Engine/Plugins/Experimental/Interchange/Runtime/Source/Dispatcher/Public/InterchangeCommands.h
--- a/Engine/Plugins/Experimental/Interchange/Runtime/Source/Dispatcher/Public/InterchangeCommands.h
+++ b/Engine/Plugins/Experimental/Interchange/Runtime/Source/Dispatcher/Public/InterchangeCommands.h
@@ -17,7 +17,7 @@ namespace UE
 		{
 		public:
 			//Major version should be updated when there is an existing API that has been change.
-			static int32 GetMajor() { return 2; }
+			static int32 GetMajor() { return 3; }
 			//Minor version must be updated when there is an addition to the API.
 			static int32 GetMinor() { return 0; }
 			//Patch version should be update if there is some bug fixes in the private code.
diff --git a/Engine/Plugins/Experimental/Interchange/Runtime/Source/Parsers/Fbx/Private/FbxHelper.cpp b/Engine/Plugins/Experimental/Interchange/Runtime/Source/Parsers/Fbx/Private/FbxHelper.cpp
--- a/Engine/Plugins/Experimental/Interchange/Runtime/Source/Parsers/Fbx/Private/FbxHelper.cpp
+++ b/Engine/Plugins/Experimental/Interchange/Runtime/Source/Parsers/Fbx/Private/FbxHelper.cpp
@@ -94,6 +94,17 @@ namespace UE
 				return NodeAttributeUniqueID;
 			}
 
+			FString FFbxHelper::GetFbxPropertyName(const FbxProperty& Property)
+			{
+				FString PropertyName = FFbxConvert::MakeName(Property.GetName());
+				if (PropertyName.Equals(TEXT("none"), ESearchCase::IgnoreCase))
+				{
+					//Replace None by Null because None clash with NAME_None and the create asset will instead call the object ClassName_X
+					PropertyName = TEXT("Null");
+				}
+				return PropertyName;
+			}
+
 			FString FFbxHelper::GetFbxObjectName(const FbxObject* Object)
 			{
 				if (!Object)
diff --git a/Engine/Plugins/Experimental/Interchange/Runtime/Source/Parsers/Fbx/Private/FbxHelper.h b/Engine/Plugins/Experimental/Interchange/Runtime/Source/Parsers/Fbx/Private/FbxHelper.h
--- a/Engine/Plugins/Experimental/Interchange/Runtime/Source/Parsers/Fbx/Private/FbxHelper.h
+++ b/Engine/Plugins/Experimental/Interchange/Runtime/Source/Parsers/Fbx/Private/FbxHelper.h
@@ -31,6 +31,11 @@ namespace UE
 				static FString GetNodeAttributeName(FbxNodeAttribute* NodeAttribute, const FStringView DefaultNamePrefix);
 				static FString GetNodeAttributeUniqueID(FbxNodeAttribute* NodeAttribute, const FStringView Prefix);
 
+				/**
+				 * Return the name of an FbxProperty, return empty string if the property is null.
+				 */
+				static FString GetFbxPropertyName(const FbxProperty& Property);
+
 				/**
 				 * Return the name of an FbxObject, return empty string if the object is null.
 				 */
diff --git a/Engine/Plugins/Experimental/Interchange/Runtime/Source/Parsers/Fbx/Private/FbxScene.cpp b/Engine/Plugins/Experimental/Interchange/Runtime/Source/Parsers/Fbx/Private/FbxScene.cpp
--- a/Engine/Plugins/Experimental/Interchange/Runtime/Source/Parsers/Fbx/Private/FbxScene.cpp
+++ b/Engine/Plugins/Experimental/Interchange/Runtime/Source/Parsers/Fbx/Private/FbxScene.cpp
@@ -15,6 +15,7 @@
 #include "InterchangeResultsContainer.h"
 #include "InterchangeSceneNode.h"
 #include "Nodes/InterchangeBaseNodeContainer.h"
+#include "Nodes/InterchangeUserDefinedAttribute.h"
 
 #define LOCTEXT_NAMESPACE "InterchangeFbxScene"
 
@@ -244,6 +245,166 @@ namespace UE
 					}
 				}
 
+				//Add all custom Attributes for the node
+				auto IsSupportedExtraAttributeType = [](int32 InDataType)
+				{
+					switch (InDataType)
+					{
+					case EFbxType::eFbxBool:
+
+					case EFbxType::eFbxChar:		//!< 8 bit signed integer.
+					case EFbxType::eFbxUChar:		//!< 8 bit unsigned integer.
+					case EFbxType::eFbxShort:		//!< 16 bit signed integer.
+					case EFbxType::eFbxUShort:		//!< 16 bit unsigned integer.
+					case EFbxType::eFbxInt:			//!< 32 bit signed integer.
+					case EFbxType::eFbxUInt:		//!< 32 bit unsigned integer.
+					case EFbxType::eFbxLongLong:	//!< 64 bit signed integer.
+					case EFbxType::eFbxULongLong:	//!< 64 bit unsigned integer.
+					case EFbxType::eFbxHalfFloat:	//!< 16 bit floating point.
+					case EFbxType::eFbxFloat:		//!< Floating point value.
+					case EFbxType::eFbxDouble:	//!< Double width floating point value.
+					case EFbxType::eFbxDouble2:	//!< Vector of two double values.
+					case EFbxType::eFbxDouble3:	//!< Vector of three double values.
+					case EFbxType::eFbxDouble4:	//!< Vector of four double values.
+
+					case EFbxType::eFbxEnum:		//!< Enumeration.
+					case EFbxType::eFbxString:	//!< String.
+						return true;
+					}
+
+					return false;
+				};
+
+				FbxProperty Property = Node->GetFirstProperty();
+				while (Property.IsValid())
+				{
+					if (Property.GetFlag(FbxPropertyFlags::eUserDefined) && IsSupportedExtraAttributeType(Property.GetPropertyDataType().GetType()))
+					{
+						FString PropertyName = FFbxHelper::GetFbxPropertyName(Property);
+
+						FbxAnimCurveNode* CurveNode = Property.GetCurveNode();
+						TOptional<FString> PayloadKey;
+						if (CurveNode && CurveNode->IsAnimated())
+						{
+							//Set the optional payload key
+						}
+						switch (Property.GetPropertyDataType().GetType())
+						{
+							case EFbxType::eFbxBool:
+								{
+									bool PropertyValue = Property.Get<bool>();
+									UInterchangeUserDefinedAttributesAPI::CreateUserDefinedAttribute(UnrealNode, PropertyName, PropertyValue, PayloadKey);
+								}
+								break;
+							case EFbxType::eFbxChar:
+								{
+									int8 PropertyValue = Property.Get<int8>();
+									UInterchangeUserDefinedAttributesAPI::CreateUserDefinedAttribute(UnrealNode, PropertyName, PropertyValue, PayloadKey);
+								}
+								break;
+							case EFbxType::eFbxUChar:
+								{
+									uint8 PropertyValue = Property.Get<uint8>();
+									UInterchangeUserDefinedAttributesAPI::CreateUserDefinedAttribute(UnrealNode, PropertyName, PropertyValue, PayloadKey);
+								}
+								break;
+							case EFbxType::eFbxShort:
+								{
+									int16 PropertyValue = Property.Get<int16>();
+									UInterchangeUserDefinedAttributesAPI::CreateUserDefinedAttribute(UnrealNode, PropertyName, PropertyValue, PayloadKey);
+								}
+								break;
+							case EFbxType::eFbxUShort:
+								{
+									uint16 PropertyValue = Property.Get<uint16>();
+									UInterchangeUserDefinedAttributesAPI::CreateUserDefinedAttribute(UnrealNode, PropertyName, PropertyValue, PayloadKey);
+								}
+								break;
+							case EFbxType::eFbxInt:
+								{
+									int32 PropertyValue = Property.Get<int32>();
+									UInterchangeUserDefinedAttributesAPI::CreateUserDefinedAttribute(UnrealNode, PropertyName, PropertyValue, PayloadKey);
+								}
+								break;
+							case EFbxType::eFbxUInt:
+								{
+									uint32 PropertyValue = Property.Get<uint32>();
+									UInterchangeUserDefinedAttributesAPI::CreateUserDefinedAttribute(UnrealNode, PropertyName, PropertyValue, PayloadKey);
+								}
+								break;
+							case EFbxType::eFbxLongLong:
+								{
+									int64 PropertyValue = Property.Get<int64>();
+									UInterchangeUserDefinedAttributesAPI::CreateUserDefinedAttribute(UnrealNode, PropertyName, PropertyValue, PayloadKey);
+								}
+								break;
+							case EFbxType::eFbxULongLong:
+								{
+									uint64 PropertyValue = Property.Get<uint64>();
+									UInterchangeUserDefinedAttributesAPI::CreateUserDefinedAttribute(UnrealNode, PropertyName, PropertyValue, PayloadKey);
+								}
+								break;
+							case EFbxType::eFbxHalfFloat:
+								{
+									FbxHalfFloat HalfFloat = Property.Get<FbxHalfFloat>();
+									FFloat16 PropertyValue = FFloat16(HalfFloat.value());
+									UInterchangeUserDefinedAttributesAPI::CreateUserDefinedAttribute(UnrealNode, PropertyName, PropertyValue, PayloadKey);
+								}
+								break;
+							case EFbxType::eFbxFloat:
+								{
+									float PropertyValue = Property.Get<float>();
+									UInterchangeUserDefinedAttributesAPI::CreateUserDefinedAttribute(UnrealNode, PropertyName, PropertyValue, PayloadKey);
+								}
+								break;
+							case EFbxType::eFbxDouble:
+								{
+									double PropertyValue = Property.Get<double>();
+									UInterchangeUserDefinedAttributesAPI::CreateUserDefinedAttribute(UnrealNode, PropertyName, PropertyValue, PayloadKey);
+								}
+								break;
+							case EFbxType::eFbxDouble2:
+								{
+									FbxDouble2 Vec = Property.Get<FbxDouble2>();
+									FVector2D PropertyValue = FVector2D(Vec[0], Vec[1]);
+									UInterchangeUserDefinedAttributesAPI::CreateUserDefinedAttribute(UnrealNode, PropertyName, PropertyValue, PayloadKey);
+								}
+								break;
+							case EFbxType::eFbxDouble3:
+								{
+									FbxDouble3 Vec = Property.Get<FbxDouble3>();
+									FVector3d PropertyValue = FVector3d(Vec[0], Vec[1], Vec[2]);
+									UInterchangeUserDefinedAttributesAPI::CreateUserDefinedAttribute(UnrealNode, PropertyName, PropertyValue, PayloadKey);
+								}
+								break;
+							case EFbxType::eFbxDouble4:
+								{
+									FbxDouble4 Vec = Property.Get<FbxDouble4>();
+									FVector4d PropertyValue = FVector4d(Vec[0], Vec[1], Vec[2], Vec[3]);
+									UInterchangeUserDefinedAttributesAPI::CreateUserDefinedAttribute(UnrealNode, PropertyName, PropertyValue, PayloadKey);
+								}
+								break;
+							case EFbxType::eFbxEnum:
+								{
+									//Convert enum to uint8
+									FbxEnum EnumValue = Property.Get<FbxEnum>();
+									uint8 PropertyValue = static_cast<uint8>(EnumValue);
+									UInterchangeUserDefinedAttributesAPI::CreateUserDefinedAttribute(UnrealNode, PropertyName, PropertyValue, PayloadKey);
+								}
+								break;
+							case EFbxType::eFbxString:
+								{
+									FbxString StringValue = Property.Get<FbxString>();
+									FString PropertyValue = FFbxConvert::MakeString(StringValue.Buffer());
+									UInterchangeUserDefinedAttributesAPI::CreateUserDefinedAttribute(UnrealNode, PropertyName, PropertyValue, PayloadKey);
+								}
+								break;
+						}
+					}
+					//Inspect next node property
+					Property = Node->GetNextProperty(Property);
+				}
+
 				const int32 ChildCount = Node->GetChildCount();
 				for (int32 ChildIndex = 0; ChildIndex < ChildCount; ++ChildIndex)
 				{
diff --git a/Engine/Source/Runtime/Experimental/Interchange/Core/Private/Nodes/InterchangeUserDefinedAttribute.cpp b/Engine/Source/Runtime/Experimental/Interchange/Core/Private/Nodes/InterchangeUserDefinedAttribute.cpp
new file mode 100644
--- /dev/null
+++ b/Engine/Source/Runtime/Experimental/Interchange/Core/Private/Nodes/InterchangeUserDefinedAttribute.cpp
@@ -0,0 +1,165 @@
+// Copyright Epic Games, Inc. All Rights Reserved.
+
+#include "Nodes/InterchangeUserDefinedAttribute.h"
+
+#include "Nodes/InterchangeBaseNode.h"
+
+const FString UInterchangeUserDefinedAttributesAPI::UserDefinedAttributeBaseKey = TEXT("UserDefined_");
+const FString UInterchangeUserDefinedAttributesAPI::UserDefinedAttributeValuePostKey = TEXT("_Value");
+const FString UInterchangeUserDefinedAttributesAPI::UserDefinedAttributePayLoadPostKey = TEXT("_Payload");
+
+bool UInterchangeUserDefinedAttributesAPI::CreateUserDefinedAttribute_Boolean(UInterchangeBaseNode* InterchangeNode, const FString& UserDefinedAttributeName, const bool& Value, const FString& PayloadKey)
+{
+	TOptional<FString> OptionalPayload;
+	if (!PayloadKey.IsEmpty())
+	{
+		OptionalPayload = PayloadKey;
+	}
+	return UInterchangeUserDefinedAttributesAPI::CreateUserDefinedAttribute(InterchangeNode, UserDefinedAttributeName, Value, OptionalPayload);
+}
+
+bool UInterchangeUserDefinedAttributesAPI::CreateUserDefinedAttribute_Float(UInterchangeBaseNode* InterchangeNode, const FString& UserDefinedAttributeName, const float& Value, const FString& PayloadKey)
+{
+	TOptional<FString> OptionalPayload;
+	if (!PayloadKey.IsEmpty())
+	{
+		OptionalPayload = PayloadKey;
+	}
+	return UInterchangeUserDefinedAttributesAPI::CreateUserDefinedAttribute(InterchangeNode, UserDefinedAttributeName, Value, OptionalPayload);
+}
+
+bool UInterchangeUserDefinedAttributesAPI::CreateUserDefinedAttribute_Double(UInterchangeBaseNode* InterchangeNode, const FString& UserDefinedAttributeName, const double& Value, const FString& PayloadKey)
+{
+	TOptional<FString> OptionalPayload;
+	if (!PayloadKey.IsEmpty())
+	{
+		OptionalPayload = PayloadKey;
+	}
+	return UInterchangeUserDefinedAttributesAPI::CreateUserDefinedAttribute(InterchangeNode, UserDefinedAttributeName, Value, OptionalPayload);
+}
+
+bool UInterchangeUserDefinedAttributesAPI::CreateUserDefinedAttribute_Int32(UInterchangeBaseNode* InterchangeNode, const FString& UserDefinedAttributeName, const int32& Value, const FString& PayloadKey)
+{
+	TOptional<FString> OptionalPayload;
+	if (!PayloadKey.IsEmpty())
+	{
+		OptionalPayload = PayloadKey;
+	}
+	return UInterchangeUserDefinedAttributesAPI::CreateUserDefinedAttribute(InterchangeNode, UserDefinedAttributeName, Value, OptionalPayload);
+}
+
+bool UInterchangeUserDefinedAttributesAPI::CreateUserDefinedAttribute_FString(UInterchangeBaseNode* InterchangeNode, const FString& UserDefinedAttributeName, const FString& Value, const FString& PayloadKey)
+{
+	TOptional<FString> OptionalPayload;
+	if (!PayloadKey.IsEmpty())
+	{
+		OptionalPayload = PayloadKey;
+	}
+	return UInterchangeUserDefinedAttributesAPI::CreateUserDefinedAttribute(InterchangeNode, UserDefinedAttributeName, Value, OptionalPayload);
+}
+
+bool UInterchangeUserDefinedAttributesAPI::RemoveUserDefinedAttribute(UInterchangeBaseNode* InterchangeNode, const FString& UserDefinedAttributeName)
+{
+
+	FString StorageBaseKey = UserDefinedAttributeBaseKey + UserDefinedAttributeName;
+	UE::Interchange::FAttributeKey UserDefinedValueKey = UE::Interchange::FAttributeKey(StorageBaseKey + UserDefinedAttributeValuePostKey);
+	if (InterchangeNode->HasAttribute(UserDefinedValueKey))
+	{
+		if (!InterchangeNode->RemoveAttribute(UserDefinedValueKey.Key))
+		{
+			return false;
+		}
+	}
+
+	UE::Interchange::FAttributeKey UserDefinedPayloadKey = UE::Interchange::FAttributeKey(StorageBaseKey + UserDefinedAttributePayLoadPostKey);
+	if (InterchangeNode->HasAttribute(UserDefinedPayloadKey))
+	{
+		if (!InterchangeNode->RemoveAttribute(UserDefinedPayloadKey.Key))
+		{
+			return false;
+		}
+	}
+
+	//Attribute was successfuly removed
+	return true;
+}
+
+bool UInterchangeUserDefinedAttributesAPI::GetUserDefinedAttribute_Boolean(const UInterchangeBaseNode* InterchangeNode, const FString& UserDefinedAttributeName, bool& OutValue, FString& OutPayloadKey)
+{
+	TOptional<FString> OptionalPayload;
+	bool bResult = UInterchangeUserDefinedAttributesAPI::GetUserDefinedAttribute(InterchangeNode, UserDefinedAttributeName, OutValue, OptionalPayload);
+	if (OptionalPayload.IsSet())
+	{
+		OutPayloadKey = OptionalPayload.GetValue();
+	}
+	return bResult;
+}
+
+bool UInterchangeUserDefinedAttributesAPI::GetUserDefinedAttribute_Float(const UInterchangeBaseNode* InterchangeNode, const FString& UserDefinedAttributeName, float& OutValue, FString& OutPayloadKey)
+{
+	TOptional<FString> OptionalPayload;
+	bool bResult = UInterchangeUserDefinedAttributesAPI::GetUserDefinedAttribute(InterchangeNode, UserDefinedAttributeName, OutValue, OptionalPayload);
+	if (OptionalPayload.IsSet())
+	{
+		OutPayloadKey = OptionalPayload.GetValue();
+	}
+	return bResult;
+}
+
+bool UInterchangeUserDefinedAttributesAPI::GetUserDefinedAttribute_Double(const UInterchangeBaseNode* InterchangeNode, const FString& UserDefinedAttributeName, double& OutValue, FString& OutPayloadKey)
+{
+	TOptional<FString> OptionalPayload;
+	bool bResult = UInterchangeUserDefinedAttributesAPI::GetUserDefinedAttribute(InterchangeNode, UserDefinedAttributeName, OutValue, OptionalPayload);
+	if (OptionalPayload.IsSet())
+	{
+		OutPayloadKey = OptionalPayload.GetValue();
+	}
+	return bResult;
+}
+
+bool UInterchangeUserDefinedAttributesAPI::GetUserDefinedAttribute_Int32(const UInterchangeBaseNode* InterchangeNode, const FString& UserDefinedAttributeName, int32& OutValue, FString& OutPayloadKey)
+{
+	TOptional<FString> OptionalPayload;
+	bool bResult = UInterchangeUserDefinedAttributesAPI::GetUserDefinedAttribute(InterchangeNode, UserDefinedAttributeName, OutValue, OptionalPayload);
+	if (OptionalPayload.IsSet())
+	{
+		OutPayloadKey = OptionalPayload.GetValue();
+	}
+	return bResult;
+}
+
+bool UInterchangeUserDefinedAttributesAPI::GetUserDefinedAttribute_FString(const UInterchangeBaseNode* InterchangeNode, const FString& UserDefinedAttributeName, FString& OutValue, FString& OutPayloadKey)
+{
+	TOptional<FString> OptionalPayload;
+	bool bResult = UInterchangeUserDefinedAttributesAPI::GetUserDefinedAttribute(InterchangeNode, UserDefinedAttributeName, OutValue, OptionalPayload);
+	if (OptionalPayload.IsSet())
+	{
+		OutPayloadKey = OptionalPayload.GetValue();
+	}
+	return bResult;
+}
+
+TArray<FString> UInterchangeUserDefinedAttributesAPI::GetUserDefinedAttributeNames(const UInterchangeBaseNode* InterchangeNode)
+{
+	check(InterchangeNode);
+	TArray<UE::Interchange::FAttributeKey> AttributeKeys;
+	InterchangeNode->GetAttributeKeys(AttributeKeys);
+	TArray<FString> UserDefinedAttributeNames;
+	int32 RightChopIndex = UserDefinedAttributeBaseKey.Len();
+	int32 LeftChopIndex = UserDefinedAttributeValuePostKey.Len();
+	for (UE::Interchange::FAttributeKey& AttributeKey : AttributeKeys)
+	{
+		if (AttributeKey.Key.StartsWith(UserDefinedAttributeBaseKey) && AttributeKey.Key.EndsWith(UserDefinedAttributeValuePostKey))
+		{
+			FString UserDefineAttributeName = AttributeKey.Key.RightChop(RightChopIndex).LeftChop(LeftChopIndex);
+			UserDefinedAttributeNames.Add(UserDefineAttributeName);
+		}
+	}
+	return UserDefinedAttributeNames;
+}
+
+void UInterchangeUserDefinedAttributesAPI::GetUserDefinedAttributeNames(const UInterchangeBaseNode* InterchangeNode, TArray<FString>& UserDefinedAttributeNames)
+{
+	check(InterchangeNode);
+	UserDefinedAttributeNames = GetUserDefinedAttributeNames(InterchangeNode);
+}
diff --git a/Engine/Source/Runtime/Experimental/Interchange/Core/Public/Nodes/InterchangeUserDefinedAttribute.h b/Engine/Source/Runtime/Experimental/Interchange/Core/Public/Nodes/InterchangeUserDefinedAttribute.h
new file mode 100644
--- /dev/null
+++ b/Engine/Source/Runtime/Experimental/Interchange/Core/Public/Nodes/InterchangeUserDefinedAttribute.h
@@ -0,0 +1,157 @@
+// Copyright Epic Games, Inc. All Rights Reserved.
+
+#pragma once
+
+#include "Misc/Optional.h"
+#include "Nodes/InterchangeBaseNode.h"
+#include "UObject/Class.h"
+#include "UObject/Object.h"
+#include "UObject/ObjectMacros.h"
+
+#include "InterchangeUserDefinedAttribute.generated.h"
+
+/**
+ * UInterchangeUserDefinedAttributesAPI is used to store and retrieve user defined attributes (i.e. DCC node attributes, pipelines will have access to those attributes)
+ * Any user defined attribute have: name, value and a optional AnimationPayloadKey (FRichCurve which is a float curve).
+ * Value type must be supported by the UE::Interchange::EAttributeTypes enumeration.
+ */
+UCLASS(BlueprintType, Experimental)
+class INTERCHANGECORE_API UInterchangeUserDefinedAttributesAPI : public UObject
+{
+	GENERATED_BODY()
+ 
+public:
+
+	/**
+	 * Create user defined attribute with a value and a optional payload key
+	 * param UserDefinedAttributeName - The name of the user defined attribute
+	 * param Value - The value of the user defined attribute
+	 * param PayloadKey - The translator payload key to retrieve the FRichCurve animation for this user defined attribute
+	 * note - User defined attributes are the DCC translated node user custom attributes (i.e. Maya extra attributes)
+	 *        Payload key will point on a FRichCurve payload.
+	 */
+	template<typename ValueType>
+	static bool CreateUserDefinedAttribute(UInterchangeBaseNode* InterchangeNode, const FString& UserDefinedAttributeName, const ValueType& Value, const TOptional<FString>& PayloadKey);
+
+	UFUNCTION(BlueprintCallable, Category = "Interchange | Node | UserDefinedAttribute")
+	static bool CreateUserDefinedAttribute_Boolean(UInterchangeBaseNode* InterchangeNode, const FString& UserDefinedAttributeName, const bool& Value, const FString& PayloadKey);
+	
+	UFUNCTION(BlueprintCallable, Category = "Interchange | Node | UserDefinedAttribute")
+	static bool CreateUserDefinedAttribute_Float(UInterchangeBaseNode* InterchangeNode, const FString& UserDefinedAttributeName, const float& Value, const FString& PayloadKey);
+
+	UFUNCTION(BlueprintCallable, Category = "Interchange | Node | UserDefinedAttribute")
+	static bool CreateUserDefinedAttribute_Double(UInterchangeBaseNode* InterchangeNode, const FString& UserDefinedAttributeName, const double& Value, const FString& PayloadKey);
+
+	UFUNCTION(BlueprintCallable, Category = "Interchange | Node | UserDefinedAttribute")
+	static bool CreateUserDefinedAttribute_Int32(UInterchangeBaseNode* InterchangeNode, const FString& UserDefinedAttributeName, const int32& Value, const FString& PayloadKey);
+
+	UFUNCTION(BlueprintCallable, Category = "Interchange | Node | UserDefinedAttribute")
+	static bool CreateUserDefinedAttribute_FString(UInterchangeBaseNode* InterchangeNode, const FString& UserDefinedAttributeName, const FString& Value, const FString& PayloadKey);
+
+	/**
+	 * Remove the specified user defined attribute
+	 * param UserDefinedAttributeName - The name of the user defined attribute to remove
+	 * return - True if the attribute exist and was remove or if the attribute doesn't exist. Return false if the attribute exist but the attribute was not properly remove.
+	 * note - User defined attributes are the DCC translated node user custom attributes (i.e. Maya extra attributes)
+	 *        Payload key will point on a FRichCurve payload.
+	 */
+	UFUNCTION(BlueprintCallable, Category = "Interchange | Node | UserDefinedAttribute")
+	static bool RemoveUserDefinedAttribute(UInterchangeBaseNode* InterchangeNode, const FString& UserDefinedAttributeName);
+
+	/**
+	 * Get user defined attribute value and optional payload key
+	 * param UserDefinedAttributeName - The name of the user defined attribute
+	 * param OutValue - The value of the user defined attribute
+	 * param OutPayloadKey - The translator payload key to retrieve the FRichCurve animation for this user defined attribute
+	 * note - User defined attributes are the DCC translated node user custom attributes (i.e. Maya extra attributes)
+	 *        Payload key will point on a FRichCurve payload.
+	 */
+	template<typename ValueType>
+	static bool GetUserDefinedAttribute(const UInterchangeBaseNode* InterchangeNode, const FString& UserDefinedAttributeName, ValueType& OutValue, TOptional<FString>& OutPayloadKey);
+
+	UFUNCTION(BlueprintCallable, Category = "Interchange | Node | UserDefinedAttribute")
+	static bool GetUserDefinedAttribute_Boolean(const UInterchangeBaseNode* InterchangeNode, const FString& UserDefinedAttributeName, bool& OutValue, FString& OutPayloadKey);
+
+	UFUNCTION(BlueprintCallable, Category = "Interchange | Node | UserDefinedAttribute")
+	static bool GetUserDefinedAttribute_Float(const UInterchangeBaseNode* InterchangeNode, const FString& UserDefinedAttributeName, float& OutValue, FString& OutPayloadKey);
+
+	UFUNCTION(BlueprintCallable, Category = "Interchange | Node | UserDefinedAttribute")
+	static bool GetUserDefinedAttribute_Double(const UInterchangeBaseNode* InterchangeNode, const FString& UserDefinedAttributeName, double& OutValue, FString& OutPayloadKey);
+
+	UFUNCTION(BlueprintCallable, Category = "Interchange | Node | UserDefinedAttribute")
+	static bool GetUserDefinedAttribute_Int32(const UInterchangeBaseNode* InterchangeNode, const FString& UserDefinedAttributeName, int32& OutValue, FString& OutPayloadKey);
+
+	UFUNCTION(BlueprintCallable, Category = "Interchange | Node | UserDefinedAttribute")
+	static bool GetUserDefinedAttribute_FString(const UInterchangeBaseNode* InterchangeNode, const FString& UserDefinedAttributeName, FString& OutValue, FString& OutPayloadKey);
+
+	static TArray<FString> GetUserDefinedAttributeNames(const UInterchangeBaseNode* InterchangeNode);
+
+	UFUNCTION(BlueprintCallable, Category = "Interchange | Node | UserDefinedAttribute")
+	static void GetUserDefinedAttributeNames(const UInterchangeBaseNode* InterchangeNode, TArray<FString>& UserDefinedAttributeNames);
+ 
+private:
+	static const FString UserDefinedAttributeBaseKey;
+	static const FString UserDefinedAttributeValuePostKey;
+	static const FString UserDefinedAttributePayLoadPostKey;
+};
+
+template<typename ValueType>
+bool UInterchangeUserDefinedAttributesAPI::CreateUserDefinedAttribute(UInterchangeBaseNode* InterchangeNode, const FString& UserDefinedAttributeName, const ValueType& Value, const TOptional<FString>& PayloadKey)
+{
+	check(InterchangeNode);
+	//Create a unique Key for this user defined attribute
+	FString StorageBaseKey = UserDefinedAttributeBaseKey + UserDefinedAttributeName;
+	UE::Interchange::FAttributeKey UserDefinedValueKey = UE::Interchange::FAttributeKey(StorageBaseKey + UserDefinedAttributeValuePostKey);
+	if (InterchangeNode->HasAttribute(UserDefinedValueKey))
+	{
+		return false;
+	}
+	UE::Interchange::FAttributeKey UserDefinedPayloadKey = UE::Interchange::FAttributeKey(StorageBaseKey + UserDefinedAttributePayLoadPostKey);
+	if (InterchangeNode->HasAttribute(UserDefinedPayloadKey))
+	{
+		return false;
+	}
+	//Add the user defined attribute to TMap
+	InterchangeNode->RegisterAttribute<ValueType>(UserDefinedValueKey, Value);
+	if (PayloadKey.IsSet())
+	{
+		InterchangeNode->RegisterAttribute<FString>(UserDefinedPayloadKey, PayloadKey.GetValue());
+	}
+	return true;
+}
+
+template<typename ValueType>
+bool UInterchangeUserDefinedAttributesAPI::GetUserDefinedAttribute(const UInterchangeBaseNode* InterchangeNode, const FString& UserDefinedAttributeName, ValueType& OutValue, TOptional<FString>& OutPayloadKey)
+{
+	check(InterchangeNode);
+	FString StorageBaseKey = UserDefinedAttributeBaseKey + UserDefinedAttributeName;
+	UE::Interchange::FAttributeKey UserDefinedValueKey = UE::Interchange::FAttributeKey(StorageBaseKey + UserDefinedAttributeValuePostKey);
+
+	if (!InterchangeNode->HasAttribute(UserDefinedValueKey))
+	{
+		return false;
+	}
+
+	if (!InterchangeNode->GetAttribute<ValueType>(UserDefinedValueKey.Key, OutValue))
+	{
+		return false;
+	}
+
+	//Payload is optional
+	OutPayloadKey.Reset();
+	UE::Interchange::FAttributeKey UserDefinedPayloadKey = UE::Interchange::FAttributeKey(StorageBaseKey + UserDefinedAttributePayLoadPostKey);
+	if (InterchangeNode->HasAttribute(UserDefinedPayloadKey))
+	{
+		FString PayloadKeyValue;
+		if (InterchangeNode->GetAttribute<FString>(UserDefinedPayloadKey.Key, PayloadKeyValue))
+		{
+			//Set the optional
+			OutPayloadKey = PayloadKeyValue;
+		}
+		else
+		{
+			return false;
+		}
+	}
+	return true;
+}
