Fix Lira crash : node edges were not removed correctly

diff --git a/Engine/Source/Runtime/Experimental/Chaos/Private/Chaos/Island/SolverIsland.cpp b/Engine/Source/Runtime/Experimental/Chaos/Private/Chaos/Island/SolverIsland.cpp
--- a/Engine/Source/Runtime/Experimental/Chaos/Private/Chaos/Island/SolverIsland.cpp
+++ b/Engine/Source/Runtime/Experimental/Chaos/Private/Chaos/Island/SolverIsland.cpp
@@ -63,7 +63,10 @@ void FPBDIslandSolver::AddConstraint(FConstraintHandle* ConstraintHandle)
 		IslandConstraints.Add(ConstraintHandle);
 
 		const int32 ContainerId = ConstraintHandle->GetContainerId();
-		ConstraintCounts[ConstraintHandle->GetContainerId()]++;
+		if(ConstraintCounts.IsValidIndex(ContainerId))
+		{
+			ConstraintCounts[ContainerId]++;
+		}
 	}
 }
 
