Chaos - Verbose Warnings
- Removed verbose warning for destructable default material.
- Add pvar bGeometryCollectionEnableForcedConvexGenerationInSerialize to disable forced convex generation on a project.

diff --git a/Engine/Source/Runtime/Engine/Private/UnrealEngine.cpp b/Engine/Source/Runtime/Engine/Private/UnrealEngine.cpp
--- a/Engine/Source/Runtime/Engine/Private/UnrealEngine.cpp
+++ b/Engine/Source/Runtime/Engine/Private/UnrealEngine.cpp
@@ -3033,7 +3033,6 @@ void UEngine::InitializeObjectReferences()
 		FString ValidDefaultDestructiblePhysMaterialName = DefaultDestructiblePhysMaterialName.ToString();
 		if (ValidDefaultDestructiblePhysMaterialName.IsEmpty())
 		{
-			UE_LOG(LogEngine, Warning, TEXT("The default destructible physics material was not found. Using default destructible physics material instead. Please make sure to have the default destructible physics material is set up correctly."));
 			ValidDefaultDestructiblePhysMaterialName = TEXT("/Engine/EngineMaterials/DefaultDestructiblePhysicalMaterial.DefaultDestructiblePhysicalMaterial");
 		}
 
diff --git a/Engine/Source/Runtime/Experimental/GeometryCollectionEngine/Private/GeometryCollection/GeometryCollectionObject.cpp b/Engine/Source/Runtime/Experimental/GeometryCollectionEngine/Private/GeometryCollection/GeometryCollectionObject.cpp
--- a/Engine/Source/Runtime/Experimental/GeometryCollectionEngine/Private/GeometryCollection/GeometryCollectionObject.cpp
+++ b/Engine/Source/Runtime/Experimental/GeometryCollectionEngine/Private/GeometryCollection/GeometryCollectionObject.cpp
@@ -43,6 +43,12 @@ FAutoConsoleVariableRef CVarGeometryCollectionBypassPhysicsAttributes(
 	GeometryCollectionAssetForceStripOnCook,
 	TEXT("Bypass the construction of simulation properties when all bodies are simply cached. for playback."));
 
+bool bGeometryCollectionEnableForcedConvexGenerationInSerialize = true;
+FAutoConsoleVariableRef CVarGeometryCollectionEnableForcedConvexGenerationInSerialize(
+	TEXT("p.GeometryCollectionEnableForcedConvexGenerationInSerialize"),
+	bGeometryCollectionEnableForcedConvexGenerationInSerialize,
+	TEXT("Enable generation of convex geometry on older destruction files.[def:true]"));
+
 
 #if ENABLE_COOK_STATS
 namespace GeometryCollectionCookStats
@@ -790,13 +796,16 @@ void UGeometryCollection::Serialize(FArchive& Ar)
 		&& Ar.CustomVer(FPhysicsObjectVersion::GUID) < FPhysicsObjectVersion::GeometryCollectionConvexDefaults)
 	{
 #if WITH_EDITOR
-		if (!FGeometryCollectionConvexUtility::HasConvexHullData(GeometryCollection.Get()))
+		if (bGeometryCollectionEnableForcedConvexGenerationInSerialize)
 		{
-			UE_LOG(LogGeometryCollectionInternal, Warning, TEXT("Regenerated GeoemtryCollections convex geometry on asset (%s). Please resave package %s."), *GetName(), *GetOutermost()->GetPathName());
-			FGeometryCollectionConvexPropertiesInterface::FConvexCreationProperties ConvexProperties = GeometryCollection->GetConvexProperties();
-			FGeometryCollectionConvexUtility::CreateNonOverlappingConvexHullData(GeometryCollection.Get(), ConvexProperties.FractionRemove, ConvexProperties.SimplificationThreshold, ConvexProperties.CanExceedFraction);
-			bCreateSimulationData = true;
-			InvalidateCollection();
+			if (!FGeometryCollectionConvexUtility::HasConvexHullData(GeometryCollection.Get()))
+			{
+				UE_LOG(LogGeometryCollectionInternal, Warning, TEXT("Regenerated GeoemtryCollections convex geometry on asset (%s). Please resave package %s."), *GetName(), *GetOutermost()->GetPathName());
+				FGeometryCollectionConvexPropertiesInterface::FConvexCreationProperties ConvexProperties = GeometryCollection->GetConvexProperties();
+				FGeometryCollectionConvexUtility::CreateNonOverlappingConvexHullData(GeometryCollection.Get(), ConvexProperties.FractionRemove, ConvexProperties.SimplificationThreshold, ConvexProperties.CanExceedFraction);
+				bCreateSimulationData = true;
+				InvalidateCollection();
+			}
 		}
 #endif
 	}
