Chaos : Joint Constraint
- Stubbed out implementation of joint constraint collision transfer.

diff --git a/Engine/Source/Runtime/Experimental/Chaos/Private/Chaos/PBDRigidsEvolutionGBF.cpp b/Engine/Source/Runtime/Experimental/Chaos/Private/Chaos/PBDRigidsEvolutionGBF.cpp
--- a/Engine/Source/Runtime/Experimental/Chaos/Private/Chaos/PBDRigidsEvolutionGBF.cpp
+++ b/Engine/Source/Runtime/Experimental/Chaos/Private/Chaos/PBDRigidsEvolutionGBF.cpp
@@ -80,6 +80,18 @@ namespace Chaos
 		bool DoTransferJointConstraintCollisions = true;
 		FAutoConsoleVariableRef CVarDoTransferJointConstraintCollisions(TEXT("p.Chaos.Solver.Joint.TransferCollisions"), DoTransferJointConstraintCollisions, TEXT("Allows joints to apply collisions to the parent from the child when the Joints TransferCollisionScale is not 0 [def:true]"));
 
+		int32 TransferCollisionsLimit = INT_MAX;
+		FAutoConsoleVariableRef CVarTransferCollisionsMultiply(TEXT("p.Chaos.Solver.Joint.TransferCollisionsLimit"), TransferCollisionsLimit, TEXT("Maximum number of constraints that are allowed to transfer to the parent. Lowering this will improve performance but reduce accuracy. [def:INT_MAX]"));
+
+		FRealSingle TransferCollisionsKinematicScale = 1.0f;
+		FAutoConsoleVariableRef CVarTransferCollisionsKinematicScale(TEXT("p.Chaos.Solver.Joint.TransferCollisionsKinematicScale"), TransferCollisionsKinematicScale, TEXT("Scale to apply to collision transfers between kinematic bodies [def:1.0]"));
+
+		FRealSingle TransferCollisionsStiffnessClamp = 1.0f;
+		FAutoConsoleVariableRef CVarTransferCollisionsStiffnessClamp(TEXT("p.Chaos.Solver.Joint.TransferCollisionsStiffnessClamp"), TransferCollisionsStiffnessClamp, TEXT("Clamp of maximum value of the stiffness clamp[def:1.0]"));
+
+		bool TransferCollisionsDebugTestAgainstMaxClamp = false;
+		FAutoConsoleVariableRef CVarTransferCollisionsDebugTestAgainstMaxClamp(TEXT("p.Chaos.Solver.Joint.TransferCollisionsDebugTestAgainstMaxClamp"), TransferCollisionsDebugTestAgainstMaxClamp, TEXT("Force all joint collision constraint settings to max clamp value to validate stability [def:false]"));
+
 		DECLARE_CYCLE_STAT(TEXT("FPBDRigidsEvolutionGBF::AdvanceOneTimeStep"), STAT_Evolution_AdvanceOneTimeStep, STATGROUP_Chaos);
 		DECLARE_CYCLE_STAT(TEXT("FPBDRigidsEvolutionGBF::UnclusterUnions"), STAT_Evolution_UnclusterUnions, STATGROUP_Chaos);
 		DECLARE_CYCLE_STAT(TEXT("FPBDRigidsEvolutionGBF::Integrate"), STAT_Evolution_Integrate, STATGROUP_Chaos);
@@ -730,6 +742,12 @@ void FPBDRigidsEvolutionGBF::SetCurrentStepResimCache(IResimCacheBase* InCurrent
 
 void FPBDRigidsEvolutionGBF::TransferJointConstraintCollisions()
 {
+	//
+	// Stubbed out implementation of joint constraint collision transfer. 
+	// Currently disabled because its not required, and the implementation
+	// requires proper unit testing before it should be released. 
+	//
+#if 0
 	// Transfer collisions from the child of a joint to the parent.
 	// E.g., if body A and B are connected by a joint, with A the parent and B the child...
 	// then a third body C collides with B...
@@ -756,11 +774,11 @@ void FPBDRigidsEvolutionGBF::TransferJointConstraintCollisions()
 				const FRigidTransform3 ChildToParentTransform = ChildTransform.GetRelativeTransform(ParentTransform);
 
 				ChildParticle->Handle()->ParticleCollisions().VisitCollisions(
-					[&](const FPBDCollisionConstraint& ChildCollisionConstraint)
+					[&](const FPBDCollisionConstraint* VisitedConstraint)
 					{
-						if (ChildCollisionConstraint.GetCCDType() != ECollisionCCDType::Disabled)
+						if (VisitedConstraint->GetCCDType() != ECollisionCCDType::Disabled)
 						{
-							return ECollisionVisitorResult::Continue;
+							return;
 						}
 
 						// @todo(chaos): implemeent this
@@ -773,121 +791,119 @@ void FPBDRigidsEvolutionGBF::TransferJointConstraintCollisions()
 						//		constraint to hold the pointer (see previous issue).
 						//	-	we should check to see if there is already an active constraint between the bodies because we don't want
 						//		to replace a legit collision with our fake one...probably
-						ensure(false);
-
-						const FGeometryParticleHandle* NewParentParticleConst = (ChildCollisionConstraint.GetParticle0() == ChildParticle->Handle()) ? ChildCollisionConstraint.GetParticle1() : ChildCollisionConstraint.GetParticle0();
-						
-						FGeometryParticleHandle* NewParticleA = const_cast<FGeometryParticleHandle*>(NewParentParticleConst);
-						FGeometryParticleHandle* NewParticleB = ParentParticle->Handle();
-
-						// Set up NewCollision - this should duplicate what happens in collision detection, except the
-						// contact points are just read from the source constraint rather than via the narrow phase
-						//FPBDCollisionConstraint NewCollision = FPBDCollisionConstraint(...);
-						// NewCollision.AddOneshotManifoldContact(...);
-						// ...
-						// NewCollision.SetStiffness(JointSettings.ContactTransferScale);
-
-						// Add collision to the system
-						//FParticlePairMidPhase* MidPhase = CollisionAllocator.GetParticlePairMidPhase(NewParticleA, NewParticleB);
-						//MidPhase->InjectCollision(NewCollision);
-
-						return ECollisionVisitorResult::Continue;
-					});
-			}
-		}
+						//const FGeometryParticleHandle* NewParentParticleConst = (CurrConstraint->GetParticle0() == ChildParticle->Handle()) ? CurrConstraint->GetParticle1() : ChildCollisionConstraint->GetParticle0();
+						//FGeometryParticleHandle* NewParticleA = const_cast<FGeometryParticleHandle*>(NewParentParticleConst);
+						//FGeometryParticleHandle* NewParticleB = ParentParticle->Handle();
+
+						FPBDCollisionConstraint& CollisionConstraint = *const_cast<FPBDCollisionConstraint*>(VisitedConstraint);
+						// If the collision constraint is colliding with the child of the joint, then 
+						// we want to map the collision between the parent of the joint and the external
+						// collision body.Otherwise, we map the collision between the childand 
+						// the external collision body.The new collision constraint is constructed with
+						// the implicit object of the colliding body on the joint, to the particle the
+						// collision was transferred to.
+
+						auto MakeCollisionAtIndex0 = [&ChildToParentTransform, &CollisionConstraint](FGenericParticleHandle& TransferToParticle)
+						{
+							return FPBDCollisionConstraint::Make(
+								TransferToParticle->Handle(),
+								CollisionConstraint.GetImplicit0(),
+								CollisionConstraint.GetCollisionParticles0(),
+								ChildToParentTransform.GetRelativeTransform(CollisionConstraint.GetShapeRelativeTransform0()),
+								CollisionConstraint.GetParticle1(),
+								CollisionConstraint.GetImplicit1(),
+								CollisionConstraint.GetCollisionParticles1(),
+								CollisionConstraint.GetShapeRelativeTransform1(),
+								CollisionConstraint.GetCullDistance(),
+								CollisionConstraint.GetUseManifold(),
+								CollisionConstraint.GetShapesType()
+							);
+						};
+
+						auto MakeCollisionAtIndex1 = [&ChildToParentTransform, &CollisionConstraint](FGenericParticleHandle& TransferToParticle)
+						{
+							return FPBDCollisionConstraint::Make(
+								CollisionConstraint.GetParticle1(),
+								CollisionConstraint.GetImplicit1(),
+								CollisionConstraint.GetCollisionParticles1(),
+								CollisionConstraint.GetShapeRelativeTransform1(),
+								TransferToParticle->Handle(),
+								CollisionConstraint.GetImplicit0(),
+								CollisionConstraint.GetCollisionParticles0(),
+								ChildToParentTransform.GetRelativeTransform(CollisionConstraint.GetShapeRelativeTransform1()),
+								CollisionConstraint.GetCullDistance(),
+								CollisionConstraint.GetUseManifold(),
+								CollisionConstraint.GetShapesType()
+							);
+						};
+
+						FGenericParticleHandle CollisionParticleA, CollisionParticleB;
+						TUniquePtr<FPBDCollisionConstraint> TransferedConstraint;
+						if (CollisionConstraint.GetParticle0() == ChildParticle->CastToRigidParticle() ||
+							CollisionConstraint.GetParticle1() == ChildParticle->CastToRigidParticle())
+						{
+							if (CollisionConstraint.GetParticle0() == ChildParticle->CastToRigidParticle())
+							{
+								CollisionParticleA = ParentParticle;
+								CollisionParticleB = CollisionConstraint.GetParticle1();
+								TransferedConstraint = MakeCollisionAtIndex0(ParentParticle);
+							}
+							else
+							{
+								CollisionParticleA = CollisionConstraint.GetParticle0();
+								CollisionParticleB = ParentParticle;
+								TransferedConstraint = MakeCollisionAtIndex1(ParentParticle);
+							}
+						}
+						else if (CollisionConstraint.GetParticle0() == ParentParticle->CastToRigidParticle() ||
+							CollisionConstraint.GetParticle1() == ParentParticle->CastToRigidParticle())
+						{
+							if (CollisionConstraint.GetParticle0() == ParentParticle->CastToRigidParticle())
+							{
+								CollisionParticleA = ChildParticle;
+								CollisionParticleB = CollisionConstraint.GetParticle0();
+								TransferedConstraint = MakeCollisionAtIndex0(ChildParticle);
+							}
+							else
+							{
+								CollisionParticleA = CollisionConstraint.GetParticle1();
+								CollisionParticleB = ChildParticle;
+								TransferedConstraint = MakeCollisionAtIndex1(ChildParticle);
+							}
+						}
 
-		CollisionAllocator.ProcessInjectedConstraints();
-	}
-#if 0
-	//
-	// Append transfer constraints. 
-	//
-	// @todo(chaos): this should be implementable with a collision modifier system
-	using FRigidHandle = TPBDRigidParticleHandleImp<FReal, 3, true>;
-	if (DoTransferJointConstraintCollisions)
-	{
-		FCollisionConstraintAllocator& CollisionAllocator = CollisionConstraints.GetConstraintAllocator();
+						TArray<FManifoldPoint> TransferManifolds;
+						for (FManifoldPoint& Manifold : CollisionConstraint.GetManifoldPoints())
+						{
+							TransferManifolds.Add(Manifold);
+						}
+						//TransferedConstraint->SetManifoldPoints(TransferManifolds);
+						TransferedConstraint->ResetActiveManifoldContacts();
+						for (FManifoldPoint& Manifold : CollisionConstraint.GetManifoldPoints())
+						{
+							TransferedConstraint->AddOneshotManifoldContact(Manifold.ContactPoint);
+						}
 
-		TMap<TGeometryParticleHandle<FReal, 3>*, TArray<FPBDCollisionConstraintHandle*> >& CollisionMap = CollisionConstraints.GetParticleCollisionsMap();
-		for (auto& CollisionPair : CollisionMap)
-		{
-			if (FRigidHandle* ChildParticle = CollisionPair.Key->CastToRigidParticle())
-			{
-				const FRigidTransform3 ChildTransform = FParticleUtilities::GetActorWorldTransform(FConstGenericParticleHandle(ChildParticle));
+						FReal CollisionConstraintStiffness = JointConstraint->GetSettings().ContactTransferScale;
+						if (TransferCollisionsDebugTestAgainstMaxClamp)
+						{
+							CollisionConstraintStiffness = TransferCollisionsStiffnessClamp;
+						}
 
-				for (FConstraintHandle* Constraint : ChildParticle->ParticleConstraints())
-				{
-					if (FPBDJointConstraintHandle* JointConstraint = Constraint->As<FPBDJointConstraintHandle>())
-					{
-						if (!FMath::IsNearlyZero(JointConstraint->GetSettings().ContactTransferScale))
+						if (CollisionParticleA.Get() && CollisionParticleA->ObjectState() != EObjectStateType::Dynamic)
 						{
-							if (FRigidHandle* ParentParticle = JointConstraint->GetConstrainedParticles()[1]->CastToRigidParticle()) // Parent
+							if (ensureMsgf(TransferCollisionsKinematicScale > 0, TEXT("Zero or Negative TransferCollisionsKinematicScale")))
 							{
-								const FRigidTransform3 ParentTransform = FParticleUtilities::GetActorWorldTransform(FConstGenericParticleHandle(ParentParticle));
-								const FTransform ChildToParentTransform = ChildTransform.GetRelativeTransform(ParentTransform);
-
-								for (FPBDCollisionConstraintHandle* ContactHandle : CollisionPair.Value)
-								{
-									if (ContactHandle->GetCCDType() == ECollisionCCDType::Disabled)
-									{
-										const FPBDCollisionConstraint& CurrConstraint = ContactHandle->GetContact();
-
-										FPBDCollisionConstraint* TransferedConstraint;
-										if (ContactHandle->GetConstrainedParticles()[0] == ChildParticle)
-										{
-											TransferedConstraint = FPBDCollisionConstraint::Make(
-												ParentParticle,
-												CurrConstraint.GetManifold().Implicit[0],
-												CurrConstraint.GetManifold().Simplicial[0],
-												ChildTransform,
-												ChildToParentTransform.GetRelativeTransform(CurrConstraint.ImplicitTransform[0]),
-												CurrConstraint.Particle[1],
-												CurrConstraint.GetManifold().Implicit[1],
-												CurrConstraint.GetManifold().Simplicial[1],
-												ParentTransform,
-												CurrConstraint.ImplicitTransform[1],
-												CurrConstraint.GetCullDistance(),
-												CurrConstraint.GetManifold().ShapesType,
-												CurrConstraint.GetUseManifold(),
-												CollisionAllocator
-											);
-										}
-										else
-										{
-											TransferedConstraint = FPBDCollisionConstraint::Make(
-												CurrConstraint.Particle[0],
-												CurrConstraint.GetManifold().Implicit[0],
-												CurrConstraint.GetManifold().Simplicial[0],
-												ParentTransform,
-												CurrConstraint.ImplicitTransform[0],
-												ParentParticle,
-												CurrConstraint.GetManifold().Implicit[1],
-												CurrConstraint.GetManifold().Simplicial[1],
-												ChildTransform,
-												ChildToParentTransform.GetRelativeTransform(CurrConstraint.ImplicitTransform[1]),
-												CurrConstraint.GetCullDistance(),
-												CurrConstraint.GetManifold().ShapesType,
-												CurrConstraint.GetUseManifold(),
-												CollisionAllocator
-											);
-										}
-										TArray<FManifoldPoint> TransferManifolds;
-										for (const FManifoldPoint& Manifold : CurrConstraint.GetManifoldPoints())
-										{
-											TransferManifolds.Add(Manifold);
-										}
-										TransferedConstraint->SetManifoldPoints(TransferManifolds);
-										for (FManifoldPoint& Manifold : TransferedConstraint->GetManifoldPoints())
-										{
-											TransferedConstraint->UpdateManifoldPointFromContact(Manifold);
-										}
-										TransferedConstraint->SetStiffness(JointConstraint->GetSettings().ContactTransferScale);
-									}
-								}
+								CollisionConstraintStiffness *= TransferCollisionsKinematicScale;
 							}
 						}
-					}
-				}
+						TransferedConstraint->SetStiffness(FMath::Clamp(CollisionConstraintStiffness, (FReal)0.f, TransferCollisionsStiffnessClamp));
+
+						FParticlePairMidPhase* MidPhase = CollisionAllocator.GetParticlePairMidPhase(CollisionParticleA->Handle(), CollisionParticleB->Handle());
+						MidPhase->InjectCollision(*TransferedConstraint);
+
+						return ECollisionVisitorResult::Continue;
+					});
 			}
 		}
 
