Chaos : Clustering
- Fix unprotected access to pointer

diff --git a/Engine/Source/Runtime/Experimental/Chaos/Private/Chaos/PBDRigidClustering.cpp b/Engine/Source/Runtime/Experimental/Chaos/Private/Chaos/PBDRigidClustering.cpp
--- a/Engine/Source/Runtime/Experimental/Chaos/Private/Chaos/PBDRigidClustering.cpp
+++ b/Engine/Source/Runtime/Experimental/Chaos/Private/Chaos/PBDRigidClustering.cpp
@@ -715,11 +715,11 @@ namespace Chaos
 						ClusterBreak.Velocity = Child->V();
 						ClusterBreak.AngularVelocity = Child->W();
 						ClusterBreak.Mass = Child->M();
-						if (Child->Geometry()->HasBoundingBox())
+						if (Child->Geometry() && Child->Geometry()->HasBoundingBox())
 						{
 							ClusterBreak.BoundingBox = Child->Geometry()->BoundingBox();
 						}
-						if (ClusterBreak.Proxy->GetType() == EPhysicsProxyType::GeometryCollectionType)
+						if (ClusterBreak.Proxy && ClusterBreak.Proxy->GetType() == EPhysicsProxyType::GeometryCollectionType)
 						{
 							FGeometryCollectionPhysicsProxy* ConcreteProxy = static_cast<FGeometryCollectionPhysicsProxy*>(ClusterBreak.Proxy);
 							ClusterBreak.TransformGroupIndex = ConcreteProxy->GetTransformGroupIndexFromHandle(Child);
