Joint Linear Solver

diff --git a/Engine/Source/Programs/HeadlessChaos/Private/HeadlessChaosTestJoints.cpp b/Engine/Source/Programs/HeadlessChaos/Private/HeadlessChaosTestJoints.cpp
--- a/Engine/Source/Programs/HeadlessChaos/Private/HeadlessChaosTestJoints.cpp
+++ b/Engine/Source/Programs/HeadlessChaos/Private/HeadlessChaosTestJoints.cpp
@@ -207,16 +207,16 @@ namespace ChaosTest {
 			Test.Evolution.EndFrame(Dt);
 
 			// Kinematic particle should have moved to animated position
-			EXPECT_LT((Test.GetParticle(Box1Id)->X() - RootPosition).Size(), (FReal)0.01 * BoxSize) << "Post-move instability on frame " << i;
+			EXPECT_LT((Test.GetParticle(Box1Id)->X() - RootPosition).Size(), (FReal)0.1 * BoxSize) << "Post-move instability on frame " << i;
 
 			// Particles should remain fixed distance apart (joint point is at Box1 location)
 			const FVec3 Delta = Test.GetParticle(Box2Id)->CastToRigidParticle()->P() - Test.GetParticle(Box1Id)->X();
 			const FReal Distance = Delta.Size();
-			EXPECT_NEAR(Distance, ExpectedDistance, (FReal)0.01 * BoxSize) << "Post-move instability on frame " << i;
+			EXPECT_NEAR(Distance, ExpectedDistance, (FReal)0.1 * BoxSize) << "Post-move instability on frame " << i;
 
 			// Joint position calculted from pose and local-space joint pos
 			const FVec3 Box2WorldSpaceJointPosition = Test.GetParticle(Box2Id)->R().RotateVector(Box2LocalSpaceJointPosition) + Test.GetParticle(Box2Id)->X();
-			EXPECT_LT((Box2WorldSpaceJointPosition - RootPosition).Size(), (FReal)0.01 * BoxSize) << "Post-move instability on frame " << i;
+			EXPECT_LT((Box2WorldSpaceJointPosition - RootPosition).Size(), (FReal)0.1 * BoxSize) << "Post-move instability on frame " << i;
 		}
 	}
 
@@ -233,7 +233,7 @@ namespace ChaosTest {
 		const FReal BoxMass = 1000;
 		const FReal Dt = (FReal)1 / 20;
 		const FReal AnimPeriod = (FReal)2;
-		const FVec3 AnimDelta = FVec3(10 * BoxSize, 0, 0);
+		const FVec3 AnimDelta = FVec3(BoxSize, 0, 0);
 
 		FJointConstraintsTest<TEvolution> Test(NumIterations, Gravity);
 
diff --git a/Engine/Source/Runtime/Experimental/Chaos/Private/Chaos/Joint/JointConstraintsCVars.cpp b/Engine/Source/Runtime/Experimental/Chaos/Private/Chaos/Joint/JointConstraintsCVars.cpp
--- a/Engine/Source/Runtime/Experimental/Chaos/Private/Chaos/Joint/JointConstraintsCVars.cpp
+++ b/Engine/Source/Runtime/Experimental/Chaos/Private/Chaos/Joint/JointConstraintsCVars.cpp
@@ -29,5 +29,5 @@ FAutoConsoleVariableRef CVarChaosJointLinearVelocityThresholdToApplyRestitution(
 float Chaos_Joint_AngularVelocityThresholdToApplyRestitution = 1e-2f;
 FAutoConsoleVariableRef CVarChaosJointAngularVelocityThresholdToApplyRestitution(TEXT("p.Chaos.Joint.AngularVelocityThresholdToApplyRestitution"), Chaos_Joint_AngularVelocityThresholdToApplyRestitution, TEXT("Apply restitution only if initial velocity is higher than this threshold (used in Quasipbd)"));
 
-bool bChaos_Joint_UseCachedSolver = false;
+bool bChaos_Joint_UseCachedSolver = true;
 FAutoConsoleVariableRef CVarChaosJointUseCachedSolver(TEXT("p.Chaos.Joint.UseCachedSolver"), bChaos_Joint_UseCachedSolver, TEXT("Use cached version of joint solver. (default is false"));
diff --git a/Engine/Source/Runtime/Experimental/Chaos/Private/Chaos/Joint/PBDJointCachedSolverGaussSeidel.cpp b/Engine/Source/Runtime/Experimental/Chaos/Private/Chaos/Joint/PBDJointCachedSolverGaussSeidel.cpp
--- a/Engine/Source/Runtime/Experimental/Chaos/Private/Chaos/Joint/PBDJointCachedSolverGaussSeidel.cpp
+++ b/Engine/Source/Runtime/Experimental/Chaos/Private/Chaos/Joint/PBDJointCachedSolverGaussSeidel.cpp
@@ -6,3188 +6,1610 @@
 #include "Chaos/ParticleHandle.h"
 #include "Chaos/PBDJointConstraintUtilities.h"
 #include "Chaos/Utilities.h"
-#include "ChaosStats.h"
-#if 0
-#include "PBDJointSolverGaussSeidel.ispc.generated.h"
-#endif
-
-//PRAGMA_DISABLE_OPTIMIZATION
 
 namespace Chaos
 {
+	
+FPBDJointCachedSolver::FPBDJointCachedSolver()
+{}
+	
+/** Derived states management */
 
+void FPBDJointCachedSolver::InitDerivedState()
+{
+	InitConnectorXs[0] = X(0) + R(0) * LocalConnectorXs[0].GetTranslation();
+	InitConnectorXs[1] = X(1) + R(1) * LocalConnectorXs[1].GetTranslation();
+	InitConnectorRs[0] = R(0) * LocalConnectorXs[0].GetRotation();
+	InitConnectorRs[1] = R(1) * LocalConnectorXs[1].GetRotation();
+	InitConnectorRs[1].EnforceShortestArcWith(InitConnectorRs[0]);
+
+	ConnectorXs[0] = InitConnectorXs[0];
+	ConnectorRs[0] = InitConnectorRs[0];
+
+	ConnectorXs[1] = InitConnectorXs[1];
+	ConnectorRs[1] = InitConnectorRs[1];
+	
+	UpdateDerivedState();
+}
 
-    //
-    //
-    //////////////////////////////////////////////////////////////////////////
-    //
-    //
-
-    FPBDJointCachedSolver::FPBDJointCachedSolver()
-    {
-        if (bChaos_Joint_ISPC_Enabled)
-        {
-#if 0
-            check(sizeof(FPBDJointCachedSolver) == ispc::SizeofFPBDJointSolver());
-#endif
-        }
-    }
-
-    void FPBDJointCachedSolver::InitDerivedState()
-    {
-        InitConnectorXs[0] = X(0) + R(0) * LocalConnectorXs[0].GetTranslation();
-        InitConnectorXs[1] = X(1) + R(1) * LocalConnectorXs[1].GetTranslation();
-        InitConnectorRs[0] = R(0) * LocalConnectorXs[0].GetRotation();
-        InitConnectorRs[1] = R(1) * LocalConnectorXs[1].GetRotation();
-        InitConnectorRs[1].EnforceShortestArcWith(InitConnectorRs[0]);
-
-        ConnectorXs[0] = InitConnectorXs[0];
-        ConnectorRs[0] = InitConnectorRs[0];
-
-        ConnectorXs[1] = InitConnectorXs[1];
-        ConnectorRs[1] = InitConnectorRs[1];
-    }
-
-
-    void FPBDJointCachedSolver::UpdateDerivedState()
-    {
-        // Kinematic bodies will not be moved, so we don't update derived state during iterations
-        if (InvM(0) > SMALL_NUMBER)
-        {
-            const FVec3 BodyP0 = P(0);
-            const FRotation3 BodyQ0 = Q(0);
-            ConnectorXs[0] = BodyP0 + BodyQ0 * LocalConnectorXs[0].GetTranslation();
-            ConnectorRs[0] = BodyQ0 * LocalConnectorXs[0].GetRotation();
-        }
-        if (InvM(1) > SMALL_NUMBER)
-        {
-            const FVec3 BodyP1 = P(1);
-            const FRotation3 BodyQ1 = Q(1);
-            ConnectorXs[1] = BodyP1 + BodyQ1 * LocalConnectorXs[1].GetTranslation();
-            ConnectorRs[1] = BodyQ1 * LocalConnectorXs[1].GetRotation();
-        }
-        ConnectorRs[1].EnforceShortestArcWith(ConnectorRs[0]);
-    }
-
-
-    void FPBDJointCachedSolver::UpdateDerivedState(const int32 BodyIndex)
-    {
-        const FVec3 BodyP = P(BodyIndex);
-        const FRotation3 BodyQ = Q(BodyIndex);
-        ConnectorXs[BodyIndex] = BodyP + BodyQ * LocalConnectorXs[BodyIndex].GetTranslation();
-        ConnectorRs[BodyIndex] = BodyQ * LocalConnectorXs[BodyIndex].GetRotation();
-        ConnectorRs[1].EnforceShortestArcWith(ConnectorRs[0]);
-    }
-
-    bool FPBDJointCachedSolver::UpdateIsActive(
-        const FReal Dt,
-        const FPBDJointSolverSettings& SolverSettings,
-        const FPBDJointSettings& JointSettings)
-    {
-        // NumActiveConstraints is initialized to -1, so there's no danger of getting invalid LastPs/Qs
-        // We also check SolverStiffness mainly for testing when solver stiffness is 0 (so we don't exit immediately)
-        if ((NumActiveConstraints >= 0) && (SolverStiffness > 0.0f))
-        {
-            bool bIsSolved =
-                FVec3::IsNearlyEqual(P(0), LastPs[0], PositionTolerance)
-                && FVec3::IsNearlyEqual(P(1), LastPs[1], PositionTolerance)
-                && FRotation3::IsNearlyEqual(Q(0), LastQs[0], 0.5f * AngleTolerance)
-                && FRotation3::IsNearlyEqual(Q(1), LastQs[1], 0.5f * AngleTolerance);
-            bIsActive = !bIsSolved;
-        }
-
-        LastPs[0] = P(0);
-        LastPs[1] = P(1);
-        LastQs[0] = Q(0);
-        LastQs[1] = Q(1);
-
-        return bIsActive;
-    }
-
-
-    void FPBDJointCachedSolver::Init(
-        const FReal Dt,
-        const FSolverBodyPtrPair& SolverBodyPair,
-        const FPBDJointSolverSettings& SolverSettings,
-        const FPBDJointSettings& JointSettings,
-        const FRigidTransform3& XL0,
-        const FRigidTransform3& XL1)
-    {
-        SolverBodies[0] = *SolverBodyPair[0];
-        SolverBodies[1] = *SolverBodyPair[1];
-
-        LocalConnectorXs[0] = XL0;
-        LocalConnectorXs[1] = XL1;
-
-        // \todo(chaos): joint should support parent/child in either order
-        SolverBodies[0].SetInvMScale(JointSettings.ParentInvMassScale);
-        SolverBodies[1].SetInvMScale(FReal(1));
-
-        FVec3 ConditionedInvILs[2];
-        FPBDJointUtilities::ConditionInverseMassAndInertia(Body0().InvM(), Body1().InvM(), Body0().InvILocal(), Body1().InvILocal(), SolverSettings.MinParentMassRatio, SolverSettings.MaxInertiaRatio, ConditionedInvMs[0], ConditionedInvMs[1], ConditionedInvILs[0], ConditionedInvILs[1]);
-        ConditionedInvIs[0] = Utilities::ComputeWorldSpaceInertia(Q(0), ConditionedInvILs[0]);
-        ConditionedInvIs[1] = Utilities::ComputeWorldSpaceInertia(Q(1), ConditionedInvILs[1]);
-
-        NetLinearImpulse = FVec3(0);
-        NetAngularImpulse = FVec3(0);
-
-        LinearSoftLambda = 0;
-        TwistSoftLambda = 0;
-        SwingSoftLambda = 0;
-        LinearDriveLambdas = FVec3(0);
-        RotationDriveLambdas = FVec3(0);
-
-        LinearConstraintPadding = FVec3(-1);
-        AngularConstraintPadding = FVec3(-1);
-
-        // Tolerances are positional errors below visible detection. But in PBD the errors
-        // we leave behind get converted to velocity, so we need to ensure that the resultant
-        // movement from that erroneous velocity is less than the desired position tolerance.
-        // Assume that the tolerances were defined for a 60Hz simulation, then it must be that
-        // the position error is less than the position change from constant external forces
-        // (e.g., gravity). So, we are saying that the tolerance was chosen because the position
-        // error is less that F.dt^2. We need to scale the tolerance to work at our current dt.
-        const FReal ToleranceScale = FMath::Min(1.f, 60.f * 60.f * Dt * Dt);
-        PositionTolerance = ToleranceScale * SolverSettings.PositionTolerance;
-        AngleTolerance = ToleranceScale * SolverSettings.AngleTolerance;
-
-        NumActiveConstraints = -1;
-        bIsActive = true;
-
-        SolverStiffness = 1.0f;
-
-        LinearHardLambda = FVec3(0);
-        AngularHardLambda = FVec3(0);
-
-        InitDerivedState();
-
-        if (JointSettings.LinearRestitution != 0.0f)
-        {
-            CalculateConstraintAxisLinearVelocities(JointSettings, InitConstraintAxisLinearVelocities);
-        }
-        if (JointSettings.TwistRestitution != 0.0f || JointSettings.SwingRestitution != 0.0f)
-        {
-            CalculateConstraintAxisAngularVelocities(SolverSettings, JointSettings, InitConstraintAxisAngularVelocities);
-        }
-
-        InitPositionConstraints(SolverSettings, JointSettings);
-        LastPs[0] = FVec3(0.f);
-        LastPs[1] = FVec3(0.f);
-        LastQs[0] = FRotation3::FromIdentity();
-        LastQs[1] = FRotation3::FromIdentity();
-    }
-
-    void FPBDJointCachedSolver::InitPositionConstraints(
-        const FPBDJointSolverSettings& SolverSettings,
-        const FPBDJointSettings& JointSettings)
-    {
-        bool bHasPositionConstraints =
-            (JointSettings.LinearMotionTypes[0] != EJointMotionType::Free)
-            || (JointSettings.LinearMotionTypes[1] != EJointMotionType::Free)
-            || (JointSettings.LinearMotionTypes[2] != EJointMotionType::Free);
-        if (!bHasPositionConstraints)
-        {
-            return;
-        }
-
-        if (FPBDJointUtilities::GetSoftLinearLimitEnabled(SolverSettings, JointSettings))
-        {
-            LinearStiffness = FPBDJointUtilities::GetSoftLinearStiffness(SolverSettings, JointSettings);
-            LinearDamping = FPBDJointUtilities::GetSoftLinearDamping(SolverSettings, JointSettings);
-            bLinearAccelerationMode = FPBDJointUtilities::GetLinearSoftAccelerationMode(SolverSettings, JointSettings);
-        }
-        else
-        {
-            LinearStiffness = FPBDJointUtilities::GetLinearStiffness(SolverSettings, JointSettings);
-        }
-
-        const TVec3<EJointMotionType>& LinearMotion = JointSettings.LinearMotionTypes;
-        const TVec3<bool> bLinearLocked =
-        {
-            (LinearMotion[0] == EJointMotionType::Locked),
-            (LinearMotion[1] == EJointMotionType::Locked),
-            (LinearMotion[2] == EJointMotionType::Locked),
-        };
-        const TVec3<bool> bLinearLimited =
-        {
-            (LinearMotion[0] == EJointMotionType::Limited),
-            (LinearMotion[1] == EJointMotionType::Limited),
-            (LinearMotion[2] == EJointMotionType::Limited),
-        };
-
-        if (bLinearLocked[0] && bLinearLocked[1] && bLinearLocked[2])
-        {
-            // Hard point constraint (most common case)
-            if (!Body0().IsDynamic())
-            {
-                InitPointPositionConstraintKD(0, 1);
-            }
-            else if (!Body1().IsDynamic())
-            {
-                InitPointPositionConstraintKD(1, 0);
-            }
-            else
-            {
-                InitPointPositionConstraintDD();
-            }
-        }
-        else if (bLinearLimited[0] && bLinearLimited[1] && bLinearLimited[2])
-        {
-            // Spherical constraint
-            InitSphericalPositionConstraint(JointSettings);
-        }
-        else if (bLinearLocked[1] && bLinearLocked[2] && !bLinearLocked[0])
-        {
-            // Line constraint along X axis
-            InitCylindricalPositionConstraint(0, LinearMotion[0], EJointMotionType::Locked, JointSettings);
-        }
-        else if (bLinearLocked[0] && bLinearLocked[2] && !bLinearLocked[1])
-        {
-            // Line constraint along Y axis
-            InitCylindricalPositionConstraint(1, LinearMotion[1], EJointMotionType::Locked, JointSettings);
-        }
-        else if (bLinearLocked[0] && bLinearLocked[1] && !bLinearLocked[2])
-        {
-            // Line constraint along Z axis
-            InitCylindricalPositionConstraint(2, LinearMotion[2], EJointMotionType::Locked, JointSettings);
-        }
-        else if (bLinearLimited[1] && bLinearLimited[2] && !bLinearLimited[0])
-        {
-            // Cylindrical constraint along X axis
-            InitCylindricalPositionConstraint(0, LinearMotion[0], EJointMotionType::Limited, JointSettings);
-        }
-        else if (bLinearLimited[0] && bLinearLimited[2] && !bLinearLimited[1])
-        {
-            // Cylindrical constraint along Y axis
-            InitCylindricalPositionConstraint(1, LinearMotion[1], EJointMotionType::Limited, JointSettings);
-        }
-        else if (bLinearLimited[0] && bLinearLimited[1] && !bLinearLimited[2])
-        {
-            // Cylindrical constraint along Z axis
-            InitCylindricalPositionConstraint(2, LinearMotion[2], EJointMotionType::Limited, JointSettings);
-        }
-        else if (bLinearLocked[0] || bLinearLimited[0])
-        {
-            // Planar constraint along X axis
-            InitPlanarPositionConstraint(0, LinearMotion[0], JointSettings);
-        }
-        else if (bLinearLocked[1] || bLinearLimited[1])
-        {
-            // Planar constraint along Y axis
-            InitPlanarPositionConstraint(1, LinearMotion[1], JointSettings);
-        }
-        else if (bLinearLocked[2] || bLinearLimited[2])
-        {
-            // Planar constraint along Z axis
-            InitPlanarPositionConstraint(2, LinearMotion[2], JointSettings);
-        }
-    }
-
-    void FPBDJointCachedSolver::InitPointPositionConstraintKD(
-        const int32 KIndex,
-        const int32 DIndex)
-    {
-        check(!Body(KIndex).IsDynamic());
-        check(Body(DIndex).IsDynamic());
-        const FVec3 ConstraintArm = ConnectorXs[DIndex] - P(DIndex);
-        const FMatrix33 M = Utilities::ComputeJointFactorMatrix(ConstraintArm, InvI(DIndex), InvM(DIndex));
-        ConstraintArms[PointPositionConstraintIndex][DIndex] = ConstraintArm;
-        ConstraintArms[PointPositionConstraintIndex][KIndex] = ConnectorXs[KIndex] - P(KIndex);
-        PointPositionConstraintIM = M.Inverse();
-    }
-
-    void FPBDJointCachedSolver::InitPointPositionConstraintDD()
-    {
-        check(Body(0).IsDynamic());
-        check(Body(1).IsDynamic());
-        const FVec3 ConstraintArm0 = ConnectorXs[0] - P(0);
-        const FVec3 ConstraintArm1 = ConnectorXs[1] - P(1);
-        const FMatrix33 M0 = Utilities::ComputeJointFactorMatrix(ConstraintArm0, InvI(0), InvM(0));
-        const FMatrix33 M1 = Utilities::ComputeJointFactorMatrix(ConstraintArm1, InvI(1), InvM(1));
-        ConstraintArms[PointPositionConstraintIndex][0] = ConstraintArm0;
-        ConstraintArms[PointPositionConstraintIndex][1] = ConstraintArm1;
-        PointPositionConstraintIM = (M0 + M1).Inverse();
-    }
-
-    void FPBDJointCachedSolver::InitSphericalPositionConstraint(
-        const FPBDJointSettings& JointSettings)
-    {
-        FVec3 Axis;
-        FReal Delta;
-        FPBDJointUtilities::GetSphericalAxisDelta(ConnectorXs[0], ConnectorXs[1], Axis, Delta);
-
-        const FReal LimitPadding = GetLinearConstraintPadding(0);
-        const FReal Limit = FMath::Max(JointSettings.LinearLimit - LimitPadding, (FReal)0.);
-
-        // This uses ConnectorXs[1] rather than ConnectorXs[0] for stability reasons. It is more stable this way when ConnectorXs[0] and ConnectorXs[1] are far apart 
-        const FVec3 ConstraintArm0 = ConnectorXs[1] - P(0);
-        const FVec3 ConstraintArm1 = ConnectorXs[1] - P(1);
-
-        InitPositionConstraint(Axis, ConstraintArm0, ConstraintArm1, SphericalPositionConstraintIndex);
-    }
-
-    void FPBDJointCachedSolver::InitCylindricalPositionConstraint(
-        const int32 AxisIndex,
-        const EJointMotionType AxialMotion,
-        const EJointMotionType RadialMotion,
-        const FPBDJointSettings& JointSettings)
-    {
-        check(AxialMotion != RadialMotion);
-
-        FVec3 Axis, RadialAxis;
-        FReal AxialDelta, RadialDelta;
-        FPBDJointUtilities::GetCylindricalAxesDeltas(ConnectorRs[0], ConnectorXs[0], ConnectorXs[1], AxisIndex, Axis, AxialDelta, RadialAxis, RadialDelta);
-
-        if (AxialDelta < 0.0f)
-        {
-            Axis = -Axis;
-        }
-
-        // This uses ConnectorXs[1] rather than ConnectorXs[0] for stability reasons. It is more stable this way when ConnectorXs[0] and ConnectorXs[1] are far apart 
-        const FVec3 ConstraintArm0 = ConnectorXs[1] - P(0);
-        const FVec3 ConstraintArm1 = ConnectorXs[1] - P(1);
-        InitPositionConstraint(Axis, ConstraintArm0, ConstraintArm1, (int32)EJointCylindricalPositionConstraintType::Axial);   
-        InitPositionConstraint(RadialAxis, ConstraintArm0, ConstraintArm1, (int32)EJointCylindricalPositionConstraintType::Radial);
-    }
-
-    void FPBDJointCachedSolver::InitPlanarPositionConstraint(
-        const int32 AxisIndex,
-        const EJointMotionType AxialMotion,
-        const FPBDJointSettings& JointSettings)
-    {
-        FVec3 Axis;
-        FReal Delta;
-        FPBDJointUtilities::GetPlanarAxisDelta(ConnectorRs[0], ConnectorXs[0], ConnectorXs[1], AxisIndex, Axis, Delta);
-
-        if (Delta < 0.0f)
-        {
-            Axis = -Axis;
-        }
-        // This uses ConnectorXs[1] rather than ConnectorXs[0] for stability reasons. It is more stable this way when ConnectorXs[0] and ConnectorXs[1] are far apart 
-        const FVec3 ConstraintArm0 = ConnectorXs[1] - P(0);
-        const FVec3 ConstraintArm1 = ConnectorXs[1] - P(1);
-        InitPositionConstraint(Axis, ConstraintArm0, ConstraintArm1, PlanarPositionConstraintIndex);
-    }
-
-    void FPBDJointCachedSolver::InitPositionConstraint(
-        const FVec3& Axis,
-        const FVec3& ConstraintArm0,
-        const FVec3& ConstraintArm1,
-        const int32 ConstraintIndex)
-    {
-        PositionConstraintAxis[ConstraintIndex] = Axis;
-        ConstraintArms[ConstraintIndex][0] = ConstraintArm0;
-        ConstraintArms[ConstraintIndex][1] = ConstraintArm1;
-
-        UpdatePositionConstraintIM(Axis, ConstraintArm0, ConstraintArm1, ConstraintIndex);
-    }
-
-    void FPBDJointCachedSolver::UpdatePositionConstraintIM(
-        const FVec3& Axis,
-        const FVec3& ConstraintArm0,
-        const FVec3& ConstraintArm1,
-        const int32 ConstraintIndex)
-    {
-        const FVec3 AngularAxis0 = FVec3::CrossProduct(ConstraintArm0, Axis);
-        const FVec3 AngularAxis1 = FVec3::CrossProduct(ConstraintArm1, Axis);
-        const FVec3 IA0 = Utilities::Multiply(InvI(0), AngularAxis0);
-        const FVec3 IA1 = Utilities::Multiply(InvI(1), AngularAxis1);
-        const FReal II0 = FVec3::DotProduct(AngularAxis0, IA0);
-        const FReal II1 = FVec3::DotProduct(AngularAxis1, IA1);
-
-        PositionConstraintDRAxis[ConstraintIndex][0] = IA0;
-        PositionConstraintDRAxis[ConstraintIndex][1] = -IA1;
-        PositionConstraintIM[ConstraintIndex] = InvM(0) + II0 + InvM(1) + II1;
-    }
-
-    void FPBDJointCachedSolver::Deinit()
-    {
-        SolverBodies[0].Reset();
-        SolverBodies[1].Reset();
-    }
-
-
-    void FPBDJointCachedSolver::Update(
-        const FReal Dt,
-        const FPBDJointSolverSettings& SolverSettings,
-        const FPBDJointSettings& JointSettings)
-    {
-        UpdateDerivedState();
-
-        UpdateIsActive(Dt, SolverSettings, JointSettings);
-    }
-
-    void FPBDJointCachedSolver::EnableProjection()
-    {
-        Body0().SetInvMScale(0);
-    }
-
-
-    void FPBDJointCachedSolver::ApplyConstraints(
-        const FReal Dt,
-        const FReal InSolverStiffness,
-        const FPBDJointSolverSettings& SolverSettings,
-        const FPBDJointSettings& JointSettings)
-    {
-        NumActiveConstraints = 0;
-        SolverStiffness = InSolverStiffness;
-
-        ApplyPositionConstraints(Dt, SolverSettings, JointSettings);
-        ApplyRotationConstraints(Dt, SolverSettings, JointSettings);
-    
-        ApplyPositionDrives(Dt, SolverSettings, JointSettings);
-        ApplyRotationDrives(Dt, SolverSettings, JointSettings);
-
-        UpdateIsActive(Dt, SolverSettings, JointSettings);
-    }
-
-
-    void FPBDJointCachedSolver::ApplyVelocityConstraints(
-        const FReal Dt,
-        const FReal InSolverStiffness,
-        const FPBDJointSolverSettings& SolverSettings,
-        const FPBDJointSettings& JointSettings)
-    {
-        SolverStiffness = InSolverStiffness;
-
-        // This is used for the QuasiPbd solver. If the Pbd step applied impulses to
-        // correct position errors, it will have introduced a velocity equal to the 
-        // corection divided by the timestep. We ensure that the velocity constraints
-        // (including restitution) are also enforced. This also prevents any position
-        // errors from the previous frame getting converted into energy.
-
-        ApplyLinearVelocityConstraints(Dt, SolverSettings, JointSettings);
-        ApplyAngularVelocityConstraints(Dt, SolverSettings, JointSettings);
-
-        // @todo(chaos): We can also apply velocity drives here rather than in the Pbd pass
-    }
-
-
-    void FPBDJointCachedSolver::ApplyProjections(
-        const FReal Dt,
-        const FReal InSolverStiffness,
-        const FPBDJointSolverSettings& SolverSettings,
-        const FPBDJointSettings& JointSettings)
-    {
-        // @todo(chaos): We need to handle parent/child being the other way round
-        if (!Body1().IsDynamic())
-        {
-            // If child is kinematic, return. 
-            return;
-        }
-
-        SolverStiffness = InSolverStiffness;
-
-        const FReal LinearProjection = FPBDJointUtilities::GetLinearProjection(SolverSettings, JointSettings);
-        const FReal AngularProjection = FPBDJointUtilities::GetAngularProjection(SolverSettings, JointSettings);
-
-        FVec3 DP1 = FVec3(0);
-        FVec3 DR1 = FVec3(0);
-
-        // Position Projection
-        const bool bLinearSoft = FPBDJointUtilities::GetSoftLinearLimitEnabled(SolverSettings, JointSettings);
-        const bool bLinearProjectionEnabled = (bLinearSoft && JointSettings.bSoftProjectionEnabled) || (!bLinearSoft && JointSettings.bProjectionEnabled);
-        const TVec3<EJointMotionType>& LinearMotion = JointSettings.LinearMotionTypes;
-        const bool bLinearLocked =
-            (LinearMotion[0] == EJointMotionType::Locked)
-            && (LinearMotion[1] == EJointMotionType::Locked)
-            && (LinearMotion[2] == EJointMotionType::Locked);
-        const bool bLinearLimited =
-            (LinearMotion[0] == EJointMotionType::Limited)
-            && (LinearMotion[1] == EJointMotionType::Limited)
-            && (LinearMotion[2] == EJointMotionType::Limited);
-        if (bLinearProjectionEnabled && (LinearProjection > 0))
-        {
-            if (bLinearLocked)
-            {
-                ApplyPointProjection(Dt, SolverSettings, JointSettings, LinearProjection, DP1, DR1);
-            }
-            else if (bLinearLimited)
-            {
-                ApplySphereProjection(Dt, SolverSettings, JointSettings, LinearProjection, DP1, DR1);
-            }
-            // @todo(ccaulfield): support mixed linear projection
-        }
-
-        // Twist projection
-        const bool bTwistSoft = FPBDJointUtilities::GetSoftTwistLimitEnabled(SolverSettings, JointSettings);
-        const bool bTwistProjectionEnabled = SolverSettings.bEnableTwistLimits && ((bTwistSoft && JointSettings.bSoftProjectionEnabled) || (!bTwistSoft && JointSettings.bProjectionEnabled));
-        if (bTwistProjectionEnabled && (AngularProjection > 0.0f))
-        {
-            const EJointMotionType TwistMotion = JointSettings.AngularMotionTypes[(int32)EJointAngularConstraintIndex::Twist];
-            if (TwistMotion != EJointMotionType::Free)
-            {
-                ApplyTwistProjection(Dt, SolverSettings, JointSettings, AngularProjection, bLinearLocked, DP1, DR1);
-            }
-        }
-
-        // Swing projection
-        const bool bSwingSoft = FPBDJointUtilities::GetSoftSwingLimitEnabled(SolverSettings, JointSettings);
-        const bool bSwingProjectionEnabled = SolverSettings.bEnableSwingLimits && ((bSwingSoft && JointSettings.bSoftProjectionEnabled) || (!bSwingSoft && JointSettings.bProjectionEnabled));
-        if (bSwingProjectionEnabled && (AngularProjection > 0.0f))
-        {
-            const EJointMotionType Swing1Motion = JointSettings.AngularMotionTypes[(int32)EJointAngularConstraintIndex::Swing1];
-            const EJointMotionType Swing2Motion = JointSettings.AngularMotionTypes[(int32)EJointAngularConstraintIndex::Swing2];
-            if ((Swing1Motion == EJointMotionType::Limited) && (Swing2Motion == EJointMotionType::Limited))
-            {
-                ApplyConeProjection(Dt, SolverSettings, JointSettings, AngularProjection, bLinearLocked, DP1, DR1);
-            }
-            else if ((Swing1Motion == EJointMotionType::Limited) && (Swing2Motion == EJointMotionType::Locked))
-            {
-                ApplySingleLockedSwingProjection(Dt, SolverSettings, JointSettings, EJointAngularConstraintIndex::Swing2, AngularProjection, bLinearLocked, DP1, DR1);
-                ApplySwingProjection(Dt, SolverSettings, JointSettings, EJointAngularConstraintIndex::Swing1, AngularProjection, bLinearLocked, DP1, DR1);
-            }
-            else if ((Swing1Motion == EJointMotionType::Limited) && (Swing2Motion == EJointMotionType::Free))
-            {
-                ApplyDualConeSwingProjection(Dt, SolverSettings, JointSettings, EJointAngularConstraintIndex::Swing1, AngularProjection, bLinearLocked, DP1, DR1);
-            }
-            else if ((Swing1Motion == EJointMotionType::Locked) && (Swing2Motion == EJointMotionType::Limited))
-            {
-                ApplySingleLockedSwingProjection(Dt, SolverSettings, JointSettings, EJointAngularConstraintIndex::Swing1, AngularProjection, bLinearLocked, DP1, DR1);
-                ApplySwingProjection(Dt, SolverSettings, JointSettings, EJointAngularConstraintIndex::Swing2, AngularProjection, bLinearLocked, DP1, DR1);
-            }
-            else if ((Swing1Motion == EJointMotionType::Locked) && (Swing2Motion == EJointMotionType::Locked))
-            {
-                ApplyDoubleLockedSwingProjection(Dt, SolverSettings, JointSettings, AngularProjection, bLinearLocked, DP1, DR1);
-            }
-            else if ((Swing1Motion == EJointMotionType::Locked) && (Swing2Motion == EJointMotionType::Free))
-            {
-                ApplySingleLockedSwingProjection(Dt, SolverSettings, JointSettings, EJointAngularConstraintIndex::Swing1, AngularProjection, bLinearLocked, DP1, DR1);
-            }
-            else if ((Swing1Motion == EJointMotionType::Free) && (Swing2Motion == EJointMotionType::Limited))
-            {
-                ApplyDualConeSwingProjection(Dt, SolverSettings, JointSettings, EJointAngularConstraintIndex::Swing2, AngularProjection, bLinearLocked, DP1, DR1);
-            }
-            else if ((Swing1Motion == EJointMotionType::Free) && (Swing2Motion == EJointMotionType::Locked))
-            {
-                ApplySingleLockedSwingProjection(Dt, SolverSettings, JointSettings, EJointAngularConstraintIndex::Swing2, AngularProjection, bLinearLocked, DP1, DR1);
-            }
-        }
-
-        // Final position fixup
-        if (bLinearLocked)
-        {
-            if (bLinearProjectionEnabled && (LinearProjection > 0))
-            {
-                ApplyTranslateProjection(Dt, SolverSettings, JointSettings, LinearProjection, DP1, DR1);
-            }
-
-            // Add velocity correction from the net projection motion
-            if (Chaos_Joint_VelProjectionAlpha > 0.0f)
-            {
-                ApplyVelocityProjection(Dt, SolverSettings, JointSettings, Chaos_Joint_VelProjectionAlpha, DP1, DR1);
-            }
-        }
-    
-        UpdateIsActive(Dt, SolverSettings, JointSettings);
-    }
-
-
-    void FPBDJointCachedSolver::ApplyRotationConstraints(
-        const FReal Dt,
-        const FPBDJointSolverSettings& SolverSettings,
-        const FPBDJointSettings& JointSettings)
-    {
-        bool bHasRotationConstraints =
-            (JointSettings.AngularMotionTypes[0] != EJointMotionType::Free)
-            || (JointSettings.AngularMotionTypes[1] != EJointMotionType::Free)
-            || (JointSettings.AngularMotionTypes[2] != EJointMotionType::Free);
-        if (!bHasRotationConstraints)
-        {
-            return;
-        }
-
-        // Locked axes always use hard constraints. Limited axes use hard or soft depending on settings
-        EJointMotionType TwistMotion = JointSettings.AngularMotionTypes[(int32)EJointAngularConstraintIndex::Twist];
-        EJointMotionType Swing1Motion = JointSettings.AngularMotionTypes[(int32)EJointAngularConstraintIndex::Swing1];
-        EJointMotionType Swing2Motion = JointSettings.AngularMotionTypes[(int32)EJointAngularConstraintIndex::Swing2];
-        bool bTwistSoft = FPBDJointUtilities::GetSoftTwistLimitEnabled(SolverSettings, JointSettings);
-        bool bSwingSoft = FPBDJointUtilities::GetSoftSwingLimitEnabled(SolverSettings, JointSettings);
-
-        // If the twist axes are opposing, we cannot decompose the orientation into swing and twist angles, so just give up
-        const FVec3 Twist0 = ConnectorRs[0] * FJointConstants::TwistAxis();
-        const FVec3 Twist1 = ConnectorRs[1] * FJointConstants::TwistAxis();
-        const FReal Twist01Dot = FVec3::DotProduct(Twist0, Twist1);
-        const bool bDegenerate = (Twist01Dot < Chaos_Joint_DegenerateRotationLimit);
-        if (bDegenerate)
-        {
-            UE_LOG(LogChaosJoint, VeryVerbose, TEXT(" Degenerate rotation at Swing %f deg"), FMath::RadiansToDegrees(FMath::Acos(Twist01Dot)));
-        }
-
-        // Apply twist constraint
-        // NOTE: Cannot calculate twist angle at 180degree swing
-        if (SolverSettings.bEnableTwistLimits && !bDegenerate)
-        {
-            if (TwistMotion == EJointMotionType::Limited)
-            {
-                ApplyTwistConstraint(Dt, SolverSettings, JointSettings, bTwistSoft);
-            }
-            else if (TwistMotion == EJointMotionType::Locked)
-            {
-                // Covered below
-            }
-            else if (TwistMotion == EJointMotionType::Free)
-            {
-            }
-        }
-
-        // Apply swing constraints
-        // NOTE: Cannot separate swing angles at 180degree swing (but we can still apply locks)
-        if (SolverSettings.bEnableSwingLimits)
-        {
-            if ((Swing1Motion == EJointMotionType::Limited) && (Swing2Motion == EJointMotionType::Limited))
-            {
-                ApplyConeConstraint(Dt, SolverSettings, JointSettings, bSwingSoft);
-            }
-            else if ((Swing1Motion == EJointMotionType::Limited) && (Swing2Motion == EJointMotionType::Locked))
-            {
-                ApplySingleLockedSwingConstraint(Dt, SolverSettings, JointSettings, EJointAngularConstraintIndex::Swing2, false);
-                if (!bDegenerate)
-                {
-                    ApplySwingConstraint(Dt, SolverSettings, JointSettings, EJointAngularConstraintIndex::Swing1, bSwingSoft);
-                }
-            }
-            else if ((Swing1Motion == EJointMotionType::Limited) && (Swing2Motion == EJointMotionType::Free))
-            {
-                if (!bDegenerate)
-                {
-                    ApplyDualConeSwingConstraint(Dt, SolverSettings, JointSettings, EJointAngularConstraintIndex::Swing1, bSwingSoft);
-                }
-            }
-            else if ((Swing1Motion == EJointMotionType::Locked) && (Swing2Motion == EJointMotionType::Limited))
-            {
-                ApplySingleLockedSwingConstraint(Dt, SolverSettings, JointSettings, EJointAngularConstraintIndex::Swing1, false);
-                if (!bDegenerate)
-                {
-                    ApplySwingConstraint(Dt, SolverSettings, JointSettings, EJointAngularConstraintIndex::Swing2, bSwingSoft);
-                }
-            }
-            else if ((Swing1Motion == EJointMotionType::Locked) && (Swing2Motion == EJointMotionType::Locked))
-            {
-                // Covered below
-            }
-            else if ((Swing1Motion == EJointMotionType::Locked) && (Swing2Motion == EJointMotionType::Free))
-            {
-                ApplySingleLockedSwingConstraint(Dt, SolverSettings, JointSettings, EJointAngularConstraintIndex::Swing1, false);
-            }
-            else if ((Swing1Motion == EJointMotionType::Free) && (Swing2Motion == EJointMotionType::Limited))
-            {
-                if (!bDegenerate)
-                {
-                    ApplyDualConeSwingConstraint(Dt, SolverSettings, JointSettings, EJointAngularConstraintIndex::Swing2, bSwingSoft);
-                }
-            }
-            else if ((Swing1Motion == EJointMotionType::Free) && (Swing2Motion == EJointMotionType::Locked))
-            {
-                ApplySingleLockedSwingConstraint(Dt, SolverSettings, JointSettings, EJointAngularConstraintIndex::Swing2, false);
-            }
-            else if ((Swing1Motion == EJointMotionType::Free) && (Swing2Motion == EJointMotionType::Free))
-            {
-            }
-        }
-
-        // Note: single-swing locks are already handled above so we only need to do something here if both are locked
-        bool bLockedTwist = SolverSettings.bEnableTwistLimits 
-            && (JointSettings.AngularMotionTypes[(int32)EJointAngularConstraintIndex::Twist] == EJointMotionType::Locked);
-        bool bLockedSwing = SolverSettings.bEnableSwingLimits 
-            && (JointSettings.AngularMotionTypes[(int32)EJointAngularConstraintIndex::Swing1] == EJointMotionType::Locked) 
-            && (JointSettings.AngularMotionTypes[(int32)EJointAngularConstraintIndex::Swing2] == EJointMotionType::Locked);
-        if (bLockedTwist || bLockedSwing)
-        {
-            ApplyLockedRotationConstraints(Dt, SolverSettings, JointSettings, bLockedTwist, bLockedSwing);
-        }
-    }
-
-
-    void FPBDJointCachedSolver::ApplyRotationDrives(
-        const FReal Dt,
-        const FPBDJointSolverSettings& SolverSettings,
-        const FPBDJointSettings& JointSettings)
-    {
-        bool bHasRotationDrives =
-            JointSettings.bAngularTwistPositionDriveEnabled
-            || JointSettings.bAngularTwistVelocityDriveEnabled
-            || JointSettings.bAngularSwingPositionDriveEnabled
-            || JointSettings.bAngularSwingVelocityDriveEnabled
-            || JointSettings.bAngularSLerpPositionDriveEnabled
-            || JointSettings.bAngularSLerpVelocityDriveEnabled;
-        if (!bHasRotationDrives)
-        {
-            return;
-        }
-
-        EJointMotionType TwistMotion = JointSettings.AngularMotionTypes[(int32)EJointAngularConstraintIndex::Twist];
-        EJointMotionType Swing1Motion = JointSettings.AngularMotionTypes[(int32)EJointAngularConstraintIndex::Swing1];
-        EJointMotionType Swing2Motion = JointSettings.AngularMotionTypes[(int32)EJointAngularConstraintIndex::Swing2];
-
-        if (SolverSettings.bEnableDrives)
-        {
-            bool bTwistLocked = TwistMotion == EJointMotionType::Locked;
-            bool bSwing1Locked = Swing1Motion == EJointMotionType::Locked;
-            bool bSwing2Locked = Swing2Motion == EJointMotionType::Locked;
-
-            // No SLerp drive if we have a locked rotation (it will be grayed out in the editor in this case, but could still have been set before the rotation was locked)
-            // @todo(ccaulfield): setting should be cleaned up before being passed to the solver
-            if ((JointSettings.bAngularSLerpPositionDriveEnabled || JointSettings.bAngularSLerpVelocityDriveEnabled) && !bTwistLocked && !bSwing1Locked && !bSwing2Locked)
-            {
-                ApplySLerpDrive(Dt, SolverSettings, JointSettings);
-            }
-            else
-            {
-                const bool bTwistDriveEnabled = ((JointSettings.bAngularTwistPositionDriveEnabled || JointSettings.bAngularTwistVelocityDriveEnabled) && !bTwistLocked);
-                const bool bSwingDriveEnabled = (JointSettings.bAngularSwingPositionDriveEnabled || JointSettings.bAngularSwingVelocityDriveEnabled);
-                const bool bSwing1DriveEnabled = bSwingDriveEnabled && !bSwing1Locked;
-                const bool bSwing2DriveEnabled = bSwingDriveEnabled && !bSwing2Locked;
-                if (bTwistDriveEnabled || bSwing1DriveEnabled || bSwing2DriveEnabled)
-                {
-                    ApplySwingTwistDrives(Dt, SolverSettings, JointSettings, bTwistDriveEnabled, bSwing1DriveEnabled, bSwing2DriveEnabled);
-                }
-            }
-        }
-    }
-
-
-    void FPBDJointCachedSolver::ApplyPositionConstraints(
-        const FReal Dt,
-        const FPBDJointSolverSettings& SolverSettings,
-        const FPBDJointSettings& JointSettings)
-    {
-        bool bHasPositionConstraints =
-            (JointSettings.LinearMotionTypes[0] != EJointMotionType::Free)
-            || (JointSettings.LinearMotionTypes[1] != EJointMotionType::Free)
-            || (JointSettings.LinearMotionTypes[2] != EJointMotionType::Free);
-        if (!bHasPositionConstraints)
-        {
-            return;
-        }
-
-        const TVec3<EJointMotionType>& LinearMotion = JointSettings.LinearMotionTypes;
-        const TVec3<bool> bLinearLocked =
-        {
-            (LinearMotion[0] == EJointMotionType::Locked),
-            (LinearMotion[1] == EJointMotionType::Locked),
-            (LinearMotion[2] == EJointMotionType::Locked),
-        };
-        const TVec3<bool> bLinearLimited =
-        {
-            (LinearMotion[0] == EJointMotionType::Limited),
-            (LinearMotion[1] == EJointMotionType::Limited),
-            (LinearMotion[2] == EJointMotionType::Limited),
-        };
-
-        if (bLinearLocked[0] && bLinearLocked[1] && bLinearLocked[2])
-        {
-            // Hard point constraint (most common case)
-            if (!Body0().IsDynamic())
-            {
-                ApplyPointPositionConstraintKD(0, 1, Dt, SolverSettings, JointSettings);
-            }
-            else if (!Body1().IsDynamic())
-            {
-                ApplyPointPositionConstraintKD(1, 0, Dt, SolverSettings, JointSettings);
-            }
-            else
-            {
-                ApplyPointPositionConstraintDD(Dt, SolverSettings, JointSettings);
-            }
-        }
-        else if (bLinearLimited[0] && bLinearLimited[1] && bLinearLimited[2])
-        {
-            // Spherical constraint
-            ApplySphericalPositionConstraint(Dt, SolverSettings, JointSettings);
-        }
-        else if (bLinearLocked[1] && bLinearLocked[2] && !bLinearLocked[0])
-        {
-            // Line constraint along X axis
-            ApplyCylindricalPositionConstraint(Dt, 0, LinearMotion[0], EJointMotionType::Locked, SolverSettings, JointSettings);
-        }
-        else if (bLinearLocked[0] && bLinearLocked[2] && !bLinearLocked[1])
-        {
-            // Line constraint along Y axis
-            ApplyCylindricalPositionConstraint(Dt, 1, LinearMotion[1], EJointMotionType::Locked, SolverSettings, JointSettings);
-        }
-        else if (bLinearLocked[0] && bLinearLocked[1] && !bLinearLocked[2])
-        {
-            // Line constraint along Z axis
-            ApplyCylindricalPositionConstraint(Dt, 2, LinearMotion[2], EJointMotionType::Locked, SolverSettings, JointSettings);
-        }
-        else if (bLinearLimited[1] && bLinearLimited[2] && !bLinearLimited[0])
-        {
-            // Cylindrical constraint along X axis
-            ApplyCylindricalPositionConstraint(Dt, 0, LinearMotion[0], EJointMotionType::Limited, SolverSettings, JointSettings);
-        }
-        else if (bLinearLimited[0] && bLinearLimited[2] && !bLinearLimited[1])
-        {
-            // Cylindrical constraint along Y axis
-            ApplyCylindricalPositionConstraint(Dt, 1, LinearMotion[1], EJointMotionType::Limited, SolverSettings, JointSettings);
-        }
-        else if (bLinearLimited[0] && bLinearLimited[1] && !bLinearLimited[2])
-        {
-            // Cylindrical constraint along Z axis
-            ApplyCylindricalPositionConstraint(Dt, 2, LinearMotion[2], EJointMotionType::Limited, SolverSettings, JointSettings);
-        }
-        else if (bLinearLocked[0] || bLinearLimited[0])
-        {
-            // Planar constraint along X axis
-            ApplyPlanarPositionConstraint(Dt, 0, LinearMotion[0], SolverSettings, JointSettings);
-        }
-        else if (bLinearLocked[1] || bLinearLimited[1])
-        {
-            // Planar constraint along Y axis
-            ApplyPlanarPositionConstraint(Dt, 1, LinearMotion[1], SolverSettings, JointSettings);
-        }
-        else if (bLinearLocked[2] || bLinearLimited[2])
-        {
-            // Planar constraint along Z axis
-            ApplyPlanarPositionConstraint(Dt, 2, LinearMotion[2], SolverSettings, JointSettings);
-        }
-    }
-
-
-    void FPBDJointCachedSolver::ApplyPositionDrives(
-        const FReal Dt,
-        const FPBDJointSolverSettings& SolverSettings,
-        const FPBDJointSettings& JointSettings)
-    {
-        if (SolverSettings.bEnableDrives)
-        {
-            TVec3<bool> bDriven =
-            {
-                (JointSettings.bLinearPositionDriveEnabled[0] || JointSettings.bLinearVelocityDriveEnabled[0]) && (JointSettings.LinearMotionTypes[0] != EJointMotionType::Locked),
-                (JointSettings.bLinearPositionDriveEnabled[1] || JointSettings.bLinearVelocityDriveEnabled[1]) && (JointSettings.LinearMotionTypes[1] != EJointMotionType::Locked),
-                (JointSettings.bLinearPositionDriveEnabled[2] || JointSettings.bLinearVelocityDriveEnabled[2]) && (JointSettings.LinearMotionTypes[2] != EJointMotionType::Locked),
-            };
-
-            // Rectangular position drives
-            if (bDriven[0] || bDriven[1] || bDriven[2])
-            {
-                const FMatrix33 R0M = ConnectorRs[0].ToMatrix();
-                const FVec3 XTarget = ConnectorXs[0] + ConnectorRs[0] * JointSettings.LinearDrivePositionTarget;
-                const FVec3 VTarget = ConnectorRs[0] * JointSettings.LinearDriveVelocityTarget;
-                const FVec3 CX = ConnectorXs[1] - XTarget;
-
-                for (int32 AxisIndex = 0; AxisIndex < 3; ++AxisIndex)
-                {
-                    if (bDriven[AxisIndex])
-                    {
-                        const FVec3 Axis = R0M.GetAxis(AxisIndex);
-                        const FReal DeltaPos = FVec3::DotProduct(CX, Axis);
-                        const FReal DeltaVel = FVec3::DotProduct(VTarget, Axis);
-
-                        ApplyPositionDrive(Dt, AxisIndex, SolverSettings, JointSettings, Axis, DeltaPos, DeltaVel);
-                    }
-                }
-            }
-        }
-    }
-
-
-    //
-    //
-    //////////////////////////////////////////////////////////////////////////
-    //
-    //
-
-
-    void FPBDJointCachedSolver::ApplyLinearVelocityConstraints(
-        const FReal Dt,
-        const FPBDJointSolverSettings& SolverSettings,
-        const FPBDJointSettings& JointSettings)
-    {
-        bool bHasPositionConstraints =
-            (JointSettings.LinearMotionTypes[0] != EJointMotionType::Free)
-            || (JointSettings.LinearMotionTypes[1] != EJointMotionType::Free)
-            || (JointSettings.LinearMotionTypes[2] != EJointMotionType::Free);
-        if (!bHasPositionConstraints)
-        {
-            return;
-        }
-
-        const TVector<EJointMotionType, 3>& LinearMotion = JointSettings.LinearMotionTypes;
-        const TVector<bool, 3> bLinearLocked =
-        {
-            (LinearMotion[0] == EJointMotionType::Locked),
-            (LinearMotion[1] == EJointMotionType::Locked),
-            (LinearMotion[2] == EJointMotionType::Locked),
-        };
-        const TVector<bool, 3> bLinearLimited =
-        {
-            (LinearMotion[0] == EJointMotionType::Limited),
-            (LinearMotion[1] == EJointMotionType::Limited),
-            (LinearMotion[2] == EJointMotionType::Limited),
-        };
-
-        if (bLinearLocked[0] && bLinearLocked[1] && bLinearLocked[2])
-        {
-            // Hard point constraint (most common case)
-            ApplyPointVelocityConstraint(Dt, SolverSettings, JointSettings);
-        }
-        else if (bLinearLimited[0] && bLinearLimited[1] && bLinearLimited[2])
-        {
-            ApplySphericalVelocityConstraint(Dt, SolverSettings, JointSettings);
-        }
-        else if (bLinearLocked[1] && bLinearLocked[2] && !bLinearLocked[0])
-        {
-            // Line constraint along X axis
-            ApplyCylindricalVelocityConstraint(Dt, 0, LinearMotion[0], EJointMotionType::Locked, SolverSettings, JointSettings);
-        }
-        else if (bLinearLocked[0] && bLinearLocked[2] && !bLinearLocked[1])
-        {
-            // Line constraint along Y axis
-            ApplyCylindricalVelocityConstraint(Dt, 1, LinearMotion[1], EJointMotionType::Locked, SolverSettings, JointSettings);
-        }
-        else if (bLinearLocked[0] && bLinearLocked[1] && !bLinearLocked[2])
-        {
-            // Line constraint along Z axis
-            ApplyCylindricalVelocityConstraint(Dt, 2, LinearMotion[2], EJointMotionType::Locked, SolverSettings, JointSettings);
-        }
-        else if (bLinearLimited[1] && bLinearLimited[2] && !bLinearLimited[0])
-        {
-            // Cylindrical constraint along X axis
-            ApplyCylindricalVelocityConstraint(Dt, 0, LinearMotion[0], EJointMotionType::Limited, SolverSettings, JointSettings);
-        }
-        else if (bLinearLimited[0] && bLinearLimited[2] && !bLinearLimited[1])
-        {
-            // Cylindrical constraint along Y axis
-            ApplyCylindricalVelocityConstraint(Dt, 1, LinearMotion[1], EJointMotionType::Limited, SolverSettings, JointSettings);
-        }
-        else if (bLinearLimited[0] && bLinearLimited[1] && !bLinearLimited[2])
-        {
-            // Cylindrical constraint along Z axis
-            ApplyCylindricalVelocityConstraint(Dt, 2, LinearMotion[2], EJointMotionType::Limited, SolverSettings, JointSettings);
-        }
-        else if (bLinearLocked[0] || bLinearLimited[0])
-        {
-            // Planar constraint along X axis
-            ApplyPlanarVelocityConstraint(Dt, 0, LinearMotion[0], SolverSettings, JointSettings);
-        }
-        else if (bLinearLocked[1] || bLinearLimited[1])
-        {
-            // Planar constraint along Y axis
-            ApplyPlanarVelocityConstraint(Dt, 1, LinearMotion[1], SolverSettings, JointSettings);
-        }
-        else if (bLinearLocked[2] || bLinearLimited[2])
-        {
-            // Planar constraint along Z axis
-            ApplyPlanarVelocityConstraint(Dt, 2, LinearMotion[2], SolverSettings, JointSettings);
-        }
-    }
-
-    void FPBDJointCachedSolver::ApplyAngularVelocityConstraints(
-        const FReal Dt,
-        const FPBDJointSolverSettings& SolverSettings,
-        const FPBDJointSettings& JointSettings)
-    {
-        bool bHasRotationConstraints =
-            (JointSettings.AngularMotionTypes[0] != EJointMotionType::Free)
-            || (JointSettings.AngularMotionTypes[1] != EJointMotionType::Free)
-            || (JointSettings.AngularMotionTypes[2] != EJointMotionType::Free);
-        if (!bHasRotationConstraints)
-        {
-            return;
-        }
-
-        // Locked axes always use hard constraints. Limited axes use hard or soft depending on settings
-        EJointMotionType TwistMotion = JointSettings.AngularMotionTypes[(int32)EJointAngularConstraintIndex::Twist];
-        EJointMotionType Swing1Motion = JointSettings.AngularMotionTypes[(int32)EJointAngularConstraintIndex::Swing1];
-        EJointMotionType Swing2Motion = JointSettings.AngularMotionTypes[(int32)EJointAngularConstraintIndex::Swing2];
-        bool bTwistSoft = FPBDJointUtilities::GetSoftTwistLimitEnabled(SolverSettings, JointSettings);
-        bool bSwingSoft = FPBDJointUtilities::GetSoftSwingLimitEnabled(SolverSettings, JointSettings);
-
-        // If the twist axes are opposing, we cannot decompose the orientation into swing and twist angles, so just give up
-        const FVec3 Twist0 = ConnectorRs[0] * FJointConstants::TwistAxis();
-        const FVec3 Twist1 = ConnectorRs[1] * FJointConstants::TwistAxis();
-        const FReal Twist01Dot = FVec3::DotProduct(Twist0, Twist1);
-        const bool bDegenerate = (Twist01Dot < Chaos_Joint_DegenerateRotationLimit);
-        if (bDegenerate)
-        {
-            UE_LOG(LogChaosJoint, VeryVerbose, TEXT(" Degenerate rotation at Swing %f deg"), FMath::RadiansToDegrees(FMath::Acos(Twist01Dot)));
-        }
-
-        // Apply twist constraint
-        // NOTE: Cannot calculate twist angle at 180degree swing
-        if (SolverSettings.bEnableTwistLimits && !bDegenerate)
-        {
-            if (TwistMotion == EJointMotionType::Limited)
-            {
-                ApplyTwistVelocityConstraint(Dt, SolverSettings, JointSettings, bTwistSoft);
-            }
-            else if (TwistMotion == EJointMotionType::Locked)
-            {
-                // Covered below
-            }
-            else if (TwistMotion == EJointMotionType::Free)
-            {
-            }
-        }
-
-        // Apply swing constraints
-        // NOTE: Cannot separate swing angles at 180degree swing (but we can still apply locks)
-        if (SolverSettings.bEnableSwingLimits)
-        {
-            if ((Swing1Motion == EJointMotionType::Limited) && (Swing2Motion == EJointMotionType::Limited))
-            {
-                ApplyConeVelocityConstraint(Dt, SolverSettings, JointSettings, bSwingSoft);
-            }
-            else if ((Swing1Motion == EJointMotionType::Limited) && (Swing2Motion == EJointMotionType::Locked))
-            {
-                ApplySingleLockedSwingVelocityConstraint(Dt, SolverSettings, JointSettings, EJointAngularConstraintIndex::Swing2, false);
-                if (!bDegenerate)
-                {
-                    ApplySwingVelocityConstraint(Dt, SolverSettings, JointSettings, EJointAngularConstraintIndex::Swing1, bSwingSoft);
-                }
-            }
-            else if ((Swing1Motion == EJointMotionType::Limited) && (Swing2Motion == EJointMotionType::Free))
-            {
-                if (!bDegenerate)
-                {
-                    ApplyDualConeSwingVelocityConstraint(Dt, SolverSettings, JointSettings, EJointAngularConstraintIndex::Swing1, bSwingSoft);
-                }
-            }
-            else if ((Swing1Motion == EJointMotionType::Locked) && (Swing2Motion == EJointMotionType::Limited))
-            {
-                ApplySingleLockedSwingVelocityConstraint(Dt, SolverSettings, JointSettings, EJointAngularConstraintIndex::Swing1, false);
-                if (!bDegenerate)
-                {
-                    ApplySwingVelocityConstraint(Dt, SolverSettings, JointSettings, EJointAngularConstraintIndex::Swing2, bSwingSoft);
-                }
-            }
-            else if ((Swing1Motion == EJointMotionType::Locked) && (Swing2Motion == EJointMotionType::Locked))
-            {
-                // Covered below
-            }
-            else if ((Swing1Motion == EJointMotionType::Locked) && (Swing2Motion == EJointMotionType::Free))
-            {
-                ApplySingleLockedSwingVelocityConstraint(Dt, SolverSettings, JointSettings, EJointAngularConstraintIndex::Swing1, false);
-            }
-            else if ((Swing1Motion == EJointMotionType::Free) && (Swing2Motion == EJointMotionType::Limited))
-            {
-                if (!bDegenerate)
-                {
-                    ApplyDualConeSwingVelocityConstraint(Dt, SolverSettings, JointSettings, EJointAngularConstraintIndex::Swing2, bSwingSoft);
-                }
-            }
-            else if ((Swing1Motion == EJointMotionType::Free) && (Swing2Motion == EJointMotionType::Locked))
-            {
-                ApplySingleLockedSwingVelocityConstraint(Dt, SolverSettings, JointSettings, EJointAngularConstraintIndex::Swing2, false);
-            }
-            else if ((Swing1Motion == EJointMotionType::Free) && (Swing2Motion == EJointMotionType::Free))
-            {
-            }
-        }
-
-        // Note: single-swing locks are already handled above so we only need to do something here if both are locked
-        bool bLockedTwist = SolverSettings.bEnableTwistLimits 
-            && (JointSettings.AngularMotionTypes[(int32)EJointAngularConstraintIndex::Twist] == EJointMotionType::Locked);
-        bool bLockedSwing = SolverSettings.bEnableSwingLimits 
-            && (JointSettings.AngularMotionTypes[(int32)EJointAngularConstraintIndex::Swing1] == EJointMotionType::Locked) 
-            && (JointSettings.AngularMotionTypes[(int32)EJointAngularConstraintIndex::Swing2] == EJointMotionType::Locked);
-        if (bLockedTwist || bLockedSwing)
-        {
-            ApplyLockedRotationVelocityConstraints(Dt, SolverSettings, JointSettings, bLockedTwist, bLockedSwing);
-        }
-    }
-
-    void FPBDJointCachedSolver::ApplyTwistVelocityConstraint(
-        const FReal Dt,
-        const FPBDJointSolverSettings& SolverSettings,
-        const FPBDJointSettings& JointSettings,
-        const bool bUseSoftLimit)
-    {
-        if(!NetAngularImpulse.IsNearlyZero() && !bUseSoftLimit && AngularHardLambda[(int32)EJointAngularConstraintIndex::Twist] > SMALL_NUMBER)
-        {
-            FVec3 TwistAxis;
-            FReal TwistAngle;
-            FPBDJointUtilities::GetTwistAxisAngle(ConnectorRs[0], ConnectorRs[1], TwistAxis, TwistAngle);
-
-            if (TwistAngle < 0.0f)
-            {
-                TwistAxis = -TwistAxis;
-            }
-
-            const FReal TwistStiffness = SolverStiffness * FPBDJointUtilities::GetTwistStiffness(SolverSettings, JointSettings);
-
-            FReal TargetVel = 0.0f;
-            if (JointSettings.TwistRestitution != 0.0f)
-            {
-                const FReal InitVel = InitConstraintAxisAngularVelocities[(int32)EJointAngularConstraintIndex::Twist];
-                TargetVel = InitVel > Chaos_Joint_AngularVelocityThresholdToApplyRestitution ? -JointSettings.TwistRestitution * InitVel : 0.0f;
-            }
-            ApplyAngularVelocityConstraint(TwistStiffness, TwistAxis, TargetVel);
-        }
-    }
-
-    void FPBDJointCachedSolver::ApplyConeVelocityConstraint(
-        const FReal Dt,
-        const FPBDJointSolverSettings& SolverSettings,
-        const FPBDJointSettings& JointSettings,
-        const bool bUseSoftLimit)
-    {
-        if(!NetAngularImpulse.IsNearlyZero() && !bUseSoftLimit && AngularHardLambda[(int32)EJointAngularConstraintIndex::Swing2] > SMALL_NUMBER)
-        {
-            FVec3 SwingAxisLocal;
-            FReal DSwingAngle = 0.0f;
-            const FReal LimitPadding = GetAngularConstraintPadding(EJointAngularConstraintIndex::Swing1);
-            const FReal Swing1Limit = FMath::Max(JointSettings.AngularLimits[(int32)EJointAngularConstraintIndex::Swing1] - LimitPadding, 0.0f);
-            const FReal Swing2Limit = FMath::Max(JointSettings.AngularLimits[(int32)EJointAngularConstraintIndex::Swing2] - LimitPadding, 0.0f);
-            FPBDJointUtilities::GetEllipticalConeAxisErrorLocal(ConnectorRs[0], ConnectorRs[1], Swing2Limit, Swing1Limit, SwingAxisLocal, DSwingAngle);
-            // SwingAxisLocal has the size of Sine(SwingAngle), thus we need to normalize it
-            SwingAxisLocal.SafeNormalize();
-            const FVec3 SwingAxis = ConnectorRs[0] * SwingAxisLocal;
-            const FReal SwingStiffness = SolverStiffness * FPBDJointUtilities::GetSwingStiffness(SolverSettings, JointSettings);
-            FReal TargetVel = 0.0f;
-            if (JointSettings.SwingRestitution != 0.0f)
-            {
-                const FReal InitVel = InitConstraintAxisAngularVelocities[(int32)EJointAngularConstraintIndex::Swing1];
-                TargetVel = InitVel > Chaos_Joint_AngularVelocityThresholdToApplyRestitution ? -JointSettings.SwingRestitution * InitVel : 0.0f;
-            }
-            ApplyAngularVelocityConstraint(SwingStiffness, SwingAxis, TargetVel);
-        }
-    }
-
-    void FPBDJointCachedSolver::ApplySingleLockedSwingVelocityConstraint(
-        const FReal Dt,
-        const FPBDJointSolverSettings& SolverSettings,
-        const FPBDJointSettings& JointSettings,
-        const EJointAngularConstraintIndex SwingConstraintIndex,
-        const bool bUseSoftLimit)
-    {
-        if(!NetAngularImpulse.IsNearlyZero() && !bUseSoftLimit && AngularHardLambda[(int32)SwingConstraintIndex] > SMALL_NUMBER)
-        {
-            FVec3 SwingAxis;
-            FReal SwingAngle;
-            FPBDJointUtilities::GetLockedSwingAxisAngle(ConnectorRs[0], ConnectorRs[1], SwingConstraintIndex, SwingAxis, SwingAngle);
-            // Swing axis has the size of Sine(SwingAngle), thus we need to normalize it
-            SwingAxis.SafeNormalize();
-
-            const FReal SwingStiffness = SolverStiffness * FPBDJointUtilities::GetSwingStiffness(SolverSettings, JointSettings);
-            ApplyAngularVelocityConstraint(SwingStiffness, SwingAxis);
-        }
-    }
-
-    void FPBDJointCachedSolver::ApplyDualConeSwingVelocityConstraint(
-        const FReal Dt,
-        const FPBDJointSolverSettings& SolverSettings,
-        const FPBDJointSettings& JointSettings,
-        const EJointAngularConstraintIndex SwingConstraintIndex,
-        const bool bUseSoftLimit)
-    {
-        if(!NetAngularImpulse.IsNearlyZero() && !bUseSoftLimit && AngularHardLambda[(int32)SwingConstraintIndex] > SMALL_NUMBER)
-        {
-            FVec3 SwingAxis;
-            FReal SwingAngle;
-            FPBDJointUtilities::GetDualConeSwingAxisAngle(ConnectorRs[0], ConnectorRs[1], SwingConstraintIndex, SwingAxis, SwingAngle);
-
-            const FReal SwingStiffness = SolverStiffness * FPBDJointUtilities::GetSwingStiffness(SolverSettings, JointSettings);
-            FReal TargetVel = 0.0f;
-            if (JointSettings.SwingRestitution != 0.0f)
-            {
-                const FReal InitVel = InitConstraintAxisAngularVelocities[(int32)SwingConstraintIndex];
-                TargetVel = InitVel > Chaos_Joint_AngularVelocityThresholdToApplyRestitution ? -JointSettings.SwingRestitution * InitVel : 0.0f;
-            }
-            ApplyAngularVelocityConstraint(SwingStiffness, SwingAxis, TargetVel);
-        }
-    }
-
-    void FPBDJointCachedSolver::ApplySwingVelocityConstraint(
-        const FReal Dt,
-        const FPBDJointSolverSettings& SolverSettings,
-        const FPBDJointSettings& JointSettings,
-        const EJointAngularConstraintIndex SwingConstraintIndex,
-        const bool bUseSoftLimit)
-    {
-        if(!NetAngularImpulse.IsNearlyZero() && !bUseSoftLimit && AngularHardLambda[(int32)SwingConstraintIndex] > SMALL_NUMBER)
-        {
-            FVec3 SwingAxis;
-            FReal SwingAngle;
-            FPBDJointUtilities::GetSwingAxisAngle(ConnectorRs[0], ConnectorRs[1], SolverSettings.SwingTwistAngleTolerance, SwingConstraintIndex, SwingAxis, SwingAngle);
-
-            const FReal SwingStiffness = SolverStiffness * FPBDJointUtilities::GetSwingStiffness(SolverSettings, JointSettings);
-            FReal TargetVel = 0.0f;
-            if (JointSettings.SwingRestitution != 0.0f)
-            {
-                const FReal InitVel = InitConstraintAxisAngularVelocities[(int32)SwingConstraintIndex];
-                TargetVel = InitVel > Chaos_Joint_AngularVelocityThresholdToApplyRestitution ? -JointSettings.SwingRestitution * InitVel : 0.0f;
-            }
-            ApplyAngularVelocityConstraint(SwingStiffness, SwingAxis, TargetVel);
-        }
-    }
-
-    void FPBDJointCachedSolver::ApplyAngularVelocityConstraint(
-        const FReal Stiffness,
-        const FVec3& Axis,
-        const FReal TargetVel)
-    {
-        const FVec3 CW = W(1) - W(0);
-        const FVec3 IA0 = Utilities::Multiply(InvI(0), Axis);
-        const FVec3 IA1 = Utilities::Multiply(InvI(1), Axis);
-        const FReal II0 = FVec3::DotProduct(Axis, IA0);
-        const FReal II1 = FVec3::DotProduct(Axis, IA1);
-
-        const FVec3 AngularImpulse = Stiffness * (FVec3::DotProduct(CW, Axis) - TargetVel) / (II0 + II1) * Axis;
-        const FVec3 DW0 = Utilities::Multiply(InvI(0), AngularImpulse);
-        const FVec3 DW1 = Utilities::Multiply(InvI(1), -AngularImpulse);
-
-        ApplyAngularVelocityDelta(DW0, DW1);
-    }
-
-    void FPBDJointCachedSolver::ApplyLockedRotationVelocityConstraints(
-        const FReal Dt,
-        const FPBDJointSolverSettings& SolverSettings,
-        const FPBDJointSettings& JointSettings,
-        const bool bApplyTwist,
-        const bool bApplySwing)
-    {
-        if(!NetAngularImpulse.IsNearlyZero())
-        {
-            FVec3 Axis0, Axis1, Axis2;
-            FPBDJointUtilities::GetLockedRotationAxes(ConnectorRs[0], ConnectorRs[1], Axis0, Axis1, Axis2);
-            const FRotation3 R01 = ConnectorRs[0].Inverse() * ConnectorRs[1];
-
-            if (bApplyTwist && AngularHardLambda[(int32)EJointAngularConstraintIndex::Twist] > SMALL_NUMBER)
-            {
-                FReal TwistStiffness = SolverStiffness * FPBDJointUtilities::GetTwistStiffness(SolverSettings, JointSettings);
-                Axis0.SafeNormalize();
-                ApplyAngularVelocityConstraint(TwistStiffness, Axis0);
-            }
-
-            if (bApplySwing && AngularHardLambda[(int32)EJointAngularConstraintIndex::Swing2] + AngularHardLambda[(int32)EJointAngularConstraintIndex::Swing1] > SMALL_NUMBER)
-            {
-                FReal SwingStiffness = SolverStiffness * FPBDJointUtilities::GetSwingStiffness(SolverSettings, JointSettings);
-                Axis1.SafeNormalize();
-                ApplyAngularVelocityConstraint(SwingStiffness, Axis1);
-                Axis2.SafeNormalize();
-                ApplyAngularVelocityConstraint(SwingStiffness, Axis2);
-            }
-        }
-    }
-
-
-    //
-    //
-    //////////////////////////////////////////////////////////////////////////
-    //
-    //
-
-
-    void FPBDJointCachedSolver::ApplyPositionDelta(
-        const int32 BodyIndex,
-        const FVec3& DP)
-    {
-        //UE_LOG(LogChaosJoint, VeryVerbose, TEXT("      Apply DP%d %f %f %f"), BodyIndex, DP.X, DP.Y, DP.Z);
-
-        Body(BodyIndex).ApplyPositionDelta(DP);
-
-        ConnectorXs[BodyIndex] += DP;
-    }
-
-
-    void FPBDJointCachedSolver::ApplyPositionDelta(
-        const FVec3& DP0,
-        const FVec3& DP1)
-    {
-        //UE_LOG(LogChaosJoint, VeryVerbose, TEXT("      Apply DP%d %f %f %f"), 0, DP0.X, DP0.Y, DP0.Z);
-        //UE_LOG(LogChaosJoint, VeryVerbose, TEXT("      Apply DP%d %f %f %f"), 1, DP1.X, DP1.Y, DP1.Z);
-
-        Body(0).ApplyPositionDelta(DP0);
-        Body(1).ApplyPositionDelta(DP1);
-
-        ConnectorXs[0] += DP0;
-        ConnectorXs[1] += DP1;
-    }
-
-
-    void FPBDJointCachedSolver::ApplyRotationDelta(
-        const int32 BodyIndex,
-        const FVec3& DR)
-    {
-        //UE_LOG(LogChaosJoint, VeryVerbose, TEXT("      Apply DR%d %f %f %f"), BodyIndex, DR.X, DR.Y, DR.Z);
-
-        Body(BodyIndex).ApplyRotationDelta(DR);
-
-        UpdateDerivedState(BodyIndex);
-    }
-
-
-    void FPBDJointCachedSolver::ApplyRotationDelta(
-        const FVec3& DR0,
-        const FVec3& DR1)
-    {
-        //UE_LOG(LogChaosJoint, VeryVerbose, TEXT("      Apply DR%d %f %f %f"), 0, DR0.X, DR0.Y, DR0.Z);
-        //UE_LOG(LogChaosJoint, VeryVerbose, TEXT("      Apply DR%d %f %f %f"), 1, DR1.X, DR1.Y, DR1.Z);
-
-        if (bRealTypeCompatibleWithISPC && bChaos_Joint_ISPC_Enabled)
-        {
-#if 0
-            ispc::ApplyRotationDelta2((ispc::FPBDJointCachedSolver*)this, (ispc::FVector&)DR0, (ispc::FVector&)DR1);
-#endif
-        }
-        else
-        {
-            if (Body0().IsDynamic())
-            {
-                Body(0).ApplyRotationDelta(DR0);
-            }
-            if (Body1().IsDynamic())
-            {
-                Body(1).ApplyRotationDelta(DR1);
-            }
-
-            UpdateDerivedState();
-        }
-    }
-
-    void FPBDJointCachedSolver::ApplyDelta(
-        const int32 BodyIndex,
-        const FVec3& DP,
-        const FVec3& DR)
-    {
-        //UE_LOG(LogChaosJoint, VeryVerbose, TEXT("      Apply DP%d %f %f %f"), BodyIndex, DP.X, DP.Y, DP.Z);
-        //UE_LOG(LogChaosJoint, VeryVerbose, TEXT("      Apply DR%d %f %f %f"), BodyIndex, DR.X, DR.Y, DR.Z);
-
-        Body(BodyIndex).ApplyTransformDelta(DP, DR);
-
-        UpdateDerivedState(BodyIndex);
-    }
-
-
-    void FPBDJointCachedSolver::ApplyVelocityDelta(
-        const int32 BodyIndex,
-        const FVec3& DV,
-        const FVec3& DW)
-    {
-        Body(BodyIndex).ApplyVelocityDelta(DV, DW);
-    }
-
-
-    void FPBDJointCachedSolver::ApplyVelocityDelta(
-        const FVec3& DV0,
-        const FVec3& DW0,
-        const FVec3& DV1,
-        const FVec3& DW1)
-    {
-        Body(0).ApplyVelocityDelta(DV0, DW0);
-        Body(1).ApplyVelocityDelta(DV1, DW1);
-    }
-
-    void FPBDJointCachedSolver::ApplyAngularVelocityDelta(
-        const FVec3& DW0,
-        const FVec3& DW1)
-    {
-        Body(0).ApplyAngularVelocityDelta(DW0);
-        Body(1).ApplyAngularVelocityDelta(DW1);
-    }
-
-    void FPBDJointCachedSolver::ApplyPositionConstraint(
-        const FReal Stiffness,
-        const FVec3& Axis,
-        const FReal Delta,
-        const FReal IM,
-        const FVec3& Arm0,
-        const FVec3& Arm1,
-        const FVec3& DRAxis0,
-        const FVec3& DRAxis1,
-        const int32 LinearHardLambdaIndex)
-    {
-        const FReal DLambda = Stiffness * Delta / IM;
-        const FVec3 DX = Axis * DLambda;
-        const FVec3 DP0 = InvM(0) * DX;
-        const FVec3 DP1 = -InvM(1) * DX;
-        const FVec3 DR0 = DRAxis0 * DLambda;
-        const FVec3 DR1 = DRAxis1 * DLambda;
-
-        ApplyPositionDelta(DP0, DP1);
-        ApplyRotationDelta(DR0, DR1);
-
-        NetLinearImpulse += DX;
-
-        if(LinearHardLambdaIndex >= 0)
-        {
-            LinearHardLambda[LinearHardLambdaIndex] += DLambda;
-        }
-
-        ++NumActiveConstraints;
-    }
-
-
-    void FPBDJointCachedSolver::ApplyPositionConstraintSoft(
-        const FReal Dt,
-        const FReal JointStiffness,
-        const FReal JointDamping,
-        const bool bAccelerationMode,
-        const FVec3& Axis,
-        const FReal Delta,
-        const FReal TargetVel,
-        FReal& Lambda)
-    {
-        if (bRealTypeCompatibleWithISPC && bChaos_Joint_ISPC_Enabled)
-        {
-#if 0
-            float ReturnedLambda = Lambda;
-            ispc::ApplyPositionConstraintSoft((ispc::FPBDJointCachedSolver*)this, Dt, JointStiffness, JointDamping, bAccelerationMode, (ispc::FVector&)Axis, Delta, TargetVel, ReturnedLambda);
-            Lambda = ReturnedLambda;
-#endif
-        }
-        else
-        {
-            // Joint-space inverse mass
-            const FVec3 AngularAxis0 = FVec3::CrossProduct(ConnectorXs[0] - P(0), Axis);
-            const FVec3 AngularAxis1 = FVec3::CrossProduct(ConnectorXs[1] - P(1), Axis);
-            const FVec3 IA0 = Utilities::Multiply(InvI(0), AngularAxis0);
-            const FVec3 IA1 = Utilities::Multiply(InvI(1), AngularAxis1);
-            const FReal II0 = FVec3::DotProduct(AngularAxis0, IA0);
-            const FReal II1 = FVec3::DotProduct(AngularAxis1, IA1);
-            const FReal II = (InvM(0) + II0 + InvM(1) + II1);
-    
-            FReal VelDt = 0;
-            if (JointDamping > KINDA_SMALL_NUMBER)
-            {
-                const FVec3 V0Dt = FVec3::CalculateVelocity(InitConnectorXs[0], ConnectorXs[0], 1.0f);
-                const FVec3 V1Dt = FVec3::CalculateVelocity(InitConnectorXs[1], ConnectorXs[1], 1.0f);
-                VelDt = TargetVel * Dt + FVec3::DotProduct(V0Dt - V1Dt, Axis);
-            }
-    
-            const FReal SpringMassScale = (bAccelerationMode) ? 1.0f / (InvM(0) + InvM(1)) : 1.0f;
-            const FReal S = SpringMassScale * JointStiffness * Dt * Dt;
-            const FReal D = SpringMassScale * JointDamping * Dt;
-            const FReal Multiplier = (FReal)1 / ((S + D) * II + (FReal)1);
-            const FReal DLambda = SolverStiffness * Multiplier * (S * Delta - D * VelDt - Lambda);
-    
-            const FVec3 DP0 = (InvM(0) * DLambda) * Axis;
-            const FVec3 DP1 = (-InvM(1) * DLambda) * Axis;
-            const FVec3 DR0 = DLambda * Utilities::Multiply(InvI(0), AngularAxis0);
-            const FVec3 DR1 = -DLambda * Utilities::Multiply(InvI(1), AngularAxis1);
-    
-            ApplyPositionDelta(DP0, DP1);
-            ApplyRotationDelta(DR0, DR1);
+void FPBDJointCachedSolver::ComputeBodyState(const int32 BodyIndex)
+{
+	CurrentPs[BodyIndex] = P(BodyIndex);
+	CurrentQs[BodyIndex] = Q(BodyIndex);
+	ConnectorXs[BodyIndex] = CurrentPs[BodyIndex] + CurrentQs[BodyIndex] * LocalConnectorXs[BodyIndex].GetTranslation();
+	ConnectorRs[BodyIndex] = CurrentQs[BodyIndex] * LocalConnectorXs[BodyIndex].GetRotation();
+}
+
+void FPBDJointCachedSolver::UpdateDerivedState()
+{
+	// Kinematic bodies will not be moved, so we don't update derived state during iterations
+	if (InvM(0) > SMALL_NUMBER)
+	{
+		ComputeBodyState(0);
+	}
+	if (InvM(1) > SMALL_NUMBER)
+	{
+		ComputeBodyState(1);
+	}
+	ConnectorRs[1].EnforceShortestArcWith(ConnectorRs[0]);
+}
+
+void FPBDJointCachedSolver::UpdateDerivedState(const int32 BodyIndex)
+{
+	ComputeBodyState(BodyIndex);
+	ConnectorRs[1].EnforceShortestArcWith(ConnectorRs[0]);
+}
+
+bool FPBDJointCachedSolver::UpdateIsActive()
+{
+	// NumActiveConstraints is initialized to -1, so there's no danger of getting invalid LastPs/Qs
+	// We also check SolverStiffness mainly for testing when solver stiffness is 0 (so we don't exit immediately)
+	if ((NumActiveConstraints >= 0) && (SolverStiffness > 0.0f))
+	{
+		bool bIsSolved =
+			FVec3::IsNearlyEqual(Body(0).DP(), LastDPs[0], PositionTolerance)
+			&& FVec3::IsNearlyEqual(Body(1).DP(), LastDPs[1], PositionTolerance)
+			&& FVec3::IsNearlyEqual(Body(0).DQ(), LastDQs[0], 0.5f * AngleTolerance)
+			&& FVec3::IsNearlyEqual(Body(1).DQ(), LastDQs[1], 0.5f * AngleTolerance);
+		bIsActive = !bIsSolved;
+	}
+
+	LastDPs[0] = Body(0).DP();
+	LastDPs[1] = Body(1).DP();
+	LastDQs[0] = Body(0).DQ();
+	LastDQs[1] = Body(1).DQ();
+
+	return bIsActive;
+}
+
+void FPBDJointCachedSolver::Update(
+	   const FReal Dt,
+	   const FPBDJointSolverSettings& SolverSettings,
+	   const FPBDJointSettings& JointSettings)
+{
+	//UpdateIsActive();
+}
+
+void FPBDJointCachedSolver::EnableProjection()
+{
+	Body0().SetInvMScale(0);
+}
+
+/** Main init function to cache datas that could be reused in the apply */
+
+void FPBDJointCachedSolver::Init(
+	const FReal Dt,
+	const FSolverBodyPtrPair& SolverBodyPair,
+	const FPBDJointSolverSettings& SolverSettings,
+	const FPBDJointSettings& JointSettings,
+	const FRigidTransform3& XL0,
+	const FRigidTransform3& XL1)
+{
+	SolverBodies[0] = *SolverBodyPair[0];
+	SolverBodies[1] = *SolverBodyPair[1];
+
+	LocalConnectorXs[0] = XL0;
+	LocalConnectorXs[1] = XL1;
+
+	// \todo(chaos): joint should support parent/child in either order
+	SolverBodies[0].SetInvMScale(JointSettings.ParentInvMassScale);
+	SolverBodies[1].SetInvMScale(FReal(1));
+
+	FVec3 ConditionedInvILs[2];
+	FPBDJointUtilities::ConditionInverseMassAndInertia(Body0().InvM(), Body1().InvM(), Body0().InvILocal(), Body1().InvILocal(),
+		SolverSettings.MinParentMassRatio, SolverSettings.MaxInertiaRatio, ConditionedInvMs[0], ConditionedInvMs[1], ConditionedInvILs[0], ConditionedInvILs[1]);
 
-            Lambda += DLambda;
-            NetLinearImpulse += DLambda * Axis;
-        }
+	ConditionedInvIs[0] = Utilities::ComputeWorldSpaceInertia(Q(0), ConditionedInvILs[0]);
+	ConditionedInvIs[1] = Utilities::ComputeWorldSpaceInertia(Q(1), ConditionedInvILs[1]);
 
-        ++NumActiveConstraints;
-    }
+	NetLinearImpulse = FVec3(0);
+	NetAngularImpulse = FVec3(0);
+
+	LinearConstraintPadding = FVec3(-1);
+	AngularConstraintPadding = FVec3(-1);
+
+	// Tolerances are positional errors below visible detection. But in PBD the errors
+	// we leave behind get converted to velocity, so we need to ensure that the resultant
+	// movement from that erroneous velocity is less than the desired position tolerance.
+	// Assume that the tolerances were defined for a 60Hz simulation, then it must be that
+	// the position error is less than the position change from constant external forces
+	// (e.g., gravity). So, we are saying that the tolerance was chosen because the position
+	// error is less that F.dt^2. We need to scale the tolerance to work at our current dt.
+	const FReal ToleranceScale = FMath::Min(1.f, 60.f * 60.f * Dt * Dt);
+	PositionTolerance = ToleranceScale * SolverSettings.PositionTolerance;
+	AngleTolerance = ToleranceScale * SolverSettings.AngleTolerance;
+
+	NumActiveConstraints = -1;
+	bIsActive = true;
+
+	SolverStiffness = 1.0f;
+
+	InitDerivedState();
+
+	// Cache all the informations for the position and rotation constraints
+	InitPositionConstraints(Dt, SolverSettings, JointSettings);
+	InitRotationConstraints(Dt, SolverSettings, JointSettings);
+
+	InitPositionDrives(Dt, SolverSettings, JointSettings);
+	InitRotationDrives(Dt, SolverSettings, JointSettings);
+
+	LastDPs[0] = FVec3(0.f);
+	LastDPs[1] = FVec3(0.f);
+	LastDQs[0] = FVec3(0.f);
+	LastDQs[1] = FVec3(0.f);
+}
+
+void FPBDJointCachedSolver::Deinit()
+{
+	SolverBodies[0].Reset();
+	SolverBodies[1].Reset();
+}
+
+/** Main Apply function to solve all the constraint*/
+
+void FPBDJointCachedSolver::ApplyConstraints(
+		const FReal Dt,
+		const FReal InSolverStiffness,
+		const FPBDJointSolverSettings& SolverSettings,
+		const FPBDJointSettings& JointSettings)
+{
+	NumActiveConstraints = 0;
+	SolverStiffness = InSolverStiffness;
+
+	ApplyPositionConstraints(Dt);
+	ApplyRotationConstraints(Dt);
+
+	ApplyPositionDrives(Dt);
+	ApplyRotationDrives(Dt);
+
+	//UpdateIsActive();
+}
+
+void FPBDJointCachedSolver::ApplyVelocityConstraints(
+	const FReal Dt,
+	const FReal InSolverStiffness,
+	const FPBDJointSolverSettings& SolverSettings,
+	const FPBDJointSettings& JointSettings)
+{
+	SolverStiffness = InSolverStiffness;
+
+	// This is used for the QuasiPbd solver. If the Pbd step applied impulses to
+	// correct position errors, it will have introduced a velocity equal to the 
+	// correction divided by the timestep. We ensure that the velocity constraints
+	// (including restitution) are also enforced. This also prevents any position
+	// errors from the previous frame getting converted into energy.
+
+	ApplyLinearVelocityConstraints();
+	ApplyAngularVelocityConstraints();
+
+	// @todo(chaos): We can also apply velocity drives here rather than in the Pbd pass
+}
+
+/** UTILS FOR POSITION CONSTRAINT **************************************************************************************/
+
+FORCEINLINE bool ExtractLinearMotion( const FPBDJointSettings& JointSettings,
+	TVec3<bool>& bLinearLocked, TVec3<bool>& bLinearLimited)
+{
+	bool bHasPositionConstraints =
+		(JointSettings.LinearMotionTypes[0] != EJointMotionType::Free)
+		|| (JointSettings.LinearMotionTypes[1] != EJointMotionType::Free)
+		|| (JointSettings.LinearMotionTypes[2] != EJointMotionType::Free);
+	if (!bHasPositionConstraints)
+	{
+		return false;
+	}
+
+	const TVec3<EJointMotionType>& LinearMotion = JointSettings.LinearMotionTypes;
+	bLinearLocked =
+	{
+		(LinearMotion[0] == EJointMotionType::Locked),
+		(LinearMotion[1] == EJointMotionType::Locked),
+		(LinearMotion[2] == EJointMotionType::Locked),
+	};
+	bLinearLimited =
+	{
+		(LinearMotion[0] == EJointMotionType::Limited),
+		(LinearMotion[1] == EJointMotionType::Limited),
+		(LinearMotion[2] == EJointMotionType::Limited),
+	};
+	return true;
+}
+
+/** INIT POSITION CONSTRAINT ******************************************************************************************/
+
+void FPBDJointCachedSolver::InitPositionConstraints(
+	const FReal Dt,
+	const FPBDJointSolverSettings& SolverSettings,
+	const FPBDJointSettings& JointSettings)
+{
+	PositionConstraints.bValidDatas[0] = false;
+	PositionConstraints.bValidDatas[1] = false;
+	PositionConstraints.bValidDatas[2] = false;
+
+	TVec3<bool> bLinearLocked, bLinearLimited;
+	if(!ExtractLinearMotion(JointSettings, bLinearLocked, bLinearLimited))
+		return;
+
+	PositionConstraints.bAccelerationMode = FPBDJointUtilities::GetLinearSoftAccelerationMode(SolverSettings, JointSettings);
+
+	for(int32 ConstraintIndex = 0; ConstraintIndex < 3; ++ ConstraintIndex)
+	{
+		PositionConstraints.InitDatas(ConstraintIndex,bLinearLimited[ConstraintIndex] &&
+			FPBDJointUtilities::GetSoftLinearLimitEnabled(SolverSettings, JointSettings),
+			FPBDJointUtilities::GetSoftLinearStiffness(SolverSettings, JointSettings),
+			FPBDJointUtilities::GetSoftLinearDamping(SolverSettings, JointSettings),
+			FPBDJointUtilities::GetLinearStiffness(SolverSettings, JointSettings));
+	}
+	PointConstraints.InitDatas(FPBDJointUtilities::GetLinearStiffness(SolverSettings, JointSettings));
+
+	const TVec3<EJointMotionType>& LinearMotion = JointSettings.LinearMotionTypes;
+
+	if (bLinearLocked[0] && bLinearLocked[1] && bLinearLocked[2])
+	{
+		// Hard point constraint (most common case)
+		InitPointPositionConstraint();
+	}
+	else if (bLinearLimited[0] && bLinearLimited[1] && bLinearLimited[2])
+	{
+		// Spherical constraint
+		InitSphericalPositionConstraint(JointSettings, Dt);
+	}
+	else if (bLinearLocked[1] && bLinearLocked[2] && !bLinearLocked[0])
+	{
+		// Line constraint along X axis
+		InitCylindricalPositionConstraint(JointSettings, Dt, 0, LinearMotion[0], EJointMotionType::Locked);
+	}
+	else if (bLinearLocked[0] && bLinearLocked[2] && !bLinearLocked[1])
+	{
+		// Line constraint along Y axis
+		InitCylindricalPositionConstraint(JointSettings, Dt, 1, LinearMotion[1], EJointMotionType::Locked);
+	}
+	else if (bLinearLocked[0] && bLinearLocked[1] && !bLinearLocked[2])
+	{
+		// Line constraint along Z axis
+		InitCylindricalPositionConstraint(JointSettings, Dt, 2, LinearMotion[2], EJointMotionType::Locked);
+	}
+	else if (bLinearLimited[1] && bLinearLimited[2] && !bLinearLimited[0])
+	{
+		// Cylindrical constraint along X axis
+		InitCylindricalPositionConstraint(JointSettings, Dt, 0, LinearMotion[0], EJointMotionType::Limited);
+	}
+	else if (bLinearLimited[0] && bLinearLimited[2] && !bLinearLimited[1])
+	{
+		// Cylindrical constraint along Y axis
+		InitCylindricalPositionConstraint(JointSettings, Dt, 1, LinearMotion[1], EJointMotionType::Limited);
+	}
+	else if (bLinearLimited[0] && bLinearLimited[1] && !bLinearLimited[2])
+	{
+		// Cylindrical constraint along Z axis
+		InitCylindricalPositionConstraint(JointSettings, Dt, 2, LinearMotion[2], EJointMotionType::Limited);
+	}
+	else if (bLinearLocked[0] || bLinearLimited[0])
+	{
+		// Planar constraint along X axis
+		InitPlanarPositionConstraint(JointSettings, Dt, 0, LinearMotion[0]);
+	}
+	else if (bLinearLocked[1] || bLinearLimited[1])
+	{
+		// Planar constraint along Y axis
+		InitPlanarPositionConstraint(JointSettings, Dt, 1, LinearMotion[1]);
+	}
+	else if (bLinearLocked[2] || bLinearLimited[2])
+	{
+		// Planar constraint along Z axis
+		InitPlanarPositionConstraint(JointSettings, Dt, 2, LinearMotion[2]);
+	}
+}
+
+void FPBDJointCachedSolver::InitPositionConstraintDatas(
+	const int32 ConstraintIndex,
+	const FVec3& ConstraintAxis,
+	const FReal& ConstraintDelta,
+	const FReal ConstraintRestitution,
+	const FReal Dt,
+	const FReal ConstraintLimit,
+	const EJointMotionType JointType)
+{
+	const FVec3 LocalAxis = (ConstraintDelta < 0.0f) ? -ConstraintAxis : ConstraintAxis;
+	const FReal LocalDelta = (ConstraintDelta < 0.0f) ? -ConstraintDelta : ConstraintDelta;
+
+	PositionConstraints.ConstraintLimits[ConstraintIndex] = (JointType == EJointMotionType::Locked) ? 0.0f :
+		FMath::Max(ConstraintLimit - GetLinearConstraintPadding(ConstraintIndex), (FReal)0.);
+	PositionConstraints.MotionType[ConstraintIndex] = JointType;
+
+	// This uses ConnectorXs[1] rather than ConnectorXs[0] for stability reasons. It is more stable this way when ConnectorXs[0] and ConnectorXs[1] are far apart 
+	const FVec3 ConstraintArm0 = ConnectorXs[1] - LocalDelta * LocalAxis - P(0);
+	const FVec3 ConstraintArm1 = ConnectorXs[1] - P(1);
+
+	const FVec3 AngularAxis0 = FVec3::CrossProduct(ConstraintArm0, LocalAxis);
+	const FVec3 AngularAxis1 = FVec3::CrossProduct(ConstraintArm1, LocalAxis);
+	const FVec3 IA0 = Utilities::Multiply(InvI(0), AngularAxis0);
+	const FVec3 IA1 = Utilities::Multiply(InvI(1), AngularAxis1);
+	const FReal II0 = FVec3::DotProduct(AngularAxis0, IA0);
+	const FReal II1 = FVec3::DotProduct(AngularAxis1, IA1);
+
+	PositionConstraints.UpdateDatas(ConstraintIndex, LocalAxis, LocalDelta,
+		IA0, IA1, InvM(0) + II0 + InvM(1) + II1, ConstraintRestitution, Dt, true, ConstraintArm0, ConstraintArm1);
+
+	const FVec3 CV0 = V(0) + FVec3::CrossProduct(W(0), ConstraintArm0);
+	const FVec3 CV1 = V(1) + FVec3::CrossProduct(W(1), ConstraintArm1);
+	const FVec3 CV = CV1 - CV0;
+
+	InitConstraintAxisLinearVelocities[ConstraintIndex] = FVec3::DotProduct(CV, LocalAxis);
+}
+
+void FPBDJointCachedSolver::InitPointPositionConstraint()
+{
+	const FVec3 ConstraintArm0 = ConnectorXs[0] - P(0);
+	const FVec3 ConstraintArm1 = ConnectorXs[1] - P(1);
+
+	const FMatrix33 IM = Body(0).IsDynamic() && Body(1).IsDynamic() ?
+		Utilities::ComputeJointFactorMatrix(ConstraintArm0, InvI(0), InvM(0)) +
+		Utilities::ComputeJointFactorMatrix(ConstraintArm1, InvI(1), InvM(1)) : Body(0).IsDynamic() ?
+		Utilities::ComputeJointFactorMatrix(ConstraintArm0, InvI(0), InvM(0)) :
+		Utilities::ComputeJointFactorMatrix(ConstraintArm1, InvI(1), InvM(1));
+
+	PointConstraints.UpdateDatas(ConnectorXs[1] - ConnectorXs[0], IM.Inverse(), ConstraintArm0, ConstraintArm1);
+}
+
+void FPBDJointCachedSolver::InitSphericalPositionConstraint(
+	const FPBDJointSettings& JointSettings,
+	const FReal Dt)
+{
+	FVec3 Axis;
+	FReal Delta;
+	FPBDJointUtilities::GetSphericalAxisDelta(ConnectorXs[0], ConnectorXs[1], Axis, Delta);
+
+	InitPositionConstraintDatas(SphericalPositionConstraintIndex, Axis,  Delta, JointSettings.LinearRestitution, Dt, JointSettings.LinearLimit, EJointMotionType::Limited);
+}
+
+void FPBDJointCachedSolver::InitCylindricalPositionConstraint(
+	const FPBDJointSettings& JointSettings,
+	const FReal Dt,
+	const int32 AxisIndex,
+	const EJointMotionType AxialMotion,
+	const EJointMotionType RadialMotion)
+{
+	FVec3 MainAxis, RadialAxis;
+	FReal AxialDelta, RadialDelta;
+	FPBDJointUtilities::GetCylindricalAxesDeltas(ConnectorRs[0], ConnectorXs[0], ConnectorXs[1],
+		AxisIndex, MainAxis, AxialDelta, RadialAxis, RadialDelta);
+
+	InitPositionConstraintDatas((int32)EJointCylindricalPositionConstraintType::Axial, MainAxis, AxialDelta,  JointSettings.LinearRestitution, Dt,
+		JointSettings.LinearLimit, AxialMotion);   
+	InitPositionConstraintDatas((int32)EJointCylindricalPositionConstraintType::Radial, RadialAxis, RadialDelta, JointSettings.LinearRestitution, Dt,
+		JointSettings.LinearLimit, RadialMotion);
+}
+
+void FPBDJointCachedSolver::InitPlanarPositionConstraint(
+	const FPBDJointSettings& JointSettings,
+	const FReal Dt,
+	const int32 AxisIndex,
+	const EJointMotionType AxialMotion)
+{
+	FVec3 Axis;
+	FReal Delta;
+	FPBDJointUtilities::GetPlanarAxisDelta(ConnectorRs[0], ConnectorXs[0], ConnectorXs[1], AxisIndex, Axis, Delta);
+
+	InitPositionConstraintDatas(PlanarPositionConstraintIndex, Axis, Delta, JointSettings.LinearRestitution, Dt,
+		JointSettings.LinearLimit, AxialMotion);
+}
+
+/** APPLY POSITION CONSTRAINT *****************************************************************************************/
+
+void FPBDJointCachedSolver::ApplyPositionConstraints(
+	const FReal Dt)
+{
+	if(PointConstraints.bValidDatas)
+	{
+		ApplyPointPositionConstraint();
+	}
+	else
+	{
+		for(int32 ConstraintIndex = 0; ConstraintIndex < 3; ++ConstraintIndex)
+		{
+			if(PositionConstraints.bValidDatas[ConstraintIndex])
+			{
+				ApplyAxisPositionConstraint(ConstraintIndex, Dt);
+			}
+		}
+	}
+}
+
+void FPBDJointCachedSolver::ApplyPointPositionConstraint()
+{
+	const FVec3 CX = PointConstraints.ConstraintCX + Body(1).DP()  - Body(0).DP() +
+		FVec3::CrossProduct(Body(1).DQ(), PointConstraints.ConstraintArms[1]) -
+			FVec3::CrossProduct(Body(0).DQ(), PointConstraints.ConstraintArms[0]) ;
+
+	if (CX.SizeSquared() > PositionTolerance * PositionTolerance)
+	{
+		// Calculate constraint correction
+		const FVec3 DX = SolverStiffness * PointConstraints.ConstraintStiffness * Utilities::Multiply(PointConstraints.ConstraintIM, CX);
+
+		// Apply constraint correction
+		if( Body(0).IsDynamic())
+		{
+			const FVec3 DP0 = InvM(0) * DX;
+			const FVec3 DR0 = Utilities::Multiply(InvI(0), FVec3::CrossProduct(PointConstraints.ConstraintArms[0], DX));
+    	
+			ApplyPositionDelta(0, DP0);
+			ApplyRotationDelta(0, DR0);
+		}
+		if( Body(1).IsDynamic())
+		{
+			const FVec3 DP1 = -InvM(1) * DX;
+			const FVec3 DR1 = Utilities::Multiply(InvI(1), FVec3::CrossProduct(PointConstraints.ConstraintArms[1], -DX));
+
+			ApplyPositionDelta(1, DP1);
+			ApplyRotationDelta(1, DR1);
+		}
+
+		NetLinearImpulse += DX;
+		PointConstraints.ConstraintLambda += DX;
+		++NumActiveConstraints;
+	}
+}
+
+void FPBDJointCachedSolver::SolvePositionConstraintDelta(
+	const int32 ConstraintIndex, 
+	const FReal DeltaLambda)
+{
+	const FVec3 DX = PositionConstraints.ConstraintAxis[ConstraintIndex] * DeltaLambda;
+
+	if(Body(0).IsDynamic())
+	{
+		const FVec3 DP0 = InvM(0) * DX;
+		const FVec3 DR0 = PositionConstraints.ConstraintDRAxis[ConstraintIndex][0] * DeltaLambda;
+		ApplyPositionDelta(0,DP0);
+		ApplyRotationDelta(0,DR0);
+	}
+	if(Body(1).IsDynamic())
+	{
+		const FVec3 DP1 = -InvM(1) * DX;
+		const FVec3 DR1 = PositionConstraints.ConstraintDRAxis[ConstraintIndex][1] * DeltaLambda;
+		ApplyPositionDelta(1,DP1);
+		ApplyRotationDelta(1,DR1);
+	}
+
+	NetLinearImpulse += DX;
+	++NumActiveConstraints;
+}
+
+void FPBDJointCachedSolver::SolvePositionConstraintHard(
+	const int32 ConstraintIndex,
+	const FReal DeltaConstraint)
+{
+	const FReal DeltaLambda = SolverStiffness * PositionConstraints.ConstraintHardStiffness[ConstraintIndex] * DeltaConstraint /
+		PositionConstraints.ConstraintHardIM[ConstraintIndex];
+
+	PositionConstraints.ConstraintLambda[ConstraintIndex] += DeltaLambda;
+	SolvePositionConstraintDelta(ConstraintIndex, DeltaLambda);
+}
+
+void FPBDJointCachedSolver::SolvePositionConstraintSoft(
+	const int32 ConstraintIndex,
+	const FReal DeltaConstraint,
+	const FReal Dt,
+	const FReal TargetVel)
+{
+	FReal VelDt = 0;
+	if (PositionConstraints.ConstraintDamping[ConstraintIndex] > KINDA_SMALL_NUMBER)
+	{
+		const FVec3 V0Dt = FVec3::CalculateVelocity(InitConnectorXs[0], ConnectorXs[0]+Body(0).DP() + FVec3::CrossProduct(Body(0).DQ(), PositionConstraints.ConstraintArms[ConstraintIndex][0]), 1.0f);
+		const FVec3 V1Dt = FVec3::CalculateVelocity(InitConnectorXs[1], ConnectorXs[1]+Body(1).DP() + FVec3::CrossProduct(Body(1).DQ(), PositionConstraints.ConstraintArms[ConstraintIndex][1]), 1.0f);
+		VelDt = TargetVel * Dt + FVec3::DotProduct(V0Dt - V1Dt, PositionConstraints.ConstraintAxis[ConstraintIndex] );
+	}
+
+	const FReal DeltaLambda = SolverStiffness * (PositionConstraints.ConstraintSoftStiffness[ConstraintIndex] * DeltaConstraint - PositionConstraints.ConstraintDamping[ConstraintIndex] * VelDt - PositionConstraints.ConstraintLambda[ConstraintIndex]) /
+		PositionConstraints.ConstraintSoftIM[ConstraintIndex];
+	PositionConstraints.ConstraintLambda[ConstraintIndex] += DeltaLambda;
+
+	SolvePositionConstraintDelta(ConstraintIndex, DeltaLambda);
+}
+
+void FPBDJointCachedSolver::ApplyAxisPositionConstraint(
+	const int32 ConstraintIndex, const FReal Dt)
+{
+	const FVec3 CX = Body(1).DP()  - Body(0).DP() +
+		FVec3::CrossProduct(Body(1).DQ(), PositionConstraints.ConstraintArms[ConstraintIndex][1]) -
+			FVec3::CrossProduct(Body(0).DQ(), PositionConstraints.ConstraintArms[ConstraintIndex][0]) ;
+	FReal DeltaPosition = PositionConstraints.ConstraintCX[ConstraintIndex] + FVec3::DotProduct(CX, PositionConstraints.ConstraintAxis[ConstraintIndex]);
+
+	bool NeedsSolve = false;
+	if(PositionConstraints.bLimitsCheck[ConstraintIndex])
+	{
+		if(DeltaPosition > PositionConstraints.ConstraintLimits[ConstraintIndex] )
+		{
+			DeltaPosition -= PositionConstraints.ConstraintLimits[ConstraintIndex];
+			NeedsSolve = true;
+		}
+		else if(DeltaPosition < -PositionConstraints.ConstraintLimits[ConstraintIndex])
+		{
+			DeltaPosition += PositionConstraints.ConstraintLimits[ConstraintIndex];
+			NeedsSolve = true;
+		}
+	}
+	if (!PositionConstraints.bLimitsCheck[ConstraintIndex] || (PositionConstraints.bLimitsCheck[ConstraintIndex] && NeedsSolve && FMath::Abs(DeltaPosition ) > PositionTolerance))
+	{
+		if ((PositionConstraints.MotionType[ConstraintIndex] == EJointMotionType::Limited) && PositionConstraints.bSoftLimit[ConstraintIndex])
+		{
+			SolvePositionConstraintSoft(ConstraintIndex, DeltaPosition, Dt, 0.0f);
+		}
+		else if (PositionConstraints.MotionType[ConstraintIndex] != EJointMotionType::Free)
+		{
+			if (PositionConstraints.ConstraintRestitution[ConstraintIndex] > 0.0f)
+			{
+				CalculateLinearConstraintPadding(ConstraintIndex, Dt, PositionConstraints.ConstraintRestitution[ConstraintIndex], DeltaPosition);
+			}
+			SolvePositionConstraintHard(ConstraintIndex, DeltaPosition);
+		}
+	}
+}
+
+/** APPLY LINEAR VELOCITY *********************************************************************************************/
+
+void FPBDJointCachedSolver::ApplyLinearVelocityConstraints()
+{
+	if(PointConstraints.bValidDatas)
+	{
+		ApplyPointVelocityConstraint();
+	}
+	else
+	{
+		for(int32 ConstraintIndex = 0; ConstraintIndex < 3; ++ConstraintIndex)
+		{
+			if(PositionConstraints.bValidDatas[ConstraintIndex])
+			{
+				ApplyAxisVelocityConstraint(ConstraintIndex);
+			}
+		}
+	}
+}
+
+void FPBDJointCachedSolver::ApplyPointVelocityConstraint()
+{
+	if (!NetLinearImpulse.IsNearlyZero())
+	{
+		const FVec3 CV0 = V(0) + FVec3::CrossProduct(W(0), PointConstraints.ConstraintArms[0]);
+		const FVec3 CV1 = V(1) + FVec3::CrossProduct(W(1), PointConstraints.ConstraintArms[1]);
     
+		const FVec3 CV = CV1 - CV0;
+		const FVec3 MDV = SolverStiffness * PointConstraints.ConstraintStiffness * Utilities::Multiply(PointConstraints.ConstraintIM, CV);
+
+		if(Body0().IsDynamic())
+		{
+			const FVec3 DV0 = InvM(0) * MDV;
+			const FVec3 DW0 = Utilities::Multiply(InvI(0), FVec3::CrossProduct(PointConstraints.ConstraintArms[0], MDV));
+			Body(0).ApplyVelocityDelta(DV0, DW0);
+		}
+
+		if(Body1().IsDynamic())
+		{
+			const FVec3 DV1 = -InvM(1) * MDV;
+			const FVec3 DW1 = Utilities::Multiply(InvI(1), FVec3::CrossProduct(PointConstraints.ConstraintArms[1], -MDV));
+			Body(1).ApplyVelocityDelta(DV1, DW1);
+		}
+	}
+}
+
+void FPBDJointCachedSolver::SolveLinearVelocityConstraint(
+	const int32 ConstraintIndex,
+	const FReal TargetVel)
+{
+	const FVec3 CV0 = V(0) + FVec3::CrossProduct(W(0), PositionConstraints.ConstraintArms[ConstraintIndex][0]);
+	const FVec3 CV1 = V(1) + FVec3::CrossProduct(W(1), PositionConstraints.ConstraintArms[ConstraintIndex][1]);
+	const FVec3 CV = CV1 - CV0;
+
+	const FReal DeltaLambda = SolverStiffness * PositionConstraints.ConstraintHardStiffness[ConstraintIndex] *
+	 (FVec3::DotProduct(CV, PositionConstraints.ConstraintAxis[ConstraintIndex]) - TargetVel) / PositionConstraints.ConstraintHardIM[ConstraintIndex];
+	
+	const FVec3 MDV = DeltaLambda * PositionConstraints.ConstraintAxis[ConstraintIndex];
+
+	if(Body(0).IsDynamic())
+	{
+		const FVec3 DV0 = InvM(0) * MDV;
+		const FVec3 DW0 = PositionConstraints.ConstraintDRAxis[ConstraintIndex][0] * DeltaLambda;
+
+		Body(0).ApplyVelocityDelta(DV0, DW0);
+	}
+	if(Body(1).IsDynamic())
+	{
+		const FVec3 DV1 = -InvM(1) * MDV;
+		const FVec3 DW1 = PositionConstraints.ConstraintDRAxis[ConstraintIndex][1] * DeltaLambda;
+
+		Body(1).ApplyVelocityDelta(DV1, DW1);
+	}
+}
+
+void FPBDJointCachedSolver::ApplyAxisVelocityConstraint(const int32 ConstraintIndex)
+{
+	if(!NetLinearImpulse.IsNearlyZero() && PositionConstraints.ConstraintLambda[ConstraintIndex] > SMALL_NUMBER)
+	{
+		FReal TargetVel = 0.0f;
+		if (PositionConstraints.MotionType[ConstraintIndex] == EJointMotionType::Limited && PositionConstraints.ConstraintRestitution[ConstraintIndex] != 0.0f)
+		{
+			const FReal InitVel = InitConstraintAxisLinearVelocities[ConstraintIndex];
+			TargetVel = InitVel > Chaos_Joint_LinearVelocityThresholdToApplyRestitution ?
+				-PositionConstraints.ConstraintRestitution[ConstraintIndex] * InitVel : 0.0f; 
+		}
+		SolveLinearVelocityConstraint(ConstraintIndex, TargetVel);
+	}
+}
+
+/** UTILS FOR ROTATION CONSTRAINT **************************************************************************************/
+
+FORCEINLINE bool ExtractAngularMotion( const FPBDJointSettings& JointSettings,
+		TVec3<bool>& bAngularLocked, TVec3<bool>& bAngularLimited, TVec3<bool>& bAngularFree)
+{
+	bool bHasRotationConstraints =
+			  (JointSettings.AngularMotionTypes[0] != EJointMotionType::Free)
+		   || (JointSettings.AngularMotionTypes[1] != EJointMotionType::Free)
+		   || (JointSettings.AngularMotionTypes[2] != EJointMotionType::Free);
+	if (!bHasRotationConstraints)
+	{
+		return false;
+	}
+
+	const TVec3<EJointMotionType>& AngularMotion = JointSettings.AngularMotionTypes;
+	bAngularLocked =
+	{
+		(AngularMotion[0] == EJointMotionType::Locked),
+		(AngularMotion[1] == EJointMotionType::Locked),
+		(AngularMotion[2] == EJointMotionType::Locked),
+	};
+	bAngularLimited =
+	{
+		(AngularMotion[0] == EJointMotionType::Limited),
+		(AngularMotion[1] == EJointMotionType::Limited),
+		(AngularMotion[2] == EJointMotionType::Limited),
+	};
+	bAngularFree=
+	{
+		(AngularMotion[0] == EJointMotionType::Free),
+		(AngularMotion[1] == EJointMotionType::Free),
+		(AngularMotion[2] == EJointMotionType::Free),
+	};
+	return true;
+}
 
-    void FPBDJointCachedSolver::ApplyRotationConstraintKD(
-        const int32 KIndex,
-        const int32 DIndex,
-        const FReal JointStiffness,
-        const FVec3& Axis,
-        const FReal Angle,
-        const int32 AngularHardLambdaIndex)
-    {
-        const FReal Stiffness = SolverStiffness * JointStiffness;
-
-        const FVec3 IA1 = Utilities::Multiply(InvI(DIndex), Axis);
-        const FReal II1 = FVec3::DotProduct(Axis, IA1);
-        const FReal DR = Stiffness * (Angle / II1);
-        const FVec3 DR1 = IA1 * -DR;
-        ApplyRotationDelta(DIndex, DR1);
-
-        NetAngularImpulse += (KIndex == 0 )? DR * Axis : -DR * Axis;
-        if (AngularHardLambdaIndex >= 0)
-        {
-            AngularHardLambda[AngularHardLambdaIndex] += DR;
-        }
-    }
-
-
-    void FPBDJointCachedSolver::ApplyRotationConstraintDD(
-        const FReal JointStiffness,
-        const FVec3& Axis,
-        const FReal Angle,
-        const int32 AngularHardLambdaIndex)
-    {
-        const FReal Stiffness = SolverStiffness * JointStiffness;
-
-        // Joint-space inverse mass
-        const FVec3 IA0 = Utilities::Multiply(InvI(0), Axis);
-        const FVec3 IA1 = Utilities::Multiply(InvI(1), Axis);
-        const FReal II0 = FVec3::DotProduct(Axis, IA0);
-        const FReal II1 = FVec3::DotProduct(Axis, IA1);
-
-        const FReal DR = Stiffness * Angle / (II0 + II1);
-        const FVec3 DR0 = IA0 * DR;
-        const FVec3 DR1 = IA1 * -DR;
-
-        ApplyRotationDelta(DR0, DR1);
-
-        NetAngularImpulse += Axis * DR;
-        if (AngularHardLambdaIndex >= 0)
-        {
-            AngularHardLambda[AngularHardLambdaIndex] += DR;
-        }
-    }
-
-
-    void FPBDJointCachedSolver::ApplyRotationConstraint(
-        const FReal JointStiffness,
-        const FVec3& Axis,
-        const FReal Angle,
-        const int32 AngularHardLambdaIndex)
-    {
-        if (!Body0().IsDynamic())
-        {
-            ApplyRotationConstraintKD(0, 1, JointStiffness, Axis, Angle, AngularHardLambdaIndex);
-        }
-        else if (!Body1().IsDynamic())
-        {
-            ApplyRotationConstraintKD(1, 0, JointStiffness, Axis, -Angle, AngularHardLambdaIndex);
-        }
-        else
-        {
-            ApplyRotationConstraintDD(JointStiffness, Axis, Angle, AngularHardLambdaIndex);
-        }
-
-        ++NumActiveConstraints;
-    }
-
-
-    // See "XPBD: Position-Based Simulation of Compliant Constrained Dynamics"
-    void FPBDJointCachedSolver::ApplyRotationConstraintSoftKD(
-        const int32 KIndex,
-        const int32 DIndex,
-        const FReal Dt,
-        const FReal JointStiffness,
-        const FReal JointDamping,
-        const bool bAccelerationMode,
-        const FVec3& Axis,
-        const FReal Angle,
-        const FReal AngVelTarget,
-        FReal& Lambda)
-    {
-        check(!Body(KIndex).IsDynamic());
-        check(Body(DIndex).IsDynamic());
-
-        if (bRealTypeCompatibleWithISPC && bChaos_Joint_ISPC_Enabled)
-        {
-#if 0
-            float ReturnedLambda = Lambda;
-            ispc::ApplyRotationConstraintSoftKD((ispc::FPBDJointCachedSolver*)this, KIndex, DIndex, Dt, JointStiffness, JointDamping, bAccelerationMode, (ispc::FVector&) Axis, Angle, AngVelTarget, ReturnedLambda);
-            Lambda = ReturnedLambda;
-#endif
-        }
-        else
-        {
-            // World-space inverse mass
-            const FVec3 IA1 = Utilities::Multiply(InvI(DIndex), Axis);
-
-            // Joint-space inverse mass
-            FReal II1 = FVec3::DotProduct(Axis, IA1);
-            const FReal II = II1;
-
-            // Damping angular velocity
-            FReal AngVelDt = 0;
-            if (JointDamping > KINDA_SMALL_NUMBER)
-            {
-                const FVec3 W0Dt = FRotation3::CalculateAngularVelocity(InitConnectorRs[KIndex], ConnectorRs[KIndex], 1.0f);
-                const FVec3 W1Dt = FRotation3::CalculateAngularVelocity(InitConnectorRs[DIndex], ConnectorRs[DIndex], 1.0f);
-                AngVelDt = AngVelTarget * Dt + FVec3::DotProduct(Axis, W0Dt - W1Dt);
-            }
-
-            const FReal SpringMassScale = (bAccelerationMode) ? 1.0f / II : 1.0f;
-            const FReal S = SpringMassScale * JointStiffness * Dt * Dt;
-            const FReal D = SpringMassScale * JointDamping * Dt;
-            const FReal Multiplier = (FReal)1 / ((S + D) * II + (FReal)1);
-            const FReal DLambda = SolverStiffness * Multiplier * (S * Angle - D * AngVelDt - Lambda);
-
-            //const FVec3 DR1 = IA1 * -DLambda;
-            const FVec3 DR1 = Axis * -(DLambda * II1);
-
-            ApplyRotationDelta(DIndex, DR1);
+/** INIT ROTATION CONSTRAINT ******************************************************************************************/
+
+void FPBDJointCachedSolver::InitRotationConstraints(
+	const FReal Dt,
+	const FPBDJointSolverSettings& SolverSettings,
+	const FPBDJointSettings& JointSettings)
+{
+	RotationConstraints.bValidDatas[0] = false;
+	RotationConstraints.bValidDatas[1] = false;
+	RotationConstraints.bValidDatas[2] = false;
+
+	TVec3<bool> bAngularLocked, bAngularLimited, bAngularFree;
+	if(!ExtractAngularMotion(JointSettings, bAngularLocked, bAngularLimited, bAngularFree))
+		return;
+
+	RotationConstraints.bAccelerationMode = FPBDJointUtilities::GetAngularSoftAccelerationMode(SolverSettings, JointSettings);
+
+	const int32 TW = (int32)EJointAngularConstraintIndex::Twist;
+	const int32 S1 = (int32)EJointAngularConstraintIndex::Swing1;
+	const int32 S2 = (int32)EJointAngularConstraintIndex::Swing2;
+
+	RotationConstraints.InitDatas(TW,FPBDJointUtilities::GetSoftTwistLimitEnabled(SolverSettings, JointSettings) && !bAngularLocked[TW],
+FPBDJointUtilities::GetSoftTwistStiffness(SolverSettings, JointSettings),
+							   FPBDJointUtilities::GetSoftTwistDamping(SolverSettings, JointSettings),
+					FPBDJointUtilities::GetTwistStiffness(SolverSettings, JointSettings));
+
+	RotationConstraints.InitDatas(S1,FPBDJointUtilities::GetSoftSwingLimitEnabled(SolverSettings, JointSettings) && !bAngularLocked[S1],
+	FPBDJointUtilities::GetSoftSwingStiffness(SolverSettings, JointSettings),
+					FPBDJointUtilities::GetSoftSwingDamping(SolverSettings, JointSettings),
+					FPBDJointUtilities::GetSwingStiffness(SolverSettings, JointSettings));
+
+	RotationConstraints.InitDatas(S2, FPBDJointUtilities::GetSoftSwingLimitEnabled(SolverSettings, JointSettings) && !bAngularLocked[S2],
+	FPBDJointUtilities::GetSoftSwingStiffness(SolverSettings, JointSettings),
+					FPBDJointUtilities::GetSoftSwingDamping(SolverSettings, JointSettings),
+					FPBDJointUtilities::GetSwingStiffness(SolverSettings, JointSettings));
+
+	const FVec3 Twist0 = ConnectorRs[0] * FJointConstants::TwistAxis();
+	const FVec3 Twist1 = ConnectorRs[1] * FJointConstants::TwistAxis();
+	const bool bDegenerate = (FVec3::DotProduct(Twist0, Twist1) < Chaos_Joint_DegenerateRotationLimit);
+
+	// Apply twist constraint
+	// NOTE: Cannot calculate twist angle at 180degree swing
+	if (SolverSettings.bEnableTwistLimits)
+	{
+		if (bAngularLimited[TW] && !bDegenerate)
+		{
+			InitTwistConstraint(JointSettings, Dt);
+		}
+	}
+
+	// Apply swing constraints
+	// NOTE: Cannot separate swing angles at 180degree swing (but we can still apply locks)
+	if (SolverSettings.bEnableSwingLimits)
+	{
+		if (bAngularLimited[S1] && bAngularLimited[S2])
+		{
+			InitConeConstraint(JointSettings, Dt);
+		}
+		else if (bAngularLimited[S1] && bAngularLocked[S2])
+		{
+			InitSingleLockedSwingConstraint(JointSettings, Dt, EJointAngularConstraintIndex::Swing2);
+			if (!bDegenerate)
+			{
+				InitSwingConstraint(JointSettings, SolverSettings, Dt, EJointAngularConstraintIndex::Swing1 );
+			}
+		}
+		else if (bAngularLimited[S1] && bAngularFree[S2])
+		{
+			if (!bDegenerate)
+			{
+				InitDualConeSwingConstraint(JointSettings, Dt, EJointAngularConstraintIndex::Swing1);
+			}
+		}
+		else if (bAngularLocked[S1] && bAngularLimited[S2])
+		{
+			InitSingleLockedSwingConstraint(JointSettings, Dt, EJointAngularConstraintIndex::Swing1);
+			if (!bDegenerate)
+			{
+				InitSwingConstraint(JointSettings, SolverSettings, Dt, EJointAngularConstraintIndex::Swing2);
+			}
+		}
+		else if (bAngularLocked[S1] && bAngularFree[S2])
+		{
+			InitSingleLockedSwingConstraint(JointSettings, Dt, EJointAngularConstraintIndex::Swing1);
+		}
+		else if (bAngularFree[S1] && bAngularLimited[S2])
+		{
+			if (!bDegenerate)
+			{
+				InitDualConeSwingConstraint(JointSettings, Dt, EJointAngularConstraintIndex::Swing2);
+			}
+		}
+		else if (bAngularFree[S1]  && bAngularLocked[S2])
+		{
+			InitSingleLockedSwingConstraint(JointSettings, Dt, EJointAngularConstraintIndex::Swing2);
+		}
+	}
+
+	// Note: single-swing locks are already handled above so we only need to do something here if both are locked
+	const bool bLockedTwist = SolverSettings.bEnableTwistLimits && bAngularLocked[TW];
+	const bool bLockedSwing = SolverSettings.bEnableSwingLimits && bAngularLocked[S1] && bAngularLocked[S2];
+	if (bLockedTwist || bLockedSwing)
+	{
+		InitLockedRotationConstraints(JointSettings, Dt, bLockedTwist, bLockedSwing);
+	}
+}
+
+void FPBDJointCachedSolver::InitRotationConstraintDatas(
+		const FPBDJointSettings& JointSettings,
+		const int32 ConstraintIndex,
+		const FVec3& ConstraintAxis,
+		const FReal ConstraintAngle,
+		const FReal ConstraintRestitution,
+		const FReal Dt,
+		const bool bCheckLimit)
+{
+	const FVec3 LocalAxis = (ConstraintAngle < 0.0f) ? -ConstraintAxis : ConstraintAxis;
+	const FReal LocalAngle = (ConstraintAngle < 0.0f) ? -ConstraintAngle : ConstraintAngle;
+	
+	const FVec3 IA0 = Utilities::Multiply(InvI(0), LocalAxis);
+	const FVec3 IA1 = Utilities::Multiply(InvI(1), LocalAxis);
+	const FReal II0 = FVec3::DotProduct(LocalAxis, IA0);
+	const FReal II1 = FVec3::DotProduct(LocalAxis, IA1);
+
+	RotationConstraints.UpdateDatas(ConstraintIndex, LocalAxis, LocalAngle, IA0, IA1, II0 + II1, ConstraintRestitution, Dt, bCheckLimit);
+
+	RotationConstraints.ConstraintLimits[ConstraintIndex] = FMath::Max(
+		JointSettings.AngularLimits[ConstraintIndex] - GetAngularConstraintPadding(ConstraintIndex), (FReal)0.);
+
+	InitConstraintAxisAngularVelocities[ConstraintIndex] = FVec3::DotProduct(W(1) - W(0), LocalAxis);
+}
+
+void FPBDJointCachedSolver::CorrectAxisAngleConstraint(
+		const FPBDJointSettings& JointSettings,
+		const int32 ConstraintIndex,
+		FVec3& ConstraintAxis,
+		FReal& ConstraintAngle) const
+{
+	const FReal AngleMax = FMath::Max(JointSettings.AngularLimits[ConstraintIndex] -
+		GetAngularConstraintPadding(ConstraintIndex), (FReal)0.);
+
+	if (ConstraintAngle > AngleMax)
+	{
+		ConstraintAngle = ConstraintAngle - AngleMax;
+	}
+	else if (ConstraintAngle < -AngleMax)
+	{
+		// Keep Twist error positive
+		ConstraintAngle = -ConstraintAngle - AngleMax;
+		ConstraintAxis = -ConstraintAxis;
+	}
+	else
+	{
+		ConstraintAngle = 0;
+	}
+}
+
+void FPBDJointCachedSolver::InitTwistConstraint(
+		const FPBDJointSettings& JointSettings,
+		const FReal Dt)
+{
+	FVec3 TwistAxis;
+	FReal TwistAngle;
+	FPBDJointUtilities::GetTwistAxisAngle(ConnectorRs[0], ConnectorRs[1], TwistAxis, TwistAngle);
+
+	InitRotationConstraintDatas( JointSettings, (int32)EJointAngularConstraintIndex::Twist, TwistAxis, TwistAngle, JointSettings.TwistRestitution, Dt, true);
+}
+
+void FPBDJointCachedSolver::InitConeConstraint(
+   const FPBDJointSettings& JointSettings,
+   const FReal Dt)
+{
+	FVec3 SwingAxisLocal;
+	FReal SwingAngle = 0.0f;
+	FPBDJointUtilities::GetEllipticalConeAxisErrorLocal(ConnectorRs[0], ConnectorRs[1], 0.0, 0.0, SwingAxisLocal, SwingAngle);
+	SwingAxisLocal.SafeNormalize();
+
+	const FVec3 SwingAxis = ConnectorRs[0] * SwingAxisLocal;
+	InitRotationConstraintDatas( JointSettings, (int32)EJointAngularConstraintIndex::Swing2, SwingAxis, SwingAngle, JointSettings.SwingRestitution, Dt, true);
+}
+
+void FPBDJointCachedSolver::InitSingleLockedSwingConstraint(
+	const FPBDJointSettings& JointSettings,
+	const FReal Dt,
+	const EJointAngularConstraintIndex SwingConstraintIndex)
+{
+	// NOTE: SwingAxis is not normalized in this mode. It has length Sin(SwingAngle).
+	// Likewise, the SwingAngle is actually Sin(SwingAngle)
+	FVec3 SwingAxis;
+	FReal SwingAngle;
+	FPBDJointUtilities::GetLockedSwingAxisAngle(ConnectorRs[0], ConnectorRs[1], SwingConstraintIndex, SwingAxis, SwingAngle);
+	SwingAxis.SafeNormalize();
+
+	InitRotationConstraintDatas(JointSettings, (int32)SwingConstraintIndex, SwingAxis, SwingAngle, 0.0, Dt, false);
+}
+
+void FPBDJointCachedSolver::InitDualConeSwingConstraint(
+	const FPBDJointSettings& JointSettings,
+	const FReal Dt,
+	const EJointAngularConstraintIndex SwingConstraintIndex)
+{
+	FVec3 SwingAxis;
+	FReal SwingAngle;
+	FPBDJointUtilities::GetDualConeSwingAxisAngle(ConnectorRs[0], ConnectorRs[1], SwingConstraintIndex, SwingAxis, SwingAngle);
+
+	InitRotationConstraintDatas(JointSettings, (int32)SwingConstraintIndex, SwingAxis, SwingAngle, JointSettings.SwingRestitution, Dt, true);
+
+}
+
+void FPBDJointCachedSolver::InitSwingConstraint(
+	const FPBDJointSettings& JointSettings,
+	const FPBDJointSolverSettings& SolverSettings,
+	const FReal Dt,
+	const EJointAngularConstraintIndex SwingConstraintIndex)
+{
+	FVec3 SwingAxis;
+	FReal SwingAngle;
+	FPBDJointUtilities::GetSwingAxisAngle(ConnectorRs[0], ConnectorRs[1], SolverSettings.SwingTwistAngleTolerance, SwingConstraintIndex, SwingAxis, SwingAngle);
+
+	InitRotationConstraintDatas( JointSettings, (int32)SwingConstraintIndex, SwingAxis, SwingAngle, JointSettings.SwingRestitution, Dt, true);
+}
+
+void FPBDJointCachedSolver::InitLockedRotationConstraints(
+	  const FPBDJointSettings& JointSettings,
+	  const FReal Dt,
+	  const bool bApplyTwist,
+	  const bool bApplySwing)
+{
+	FVec3 Axis0, Axis1, Axis2;
+	FPBDJointUtilities::GetLockedRotationAxes(ConnectorRs[0], ConnectorRs[1], Axis0, Axis1, Axis2);
+
+	const FRotation3 R01 = ConnectorRs[0].Inverse() * ConnectorRs[1];
+
+	if (bApplyTwist)
+	{
+		InitRotationConstraintDatas(JointSettings, (int32)EJointAngularConstraintIndex::Twist, Axis0, R01.X, 0.0, Dt, false);
+	}
+
+	if (bApplySwing)
+	{
+		InitRotationConstraintDatas(JointSettings, (int32)EJointAngularConstraintIndex::Swing2, Axis1, R01.Y, 0.0, Dt, false);
+		InitRotationConstraintDatas(JointSettings, (int32)EJointAngularConstraintIndex::Swing1, Axis2, R01.Z, 0.0, Dt, false);
+	}
+}
+
+/** APPLY ROTATION CONSTRAINT ******************************************************************************************/
+
+void FPBDJointCachedSolver::ApplyRotationConstraints(
+	const FReal Dt)
+{
+	for(int32 ConstraintIndex = 0; ConstraintIndex < 3; ++ConstraintIndex)
+	{
+		if(RotationConstraints.bValidDatas[ConstraintIndex])
+		{
+			ApplyRotationConstraint(ConstraintIndex, Dt);
+		}
+	}
+}
+
+void FPBDJointCachedSolver::SolveRotationConstraintDelta(
+		const int32 ConstraintIndex, 
+		const FReal DeltaLambda)
+{
+	if(Body(0).IsDynamic())
+	{
+		const FVec3 DR0 = RotationConstraints.ConstraintDRAxis[ConstraintIndex][0] * DeltaLambda;
+		ApplyRotationDelta(0, DR0);
+	}
+	if(Body(1).IsDynamic())
+	{
+		const FVec3 DR1 = RotationConstraints.ConstraintDRAxis[ConstraintIndex][1] * DeltaLambda;
+		ApplyRotationDelta(1, DR1);
+	}
+	NetAngularImpulse += RotationConstraints.ConstraintAxis[ConstraintIndex]  * DeltaLambda;
+	++NumActiveConstraints;
+}
+
+void FPBDJointCachedSolver::SolveRotationConstraintHard(
+			const int32 ConstraintIndex,
+			const FReal DeltaConstraint)
+{
+	const FReal DeltaLambda = SolverStiffness * RotationConstraints.ConstraintHardStiffness[ConstraintIndex] * DeltaConstraint /
+		RotationConstraints.ConstraintHardIM[ConstraintIndex];
+
+	if (DeltaLambda >= 0)
+	{
+		RotationConstraints.ConstraintLambda[ConstraintIndex] += DeltaLambda;
+	}
+	SolveRotationConstraintDelta(ConstraintIndex, DeltaLambda);
+}
+
+void FPBDJointCachedSolver::SolveRotationConstraintSoft(
+			const int32 ConstraintIndex,
+			const FReal DeltaConstraint,
+			const FReal Dt,
+			const FReal TargetVel)
+{
+	// Damping angular velocity
+	FReal AngVelDt = 0;
+	if (RotationConstraints.ConstraintDamping[ConstraintIndex] > KINDA_SMALL_NUMBER)
+	{
+		const FVec3 W0Dt = FRotation3::CalculateAngularVelocity(InitConnectorRs[0], ConnectorRs[0], 1.0f);
+		const FVec3 W1Dt = FRotation3::CalculateAngularVelocity(InitConnectorRs[1], ConnectorRs[1], 1.0f);
+		AngVelDt = TargetVel * Dt + FVec3::DotProduct(RotationConstraints.ConstraintAxis[ConstraintIndex] , W0Dt - W1Dt);
+	}
+
+	const FReal DeltaLambda = SolverStiffness * (RotationConstraints.ConstraintSoftStiffness[ConstraintIndex] * DeltaConstraint -
+		RotationConstraints.ConstraintDamping[ConstraintIndex] * AngVelDt - RotationConstraints.ConstraintLambda[ConstraintIndex]) /
+		RotationConstraints.ConstraintSoftIM[ConstraintIndex];
+	RotationConstraints.ConstraintLambda[ConstraintIndex] += DeltaLambda;
+
+	SolveRotationConstraintDelta(ConstraintIndex, DeltaLambda);
+}
+
+void FPBDJointCachedSolver::ApplyRotationConstraint(
+	const int32 ConstraintIndex,
+	const FReal Dt)
+{
+	FReal DeltaAngle = RotationConstraints.ConstraintCX[ConstraintIndex] +
+		FVec3::DotProduct(Body(1).DQ() - Body(0).DQ(), RotationConstraints.ConstraintAxis[ConstraintIndex]);
+	
+	bool NeedsSolve = false;
+	if(RotationConstraints.bLimitsCheck[ConstraintIndex])
+	{
+		if(DeltaAngle > RotationConstraints.ConstraintLimits[ConstraintIndex] )
+		{
+			DeltaAngle -= RotationConstraints.ConstraintLimits[ConstraintIndex];
+			NeedsSolve = true;
+		}
+		else if(DeltaAngle < -RotationConstraints.ConstraintLimits[ConstraintIndex])
+		{
+			DeltaAngle += RotationConstraints.ConstraintLimits[ConstraintIndex];
+			NeedsSolve = true;
+		}
+	}
+
+	if (!RotationConstraints.bLimitsCheck[ConstraintIndex] || (RotationConstraints.bLimitsCheck[ConstraintIndex] && NeedsSolve && FMath::Abs(DeltaAngle) > AngleTolerance))
+	{
+		if (RotationConstraints.bSoftLimit[ConstraintIndex])
+		{
+			SolveRotationConstraintSoft(ConstraintIndex, DeltaAngle, Dt, 0.0f);
+		}
+		else
+		{
+			if (RotationConstraints.ConstraintRestitution[ConstraintIndex] > 0.0f)
+			{
+				CalculateAngularConstraintPadding(ConstraintIndex, RotationConstraints.ConstraintRestitution[ConstraintIndex], DeltaAngle);
+			}
+			SolveRotationConstraintHard(ConstraintIndex, DeltaAngle);
+		}
+	}
+}
+
+/** APPLY ANGULAR VELOCITY CONSTRAINT *********************************************************************************/
+
+void FPBDJointCachedSolver::ApplyAngularVelocityConstraints()
+{
+	for(int32 ConstraintIndex = 0; ConstraintIndex < 3; ++ConstraintIndex)
+	{
+		if(RotationConstraints.bValidDatas[ConstraintIndex])
+		{
+			ApplyAngularVelocityConstraint(ConstraintIndex);
+		}
+	}
+}
+
+void FPBDJointCachedSolver::SolveAngularVelocityConstraint(
+	const int32 ConstraintIndex,
+	const FReal TargetVel)
+{
+	const FVec3 CW = W(1) - W(0);
+
+	const FReal DeltaLambda = SolverStiffness * RotationConstraints.ConstraintHardStiffness[ConstraintIndex] *
+	 (FVec3::DotProduct(CW, RotationConstraints.ConstraintAxis[ConstraintIndex]) - TargetVel) / RotationConstraints.ConstraintHardIM[ConstraintIndex];
+
+	if(Body(0).IsDynamic())
+	{
+		const FVec3 DW0 = RotationConstraints.ConstraintDRAxis[ConstraintIndex][0] * DeltaLambda;
+	
+		Body(0).ApplyAngularVelocityDelta(DW0);
+	}
+	if(Body(1).IsDynamic())
+	{
+		const FVec3 DW1 = RotationConstraints.ConstraintDRAxis[ConstraintIndex][1] * DeltaLambda;
+	
+		Body(1).ApplyAngularVelocityDelta(DW1);
+	}
+}
+
+void FPBDJointCachedSolver::ApplyAngularVelocityConstraint(const int32 ConstraintIndex)
+{
+	if(!NetAngularImpulse.IsNearlyZero() && RotationConstraints.ConstraintLambda[ConstraintIndex] > SMALL_NUMBER)
+	{
+		FReal TargetVel = 0.0f;
+		if (RotationConstraints.ConstraintRestitution[ConstraintIndex] != 0.0f)
+		{
+			const FReal InitVel = InitConstraintAxisAngularVelocities[ConstraintIndex];
+			TargetVel = InitVel > Chaos_Joint_AngularVelocityThresholdToApplyRestitution ?
+				-RotationConstraints.ConstraintRestitution[ConstraintIndex] * InitVel : 0.0f;
+		}
+		SolveAngularVelocityConstraint(ConstraintIndex, TargetVel);
+	}
+}
+
+/** INIT POSITION DRIVES *********************************************************************************/
+
+void FPBDJointCachedSolver::InitPositionDrives(
+	const FReal Dt,
+	const FPBDJointSolverSettings& SolverSettings,
+	const FPBDJointSettings& JointSettings)
+{
+	PositionDrives.bValidDatas[0] = false;
+	PositionDrives.bValidDatas[1] = false;
+	PositionDrives.bValidDatas[2] = false;
+
+	if (SolverSettings.bEnableDrives)
+	{
+		TVec3<bool> bDriven =
+		{
+			(JointSettings.bLinearPositionDriveEnabled[0] || JointSettings.bLinearVelocityDriveEnabled[0]) && (JointSettings.LinearMotionTypes[0] != EJointMotionType::Locked),
+			(JointSettings.bLinearPositionDriveEnabled[1] || JointSettings.bLinearVelocityDriveEnabled[1]) && (JointSettings.LinearMotionTypes[1] != EJointMotionType::Locked),
+			(JointSettings.bLinearPositionDriveEnabled[2] || JointSettings.bLinearVelocityDriveEnabled[2]) && (JointSettings.LinearMotionTypes[2] != EJointMotionType::Locked),
+		};
+	
+		PositionDrives.bAccelerationMode = FPBDJointUtilities::GetDriveAccelerationMode(SolverSettings, JointSettings);
+
+		// Rectangular position drives
+		if (bDriven[0] || bDriven[1] || bDriven[2])
+		{
+			const FMatrix33 R0M = ConnectorRs[0].ToMatrix();
+			const FVec3 XTarget = ConnectorXs[0] + ConnectorRs[0] * JointSettings.LinearDrivePositionTarget;
+			const FVec3 VTarget = ConnectorRs[0] * JointSettings.LinearDriveVelocityTarget;
+			const FVec3 CX = ConnectorXs[1] - XTarget;
+
+			for (int32 AxisIndex = 0; AxisIndex < 3; ++AxisIndex)
+			{
+				if (bDriven[AxisIndex])
+				{
+					PositionDrives.InitDatas(AxisIndex, true, FPBDJointUtilities::GetLinearDriveStiffness(SolverSettings, JointSettings, AxisIndex),
+						FPBDJointUtilities::GetLinearDriveDamping(SolverSettings, JointSettings, AxisIndex), 0.0f);
+					const FVec3 Axis = R0M.GetAxis(AxisIndex);
+				
+					if ((FMath::Abs(FVec3::DotProduct(CX,Axis)) > PositionTolerance) || (PositionDrives.ConstraintDamping[AxisIndex] > 0.0f))
+					{
+						InitAxisPositionDrive(AxisIndex, Axis, CX, VTarget, Dt);
+					}
+				}
+			}
+		}
+	}
+}
+
+void FPBDJointCachedSolver::InitAxisPositionDrive(
+		const int32 ConstraintIndex,
+		const FVec3& ConstraintAxis,
+		const FVec3& DeltaPosition,
+		const FVec3& DeltaVelocity,
+		const FReal Dt)
+{	
+	const FVec3 ConstraintArm0 = ConnectorXs[0] - P(0);
+	const FVec3 ConstraintArm1 = ConnectorXs[1] - P(1);
+
+	const FVec3 AngularAxis0 = FVec3::CrossProduct(ConstraintArm0, ConstraintAxis);
+	const FVec3 AngularAxis1 = FVec3::CrossProduct(ConstraintArm1, ConstraintAxis);
+	const FVec3 IA0 = Utilities::Multiply(InvI(0), AngularAxis0);
+	const FVec3 IA1 = Utilities::Multiply(InvI(1), AngularAxis1);
+	const FReal II0 = FVec3::DotProduct(AngularAxis0, IA0);
+	const FReal II1 = FVec3::DotProduct(AngularAxis1, IA1);
+
+	PositionDrives.UpdateDatas(ConstraintIndex, ConstraintAxis, FVec3::DotProduct(DeltaPosition, ConstraintAxis),
+		IA0, IA1, InvM(0) + II0 + InvM(1) + II1, 0.0f, Dt, true, ConstraintArm0, ConstraintArm1,
+		FVec3::DotProduct(DeltaVelocity, ConstraintAxis));
+}
+
+/** APPLY POSITION  DRIVES *********************************************************************************/
+
+void FPBDJointCachedSolver::ApplyPositionDrives(
+		const FReal Dt)
+{
+	for(int32 ConstraintIndex = 0; ConstraintIndex < 3; ++ConstraintIndex)
+	{
+		if(PositionDrives.bValidDatas[ConstraintIndex])
+		{
+			ApplyAxisPositionDrive(ConstraintIndex, Dt);
+		}
+	}
+}
+
+void FPBDJointCachedSolver::ApplyAxisPositionDrive(
+		const int32 ConstraintIndex,
+		const FReal Dt)
+{
+	const FVec3 Delta0 = Body(0).DP() + FVec3::CrossProduct(Body(0).DQ(), PositionDrives.ConstraintArms[ConstraintIndex][0]);
+	const FVec3 Delta1 = Body(1).DP() + FVec3::CrossProduct(Body(1).DQ(), PositionDrives.ConstraintArms[ConstraintIndex][1]);
+
+	const FReal DeltaPos = PositionDrives.ConstraintCX[ConstraintIndex] + FVec3::DotProduct(Delta1 - Delta0, PositionDrives.ConstraintAxis[ConstraintIndex]);
+
+	FReal VelDt = 0;
+	if (PositionDrives.ConstraintDamping[ConstraintIndex] > KINDA_SMALL_NUMBER)
+	{
+		const FVec3 V0Dt = FVec3::CalculateVelocity(InitConnectorXs[0], ConnectorXs[0]+ Delta0, 1.0f);
+		const FVec3 V1Dt = FVec3::CalculateVelocity(InitConnectorXs[1], ConnectorXs[1]+ Delta1, 1.0f);
+		VelDt = PositionDrives.ConstraintVX[ConstraintIndex] * Dt + FVec3::DotProduct(V0Dt - V1Dt, PositionDrives.ConstraintAxis[ConstraintIndex] );
+	}
+
+	const FReal DeltaLambda = SolverStiffness * (PositionDrives.ConstraintSoftStiffness[ConstraintIndex] * DeltaPos -
+		PositionDrives.ConstraintDamping[ConstraintIndex] * VelDt - PositionDrives.ConstraintLambda[ConstraintIndex]) /
+		PositionDrives.ConstraintSoftIM[ConstraintIndex];
+	PositionDrives.ConstraintLambda[ConstraintIndex] += DeltaLambda;
+
+	const FVec3 DX = PositionDrives.ConstraintAxis[ConstraintIndex] * DeltaLambda;
+
+	if(Body(0).IsDynamic())
+	{
+		const FVec3 DP0 = InvM(0) * DX;
+		const FVec3 DR0 = PositionDrives.ConstraintDRAxis[ConstraintIndex][0] * DeltaLambda;
+		ApplyPositionDelta(0,DP0);
+		ApplyRotationDelta(0,DR0);
+	}
+	if(Body(1).IsDynamic())
+	{
+		const FVec3 DP1 = -InvM(1) * DX;
+		const FVec3 DR1 = PositionDrives.ConstraintDRAxis[ConstraintIndex][1] * DeltaLambda;
+		ApplyPositionDelta(1,DP1);
+		ApplyRotationDelta(1,DR1);
+	}
+
+	NetLinearImpulse += DX;
+	++NumActiveConstraints;
+}
+
+/** INIT ROTATION DRIVES *********************************************************************************/
+
+void FPBDJointCachedSolver::InitRotationDrives(
+	const FReal Dt,
+	const FPBDJointSolverSettings& SolverSettings,
+	const FPBDJointSettings& JointSettings)
+{
+	RotationDrives.bValidDatas[0] = false;
+	RotationDrives.bValidDatas[1] = false;
+	RotationDrives.bValidDatas[2] = false;
+
+	bool bHasRotationDrives =
+		JointSettings.bAngularTwistPositionDriveEnabled
+			|| JointSettings.bAngularTwistVelocityDriveEnabled
+			|| JointSettings.bAngularSwingPositionDriveEnabled
+			|| JointSettings.bAngularSwingVelocityDriveEnabled
+			|| JointSettings.bAngularSLerpPositionDriveEnabled
+			|| JointSettings.bAngularSLerpVelocityDriveEnabled;
+	if (!bHasRotationDrives)
+	{
+		return;
+	}
+
+	EJointMotionType TwistMotion = JointSettings.AngularMotionTypes[(int32)EJointAngularConstraintIndex::Twist];
+	EJointMotionType Swing1Motion = JointSettings.AngularMotionTypes[(int32)EJointAngularConstraintIndex::Swing1];
+	EJointMotionType Swing2Motion = JointSettings.AngularMotionTypes[(int32)EJointAngularConstraintIndex::Swing2];
+
+	if (SolverSettings.bEnableDrives)
+	{
+		bool bTwistLocked = TwistMotion == EJointMotionType::Locked;
+		bool bSwing1Locked = Swing1Motion == EJointMotionType::Locked;
+		bool bSwing2Locked = Swing2Motion == EJointMotionType::Locked;
+
+		// No SLerp drive if we have a locked rotation (it will be grayed out in the editor in this case, but could still have been set before the rotation was locked)
+		// @todo(ccaulfield): setting should be cleaned up before being passed to the solver
+		if ((JointSettings.bAngularSLerpPositionDriveEnabled || JointSettings.bAngularSLerpVelocityDriveEnabled) && !bTwistLocked && !bSwing1Locked && !bSwing2Locked)
+		{
+			InitSLerpDrive(Dt, SolverSettings, JointSettings);
+		}
+		else
+		{
+			const bool bTwistDriveEnabled = ((JointSettings.bAngularTwistPositionDriveEnabled || JointSettings.bAngularTwistVelocityDriveEnabled) && !bTwistLocked);
+			const bool bSwingDriveEnabled = (JointSettings.bAngularSwingPositionDriveEnabled || JointSettings.bAngularSwingVelocityDriveEnabled);
+			const bool bSwing1DriveEnabled = bSwingDriveEnabled && !bSwing1Locked;
+			const bool bSwing2DriveEnabled = bSwingDriveEnabled && !bSwing2Locked;
+			if (bTwistDriveEnabled || bSwing1DriveEnabled || bSwing2DriveEnabled)
+			{
+				InitSwingTwistDrives(Dt, SolverSettings, JointSettings, bTwistDriveEnabled, bSwing1DriveEnabled, bSwing2DriveEnabled);
+			}
+		}
+	}
+}
+
+void FPBDJointCachedSolver::InitRotationConstraintDrive(
+			const int32 ConstraintIndex,
+			const FVec3& ConstraintAxis,
+			const FReal Dt,
+			const FReal DeltaAngle)
+{
+	const FVec3 IA0 = Utilities::Multiply(InvI(0), ConstraintAxis);
+	const FVec3 IA1 = Utilities::Multiply(InvI(1), ConstraintAxis);
+	const FReal II0 = FVec3::DotProduct(ConstraintAxis, IA0);
+	const FReal II1 = FVec3::DotProduct(ConstraintAxis, IA1);
+
+	RotationDrives.UpdateDatas(ConstraintIndex, ConstraintAxis, DeltaAngle,
+		IA0, IA1, II0 + II1, 0.0f, Dt);
+}
+
+void FPBDJointCachedSolver::InitSwingTwistDrives(
+	const FReal Dt,
+	const FPBDJointSolverSettings& SolverSettings,
+	const FPBDJointSettings& JointSettings,
+	const bool bTwistDriveEnabled,
+	const bool bSwing1DriveEnabled,
+	const bool bSwing2DriveEnabled)
+{
+	FRotation3 R1Target = ConnectorRs[0] * JointSettings.AngularDrivePositionTarget;
+	R1Target.EnforceShortestArcWith(ConnectorRs[1]);
+	FRotation3 R1Error = R1Target.Inverse() * ConnectorRs[1];
+	FVec3 R1TwistAxisError = R1Error * FJointConstants::TwistAxis();
+
+	// Angle approximation Angle ~= Sin(Angle) for small angles, underestimates for large angles
+	const FReal DTwistAngle = 2.0f * R1Error.X;
+	const FReal DSwing1Angle = R1TwistAxisError.Y;
+	const FReal DSwing2Angle = -R1TwistAxisError.Z;
+
+	const int32 TW = (int32)EJointAngularConstraintIndex::Twist;
+	const int32 S1 = (int32)EJointAngularConstraintIndex::Swing1;
+	const int32 S2 = (int32)EJointAngularConstraintIndex::Swing2;
+
+	RotationDrives.InitDatas(TW, true, FPBDJointUtilities::GetAngularTwistDriveStiffness(SolverSettings, JointSettings),
+		FPBDJointUtilities::GetAngularTwistDriveDamping(SolverSettings, JointSettings), 0.0);
+	RotationDrives.InitDatas(S1, true, FPBDJointUtilities::GetAngularSwingDriveStiffness(SolverSettings, JointSettings),
+		FPBDJointUtilities::GetAngularSwingDriveDamping(SolverSettings, JointSettings), 0.0);
+	RotationDrives.InitDatas(S2, true, FPBDJointUtilities::GetAngularSwingDriveStiffness(SolverSettings, JointSettings),
+		FPBDJointUtilities::GetAngularSwingDriveDamping(SolverSettings, JointSettings), 0.0);
+
+	RotationDrives.bAccelerationMode = FPBDJointUtilities::GetDriveAccelerationMode(SolverSettings, JointSettings);
+
+	const bool bUseTwistDrive = bTwistDriveEnabled && (((FMath::Abs(DTwistAngle) > AngleTolerance) && (RotationDrives.ConstraintSoftStiffness[TW] > 0.0f)) || (RotationDrives.ConstraintDamping[TW]  > 0.0f));
+	if (bUseTwistDrive)
+	{
+		InitRotationConstraintDrive(TW, ConnectorRs[1] * FJointConstants::TwistAxis(), Dt, DTwistAngle);
+		RotationDrives.ConstraintVX[TW] = JointSettings.AngularDriveVelocityTarget[TW];
+	}
+
+	const bool bUseSwing1Drive = bSwing1DriveEnabled && (((FMath::Abs(DSwing1Angle) > AngleTolerance) && (RotationDrives.ConstraintSoftStiffness[S1] > 0.0f)) || (RotationDrives.ConstraintDamping[S1] > 0.0f));
+	if (bUseSwing1Drive)
+	{
+		InitRotationConstraintDrive(S1, ConnectorRs[1] * FJointConstants::Swing1Axis(),  Dt, DSwing1Angle);
+		RotationDrives.ConstraintVX[S1] = JointSettings.AngularDriveVelocityTarget[S1];
+	}
+
+	const bool bUseSwing2Drive = bSwing2DriveEnabled && (((FMath::Abs(DSwing2Angle) > AngleTolerance) && (RotationDrives.ConstraintSoftStiffness[S2] > 0.0f)) || (RotationDrives.ConstraintDamping[S2] > 0.0f));
+	if (bUseSwing2Drive)
+	{
+		InitRotationConstraintDrive(S2, ConnectorRs[1] * FJointConstants::Swing2Axis(),  Dt, DSwing2Angle);
+		RotationDrives.ConstraintVX[S2] = JointSettings.AngularDriveVelocityTarget[S2];
+	}
+}
+
+void FPBDJointCachedSolver::InitSLerpDrive(
+	const FReal Dt,
+	const FPBDJointSolverSettings& SolverSettings,
+	const FPBDJointSettings& JointSettings)
+{
+	for (int32 AxisIndex = 0; AxisIndex < 3; ++AxisIndex)
+	{
+		RotationDrives.InitDatas(AxisIndex, true, FPBDJointUtilities::GetAngularSLerpDriveStiffness(SolverSettings, JointSettings),
+						FPBDJointUtilities::GetAngularSLerpDriveDamping(SolverSettings, JointSettings), 0.0);
+	}
+	RotationDrives.bAccelerationMode = FPBDJointUtilities::GetDriveAccelerationMode(SolverSettings, JointSettings);
+
+	// If damping is enabled, we need to apply the drive about all 3 axes, but without damping we can just drive along the axis of error
+	if (RotationDrives.ConstraintDamping[0]  > 0.0f)
+	{
+		// NOTE: Slerp target velocity only works properly if we have a stiffness of zero.
+		FVec3 Axes[3] = { FVec3(1, 0, 0), FVec3(0, 1, 0), FVec3(0, 0, 1) };
+		if (RotationDrives.ConstraintSoftStiffness[0] > 0.0f)
+		{
+			FPBDJointUtilities::GetLockedRotationAxes(ConnectorRs[0], ConnectorRs[1], Axes[0], Axes[1], Axes[2]);
+			Utilities::NormalizeSafe(Axes[0], KINDA_SMALL_NUMBER);
+			Utilities::NormalizeSafe(Axes[1], KINDA_SMALL_NUMBER);
+			Utilities::NormalizeSafe(Axes[2], KINDA_SMALL_NUMBER);
+		}
+		const FRotation3 R01 = ConnectorRs[0].Inverse() * ConnectorRs[1];
+		FRotation3 TargetAngPos = JointSettings.AngularDrivePositionTarget;
+		TargetAngPos.EnforceShortestArcWith(R01);
+		const FRotation3 R1Error = TargetAngPos.Inverse() * R01;
+		FReal AxisAngles[3] = 
+		{ 
+			2.0f * FMath::Asin(R1Error.X), 
+			2.0f * FMath::Asin(R1Error.Y), 
+			2.0f * FMath::Asin(R1Error.Z) 
+		};
+
+		const FVec3 TargetAngVel = ConnectorRs[0] * JointSettings.AngularDriveVelocityTarget;
+
+		for (int32 AxisIndex = 0; AxisIndex < 3; ++AxisIndex)
+		{
+			InitRotationConstraintDrive(AxisIndex, Axes[AxisIndex],  Dt, AxisAngles[AxisIndex]);
+			RotationDrives.ConstraintVX[AxisIndex] = FVec3::DotProduct(TargetAngVel, RotationDrives.ConstraintAxis[AxisIndex]);
+		}
+	}
+	else
+	{
+		const FRotation3 TargetR1 = ConnectorRs[0] * JointSettings.AngularDrivePositionTarget;
+		const FRotation3 DR = TargetR1 * ConnectorRs[1].Inverse();
+
+		FVec3 SLerpAxis;
+		FReal SLerpAngle;
+		if (DR.ToAxisAndAngleSafe(SLerpAxis, SLerpAngle, FVec3(1, 0, 0)))
+		{
+			if (SLerpAngle > (FReal)PI)
+			{
+				SLerpAngle = SLerpAngle - (FReal)2 * PI;
+			}
+
+			if (FMath::Abs(SLerpAngle) > AngleTolerance)
+			{
+				InitRotationConstraintDrive(0, SLerpAxis, Dt, -SLerpAngle);
+			}
+		}
+	}
+}
+
+/** APPLY ROTATION DRIVES *********************************************************************************/
+
+void FPBDJointCachedSolver::ApplyRotationDrives(
+	const FReal Dt)
+{
+	for(int32 ConstraintIndex = 0; ConstraintIndex < 3; ++ConstraintIndex)
+	{
+		if(RotationDrives.bValidDatas[ConstraintIndex])
+		{
+			ApplyAxisRotationDrive(ConstraintIndex, Dt);
+		}
+	}
+}
+
+void FPBDJointCachedSolver::ApplyAxisRotationDrive(
+		const int32 ConstraintIndex,
+		const FReal Dt)
+{
+	const FReal DeltaConstraint = RotationDrives.ConstraintCX[ConstraintIndex]+
+		FVec3::DotProduct(Body(1).DQ()-Body(0).DQ(), RotationDrives.ConstraintAxis[ConstraintIndex]);
+
+	// Damping angular velocity
+	FReal AngVelDt = 0;
+	if (RotationDrives.ConstraintDamping[ConstraintIndex] > KINDA_SMALL_NUMBER)
+	{
+		const FVec3 W0Dt = FRotation3::CalculateAngularVelocity(InitConnectorRs[0], ConnectorRs[0], 1.0f);
+		const FVec3 W1Dt = FRotation3::CalculateAngularVelocity(InitConnectorRs[1], ConnectorRs[1], 1.0f);
+		AngVelDt = RotationDrives.ConstraintVX[ConstraintIndex] * Dt + FVec3::DotProduct(RotationDrives.ConstraintAxis[ConstraintIndex] , W0Dt - W1Dt);
+	}
+
+	const FReal DeltaLambda = SolverStiffness * (RotationDrives.ConstraintSoftStiffness[ConstraintIndex] * DeltaConstraint -
+		RotationDrives.ConstraintDamping[ConstraintIndex] * AngVelDt - RotationDrives.ConstraintLambda[ConstraintIndex]) /
+		RotationDrives.ConstraintSoftIM[ConstraintIndex];
+	RotationDrives.ConstraintLambda[ConstraintIndex] += DeltaLambda;
+
+	if(Body(0).IsDynamic())
+	{
+		const FVec3 DR0 = RotationDrives.ConstraintDRAxis[ConstraintIndex][0] * DeltaLambda;
+		ApplyRotationDelta(0, DR0);
+	}
+	if(Body(1).IsDynamic())
+	{
+		const FVec3 DR1 = RotationDrives.ConstraintDRAxis[ConstraintIndex][1] * DeltaLambda;
+		ApplyRotationDelta(1, DR1);
+	}
+	NetAngularImpulse += RotationDrives.ConstraintAxis[ConstraintIndex]  * DeltaLambda;
+	++NumActiveConstraints;
+}
+
+ // Joint utilities
+
+void FPBDJointCachedSolver::ApplyPositionDelta(
+	const int32 BodyIndex,
+	const FVec3& DP)
+{
+	Body(BodyIndex).ApplyPositionDelta(DP);
+}
+
+void FPBDJointCachedSolver::ApplyRotationDelta(
+	const int32 BodyIndex,
+	const FVec3& DR)
+{
+	Body(BodyIndex).ApplyRotationDelta(DR);
+}
+
+// Used for non-zero restitution. We pad constraints by an amount such that the velocity
+// calculated after solving constraint positions will as required for the restitution.
+void FPBDJointCachedSolver::CalculateLinearConstraintPadding(
+	const int32 ConstraintIndex,
+	const FReal Dt,
+	const FReal Restitution,
+	FReal& InOutPos)
+{
+	// NOTE: We only calculate the padding after the constraint is first violated, and after
+	// that the padding is fixed for the rest of the iterations in the current step.
+	if ((Restitution > 0.0f) && (InOutPos > 0.0f) && !HasLinearConstraintPadding(ConstraintIndex))
+	{
+		SetLinearConstraintPadding(ConstraintIndex, 0.0f);
+
+		// Calculate the velocity we want to match
     
-            Lambda += DLambda;
-            NetAngularImpulse += (KIndex == 0 ? 1 : -1) * DLambda * Axis;
-        }
-    }
-
-    // See "XPBD: Position-Based Simulation of Compliant Constrained Dynamics"
-    void FPBDJointCachedSolver::ApplyRotationConstraintSoftDD(
-        const FReal Dt,
-        const FReal JointStiffness,
-        const FReal JointDamping,
-        const bool bAccelerationMode,
-        const FVec3& Axis,
-        const FReal Angle,
-        const FReal AngVelTarget,
-        FReal& Lambda)
-    {
-        check(Body(0).IsDynamic());
-        check(Body(1).IsDynamic());
-
-        if (bRealTypeCompatibleWithISPC && bChaos_Joint_ISPC_Enabled)
-        {
-#if 0
-            float ReturnedLambda = Lambda;
-            ispc::ApplyRotationConstraintSoftDD((ispc::FPBDJointCachedSolver*)this, Dt, JointStiffness, JointDamping, bAccelerationMode, (ispc::FVector&) Axis, Angle, AngVelTarget, ReturnedLambda);
-            Lambda = ReturnedLambda;
-#endif
-        }
-        else
-        {
-            // World-space inverse mass
-            const FVec3 IA0 = Utilities::Multiply(InvI(0), Axis);
-            const FVec3 IA1 = Utilities::Multiply(InvI(1), Axis);
-
-            // Joint-space inverse mass
-            FReal II0 = FVec3::DotProduct(Axis, IA0);
-            FReal II1 = FVec3::DotProduct(Axis, IA1);
-            const FReal II = (II0 + II1);
-
-            // Damping angular velocity
-            FReal AngVelDt = 0;
-            if (JointDamping > KINDA_SMALL_NUMBER)
-            {
-                const FVec3 W0Dt = FRotation3::CalculateAngularVelocity(InitConnectorRs[0], ConnectorRs[0], 1.0f);
-                const FVec3 W1Dt = FRotation3::CalculateAngularVelocity(InitConnectorRs[1], ConnectorRs[1], 1.0f);
-                AngVelDt = AngVelTarget * Dt + FVec3::DotProduct(Axis, W0Dt - W1Dt);
-            }
-
-            const FReal SpringMassScale = (bAccelerationMode) ? 1.0f / II : 1.0f;
-            const FReal S = SpringMassScale * JointStiffness * Dt * Dt;
-            const FReal D = SpringMassScale * JointDamping * Dt;
-            const FReal Multiplier = (FReal)1 / ((S + D) * II + (FReal)1);
-            const FReal DLambda = SolverStiffness * Multiplier * (S * Angle - D * AngVelDt - Lambda);
-
-            //const FVec3 DR0 = IA0 * DLambda;
-            //const FVec3 DR1 = IA1 * -DLambda;
-            const FVec3 DR0 = Axis * (DLambda * II0);
-            const FVec3 DR1 = Axis * -(DLambda * II1);
-
-            ApplyRotationDelta(DR0, DR1);
-
-            Lambda += DLambda;
-            NetAngularImpulse += DLambda * Axis;
-        }
-    }
-
-    void FPBDJointCachedSolver::ApplyRotationConstraintSoft(
-        const FReal Dt,
-        const FReal JointStiffness,
-        const FReal JointDamping,
-        const bool bAccelerationMode,
-        const FVec3& Axis,
-        const FReal Angle,
-        const FReal AngVelTarget,
-        FReal& Lambda)
-    {
-        if (!Body0().IsDynamic())
-        {
-            ApplyRotationConstraintSoftKD(0, 1, Dt, JointStiffness, JointDamping, bAccelerationMode, Axis, Angle, AngVelTarget, Lambda);
-        }
-        else if (!Body1().IsDynamic())
-        {
-            ApplyRotationConstraintSoftKD(1, 0, Dt, JointStiffness, JointDamping, bAccelerationMode, Axis, -Angle, -AngVelTarget, Lambda);
-        }
-        else
-        {
-            ApplyRotationConstraintSoftDD(Dt, JointStiffness, JointDamping, bAccelerationMode, Axis, Angle, AngVelTarget, Lambda);
-        }
-
-        ++NumActiveConstraints;
-    }
-
-    // Used for non-zero restitution. We pad constraints by an amount such that the velocity
-    // calculated after solving constraint positions will as required for the restitution.
-    void FPBDJointCachedSolver::CalculateLinearConstraintPadding(
-        const FReal Dt,
-        const FPBDJointSolverSettings& SolverSettings,
-        const FPBDJointSettings& JointSettings,
-        const FReal Restitution,
-        const int32 AxisIndex,
-        const FVec3 Axis,
-        FReal& InOutPos)
-    {
-        // NOTE: We only calculate the padding after the constraint is first violated, and after
-        // that the padding is fixed for the rest of the iterations in the current step.
-        if ((Restitution > 0.0f) && (InOutPos > 0.0f) && !HasLinearConstraintPadding(AxisIndex))
-        {
-            SetLinearConstraintPadding(AxisIndex, 0.0f);
-
-            // Calculate the velocity we want to match
-            const FVec3 V0Dt = FVec3::CalculateVelocity(InitConnectorXs[0], ConnectorXs[0], 1.0f);
-            const FVec3 V1Dt = FVec3::CalculateVelocity(InitConnectorXs[1], ConnectorXs[1], 1.0f);
-            const FReal AxisVDt = FVec3::DotProduct(V1Dt - V0Dt, Axis);
-
-            // Calculate the padding to apply to the constraint that will result in the
-            // desired outward velocity (assuming the constraint is fully resolved)
-            const FReal Padding = (1.0f + Restitution) * AxisVDt - InOutPos;
-            if (Padding > 0.0f)
-            {
-                SetLinearConstraintPadding(AxisIndex, Padding);
-                InOutPos += Padding;
-            }
-        }
-    }
-
-    // Used for non-zero restitution. We pad constraints by an amount such that the velocity
-    // calculated after solving constraint positions will as required for the restitution.
-    void FPBDJointCachedSolver::CalculateAngularConstraintPadding(
-        const FReal Dt,
-        const FPBDJointSolverSettings& SolverSettings,
-        const FPBDJointSettings& JointSettings,
-        const FReal Restitution,
-        const EJointAngularConstraintIndex ConstraintIndex,
-        const FVec3 Axis,
-        FReal& InOutAngle)
-    {
-        // NOTE: We only calculate the padding after the constraint is first violated, and after
-        // that the padding is fixed for the rest of the iterations in the current step.
-        if ((Restitution > 0.0f) && (InOutAngle > 0.0f) && !HasAngularConstraintPadding(ConstraintIndex))
-        {
-            SetAngularConstraintPadding(ConstraintIndex, 0.0f);
-
-            // Calculate the velocity we want to match
-            const FVec3 W0Dt = FRotation3::CalculateAngularVelocity(InitConnectorRs[0], ConnectorRs[0], 1.0f);
-            const FVec3 W1Dt = FRotation3::CalculateAngularVelocity(InitConnectorRs[1], ConnectorRs[1], 1.0f);
-            const FReal AxisWDt = FVec3::DotProduct(W1Dt - W0Dt, Axis);
-
-            // Calculate the padding to apply to the constraint that will result in the
-            // desired outward velocity (assuming the constraint is fully resolved)
-            const FReal Padding = (1.0f + Restitution) * AxisWDt - InOutAngle;
-            if (Padding > 0.0f)
-            {
-                SetAngularConstraintPadding(ConstraintIndex, Padding);
-                InOutAngle += Padding;
-            }
-        }
-    }
-
-    //
-    //
-    //////////////////////////////////////////////////////////////////////////
-    //
-    //
-
-    void FPBDJointCachedSolver::ApplyLockedRotationConstraints(
-        const FReal Dt,
-        const FPBDJointSolverSettings& SolverSettings,
-        const FPBDJointSettings& JointSettings,
-        const bool bApplyTwist,
-        const bool bApplySwing)
-    {
-        FVec3 Axis0, Axis1, Axis2;
-        FPBDJointUtilities::GetLockedRotationAxes(ConnectorRs[0], ConnectorRs[1], Axis0, Axis1, Axis2);
-
-        const FRotation3 R01 = ConnectorRs[0].Inverse() * ConnectorRs[1];
-
-        if (bApplyTwist)
-        {
-            FReal TwistStiffness = FPBDJointUtilities::GetTwistStiffness(SolverSettings, JointSettings);
-            ApplyRotationConstraint(TwistStiffness, Axis0, R01.X, (int32)EJointAngularConstraintIndex::Twist);
-        }
-
-        if (bApplySwing)
-        {
-            FReal SwingStiffness = FPBDJointUtilities::GetSwingStiffness(SolverSettings, JointSettings);
-            ApplyRotationConstraint(SwingStiffness, Axis1, R01.Y, (int32)EJointAngularConstraintIndex::Swing2);
-            ApplyRotationConstraint(SwingStiffness, Axis2, R01.Z, (int32)EJointAngularConstraintIndex::Swing1);
-        }
-    }
-
-    void FPBDJointCachedSolver::ApplyTwistConstraint(
-        const FReal Dt,
-        const FPBDJointSolverSettings& SolverSettings,
-        const FPBDJointSettings& JointSettings,
-        const bool bUseSoftLimit)
-    {
-        FVec3 TwistAxis;
-        FReal TwistAngle;
-        FPBDJointUtilities::GetTwistAxisAngle(ConnectorRs[0], ConnectorRs[1], TwistAxis, TwistAngle);
-
-        // Calculate the twist correction to apply to each body
-        const FReal LimitPadding = GetAngularConstraintPadding(EJointAngularConstraintIndex::Twist);
-        FReal DTwistAngle = 0;
-        FReal TwistAngleMax = FMath::Max(JointSettings.AngularLimits[(int32)EJointAngularConstraintIndex::Twist] - LimitPadding, (FReal)0.);
-        if (TwistAngle > TwistAngleMax)
-        {
-            DTwistAngle = TwistAngle - TwistAngleMax;
-        }
-        else if (TwistAngle < -TwistAngleMax)
-        {
-            // Keep Twist error positive
-            DTwistAngle = -TwistAngle - TwistAngleMax;
-            TwistAxis = -TwistAxis;
-        }
-
-        UE_LOG(LogChaosJoint, VeryVerbose, TEXT("    Twist Angle %f [Limit %f]"), FMath::RadiansToDegrees(TwistAngle), FMath::RadiansToDegrees(TwistAngleMax));
-
-        // Apply twist correction
-        if (DTwistAngle > AngleTolerance)
-        {
-            if (bUseSoftLimit)
-            {
-                const FReal TwistStiffness = FPBDJointUtilities::GetSoftTwistStiffness(SolverSettings, JointSettings);
-                const FReal TwistDamping = FPBDJointUtilities::GetSoftTwistDamping(SolverSettings, JointSettings);
-                const bool bAccelerationMode = FPBDJointUtilities::GetAngularSoftAccelerationMode(SolverSettings, JointSettings);
-                ApplyRotationConstraintSoft(Dt, TwistStiffness, TwistDamping, bAccelerationMode, TwistAxis, DTwistAngle, 0.0f, TwistSoftLambda);
-            }
-            else
-            {
-                if (JointSettings.TwistRestitution > 0.0f)
-                {
-                    CalculateAngularConstraintPadding(Dt, SolverSettings, JointSettings, JointSettings.TwistRestitution, EJointAngularConstraintIndex::Twist, TwistAxis, DTwistAngle);
-                }
-
-                FReal TwistStiffness = FPBDJointUtilities::GetTwistStiffness(SolverSettings, JointSettings);
-                ApplyRotationConstraint(TwistStiffness, TwistAxis, DTwistAngle, (int32)EJointAngularConstraintIndex::Twist);
-            }
-        }
-    }
-
-    void FPBDJointCachedSolver::ApplyConeConstraint(
-        const FReal Dt,
-        const FPBDJointSolverSettings& SolverSettings,
-        const FPBDJointSettings& JointSettings,
-        const bool bUseSoftLimit)
-    {
-        FVec3 SwingAxisLocal;
-        FReal DSwingAngle = 0.0f;
-
-        const FReal LimitPadding = GetAngularConstraintPadding(EJointAngularConstraintIndex::Swing1);
-        const FReal Swing1Limit = FMath::Max(JointSettings.AngularLimits[(int32)EJointAngularConstraintIndex::Swing1] - LimitPadding, (FReal)0.);
-        const FReal Swing2Limit = FMath::Max(JointSettings.AngularLimits[(int32)EJointAngularConstraintIndex::Swing2] - LimitPadding, (FReal)0.);
-        FPBDJointUtilities::GetEllipticalConeAxisErrorLocal(ConnectorRs[0], ConnectorRs[1], Swing2Limit, Swing1Limit, SwingAxisLocal, DSwingAngle);
-
-        UE_LOG(LogChaosJoint, VeryVerbose, TEXT("    Cone Error %f [Limits %f %f]"), FMath::RadiansToDegrees(DSwingAngle), FMath::RadiansToDegrees(Swing2Limit), FMath::RadiansToDegrees(Swing1Limit));
-
-        const FVec3 SwingAxis = ConnectorRs[0] * SwingAxisLocal;
-
-        // Apply swing correction to each body
-        if (DSwingAngle > AngleTolerance)
-        {
-            if (bUseSoftLimit)
-            {
-                const FReal SoftSwingStiffness = FPBDJointUtilities::GetSoftSwingStiffness(SolverSettings, JointSettings);
-                const FReal SoftSwingDamping = FPBDJointUtilities::GetSoftSwingDamping(SolverSettings, JointSettings);
-                const bool bAccelerationMode = FPBDJointUtilities::GetAngularSoftAccelerationMode(SolverSettings, JointSettings);
-                ApplyRotationConstraintSoft(Dt, SoftSwingStiffness, SoftSwingDamping, bAccelerationMode, SwingAxis, DSwingAngle, 0.0f, SwingSoftLambda);
-            }
-            else
-            {
-                if (JointSettings.SwingRestitution > 0.0f)
-                {
-                    CalculateAngularConstraintPadding(Dt, SolverSettings, JointSettings, JointSettings.SwingRestitution, EJointAngularConstraintIndex::Swing1, SwingAxis, DSwingAngle);
-                }
-
-                FReal SwingStiffness = FPBDJointUtilities::GetSwingStiffness(SolverSettings, JointSettings);
-                // For cone constraints, the lambda are all accumulated in Swing2
-                ApplyRotationConstraint(SwingStiffness, SwingAxis, DSwingAngle, (int32)EJointAngularConstraintIndex::Swing2);
-            }
-        }
-    }
-
-
-    void FPBDJointCachedSolver::ApplySingleLockedSwingConstraint(
-        const FReal Dt,
-        const FPBDJointSolverSettings& SolverSettings,
-        const FPBDJointSettings& JointSettings,
-        const EJointAngularConstraintIndex SwingConstraintIndex,
-        const bool bUseSoftLimit)
-    {
-        // NOTE: SwingAxis is not normalized in this mode. It has length Sin(SwingAngle).
-        // Likewise, the SwingAngle is actually Sin(SwingAngle)
-        FVec3 SwingAxis;
-        FReal SwingAngle;
-        FPBDJointUtilities::GetLockedSwingAxisAngle(ConnectorRs[0], ConnectorRs[1], SwingConstraintIndex, SwingAxis, SwingAngle);
-
-        UE_LOG(LogChaosJoint, VeryVerbose, TEXT("    LockedSwing%d Angle %f [Tolerance %f]"), (SwingConstraintIndex == EJointAngularConstraintIndex::Swing1) ? 1 : 2, FMath::RadiansToDegrees(SwingAngle), FMath::RadiansToDegrees(AngleTolerance));
-
-        // Apply swing correction
-        FReal DSwingAngle = SwingAngle;
-        if (FMath::Abs(DSwingAngle) > AngleTolerance)
-        {
-            if (bUseSoftLimit)
-            {
-                const FReal SoftSwingStiffness = FPBDJointUtilities::GetSoftSwingStiffness(SolverSettings, JointSettings);
-                const FReal SoftSwingDamping = FPBDJointUtilities::GetSoftSwingDamping(SolverSettings, JointSettings);
-                const bool bAccelerationMode = FPBDJointUtilities::GetAngularSoftAccelerationMode(SolverSettings, JointSettings);
-                ApplyRotationConstraintSoft(Dt, SoftSwingStiffness, SoftSwingDamping, bAccelerationMode, SwingAxis, DSwingAngle, 0.0f, SwingSoftLambda);
-            }
-            else
-            {
-                const FReal SwingStiffness = FPBDJointUtilities::GetSwingStiffness(SolverSettings, JointSettings);
-                ApplyRotationConstraint(SwingStiffness, SwingAxis, DSwingAngle, (int32)SwingConstraintIndex);
-            }
-        }
-    }
-
-
-    void FPBDJointCachedSolver::ApplyDualConeSwingConstraint(
-        const FReal Dt,
-        const FPBDJointSolverSettings& SolverSettings,
-        const FPBDJointSettings& JointSettings,
-        const EJointAngularConstraintIndex SwingConstraintIndex,
-        const bool bUseSoftLimit)
-    {
-        FVec3 SwingAxis;
-        FReal SwingAngle;
-        FPBDJointUtilities::GetDualConeSwingAxisAngle(ConnectorRs[0], ConnectorRs[1], SwingConstraintIndex, SwingAxis, SwingAngle);
-
-        // Calculate swing error we need to correct
-        const FReal LimitPadding = GetAngularConstraintPadding(SwingConstraintIndex);
-        FReal DSwingAngle = 0;
-        const FReal SwingAngleMax = FMath::Max(JointSettings.AngularLimits[(int32)SwingConstraintIndex] - LimitPadding, (FReal)0.);
-        if (SwingAngle > SwingAngleMax)
-        {
-            DSwingAngle = SwingAngle - SwingAngleMax;
-        }
-        else if (SwingAngle < -SwingAngleMax)
-        {
-            // Keep the error positive
-            DSwingAngle = -SwingAngle - SwingAngleMax;
-            SwingAxis = -SwingAxis;
-        }
-
-        UE_LOG(LogChaosJoint, VeryVerbose, TEXT("    DualConeSwing%d Angle %f [Limit %f]"), (SwingConstraintIndex == EJointAngularConstraintIndex::Swing1) ? 1 : 2, FMath::RadiansToDegrees(SwingAngle), FMath::RadiansToDegrees(SwingAngleMax));
-
-        // Apply swing correction
-        if (DSwingAngle > SolverSettings.AngleTolerance)
-        {
-            if (bUseSoftLimit)
-            {
-                const FReal SoftSwingStiffness = FPBDJointUtilities::GetSoftSwingStiffness(SolverSettings, JointSettings);
-                const FReal SoftSwingDamping = FPBDJointUtilities::GetSoftSwingDamping(SolverSettings, JointSettings);
-                const bool bAccelerationMode = FPBDJointUtilities::GetAngularSoftAccelerationMode(SolverSettings, JointSettings);
-                ApplyRotationConstraintSoft(Dt, SoftSwingStiffness, SoftSwingDamping, bAccelerationMode, SwingAxis, DSwingAngle, 0.0f, SwingSoftLambda);
-            }
-            else
-            {
-                if (JointSettings.SwingRestitution > 0.0f)
-                {
-                    CalculateAngularConstraintPadding(Dt, SolverSettings, JointSettings, JointSettings.SwingRestitution, SwingConstraintIndex, SwingAxis, DSwingAngle);
-                }
-
-                const FReal SwingStiffness = FPBDJointUtilities::GetSwingStiffness(SolverSettings, JointSettings);
-                ApplyRotationConstraint(SwingStiffness, SwingAxis, DSwingAngle, (int32)SwingConstraintIndex);
-            }
-        }
-    }
-
-
-    void FPBDJointCachedSolver::ApplySwingConstraint(
-        const FReal Dt,
-        const FPBDJointSolverSettings& SolverSettings,
-        const FPBDJointSettings& JointSettings,
-        const EJointAngularConstraintIndex SwingConstraintIndex,
-        const bool bUseSoftLimit)
-    {
-        FVec3 SwingAxis;
-        FReal SwingAngle;
-        FPBDJointUtilities::GetSwingAxisAngle(ConnectorRs[0], ConnectorRs[1], SolverSettings.SwingTwistAngleTolerance, SwingConstraintIndex, SwingAxis, SwingAngle);
-
-        // Calculate swing error we need to correct
-        const FReal LimitPadding = GetAngularConstraintPadding(SwingConstraintIndex);
-        FReal DSwingAngle = 0;
-        const FReal SwingAngleMax = FMath::Max(JointSettings.AngularLimits[(int32)SwingConstraintIndex] - LimitPadding, (FReal)0.);
-        if (SwingAngle > SwingAngleMax)
-        {
-            DSwingAngle = SwingAngle - SwingAngleMax;
-        }
-        else if (SwingAngle < -SwingAngleMax)
-        {
-            // Keep swing error positive
-            DSwingAngle = -SwingAngle - SwingAngleMax;
-            SwingAxis = -SwingAxis;
-        }
-
-        UE_LOG(LogChaosJoint, VeryVerbose, TEXT("    Swing%d Angle %f [Limit %f]"), (SwingConstraintIndex == EJointAngularConstraintIndex::Swing1) ? 1 : 2, FMath::RadiansToDegrees(SwingAngle), FMath::RadiansToDegrees(SwingAngleMax));
-
-        // Apply swing correction
-        if (DSwingAngle > AngleTolerance)
-        {
-            if (bUseSoftLimit)
-            {
-                const FReal SoftSwingStiffness = FPBDJointUtilities::GetSoftSwingStiffness(SolverSettings, JointSettings);
-                const FReal SoftSwingDamping = FPBDJointUtilities::GetSoftSwingDamping(SolverSettings, JointSettings);
-                const bool bAccelerationMode = FPBDJointUtilities::GetAngularSoftAccelerationMode(SolverSettings, JointSettings);
-                ApplyRotationConstraintSoft(Dt, SoftSwingStiffness, SoftSwingDamping, bAccelerationMode, SwingAxis, DSwingAngle, 0.0f, SwingSoftLambda);
-            }
-            else
-            {
-                if (JointSettings.SwingRestitution > 0.0f)
-                {
-                    CalculateAngularConstraintPadding(Dt, SolverSettings, JointSettings, JointSettings.SwingRestitution, SwingConstraintIndex, SwingAxis, DSwingAngle);
-                }
-
-                const FReal SwingStiffness = FPBDJointUtilities::GetSwingStiffness(SolverSettings, JointSettings);
-                ApplyRotationConstraint(SwingStiffness, SwingAxis, DSwingAngle, (int32)SwingConstraintIndex);
-            }
-        }
-    }
-
-
-    void FPBDJointCachedSolver::ApplySwingTwistDrives(
-        const FReal Dt,
-        const FPBDJointSolverSettings& SolverSettings,
-        const FPBDJointSettings& JointSettings,
-        const bool bTwistDriveEnabled,
-        const bool bSwing1DriveEnabled,
-        const bool bSwing2DriveEnabled)
-    {
-        FRotation3 R1Target = ConnectorRs[0] * JointSettings.AngularDrivePositionTarget;
-        R1Target.EnforceShortestArcWith(ConnectorRs[1]);
-        FRotation3 R1Error = R1Target.Inverse() * ConnectorRs[1];
-        FVec3 R1TwistAxisError = R1Error * FJointConstants::TwistAxis();
-
-        // Angle approximation Angle ~= Sin(Angle) for small angles, underestimates for large angles
-        const FReal DTwistAngle = 2.0f * R1Error.X;
-        const FReal DSwing1Angle = R1TwistAxisError.Y;
-        const FReal DSwing2Angle = -R1TwistAxisError.Z;
-
-        const FReal AngularTwistDriveStiffness = FPBDJointUtilities::GetAngularTwistDriveStiffness(SolverSettings, JointSettings);
-        const FReal AngularTwistDriveDamping = FPBDJointUtilities::GetAngularTwistDriveDamping(SolverSettings, JointSettings);
-        const FReal AngularSwingDriveStiffness = FPBDJointUtilities::GetAngularSwingDriveStiffness(SolverSettings, JointSettings);
-        const FReal AngularSwingDriveDamping = FPBDJointUtilities::GetAngularSwingDriveDamping(SolverSettings, JointSettings);
-        const bool bAccelerationMode = FPBDJointUtilities::GetDriveAccelerationMode(SolverSettings, JointSettings);
-
-        const bool bUseTwistDrive = bTwistDriveEnabled && (((FMath::Abs(DTwistAngle) > AngleTolerance) && (AngularTwistDriveStiffness > 0.0f)) || (AngularTwistDriveDamping > 0.0f));
-        if (bUseTwistDrive)
-        {
-            FReal ReturnedLambda = RotationDriveLambdas[(int32)EJointAngularConstraintIndex::Twist];
-            const FVec3 TwistAxis = ConnectorRs[1] * FJointConstants::TwistAxis();
-            ApplyRotationConstraintSoft(Dt, AngularTwistDriveStiffness, AngularTwistDriveDamping, bAccelerationMode, TwistAxis, DTwistAngle, JointSettings.AngularDriveVelocityTarget[(int32)EJointAngularConstraintIndex::Twist], ReturnedLambda);
-            RotationDriveLambdas[(int32)EJointAngularConstraintIndex::Twist] = ReturnedLambda;
-        }
-
-        const bool bUseSwing1Drive = bSwing1DriveEnabled && (((FMath::Abs(DSwing1Angle) > AngleTolerance) && (AngularSwingDriveStiffness > 0.0f)) || (AngularSwingDriveDamping > 0.0f));
-        if (bUseSwing1Drive)
-        {
-            FReal ReturnedLambda = RotationDriveLambdas[(int32)EJointAngularConstraintIndex::Swing1];
-            const FVec3 Swing1Axis = ConnectorRs[1] * FJointConstants::Swing1Axis();
-            ApplyRotationConstraintSoft(Dt, AngularSwingDriveStiffness, AngularSwingDriveDamping, bAccelerationMode, Swing1Axis, DSwing1Angle, JointSettings.AngularDriveVelocityTarget[(int32)EJointAngularConstraintIndex::Swing1], ReturnedLambda);
-            RotationDriveLambdas[(int32)EJointAngularConstraintIndex::Swing1] = ReturnedLambda;
-        }
-
-        const bool bUseSwing2Drive = bSwing2DriveEnabled && (((FMath::Abs(DSwing2Angle) > AngleTolerance) && (AngularSwingDriveStiffness > 0.0f)) || (AngularSwingDriveDamping > 0.0f));
-        if (bUseSwing2Drive)
-        {
-            FReal ReturnedLambda = RotationDriveLambdas[(int32)EJointAngularConstraintIndex::Swing2];
-            const FVec3 Swing2Axis = ConnectorRs[1] * FJointConstants::Swing2Axis();
-            ApplyRotationConstraintSoft(Dt, AngularSwingDriveStiffness, AngularSwingDriveDamping, bAccelerationMode, Swing2Axis, DSwing2Angle, JointSettings.AngularDriveVelocityTarget[(int32)EJointAngularConstraintIndex::Swing2], ReturnedLambda);
-            RotationDriveLambdas[(int32)EJointAngularConstraintIndex::Swing2] = ReturnedLambda;
-        }
-    }
-
-
-    void FPBDJointCachedSolver::ApplySLerpDrive(
-        const FReal Dt,
-        const FPBDJointSolverSettings& SolverSettings,
-        const FPBDJointSettings& JointSettings)
-    {
-        const FReal AngularDriveStiffness = FPBDJointUtilities::GetAngularSLerpDriveStiffness(SolverSettings, JointSettings);
-        const FReal AngularDriveDamping = FPBDJointUtilities::GetAngularSLerpDriveDamping(SolverSettings, JointSettings);
-        const bool bAccelerationMode = FPBDJointUtilities::GetDriveAccelerationMode(SolverSettings, JointSettings);
-
-        // If damping is enabled, we need to apply the drive about all 3 axes, but without damping we can just drive along the axis of error
-        if (AngularDriveDamping > 0.0f)
-        {
-            // NOTE: Slerp target velocity only works properly if we have a stiffness of zero.
-            FVec3 Axes[3] = { FVec3(1, 0, 0), FVec3(0, 1, 0), FVec3(0, 0, 1) };
-            if (AngularDriveStiffness > 0.0f)
-            {
-                FPBDJointUtilities::GetLockedRotationAxes(ConnectorRs[0], ConnectorRs[1], Axes[0], Axes[1], Axes[2]);
-                Utilities::NormalizeSafe(Axes[0], KINDA_SMALL_NUMBER);
-                Utilities::NormalizeSafe(Axes[1], KINDA_SMALL_NUMBER);
-                Utilities::NormalizeSafe(Axes[2], KINDA_SMALL_NUMBER);
-            }
-
-            const FRotation3 R01 = ConnectorRs[0].Inverse() * ConnectorRs[1];
-            FRotation3 TargetAngPos = JointSettings.AngularDrivePositionTarget;
-            TargetAngPos.EnforceShortestArcWith(R01);
-            const FRotation3 R1Error = TargetAngPos.Inverse() * R01;
-            FReal AxisAngles[3] = 
-            { 
-                2.0f * FMath::Asin(R1Error.X), 
-                2.0f * FMath::Asin(R1Error.Y), 
-                2.0f * FMath::Asin(R1Error.Z) 
-            };
-
-            const FVec3 TargetAngVel = ConnectorRs[0] * JointSettings.AngularDriveVelocityTarget;
-
-            for (int32 AxisIndex = 0; AxisIndex < 3; ++AxisIndex)
-            {
-                FReal ReturnedLambda = RotationDriveLambdas[AxisIndex];
-                const FReal AxisAngVel = FVec3::DotProduct(TargetAngVel, Axes[AxisIndex]);
-                ApplyRotationConstraintSoft(Dt, AngularDriveStiffness, AngularDriveDamping, bAccelerationMode, Axes[AxisIndex], AxisAngles[AxisIndex], AxisAngVel, ReturnedLambda);
-                RotationDriveLambdas[AxisIndex] = ReturnedLambda;
-            }
-        }
-        else
-        {
-            const FRotation3 TargetR1 = ConnectorRs[0] * JointSettings.AngularDrivePositionTarget;
-            const FRotation3 DR = TargetR1 * ConnectorRs[1].Inverse();
-
-            FVec3 SLerpAxis;
-            FReal SLerpAngle;
-            if (DR.ToAxisAndAngleSafe(SLerpAxis, SLerpAngle, FVec3(1, 0, 0)))
-            {
-                if (SLerpAngle > (FReal)PI)
-                {
-                    SLerpAngle = SLerpAngle - (FReal)2 * PI;
-                }
-
-                UE_LOG(LogChaosJoint, VeryVerbose, TEXT("      SLerpDrive Pos: %f Axis: %f %f %f"), -SLerpAngle, SLerpAxis.X, SLerpAxis.Y, SLerpAxis.Z);
-
-                if (FMath::Abs(SLerpAngle) > AngleTolerance)
-                {
-                    FReal ReturnedLambda = RotationDriveLambdas[(int32)EJointAngularConstraintIndex::Swing1];
-                    FReal AngVelTarget = (JointSettings.AngularDriveDamping > FReal(0)) ? FVec3::DotProduct(SLerpAxis, ConnectorRs[0] * JointSettings.AngularDriveVelocityTarget) : 0.0f;
-                    ApplyRotationConstraintSoft(Dt, AngularDriveStiffness, AngularDriveDamping, bAccelerationMode, SLerpAxis, -SLerpAngle, AngVelTarget, ReturnedLambda);
-                    RotationDriveLambdas[(int32)EJointAngularConstraintIndex::Swing1] = ReturnedLambda;
-                }
-            }
-        }
-    }
-
-
-    // Kinematic-Dynamic bodies
-    void FPBDJointCachedSolver::ApplyPointPositionConstraintKD(
-        const int32 KIndex,
-        const int32 DIndex,
-        const FReal Dt,
-        const FPBDJointSolverSettings& SolverSettings,
-        const FPBDJointSettings& JointSettings)
-    {
-        check(!Body(KIndex).IsDynamic());
-        check(Body(DIndex).IsDynamic());
-
-        FReal Stiffness = SolverStiffness * FPBDJointUtilities::GetLinearStiffness(SolverSettings, JointSettings);
-        const FVec3 CX = ConnectorXs[DIndex] - ConnectorXs[KIndex];
-
-        UE_LOG(LogChaosJoint, VeryVerbose, TEXT("    PointKD Delta %f [Limit %f]"), CX.Size(), PositionTolerance);
-
-        if (CX.SizeSquared() > PositionTolerance * PositionTolerance)
-        {
-            if (bRealTypeCompatibleWithISPC && bChaos_Joint_ISPC_Enabled)
-            {
-                check(false);
-#if 0
-                ispc::ApplyPointPositionConstraintKD((ispc::FPBDJointCachedSolver*)this, KIndex, DIndex, (ispc::FVector&)CX, Stiffness);
-#endif
-            }
-            else
-            {
-                // Calculate constraint correction
-                const FVec3 DX = Stiffness * Utilities::Multiply(PointPositionConstraintIM, CX);
-
-                // Apply constraint correction
-                const FVec3 ArmD = ConstraintArms[PointPositionConstraintIndex][DIndex];
-                const FVec3 DP1 = -InvM(DIndex) * DX;
-                const FVec3 DR1 = Utilities::Multiply(InvI(DIndex), FVec3::CrossProduct(ArmD, -DX));
-
-                ApplyDelta(DIndex, DP1, DR1);
-
-                NetLinearImpulse += (KIndex == 0) ? DX : -DX;
-            }
-
-            ++NumActiveConstraints;
-        }
-    }
-
-
-    // Dynamic-Dynamic bodies
-    void FPBDJointCachedSolver::ApplyPointPositionConstraintDD(
-        const FReal Dt,
-        const FPBDJointSolverSettings& SolverSettings,
-        const FPBDJointSettings& JointSettings)
-    {
-        check(Body(0).IsDynamic());
-        check(Body(1).IsDynamic());
-
-        FReal Stiffness = SolverStiffness * FPBDJointUtilities::GetLinearStiffness(SolverSettings, JointSettings);
-        const FVec3 CX = ConnectorXs[1] - ConnectorXs[0];
-
-        UE_LOG(LogChaosJoint, VeryVerbose, TEXT("    PointDD Delta %f [Limit %f]"), CX.Size(), PositionTolerance);
-
-        if (CX.SizeSquared() > PositionTolerance * PositionTolerance)
-        {
-            if (bRealTypeCompatibleWithISPC && bChaos_Joint_ISPC_Enabled)
-            {
-                check(false);
-#if 0
-                ispc::ApplyPointPositionConstraintDD((ispc::FPBDJointCachedSolver*)this, (ispc::FVector&)CX, Stiffness);
-#endif
-            }
-            else
-            {
-                // Calculate constraint correction
-                const FVec3 DX = Stiffness * Utilities::Multiply(PointPositionConstraintIM, CX);
-
-                // Apply constraint correction
-                const FVec3 Arm0 = ConstraintArms[PointPositionConstraintIndex][0];
-                const FVec3 Arm1 = ConstraintArms[PointPositionConstraintIndex][1];
-                const FVec3 DP0 = InvM(0) * DX;
-                const FVec3 DP1 = -InvM(1) * DX;
-                const FVec3 DR0 = Utilities::Multiply(InvI(0), FVec3::CrossProduct(Arm0, DX));
-                const FVec3 DR1 = Utilities::Multiply(InvI(1), FVec3::CrossProduct(Arm1, -DX));
-
-                ApplyPositionDelta(DP0, DP1);
-                ApplyRotationDelta(DR0, DR1);
-
-                NetLinearImpulse += DX;
-            }
-            
-            ++NumActiveConstraints;
-        }
-    }
-
-
-    void FPBDJointCachedSolver::ApplySphericalPositionConstraint(
-        const FReal Dt,
-        const FPBDJointSolverSettings& SolverSettings,
-        const FPBDJointSettings& JointSettings)
-    {
-        const FVec3 Axis = PositionConstraintAxis[SphericalPositionConstraintIndex];
-        const FReal Delta = FVec3::DotProduct(ConnectorXs[1] - ConnectorXs[0], Axis);
-
-
-        const FReal LimitPadding = GetLinearConstraintPadding(0);
-        const FReal Limit = FMath::Max(JointSettings.LinearLimit - LimitPadding, (FReal)0.);
-
-        FReal Error = Delta - Limit;
-        if (Error > PositionTolerance)
-        {
-            if (!FPBDJointUtilities::GetSoftLinearLimitEnabled(SolverSettings, JointSettings))
-            {
-                if (JointSettings.LinearRestitution > 0.0f)
-                {
-                    CalculateLinearConstraintPadding(Dt, SolverSettings, JointSettings, JointSettings.LinearRestitution, 0, Axis, Error);
-                }
-
-                const FReal& IM = PositionConstraintIM[SphericalPositionConstraintIndex];
-                const FVec3& Arm0 = ConstraintArms[SphericalPositionConstraintIndex][0];
-                const FVec3& Arm1 = ConstraintArms[SphericalPositionConstraintIndex][1];
-                const FVec3& DRAxis0 = PositionConstraintDRAxis[SphericalPositionConstraintIndex][0];
-                const FVec3& DRAxis1 = PositionConstraintDRAxis[SphericalPositionConstraintIndex][1];
-                ApplyPositionConstraint(LinearStiffness, Axis, Error, IM, Arm0, Arm1, DRAxis0, DRAxis1);
-            }
-            else
-            {
-                const FReal JointStiffness = FPBDJointUtilities::GetSoftLinearStiffness(SolverSettings, JointSettings);
-                const FReal JointDamping = FPBDJointUtilities::GetSoftLinearDamping(SolverSettings, JointSettings);
-                const bool bAccelerationMode = FPBDJointUtilities::GetLinearSoftAccelerationMode(SolverSettings, JointSettings);
-                ApplyPositionConstraintSoft(Dt, JointStiffness, JointDamping, bAccelerationMode, Axis, Error, 0.0f, LinearSoftLambda);
-            }
-        }
-    }
-
-
-    void FPBDJointCachedSolver::ApplyCylindricalPositionConstraint(
-        const FReal Dt,
-        const int32 AxisIndex,
-        const EJointMotionType AxialMotion,
-        const EJointMotionType RadialMotion,
-        const FPBDJointSolverSettings& SolverSettings,
-        const FPBDJointSettings& JointSettings)
-    {
-        check(AxialMotion != RadialMotion);
-
-        FVec3 Axis = PositionConstraintAxis[(int32)EJointCylindricalPositionConstraintType::Axial];
-        FReal AxialDelta = FVec3::DotProduct(ConnectorXs[1] - ConnectorXs[0], Axis);
-        if (AxialDelta < 0.f)
-        {
-            AxialDelta = -AxialDelta;
-            Axis = -Axis;
-        }
-        const FReal AxialLimitPadding = GetLinearConstraintPadding(0);
-        const FReal AxialLimit = (AxialMotion == EJointMotionType::Locked) ? 0.0f : FMath::Max(JointSettings.LinearLimit - AxialLimitPadding, (FReal)0.);
-        FReal AxialError = AxialDelta - AxialLimit;
-
-        if (AxialError > PositionTolerance)
-        {
-            if ((AxialMotion == EJointMotionType::Limited) && FPBDJointUtilities::GetSoftLinearLimitEnabled(SolverSettings, JointSettings))
-            {
-                // Soft Axial constraint
-                const FReal JointStiffness = FPBDJointUtilities::GetSoftLinearStiffness(SolverSettings, JointSettings);
-                const FReal JointDamping = FPBDJointUtilities::GetSoftLinearDamping(SolverSettings, JointSettings);
-                const bool bAccelerationMode = FPBDJointUtilities::GetLinearSoftAccelerationMode(SolverSettings, JointSettings);
-                ApplyPositionConstraintSoft(Dt, JointStiffness, JointDamping, bAccelerationMode, Axis, AxialError, 0.0f, LinearSoftLambda);
-            }
-            else if (AxialMotion != EJointMotionType::Free)
-            {
-                if (JointSettings.LinearRestitution > 0.0f)
-                {
-                    CalculateLinearConstraintPadding(Dt, SolverSettings, JointSettings, JointSettings.LinearRestitution, 0, Axis, AxialError);
-                }
-                const FReal& IM = PositionConstraintIM[(int32)EJointCylindricalPositionConstraintType::Axial];
-                const FVec3& Arm0 = ConstraintArms[(int32)EJointCylindricalPositionConstraintType::Axial][0];
-                const FVec3& Arm1 = ConstraintArms[(int32)EJointCylindricalPositionConstraintType::Axial][1];
-                const FVec3& DRAxis0 = PositionConstraintDRAxis[(int32)EJointCylindricalPositionConstraintType::Axial][0];
-                const FVec3& DRAxis1 = PositionConstraintDRAxis[(int32)EJointCylindricalPositionConstraintType::Axial][1];
-                ApplyPositionConstraint(LinearStiffness, Axis, AxialError, IM, Arm0, Arm1, DRAxis0, DRAxis1, (int32)EJointCylindricalPositionConstraintType::Axial);
-            }
-        }
-
-        const FVec3& RadialAxis = PositionConstraintAxis[(int32)EJointCylindricalPositionConstraintType::Radial];
-        const FReal RadialDelta = FVec3::DotProduct(ConnectorXs[1] - ConnectorXs[0], RadialAxis);
-        const FReal RadialLimitPadding = GetLinearConstraintPadding(1);
-        const FReal RadialLimit = (RadialMotion == EJointMotionType::Locked) ? 0.0f : FMath::Max(JointSettings.LinearLimit - AxialLimitPadding, (FReal)0.);
-        FReal RadialError = RadialDelta - RadialLimit;
-        if (RadialError > PositionTolerance)
-        {
-            if ((RadialMotion == EJointMotionType::Limited) && FPBDJointUtilities::GetSoftLinearLimitEnabled(SolverSettings, JointSettings))
-            {
-                // Soft Radial constraint
-                const FReal JointStiffness = FPBDJointUtilities::GetSoftLinearStiffness(SolverSettings, JointSettings);
-                const FReal JointDamping = FPBDJointUtilities::GetSoftLinearDamping(SolverSettings, JointSettings);
-                const bool bAccelerationMode = FPBDJointUtilities::GetLinearSoftAccelerationMode(SolverSettings, JointSettings);
-                ApplyPositionConstraintSoft(Dt, JointStiffness, JointDamping, bAccelerationMode, RadialAxis, RadialError, 0.0f, LinearSoftLambda);
-            }
-            else if (RadialMotion != EJointMotionType::Free)
-            {
-                // Hard Radial constraint
-                if (JointSettings.LinearRestitution > 0.0f)
-                {
-                    CalculateLinearConstraintPadding(Dt, SolverSettings, JointSettings, JointSettings.LinearRestitution, 1, RadialAxis, RadialError);
-                }
-
-                const FReal& IM = PositionConstraintIM[(int32)EJointCylindricalPositionConstraintType::Radial];
-                const FVec3& Arm0 = ConstraintArms[(int32)EJointCylindricalPositionConstraintType::Radial][0];
-                const FVec3& Arm1 = ConstraintArms[(int32)EJointCylindricalPositionConstraintType::Radial][1];
-                const FVec3& DRAxis0 = PositionConstraintDRAxis[(int32)EJointCylindricalPositionConstraintType::Radial][0];
-                const FVec3& DRAxis1 = PositionConstraintDRAxis[(int32)EJointCylindricalPositionConstraintType::Radial][1];
-                ApplyPositionConstraint(LinearStiffness, RadialAxis, RadialError, IM, Arm0, Arm1, DRAxis0, DRAxis1, (int32)EJointCylindricalPositionConstraintType::Radial);
-            }
-        }
-    }
-
-    void FPBDJointCachedSolver::ApplyPlanarPositionConstraint(
-        const FReal Dt,
-        const int32 AxisIndex,
-        const EJointMotionType AxialMotion,
-        const FPBDJointSolverSettings& SolverSettings,
-        const FPBDJointSettings& JointSettings)
-    {
-        FVec3 Axis = PositionConstraintAxis[PlanarPositionConstraintIndex];
-        FReal Delta = FVec3::DotProduct(ConnectorXs[1] - ConnectorXs[0], Axis);
-        if (Delta < 0.0f)
-        {
-            Delta = -Delta;
-            Axis = -Axis;
-        }
-
-        const FReal LimitPadding = GetLinearConstraintPadding(0);
-        const FReal Limit = (AxialMotion == EJointMotionType::Locked) ? 0 : FMath::Max(JointSettings.LinearLimit - LimitPadding, (FReal)0.);
-        FReal Error = Delta - Limit;
-        if (Error > PositionTolerance)
-        {
-            if ((AxialMotion == EJointMotionType::Limited) && FPBDJointUtilities::GetSoftLinearLimitEnabled(SolverSettings, JointSettings))
-            {
-                const FReal JointStiffness = FPBDJointUtilities::GetSoftLinearStiffness(SolverSettings, JointSettings);
-                const FReal JointDamping = FPBDJointUtilities::GetSoftLinearDamping(SolverSettings, JointSettings);
-                const bool bAccelerationMode = FPBDJointUtilities::GetLinearSoftAccelerationMode(SolverSettings, JointSettings);
-                ApplyPositionConstraintSoft(Dt, JointStiffness, JointDamping, bAccelerationMode, Axis, Error, 0.0f, LinearSoftLambda);
-            }
-            else
-            {
-                if (JointSettings.LinearRestitution > 0.0f)
-                {
-                    CalculateLinearConstraintPadding(Dt, SolverSettings, JointSettings, JointSettings.LinearRestitution, 0, Axis, Error);
-                }
-
-                const FReal& IM = PositionConstraintIM[PlanarPositionConstraintIndex];
-                const FVec3& Arm0 = ConstraintArms[PlanarPositionConstraintIndex][0];
-                const FVec3& Arm1 = ConstraintArms[PlanarPositionConstraintIndex][1];
-                const FVec3& DRAxis0 = PositionConstraintDRAxis[PlanarPositionConstraintIndex][0];
-                const FVec3& DRAxis1 = PositionConstraintDRAxis[PlanarPositionConstraintIndex][1];
-                ApplyPositionConstraint(LinearStiffness, Axis, Error, IM, Arm0, Arm1, DRAxis0, DRAxis1, PlanarPositionConstraintIndex);
-            }
-        }
-    }
-
-
-    void FPBDJointCachedSolver::ApplyPositionDrive(
-        const FReal Dt,
-        const int32 AxisIndex,
-        const FPBDJointSolverSettings& SolverSettings,
-        const FPBDJointSettings& JointSettings,
-        const FVec3& Axis,
-        const FReal DeltaPos,
-        const FReal DeltaVel)
-    {
-        const FReal JointStiffness = FPBDJointUtilities::GetLinearDriveStiffness(SolverSettings, JointSettings, AxisIndex);
-        const FReal JointDamping = FPBDJointUtilities::GetLinearDriveDamping(SolverSettings, JointSettings, AxisIndex);
-        const bool bAccelerationMode = FPBDJointUtilities::GetDriveAccelerationMode(SolverSettings, JointSettings);
-
-        if ((FMath::Abs(DeltaPos) > PositionTolerance) || (JointDamping > 0.0f))
-        {
-            FReal ReturnedLambda = LinearDriveLambdas[AxisIndex];
-            ApplyPositionConstraintSoft(Dt, JointStiffness, JointDamping, bAccelerationMode, Axis, DeltaPos, DeltaVel, ReturnedLambda);
-            LinearDriveLambdas[AxisIndex] = ReturnedLambda;
-        }
-    }
-
-
-    void FPBDJointCachedSolver::ApplyPointProjection(
-        const FReal Dt,
-        const FPBDJointSolverSettings& SolverSettings,
-        const FPBDJointSettings& JointSettings,
-        const FReal Alpha,
-        FVec3& NetDP1,
-        FVec3& NetDR1)
-    {
-        const FReal ProjectionPositionTolerance = 0.0f;//PositionTolerance;
-
-        const FVec3 CX = ConnectorXs[1] - ConnectorXs[0];
-        if (CX.Size() > ProjectionPositionTolerance)
-        {
-            FMatrix33 J = Utilities::ComputeJointFactorMatrix(ConnectorXs[1] - P(1), InvI(1), InvM(1));
-            const FMatrix33 IJ = J.Inverse();
-            const FVec3 DX = Utilities::Multiply(IJ, CX);
-
-            const FVec3 DP1 = -Alpha * InvM(1) * DX;
-            const FVec3 DR1 = -Alpha * Utilities::Multiply(InvI(1), FVec3::CrossProduct(ConnectorXs[1] - P(1), DX));
-            ApplyDelta(1, DP1, DR1);
-
-            NetDP1 += DP1;
-            NetDR1 += DR1;
-        }
-    }
-
-    void FPBDJointCachedSolver::ApplySphereProjection(
-        const FReal Dt,
-        const FPBDJointSolverSettings& SolverSettings,
-        const FPBDJointSettings& JointSettings,
-        const FReal Alpha,
-        FVec3& NetDP1,
-        FVec3& NetDR1)
-    {
-        const FReal ProjectionPositionTolerance = 0.0f;//PositionTolerance;
-
-        FVec3 Axis;
-        FReal Delta;
-        FPBDJointUtilities::GetSphericalAxisDelta(ConnectorXs[0], ConnectorXs[1], Axis, Delta);
-        const FReal Error = FMath::Max((FReal)0, Delta - JointSettings.LinearLimit);
-        if (FMath::Abs(Error) > ProjectionPositionTolerance)
-        {
-            const FVec3 AngularAxis1 = FVec3::CrossProduct(ConnectorXs[1] - P(1), Axis);
-            const FVec3 IA1 = Utilities::Multiply(InvI(1), AngularAxis1);
-            const FReal II1 = FVec3::DotProduct(AngularAxis1, IA1);
-            const FReal IM = InvM(1) + II1;
-            const FVec3 DX = Axis * Error / IM;
-
-            const FVec3 DP1 = -Alpha * InvM(1) * DX;
-            const FVec3 DR1 = -Alpha * Utilities::Multiply(InvI(1), FVec3::CrossProduct(ConnectorXs[1] - P(1), DX));
-            ApplyDelta(1, DP1, DR1);
-
-            NetDP1 += DP1;
-            NetDR1 += DR1;
-        }
-    }
-
-    void FPBDJointCachedSolver::ApplyTranslateProjection(
-        const FReal Dt,
-        const FPBDJointSolverSettings& SolverSettings,
-        const FPBDJointSettings& JointSettings,
-        const FReal Alpha,
-        FVec3& NetDP1,
-        FVec3& NetDR1)
-    {
-        const FReal ProjectionPositionTolerance = 0.0f;//PositionTolerance;
-
-        FVec3 Axis;
-        FReal Delta;
-        FPBDJointUtilities::GetSphericalAxisDelta(ConnectorXs[0], ConnectorXs[1], Axis, Delta);
-        const FReal Error = FMath::Max((FReal)0, Delta - JointSettings.LinearLimit);
-        if (Error > ProjectionPositionTolerance)
-        {
-            const FVec3 DP1 = -Alpha * Error * Axis;
-            ApplyPositionDelta(1, DP1);
-            
-            NetDP1 += DP1;
-        }
-    }
-
-    void FPBDJointCachedSolver::ApplyConeProjection(
-        const FReal Dt,
-        const FPBDJointSolverSettings& SolverSettings,
-        const FPBDJointSettings& JointSettings,
-        const FReal Alpha,
-        const bool bPositionLocked,
-        FVec3& NetDP1,
-        FVec3& NetDR1)
-    {
-        const FReal ProjectionAngleTolerance = 0.0f;//AngleTolerance;
-
-        FVec3 SwingAxisLocal;
-        FReal DSwingAngle = 0.0f;
-        const FReal Swing1Limit = JointSettings.AngularLimits[(int32)EJointAngularConstraintIndex::Swing1];
-        const FReal Swing2Limit = JointSettings.AngularLimits[(int32)EJointAngularConstraintIndex::Swing2];
-        FPBDJointUtilities::GetEllipticalConeAxisErrorLocal(ConnectorRs[0], ConnectorRs[1], Swing2Limit, Swing1Limit, SwingAxisLocal, DSwingAngle);
-        FVec3 SwingAxis = ConnectorRs[0] * SwingAxisLocal;
-        if (DSwingAngle > ProjectionAngleTolerance)
-        {
-            const FVec3 DR1 = -Alpha * DSwingAngle * SwingAxis;
-            FVec3 DP1 = FVec3(0);
-            if (bPositionLocked)
-            {
-                DP1 = -Alpha * FVec3::CrossProduct(DR1, ConnectorXs[1] - P(1));
-            }
-            ApplyDelta(1, DP1, DR1);
-
-            NetDP1 += DP1;
-            NetDR1 += DR1;
-        }
-    }
-
-    void FPBDJointCachedSolver::ApplySwingProjection(
-        const FReal Dt,
-        const FPBDJointSolverSettings& SolverSettings,
-        const FPBDJointSettings& JointSettings,
-        const EJointAngularConstraintIndex SwingConstraintIndex,
-        const FReal Alpha,
-        const bool bPositionLocked,
-        FVec3& NetDP1,
-        FVec3& NetDR1)
-    {
-        const FReal ProjectionAngleTolerance = 0.0f;//AngleTolerance;
-
-        FVec3 SwingAxis;
-        FReal SwingAngle;
-        FPBDJointUtilities::GetSwingAxisAngle(ConnectorRs[0], ConnectorRs[1], SolverSettings.SwingTwistAngleTolerance, SwingConstraintIndex, SwingAxis, SwingAngle);
-
-        // Calculate swing error we need to correct
-        FReal DSwingAngle = 0;
-        const FReal SwingAngleMax = JointSettings.AngularLimits[(int32)SwingConstraintIndex];
-        if (SwingAngle > SwingAngleMax)
-        {
-            DSwingAngle = SwingAngle - SwingAngleMax;
-        }
-        else if (SwingAngle < -SwingAngleMax)
-        {
-            DSwingAngle = SwingAngle + SwingAngleMax;
-        }
-
-        UE_LOG(LogChaosJoint, VeryVerbose, TEXT("    Swing%d Angle %f [Limit %f]"), (SwingConstraintIndex == EJointAngularConstraintIndex::Swing1) ? 1 : 2, FMath::RadiansToDegrees(SwingAngle), FMath::RadiansToDegrees(SwingAngleMax));
-
-        // Apply swing correction
-        if (FMath::Abs(DSwingAngle) > ProjectionAngleTolerance)
-        {
-            const FVec3 DR1 = -Alpha * DSwingAngle * SwingAxis;
-            FVec3 DP1 = FVec3(0);
-            if (bPositionLocked)
-            {
-                DP1 = -Alpha * FVec3::CrossProduct(DR1, ConnectorXs[1] - P(1));
-            }
-            ApplyDelta(1, DP1, DR1);
-
-            NetDP1 += DP1;
-            NetDR1 += DR1;
-        }
-    }
-
-    void FPBDJointCachedSolver::ApplySingleLockedSwingProjection(
-        const FReal Dt,
-        const FPBDJointSolverSettings& SolverSettings,
-        const FPBDJointSettings& JointSettings,
-        const EJointAngularConstraintIndex SwingConstraintIndex,
-        const FReal Alpha,
-        const bool bPositionLocked,
-        FVec3& NetDP1,
-        FVec3& NetDR1)
-    {
-        ApplySwingProjection(Dt, SolverSettings, JointSettings, SwingConstraintIndex, Alpha, bPositionLocked, NetDP1, NetDR1);
-    }
-
-    void FPBDJointCachedSolver::ApplyDoubleLockedSwingProjection(
-        const FReal Dt,
-        const FPBDJointSolverSettings& SolverSettings,
-        const FPBDJointSettings& JointSettings,
-        const FReal Alpha,
-        const bool bPositionLocked,
-        FVec3& NetDP1,
-        FVec3& NetDR1)
-    {
-        const FReal ProjectionAngleTolerance = 0.0f;//AngleTolerance;
-
-        FVec3 SwingAxisLocal;
-        FReal DSwingAngle = 0.0f;
-        FPBDJointUtilities::GetCircularConeAxisErrorLocal(ConnectorRs[0], ConnectorRs[1], 0.0f, SwingAxisLocal, DSwingAngle);
-        FVec3 SwingAxis = ConnectorRs[0] * SwingAxisLocal;
-        if (DSwingAngle > ProjectionAngleTolerance)
-        {
-            const FVec3 DR1 = -Alpha * DSwingAngle * SwingAxis;
-            FVec3 DP1 = FVec3(0);
-            if (bPositionLocked)
-            {
-                DP1 = -Alpha * FVec3::CrossProduct(DR1, ConnectorXs[1] - P(1));
-            }
-            ApplyDelta(1, DP1, DR1);
-
-            NetDP1 += DP1;
-            NetDR1 += DR1;
-        }
-    }
-
-    void FPBDJointCachedSolver::ApplyDualConeSwingProjection(
-        const FReal Dt,
-        const FPBDJointSolverSettings& SolverSettings,
-        const FPBDJointSettings& JointSettings,
-        const EJointAngularConstraintIndex SwingConstraintIndex,
-        const FReal Alpha,
-        const bool bPositionLocked,
-        FVec3& NetDP1,
-        FVec3& NetDR1)
-    {
-        const FReal ProjectionAngleTolerance = 0.0f;//AngleTolerance;
-
-        FVec3 SwingAxis;
-        FReal SwingAngle;
-        FPBDJointUtilities::GetDualConeSwingAxisAngle(ConnectorRs[0], ConnectorRs[1], SwingConstraintIndex, SwingAxis, SwingAngle);
-
-        // Calculate swing error we need to correct
-        FReal DSwingAngle = 0;
-        const FReal SwingAngleMax = JointSettings.AngularLimits[(int32)SwingConstraintIndex];
-        if (SwingAngle > SwingAngleMax)
-        {
-            DSwingAngle = SwingAngle - SwingAngleMax;
-        }
-        else if (SwingAngle < -SwingAngleMax)
-        {
-            DSwingAngle = SwingAngle + SwingAngleMax;
-        }
-
-        // Apply swing correction
-        if (FMath::Abs(DSwingAngle) > ProjectionAngleTolerance)
-        {
-            const FVec3 DR1 = -Alpha * DSwingAngle * SwingAxis;
-            FVec3 DP1 = FVec3(0);
-            if (bPositionLocked)
-            {
-                DP1 = -Alpha * FVec3::CrossProduct(DR1, ConnectorXs[1] - P(1));
-            }
-            ApplyDelta(1, DP1, DR1);
-
-            NetDP1 += DP1;
-            NetDR1 += DR1;
-        }
-    }
-
-    void FPBDJointCachedSolver::ApplyTwistProjection(
-        const FReal Dt,
-        const FPBDJointSolverSettings& SolverSettings,
-        const FPBDJointSettings& JointSettings,
-        const FReal Alpha,
-        const bool bPositionLocked,
-        FVec3& NetDP1,
-        FVec3& NetDR1)
-    {
-        const FReal ProjectionAngleTolerance = 0.0f;//AngleTolerance;
-
-        FVec3 TwistAxis;
-        FReal TwistAngle;
-        FPBDJointUtilities::GetTwistAxisAngle(ConnectorRs[0], ConnectorRs[1], TwistAxis, TwistAngle);
-        FReal DTwistAngle = 0;
-        const FReal TwistLimit = JointSettings.AngularLimits[(int32)EJointAngularConstraintIndex::Twist];
-        if (TwistAngle > TwistLimit)
-        {
-            DTwistAngle = TwistAngle - TwistLimit;
-        }
-        else if (TwistAngle < -TwistLimit)
-        {
-            DTwistAngle = TwistAngle + TwistLimit;
-        }
-
-        if (FMath::Abs(DTwistAngle) > ProjectionAngleTolerance)
-        {
-            const FVec3 DR1 = -Alpha * DTwistAngle * TwistAxis;
-            FVec3 DP1 = FVec3(0);
-            if (bPositionLocked)
-            {
-                DP1 = -Alpha * FVec3::CrossProduct(DR1, ConnectorXs[1] - P(1));
-            }
-            ApplyDelta(1, DP1, DR1);
-
-            NetDP1 += DP1;
-            NetDR1 += DR1;
-        }
-    }
-
-    void FPBDJointCachedSolver::ApplyVelocityProjection(
-        const FReal Dt,
-        const FPBDJointSolverSettings& SolverSettings,
-        const FPBDJointSettings& JointSettings,
-        const FReal Alpha,
-        const FVec3& DP1,
-        const FVec3& DR1)
-    {
-        const FVec3 DV1 = Alpha * DP1 / Dt;
-        const FVec3 DW1 = Alpha * DR1 / Dt;
-        ApplyVelocityDelta(1, DV1, DW1);
-    }
-
-    void FPBDJointCachedSolver::ApplyLinearVelocityConstraint(
-        const FReal Stiffness,
-        const FVec3& Axis,
-        const FReal IM,
-        const FVec3& Arm0,
-        const FVec3& Arm1,
-        const FVec3& DRAxis0,
-        const FVec3& DRAxis1,
-        const FReal TargetVel)
-    {
-        const FVec3 CV0 = V(0) + FVec3::CrossProduct(W(0), Arm0);
-        const FVec3 CV1 = V(1) + FVec3::CrossProduct(W(1), Arm1);
-        const FVec3 CV = CV1 - CV0;
-        const FReal ImpulseSize = Stiffness * (FVec3::DotProduct(CV, Axis) - TargetVel) / IM;
-        const FVec3 Impulse = ImpulseSize * Axis;
-
-        const FVec3 DV0 = InvM(0) * Impulse;
-        const FVec3 DV1 = -InvM(1) * Impulse;
-        const FVec3 DW0 = DRAxis0 * ImpulseSize;
-        const FVec3 DW1 = DRAxis1 * ImpulseSize;
-        ApplyVelocityDelta(DV0, DW0, DV1, DW1);
-    }
-
-    void FPBDJointCachedSolver::ApplyPointVelocityConstraint(
-        const FReal Stiffness,
-        const FMatrix33& ConstraintIM,
-        const FVec3& ConstraintArm0,
-        const FVec3& ConstraintArm1)
-    {
-        const FVec3 CV0 = V(0) + FVec3::CrossProduct(W(0), ConstraintArm0);
-        const FVec3 CV1 = V(1) + FVec3::CrossProduct(W(1), ConstraintArm1);
-        const FVec3 CV = CV1 - CV0;
-        const FVec3 Impulse = Stiffness * Utilities::Multiply(ConstraintIM, CV);
-
-        const FVec3 DV0 = InvM(0) * Impulse;
-        const FVec3 DV1 = -InvM(1) * Impulse;
-        const FVec3 DW0 = Utilities::Multiply(InvI(0), FVec3::CrossProduct(ConstraintArm0, Impulse));
-        const FVec3 DW1 = Utilities::Multiply(InvI(1), FVec3::CrossProduct(ConstraintArm1, -Impulse));
-
-        ApplyVelocityDelta(DV0, DW0, DV1, DW1);
-    }
-
-    void FPBDJointCachedSolver::ApplyPointVelocityConstraint(
-        const FReal Dt,
-        const FPBDJointSolverSettings& SolverSettings,
-        const FPBDJointSettings& JointSettings)
-    {
-        if (!NetLinearImpulse.IsNearlyZero() && !FPBDJointUtilities::GetSoftLinearLimitEnabled(SolverSettings, JointSettings))
-        {
-            const FVec3& Arm0 = ConstraintArms[PointPositionConstraintIndex][0];
-            const FVec3& Arm1 = ConstraintArms[PointPositionConstraintIndex][1];
-			ApplyPointVelocityConstraint(LinearStiffness, PointPositionConstraintIM, Arm0, Arm1);
-        }
-    }
-
-    void FPBDJointCachedSolver::ApplySphericalVelocityConstraint(
-        const FReal Dt,
-        const FPBDJointSolverSettings& SolverSettings,
-        const FPBDJointSettings& JointSettings)
-    {
-        if(!NetLinearImpulse.IsNearlyZero() && !FPBDJointUtilities::GetSoftLinearLimitEnabled(SolverSettings, JointSettings))
-        {
-            const FVec3 Axis = PositionConstraintAxis[SphericalPositionConstraintIndex];
-            FReal TargetVel = 0.0f;
-            if (JointSettings.LinearRestitution != 0.0f)
-            {
-                const FReal InitVel = InitConstraintAxisLinearVelocities[0];
-                TargetVel = InitVel > Chaos_Joint_LinearVelocityThresholdToApplyRestitution ? -JointSettings.LinearRestitution * InitVel : 0.0f; 
-            }
-            const FReal& IM = PositionConstraintIM[SphericalPositionConstraintIndex];
-            const FVec3& Arm0 = ConstraintArms[SphericalPositionConstraintIndex][0];
-            const FVec3& Arm1 = ConstraintArms[SphericalPositionConstraintIndex][1];
-            const FVec3& DRAxis0 = PositionConstraintDRAxis[SphericalPositionConstraintIndex][0];
-            const FVec3& DRAxis1 = PositionConstraintDRAxis[SphericalPositionConstraintIndex][1];
-            ApplyLinearVelocityConstraint(LinearStiffness, Axis, IM, Arm0, Arm1, DRAxis0, DRAxis1, TargetVel);
-        }
-    }
-
-    void FPBDJointCachedSolver::ApplyCylindricalVelocityConstraint(
-        const FReal Dt,
-        const int32 AxisIndex,
-        const EJointMotionType AxialMotion,
-        const EJointMotionType RadialMotion,
-        const FPBDJointSolverSettings& SolverSettings,
-        const FPBDJointSettings& JointSettings)
-    {
-        if(!NetLinearImpulse.IsNearlyZero() && !FPBDJointUtilities::GetSoftLinearLimitEnabled(SolverSettings, JointSettings))
-        {
-            if (LinearHardLambda[(int32)EJointCylindricalPositionConstraintType::Axial] > SMALL_NUMBER)
-            {
-                const FVec3 Axis = PositionConstraintAxis[(int32)EJointCylindricalPositionConstraintType::Axial];
-                FReal TargetVel = 0.0f;
-                if (AxialMotion == EJointMotionType::Limited && JointSettings.LinearRestitution != 0.0f)
-                {
-                    const FReal InitVel = InitConstraintAxisLinearVelocities[(int32)EJointCylindricalPositionConstraintType::Axial];
-                    TargetVel = InitVel > Chaos_Joint_LinearVelocityThresholdToApplyRestitution ? -JointSettings.LinearRestitution * InitVel : 0.0f; 
-                }
-                const FReal& IM = PositionConstraintIM[(int32)EJointCylindricalPositionConstraintType::Axial];
-                const FVec3& Arm0 = ConstraintArms[(int32)EJointCylindricalPositionConstraintType::Axial][0];
-                const FVec3& Arm1 = ConstraintArms[(int32)EJointCylindricalPositionConstraintType::Axial][1];
-                const FVec3& DRAxis0 = PositionConstraintDRAxis[(int32)EJointCylindricalPositionConstraintType::Axial][0];
-                const FVec3& DRAxis1 = PositionConstraintDRAxis[(int32)EJointCylindricalPositionConstraintType::Axial][1];
-                ApplyLinearVelocityConstraint(LinearStiffness, Axis, IM, Arm0, Arm1, DRAxis0, DRAxis1, TargetVel);
-            }
-
-            if (LinearHardLambda[(int32)EJointCylindricalPositionConstraintType::Radial] > SMALL_NUMBER)
-            {
-                const FVec3 RadialAxis = PositionConstraintAxis[(int32)EJointCylindricalPositionConstraintType::Radial];
-                FReal TargetVel = 0.0f;
-                if (RadialMotion == EJointMotionType::Limited && JointSettings.LinearRestitution != 0.0f)
-                {
-                    const FReal InitVel = InitConstraintAxisLinearVelocities[(int32)EJointCylindricalPositionConstraintType::Radial];
-                    TargetVel = InitVel > Chaos_Joint_LinearVelocityThresholdToApplyRestitution ? -JointSettings.LinearRestitution * InitVel : 0.0f; 
-                }
-                const FReal& IM = PositionConstraintIM[(int32)EJointCylindricalPositionConstraintType::Radial];
-                const FVec3& Arm0 = ConstraintArms[(int32)EJointCylindricalPositionConstraintType::Radial][0];
-                const FVec3& Arm1 = ConstraintArms[(int32)EJointCylindricalPositionConstraintType::Radial][1];
-                const FVec3& DRAxis0 = PositionConstraintDRAxis[(int32)EJointCylindricalPositionConstraintType::Radial][0];
-                const FVec3& DRAxis1 = PositionConstraintDRAxis[(int32)EJointCylindricalPositionConstraintType::Radial][1];
-                ApplyLinearVelocityConstraint(LinearStiffness, RadialAxis, IM, Arm0, Arm1, DRAxis0, DRAxis1, TargetVel);
-            }
-        }
-    }
-
-    void FPBDJointCachedSolver::ApplyPlanarVelocityConstraint(
-        const FReal Dt,
-        const int32 AxisIndex,
-        const EJointMotionType AxialMotion,
-        const FPBDJointSolverSettings& SolverSettings,
-        const FPBDJointSettings& JointSettings)
-    {
-        if(!NetLinearImpulse.IsNearlyZero() && !FPBDJointUtilities::GetSoftLinearLimitEnabled(SolverSettings, JointSettings))
-        {
-			FVec3 Axis = PositionConstraintAxis[PlanarPositionConstraintIndex];
-            FReal TargetVel = 0.0f;
-            if (AxialMotion == EJointMotionType::Limited && JointSettings.LinearRestitution != 0.0f)
-            {
-                const FReal InitVel = InitConstraintAxisLinearVelocities[0];
-                TargetVel = InitVel > Chaos_Joint_LinearVelocityThresholdToApplyRestitution ? -JointSettings.LinearRestitution * InitVel : 0.0f; 
-            }
-            const FReal& IM = PositionConstraintIM[PlanarPositionConstraintIndex];
-            const FVec3& Arm0 = ConstraintArms[PlanarPositionConstraintIndex][0];
-            const FVec3& Arm1 = ConstraintArms[PlanarPositionConstraintIndex][1];
-            const FVec3& DRAxis0 = PositionConstraintDRAxis[PlanarPositionConstraintIndex][0];
-            const FVec3& DRAxis1 = PositionConstraintDRAxis[PlanarPositionConstraintIndex][1];
-            ApplyLinearVelocityConstraint(LinearStiffness, Axis, IM, Arm0, Arm1, DRAxis0, DRAxis1, TargetVel);
-        }
-    }
-
-    void FPBDJointCachedSolver::CalculateConstraintAxisLinearVelocities(
-        const FPBDJointSettings& JointSettings,
-        FVec3& ConstraintAxisLinearVelocities) const
-    {
-        bool bHasPositionConstraints =
-            (JointSettings.LinearMotionTypes[0] != EJointMotionType::Free)
-            || (JointSettings.LinearMotionTypes[1] != EJointMotionType::Free)
-            || (JointSettings.LinearMotionTypes[2] != EJointMotionType::Free);
-        if (!bHasPositionConstraints)
-        {
-            return;
-        }
-
-        const TVector<EJointMotionType, 3>& LinearMotion = JointSettings.LinearMotionTypes;
-        const TVector<bool, 3> bLinearLocked =
-        {
-            (LinearMotion[0] == EJointMotionType::Locked),
-            (LinearMotion[1] == EJointMotionType::Locked),
-            (LinearMotion[2] == EJointMotionType::Locked),
-        };
-        const TVector<bool, 3> bLinearLimited =
-        {
-            (LinearMotion[0] == EJointMotionType::Limited),
-            (LinearMotion[1] == EJointMotionType::Limited),
-            (LinearMotion[2] == EJointMotionType::Limited),
-        };
-
-        if (bLinearLocked[0] && bLinearLocked[1] && bLinearLocked[2])
-        {
-            
-        }
-        else if (bLinearLimited[0] && bLinearLimited[1] && bLinearLimited[2])
-        {
-            CalculateSphericalConstraintAxisLinearVelocities(JointSettings, ConstraintAxisLinearVelocities);
-        }
-        else if (bLinearLocked[1] && bLinearLocked[2] && !bLinearLocked[0])
-        {
-            // Line constraint along X axis
-            CalculateCylindricalConstraintAxisLinearVelocities(0, LinearMotion[0], EJointMotionType::Locked, JointSettings, ConstraintAxisLinearVelocities);
-        }
-        else if (bLinearLocked[0] && bLinearLocked[2] && !bLinearLocked[1])
-        {
-            // Line constraint along Y axis
-            CalculateCylindricalConstraintAxisLinearVelocities(1, LinearMotion[1], EJointMotionType::Locked, JointSettings, ConstraintAxisLinearVelocities);
-        }
-        else if (bLinearLocked[0] && bLinearLocked[1] && !bLinearLocked[2])
-        {
-            // Line constraint along Z axis
-            CalculateCylindricalConstraintAxisLinearVelocities(2, LinearMotion[2], EJointMotionType::Locked, JointSettings, ConstraintAxisLinearVelocities);
-        }
-        else if (bLinearLimited[1] && bLinearLimited[2] && !bLinearLimited[0])
-        {
-            // Cylindrical constraint along X axis
-            CalculateCylindricalConstraintAxisLinearVelocities(0, LinearMotion[0], EJointMotionType::Limited, JointSettings, ConstraintAxisLinearVelocities);
-        }
-        else if (bLinearLimited[0] && bLinearLimited[2] && !bLinearLimited[1])
-        {
-            // Cylindrical constraint along Y axis
-            CalculateCylindricalConstraintAxisLinearVelocities(1, LinearMotion[1], EJointMotionType::Limited, JointSettings, ConstraintAxisLinearVelocities);
-        }
-        else if (bLinearLimited[0] && bLinearLimited[1] && !bLinearLimited[2])
-        {
-            // Cylindrical constraint along Z axis
-            CalculateCylindricalConstraintAxisLinearVelocities(2, LinearMotion[2], EJointMotionType::Limited, JointSettings, ConstraintAxisLinearVelocities);
-        }
-        else if (bLinearLimited[0])
-        {
-            // Planar constraint along X axis
-            CalculatePlanarConstraintAxisLinearVelocities(0, LinearMotion[0], JointSettings, ConstraintAxisLinearVelocities);
-        }
-        else if (bLinearLimited[1])
-        {
-            // Planar constraint along Y axis
-            CalculatePlanarConstraintAxisLinearVelocities(1, LinearMotion[1], JointSettings, ConstraintAxisLinearVelocities);
-        }
-        else if (bLinearLimited[2])
-        {
-            // Planar constraint along Z axis
-            CalculatePlanarConstraintAxisLinearVelocities(2, LinearMotion[2], JointSettings, ConstraintAxisLinearVelocities);
-        }
-    }
-
-    void FPBDJointCachedSolver::CalculateSphericalConstraintAxisLinearVelocities(
-            const FPBDJointSettings& JointSettings,
-            FVec3& ConstraintAxisLinearVelocities) const
-    {
-        FVec3 Axis;
-        FReal Delta;
-        FPBDJointUtilities::GetSphericalAxisDelta(ConnectorXs[0], ConnectorXs[1], Axis, Delta);
-        const FVec3 DConnectorX = ConnectorXs[1] - ConnectorXs[0];
-        const FVec3 Connector0Correction = Delta > JointSettings.LinearLimit ? Axis * JointSettings.LinearLimit : DConnectorX;
-
-        const FVec3 Arm0 = ConnectorXs[0] + Connector0Correction - P(0);
-        const FVec3 Arm1 = ConnectorXs[1] - P(1);
-        const FVec3 CV0 = V(0) + FVec3::CrossProduct(W(0), Arm0);
-        const FVec3 CV1 = V(1) + FVec3::CrossProduct(W(1), Arm1);
-        const FVec3 CV = CV1 - CV0;
-
-        ConstraintAxisLinearVelocities[0] = FVec3::DotProduct(CV, Axis);
-    }
-
-    void FPBDJointCachedSolver::CalculateCylindricalConstraintAxisLinearVelocities(
-            const int32 AxisIndex,
-            const EJointMotionType AxialMotion,
-            const EJointMotionType RadialMotion,
-            const FPBDJointSettings& JointSettings,
-            FVec3& ConstraintAxisLinearVelocities) const
-    {
-        FVec3 Axis, RadialAxis;
-        FReal AxialDelta, RadialDelta;
-        FPBDJointUtilities::GetCylindricalAxesDeltas(ConnectorRs[0], ConnectorXs[0], ConnectorXs[1], AxisIndex, Axis, AxialDelta, RadialAxis, RadialDelta);
-        const FVec3 DConnectorX = ConnectorXs[1] - ConnectorXs[0];
-
-        if (AxialMotion == EJointMotionType::Limited)
-        {
-            if (AxialDelta < 0.0f)
-            {
-                AxialDelta = -AxialDelta;
-                Axis = -Axis;
-            }
-            const FReal AxialLimit = (AxialMotion == EJointMotionType::Locked) ? 0.0f :JointSettings.LinearLimit;
-            const FReal AxialError = AxialDelta - AxialLimit;
-            
-            const FVec3 Connector0Correction = DConnectorX - AxialError * Axis;
-
-            const FVec3 Arm0 = ConnectorXs[0] + Connector0Correction - P(0);
-            const FVec3 Arm1 = ConnectorXs[1] - P(1);
-            const FVec3 CV0 = V(0) + FVec3::CrossProduct(W(0), Arm0);
-            const FVec3 CV1 = V(1) + FVec3::CrossProduct(W(1), Arm1);
-            const FVec3 CV = CV1 - CV0;
-
-            ConstraintAxisLinearVelocities[(int32)EJointCylindricalPositionConstraintType::Axial] = FVec3::DotProduct(CV, Axis);
-        }
-
-        if (RadialMotion == EJointMotionType::Limited)
-        {
-            const FReal RadialLimit = (RadialMotion == EJointMotionType::Locked) ? 0.0f : JointSettings.LinearLimit;
-            const FReal RadialError = FMath::Max(RadialDelta - RadialLimit, 0.0f);
-            const FVec3 Connector0Correction = DConnectorX - RadialError * RadialAxis;
-            
-            const FVec3 Arm0 = ConnectorXs[0] + Connector0Correction - P(0);
-            const FVec3 Arm1 = ConnectorXs[1] - P(1);
-            const FVec3 CV0 = V(0) + FVec3::CrossProduct(W(0), Arm0);
-            const FVec3 CV1 = V(1) + FVec3::CrossProduct(W(1), Arm1);
-            const FVec3 CV = CV1 - CV0;
-
-            ConstraintAxisLinearVelocities[(int32)EJointCylindricalPositionConstraintType::Radial] = FVec3::DotProduct(CV, RadialAxis);
-        }
-    }
-
-    void FPBDJointCachedSolver::CalculatePlanarConstraintAxisLinearVelocities(
-            const int32 AxisIndex,
-            const EJointMotionType AxialMotion,
-            const FPBDJointSettings& JointSettings,
-            FVec3& ConstraintAxisLinearVelocities) const
-    {
-        FVec3 Axis;
-        FReal Delta;
-        FPBDJointUtilities::GetPlanarAxisDelta(ConnectorRs[0], ConnectorXs[0], ConnectorXs[1], AxisIndex, Axis, Delta);
-
-        if (Delta < 0.0f)
-        {
-            Delta = -Delta;
-            Axis = -Axis;
-        }
-
-        const FReal Limit = (AxialMotion == EJointMotionType::Locked) ? 0 : JointSettings.LinearLimit;
-        const FReal Error = Delta - Limit;
-
-        const FVec3 DConnectorX = ConnectorXs[1] - ConnectorXs[0];
-        const FVec3 Connector0Correction = DConnectorX - Axis * Error; 
-
-        const FVec3 Arm0 = ConnectorXs[0] + Connector0Correction - P(0);
-        const FVec3 Arm1 = ConnectorXs[1] - P(1);
-        const FVec3 CV0 = V(0) + FVec3::CrossProduct(W(0), Arm0);
-        const FVec3 CV1 = V(1) + FVec3::CrossProduct(W(1), Arm1);
-        const FVec3 CV = CV1 - CV0;
-
-        ConstraintAxisLinearVelocities[0] = FVec3::DotProduct(CV, Axis);
-    }
-
-    void FPBDJointCachedSolver::CalculateConstraintAxisAngularVelocities(
-        const FPBDJointSolverSettings& SolverSettings,
-        const FPBDJointSettings& JointSettings,
-        FVec3& ConstraintAxisAngularVelocities
-        ) const
-    {
-        bool bHasRotationConstraints =
-            (JointSettings.AngularMotionTypes[0] != EJointMotionType::Free)
-            || (JointSettings.AngularMotionTypes[1] != EJointMotionType::Free)
-            || (JointSettings.AngularMotionTypes[2] != EJointMotionType::Free);
-        if (!bHasRotationConstraints)
-        {
-            return;
-        }
-
-        // Locked axes always use hard constraints. Limited axes use hard or soft depending on settings
-        EJointMotionType TwistMotion = JointSettings.AngularMotionTypes[(int32)EJointAngularConstraintIndex::Twist];
-        EJointMotionType Swing1Motion = JointSettings.AngularMotionTypes[(int32)EJointAngularConstraintIndex::Swing1];
-        EJointMotionType Swing2Motion = JointSettings.AngularMotionTypes[(int32)EJointAngularConstraintIndex::Swing2];
-
-        // If the twist axes are opposing, we cannot decompose the orientation into swing and twist angles, so just give up
-        const FVec3 Twist0 = ConnectorRs[0] * FJointConstants::TwistAxis();
-        const FVec3 Twist1 = ConnectorRs[1] * FJointConstants::TwistAxis();
-        const FReal Twist01Dot = FVec3::DotProduct(Twist0, Twist1);
-        const bool bDegenerate = (Twist01Dot < Chaos_Joint_DegenerateRotationLimit);
-        if (bDegenerate)
-        {
-            UE_LOG(LogChaosJoint, VeryVerbose, TEXT(" Degenerate rotation at Swing %f deg"), FMath::RadiansToDegrees(FMath::Acos(Twist01Dot)));
-        }
-
-        // Apply twist constraint
-        // NOTE: Cannot calculate twist angle at 180degree swing
-        if (SolverSettings.bEnableTwistLimits && !bDegenerate)
-        {
-            if (TwistMotion == EJointMotionType::Limited)
-            {
-                CalculateTwistConstraintAxisAngularVelocities(JointSettings, ConstraintAxisAngularVelocities);
-            }
-        }
-
-        // Apply swing constraints
-        // NOTE: Cannot separate swing angles at 180degree swing (but we can still apply locks)
-        if (SolverSettings.bEnableSwingLimits)
-        {
-            if ((Swing1Motion == EJointMotionType::Limited) && (Swing2Motion == EJointMotionType::Limited))
-            {
-                CalculateConeConstraintAxisAngularVelocities(JointSettings, ConstraintAxisAngularVelocities);
-            }
-            else if ((Swing1Motion == EJointMotionType::Limited) && (Swing2Motion == EJointMotionType::Locked))
-            {
-                if (!bDegenerate)
-                {
-                    CalculateSwingConstraintAxisAngularVelocities(SolverSettings, JointSettings, EJointAngularConstraintIndex::Swing1, ConstraintAxisAngularVelocities);
-                }
-            }
-            else if ((Swing1Motion == EJointMotionType::Limited) && (Swing2Motion == EJointMotionType::Free))
-            {
-                if (!bDegenerate)
-                {
-                    CalculateDualConeSwingConstraintAxisAngularVelocities(JointSettings, EJointAngularConstraintIndex::Swing1, ConstraintAxisAngularVelocities);
-                }
-            }
-            else if ((Swing1Motion == EJointMotionType::Locked) && (Swing2Motion == EJointMotionType::Limited))
-            {
-                if (!bDegenerate)
-                {
-                    CalculateSwingConstraintAxisAngularVelocities(SolverSettings, JointSettings, EJointAngularConstraintIndex::Swing2, ConstraintAxisAngularVelocities);
-                }
-            }
-            else if ((Swing1Motion == EJointMotionType::Locked) && (Swing2Motion == EJointMotionType::Locked))
-            {
-            }
-            else if ((Swing1Motion == EJointMotionType::Locked) && (Swing2Motion == EJointMotionType::Free))
-            {
-            }
-            else if ((Swing1Motion == EJointMotionType::Free) && (Swing2Motion == EJointMotionType::Limited))
-            {
-                if (!bDegenerate)
-                {
-                    CalculateDualConeSwingConstraintAxisAngularVelocities(JointSettings, EJointAngularConstraintIndex::Swing2, ConstraintAxisAngularVelocities);
-                }
-            }
-            else if ((Swing1Motion == EJointMotionType::Free) && (Swing2Motion == EJointMotionType::Locked))
-            {
-            }
-            else if ((Swing1Motion == EJointMotionType::Free) && (Swing2Motion == EJointMotionType::Free))
-            {
-            }
-        }
-    }
-
-    void FPBDJointCachedSolver::CalculateTwistConstraintAxisAngularVelocities(
-        const FPBDJointSettings& JointSettings,
-        FVec3& ConstraintAxisAngularVelocities) const
-    {
-        FVec3 TwistAxis;
-        FReal TwistAngle;
-        FPBDJointUtilities::GetTwistAxisAngle(ConnectorRs[0], ConnectorRs[1], TwistAxis, TwistAngle);
-
-        if (TwistAngle < 0.0f)
-        {
-            TwistAxis = -TwistAxis;
-        }
-        ConstraintAxisAngularVelocities[(int32)EJointAngularConstraintIndex::Twist] = FVec3::DotProduct(W(1) - W(0), TwistAxis);
-    }
-
-    void FPBDJointCachedSolver::CalculateConeConstraintAxisAngularVelocities(
-            const FPBDJointSettings& JointSettings,
-            FVec3& ConstraintAxisAngularVelocities) const
-    {
-        FVec3 SwingAxisLocal;
-        FReal DSwingAngle = 0.0f;
-        const FReal Swing1Limit = JointSettings.AngularLimits[(int32)EJointAngularConstraintIndex::Swing1];
-        const FReal Swing2Limit = JointSettings.AngularLimits[(int32)EJointAngularConstraintIndex::Swing2];
-        FPBDJointUtilities::GetEllipticalConeAxisErrorLocal(ConnectorRs[0], ConnectorRs[1], Swing2Limit, Swing1Limit, SwingAxisLocal, DSwingAngle);
-        SwingAxisLocal.SafeNormalize();
-        const FVec3 SwingAxis = ConnectorRs[0] * SwingAxisLocal;
-        ConstraintAxisAngularVelocities[(int32)EJointAngularConstraintIndex::Swing1] = FVec3::DotProduct(W(1) - W(0), SwingAxis);   
-    }
-
-    void FPBDJointCachedSolver::CalculateDualConeSwingConstraintAxisAngularVelocities(
-            const FPBDJointSettings& JointSettings,
-            const EJointAngularConstraintIndex SwingConstraintIndex,
-            FVec3& ConstraintAxisAngularVelocities) const
-    {
-        FVec3 SwingAxis;
-        FReal SwingAngle;
-        FPBDJointUtilities::GetDualConeSwingAxisAngle(ConnectorRs[0], ConnectorRs[1], SwingConstraintIndex, SwingAxis, SwingAngle);
-        ConstraintAxisAngularVelocities[(int32)SwingConstraintIndex] = FVec3::DotProduct(W(1) - W(0), SwingAxis);
-    }
-
-    void FPBDJointCachedSolver::CalculateSwingConstraintAxisAngularVelocities(
-            const FPBDJointSolverSettings& SolverSettings,
-            const FPBDJointSettings& JointSettings,
-            const EJointAngularConstraintIndex SwingConstraintIndex,
-            FVec3& ConstraintAxisAngularVelocities) const
-    {
-        FVec3 SwingAxis;
-        FReal SwingAngle;
-        FPBDJointUtilities::GetSwingAxisAngle(ConnectorRs[0], ConnectorRs[1], SolverSettings.SwingTwistAngleTolerance, SwingConstraintIndex, SwingAxis, SwingAngle);
-        ConstraintAxisAngularVelocities[(int32)SwingConstraintIndex] = FVec3::DotProduct(W(1) - W(0), SwingAxis);
-    }
+		const FVec3 V0Dt = FVec3::CalculateVelocity(InitConnectorXs[0], ConnectorXs[0]+Body(0).DP() + FVec3::CrossProduct(Body(0).DQ(), PositionConstraints.ConstraintArms[PointPositionConstraintIndex][0]), 1.0f);
+		const FVec3 V1Dt = FVec3::CalculateVelocity(InitConnectorXs[1], ConnectorXs[1]+Body(1).DP() + FVec3::CrossProduct(Body(1).DQ(), PositionConstraints.ConstraintArms[PointPositionConstraintIndex][1]), 1.0f);
+		const FReal AxisVDt = FVec3::DotProduct(V1Dt - V0Dt, PositionConstraints.ConstraintAxis[ConstraintIndex]);
+
+		// Calculate the padding to apply to the constraint that will result in the
+		// desired outward velocity (assuming the constraint is fully resolved)
+		const FReal Padding = (1.0f + Restitution) * AxisVDt - InOutPos;
+		if (Padding > 0.0f)
+		{
+			SetLinearConstraintPadding(ConstraintIndex, Padding);
+			InOutPos += Padding;
+		}
+	}
+}
+
+// Used for non-zero restitution. We pad constraints by an amount such that the velocity
+// calculated after solving constraint positions will as required for the restitution.
+void FPBDJointCachedSolver::CalculateAngularConstraintPadding(
+	const int32 ConstraintIndex,
+	const FReal Restitution,
+	FReal& InOutAngle)
+{
+	// NOTE: We only calculate the padding after the constraint is first violated, and after
+	// that the padding is fixed for the rest of the iterations in the current step.
+	if ((Restitution > 0.0f) && (InOutAngle > 0.0f) && !HasAngularConstraintPadding(ConstraintIndex))
+	{
+		SetAngularConstraintPadding(ConstraintIndex, 0.0f);
+
+		// Calculate the velocity we want to match
+		const FVec3 W0Dt = FRotation3::CalculateAngularVelocity(InitConnectorRs[0], ConnectorRs[0], 1.0f);
+		const FVec3 W1Dt = FRotation3::CalculateAngularVelocity(InitConnectorRs[1], ConnectorRs[1], 1.0f);
+		const FReal AxisWDt = FVec3::DotProduct(W1Dt - W0Dt, RotationConstraints.ConstraintAxis[(int32)ConstraintIndex]);
+
+		// Calculate the padding to apply to the constraint that will result in the
+		// desired outward velocity (assuming the constraint is fully resolved)
+		const FReal Padding = (1.0f + Restitution) * AxisWDt - InOutAngle;
+		if (Padding > 0.0f)
+		{
+			SetAngularConstraintPadding(ConstraintIndex, Padding);
+			InOutAngle += Padding;
+		}
+	}
 }
+
+void FAxisConstraintDatas::InitDatas(
+	const int32 ConstraintIndex,
+	const bool bHasSoftLimits,
+	const FReal SoftStiffness,
+	const FReal SoftDamping,
+	const FReal HardStiffness)
+{
+	bSoftLimit[ConstraintIndex] = bHasSoftLimits;
+	ConstraintSoftStiffness[ConstraintIndex] = SoftStiffness;
+	ConstraintHardStiffness[ConstraintIndex] = HardStiffness;
+	ConstraintDamping[ConstraintIndex] = bHasSoftLimits ? SoftDamping : 0.0;
+	bValidDatas[ConstraintIndex] = false;
+	bLimitsCheck[ConstraintIndex] = true;
+	ConstraintLambda = FVec3::Zero();
+	ConstraintLimits = FVec3::Zero();
+	MotionType[ConstraintIndex] = EJointMotionType::Free;
+}
+
+void FAxisConstraintDatas::UpdateDatas(
+	const int32 ConstraintIndex,
+	const FVec3& DatasAxis,
+	const FReal DatasCX,
+	const FVec3& DatasIA0,
+	const FVec3& DatasIA1,
+	const FReal DatasIM,
+	const FReal DatasRestitution,
+	const FReal Dt,
+	const bool bCheckLimit,
+	const FVec3& DatasArm0 ,
+	const FVec3& DatasArm1 ,
+	const FReal DatasVX)
+{
+	bValidDatas[ConstraintIndex] = true;
+	bLimitsCheck[ConstraintIndex] = bCheckLimit;
+
+	ConstraintCX[ConstraintIndex] = DatasCX;
+	ConstraintVX[ConstraintIndex] = DatasVX;
+	ConstraintAxis[ConstraintIndex] = DatasAxis;
+	ConstraintDRAxis[ConstraintIndex][0] = DatasIA0;
+	ConstraintDRAxis[ConstraintIndex][1] = -DatasIA1;
+	ConstraintHardIM[ConstraintIndex] = DatasIM;
+	ConstraintRestitution[ConstraintIndex] = DatasRestitution;
+	ConstraintArms[ConstraintIndex][0] = DatasArm0;
+	ConstraintArms[ConstraintIndex][1] = DatasArm1;
+
+	if(bSoftLimit[ConstraintIndex])
+	{
+		const FReal SpringMassScale = (bAccelerationMode) ? (FReal)1 / (ConstraintHardIM[ConstraintIndex]) : (FReal)1;
+		ConstraintSoftStiffness[ConstraintIndex] = SpringMassScale * ConstraintSoftStiffness[ConstraintIndex] * Dt * Dt;
+		ConstraintDamping[ConstraintIndex] = SpringMassScale * ConstraintDamping[ConstraintIndex] * Dt;
+		ConstraintSoftIM[ConstraintIndex] = (ConstraintSoftStiffness[ConstraintIndex] + ConstraintDamping[ConstraintIndex]) * ConstraintHardIM[ConstraintIndex] + (FReal)1;
+	}
+}
+
+void FPointConstraintDatas::InitDatas(
+	const FReal PointStiffness)
+{
+	ConstraintStiffness = PointStiffness;
+	bValidDatas= false;
+	ConstraintLambda = FVec3::Zero();
+}
+
+void FPointConstraintDatas::UpdateDatas(
+	const FVec3& DatasCX,
+	const FMatrix33& DatasIM,
+	const FVec3& DatasArm0,
+	const FVec3& DatasArm1)
+{
+	bValidDatas= true;
+	ConstraintCX = DatasCX;
+	ConstraintIM = DatasIM;
+	ConstraintArms[0] = DatasArm0;
+	ConstraintArms[1] = DatasArm1;
+}
+	
+}
+
+
diff --git a/Engine/Source/Runtime/Experimental/Chaos/Private/Chaos/Joint/PBDJointSolverGaussSeidel.cpp b/Engine/Source/Runtime/Experimental/Chaos/Private/Chaos/Joint/PBDJointSolverGaussSeidel.cpp
--- a/Engine/Source/Runtime/Experimental/Chaos/Private/Chaos/Joint/PBDJointSolverGaussSeidel.cpp
+++ b/Engine/Source/Runtime/Experimental/Chaos/Private/Chaos/Joint/PBDJointSolverGaussSeidel.cpp
@@ -2103,7 +2103,7 @@ namespace Chaos
 		}
 
 		const FReal RadialLimitPadding = GetLinearConstraintPadding(1);
-		const FReal RadialLimit = (RadialMotion == EJointMotionType::Locked) ? 0.0f : FMath::Max(JointSettings.LinearLimit - AxialLimitPadding, (FReal)0.);
+		const FReal RadialLimit = (RadialMotion == EJointMotionType::Locked) ? 0.0f : FMath::Max(JointSettings.LinearLimit - RadialLimitPadding, (FReal)0.);
 		FReal RadialError = RadialDelta - RadialLimit;
 
 		if (RadialError > PositionTolerance)
diff --git a/Engine/Source/Runtime/Experimental/Chaos/Public/Chaos/Joint/PBDJointCachedSolverGaussSeidel.h b/Engine/Source/Runtime/Experimental/Chaos/Public/Chaos/Joint/PBDJointCachedSolverGaussSeidel.h
--- a/Engine/Source/Runtime/Experimental/Chaos/Public/Chaos/Joint/PBDJointCachedSolverGaussSeidel.h
+++ b/Engine/Source/Runtime/Experimental/Chaos/Public/Chaos/Joint/PBDJointCachedSolverGaussSeidel.h
@@ -12,7 +12,78 @@
 
 namespace Chaos
 {
-	class FPBDJointCachedSolver;
+
+/** Cached axis joint datas that will be used during the apply */
+struct FAxisConstraintDatas
+{
+	/** Init the axis joint datas with stiffness / damping */
+	void InitDatas(
+		const int32 ConstraintIndex,
+		const bool bHasSoftLimits,
+		const FReal SoftStiffness,
+		const FReal SoftDamping,
+		const FReal HardStiffness);
+
+	/** Update the axis joint datas with axis, limits, arms... */
+	void UpdateDatas(
+		const int32 ConstraintIndex,
+		const FVec3& DatasAxis,
+		const FReal DatasCX,
+		const FVec3& DatasIA0,
+		const FVec3& DatasIA1,
+		const FReal DatasIM,
+		const FReal DatasRestitution,
+		const FReal Dt,
+		const bool bCheckLimit = true,
+		const FVec3& DatasArm0 = FVec3::Zero(),
+		const FVec3& DatasArm1 = FVec3::Zero(),
+		const FReal DatasVX = 0.0);
+	
+	FVec3 ConstraintSoftStiffness;
+	FVec3 ConstraintHardStiffness;
+	FVec3 ConstraintDamping;
+	FVec3 ConstraintArms[3][2];
+	FVec3 ConstraintAxis[3];
+	FVec3 ConstraintLimits;
+	
+	FVec3 ConstraintCX;
+	FVec3 ConstraintVX;
+	FVec3 ConstraintSoftIM;
+	FVec3 ConstraintHardIM;
+	FVec3 ConstraintDRAxis[3][2];
+	FVec3 ConstraintRestitution;
+	FVec3 ConstraintLambda;
+	
+	bool bValidDatas[3];
+	bool bLimitsCheck[3];
+	bool bAccelerationMode;
+	bool bSoftLimit[3];
+	EJointMotionType MotionType[3];
+};
+
+/** Cached point joint datas that will be used during the apply */
+struct FPointConstraintDatas
+{
+	/** Init the point joint datas with stiffness / damping */
+	void InitDatas(
+		const FReal PointStiffness);
+
+	/** Update the point joint datas with axis, limits, arms... */
+	void UpdateDatas(
+		const FVec3& DatasCX,
+		const FMatrix33& DatasIM,
+		const FVec3& DatasArm0,
+		const FVec3& DatasArm1);
+	
+	bool bValidDatas;
+	
+	FReal ConstraintStiffness;
+	FVec3 ConstraintArms[2];
+	
+	FVec3 ConstraintCX;
+	FMatrix33 ConstraintIM;
+	FVec3 ConstraintLambda;
+};
 
 	/**
 	 * Calculate new positions and rotations for a pair of bodies connected by a joint.
@@ -129,11 +200,10 @@ namespace Chaos
 		{
 			return bIsActive;
 		}
-
-
+		
 		FPBDJointCachedSolver();
 
-		// Called once per frame to initialze the joint solver from the joint settings
+		// Called once per frame to initialize the joint solver from the joint settings
 		void Init(
 			const FReal Dt,
 			const FSolverBodyPtrPair& SolverBodyPair,
@@ -142,42 +212,6 @@ namespace Chaos
 			const FRigidTransform3& XL0,
 			const FRigidTransform3& XL1);
 
-		void InitPositionConstraints(
-			const FPBDJointSolverSettings& SolverSettings,
-			const FPBDJointSettings& JointSettings);
-
-		void InitPointPositionConstraintKD(
-			const int32 KIndex,
-			const int32 DIndex);
-
-		void InitPointPositionConstraintDD();
-
-		void InitSphericalPositionConstraint(
-			const FPBDJointSettings& JointSettings);
-
-		void InitCylindricalPositionConstraint(
-		    const int32 AxisIndex,
-		    const EJointMotionType AxialMotion,
-		    const EJointMotionType RadialMotion,
-		    const FPBDJointSettings& JointSettings);
-
-		void InitPlanarPositionConstraint(
-			const int32 AxisIndex,
-			const EJointMotionType AxialMotion,
-			const FPBDJointSettings& JointSettings);
-
-		void InitPositionConstraint(
-			const FVec3& Axis,
-			const FVec3& ConstraintArm0,
-			const FVec3& ConstraintArm1,
-			const int32 ConstraintIndex);
-
-		void UpdatePositionConstraintIM(
-			const FVec3& Axis,
-			const FVec3& ConstraintArm0,
-			const FVec3& ConstraintArm1,
-			const int32 ConstraintIndex);
-
 		void Deinit();
 
 		// @todo(chaos): this doesn't do much now we have SolverBodies - remove it
@@ -206,7 +240,7 @@ namespace Chaos
 			const FReal Dt,
 			const FReal InSolverStiffness,
 			const FPBDJointSolverSettings& SolverSettings,
-			const FPBDJointSettings& JointSettings);
+			const FPBDJointSettings& JointSettings){}
 
 		void EnableProjection();
 
@@ -222,442 +256,248 @@ namespace Chaos
 		// Check to see if this constraint still needs further solved
 		// @todo(chaos): the term "active" is used inconsistently with the meaning elsewhere. Active should
 		// mean "contributed impusles". Should use "solved" rather than "active"
-		bool UpdateIsActive(
-			const FReal Dt,
-			const FPBDJointSolverSettings& SolverSettings,
-			const FPBDJointSettings& JointSettings);
-
-		void ApplyPositionConstraints(
-			const FReal Dt,
-			const FPBDJointSolverSettings& SolverSettings,
-			const FPBDJointSettings& JointSettings);
-
-		void ApplyRotationConstraints(
-			const FReal Dt,
-			const FPBDJointSolverSettings& SolverSettings,
-			const FPBDJointSettings& JointSettings);
-
-		void ApplyPositionDrives(
-			const FReal Dt,
-			const FPBDJointSolverSettings& SolverSettings,
-			const FPBDJointSettings& JointSettings);
-
-		void ApplyRotationDrives(
-			const FReal Dt,
-			const FPBDJointSolverSettings& SolverSettings,
-			const FPBDJointSettings& JointSettings);
-
-		void ApplyLinearVelocityConstraints(
-			const FReal Dt,
-			const FPBDJointSolverSettings& SolverSettings,
-			const FPBDJointSettings& JointSettings);
-
-		void ApplyAngularVelocityConstraints(
-			const FReal Dt,
-			const FPBDJointSolverSettings& SolverSettings,
-			const FPBDJointSettings& JointSettings);
+		bool UpdateIsActive();
 
+		/** Common function to apply the lagrange multiplier update */
+		
 		void ApplyPositionDelta(
 			const int32 BodyIndex,
 			const FVec3& DP);
 
-		void ApplyPositionDelta(
-			const FVec3& DP0,
-			const FVec3& DP1);
-
 		void ApplyRotationDelta(
 			const int32 BodyIndex,
 			const FVec3& DR);
 
-		void ApplyRotationDelta(
-			const FVec3& DR0,
-			const FVec3& DR1);
-
-		void ApplyDelta(
-			const int32 BodyIndex,
-			const FVec3& DP,
-			const FVec3& DR1);
-
-		void ApplyVelocityDelta(
-			const int32 BodyIndex,
-			const FVec3& DV,
-			const FVec3& DW);
-
-		void ApplyVelocityDelta(
-			const FVec3& DV0,
-			const FVec3& DW0,
-			const FVec3& DV1,
-			const FVec3& DW1);
-
-		void ApplyAngularVelocityDelta(
-			const FVec3& DW0,
-			const FVec3& DW1);
-
-		void ApplyPositionConstraint(
-			const FReal Stiffness,
-			const FVec3& Axis,
-			const FReal Delta,
-			const FReal IM,
-			const FVec3& Arm0,
-			const FVec3& Arm1,
-			const FVec3& DRAxis0,
-			const FVec3& DRAxis1,
-			const int32 LinearHardLambdaIndex = -1);
-
-		void ApplyPositionConstraintSoft(
-			const FReal Dt,
-			const FReal JointStiffness,
-			const FReal JointDamping,
-			const bool bAccelerationMode,
-			const FVec3& Axis,
-			const FReal Delta,
-			const FReal TargetVel,
-			FReal& Lambda);
-
-		void ApplyRotationConstraint(
-			const FReal JointStiffness,
-			const FVec3& Axis,
-			const FReal Angle,
-			const int32 AngularHardLambdaIndex = -1);
-
-		void ApplyRotationConstraintKD(
-			const int32 KIndex,
-			const int32 DIndex,
-			const FReal JointStiffness,
-			const FVec3& Axis,
-			const FReal Angle,
-			const int32 AngularHardLambdaIndex = -1);
-
-		void ApplyRotationConstraintDD(
-			const FReal JointStiffness,
-			const FVec3& Axis,
-			const FReal Angle,
-			const int32 AngularHardLambdaIndex = -1);
-
-		void ApplyRotationConstraintSoft(
-			const FReal Dt,
-			const FReal JointStiffness,
-			const FReal JointDamping,
-			const bool bAccelerationMode,
-			const FVec3& Axis,
-			const FReal Angle,
-			const FReal AngVelTarget,
-			FReal& Lambda);
-
-		void ApplyRotationConstraintSoftKD(
-			const int32 KIndex,
-			const int32 DIndex,
-			const FReal Dt,
-			const FReal JointStiffness,
-			const FReal JointDamping,
-			const bool bAccelerationMode,
-			const FVec3& Axis,
-			const FReal Angle,
-			const FReal AngVelTarget,
-			FReal& Lambda);
-
-		void ApplyRotationConstraintSoftDD(
-			const FReal Dt,
-			const FReal JointStiffness,
-			const FReal JointDamping,
-			const bool bAccelerationMode,
-			const FVec3& Axis,
-			const FReal Angle,
-			const FReal AngVelTarget,
-			FReal& Lambda);
-
-		void ApplyLockedRotationConstraints(
-			const FReal Dt,
-			const FPBDJointSolverSettings& SolverSettings,
-			const FPBDJointSettings& JointSettings,
-			const bool bApplyTwist,
-			const bool bApplySwing);
-
-		void ApplyTwistConstraint(
+		/** Init Position constraints */
+		
+		void InitPositionConstraints(
 			const FReal Dt,
 			const FPBDJointSolverSettings& SolverSettings,
-			const FPBDJointSettings& JointSettings,
-			const bool bUseSoftLimit);
+			const FPBDJointSettings& JointSettings);
+		
+		void InitPointPositionConstraint();
+		
+		void InitPositionConstraintDatas(
+			const int32 ConstraintIndex,
+			const FVec3& ConstraintAxis,
+			const FReal& ConstraintDelta,
+			const FReal ConstraintRestitution,
+			const FReal Dt,
+			const FReal ConstraintLimit,
+			const EJointMotionType JointType);
 
-		void ApplyConeConstraint(
-			const FReal Dt,
-			const FPBDJointSolverSettings& SolverSettings,
+		void InitSphericalPositionConstraint(
 			const FPBDJointSettings& JointSettings,
-			const bool bUseSoftLimit);
+			const FReal Dt);
 
-		// One Swing axis is free, and the other locked. This applies the lock: Body1 Twist axis is confined to a plane.
-		void ApplySingleLockedSwingConstraint(
-			const FReal Dt,
-			const FPBDJointSolverSettings& SolverSettings,
+		void InitCylindricalPositionConstraint(
 			const FPBDJointSettings& JointSettings,
-			const EJointAngularConstraintIndex SwingConstraintIndex,
-			const bool bUseSoftLimit);
-
-		// One Swing axis is free, and the other limited. This applies the limit: Body1 Twist axis is confined to space between two cones.
-		void ApplyDualConeSwingConstraint(
 			const FReal Dt,
-			const FPBDJointSolverSettings& SolverSettings,
-			const FPBDJointSettings& JointSettings,
-			const EJointAngularConstraintIndex SwingConstraintIndex,
-			const bool bUseSoftLimit);
+			const int32 AxisIndex,
+			const EJointMotionType AxialMotion,
+			const EJointMotionType RadialMotion);
 
-		// One swing axis is locked, the other limited or locked. This applies the Limited axis (ApplyDualConeSwingConstraint is used for the locked axis).
-		void ApplySwingConstraint(
-			const FReal Dt,
-			const FPBDJointSolverSettings& SolverSettings,
+		void InitPlanarPositionConstraint(
 			const FPBDJointSettings& JointSettings,
-			const EJointAngularConstraintIndex SwingConstraintIndex,
-			const bool bUseSoftLimit);
-
-		void ApplySwingTwistDrives(
 			const FReal Dt,
-			const FPBDJointSolverSettings& SolverSettings,
-			const FPBDJointSettings& JointSettings,
-			const bool bTwistDriveEnabled,
-			const bool bSwing1DriveEnabled,
-			const bool bSwing2DriveEnabled);
+			const int32 AxisIndex,
+			const EJointMotionType AxialMotion);
+		
+		/** Apply Position constraints */
 
-		void ApplySLerpDrive(
-			const FReal Dt,
-			const FPBDJointSolverSettings& SolverSettings,
-			const FPBDJointSettings& JointSettings);
+		void ApplyPositionConstraints(
+			const FReal Dt);
+		
+		void ApplyPointPositionConstraint();
+		
+		void ApplyAxisPositionConstraint(
+			const int32 ConstraintIndex,
+			const FReal Dt);
+		
+		void SolvePositionConstraintDelta(
+			const int32 ConstraintIndex, 
+			const FReal DeltaLambda);
+
+		void SolvePositionConstraintHard(
+			const int32 ConstraintIndex,
+			const FReal DeltaConstraint);
+
+		void SolvePositionConstraintSoft(
+			const int32 ConstraintIndex,
+			const FReal DeltaConstraint,
+			const FReal Dt,
+			const FReal TargetVel);
+		
+		/** Apply Linear Velocity constraints */
+		
+		void ApplyLinearVelocityConstraints();
+
+		void ApplyPointVelocityConstraint();
+
+		void ApplyAxisVelocityConstraint(
+			const int32 ConstraintIndex);
 
-		void ApplyPointPositionConstraintKD(
-			const int32 KIndex,
-			const int32 DIndex,
+		void SolveLinearVelocityConstraint(
+			const int32 ConstraintIndex,
+			const FReal TargetVel);
+		
+		/** Init Rotation constraints */
+		
+		void InitRotationConstraints(
 			const FReal Dt,
 			const FPBDJointSolverSettings& SolverSettings,
 			const FPBDJointSettings& JointSettings);
+		
+		void CorrectAxisAngleConstraint(
+			const FPBDJointSettings& JointSettings,
+			const int32 ConstraintIndex,
+			FVec3& ConstraintAxis,
+			FReal& ConstraintAngle) const;
 
-		void ApplyPointPositionConstraintDD(
+		void InitRotationConstraintDatas(
+			const FPBDJointSettings& JointSettings,
+			const int32 ConstraintIndex,
+			const FVec3& ConstraintAxis,
+			const FReal ConstraintAngle,
+			const FReal ConstraintRestitution,
 			const FReal Dt,
-			const FPBDJointSolverSettings& SolverSettings,
-			const FPBDJointSettings& JointSettings);
+			const bool bCheckLimit);
+		
+		void InitTwistConstraint(
+			const FPBDJointSettings& JointSettings,
+			const FReal Dt);
+		
+		void InitConeConstraint(
+			const FPBDJointSettings& JointSettings,
+			const FReal Dt);
 
-		void ApplySphericalPositionConstraint(
+		// One Swing axis is free, and the other locked. This applies the lock: Body1 Twist axis is confined to a plane.
+		void InitSingleLockedSwingConstraint(
+			const FPBDJointSettings& JointSettings,
 			const FReal Dt,
-			const FPBDJointSolverSettings& SolverSettings,
-			const FPBDJointSettings& JointSettings);
+			const EJointAngularConstraintIndex SwingConstraintIndex);
 
-		void ApplyCylindricalPositionConstraint(
-			const FReal Dt,
-			const int32 AxisIndex,
-			const EJointMotionType AxialMotion,
-			const EJointMotionType RadialMotion,
-			const FPBDJointSolverSettings& SolverSettings,
-			const FPBDJointSettings& JointSettings);
-
-		void ApplyPlanarPositionConstraint(
+		// One Swing axis is free, and the other limited. This applies the limit: Body1 Twist axis is confined to space between two cones.
+		void InitDualConeSwingConstraint(
+			const FPBDJointSettings& JointSettings,
 			const FReal Dt,
-			const int32 AxisIndex,
-			const EJointMotionType AxialMotion,
-			const FPBDJointSolverSettings& SolverSettings,
-			const FPBDJointSettings& JointSettings);
+			const EJointAngularConstraintIndex SwingConstraintIndex);
 
-		void ApplyPositionDrive(
-			const FReal Dt,
-			const int32 AxisIndex,
-			const FPBDJointSolverSettings& SolverSettings,
+		// One swing axis is locked, the other limited or locked. This applies the Limited axis (ApplyDualConeSwingConstraint is used for the locked axis).
+		void InitSwingConstraint(
 			const FPBDJointSettings& JointSettings,
-			const FVec3& Axis,
-			const FReal DeltaPos,
-			const FReal DeltaVel);
-
-
-		void ApplyPointProjection(
-			const FReal Dt,
 			const FPBDJointSolverSettings& SolverSettings,
-			const FPBDJointSettings& JointSettings,
-			const FReal Alpha,
-			FVec3& NetDP1,
-			FVec3& NetDR1);
-
-		void ApplySphereProjection(
 			const FReal Dt,
-			const FPBDJointSolverSettings& SolverSettings,
-			const FPBDJointSettings& JointSettings,
-			const FReal Alpha,
-			FVec3& NetDP1,
-			FVec3& NetDR1);
+			const EJointAngularConstraintIndex SwingConstraintIndex);
 
-		void ApplyTranslateProjection(
-			const FReal Dt,
-			const FPBDJointSolverSettings& SolverSettings,
+		void InitLockedRotationConstraints(
 			const FPBDJointSettings& JointSettings,
-			const FReal Alpha,
-			FVec3& NetDP1,
-			FVec3& NetDR1);
-
-		void ApplyConeProjection(
 			const FReal Dt,
-			const FPBDJointSolverSettings& SolverSettings,
-			const FPBDJointSettings& JointSettings,
-			const FReal Alpha,
-			const bool bPositionLocked,
-			FVec3& NetDP1,
-			FVec3& NetDR1);
+			const bool bApplyTwist,
+			const bool bApplySwing);
 
-		void ApplySwingProjection(
-			const FReal Dt,
-			const FPBDJointSolverSettings& SolverSettings,
-			const FPBDJointSettings& JointSettings,
-			const EJointAngularConstraintIndex SwingConstraintIndex,
-			const FReal Alpha,
-			const bool bPositionLocked,
-			FVec3& NetDP1,
-			FVec3& NetDR1);
+		/** Apply Rotation constraints */
 
-		void ApplySingleLockedSwingProjection(
-			const FReal Dt,
-			const FPBDJointSolverSettings& SolverSettings,
-			const FPBDJointSettings& JointSettings,
-			const EJointAngularConstraintIndex SwingConstraintIndex,
-			const FReal Alpha,
-			const bool bPositionLocked,
-			FVec3& NetDP1,
-			FVec3& NetDR1);
+		void ApplyRotationConstraints(
+			const FReal Dt);
+		
+		void ApplyRotationConstraint(
+			const int32 ConstraintIndex,
+			const FReal Dt);
 
-		void ApplyDoubleLockedSwingProjection(
-			const FReal Dt,
-			const FPBDJointSolverSettings& SolverSettings,
-			const FPBDJointSettings& JointSettings,
-			const FReal Alpha,
-			const bool bPositionLocked,
-			FVec3& NetDP1,
-			FVec3& NetDR1);
+		void SolveRotationConstraintDelta(
+			const int32 ConstraintIndex, 
+			const FReal DeltaLambda);
 
-		void ApplyDualConeSwingProjection(
-			const FReal Dt,
-			const FPBDJointSolverSettings& SolverSettings,
-			const FPBDJointSettings& JointSettings,
-			const EJointAngularConstraintIndex SwingConstraintIndex,
-			const FReal Alpha,
-			const bool bPositionLocked,
-			FVec3& NetDP1,
-			FVec3& NetDR1);
+		void SolveRotationConstraintHard(
+			const int32 ConstraintIndex,
+			const FReal DeltaConstraint);
 
-		void ApplyTwistProjection(
+		void SolveRotationConstraintSoft(
+			const int32 ConstraintIndex,
+			const FReal DeltaConstraint,
 			const FReal Dt,
-			const FPBDJointSolverSettings& SolverSettings,
-			const FPBDJointSettings& JointSettings,
-			const FReal Alpha,
-			const bool bPositionLocked,
-			FVec3& NetDP1,
-			FVec3& NetDR1);
+			const FReal TargetVel);
 
-		void ApplyVelocityProjection(
-			const FReal Dt,
-			const FPBDJointSolverSettings& SolverSettings,
-			const FPBDJointSettings& JointSettings,
-			const FReal Alpha,
-			const FVec3& DP1,
-			const FVec3& DR1);
-
-		void ApplyLinearVelocityConstraint(
-			const FReal Stiffness,
-			const FVec3& Axis,
-			const FReal IM,
-			const FVec3& Arm0,
-			const FVec3& Arm1,
-			const FVec3& DRAxis0,
-			const FVec3& DRAxis1,
-			const FReal TargetVel = 0);
-
-		void ApplyPointVelocityConstraint(
-			const FReal Dt,
-			const FPBDJointSolverSettings& SolverSettings,
-			const FPBDJointSettings& JointSettings);
+		/** Apply Angular velocity  constraint */
 
-		void ApplyPointVelocityConstraint(
-		    const FReal Stiffness,
-		    const FMatrix33& ConstraintIM,
-		    const FVec3& ConstraintArm0,
-		    const FVec3& ConstraintArm1);
+		void ApplyAngularVelocityConstraints();
 
-		void ApplySphericalVelocityConstraint(
-			const FReal Dt,
-			const FPBDJointSolverSettings& SolverSettings,
-			const FPBDJointSettings& JointSettings);
+		void SolveAngularVelocityConstraint(
+			const int32 ConstraintIndex,
+			const FReal TargetVel);
 
-		void ApplyCylindricalVelocityConstraint(
-			const FReal Dt,
-			const int32 AxisIndex,
-			const EJointMotionType AxialMotion,
-			const EJointMotionType RadialMotion,
-			const FPBDJointSolverSettings& SolverSettings,
-			const FPBDJointSettings& JointSettings);
+		void ApplyAngularVelocityConstraint(const int32 ConstraintIndex);
 
-		void ApplyPlanarVelocityConstraint(
+		/** Init Position Drives */
+		
+		void InitPositionDrives(
 			const FReal Dt,
-			const int32 AxisIndex,
-			const EJointMotionType AxialMotion,
 			const FPBDJointSolverSettings& SolverSettings,
 			const FPBDJointSettings& JointSettings);
 
-		void ApplyTwistVelocityConstraint(
-			const FReal Dt,
-			const FPBDJointSolverSettings& SolverSettings,
-			const FPBDJointSettings& JointSettings,
-			const bool bUseSoftLimit);
+		void InitAxisPositionDrive(
+            const int32 ConstraintIndex,
+            const FVec3& ConstraintAxis,
+            const FVec3& DeltaPosition,
+            const FVec3& DeltaVelocity,
+            const FReal Dt);
 
-		void ApplyConeVelocityConstraint(
-			const FReal Dt,
-			const FPBDJointSolverSettings& SolverSettings,
-			const FPBDJointSettings& JointSettings,
-			const bool bUseSoftLimit);
+		/** Apply Position Drives */
+		
+		void ApplyPositionDrives(
+			const FReal Dt);
+		
+		void ApplyAxisPositionDrive(
+			const int32 ConstraintIndex,
+			const FReal Dt);
 
-		// One Swing axis is free, and the other locked. This applies the lock: Body1 Twist axis is confined to a plane.
-		void ApplySingleLockedSwingVelocityConstraint(
-			const FReal Dt,
-			const FPBDJointSolverSettings& SolverSettings,
-			const FPBDJointSettings& JointSettings,
-			const EJointAngularConstraintIndex SwingConstraintIndex,
-			const bool bUseSoftLimit);
+		/** Init Rotation Drives */
 
-		// One Swing axis is free, and the other limited. This applies the limit: Body1 Twist axis is confined to space between two cones.
-		void ApplyDualConeSwingVelocityConstraint(
+		void InitRotationDrives(
 			const FReal Dt,
 			const FPBDJointSolverSettings& SolverSettings,
-			const FPBDJointSettings& JointSettings,
-			const EJointAngularConstraintIndex SwingConstraintIndex,
-			const bool bUseSoftLimit);
+			const FPBDJointSettings& JointSettings);
 
-		// One swing axis is locked, the other limited or locked. This applies the Limited axis (ApplyDualConeSwingConstraint is used for the locked axis).
-		void ApplySwingVelocityConstraint(
+		void InitRotationConstraintDrive(
+			const int32 ConstraintIndex,
+			const FVec3& ConstraintAxis,
+			const FReal Dt,
+			const FReal DeltaAngle);
+		
+		void InitSwingTwistDrives(
 			const FReal Dt,
 			const FPBDJointSolverSettings& SolverSettings,
 			const FPBDJointSettings& JointSettings,
-			const EJointAngularConstraintIndex SwingConstraintIndex,
-			const bool bUseSoftLimit);
-
-		void ApplyAngularVelocityConstraint(
-			const FReal Stiffness,
-			const FVec3& Axis,
-			const FReal TargetVel = 0.0f);
-
-		void ApplyLockedRotationVelocityConstraints(
+			const bool bTwistDriveEnabled,
+			const bool bSwing1DriveEnabled,
+			const bool bSwing2DriveEnabled);
+		
+		void InitSLerpDrive(
 			const FReal Dt,
 			const FPBDJointSolverSettings& SolverSettings,
-			const FPBDJointSettings& JointSettings,
-			const bool bApplyTwist,
-			const bool bApplySwing);
+			const FPBDJointSettings& JointSettings);
 
+		/** Apply Rotation Drives */
+		
+		void ApplyRotationDrives(
+			const FReal Dt);
+		
+		void ApplyAxisRotationDrive(
+    		const int32 ConstraintIndex,
+    		const FReal Dt);
+		
+		/** Linear / Angular padding */
+		
 		void CalculateLinearConstraintPadding(
+			const int32 ConstraintIndex,
 			const FReal Dt,
-			const FPBDJointSolverSettings& SolverSettings,
-			const FPBDJointSettings& JointSettings,
 			const FReal Restitution,
-			const int32 AxisIndex,
-			const FVec3 Axis,
 			FReal& InOutPos);
 
 		void CalculateAngularConstraintPadding(
-			const FReal Dt,
-			const FPBDJointSolverSettings& SolverSettings,
-			const FPBDJointSettings& JointSettings,
+			const int32 ConstraintIndex,
 			const FReal Restitution,
-			const EJointAngularConstraintIndex ConstraintIndex,
-			const FVec3 Axis,
 			FReal& InOutAngle);
 
 		inline bool HasLinearConstraintPadding(const int32 AxisIndex) const
@@ -674,72 +514,22 @@ namespace Chaos
 		{
 			LinearConstraintPadding[AxisIndex] = Padding;
 		}
-
-
-		inline bool HasAngularConstraintPadding(const EJointAngularConstraintIndex ConstraintIndex) const
+		
+		inline bool HasAngularConstraintPadding(const int32 ConstraintIndex) const
 		{
-			return AngularConstraintPadding[(int32)ConstraintIndex] >= 0.0f;
+			return AngularConstraintPadding[ConstraintIndex] >= 0.0f;
 		}
 
-		inline FReal GetAngularConstraintPadding(const EJointAngularConstraintIndex ConstraintIndex) const
+		inline FReal GetAngularConstraintPadding(const int32 ConstraintIndex) const
 		{
-			return FMath::Max(AngularConstraintPadding[(int32)ConstraintIndex], 0.0f);
+			return FMath::Max(AngularConstraintPadding[ConstraintIndex], 0.0f);
 		}
 
-		inline void SetAngularConstraintPadding(const EJointAngularConstraintIndex ConstraintIndex, FReal Padding)
+		inline void SetAngularConstraintPadding(const int32 ConstraintIndex, FReal Padding)
 		{
-			AngularConstraintPadding[(int32)ConstraintIndex] = Padding;
+			AngularConstraintPadding[ConstraintIndex] = Padding;
 		}
-
-		// Calculate linear velocities along constraint axes based on the constraint type in JointSettings. CV is the relative linear velocity of joint connectors.
-		void CalculateConstraintAxisLinearVelocities(
-			const FPBDJointSettings& JointSettings,
-			FVec3& ConstraintAxisLinearVelocities) const;
-
-		void CalculateSphericalConstraintAxisLinearVelocities(
-			const FPBDJointSettings& JointSettings,
-			FVec3& ConstraintAxisLinearVelocities
-			) const;
-
-		void CalculateCylindricalConstraintAxisLinearVelocities(
-			const int32 AxisIndex,
-			const EJointMotionType AxialMotion,
-			const EJointMotionType RadialMotion,
-			const FPBDJointSettings& JointSettings,
-			FVec3& ConstraintAxisLinearVelocities) const;
-
-		void CalculatePlanarConstraintAxisLinearVelocities(
-			const int32 AxisIndex,
-			const EJointMotionType AxialMotion,
-			const FPBDJointSettings& JointSettings,
-			FVec3& ConstraintAxisLinearVelocities) const;
-
-		// Calculate angular velocities along constraint axes based on the constraint type in JointSettings. CW is the relative angular velocity of joint connectors.
-		void CalculateConstraintAxisAngularVelocities(
-			const FPBDJointSolverSettings& SolverSettings,
-			const FPBDJointSettings& JointSettings,
-			FVec3& ConstraintAxisAngularVelocities
-			) const;
-
-		void CalculateTwistConstraintAxisAngularVelocities(
-			const FPBDJointSettings& JointSettings,
-			FVec3& ConstraintAxisAngularVelocities) const;
-
-		void CalculateConeConstraintAxisAngularVelocities(
-			const FPBDJointSettings& JointSettings,
-			FVec3& ConstraintAxisAngularVelocities) const;
-
-		void CalculateDualConeSwingConstraintAxisAngularVelocities(
-			const FPBDJointSettings& JointSettings,
-			const EJointAngularConstraintIndex SwingConstraintIndex,
-			FVec3& ConstraintAxisAngularVelocities) const;
-
-		void CalculateSwingConstraintAxisAngularVelocities(
-			const FPBDJointSolverSettings& SolverSettings,
-			const FPBDJointSettings& JointSettings,
-			const EJointAngularConstraintIndex SwingConstraintIndex,
-			FVec3& ConstraintAxisAngularVelocities) const;
-
+		
 		// The cached body state on which the joint operates
 		FConstraintSolverBody SolverBodies[MaxConstrainedBodies];
 
@@ -761,21 +551,7 @@ namespace Chaos
 		// Accumulated Impulse and AngularImpulse (Impulse * Dt since they are mass multiplied position corrections)
 		FVec3 NetLinearImpulse;
 		FVec3 NetAngularImpulse;
-
-		// Lagrange multipliers of the position constraints.
-		// Currently these are only used in ApplyCylindricalVelocityConstraints
-		FVec3 LinearHardLambda;
-		// Lagrange multipliers of the rotation constraints
-		// Currently these are used in ApplyAngularVelocityConstraints
-		FVec3 AngularHardLambda;
-
-		// XPBD Accumulators (net impulse for each soft constraint/drive)
-		FReal LinearSoftLambda;
-		FReal TwistSoftLambda;
-		FReal SwingSoftLambda;
-		FVec3 LinearDriveLambdas;
-		FVec3 RotationDriveLambdas;
-
+		
 		// Solver stiffness - increased over iterations for stability
 		// \todo(chaos): remove Stiffness from SolverSettings (because it is not a solver constant)
 		FReal SolverStiffness;
@@ -789,27 +565,33 @@ namespace Chaos
 		FReal AngleTolerance;						// Angle error below which we consider a constraint or drive solved
 
 		// Tracking whether the solver is resolved
-		FVec3 LastPs[MaxConstrainedBodies];			// Positions at the beginning of the iteration
-		FRotation3 LastQs[MaxConstrainedBodies];	// Rotations at the beginning of the iteration
+		FVec3 LastDPs[MaxConstrainedBodies];			// Positions at the beginning of the iteration
+		FVec3 LastDQs[MaxConstrainedBodies];	// Rotations at the beginning of the iteration
+		FVec3 CurrentPs[MaxConstrainedBodies];			// Positions at the beginning of the iteration
+		FRotation3 CurrentQs[MaxConstrainedBodies];	// Rotations at the beginning of the iteration
 		FVec3 InitConstraintAxisLinearVelocities; // Linear velocities along the constraint axes at the begining of the frame, used by restitution
 		FVec3 InitConstraintAxisAngularVelocities; // Angular velocities along the constraint axes at the begining of the frame, used by restitution
 		int32 NumActiveConstraints;					// The number of active constraints and drives in the last iteration (-1 initial value)
 		bool bIsActive;								// Whether any constraints actually moved any bodies in last iteration
+		
+		FPointConstraintDatas PointConstraints;
+		
+		FAxisConstraintDatas PositionConstraints;
+		FAxisConstraintDatas RotationConstraints;
 
-		// Caching data
-		FReal LinearStiffness, LinearDamping;
-		FReal AngularStiffness, AngularDamping;
-		FVec3 ConstraintArms[3][2];
-		FMatrix33 PointPositionConstraintIM;
-		FReal PositionConstraintIM[2];
-		FVec3 PositionConstraintAxis[2];
-		FVec3 PositionConstraintDRAxis[2][2]; // Cache for InvI * cross(Arm, Axis)
-		bool bLinearAccelerationMode;
+		FAxisConstraintDatas PositionDrives;
+		FAxisConstraintDatas RotationDrives;
 
 		// dummy indices
 		static constexpr int32 PointPositionConstraintIndex = 0;
 		static constexpr int32 SphericalPositionConstraintIndex = 0;
 		static constexpr int32 PlanarPositionConstraintIndex = 0;
+
+	private :
+		/** Compute the body and connector state
+		 * @param BodyIndex Index of the body (0,1) on which the state will be updated
+		 */
+		void ComputeBodyState(const int32 BodyIndex);
 	};
 
 }
\ No newline at end of file
