AABBTree overlap caching fix when time slicing is used

diff --git a/Engine/Source/Runtime/Experimental/Chaos/Private/Chaos/PBDRigidsEvolution.cpp b/Engine/Source/Runtime/Experimental/Chaos/Private/Chaos/PBDRigidsEvolution.cpp
--- a/Engine/Source/Runtime/Experimental/Chaos/Private/Chaos/PBDRigidsEvolution.cpp
+++ b/Engine/Source/Runtime/Experimental/Chaos/Private/Chaos/PBDRigidsEvolution.cpp
@@ -789,6 +789,12 @@ namespace Chaos
 			else
 			{
 				FlushInternalAccelerationQueue();
+
+				if (GAccelerationStructureCacheOverlappingLeaves)
+				{
+					// Caching of the overlapping leaves
+					InternalAcceleration->CacheOverlappingLeaves();
+				}
 			}
 			
 			if (bCanStartAsyncTasks)
