GeometryCollection : Properties
- Adding damage threshold back to the GeometryCollectionObject.

diff --git a/Engine/Source/Runtime/Experimental/GeometryCollectionEngine/Private/GeometryCollection/GeometryCollectionObject.cpp b/Engine/Source/Runtime/Experimental/GeometryCollectionEngine/Private/GeometryCollection/GeometryCollectionObject.cpp
--- a/Engine/Source/Runtime/Experimental/GeometryCollectionEngine/Private/GeometryCollection/GeometryCollectionObject.cpp
+++ b/Engine/Source/Runtime/Experimental/GeometryCollectionEngine/Private/GeometryCollection/GeometryCollectionObject.cpp
@@ -69,7 +69,7 @@ UGeometryCollection::UGeometryCollection(const FObjectInitializer& ObjectInitial
 	, EnableClustering(true)
 	, ClusterGroupIndex(0)
 	, MaxClusterLevel(100)
-	, DamageThreshold_DEPRECATED({ 500000.f, 50000.f, 5000.f })
+	, DamageThreshold({ 500000.f, 50000.f, 5000.f })
 	, ClusterConnectionType(EClusterConnectionTypeEnum::Chaos_MinimalSpanningSubsetDelaunayTriangulation)
 	, bUseFullPrecisionUVs(false)
 	, bStripOnCook(false)
diff --git a/Engine/Source/Runtime/Experimental/GeometryCollectionEngine/Public/GeometryCollection/GeometryCollectionObject.h b/Engine/Source/Runtime/Experimental/GeometryCollectionEngine/Public/GeometryCollection/GeometryCollectionObject.h
--- a/Engine/Source/Runtime/Experimental/GeometryCollectionEngine/Public/GeometryCollection/GeometryCollectionObject.h
+++ b/Engine/Source/Runtime/Experimental/GeometryCollectionEngine/Public/GeometryCollection/GeometryCollectionObject.h
@@ -424,8 +424,8 @@ public:
 	int32 MaxClusterLevel;
 
 	/** Damage threshold for clusters at different levels. */
-	UPROPERTY(meta = (DeprecatedProperty, DeprecationMessage = "This property is deprecated. Use the default SizeSpecificData instead."))
-	TArray<float> DamageThreshold_DEPRECATED;
+	UPROPERTY(EditAnywhere, Category = "Clustering")
+	TArray<float> DamageThreshold;
 
 	/** */
 	UPROPERTY(EditAnywhere, Category = "Clustering")
