Chaos: Fixed raycast and collision for negative scales on trimesh and heightfield

diff --git a/Engine/Source/Runtime/Experimental/Chaos/Private/Chaos/HeightField.cpp b/Engine/Source/Runtime/Experimental/Chaos/Private/Chaos/HeightField.cpp
--- a/Engine/Source/Runtime/Experimental/Chaos/Private/Chaos/HeightField.cpp
+++ b/Engine/Source/Runtime/Experimental/Chaos/Private/Chaos/HeightField.cpp
@@ -873,8 +873,11 @@ namespace Chaos
 
 			OutMin = FVec3(static_cast<FReal>(InCoord[0]), static_cast<FReal>(InCoord[1]), GeomData.GetMinHeight());
 			OutMax = FVec3(static_cast<FReal>(InCoord[0] + 1), static_cast<FReal>(InCoord[1] + 1), Max[2]);
-			OutMin = OutMin * GeomData.Scale - InInflate;
-			OutMax = OutMax * GeomData.Scale + InInflate;
+
+			const FAABB3 CellBoundScaled = FAABB3::FromPoints(OutMin * GeomData.Scale, OutMax * GeomData.Scale);
+
+			OutMin = CellBoundScaled.Min() - InInflate;
+			OutMax = CellBoundScaled.Max() + InInflate;
 			return true;
 		}
 
@@ -886,18 +889,13 @@ namespace Chaos
 		if(FlatGrid.IsValid(InCoord))
 		{
 			int32 Index = InCoord[1] * (GeomData.NumCols) + InCoord[0];
-			static FVec3 Points[4];
+			FVec3 Points[4];
 			GeomData.GetPoints(Index, Points);
 
-			OutMin = OutMax = Points[0];
-
-			for(int32 PointIndex = 1; PointIndex < 4; ++PointIndex)
-			{
-				const FVec3& Point = Points[PointIndex];
-				OutMin = FVec3(FGenericPlatformMath::Min(OutMin[0], Point[0]), FGenericPlatformMath::Min(OutMin[1], Point[1]), FGenericPlatformMath::Min(OutMin[2], Point[2]));
-				OutMax = FVec3(FGenericPlatformMath::Max(OutMax[0], Point[0]), FGenericPlatformMath::Max(OutMax[1], Point[1]), FGenericPlatformMath::Max(OutMax[2], Point[2]));
-			}
+			const FAABB3 CellBound = FAABB3::FromPoints(Points[0], Points[1], Points[2], Points[3]);
 
+			OutMin = CellBound.Min();
+			OutMax = CellBound.Max();
 			OutMin -= InInflate;
 			OutMax += InInflate;
 
@@ -931,27 +929,22 @@ namespace Chaos
 			bParallel[Axis] = FMath::IsNearlyZero(Dir[Axis], (FReal)1.e-8);
 			InvDir[Axis] = bParallel[Axis] ? 0 : 1 / Dir[Axis];
 		}
-
-		const FBounds2D FlatBounds = GetFlatBounds();
-		FAABB3 Bounds(
-			FVec3(FlatBounds.Min[0],FlatBounds.Min[1],GeomData.GetMinHeight() * GeomData.Scale[2]),
-			FVec3(FlatBounds.Max[0],FlatBounds.Max[1],GeomData.GetMaxHeight() * GeomData.Scale[2])
-			);
-
+		
 		FReal RayEntryTime;
 		FReal RayExitTime;
-		if(Bounds.RaycastFast(StartPoint, Dir, InvDir, bParallel, Length, InvCurrentLength, RayEntryTime, RayExitTime))
+		if(CachedBounds.RaycastFast(StartPoint, Dir, InvDir, bParallel, Length, InvCurrentLength, RayEntryTime, RayExitTime))
 		{
 			CurrentLength = RayExitTime + SMALL_NUMBER; // to account for precision errors 
 			FVec3 NextStart = StartPoint + (Dir * RayEntryTime);
 
 			const FVec2 Scale2D(GeomData.Scale[0],GeomData.Scale[1]);
 			TVec2<int32> CellIdx = FlatGrid.Cell(TVec2<int32>(static_cast<int32>(NextStart[0] / Scale2D[0]), static_cast<int32>(NextStart[1] / Scale2D[1])));
-			const FReal ZDx = Bounds.Extents()[2];
-			const FReal ZMidPoint = Bounds.Min()[2] + ZDx * 0.5f;
+			const FReal ZDx = CachedBounds.Extents()[2];
+			const FReal ZMidPoint = CachedBounds.Min()[2] + ZDx * 0.5f;
 			const FVec3 ScaledDx(FlatGrid.Dx()[0] * Scale2D[0],FlatGrid.Dx()[1] * Scale2D[1],ZDx);
 			const FVec2 ScaledDx2D(ScaledDx[0],ScaledDx[1]);
 			const FVec2 ScaledMin = FlatGrid.MinCorner() * Scale2D;
+			const FVec3 ScaleSign = GeomData.Scale.GetSignVector();
 
 			//START
 			do
@@ -981,7 +974,7 @@ namespace Chaos
 				{
 					if(!bParallel[Axis])
 					{
-						const FReal CrossPoint = Dir[Axis] > 0 ? ScaledCellCenter[Axis] + ScaledDx[Axis] / 2 : ScaledCellCenter[Axis] - ScaledDx[Axis] / 2;
+						const FReal CrossPoint = (Dir[Axis] * ScaleSign[Axis]) > 0 ? ScaledCellCenter[Axis] + ScaledDx[Axis] / 2 : ScaledCellCenter[Axis] - ScaledDx[Axis] / 2;
 						const FReal Distance = CrossPoint - NextStart[Axis];	//note: CellCenter already has /2, we probably want to use the corner instead
 						const FReal Time = Distance * InvDir[Axis];
 						Times[Axis] = Time;
@@ -1005,7 +998,7 @@ namespace Chaos
 
 				for(int Axis = 0; Axis < 2; ++Axis)
 				{
-					CellIdx[Axis] += (Times[Axis] <= BestTime) ? (Dir[Axis] > 0 ? 1 : -1) : 0;
+					CellIdx[Axis] += (Times[Axis] <= BestTime) ? ((Dir[Axis] * ScaleSign[Axis]) > 0 ? 1 : -1) : 0;
 					if(CellIdx[Axis] < 0 || CellIdx[Axis] >= FlatGrid.Counts()[Axis])
 					{
 						return false;
@@ -1229,7 +1222,7 @@ namespace Chaos
 							ExpandAxis = 0;
 						}
 						FReal ExpandSize = HalfExtents3D[ExpandAxis];
-						int32 Steps = FMath::TruncToInt(FMath::RoundFromZero(ExpandSize / GeomData.Scale[ExpandAxis]));
+						int32 Steps = FMath::TruncToInt(FMath::RoundFromZero(ExpandSize / FMath::Abs(GeomData.Scale[ExpandAxis])));
 
 						Expand(Coord, ThickenDir, Steps);
 						Expand(Coord, -ThickenDir, Steps);
@@ -2120,7 +2113,9 @@ namespace Chaos
 
 			const FVec3 AB = B - A;
 			const FVec3 AC = C - A;
-			FVec3 Normal = FVec3::CrossProduct(AB, AC);
+			const FVec3 ScaleSigns = GeomData.Scale.GetSignVector();
+			const FReal ScaleInversion = ScaleSigns.X * ScaleSigns.Y * ScaleSigns.Z;
+			FVec3 Normal = FVec3::CrossProduct(AB, AC) * ScaleInversion;
 			const FReal Length = Normal.SafeNormalize();
 			ensure(Length > 0);
 			return Normal;
diff --git a/Engine/Source/Runtime/Experimental/Chaos/Private/Chaos/TriangleMeshImplicitObject.cpp b/Engine/Source/Runtime/Experimental/Chaos/Private/Chaos/TriangleMeshImplicitObject.cpp
--- a/Engine/Source/Runtime/Experimental/Chaos/Private/Chaos/TriangleMeshImplicitObject.cpp
+++ b/Engine/Source/Runtime/Experimental/Chaos/Private/Chaos/TriangleMeshImplicitObject.cpp
@@ -21,6 +21,11 @@ namespace Chaos
 	bool TriMeshPerPolySupport = 1;
 	FAutoConsoleVariableRef CVarPerPolySupport(TEXT("p.Chaos.TriMeshPerPolySupport"), TriMeshPerPolySupport, TEXT("Disabling removes memory cost of vertex map on triangle mesh. Note: Changing at runtime will not work."));
 
+	FReal GetWindingOrder(const FVec3& Scale)
+	{
+		const FVec3 SignVector = Scale.GetSignVector();
+		return SignVector.X * SignVector.Y * SignVector.Z;
+	}
 
 template <typename QueryGeomType>
 static auto MakeScaledHelper(const QueryGeomType& B, const FVec3& InvScale)
@@ -360,7 +365,7 @@ bool FTriangleMeshImplicitObject::ContactManifoldImp(const GeomType& QueryGeom,
 	FAABB3 QueryBounds = WorldScaleGeom.BoundingBox();
 	QueryBounds = QueryBounds.TransformedAABB(TriMeshToGeomNoScale);
 	QueryBounds.ThickenSymmetrically(FVec3(WorldThickness));
-	QueryBounds.Scale(InvTriMeshScale);
+	QueryBounds.ScaleWithNegative(InvTriMeshScale);
 
 	TRigidTransform<FReal, 3> WorldScaleQueryTM;
 	ScaleTransformHelper(TriMeshScale, QueryTM, WorldScaleQueryTM);
@@ -480,7 +485,7 @@ bool FTriangleMeshImplicitObject::GJKContactPointImp(const QueryGeomType& QueryG
 	FAABB3 QueryBounds = WorldScaleGeom.BoundingBox();
 	QueryBounds = QueryBounds.TransformedAABB(TriMeshToGeomNoScale);
 	QueryBounds.ThickenSymmetrically(FVec3(WorldThickness));
-	QueryBounds.Scale(InvTriMeshScale);
+	QueryBounds.ScaleWithNegative(InvTriMeshScale);
 
 	TRigidTransform<FReal, 3> WorldScaleQueryTM;
 	ScaleTransformHelper(TriMeshScale, QueryTM, WorldScaleQueryTM);
@@ -748,7 +753,7 @@ bool FTriangleMeshImplicitObject::OverlapGeomImp(const QueryGeomType& QueryGeom,
 	FAABB3 QueryBounds = WorldScaleQueryGeom.BoundingBox();
 	QueryBounds = QueryBounds.TransformedAABB(TriMeshToGeomNoScale);
 	QueryBounds.ThickenSymmetrically(FVec3(Thickness));
-	QueryBounds.Scale(InvTriMeshScale);
+	QueryBounds.ScaleWithNegative(InvTriMeshScale);
 
 	const TArray<int32> PotentialIntersections = BVH.FindAllIntersections(QueryBounds);
 
@@ -993,12 +998,6 @@ void ComputeScaledSweepInputs(FVec3 TriMeshScale, const FRigidTransform3& StartT
 	OutScaledStartTM = FRigidTransform3(StartTM.GetLocation() * TriMeshScale, StartTM.GetRotation());
 }
 
-FReal GetWindingOrder(const FVec3& Scale)
-{
-	const FVec3 SignVector = Scale.GetSignVector();
-	return SignVector.X * SignVector.Y * SignVector.Z;
-}
-
 template <typename QueryGeomType>
 bool FTriangleMeshImplicitObject::SweepGeomImp(const QueryGeomType& QueryGeom, const FRigidTransform3& StartTM, const FVec3& Dir, const FReal Length,
 	FReal& OutTime, FVec3& OutPosition, FVec3& OutNormal, int32& OutFaceIndex, const FReal Thickness, const bool bComputeMTD, FVec3 TriMeshScale) const
diff --git a/Engine/Source/Runtime/Experimental/Chaos/Public/Chaos/AABB.h b/Engine/Source/Runtime/Experimental/Chaos/Public/Chaos/AABB.h
--- a/Engine/Source/Runtime/Experimental/Chaos/Public/Chaos/AABB.h
+++ b/Engine/Source/Runtime/Experimental/Chaos/Public/Chaos/AABB.h
@@ -475,6 +475,12 @@ namespace Chaos
 			}
 		}
 
+		FORCEINLINE TAABB<T, d>& ScaleWithNegative(const TVector<T, d>& InScale)
+		{
+			*this = FromPoints(MMin * InScale, MMax * InScale);
+			return *this;
+		}
+
 		/**
 		* Scale the AABB relative to the origin
 		* IMPORTANT : this does not support negative scale
diff --git a/Engine/Source/Runtime/Landscape/Private/LandscapeCollision.cpp b/Engine/Source/Runtime/Landscape/Private/LandscapeCollision.cpp
--- a/Engine/Source/Runtime/Landscape/Private/LandscapeCollision.cpp
+++ b/Engine/Source/Runtime/Landscape/Private/LandscapeCollision.cpp
@@ -481,7 +481,7 @@ void ULandscapeHeightfieldCollisionComponent::OnCreatePhysicsState()
 				// Create a shape for a heightfield which is used only by the landscape editor
 				if(!GetWorld()->IsGameWorld() && !GetOutermost()->bIsCookedForEditor)
 				{
-					HeightfieldRef->EditorHeightfield->SetScale(FinalScale);
+					HeightfieldRef->EditorHeightfield->SetScale(FinalScale * LandscapeComponentTransform.GetScale3D().GetSignVector());
 					TUniquePtr<Chaos::TImplicitObjectTransformed<Chaos::FReal, 3>> ChaosEditorHeightFieldFromCooked = MakeUnique<Chaos::TImplicitObjectTransformed<Chaos::FReal, 3>>(MakeSerializable(HeightfieldRef->EditorHeightfield), Chaos::FRigidTransform3(FTransform::Identity));
 
 					TUniquePtr<Chaos::FPerShapeData> NewEditorShape = Chaos::FPerShapeData::CreatePerShapeData(ShapeArray.Num(), MakeSerializable(ChaosEditorHeightFieldFromCooked));
