Chaos: Fixed triangle meshes taking bounds from their vertex buffer and not their index buffer.
 - When using BSP we were giving every generated triangle mesh the same bounds encompassing the whole level
 - This led to severe performance degradation when worlds contained many BSP brushes

diff --git a/Engine/Source/Runtime/Experimental/Chaos/Public/Chaos/TriangleMeshImplicitObject.h b/Engine/Source/Runtime/Experimental/Chaos/Public/Chaos/TriangleMeshImplicitObject.h
--- a/Engine/Source/Runtime/Experimental/Chaos/Public/Chaos/TriangleMeshImplicitObject.h
+++ b/Engine/Source/Runtime/Experimental/Chaos/Public/Chaos/TriangleMeshImplicitObject.h
@@ -16,6 +16,7 @@
 #include "UObject/ExternalPhysicsCustomObjectVersion.h"
 #include "UObject/PhysicsObjectVersion.h"
 
+#include <type_traits>
 
 namespace Chaos
 {
@@ -78,18 +79,47 @@ namespace Chaos
 			return bRequiresLargeIndices;
 		}
 
-		const auto& GetLargeIndexBuffer() const
+		const TArray<TVec3<LargeIdxType>>& GetLargeIndexBuffer() const
 		{
 			check(bRequiresLargeIndices);
 			return LargeIdxBuffer;
 		}
 
-		const auto& GetSmallIndexBuffer() const
+		const TArray<TVec3<SmallIdxType>>& GetSmallIndexBuffer() const
 		{
 			check(!bRequiresLargeIndices);
 			return SmallIdxBuffer;
 		}
 
+		int32 GetNumTriangles() const
+		{
+			if(bRequiresLargeIndices)
+			{
+				return LargeIdxBuffer.Num();
+			}
+
+			return SmallIdxBuffer.Num();
+		}
+
+		template<typename ExpectedType>
+		const TArray<TVec3<ExpectedType>>& GetIndexBuffer() const
+		{
+			if constexpr(std::is_same_v<ExpectedType, LargeIdxType>)
+			{
+				check(bRequiresLargeIndices);
+				return LargeIdxBuffer;
+			}
+			else if constexpr(std::is_same_v<ExpectedType, SmallIdxType>)
+			{
+				check(!bRequiresLargeIndices);
+				return SmallIdxBuffer;
+			}
+			else
+			{
+				static_assert(sizeof(ExpectedType) == 0, "Unsupported index buffer type");
+			}
+		}
+
 	private:
 		TArray<TVec3<LargeIdxType>> LargeIdxBuffer;
 		TArray<TVec3<SmallIdxType>> SmallIdxBuffer;
@@ -141,10 +171,26 @@ namespace Chaos
 		, ExternalVertexIndexMap(MoveTemp(InExternalVertexIndexMap))
 		, bCullsBackFaceRaycast(bInCullsBackFaceRaycast)
 		{
-			for (uint32 Idx = 1; Idx < MParticles.Size(); ++Idx)
+			const int32 NumTriangles = MElements.GetNumTriangles();
+			if(NumTriangles > 0)
 			{
-				MLocalBoundingBox.GrowToInclude(MParticles.X(Idx));
+
+				const TArray<TVec3<IdxType>>& Tris = MElements.GetIndexBuffer<IdxType>();
+				const TVec3<IdxType>& FirstTri = Tris[0];
+
+				MLocalBoundingBox = FAABB3(MParticles.X(FirstTri[0]), MParticles.X(FirstTri[0]));
+				MLocalBoundingBox.GrowToInclude(MParticles.X(FirstTri[1]));
+				MLocalBoundingBox.GrowToInclude(MParticles.X(FirstTri[2]));
+
+				for(int32 TriangleIndex = 1; TriangleIndex < NumTriangles; ++TriangleIndex)
+				{
+					const TVec3<IdxType>& Tri = Tris[TriangleIndex];
+					MLocalBoundingBox.GrowToInclude(MParticles.X(Tri[0]));
+					MLocalBoundingBox.GrowToInclude(MParticles.X(Tri[1]));
+					MLocalBoundingBox.GrowToInclude(MParticles.X(Tri[2]));
+				}
 			}
+			
 			RebuildBV();
 		}
 
@@ -378,10 +424,25 @@ namespace Chaos
 			, ExternalVertexIndexMap(MoveTemp(InExternalVertexIndexMap))
 			, bCullsBackFaceRaycast(bInCullsBackFaceRaycast)
 		{
-			for(uint32 Idx = 1; Idx < MParticles.Size(); ++Idx)
+			const int32 NumTriangles = MElements.GetNumTriangles();
+			if(NumTriangles > 0)
 			{
-				MLocalBoundingBox.GrowToInclude(MParticles.X(Idx));
+				const TArray<TVec3<IdxType>>& Tris = MElements.GetIndexBuffer<IdxType>();
+				const TVec3<IdxType>& FirstTri = Tris[0];
+
+				MLocalBoundingBox = FAABB3(MParticles.X(FirstTri[0]), MParticles.X(FirstTri[0]));
+				MLocalBoundingBox.GrowToInclude(MParticles.X(FirstTri[1]));
+				MLocalBoundingBox.GrowToInclude(MParticles.X(FirstTri[2]));
+
+				for(int32 TriangleIndex = 1; TriangleIndex < NumTriangles; ++TriangleIndex)
+				{
+					const TVec3<IdxType>& Tri = Tris[TriangleIndex];
+					MLocalBoundingBox.GrowToInclude(MParticles.X(Tri[0]));
+					MLocalBoundingBox.GrowToInclude(MParticles.X(Tri[1]));
+					MLocalBoundingBox.GrowToInclude(MParticles.X(Tri[2]));
+				}
 			}
+			
 			BVH.CopyFrom(InBvhToCopy);
 		}
 
diff --git a/Engine/Source/Runtime/Experimental/ChaosCore/Public/Chaos/Core.h b/Engine/Source/Runtime/Experimental/ChaosCore/Public/Chaos/Core.h
--- a/Engine/Source/Runtime/Experimental/ChaosCore/Public/Chaos/Core.h
+++ b/Engine/Source/Runtime/Experimental/ChaosCore/Public/Chaos/Core.h
@@ -42,5 +42,5 @@ namespace Chaos
 	// NOTE: if you get a merge conflict on the GUID, you must replace it with a new GUID - do not accept the source or target
 	// or you will likely get DDC version conflicts resulting in crashes during load.
 	// Core version string for Chaos data. Any DDC builder dependent on Chaos for serialization should depend on this version
-	static const FString ChaosVersionString = TEXT("F8EF37A7-2C6C-4269-BE5F-6B8B7C6B5295");
+	static const FString ChaosVersionString = TEXT("997CB556-41E1-4D22-B53C-986B3562586C");
 }
