Chaos - Sphere-Convex manifolds

diff --git a/Engine/Source/Runtime/Experimental/Chaos/Private/Chaos/CollisionOneShotManifolds.cpp b/Engine/Source/Runtime/Experimental/Chaos/Private/Chaos/CollisionOneShotManifolds.cpp
--- a/Engine/Source/Runtime/Experimental/Chaos/Private/Chaos/CollisionOneShotManifolds.cpp
+++ b/Engine/Source/Runtime/Experimental/Chaos/Private/Chaos/CollisionOneShotManifolds.cpp
@@ -20,6 +20,10 @@
 
 namespace Chaos
 {
+	// If a sphere is this much smaller than a capsule, we create a single-point manifold. Otherwise a three-point manifold
+	FRealSingle Chaos_Collision_Manifold_SphereCapsuleSizeThreshold = 0.5f;
+	FAutoConsoleVariableRef CVarChaos_Manifold_SphereCapsuleSizeThreshold(TEXT("p.Chaos.Collision.Manifold.SphereCapsuleSizeThreshold"), Chaos_Collision_Manifold_SphereCapsuleSizeThreshold, TEXT(""));
+
 	FRealSingle Chaos_Collision_Manifold_PlaneContactNormalEpsilon = 0.001f;
 	FAutoConsoleVariableRef CVarChaos_Manifold_PlaneContactNormalEpsilon(TEXT("p.Chaos.Collision.Manifold.PlaneContactNormalEpsilon"), Chaos_Collision_Manifold_PlaneContactNormalEpsilon, TEXT("Normal tolerance used to distinguish face contacts from edge-edge contacts"));
 
@@ -29,10 +33,6 @@ namespace Chaos
 	FRealSingle Chaos_Collision_Manifold_EdgeContactNormalThreshold = 0.9f;	// About 25deg
 	FAutoConsoleVariableRef CVarChaos_Manifold_EdgeContactNormalThreshold(TEXT("p.Chaos.Collision.Manifold.EdgeNormalThreshold"), Chaos_Collision_Manifold_EdgeContactNormalThreshold, TEXT(""));
 
-	// @todo(chaos): TEMP - use convex-convex collisio detection for box-box until TBox::GetClosestEdgePosition is implemented for that path (without plane hint)
-	bool bChaos_Collision_Manifold_BoxAsConvex = true;
-	FAutoConsoleVariableRef CVarChaosCollisioConvexManifoldBoxAsConvex(TEXT("p.Chaos.Collision.Manifold.BoxAsConvex"), bChaos_Collision_Manifold_BoxAsConvex, TEXT(""));
-
 	FRealSingle Chaos_Collision_Manifold_CullDistanceMarginMultiplier = 1.0f;
 	FAutoConsoleVariableRef CVarChaosCollisioConvexManifoldCullDistanceMarginMultiplier(TEXT("p.Chaos.Collision.Manifold.CullDistanceMarginMultiplier"), Chaos_Collision_Manifold_CullDistanceMarginMultiplier, TEXT(""));
 
diff --git a/Engine/Source/Runtime/Experimental/Chaos/Private/Chaos/CollisionOneShotManifoldsMiscShapes.cpp b/Engine/Source/Runtime/Experimental/Chaos/Private/Chaos/CollisionOneShotManifoldsMiscShapes.cpp
--- a/Engine/Source/Runtime/Experimental/Chaos/Private/Chaos/CollisionOneShotManifoldsMiscShapes.cpp
+++ b/Engine/Source/Runtime/Experimental/Chaos/Private/Chaos/CollisionOneShotManifoldsMiscShapes.cpp
@@ -25,6 +25,7 @@
 namespace Chaos
 {
 	extern FRealSingle Chaos_Collision_EdgePrunePlaneDistance;
+	extern FRealSingle Chaos_Collision_Manifold_SphereCapsuleSizeThreshold;
 
 	namespace Collisions
 	{
@@ -93,6 +94,12 @@ namespace Chaos
 			}
 		}
 
+		// Build a sphere-Capsule manifold.
+		// When the sphere and capsule are of similar size, we usually only need a 1-point manifold.
+		// If the sphere is larger than the capsule, we need to generate a multi-point manifold so that
+		// we don't end up jittering between collisions on each end cap. E.g., consider a small capsule
+		// lying horizontally on a very large sphere (almost flat) - we need at least 2 contact points to 
+		// make this stable.
 		void ConstructSphereCapsuleOneShotManifold(const TSphere<FReal, 3>& Sphere, const FRigidTransform3& SphereTransform, const FCapsule& Capsule, const FRigidTransform3& CapsuleTransform, const FReal Dt, FPBDCollisionConstraint& Constraint)
 		{
 			// We only build one shot manifolds once
@@ -100,13 +107,115 @@ namespace Chaos
 			ensure(SphereTransform.GetScale3D() == FVec3(1.0f, 1.0f, 1.0f));
 			ensure(CapsuleTransform.GetScale3D() == FVec3(1.0f, 1.0f, 1.0f));
 
-			// @todo(chaos): support manifold maintenance
 			Constraint.ResetActiveManifoldContacts();
 
-			FContactPoint ContactPoint = SphereCapsuleContactPoint(Sphere, SphereTransform, Capsule, CapsuleTransform, Constraint.GetRestitutionPadding());
-			if (ContactPoint.Phi < Constraint.GetCullDistance())
+			// Build a multi-point manifold
+			const FReal NetCullDistance = Sphere.GetRadius() + Capsule.GetRadius() + Constraint.GetCullDistance() + Constraint.GetRestitutionPadding();
+			const FReal NetCullDistanceSq = FMath::Square(NetCullDistance);
+
+			// Transform the sphere into capsule space and find the closest point on the capsule line segment
+			// @todo(chaos) this would be much simpler if the sphere's were always at the origin and capsules were at the origin and axis aligned
+			const FRigidTransform3 SphereToCapsuleTransform = SphereTransform.GetRelativeTransformNoScale(CapsuleTransform);
+			const FVec3 SpherePos = SphereToCapsuleTransform.TransformPositionNoScale(Sphere.GetCenter());
+			const FReal NearPosT = Utilities::ClosestTimeOnLineSegment(SpherePos, Capsule.GetX1(), Capsule.GetX2());
+
+			// Add the closest contact point to the manifold
+			const FVec3 NearPos = FMath::Lerp(Capsule.GetX1(), Capsule.GetX2(), NearPosT);
+			const FVec3 NearPosDelta = SpherePos - NearPos;
+			const FReal NearPosDistanceSq = NearPosDelta.SizeSquared();
+			if (NearPosDistanceSq > SMALL_NUMBER)
 			{
-				Constraint.AddOneshotManifoldContact(ContactPoint);
+				if (NearPosDistanceSq < NetCullDistanceSq)
+				{
+					const FReal NearPosDistance = FMath::Sqrt(NearPosDistanceSq);
+					const FVec3 NearPosDir = NearPosDelta / NearPosDistance;
+					const FReal NearPhi = NearPosDistance - Sphere.GetRadius() - Capsule.GetRadius() - Constraint.GetRestitutionPadding();
+
+					FContactPoint NearContactPoint;
+					NearContactPoint.ShapeContactPoints[0] = SphereToCapsuleTransform.InverseTransformPositionNoScale(SpherePos - Sphere.GetRadius() * NearPosDir);
+					NearContactPoint.ShapeContactPoints[1] = NearPos + Capsule.GetRadius() * NearPosDir;
+					NearContactPoint.ShapeContactNormal = NearPosDir;
+					NearContactPoint.Phi = NearPhi;
+					NearContactPoint.FaceIndex = INDEX_NONE;
+					NearContactPoint.ContactType = EContactPointType::VertexPlane;
+					Constraint.AddOneshotManifoldContact(NearContactPoint);
+
+					// If we have a small sphere, just stick with the 1-point manifold
+					const FReal SphereCapsuleSizeThreshold = FReal(Chaos_Collision_Manifold_SphereCapsuleSizeThreshold);
+					if (Sphere.GetRadius() < SphereCapsuleSizeThreshold * (Capsule.GetHeight() + Capsule.GetRadius()))
+					{
+						return;
+					}
+
+					// If the capsule is non-dynamic there's no point in creating the multipoint manifold
+					if (!FConstGenericParticleHandle(Constraint.GetParticle1())->IsDynamic())
+					{
+						return;
+					}
+
+					// Now add the two end caps
+					// Calculate the vector orthogonal to the capsule axis that gives the nearest points on the capsule cyclinder to the sphere
+					// The initial length will be proportional to the sine of the angle between the axis and the delta position and will approach
+					// zero when the capsule is end-on to the sphere, in which case we won't add the end caps.
+					constexpr FReal EndCapSinAngleThreshold = FReal(0.35);	// about 20deg
+					constexpr FReal EndCapDistanceThreshold = FReal(0.2);	// fraction
+					FVec3 CapsuleOrthogonal = FVec3::CrossProduct(Capsule.GetAxis(), FVec3::CrossProduct(Capsule.GetAxis(), SpherePos - Capsule.GetCenter()));
+					const FReal CapsuleOrthogonalLenSq = CapsuleOrthogonal.SizeSquared();
+					if (CapsuleOrthogonalLenSq > FMath::Square(EndCapSinAngleThreshold))
+					{
+						// Orthogonal must point towards the sphere, but currently depends on the relative axis orientation
+						const FReal CapsuleOrthogonalLen = FMath::Sqrt(CapsuleOrthogonalLenSq);
+						CapsuleOrthogonal = CapsuleOrthogonal / CapsuleOrthogonalLen;
+						if (FVec3::DotProduct(CapsuleOrthogonal, SpherePos - Capsule.GetCenter()) < FReal(0))
+						{
+							CapsuleOrthogonal = -CapsuleOrthogonal;
+						}
+
+						if (NearPosT > EndCapDistanceThreshold)
+						{
+							const FVec3 EndCapPos0 = Capsule.GetX1() + CapsuleOrthogonal * Capsule.GetRadius();
+							const FReal EndCapDistance0 = (SpherePos - EndCapPos0).Size();
+							const FReal EndCapPhi0 = EndCapDistance0 - Sphere.GetRadius() - Constraint.GetRestitutionPadding();
+							
+							if (EndCapPhi0 < Constraint.GetCullDistance())
+							{
+								const FVec3 EndCapPosDir0 = (SpherePos - EndCapPos0) / EndCapDistance0;
+								const FVec3 SpherePos0 = SpherePos - EndCapPosDir0 * Sphere.GetRadius();
+						
+								FContactPoint EndCapContactPoint0;
+								EndCapContactPoint0.ShapeContactPoints[0] = SphereToCapsuleTransform.InverseTransformPositionNoScale(SpherePos0);
+								EndCapContactPoint0.ShapeContactPoints[1] = EndCapPos0;
+								EndCapContactPoint0.ShapeContactNormal = EndCapPosDir0;
+								EndCapContactPoint0.Phi = EndCapPhi0;
+								EndCapContactPoint0.FaceIndex = INDEX_NONE;
+								EndCapContactPoint0.ContactType = EContactPointType::VertexPlane;
+								Constraint.AddOneshotManifoldContact(EndCapContactPoint0);
+							}
+						}
+
+						if (NearPosT < FReal(1) - EndCapDistanceThreshold)
+						{
+							const FVec3 EndCapPos1 = Capsule.GetX2() + CapsuleOrthogonal * Capsule.GetRadius();
+							const FReal EndCapDistance1 = (SpherePos - EndCapPos1).Size();
+							const FReal EndCapPhi1 = EndCapDistance1 - Sphere.GetRadius() - Constraint.GetRestitutionPadding();
+
+							if (EndCapPhi1 < Constraint.GetCullDistance())
+							{
+								const FVec3 EndCapPosDir1 = (SpherePos - EndCapPos1) / EndCapDistance1;
+								const FVec3 SpherePos1 = SpherePos - EndCapPosDir1 * Sphere.GetRadius();
+
+								FContactPoint EndCapContactPoint0;
+								EndCapContactPoint0.ShapeContactPoints[0] = SphereToCapsuleTransform.InverseTransformPositionNoScale(SpherePos1);
+								EndCapContactPoint0.ShapeContactPoints[1] = EndCapPos1;
+								EndCapContactPoint0.ShapeContactNormal = EndCapPosDir1;
+								EndCapContactPoint0.Phi = EndCapPhi1;
+								EndCapContactPoint0.FaceIndex = INDEX_NONE;
+								EndCapContactPoint0.ContactType = EContactPointType::VertexPlane;
+								Constraint.AddOneshotManifoldContact(EndCapContactPoint0);
+							}
+						}
+					}
+				}
 			}
 		}
 
diff --git a/Engine/Source/Runtime/Experimental/Chaos/Public/Chaos/Utilities.h b/Engine/Source/Runtime/Experimental/Chaos/Public/Chaos/Utilities.h
--- a/Engine/Source/Runtime/Experimental/Chaos/Public/Chaos/Utilities.h
+++ b/Engine/Source/Runtime/Experimental/Chaos/Public/Chaos/Utilities.h
@@ -563,6 +563,29 @@ namespace Chaos
 			C2 = P2 + T * D2;
 		}
 
+		inline FReal ClosestTimeOnLineSegment(const FVector& Point, const FVector& StartPoint, const FVector& EndPoint)
+		{
+			const FVector Segment = EndPoint - StartPoint;
+			const FVector VectToPoint = Point - StartPoint;
+
+			// See if closest point is before StartPoint
+			const FVector::FReal Dot1 = FVec3::DotProduct(VectToPoint, Segment);
+			if (Dot1 <= 0)
+			{
+				return FReal(0);
+			}
+
+			// See if closest point is beyond EndPoint
+			const FVector::FReal Dot2 = FVec3::DotProduct(Segment, Segment);
+			if (Dot2 <= Dot1)
+			{
+				return FReal(1);
+			}
+
+			// Closest Point is within segment
+			return Dot1 / Dot2;
+		}
+
 
 	} // namespace Utilities
 } // namespace Chaos
\ No newline at end of file
