speculative fix (to avoid crash) for convex builder failure -- just build an empty hull in the failure cases rather than crash

diff --git a/Engine/Source/Runtime/Experimental/Chaos/Public/Chaos/CollisionConvexMesh.h b/Engine/Source/Runtime/Experimental/Chaos/Public/Chaos/CollisionConvexMesh.h
--- a/Engine/Source/Runtime/Experimental/Chaos/Public/Chaos/CollisionConvexMesh.h
+++ b/Engine/Source/Runtime/Experimental/Chaos/Public/Chaos/CollisionConvexMesh.h
@@ -307,7 +307,13 @@ namespace Chaos
 #endif
 #endif
 
-				AddVertex(Pool, InVertices, ConflictV, InParams);
+				if (!AddVertex(Pool, InVertices, ConflictV, InParams))
+				{
+					// AddVertex failed to process the conflict vertex -- 
+					// subsequent calls to FindConflictVertex will just keep finding the same one,
+					// so all we can do here is fail to build a convex hull at all
+					return;
+				}
 				ConflictV = FindConflictVertex(InVertices, DummyFace->Next);
 			}
 
@@ -962,10 +968,13 @@ namespace Chaos
 
 		}
 
-		static void BuildFaces(FMemPool& Pool, const TArray<FVec3Type>& InVertices, const FHalfEdge* ConflictV, const TArray<FHalfEdge*>& HorizonEdges, const TArray<FConvexFace*> OldFaces, TArray<FConvexFace*>& NewFaces)
+		static bool BuildFaces(FMemPool& Pool, const TArray<FVec3Type>& InVertices, const FHalfEdge* ConflictV, const TArray<FHalfEdge*>& HorizonEdges, const TArray<FConvexFace*>& OldFaces, TArray<FConvexFace*>& NewFaces)
 		{
 			//The HorizonEdges are in CCW order. We must make new faces and edges to join from ConflictV to these edges
-			check(HorizonEdges.Num() >= 3);
+			if (!(HorizonEdges.Num() >= 3)) // TODO: previously this was a check(), but it sometimes failed; can we fix things so this always holds?
+			{
+				return false;
+			}
 			NewFaces.Reserve(HorizonEdges.Num());
 			FHalfEdge* PrevEdge = nullptr;
 			for(int32 HorizonIdx = 0; HorizonIdx < HorizonEdges.Num(); ++HorizonIdx)
@@ -975,7 +984,11 @@ namespace Chaos
 				NewHorizonEdge->Twin = OriginalEdge->Twin; //swap edges
 				NewHorizonEdge->Twin->Twin = NewHorizonEdge;
 				FHalfEdge* HorizonNext = Pool.AllocHalfEdge(OriginalEdge->Next->Vertex);
-				check(HorizonNext->Vertex == HorizonEdges[(HorizonIdx + 1) % HorizonEdges.Num()]->Vertex); //should be ordered properly
+				// TODO: previously this was a check(), but it sometimes failed; can we fix things so this always holds?
+				if (!(HorizonNext->Vertex == HorizonEdges[(HorizonIdx + 1) % HorizonEdges.Num()]->Vertex)) //should be ordered properly
+				{
+					return false;
+				}
 				FHalfEdge* V = Pool.AllocHalfEdge(ConflictV->Vertex);
 				V->Twin = PrevEdge;
 				if(PrevEdge)
@@ -998,7 +1011,11 @@ namespace Chaos
 				NewFaces.Add(NewFace);
 			}
 
-			check(PrevEdge);
+			// TODO: previously this was a check(); can we determine if this always hold and switch it back if so?
+			if (!PrevEdge)
+			{
+				return false;
+			}
 			NewFaces[0]->FirstEdge->Prev->Twin = PrevEdge;
 			PrevEdge->Twin = NewFaces[0]->FirstEdge->Prev;
 			NewFaces[NewFaces.Num() - 1]->Next = nullptr;
@@ -1020,9 +1037,11 @@ namespace Chaos
 			EndFace->Next = OldFace->Next;
 			OldFace->Next = StartFace;
 			StartFace->Prev = OldFace;
+
+			return true;
 		}
 
-		static void AddVertex(FMemPool& Pool, const TArray<FVec3Type>& InVertices, FHalfEdge* ConflictV, const Params& InParams)
+		static bool AddVertex(FMemPool& Pool, const TArray<FVec3Type>& InVertices, FHalfEdge* ConflictV, const Params& InParams)
 		{
 			UE_CLOG(DEBUG_HULL_GENERATION, LogChaos, VeryVerbose, TEXT("Adding Vertex %d"), ConflictV->Vertex);
 
@@ -1031,7 +1050,10 @@ namespace Chaos
 			BuildHorizon(InVertices, ConflictV, HorizonEdges, FacesToDelete, InParams);
 
 			TArray<FConvexFace*> NewFaces;
-			BuildFaces(Pool, InVertices, ConflictV, HorizonEdges, FacesToDelete, NewFaces);
+			if (!ensure(BuildFaces(Pool, InVertices, ConflictV, HorizonEdges, FacesToDelete, NewFaces)))
+			{
+				return false;
+			}
 
 #if DEBUG_HULL_GENERATION
 			FString NewFaceString(TEXT("\tNew Faces: "));
@@ -1071,6 +1093,8 @@ namespace Chaos
 
 			//todo(ocohen): need to explicitly test for merge failures. Coplaner, nonconvex, etc...
 			//getting this in as is for now to unblock other systems
+
+			return true;
 		}
 
 	};
