Geometry Collection : Size Specific Defaults
- Don't update to convex if the older file does not have convex properties.
- Fix nullptr access in the shapes debug drawing.

diff --git a/Engine/Source/Runtime/Experimental/Chaos/Private/Chaos/ChaosDebugDraw.cpp b/Engine/Source/Runtime/Experimental/Chaos/Private/Chaos/ChaosDebugDraw.cpp
--- a/Engine/Source/Runtime/Experimental/Chaos/Private/Chaos/ChaosDebugDraw.cpp
+++ b/Engine/Source/Runtime/Experimental/Chaos/Private/Chaos/ChaosDebugDraw.cpp
@@ -483,6 +483,9 @@ namespace Chaos
 
 		void DrawShapesImpl(const FGeometryParticleHandle* Particle, const FRigidTransform3& ShapeTransform, const FImplicitObject* Shape, const FReal Margin, const FColor& Color, const FChaosDebugDrawSettings& Settings)
 		{
+
+			if (!Particle || !Shape) return;
+
 			const EImplicitObjectType PackedType = Shape->GetType(); // Type includes scaling and instancing data
 			const EImplicitObjectType InnerType = GetInnerType(Shape->GetType());
 
diff --git a/Engine/Source/Runtime/Experimental/Chaos/Private/GeometryCollection/GeometryCollectionSizeSpecificUtility.cpp b/Engine/Source/Runtime/Experimental/Chaos/Private/GeometryCollection/GeometryCollectionSizeSpecificUtility.cpp
new file mode 100644
--- /dev/null
+++ b/Engine/Source/Runtime/Experimental/Chaos/Private/GeometryCollection/GeometryCollectionSizeSpecificUtility.cpp
@@ -0,0 +1,43 @@
+// Copyright Epic Games, Inc. All Rights Reserved.
+
+/*=============================================================================
+	GeometryCollection->cpp: FGeometryCollection methods.
+=============================================================================*/
+
+#include "GeometryCollection/GeometryCollectionSizeSpecificUtility.h"
+#include "GeometryCollection/GeometryCollectionAlgo.h"
+
+#include "CoreMinimal.h"
+#include "Templates/SharedPointer.h"
+#include "Chaos/Real.h"
+
+DEFINE_LOG_CATEGORY_STATIC(FGeometryCollectionSizeSpecificUtilityLogging, Log, All);
+
+namespace GeometryCollection
+{
+	namespace SizeSpecific
+	{
+		int32 FindIndexForVolume(const TArray<FSharedSimulationSizeSpecificData>& SizeSpecificData, Chaos::FReal Volume)
+		{
+			check(SizeSpecificData.Num());
+			int32 UseIdx = 0;
+			float PreSize = FLT_MAX;
+			for (int32 Idx = SizeSpecificData.Num() - 1; Idx >= 0; --Idx)
+			{
+				ensureMsgf(PreSize >= SizeSpecificData[Idx].MaxSize, TEXT("SizeSpecificData is not sorted"));
+				PreSize = SizeSpecificData[Idx].MaxSize;
+				if (Volume < SizeSpecificData[Idx].MaxSize)
+					UseIdx = Idx;
+				else
+					break;
+			}
+			return UseIdx;
+		}
+
+		int32 FindIndexForVolume(const TArray<FSharedSimulationSizeSpecificData>& SizeSpecificData, const FBox& Bounds)
+		{
+			return FindIndexForVolume(SizeSpecificData, Bounds.GetVolume());
+		}
+
+	}
+}
\ No newline at end of file
diff --git a/Engine/Source/Runtime/Experimental/Chaos/Private/PhysicsProxy/GeometryCollectionPhysicsProxy.cpp b/Engine/Source/Runtime/Experimental/Chaos/Private/PhysicsProxy/GeometryCollectionPhysicsProxy.cpp
--- a/Engine/Source/Runtime/Experimental/Chaos/Private/PhysicsProxy/GeometryCollectionPhysicsProxy.cpp
+++ b/Engine/Source/Runtime/Experimental/Chaos/Private/PhysicsProxy/GeometryCollectionPhysicsProxy.cpp
@@ -24,6 +24,7 @@
 #include "Chaos/PBDRigidClustering.h"
 #include "GeometryCollection/GeometryCollectionAlgo.h"
 #include "GeometryCollection/GeometryCollectionUtility.h"
+#include "GeometryCollection/GeometryCollectionSizeSpecificUtility.h"
 #include "GeometryCollection/GeometryCollectionSimulationTypes.h"
 #include "GeometryCollection/ManagedArrayCollection.h"
 #include "Modules/ModuleManager.h"
@@ -921,27 +922,7 @@ FAutoConsoleVariableRef CVarReportNoLevelsetCluster(TEXT("p.gc.ReportNoLevelsetC
 DECLARE_CYCLE_STAT(TEXT("FGeometryCollectionPhysicsProxy::BuildClusters"), STAT_BuildClusters, STATGROUP_Chaos);
 DECLARE_CYCLE_STAT(TEXT("FGeometryCollectionPhysicsProxy::BuildClusters:GlobalMatrices"), STAT_BuildClustersGlobalMatrices, STATGROUP_Chaos);
 
-int32 FindSizeSpecificIdx(const TArray<FSharedSimulationSizeSpecificData>& SizeSpecificData, Chaos::FReal Volume)
-{
-	check(SizeSpecificData.Num());
-	int32 UseIdx = 0;
-	float PreSize = FLT_MAX;
-	for (int32 Idx = SizeSpecificData.Num() - 1; Idx >= 0; --Idx)
-	{
-		ensureMsgf(PreSize >= SizeSpecificData[Idx].MaxSize, TEXT("SizeSpecificData is not sorted"));
-		PreSize = SizeSpecificData[Idx].MaxSize;
-		if (Volume < SizeSpecificData[Idx].MaxSize)
-			UseIdx = Idx;
-		else
-			break;
-	}
-	return UseIdx;
-}
 
-int32 FindSizeSpecificIdx(const TArray<FSharedSimulationSizeSpecificData>& SizeSpecificData, const FBox& Bounds)
-{
-	return FindSizeSpecificIdx(SizeSpecificData, Bounds.GetVolume());
-}
 
 Chaos::TPBDRigidClusteredParticleHandle<Chaos::FReal, 3>*
 FGeometryCollectionPhysicsProxy::BuildClusters(
@@ -1058,7 +1039,7 @@ FGeometryCollectionPhysicsProxy::BuildClusters(
 		if (Parameters.RestCollection->HasAttribute("Size", FTransformCollection::TransformGroup))
 		{
 			const TManagedArray<float>& RelativeSize = Parameters.RestCollection->GetAttribute<float>("Size", FTransformCollection::TransformGroup);
-			SizeSpecificIdx = FindSizeSpecificIdx(Parameters.Shared.SizeSpecificData, RelativeSize[CollectionClusterIndex]);
+			SizeSpecificIdx = GeometryCollection::SizeSpecific::FindIndexForVolume(Parameters.Shared.SizeSpecificData, RelativeSize[CollectionClusterIndex]);
 		}
 		else
 		{
@@ -1066,7 +1047,7 @@ FGeometryCollectionPhysicsProxy::BuildClusters(
 			if (Implicit[CollectionClusterIndex] && Implicit[CollectionClusterIndex]->HasBoundingBox())
 			{
 				FBox LocalBoundingBox(Implicit[CollectionClusterIndex]->BoundingBox().Min(), Implicit[CollectionClusterIndex]->BoundingBox().Max());
-				SizeSpecificIdx = FindSizeSpecificIdx(Parameters.Shared.SizeSpecificData, LocalBoundingBox);
+				SizeSpecificIdx = GeometryCollection::SizeSpecific::FindIndexForVolume(Parameters.Shared.SizeSpecificData, LocalBoundingBox);
 			}
 		}
 
@@ -2160,11 +2141,11 @@ void FGeometryCollectionPhysicsProxy::InitializeSharedCollisionStructures(
 			if (bUseRelativeSize)
 			{
 				const TManagedArray<float>& RelativeSize = RestCollection.GetAttribute<float>(TEXT("Size"), FTransformCollection::TransformGroup);
-				SizeSpecificIdx = FindSizeSpecificIdx(SharedParams.SizeSpecificData, RelativeSize[TransformGroupIndex]);
+				SizeSpecificIdx = GeometryCollection::SizeSpecific::FindIndexForVolume(SharedParams.SizeSpecificData, RelativeSize[TransformGroupIndex]);
 			}
 			else
 			{
-				SizeSpecificIdx = FindSizeSpecificIdx(SharedParams.SizeSpecificData, InstanceBoundingBox);
+				SizeSpecificIdx = GeometryCollection::SizeSpecific::FindIndexForVolume(SharedParams.SizeSpecificData, InstanceBoundingBox);
 			}
 			
 			const FSharedSimulationSizeSpecificData& SizeSpecificData = SharedParams.SizeSpecificData[SizeSpecificIdx];
@@ -2427,7 +2408,7 @@ void FGeometryCollectionPhysicsProxy::InitializeSharedCollisionStructures(
 				CollectionInertiaTensor[ClusterTransformIdx] = (FVector3f)InertiaDiagonal;	// LWC_TODO: Precision loss
 				CollectionMass[ClusterTransformIdx] = (FRealSingle)CollectionSpaceParticles->M(ClusterTransformIdx);
 
-				const int32 SizeSpecificIdx = FindSizeSpecificIdx(SharedParams.SizeSpecificData, InstanceBoundingBox);
+				const int32 SizeSpecificIdx = GeometryCollection::SizeSpecific::FindIndexForVolume(SharedParams.SizeSpecificData, InstanceBoundingBox);
 				const FSharedSimulationSizeSpecificData& SizeSpecificData = SharedParams.SizeSpecificData[SizeSpecificIdx];
 
 				if (SizeSpecificData.CollisionShapesData.Num())
diff --git a/Engine/Source/Runtime/Experimental/Chaos/Public/GeometryCollection/GeometryCollectionSizeSpecificUtility.h b/Engine/Source/Runtime/Experimental/Chaos/Public/GeometryCollection/GeometryCollectionSizeSpecificUtility.h
new file mode 100644
--- /dev/null
+++ b/Engine/Source/Runtime/Experimental/Chaos/Public/GeometryCollection/GeometryCollectionSizeSpecificUtility.h
@@ -0,0 +1,25 @@
+// Copyright Epic Games, Inc. All Rights Reserved.
+
+#pragma once
+
+#include "CoreMinimal.h"
+#include "GeometryCollection/GeometryCollection.h"
+#include "GeometryCollection/GeometryCollectionSimulationCoreTypes.h"
+#include "Templates/SharedPointer.h"
+
+namespace GeometryCollection 
+{
+	namespace SizeSpecific
+	{
+		//
+		//  Given a volume, return the first index of the size specific data that matches this volume. 
+		//
+		int32 CHAOS_API FindIndexForVolume(const TArray<FSharedSimulationSizeSpecificData>& SizeSpecificData, Chaos::FReal Volume);
+
+		//
+		//  Given a FBox, return the first index of the size specific data that matches this volume. 
+		//
+		int32 CHAOS_API FindIndexForVolume(const TArray<FSharedSimulationSizeSpecificData>& SizeSpecificData, const FBox& Bounds);
+
+	}
+}
\ No newline at end of file
diff --git a/Engine/Source/Runtime/Experimental/GeometryCollectionEngine/Private/GeometryCollection/GeometryCollectionEngineSizeSpecificUtility.cpp b/Engine/Source/Runtime/Experimental/GeometryCollectionEngine/Private/GeometryCollection/GeometryCollectionEngineSizeSpecificUtility.cpp
new file mode 100644
--- /dev/null
+++ b/Engine/Source/Runtime/Experimental/GeometryCollectionEngine/Private/GeometryCollection/GeometryCollectionEngineSizeSpecificUtility.cpp
@@ -0,0 +1,40 @@
+// Copyright Epic Games, Inc. All Rights Reserved.
+
+#include "GeometryCollection/GeometryCollectionEngineSizeSpecificUtility.h"
+
+namespace GeometryCollection
+{
+	namespace SizeSpecific
+	{
+		bool UsesImplicitCollisionType(const TArray<FGeometryCollectionSizeSpecificData>& SizeSpecificData, EImplicitTypeEnum ImplicitType)
+		{
+			for (int32 Idx = SizeSpecificData.Num() - 1; Idx >= 0; --Idx)
+			{
+				for (int Cdx = 0; Cdx < SizeSpecificData[Idx].CollisionShapes.Num(); Cdx++)
+				{
+					if (SizeSpecificData[Idx].CollisionShapes[Cdx].ImplicitType == ImplicitType)
+					{
+						return true;
+					}
+				}
+			}
+			return false;
+		}
+
+		void SetImplicitCollisionType(TArray<FGeometryCollectionSizeSpecificData>& SizeSpecificData, EImplicitTypeEnum ImplicitTypeTo, EImplicitTypeEnum ImplicitTypeFrom)
+		{
+			for (int32 Idx = SizeSpecificData.Num() - 1; Idx >= 0; --Idx)
+			{
+				for (int Cdx = 0; Cdx < SizeSpecificData[Idx].CollisionShapes.Num(); Cdx++)
+				{
+					if (SizeSpecificData[Idx].CollisionShapes[Cdx].ImplicitType == ImplicitTypeFrom ||
+						ImplicitTypeFrom == EImplicitTypeEnum::Chaos_Implicit_None)
+					{
+						SizeSpecificData[Idx].CollisionShapes[Cdx].ImplicitType = ImplicitTypeTo;
+					}
+				}
+			}
+		}
+
+	}
+}
\ No newline at end of file
diff --git a/Engine/Source/Runtime/Experimental/GeometryCollectionEngine/Private/GeometryCollection/GeometryCollectionObject.cpp b/Engine/Source/Runtime/Experimental/GeometryCollectionEngine/Private/GeometryCollection/GeometryCollectionObject.cpp
--- a/Engine/Source/Runtime/Experimental/GeometryCollectionEngine/Private/GeometryCollection/GeometryCollectionObject.cpp
+++ b/Engine/Source/Runtime/Experimental/GeometryCollectionEngine/Private/GeometryCollection/GeometryCollectionObject.cpp
@@ -22,6 +22,7 @@
 #if WITH_EDITOR
 #include "GeometryCollection/DerivedDataGeometryCollectionCooker.h"
 #include "GeometryCollection/GeometryCollectionConvexUtility.h"
+#include "GeometryCollection/GeometryCollectionEngineSizeSpecificUtility.h"
 #include "GeometryCollection/GeometryCollectionComponent.h"
 #include "DerivedDataCacheInterface.h"
 #include "Serialization/MemoryReader.h"
@@ -798,11 +799,10 @@ void UGeometryCollection::Serialize(FArchive& Ar)
 #if WITH_EDITOR
 		if (bGeometryCollectionEnableForcedConvexGenerationInSerialize)
 		{
-			if (!FGeometryCollectionConvexUtility::HasConvexHullData(GeometryCollection.Get()))
+			if (!FGeometryCollectionConvexUtility::HasConvexHullData(GeometryCollection.Get()) &&
+				GeometryCollection::SizeSpecific::UsesImplicitCollisionType(SizeSpecificData, EImplicitTypeEnum::Chaos_Implicit_Convex))
 			{
-				UE_LOG(LogGeometryCollectionInternal, Warning, TEXT("Regenerated GeoemtryCollections convex geometry on asset (%s). Please resave package %s."), *GetName(), *GetOutermost()->GetPathName());
-				FGeometryCollectionConvexPropertiesInterface::FConvexCreationProperties ConvexProperties = GeometryCollection->GetConvexProperties();
-				FGeometryCollectionConvexUtility::CreateNonOverlappingConvexHullData(GeometryCollection.Get(), ConvexProperties.FractionRemove, ConvexProperties.SimplificationThreshold, ConvexProperties.CanExceedFraction);
+				GeometryCollection::SizeSpecific::SetImplicitCollisionType(SizeSpecificData, EImplicitTypeEnum::Chaos_Implicit_Box, EImplicitTypeEnum::Chaos_Implicit_Convex);
 				bCreateSimulationData = true;
 				InvalidateCollection();
 			}
diff --git a/Engine/Source/Runtime/Experimental/GeometryCollectionEngine/Public/GeometryCollection/GeometryCollectionEngineSizeSpecificUtility.h b/Engine/Source/Runtime/Experimental/GeometryCollectionEngine/Public/GeometryCollection/GeometryCollectionEngineSizeSpecificUtility.h
new file mode 100644
--- /dev/null
+++ b/Engine/Source/Runtime/Experimental/GeometryCollectionEngine/Public/GeometryCollection/GeometryCollectionEngineSizeSpecificUtility.h
@@ -0,0 +1,23 @@
+// Copyright Epic Games, Inc. All Rights Reserved.
+
+#pragma once
+
+#include "CoreMinimal.h"
+#include "GeometryCollection/GeometryCollectionObject.h"
+
+namespace GeometryCollection
+{
+	namespace SizeSpecific
+	{
+		//
+		// Return if true/false if this size specific object uses the specified collision type. 
+		//
+		bool GEOMETRYCOLLECTIONENGINE_API UsesImplicitCollisionType(const TArray<FGeometryCollectionSizeSpecificData>& SizeSpecificData, EImplicitTypeEnum ImplicitType);
+
+		//
+		// Set the collision types for a size specific data, also will only set the types matching ImplicitTypeFrom if specified.
+		//
+		void GEOMETRYCOLLECTIONENGINE_API SetImplicitCollisionType(TArray<FGeometryCollectionSizeSpecificData>& SizeSpecificData, EImplicitTypeEnum ImplicitTypeTo, EImplicitTypeEnum ImplicitTypeFrom = EImplicitTypeEnum::Chaos_Implicit_None);
+
+	}
+}
\ No newline at end of file
