Chaos Cloth Optimization - Use atomic instead of critical section when building self collision constraints

We were hitting hitches due to lock contention when there were a lot of self collisions. Switch to preallocating the arrays to be large enough to accommodate the maximum number of constraints and using an atomic to determine the write index (and then shrinking the array at the end).

diff --git a/Engine/Source/Runtime/Experimental/Chaos/Private/Chaos/PBDCollisionSpringConstraintsBase.cpp b/Engine/Source/Runtime/Experimental/Chaos/Private/Chaos/PBDCollisionSpringConstraintsBase.cpp
--- a/Engine/Source/Runtime/Experimental/Chaos/Private/Chaos/PBDCollisionSpringConstraintsBase.cpp
+++ b/Engine/Source/Runtime/Experimental/Chaos/Private/Chaos/PBDCollisionSpringConstraintsBase.cpp
@@ -6,6 +6,7 @@
 
 #if !COMPILE_WITHOUT_UNREAL_SUPPORT
 #include "Chaos/Framework/Parallel.h"
+#include <atomic>
 
 namespace Chaos::Softs {
 
@@ -35,71 +36,84 @@ void FPBDCollisionSpringConstraintsBase::Init(const FSolverParticles& Particles)
 		return;
 	}
 
-	Constraints.Reset();
-	Barys.Reset();
-
 	FTriangleMesh::TBVHType<FSolverReal> BVH;
-	TriangleMesh.BuildBVH(static_cast<const TArrayView<const FSolverVec3>&>(Particles.XArray()), BVH);
-
-	FCriticalSection CriticalSection;
-	const FSolverReal HeightSq = FMath::Square(Thickness + Thickness);
-	PhysicsParallelFor(NumParticles,
-		[this, &BVH, &Particles, &CriticalSection, HeightSq](int32 i)
-		{
-			const int32 Index = i + Offset;
-
-			TArray< TTriangleCollisionPoint<FSolverReal> > Result;
-			if (TriangleMesh.PointProximityQuery(BVH, static_cast<const TArrayView<const FSolverVec3>&>(Particles.XArray()), Index, Particles.X(Index), Thickness, Thickness,
-				[this](const int32 PointIndex, const int32 TriangleIndex)->bool
-				{
-					const TVector<int32, 3>& Elem = Elements[TriangleIndex];
-					if (DisabledCollisionElements.Contains({ PointIndex, Elem[0] }) ||
-						DisabledCollisionElements.Contains({ PointIndex, Elem[1] }) ||
-						DisabledCollisionElements.Contains({ PointIndex, Elem[2] }))
-					{
-						return false;
-					}
-					return true;
-				},
-				Result))
+	{
+		TRACE_CPUPROFILER_EVENT_SCOPE(ChaosPBDCollisionSpring_BuildBVH);
+		TriangleMesh.BuildBVH(static_cast<const TArrayView<const FSolverVec3>&>(Particles.XArray()), BVH);
+	}
+	{
+		TRACE_CPUPROFILER_EVENT_SCOPE(ChaosPBDCollisionSpring_ProximityQuery);
+
+		// Preallocate enough space for all possible connections.
+		constexpr int32 MaxConnectionsPerPoint = 3;
+		Constraints.SetNum(NumParticles * MaxConnectionsPerPoint);
+		Barys.SetNum(NumParticles * MaxConnectionsPerPoint);
+
+		std::atomic<int32> ConstraintIndex(0);
+
+		const FSolverReal HeightSq = FMath::Square(Thickness + Thickness);
+		PhysicsParallelFor(NumParticles,
+			[this, &BVH, &Particles, &ConstraintIndex, HeightSq, MaxConnectionsPerPoint](int32 i)
 			{
+				const int32 Index = i + Offset;
 
-				constexpr int32 MaxConnectionsPerPoint = 3;
-				if (Result.Num() > MaxConnectionsPerPoint)
-				{
-					// TODO: once we have a PartialSort, use that instead here.
-					Result.Sort(
-						[](const TTriangleCollisionPoint<FSolverReal>& First, const TTriangleCollisionPoint<FSolverReal>& Second)->bool
+				TArray< TTriangleCollisionPoint<FSolverReal> > Result;
+				if (TriangleMesh.PointProximityQuery(BVH, static_cast<const TArrayView<const FSolverVec3>&>(Particles.XArray()), Index, Particles.X(Index), Thickness, Thickness,
+					[this](const int32 PointIndex, const int32 TriangleIndex)->bool
+					{
+						const TVector<int32, 3>& Elem = Elements[TriangleIndex];
+						if (DisabledCollisionElements.Contains({ PointIndex, Elem[0] }) ||
+							DisabledCollisionElements.Contains({ PointIndex, Elem[1] }) ||
+							DisabledCollisionElements.Contains({ PointIndex, Elem[2] }))
 						{
-							return First.Phi < Second.Phi;
+							return false;
 						}
-					);
-					Result.SetNum(MaxConnectionsPerPoint, false /*bAllowShrinking*/);
-				}
-
-				for (const TTriangleCollisionPoint<FSolverReal>& CollisionPoint : Result)
+						return true;
+					},
+					Result))
 				{
-					const TVector<int32, 3>& Elem = Elements[CollisionPoint.Indices[1]];
-					if (ReferencePositions)
+
+					if (Result.Num() > MaxConnectionsPerPoint)
+					{
+						// TODO: once we have a PartialSort, use that instead here.
+						Result.Sort(
+							[](const TTriangleCollisionPoint<FSolverReal>& First, const TTriangleCollisionPoint<FSolverReal>& Second)->bool
+							{
+								return First.Phi < Second.Phi;
+							}
+						);
+						Result.SetNum(MaxConnectionsPerPoint, false /*bAllowShrinking*/);
+					}
+
+					for (const TTriangleCollisionPoint<FSolverReal>& CollisionPoint : Result)
 					{
-						const FSolverVec3& RefP = (*ReferencePositions)[Index];
-						const FSolverVec3& RefP0 = (*ReferencePositions)[Elem[0]];
-						const FSolverVec3& RefP1 = (*ReferencePositions)[Elem[1]];
-						const FSolverVec3& RefP2 = (*ReferencePositions)[Elem[2]];
-						const FSolverVec3 RefDiff = RefP - CollisionPoint.Bary[1] * RefP0 - CollisionPoint.Bary[2] * RefP1 - CollisionPoint.Bary[3] * RefP2;
-						if (RefDiff.SizeSquared() < HeightSq)
+						const TVector<int32, 3>& Elem = Elements[CollisionPoint.Indices[1]];
+						if (ReferencePositions)
 						{
-							continue;
+							const FSolverVec3& RefP = (*ReferencePositions)[Index];
+							const FSolverVec3& RefP0 = (*ReferencePositions)[Elem[0]];
+							const FSolverVec3& RefP1 = (*ReferencePositions)[Elem[1]];
+							const FSolverVec3& RefP2 = (*ReferencePositions)[Elem[2]];
+							const FSolverVec3 RefDiff = RefP - CollisionPoint.Bary[1] * RefP0 - CollisionPoint.Bary[2] * RefP1 - CollisionPoint.Bary[3] * RefP2;
+							if (RefDiff.SizeSquared() < HeightSq)
+							{
+								continue;
+							}
 						}
+						const int32 IndexToWrite = ConstraintIndex.fetch_add(1);
+
+						Constraints[IndexToWrite] = { Index, Elem[0], Elem[1], Elem[2] };
+						Barys[IndexToWrite] = { CollisionPoint.Bary[1], CollisionPoint.Bary[2], CollisionPoint.Bary[3] };
 					}
-					CriticalSection.Lock();
-					Constraints.Add({ Index, Elem[0], Elem[1], Elem[2] });
-					Barys.Add({ CollisionPoint.Bary[1], CollisionPoint.Bary[2], CollisionPoint.Bary[3] });
-					CriticalSection.Unlock();
 				}
 			}
-		}
-	);
+		);
+
+		// Shrink the arrays to the actual number of found constraints.
+		const int32 ConstraintNum = ConstraintIndex.load();
+		Constraints.SetNum(ConstraintNum, /*bAllowShrinking*/ true);
+		Barys.SetNum(ConstraintNum, /*bAllowShrinking*/ true);
+	}
 }
 
 FSolverVec3 FPBDCollisionSpringConstraintsBase::GetDelta(const FSolverParticles& Particles, const int32 i) const
