fix issues with selection / clustering tools interaction with fracture mode 'view level', causing invalid selections of bones that are not visible

diff --git a/Engine/Plugins/Experimental/ChaosEditor/Source/FractureEditor/Private/FractureEditorMode.cpp b/Engine/Plugins/Experimental/ChaosEditor/Source/FractureEditor/Private/FractureEditorMode.cpp
--- a/Engine/Plugins/Experimental/ChaosEditor/Source/FractureEditor/Private/FractureEditorMode.cpp
+++ b/Engine/Plugins/Experimental/ChaosEditor/Source/FractureEditor/Private/FractureEditorMode.cpp
@@ -293,27 +293,31 @@ bool UFractureEditorMode::UpdateSelectionInFrustum(const FConvexVolume& InFrustu
 
 	FConvexVolume SelectionFrustum(TranformFrustum(InFrustum, InvActorMatrix));
 
-	TMap<int32, FBox> BoundsToBone;
+	TArray<FBox> BoundsPerBone;
 	TArray<int32> SelectedBonesArray;
 	bool bSelectionWasUpdated = false;
 
 	for (UGeometryCollectionComponent* GeometryCollectionComponent : GeometryComponents)
 	{
-		BoundsToBone.Reset();
+		BoundsPerBone.Reset();
 		SelectedBonesArray.Reset();
 
-		GetComponentGlobalBounds(GeometryCollectionComponent, BoundsToBone);
+		GetComponentGlobalBounds(GeometryCollectionComponent, BoundsPerBone);
 
-		for (const TPair<int32, FBox>& Bone : BoundsToBone)
+		for (int32 Idx = 0; Idx < BoundsPerBone.Num(); Idx++)
 		{
-			const FBox& TransformedBoneBox = Bone.Value;
+			const FBox& TransformedBoneBox = BoundsPerBone[Idx];
+			if (!TransformedBoneBox.IsValid)
+			{
+				continue;
+			}
 			bool bFullyContained = false;
 			bool bIntersected = SelectionFrustum.IntersectBox(TransformedBoneBox.GetCenter(), TransformedBoneBox.GetExtent(), bFullyContained);
 			if (bIntersected)
 			{
 				if (!bStrictDragSelection || (bFullyContained && bStrictDragSelection))
 				{
-					SelectedBonesArray.Add(Bone.Key);
+					SelectedBonesArray.Add(Idx);
 				}
 			}
 		}
@@ -324,7 +328,8 @@ bool UFractureEditorMode::UpdateSelectionInFrustum(const FConvexVolume& InFrustu
 			FScopedColorEdit ColorEdit = GeometryCollectionComponent->EditBoneSelection();
 			ColorEdit.SelectBones(GeometryCollection::ESelectionMode::None);
 			ColorEdit.SetSelectedBones(SelectedBonesArray);
-			ColorEdit.SetHighlightedBones(SelectedBonesArray);
+			ColorEdit.FilterSelectionToLevel(/*bPreferLowestOnly*/true);
+			ColorEdit.SetHighlightedBones(ColorEdit.GetSelectedBones(), true);
 
 			FFractureEditorModeToolkit* FractureToolkit = (FFractureEditorModeToolkit*)Toolkit.Get();
 
@@ -397,16 +402,13 @@ bool UFractureEditorMode::ComputeBoundingBoxForViewportFocus(AActor* Actor, UPri
 	UGeometryCollectionComponent* GeometryCollectionComponent = Cast<UGeometryCollectionComponent>(PrimitiveComponent);
 	if (GeometryCollectionComponent && GeometryCollectionComponent->GetSelectedBones().Num() > 0 && SelectedGeometryComponents.Contains(GeometryCollectionComponent))
 	{
-		TMap<int32, FBox> BoundsToBoneMap;
-		GetComponentGlobalBounds(GeometryCollectionComponent, BoundsToBoneMap);
+		TArray<FBox> BoundsPerBone;
+		GetComponentGlobalBounds(GeometryCollectionComponent, BoundsPerBone);
 
 		FBox TotalBoneBox(ForceInit);
 		for (int32 BoneIndex : GeometryCollectionComponent->GetSelectedBones())
 		{
-			if (FBox* LocalBoneBox = BoundsToBoneMap.Find(BoneIndex))
-			{
-				TotalBoneBox += *LocalBoneBox;
-			}
+			TotalBoneBox += BoundsPerBone[BoneIndex];
 		}
 
 		InOutBox += TotalBoneBox.TransformBy(GeometryCollectionComponent->GetComponentToWorld());;
@@ -523,7 +525,7 @@ void UFractureEditorMode::RefreshOutlinerWithCurrentSelection()
 	}
 }
 
-void UFractureEditorMode::GetComponentGlobalBounds(UGeometryCollectionComponent* GeometryCollectionComponent, TMap<int32, FBox> &BoundsToBone) const
+void UFractureEditorMode::GetComponentGlobalBounds(UGeometryCollectionComponent* GeometryCollectionComponent, TArray<FBox>& BoundsPerBone) const
 {
 	FGeometryCollectionEdit RestCollection = GeometryCollectionComponent->EditRestCollection(GeometryCollection::EEditUpdate::None);
 	UGeometryCollection* GeometryCollection = RestCollection.GetRestCollection();
@@ -531,38 +533,33 @@ void UFractureEditorMode::GetComponentGlobalBounds(UGeometryCollectionComponent*
 	TSharedPtr<FGeometryCollection, ESPMode::ThreadSafe> GeometryCollectionPtr = GeometryCollection->GetGeometryCollection();
 	FGeometryCollection* OutGeometryCollection = GeometryCollectionPtr.Get();
 
-	const TManagedArray<FTransform>& Transform = OutGeometryCollection->GetAttribute<FTransform>("Transform", FGeometryCollection::TransformGroup);
-	const TManagedArray<FBox>& BoundingBox = OutGeometryCollection->GetAttribute<FBox>("BoundingBox", FGeometryCollection::GeometryGroup);
-	const TManagedArray<int32>& TransformToGeometryIndex = OutGeometryCollection->GetAttribute<int32>("TransformToGeometryIndex", FGeometryCollection::TransformGroup);
+	const TManagedArray<FTransform>& Transform = OutGeometryCollection->Transform;
+	const TManagedArray<FBox>& BoundingBox = OutGeometryCollection->BoundingBox;
+	const TManagedArray<int32>& TransformToGeometryIndex = OutGeometryCollection->TransformToGeometryIndex;
+	const TManagedArray<int32>& Parent = OutGeometryCollection->Parent;
 
 	TManagedArray<FVector3f>* ExplodedVectorsPtr = OutGeometryCollection->FindAttribute<FVector3f>("ExplodedVector", FGeometryCollection::TransformGroup);
 
 	TArray<FTransform> Transforms;
 	GeometryCollectionAlgo::GlobalMatrices(Transform, OutGeometryCollection->Parent, Transforms);
 
-	BoundsToBone.Reset();
-
-	if (ExplodedVectorsPtr)
+	FBox EmptyBounds(EForceInit::ForceInit);
+	int32 NumBones = OutGeometryCollection->NumElements(FGeometryCollection::TransformGroup);
+	BoundsPerBone.Init(EmptyBounds, NumBones);
+	TArray<int32> BoneIndices = GeometryCollectionAlgo::ComputeRecursiveOrder(*OutGeometryCollection);
+	for (int32 BoneIdx : BoneIndices)
 	{
-		TManagedArray<FVector3f>& ExplodedVectors = *ExplodedVectorsPtr;
-		for (int32 Idx = 0, ni = GeometryCollection->NumElements(FGeometryCollection::TransformGroup); Idx < ni; ++Idx)
+		int32 GeoIdx = TransformToGeometryIndex[BoneIdx];
+		if (GeoIdx != INDEX_NONE)
 		{
-			if (TransformToGeometryIndex[Idx] > -1)
-			{
-				const FVector3f& Offset = ExplodedVectors[Idx];
-				const FBox& Bounds = BoundingBox[TransformToGeometryIndex[Idx]];
-				BoundsToBone.Add(Idx, Bounds.ShiftBy((FVector)Offset).TransformBy(Transforms[Idx]));
-			}
+			const FVector3f& Offset = ExplodedVectorsPtr ? (*ExplodedVectorsPtr)[BoneIdx] : FVector3f::ZeroVector;
+			const FBox& Bounds = BoundingBox[TransformToGeometryIndex[BoneIdx]];
+			BoundsPerBone[BoneIdx] += Bounds.ShiftBy((FVector)Offset).TransformBy(Transforms[BoneIdx]);
 		}
-	}
-	else
-	{
-		for (int32 Idx = 0, ni = GeometryCollection->NumElements(FGeometryCollection::TransformGroup); Idx < ni; ++Idx)
+		int32 ParentIdx = Parent[BoneIdx];
+		if (ParentIdx != INDEX_NONE)
 		{
-			if (TransformToGeometryIndex[Idx] > -1)
-			{
-				BoundsToBone.Add(Idx, BoundingBox[TransformToGeometryIndex[Idx]].TransformBy(Transforms[Idx]));
-			}
+			BoundsPerBone[ParentIdx] += BoundsPerBone[BoneIdx];
 		}
 	}
 }
diff --git a/Engine/Plugins/Experimental/ChaosEditor/Source/FractureEditor/Private/FractureEditorModeToolkit.cpp b/Engine/Plugins/Experimental/ChaosEditor/Source/FractureEditor/Private/FractureEditorModeToolkit.cpp
--- a/Engine/Plugins/Experimental/ChaosEditor/Source/FractureEditor/Private/FractureEditorModeToolkit.cpp
+++ b/Engine/Plugins/Experimental/ChaosEditor/Source/FractureEditor/Private/FractureEditorModeToolkit.cpp
@@ -1098,8 +1098,9 @@ void FFractureEditorModeToolkit::OnLevelViewValueChanged()
 		if(EditBoneColor.GetViewLevel() != FractureLevel)
 		{
 			EditBoneColor.SetLevelViewMode(FractureLevel);
-			EditBoneColor.ResetBoneSelection();
-			EditBoneColor.ResetHighlightedBones();
+			// Clear selection below currently-selected view level and update highlights,
+			// so the selection is compatible with the current 3D view and outliner (e.g., doesn't hide selection of children)
+			EditBoneColor.FilterSelectionToLevel();
 			UpdateExplodedVectors(Comp);
 			Comp->MarkRenderStateDirty();
 			Comp->MarkRenderDynamicDataDirty();
@@ -1693,8 +1694,8 @@ void FFractureEditorModeToolkit::OnOutlinerBoneSelectionChanged(UGeometryCollect
 	{
 		if (SelectedBones.Num())
 		{
-
-			FFractureSelectionTools::ToggleSelectedBones(RootComponent, SelectedBones, true, false);
+			// don't need to snap the bones to the current level because they are directly selected from the outliner
+			FFractureSelectionTools::ToggleSelectedBones(RootComponent, SelectedBones, true, false, false /*bSnapToLevel*/);
 			OutlinerView->SetBoneSelection(RootComponent, SelectedBones, true);
 			HistogramView->SetBoneSelection(RootComponent, SelectedBones, true);
 		}
diff --git a/Engine/Plugins/Experimental/ChaosEditor/Source/FractureEditor/Private/FractureSelectionTools.cpp b/Engine/Plugins/Experimental/ChaosEditor/Source/FractureEditor/Private/FractureSelectionTools.cpp
--- a/Engine/Plugins/Experimental/ChaosEditor/Source/FractureEditor/Private/FractureSelectionTools.cpp
+++ b/Engine/Plugins/Experimental/ChaosEditor/Source/FractureEditor/Private/FractureSelectionTools.cpp
@@ -5,7 +5,7 @@
 #include "GeometryCollection/GeometryCollectionClusteringUtility.h"
 #include "EditorSupportDelegates.h"
 
-void FFractureSelectionTools::ToggleSelectedBones(UGeometryCollectionComponent* GeometryCollectionComponent, TArray<int32>& BoneIndices, bool bClearCurrentSelection, bool bAdd)
+void FFractureSelectionTools::ToggleSelectedBones(UGeometryCollectionComponent* GeometryCollectionComponent, TArray<int32>& BoneIndices, bool bClearCurrentSelection, bool bAdd, bool bSnapToLevel)
 {
 	if (GeometryCollectionComponent)
 	{
@@ -24,26 +24,8 @@ void FFractureSelectionTools::ToggleSelectedBones(UGeometryCollectionComponent*
 					EditBoneColor.ResetBoneSelection();
 				}
 		
-				EditBoneColor.ToggleSelectedBones(BoneIndices, bAdd);
-
-				// The actual selection made is based on the hierarchy and the view mode
-				if (GeometryCollection)
-				{
-					const TArray<int32>& Selected = EditBoneColor.GetSelectedBones();
-					TArray<int32> RevisedSelected;
-					TArray<int32> Highlighted;
-					FGeometryCollectionClusteringUtility::ContextBasedClusterSelection(GeometryCollection, EditBoneColor.GetViewLevel(), Selected, RevisedSelected, Highlighted);
-					EditBoneColor.SetSelectedBones(RevisedSelected);
-					
-					const TArray<int32>& SelectedBones = EditBoneColor.GetSelectedBones();
-					TArray<int32> HighlightBones;
-					for (int32 SelectedBone : SelectedBones)
-					{
-						FGeometryCollectionClusteringUtility::RecursiveAddAllChildren(GeometryCollectionPtr->Children, SelectedBone, HighlightBones);
-					}
-					
-					EditBoneColor.SetHighlightedBones(HighlightBones);
-				}
+				EditBoneColor.ToggleSelectedBones(BoneIndices, bAdd, bSnapToLevel);
+				EditBoneColor.SetHighlightedBones(EditBoneColor.GetSelectedBones(), true);
 			}
 
 			FEditorSupportDelegates::RedrawAllViewports.Broadcast();
diff --git a/Engine/Plugins/Experimental/ChaosEditor/Source/FractureEditor/Private/FractureToolClustering.cpp b/Engine/Plugins/Experimental/ChaosEditor/Source/FractureEditor/Private/FractureToolClustering.cpp
--- a/Engine/Plugins/Experimental/ChaosEditor/Source/FractureEditor/Private/FractureToolClustering.cpp
+++ b/Engine/Plugins/Experimental/ChaosEditor/Source/FractureEditor/Private/FractureToolClustering.cpp
@@ -50,13 +50,21 @@ void UFractureToolFlattenAll::Execute(TWeakPtr<FFractureEditorModeToolkit> InToo
 
 			Context.ConvertSelectionToClusterNodes();
 
+			int32 MaxClusterLevel = -1;
 			for (int32 ClusterIndex : Context.GetSelection())
 			{
 				TArray<int32> LeafBones;
+				MaxClusterLevel = FMath::Max(MaxClusterLevel, Levels[ClusterIndex]);
 				FGeometryCollectionClusteringUtility::GetLeafBones(Context.GetGeometryCollection().Get(), ClusterIndex, true, LeafBones);
 				FGeometryCollectionClusteringUtility::ClusterBonesUnderExistingNode(Context.GetGeometryCollection().Get(), ClusterIndex, LeafBones);
 			}
 
+			// if not viewing all levels, switch the view level to show make sure we can see the cluster levels
+			if (Toolkit->GetLevelViewValue() != -1 && MaxClusterLevel != -1)
+			{
+				Toolkit->OnSetLevelViewValue(MaxClusterLevel + 1);
+			}
+
 			// Cleanup: Remove any clusters remaining in the flattened branch.
 			FGeometryCollectionClusteringUtility::RemoveDanglingClusters(Context.GetGeometryCollection().Get());
 
diff --git a/Engine/Plugins/Experimental/ChaosEditor/Source/FractureEditor/Private/FractureToolSelection.cpp b/Engine/Plugins/Experimental/ChaosEditor/Source/FractureEditor/Private/FractureToolSelection.cpp
--- a/Engine/Plugins/Experimental/ChaosEditor/Source/FractureEditor/Private/FractureToolSelection.cpp
+++ b/Engine/Plugins/Experimental/ChaosEditor/Source/FractureEditor/Private/FractureToolSelection.cpp
@@ -2,6 +2,7 @@
 
 #include "FractureToolSelection.h"
 #include "FractureEditorStyle.h"
+#include "FractureSettings.h"
 #include "PlanarCut.h"
 #include "GeometryCollection/GeometryCollection.h"
 #include "GeometryCollection/GeometryCollectionObject.h"
@@ -278,6 +279,10 @@ void UFractureToolSelection::Shutdown()
 {
 	UInteractiveToolsContext* Context = GLevelEditorModeTools().GetInteractiveToolsContext();
 	Context->InputRouter->DeregisterSource(SelectionBehaviorSource);
+
+	SelectionBehaviorSet = nullptr;
+	SelectionBehaviorSource = nullptr;
+	RectangleMarqueeManager = nullptr;
 }
 
 void UFractureToolSelection::DrawHUD(FEditorViewportClient* ViewportClient, FViewport* Viewport, const FSceneView* View, FCanvas* Canvas)
@@ -386,12 +391,12 @@ void UFractureToolSelection::FractureContextChanged()
 	}
 }
 
-TInterval<float> UFractureToolSelection::GetVolumeRange(const TManagedArray<float>& Volumes, const TManagedArray<int32>& TransformToGeometryIndex)
+TInterval<float> UFractureToolSelection::GetVolumeRange(const TManagedArray<float>& Volumes, const TManagedArray<int32>& SimulationType) const
 {
-	check(Volumes.Num() == TransformToGeometryIndex.Num());
+	check(Volumes.Num() == SimulationType.Num());
 
 	TInterval<float> Range;
-	if (SelectionSettings->VolumeSelectionMethod == EVolumeSelectionMethod::CubeRoot)
+	if (SelectionSettings->VolumeSelectionMethod == EVolumeSelectionMethod::CubeRootOfVolume)
 	{
 		Range.Min = SelectionSettings->MinVolume * VolDimScale;
 		Range.Max = SelectionSettings->MaxVolume * VolDimScale;
@@ -403,7 +408,7 @@ TInterval<float> UFractureToolSelection::GetVolumeRange(const TManagedArray<floa
 		double LargestVolume = KINDA_SMALL_NUMBER;
 		for (int32 Idx = 0; Idx < Volumes.Num(); Idx++)
 		{
-			if (TransformToGeometryIndex[Idx] >= 0)
+			if (SimulationType[Idx] == FGeometryCollection::ESimulationTypes::FST_Rigid)
 			{
 				float Volume = Volumes[Idx];
 				LargestVolume = FMath::Max(LargestVolume, Volume);
@@ -417,7 +422,7 @@ TInterval<float> UFractureToolSelection::GetVolumeRange(const TManagedArray<floa
 		double VolumeSum = 0;
 		for (int32 Idx = 0; Idx < Volumes.Num(); Idx++)
 		{
-			if (TransformToGeometryIndex[Idx] >= 0)
+			if (SimulationType[Idx] == FGeometryCollection::ESimulationTypes::FST_Rigid)
 			{
 				float Volume = Volumes[Idx];
 				VolumeSum += Volume;
@@ -438,10 +443,25 @@ bool UFractureToolSelection::GetBonesByVolume(const FGeometryCollection& Collect
 		return false;
 	}
 
-	TInterval<float> VolumeRange = GetVolumeRange(*Volumes, Collection.TransformToGeometryIndex);
+	const UFractureSettings* FractureSettings = GetDefault<UFractureSettings>();
+
+	TInterval<float> VolumeRange = GetVolumeRange(*Volumes, Collection.SimulationType);
+
+	const TManagedArray<int32>* Levels = Collection.FindAttribute<int32>("Level", FGeometryCollection::TransformGroup);
+	const TManagedArray<int32>& SimTypes = Collection.SimulationType;
 
 	for (int32 BoneIdx = 0; BoneIdx < Volumes->Num(); BoneIdx++)
 	{
+		// Skip embedded geometry (volumes aren't computed for embedded geo)
+		if (SimTypes[BoneIdx] == FGeometryCollection::ESimulationTypes::FST_None)
+		{
+			continue;
+		}
+		// only consider bones at the filtered level (if any)
+		if (FractureSettings->FractureLevel > -1 && Levels && (*Levels)[BoneIdx] != FractureSettings->FractureLevel)
+		{
+			continue;
+		}
 		if (VolumeRange.Contains((*Volumes)[BoneIdx]))
 		{
 			FilterIndices.Add(BoneIdx);
diff --git a/Engine/Plugins/Experimental/ChaosEditor/Source/FractureEditor/Private/FractureToolSelection.h b/Engine/Plugins/Experimental/ChaosEditor/Source/FractureEditor/Private/FractureToolSelection.h
--- a/Engine/Plugins/Experimental/ChaosEditor/Source/FractureEditor/Private/FractureToolSelection.h
+++ b/Engine/Plugins/Experimental/ChaosEditor/Source/FractureEditor/Private/FractureToolSelection.h
@@ -18,7 +18,7 @@ UENUM()
 enum class EVolumeSelectionMethod
 {
 	// Select by cube root of volume
-	CubeRoot,
+	CubeRootOfVolume,
 	// Select by cube root of volume relative to the overall shape's cube root of volume
 	RelativeToWhole,
 	// Select by cube root of volume relative to the largest single geometry
@@ -58,9 +58,7 @@ public:
 	UPROPERTY(EditAnywhere, Category = MouseInteraction)
 	EMouseSelectionMethod MouseSelectionMethod = EMouseSelectionMethod::RectSelect;
 
-
-
-	/** What values to use when filtering by volume */
+	/** What values to use when filtering by volume.  Note all values are presented as cube roots to give more intuitive scales (e.g., to select bones with volume less than a 10x10x10 cube, choose CubeRootOfVolume and MaxVolume=10, rather than needing to multiply out to 1000) */
 	UPROPERTY(EditAnywhere, Category = FilterSettings)
 	EVolumeSelectionMethod VolumeSelectionMethod = EVolumeSelectionMethod::RelativeToLargest;
 
@@ -267,7 +265,7 @@ private:
 	TArray<FBox> SelectionBounds; // Bounds in global space but without exploded vectors applied
 	FVisualizationMappings SelectionMappings;
 
-	TInterval<float> GetVolumeRange(const TManagedArray<float>& Volumes, const TManagedArray<int32>& TransformToGeometryIndex);
+	TInterval<float> GetVolumeRange(const TManagedArray<float>& Volumes, const TManagedArray<int32>& SimulationType) const;
 	bool GetBonesByVolume(const FGeometryCollection& Collection, TArray<int32>& FilterIndices);
 	const double VolDimScale = .01; // compute volumes in meters instead of cm, for saner units at typical scales
 
diff --git a/Engine/Plugins/Experimental/ChaosEditor/Source/FractureEditor/Private/FractureToolSelectors.cpp b/Engine/Plugins/Experimental/ChaosEditor/Source/FractureEditor/Private/FractureToolSelectors.cpp
--- a/Engine/Plugins/Experimental/ChaosEditor/Source/FractureEditor/Private/FractureToolSelectors.cpp
+++ b/Engine/Plugins/Experimental/ChaosEditor/Source/FractureEditor/Private/FractureToolSelectors.cpp
@@ -59,6 +59,30 @@ void UFractureToolSelectAll::SelectByMode(FFractureEditorModeToolkit* InToolkit,
 		{
 			FScopedColorEdit EditBoneColor = GeometryCollectionComponent->EditBoneSelection();
 			EditBoneColor.SelectBones(SelectionMode);
+
+			// Increase level so that we can see newly selected pieces
+			int32 LevelView = InToolkit->GetLevelViewValue();
+			if (LevelView >= 0)
+			{
+				int32 TargetLevel = LevelView;
+				// if we've selected bones beyond the current level, try to increase the level
+				int32 MaxLevel = EditBoneColor.GetMaxSelectedLevel(false);
+				if (MaxLevel > LevelView)
+				{
+					TargetLevel = MaxLevel;
+				}
+				// if the selected bones would not show up in the outliner at the target level, instead choose "All" level
+				if (!EditBoneColor.IsSelectionValidAtLevel(TargetLevel))
+				{
+					TargetLevel = -1;
+				}
+				if (TargetLevel != LevelView)
+				{
+					EditBoneColor.SetLevelViewMode(TargetLevel);
+					InToolkit->OnSetLevelViewValue(TargetLevel);
+				}
+			}
+
 			InToolkit->SetBoneSelection(GeometryCollectionComponent, EditBoneColor.GetSelectedBones(), true);
 		}
 	}
diff --git a/Engine/Plugins/Experimental/ChaosEditor/Source/FractureEditor/Private/SGeometryCollectionOutliner.cpp b/Engine/Plugins/Experimental/ChaosEditor/Source/FractureEditor/Private/SGeometryCollectionOutliner.cpp
--- a/Engine/Plugins/Experimental/ChaosEditor/Source/FractureEditor/Private/SGeometryCollectionOutliner.cpp
+++ b/Engine/Plugins/Experimental/ChaosEditor/Source/FractureEditor/Private/SGeometryCollectionOutliner.cpp
@@ -537,18 +537,21 @@ void FGeometryCollectionTreeItemComponent::SetHistogramSelection(TArray<int32>&
 
 bool FGeometryCollectionTreeItemComponent::FilterBoneIndex(int32 BoneIndex) const
 {
-	FGeometryCollection* Collection = Component->GetRestCollection()->GetGeometryCollection().Get();
-	TManagedArray<TSet<int32>>& Children = Collection->GetAttribute<TSet<int32>>("Children", FGeometryCollection::TransformGroup);
+	const FGeometryCollection* Collection = Component->GetRestCollection()->GetGeometryCollection().Get();
+	const TManagedArray<int32>& SimTypes = Collection->SimulationType;
+	bool bHasChildren = Collection->Children[BoneIndex].Num() > 0;
 
-	if (Children[BoneIndex].Num() == 0)
+	if (SimTypes[BoneIndex] != FGeometryCollection::ESimulationTypes::FST_Clustered)
 	{
-		// We don't display leaf nodes deeper than the view level.
+		// We only display cluster nodes deeper than the view level.
 		UFractureSettings* FractureSettings = GetMutableDefault<UFractureSettings>();
 
 		if (FractureSettings->FractureLevel >= 0)
 		{
-			TManagedArray<int32>& Level = Collection->GetAttribute<int32>("Level", FTransformCollection::TransformGroup);
-			if (Level[BoneIndex] != FractureSettings->FractureLevel)
+			const TManagedArray<int32>& Level = Collection->GetAttribute<int32>("Level", FTransformCollection::TransformGroup);
+			int32 BoneLevel = Level[BoneIndex];
+			// bone is not at the right level itself and doesn't have child(ren) at the right level
+			if (BoneLevel != FractureSettings->FractureLevel && (!bHasChildren || BoneLevel + 1 != FractureSettings->FractureLevel))
 			{
 				return false;
 			}
diff --git a/Engine/Plugins/Experimental/ChaosEditor/Source/FractureEditor/Public/FractureEditorMode.h b/Engine/Plugins/Experimental/ChaosEditor/Source/FractureEditor/Public/FractureEditorMode.h
--- a/Engine/Plugins/Experimental/ChaosEditor/Source/FractureEditor/Public/FractureEditorMode.h
+++ b/Engine/Plugins/Experimental/ChaosEditor/Source/FractureEditor/Public/FractureEditorMode.h
@@ -73,7 +73,7 @@ public:
 private:
 	void OnUndoRedo();
 	void OnActorSelectionChanged(const TArray<UObject*>& NewSelection, bool bForceRefresh);
-	void GetComponentGlobalBounds(UGeometryCollectionComponent* GeometryCollectionComponent, TMap<int32, FBox> &BoundsToBone) const;
+	void GetComponentGlobalBounds(UGeometryCollectionComponent* GeometryCollectionComponent, TArray<FBox>& BoundsPerBone) const;
 	void SelectionStateChanged();
 	
 	/** Handle package reloading (might be our geometry collection) */
diff --git a/Engine/Plugins/Experimental/ChaosEditor/Source/FractureEditor/Public/FractureSelectionTools.h b/Engine/Plugins/Experimental/ChaosEditor/Source/FractureEditor/Public/FractureSelectionTools.h
--- a/Engine/Plugins/Experimental/ChaosEditor/Source/FractureEditor/Public/FractureSelectionTools.h
+++ b/Engine/Plugins/Experimental/ChaosEditor/Source/FractureEditor/Public/FractureSelectionTools.h
@@ -16,8 +16,9 @@ public:
 	 * @param GeometryCollectionComponent	The component with bones being selected
 	 * @param BoneIndices					The indices to select
 	 * @param bClearCurrentSelection		Whether or not to clear the current selection set or just append the new indices to it				
+	 * @param bSnapToLevel					Whether to snap the selection to the nearest visible bones (in the 3D view and outliner) by searching the parent bones
 	 */
-	static void ToggleSelectedBones(UGeometryCollectionComponent* GeometryCollectionComponent, TArray<int32>& BoneIndices, bool bClearCurrentSelection, bool bAdd);
+	static void ToggleSelectedBones(UGeometryCollectionComponent* GeometryCollectionComponent, TArray<int32>& BoneIndices, bool bClearCurrentSelection, bool bAdd, bool bSnapToLevel = true);
 	static void ClearSelectedBones(UGeometryCollectionComponent* GeometryCollectionComponent);
 
 	static void SelectNeighbors(UGeometryCollectionComponent* GeometryCollectionComponent);
diff --git a/Engine/Plugins/Experimental/ChaosEditor/Source/FractureEditor/Public/FractureTool.cpp b/Engine/Plugins/Experimental/ChaosEditor/Source/FractureEditor/Public/FractureTool.cpp
--- a/Engine/Plugins/Experimental/ChaosEditor/Source/FractureEditor/Public/FractureTool.cpp
+++ b/Engine/Plugins/Experimental/ChaosEditor/Source/FractureEditor/Public/FractureTool.cpp
@@ -193,8 +193,6 @@ void UFractureActionTool::Refresh(FFractureToolContext& Context, FFractureEditor
 		EditBoneColor.SetSelectedBones(Context.GetSelection());
 		FFractureSelectionTools::ToggleSelectedBones(GeometryCollectionComponent, Context.GetSelection(), true, false);
 	}
-	
-	Toolkit->UpdateExplodedVectors(GeometryCollectionComponent);
 
 	GeometryCollectionComponent->MarkRenderDynamicDataDirty();
 	GeometryCollectionComponent->MarkRenderStateDirty();
diff --git a/Engine/Source/Runtime/Experimental/Chaos/Private/GeometryCollection/GeometryCollectionClusteringUtility.cpp b/Engine/Source/Runtime/Experimental/Chaos/Private/GeometryCollection/GeometryCollectionClusteringUtility.cpp
--- a/Engine/Source/Runtime/Experimental/Chaos/Private/GeometryCollection/GeometryCollectionClusteringUtility.cpp
+++ b/Engine/Source/Runtime/Experimental/Chaos/Private/GeometryCollection/GeometryCollectionClusteringUtility.cpp
@@ -483,6 +483,36 @@ void FGeometryCollectionClusteringUtility::GetChildBonesFromLevel(const FGeometr
 
 }
 
+void FGeometryCollectionClusteringUtility::GetBonesToLevel(const FGeometryCollection* GeometryCollection, int32 Level, TArray<int32>& BonesOut, bool bOnlyClusteredOrRigid, bool bSkipFiltered)
+{
+	check(GeometryCollection);
+
+	const TManagedArray<int32>& Levels = GeometryCollection->GetAttribute<int32>("Level", FGeometryCollection::TransformGroup);
+	const TManagedArray<int32>& SimType = GeometryCollection->SimulationType;
+
+	bool bAllLevels = Level == -1;
+
+	int32 NumBones = GeometryCollection->NumElements(FGeometryCollection::TransformGroup);
+	for (int32 BoneIdx = 0; BoneIdx < NumBones; BoneIdx++)
+	{
+		bool bIsRigid = SimType[BoneIdx] == FGeometryCollection::ESimulationTypes::FST_Rigid;
+		bool bIsClustered = SimType[BoneIdx] == FGeometryCollection::ESimulationTypes::FST_Clustered;
+		if (
+			// (if skipping embedded) sim type is clustered or rigid 
+			(!bOnlyClusteredOrRigid || bIsClustered || bIsRigid)
+			&&
+			// (if skipping nodes the outliner has filtered) sim type is clustered or level is an exact match or level has an exact-match child
+			(bAllLevels || !bSkipFiltered || bIsClustered || Levels[BoneIdx] == Level || (GeometryCollection->Children[BoneIdx].Num() > 0 && Levels[BoneIdx] + 1 == Level))
+			&&
+			// level is at or before the target
+			(bAllLevels || Levels[BoneIdx] <= Level)
+			)
+		{
+			BonesOut.Add(BoneIdx);
+		}
+	}
+}
+
 void FGeometryCollectionClusteringUtility::GetChildBonesAtLevel(const FGeometryCollection* GeometryCollection, int32 SourceBone, int32 Level, TArray<int32>& BonesOut)
 {
 	check(GeometryCollection);
@@ -519,19 +549,28 @@ void FGeometryCollectionClusteringUtility::RecursiveAddAllChildren(const TManage
 
 }
 
-int32 FGeometryCollectionClusteringUtility::GetParentOfBoneAtSpecifiedLevel(const FGeometryCollection* GeometryCollection, int32 SourceBone, int32 Level)
+int32 FGeometryCollectionClusteringUtility::GetParentOfBoneAtSpecifiedLevel(const FGeometryCollection* GeometryCollection, int32 SourceBone, int32 Level, bool bSkipFiltered)
 {
 	check(GeometryCollection);
 	const TManagedArray<int32>& Parents = GeometryCollection->Parent;
 	const TManagedArray<int32>& Levels = GeometryCollection->GetAttribute<int32>("Level", FGeometryCollection::TransformGroup);
+	const TManagedArray<int32>& SimTypes = GeometryCollection->SimulationType;
 
 	if (SourceBone >= 0)
 	{
 		int32 SourceParent = SourceBone;
-		while (Levels[SourceParent] > Level)
+		while (Levels[SourceParent] > Level || 
+			// go to parents of bones that will be filtered by the outliner (i.e., rigid/embedded at the wrong level)
+			(bSkipFiltered && 
+				Levels[SourceParent] != Level && 
+				GeometryCollection->SimulationType[SourceParent] != FGeometryCollection::ESimulationTypes::FST_Clustered &&
+				(GeometryCollection->Children[SourceParent].Num() == 0 || Levels[SourceParent] + 1 != Level)
+			))
 		{
 			if (Parents[SourceParent] == -1)
+			{
 				break;
+			}
 
 			SourceParent = Parents[SourceParent];
 		}
@@ -718,48 +757,6 @@ void FGeometryCollectionClusteringUtility::ValidateResults(FGeometryCollection*
 	ensure(GeometryCollection->HasContiguousVertices());
 }
 
-void FGeometryCollectionClusteringUtility::ContextBasedClusterSelection(
-	FGeometryCollection* GeometryCollection,
-	int ViewLevel,
-	const TArray<int32>& SelectedComponentBonesIn,
-	TArray<int32>& SelectedComponentBonesOut,
-	TArray<int32>& HighlightedComponentBonesOut)
-{
-	HighlightedComponentBonesOut.Empty();
-	SelectedComponentBonesOut.Empty();
-
-	for (int32 BoneIndex : SelectedComponentBonesIn)
-	{
-		TArray <int32> SelectionHighlightedBones;
-		if (ViewLevel == -1)
-		{
-			SelectionHighlightedBones.AddUnique(BoneIndex);
-			SelectedComponentBonesOut.AddUnique(BoneIndex);
-		}
-		else
-		{
-			// select all children under bone as selected hierarchy level
-			int32 ParentBoneIndex = GetParentOfBoneAtSpecifiedLevel(GeometryCollection, BoneIndex, ViewLevel);
-			if (ParentBoneIndex != FGeometryCollection::Invalid)
-			{
-				SelectedComponentBonesOut.AddUnique(ParentBoneIndex);
-			}
-			else
-			{
-				SelectedComponentBonesOut.AddUnique(BoneIndex);
-			}
-
-			for (int32 Bone : SelectedComponentBonesOut)
-			{
-				GetChildBonesFromLevel(GeometryCollection, Bone, ViewLevel, SelectionHighlightedBones);
-			}
-		}
-
-		HighlightedComponentBonesOut.Append(SelectionHighlightedBones);
-	}
-
-}
-
 void FGeometryCollectionClusteringUtility::GetLeafBones(const FGeometryCollection* GeometryCollection, int BoneIndex, bool bOnlyRigids, TArray<int32>& LeafBonesOut)
 {
 	if (!ensure(BoneIndex >= 0))
diff --git a/Engine/Source/Runtime/Experimental/Chaos/Public/GeometryCollection/GeometryCollectionClusteringUtility.h b/Engine/Source/Runtime/Experimental/Chaos/Public/GeometryCollection/GeometryCollectionClusteringUtility.h
--- a/Engine/Source/Runtime/Experimental/Chaos/Public/GeometryCollection/GeometryCollectionClusteringUtility.h
+++ b/Engine/Source/Runtime/Experimental/Chaos/Public/GeometryCollection/GeometryCollectionClusteringUtility.h
@@ -60,6 +60,9 @@ public:
 	
 	/** Finds all Bones in same cluster as the one specified */
 	static void GetClusteredBonesWithCommonParent(const FGeometryCollection* GeometryCollection, int32 SourceBone, TArray<int32>& BonesOut);
+
+	/** Get list of all bones above the specified hierarchy level */
+	static void GetBonesToLevel(const FGeometryCollection* GeometryCollection, int32 Level, TArray<int32>& BonesOut, bool bOnlyClusteredOrRigid = true, bool bSkipFiltered = true);
 	
 	/** Get list of child bones down from the source bone below the specified hierarchy level */
 	static void GetChildBonesFromLevel(const FGeometryCollection* GeometryCollection, int32 SourceBone, int32 Level, TArray<int32>& BonesOut);
@@ -70,8 +73,11 @@ public:
 	/** Recursively Add all children to output bone list from source bone down to the leaf nodes */
 	static void RecursiveAddAllChildren(const TManagedArray<TSet<int32>>& Children, int32 SourceBone, TArray<int32>& BonesOut);
 
-	/** Search hierarchy for the parent of the specified bone, where the parent exists at the given level in the hierarchy */
-	static int32 GetParentOfBoneAtSpecifiedLevel(const FGeometryCollection* GeometryCollection, int32 SourceBone, int32 Level);
+	/**
+	 * Search hierarchy for the parent of the specified bone, where the parent exists at the given level in the hierarchy.
+	 * If bSkipFiltered, also skip to parents of any rigid/embedded nodes that will be filtered from the outliner.
+	 */
+	static int32 GetParentOfBoneAtSpecifiedLevel(const FGeometryCollection* GeometryCollection, int32 SourceBone, int32 Level, bool bSkipFiltered = false);
 
 	/**
 	* Maintains the bone naming convention of
@@ -103,9 +109,6 @@ public:
 	/** Rename a bone node, will automatically update all child node names if requested (current default) */
 	static void RenameBone(FGeometryCollection* GeometryCollection, int32 BoneIndex, const FString& NewName, bool UpdateChildren = true);
 
-	/** Make logical editor selections based on the current selected hierarchy view level */
-	static void ContextBasedClusterSelection(FGeometryCollection* GeometryCollection, int ViewLevel, const TArray<int32>& SelectedComponentBonesIn, TArray<int32>& SelectedComponentBonesOut, TArray<int32>& HighlightedComponentBonesOut);
-
 	/** return an array of all child leaf nodes below the specified node. If bOnlyRigids is true, the first Rigid node dound is considered a leaf, regardless of an children it might have. */
 	static void GetLeafBones(const FGeometryCollection* GeometryCollection, int BoneIndex, bool bOnlyRigids, TArray<int32>& LeafBonesOut);
 
diff --git a/Engine/Source/Runtime/Experimental/GeometryCollectionEngine/Private/GeometryCollection/GeometryCollectionComponent.cpp b/Engine/Source/Runtime/Experimental/GeometryCollectionEngine/Private/GeometryCollection/GeometryCollectionComponent.cpp
--- a/Engine/Source/Runtime/Experimental/GeometryCollectionEngine/Private/GeometryCollection/GeometryCollectionComponent.cpp
+++ b/Engine/Source/Runtime/Experimental/GeometryCollectionEngine/Private/GeometryCollection/GeometryCollectionComponent.cpp
@@ -875,9 +875,12 @@ UPhysicalMaterial* UGeometryCollectionComponent::GetPhysicalMaterial() const
 
 void UGeometryCollectionComponent::RefreshEmbeddedGeometry()
 {
-	
 	const TManagedArray<int32>& ExemplarIndexArray = GetExemplarIndexArray();
 	const int32 TransformCount = GlobalMatrices.Num();
+	if (!ensureMsgf(TransformCount == ExemplarIndexArray.Num(), TEXT("GlobalMatrices (Num=%d) cached on GeometryCollectionComponent are not in sync with ExemplarIndexArray (Num=%d) on underlying GeometryCollection; likely missed a dynamic data update"), TransformCount, ExemplarIndexArray.Num()))
+	{
+		return;
+	}
 	
 	const TManagedArray<bool>* HideArray = nullptr;
 	if (RestCollection->GetGeometryCollection()->HasAttribute("Hide", FGeometryCollection::TransformGroup))
@@ -2352,7 +2355,7 @@ void FScopedColorEdit::AppendSelectedBones(const TArray<int32>& SelectedBonesIn)
 	Component->SelectedBones.Append(SelectedBonesIn);
 }
 
-void FScopedColorEdit::ToggleSelectedBones(const TArray<int32>& SelectedBonesIn, bool bAdd)
+void FScopedColorEdit::ToggleSelectedBones(const TArray<int32>& SelectedBonesIn, bool bAdd, bool bSnapToLevel)
 {
 	bUpdated = true;
 	
@@ -2363,7 +2366,13 @@ void FScopedColorEdit::ToggleSelectedBones(const TArray<int32>& SelectedBonesIn,
 		for (int32 BoneIndex : SelectedBonesIn)
 		{
 		
-			int32 ContextBoneIndex = (GetViewLevel() > -1) ? FGeometryCollectionClusteringUtility::GetParentOfBoneAtSpecifiedLevel(GeometryCollectionPtr.Get(), BoneIndex, GetViewLevel()) : BoneIndex;
+			int32 ContextBoneIndex = (bSnapToLevel && GetViewLevel() > -1) ? 
+				FGeometryCollectionClusteringUtility::GetParentOfBoneAtSpecifiedLevel(GeometryCollectionPtr.Get(), BoneIndex, GetViewLevel(), true /*bSkipFiltered*/) 
+				: BoneIndex;
+			if (ContextBoneIndex == FGeometryCollection::Invalid)
+			{
+				continue;
+			}
 		
 			if (bAdd) // shift select
 			{
@@ -2407,6 +2416,52 @@ const TArray<int32>& FScopedColorEdit::GetSelectedBones() const
 	return Component->GetSelectedBones();
 }
 
+int32 FScopedColorEdit::GetMaxSelectedLevel(bool bOnlyRigid) const
+{
+	int32 MaxSelectedLevel = -1;
+	const UGeometryCollection* GeometryCollection = Component->GetRestCollection();
+	if (GeometryCollection)
+	{
+		TSharedPtr<FGeometryCollection, ESPMode::ThreadSafe> GeometryCollectionPtr = GeometryCollection->GetGeometryCollection();
+		const TManagedArray<int32>& Levels = GeometryCollectionPtr->GetAttribute<int32>("Level", FGeometryCollection::TransformGroup);
+		const TManagedArray<int32>& SimTypes = GeometryCollectionPtr->SimulationType;
+		for (int32 BoneIndex : Component->SelectedBones)
+		{
+			if (!bOnlyRigid || SimTypes[BoneIndex] == FGeometryCollection::ESimulationTypes::FST_Rigid)
+			{
+				MaxSelectedLevel = FMath::Max(MaxSelectedLevel, Levels[BoneIndex]);
+			}
+		}
+	}
+	return MaxSelectedLevel;
+}
+
+bool FScopedColorEdit::IsSelectionValidAtLevel(int32 TargetLevel) const
+{
+	if (TargetLevel == -1)
+	{
+		return true;
+	}
+	const UGeometryCollection* GeometryCollection = Component->GetRestCollection();
+	if (GeometryCollection)
+	{
+		TSharedPtr<FGeometryCollection, ESPMode::ThreadSafe> GeometryCollectionPtr = GeometryCollection->GetGeometryCollection();
+		const TManagedArray<int32>& Levels = GeometryCollectionPtr->GetAttribute<int32>("Level", FGeometryCollection::TransformGroup);
+		const TManagedArray<int32>& SimTypes = GeometryCollectionPtr->SimulationType;
+		for (int32 BoneIndex : Component->SelectedBones)
+		{
+			if (SimTypes[BoneIndex] != FGeometryCollection::ESimulationTypes::FST_Clustered && // clusters are always shown in outliner
+				Levels[BoneIndex] != TargetLevel && // nodes at the target level are shown in outliner
+				// non-cluster parents are shown if they have children that are exact matches (i.e., a rigid parent w/ embedded at the target level)
+				(GeometryCollectionPtr->Children[BoneIndex].Num() == 0 || Levels[BoneIndex] + 1 != TargetLevel))
+			{
+				return false;
+			}
+		}
+	}
+	return true;
+}
+
 void FScopedColorEdit::ResetBoneSelection()
 {
 	if (Component->SelectedBones.Num() > 0)
@@ -2417,6 +2472,56 @@ void FScopedColorEdit::ResetBoneSelection()
 	Component->SelectedBones.Empty();
 }
 
+void FScopedColorEdit::FilterSelectionToLevel(bool bPreferLowestOnly)
+{
+	const UGeometryCollection* GeometryCollection = Component->GetRestCollection();
+	int32 ViewLevel = GetViewLevel();
+	bool bNeedsFiltering = ViewLevel >= 0 || bPreferLowestOnly;
+	if (GeometryCollection && Component->SelectedBones.Num() > 0 && bNeedsFiltering)
+	{
+		TSharedPtr<FGeometryCollection, ESPMode::ThreadSafe> GeometryCollectionPtr = GeometryCollection->GetGeometryCollection();
+
+		const TManagedArray<int32>& Levels = GeometryCollectionPtr->GetAttribute<int32>("Level", FGeometryCollection::TransformGroup);
+		const TManagedArray<int32>& SimTypes = GeometryCollectionPtr->SimulationType;
+
+		TArray<int32> NewSelection;
+		NewSelection.Reserve(Component->SelectedBones.Num());
+		if (ViewLevel >= 0)
+		{
+			for (int32 BoneIdx : Component->SelectedBones)
+			{
+				bool bIsCluster = SimTypes[BoneIdx] == FGeometryCollection::ESimulationTypes::FST_Clustered;
+				if (bPreferLowestOnly && bIsCluster && Levels[BoneIdx] < ViewLevel)
+				{
+					continue;
+				}
+				if (Levels[BoneIdx] == ViewLevel || (bIsCluster && Levels[BoneIdx] <= ViewLevel))
+				{
+					NewSelection.Add(BoneIdx);
+				}
+			}
+		}
+		else // bPreferLowestOnly && ViewLevel == -1
+		{
+			// If view level is "all" and we prefer lowest selection, just select any non-cluster nodes
+			for (int32 BoneIdx : Component->SelectedBones)
+			{
+				bool bIsCluster = SimTypes[BoneIdx] == FGeometryCollection::ESimulationTypes::FST_Clustered;
+				if (!bIsCluster)
+				{
+					NewSelection.Add(BoneIdx);
+				}
+			}
+		}
+
+		if (NewSelection.Num() != Component->SelectedBones.Num())
+		{
+			SetSelectedBones(NewSelection);
+			SetHighlightedBones(NewSelection, true);
+		}
+	}
+}
+
 void FScopedColorEdit::SelectBones(GeometryCollection::ESelectionMode SelectionMode)
 {
 	check(Component);
@@ -2434,15 +2539,10 @@ void FScopedColorEdit::SelectBones(GeometryCollection::ESelectionMode SelectionM
 
 		case GeometryCollection::ESelectionMode::AllGeometry:
 		{
-			TArray<int32> Roots;
-			FGeometryCollectionClusteringUtility::GetRootBones(GeometryCollectionPtr.Get(), Roots);
 			ResetBoneSelection();
-			for (int32 RootElement : Roots)
-			{
-				TArray<int32> LeafBones;
-				FGeometryCollectionClusteringUtility::GetLeafBones(GeometryCollectionPtr.Get(), RootElement, true, LeafBones);
-				AppendSelectedBones(LeafBones);
-			}
+			TArray<int32> BonesToSelect;
+			FGeometryCollectionClusteringUtility::GetBonesToLevel(GeometryCollectionPtr.Get(), GetViewLevel(), BonesToSelect, true, true);
+			AppendSelectedBones(BonesToSelect);
 		}
 		break;
 
@@ -2476,9 +2576,6 @@ void FScopedColorEdit::SelectBones(GeometryCollection::ESelectionMode SelectionM
 						if (!IsBoneSelected(ViewLevelBone))
 						{
 							NewSelection.Push(ViewLevelBone);
-							TArray<int32> ChildBones;
-							FGeometryCollectionClusteringUtility::GetChildBonesFromLevel(GeometryCollectionPtr.Get(), ViewLevelBone, GetViewLevel(), ChildBones);
-							NewSelection.Append(ChildBones);
 						}
 					}
 				}
@@ -2635,12 +2732,25 @@ bool FScopedColorEdit::IsBoneHighlighted(int BoneIndex) const
 	return Component->HighlightedBones.Contains(BoneIndex);
 }
 
-void FScopedColorEdit::SetHighlightedBones(const TArray<int32>& HighlightedBonesIn)
+void FScopedColorEdit::SetHighlightedBones(const TArray<int32>& HighlightedBonesIn, bool bHighlightChildren)
 {
 	if (Component->HighlightedBones != HighlightedBonesIn)
 	{
+		const UGeometryCollection* GeometryCollection = Component->GetRestCollection();
+		if (bHighlightChildren && GeometryCollection)
+		{
+			Component->HighlightedBones.Reset();
+			TSharedPtr<FGeometryCollection, ESPMode::ThreadSafe> GeometryCollectionPtr = GeometryCollection->GetGeometryCollection();
+			for (int32 SelectedBone : HighlightedBonesIn)
+			{
+				FGeometryCollectionClusteringUtility::RecursiveAddAllChildren(GeometryCollectionPtr->Children, SelectedBone, Component->HighlightedBones);
+			}
+		}
+		else
+		{
+			Component->HighlightedBones = HighlightedBonesIn;
+		}
 		bUpdated = true;
-		Component->HighlightedBones = HighlightedBonesIn;
 	}
 }
 
diff --git a/Engine/Source/Runtime/Experimental/GeometryCollectionEngine/Public/GeometryCollection/GeometryCollectionComponent.h b/Engine/Source/Runtime/Experimental/GeometryCollectionEngine/Public/GeometryCollection/GeometryCollectionComponent.h
--- a/Engine/Source/Runtime/Experimental/GeometryCollectionEngine/Public/GeometryCollection/GeometryCollectionComponent.h
+++ b/Engine/Source/Runtime/Experimental/GeometryCollectionEngine/Public/GeometryCollection/GeometryCollectionComponent.h
@@ -210,15 +210,18 @@ public:
 	bool IsBoneSelected(int BoneIndex) const;
 	void SetSelectedBones(const TArray<int32>& SelectedBonesIn);
 	void AppendSelectedBones(const TArray<int32>& SelectedBonesIn);
-	void ToggleSelectedBones(const TArray<int32>& SelectedBonesIn, bool bAdd);
+	void ToggleSelectedBones(const TArray<int32>& SelectedBonesIn, bool bAdd, bool bSnapToLevel = true);
 	void AddSelectedBone(int32 BoneIndex);
 	void ClearSelectedBone(int32 BoneIndex);
 	const TArray<int32>& GetSelectedBones() const;
 	void ResetBoneSelection();
 	void SelectBones(GeometryCollection::ESelectionMode SelectionMode);
+	void FilterSelectionToLevel(bool bPreferLowestOnly = false);
+	int32 GetMaxSelectedLevel(bool bOnlyRigid) const;
+	bool IsSelectionValidAtLevel(int32 TargetLevel) const;
 
 	bool IsBoneHighlighted(int BoneIndex) const;
-	void SetHighlightedBones(const TArray<int32>& HighlightedBonesIn);
+	void SetHighlightedBones(const TArray<int32>& HighlightedBonesIn, bool bHighlightChildren = false);
 	void AddHighlightedBone(int32 BoneIndex);
 	const TArray<int32>& GetHighlightedBones() const;
 	void ResetHighlightedBones();
