Fixed triangle mesh collision breakage post-vectorization causing queries to sometimes fail (introduced in 18935688)

diff --git a/Engine/Source/Runtime/Experimental/Chaos/Private/Chaos/TriangleMeshImplicitObject.cpp b/Engine/Source/Runtime/Experimental/Chaos/Private/Chaos/TriangleMeshImplicitObject.cpp
--- a/Engine/Source/Runtime/Experimental/Chaos/Private/Chaos/TriangleMeshImplicitObject.cpp
+++ b/Engine/Source/Runtime/Experimental/Chaos/Private/Chaos/TriangleMeshImplicitObject.cpp
@@ -944,7 +944,7 @@ struct FTriangleMeshSweepVisitor
 		if(CullsBackFaceSweepsCode != 0)
 		{
 			const VectorRegister4Float TriNormal = VectorCross(VectorSubtract(B, A), VectorSubtract(C, A));
-			const VectorRegister4Float ReturnTrue = VectorCompareGT(VectorMultiply(VectorCross(TriNormal, VectorScaledDirNormalized), VectorCullsBackFaceSweepsCode), VectorZero());
+			const VectorRegister4Float ReturnTrue = VectorCompareGT(VectorMultiply(VectorDot3(TriNormal, VectorScaledDirNormalized), VectorCullsBackFaceSweepsCode), VectorZero());
 			if (VectorMaskBits(ReturnTrue))
 			{
 				return true;
