Chaos - joint shock propagation in QPBD solver (disabled)

diff --git a/Engine/Source/Editor/DetailCustomizations/Private/PhysicsConstraintComponentDetails.cpp b/Engine/Source/Editor/DetailCustomizations/Private/PhysicsConstraintComponentDetails.cpp
--- a/Engine/Source/Editor/DetailCustomizations/Private/PhysicsConstraintComponentDetails.cpp
+++ b/Engine/Source/Editor/DetailCustomizations/Private/PhysicsConstraintComponentDetails.cpp
@@ -180,6 +180,7 @@ void FPhysicsConstraintComponentDetails::AddConstraintBehaviorProperties(IDetail
 	ConstraintCat.AddProperty(ProfilePropertiesProperty->GetChildHandle(GET_MEMBER_NAME_CHECKED(FConstraintProfileProperties, bEnableSoftProjection)));
 	ConstraintCat.AddProperty(ProfilePropertiesProperty->GetChildHandle(GET_MEMBER_NAME_CHECKED(FConstraintProfileProperties, ProjectionLinearAlpha)));
 	ConstraintCat.AddProperty(ProfilePropertiesProperty->GetChildHandle(GET_MEMBER_NAME_CHECKED(FConstraintProfileProperties, ProjectionAngularAlpha)));
+	ConstraintCat.AddProperty(ProfilePropertiesProperty->GetChildHandle(GET_MEMBER_NAME_CHECKED(FConstraintProfileProperties, ShockPropagationAlpha)));
 #endif
 
 #if WITH_CHAOS
@@ -193,6 +194,7 @@ void FPhysicsConstraintComponentDetails::AddConstraintBehaviorProperties(IDetail
 	ProfilePropertiesProperty->GetChildHandle(GET_MEMBER_NAME_CHECKED(FConstraintProfileProperties, bEnableSoftProjection))->MarkHiddenByCustomization();
 	ProfilePropertiesProperty->GetChildHandle(GET_MEMBER_NAME_CHECKED(FConstraintProfileProperties, ProjectionLinearAlpha))->MarkHiddenByCustomization();
 	ProfilePropertiesProperty->GetChildHandle(GET_MEMBER_NAME_CHECKED(FConstraintProfileProperties, ProjectionAngularAlpha))->MarkHiddenByCustomization();
+	ProfilePropertiesProperty->GetChildHandle(GET_MEMBER_NAME_CHECKED(FConstraintProfileProperties, ShockPropagationAlpha))->MarkHiddenByCustomization();
 #endif
 
 	//Add the rest
diff --git a/Engine/Source/Runtime/Engine/Classes/PhysicsEngine/ConstraintInstance.h b/Engine/Source/Runtime/Engine/Classes/PhysicsEngine/ConstraintInstance.h
--- a/Engine/Source/Runtime/Engine/Classes/PhysicsEngine/ConstraintInstance.h
+++ b/Engine/Source/Runtime/Engine/Classes/PhysicsEngine/ConstraintInstance.h
@@ -42,6 +42,10 @@ struct ENGINE_API FConstraintProfileProperties
 	UPROPERTY(EditAnywhere, Category = Projection, meta = (ClampMin = "0.0", ClampMax = "1.0"))
 	float ProjectionAngularAlpha;
 
+	/** How much shock propagation to apply [0-1] */
+	UPROPERTY(EditAnywhere, Category = Projection, meta = (ClampMin = "0.0", ClampMax = "1.0"))
+	float ShockPropagationAlpha;
+
 	/** Force needed to break the distance constraint. */
 	UPROPERTY(EditAnywhere, AdvancedDisplay, Category = Linear, meta = (editcondition = "bLinearBreakable", ClampMin = "0.0"))
 	float LinearBreakThreshold;
@@ -908,6 +912,12 @@ public:
 	*/
 	void GetProjectionAlphasOrTolerances(float& ProjectionLinearAlphaOrTolerance, float& ProjectionAngularAlphaOrTolerance) const;
 
+	/** Set the shock propagation amount [0, 1] */
+	void SetShockPropagationParams(float ShockPropagationAlpha);
+
+	/** Get the shock propagation amount [0, 1] */
+	float GetShockPropagationAlpha() const;
+
 	/** Whether parent domination is enabled (meaning the parent body cannot be be affected at all by a child) */
 	bool IsParentDominatesEnabled() const
 	{
diff --git a/Engine/Source/Runtime/Engine/Private/PhysicsEngine/ConstraintInstance.cpp b/Engine/Source/Runtime/Engine/Private/PhysicsEngine/ConstraintInstance.cpp
--- a/Engine/Source/Runtime/Engine/Private/PhysicsEngine/ConstraintInstance.cpp
+++ b/Engine/Source/Runtime/Engine/Private/PhysicsEngine/ConstraintInstance.cpp
@@ -135,6 +135,7 @@ FConstraintProfileProperties::FConstraintProfileProperties()
 	, ProjectionAngularTolerance(180.f)
 	, ProjectionLinearAlpha(1.0f)
 	, ProjectionAngularAlpha(0.0f)
+	, ShockPropagationAlpha(0.0f)
 	, LinearBreakThreshold(300.f)
 	, LinearPlasticityThreshold(0.1f)
 	, AngularBreakThreshold(500.f)
@@ -466,6 +467,7 @@ void FConstraintProfileProperties::UpdateConstraintFlags_AssumesLocked(const FPh
 	FPhysicsInterface::SetCollisionEnabled(InConstraintRef, !bDisableCollision);
 #if WITH_CHAOS
 	FPhysicsInterface::SetProjectionEnabled_AssumesLocked(InConstraintRef, bEnableProjection, ProjectionLinearAlpha, ProjectionAngularAlpha);
+	FPhysicsInterface::SetShockPropagationAlpha_AssumesLocked(InConstraintRef, ShockPropagationAlpha);
 #else
 	FPhysicsInterface::SetProjectionEnabled_AssumesLocked(InConstraintRef, bEnableProjection, ProjectionLinearTolerance, ProjectionAngularTolerance);
 #endif
@@ -1215,6 +1217,26 @@ void FConstraintInstance::SetProjectionParams(bool bEnableProjection, float Proj
 		});
 }
 
+float FConstraintInstance::GetShockPropagationAlpha() const
+{
+#if WITH_CHAOS
+	return ProfileInstance.ShockPropagationAlpha;
+#else
+	return 0;
+#endif
+}
+
+void FConstraintInstance::SetShockPropagationParams(float ShockPropagationAlpha)
+{
+#if WITH_CHAOS
+	ProfileInstance.ShockPropagationAlpha = ShockPropagationAlpha;
+	FPhysicsCommand::ExecuteWrite(ConstraintHandle, [&](const FPhysicsConstraintHandle& Constraint)
+		{
+			FPhysicsInterface::SetShockPropagationAlpha_AssumesLocked(Constraint, ShockPropagationAlpha);
+		});
+#endif
+}
+
 void FConstraintInstance::EnableProjection()
 {
 #if WITH_CHAOS
diff --git a/Engine/Source/Runtime/Engine/Private/PhysicsEngine/ImmediatePhysics/ImmediatePhysicsChaos/ImmediatePhysicsSimulation_Chaos.cpp b/Engine/Source/Runtime/Engine/Private/PhysicsEngine/ImmediatePhysics/ImmediatePhysicsChaos/ImmediatePhysicsSimulation_Chaos.cpp
--- a/Engine/Source/Runtime/Engine/Private/PhysicsEngine/ImmediatePhysics/ImmediatePhysicsChaos/ImmediatePhysicsSimulation_Chaos.cpp
+++ b/Engine/Source/Runtime/Engine/Private/PhysicsEngine/ImmediatePhysics/ImmediatePhysicsChaos/ImmediatePhysicsSimulation_Chaos.cpp
@@ -88,11 +88,14 @@ int32 ChaosImmediate_Joint_PushOutPairIterations = -1;
 Chaos::FRealSingle ChaosImmediate_Joint_SwingTwistAngleTolerance = 1.0e-6f;
 Chaos::FRealSingle ChaosImmediate_Joint_PositionTolerance = 0.025f;
 Chaos::FRealSingle ChaosImmediate_Joint_AngleTolerance = 0.001f;
+int32 ChaosImmediate_Joint_NumShockPropagationIterations = 0;
+int32 ChaosImmediate_Joint_SolvePositionLast = 1;
 int32 ChaosImmediate_Joint_EnableTwistLimits = 1;
 int32 ChaosImmediate_Joint_EnableSwingLimits = 1;
 int32 ChaosImmediate_Joint_EnableDrives = 1;
 Chaos::FRealSingle ChaosImmediate_Joint_LinearProjection = -1.0f;
 Chaos::FRealSingle ChaosImmediate_Joint_AngularProjection = -1.0f;
+Chaos::FRealSingle ChaosImmediate_Joint_ShockPropagation = -1.0f;
 Chaos::FRealSingle ChaosImmediate_Joint_Stiffness = -1.0f;
 Chaos::FRealSingle ChaosImmediate_Joint_SoftLinearStiffness = -1.0f;
 Chaos::FRealSingle ChaosImmediate_Joint_SoftTwistStiffness = -1.0f;
@@ -110,11 +113,14 @@ FAutoConsoleVariableRef CVarChaosImmPhysJointPushOutPairIterations(TEXT("p.Chaos
 FAutoConsoleVariableRef CVarChaosImmPhysJointSwingTwistAngleTolerance(TEXT("p.Chaos.ImmPhys.Joint.SwingTwistAngleTolerance"), ChaosImmediate_Joint_SwingTwistAngleTolerance, TEXT("SwingTwistAngleTolerance."));
 FAutoConsoleVariableRef CVarChaosImmPhysJointPositionTolerance(TEXT("p.Chaos.ImmPhys.Joint.PositionTolerance"), ChaosImmediate_Joint_PositionTolerance, TEXT("PositionTolerance."));
 FAutoConsoleVariableRef CVarChaosImmPhysJointAngleTolerance(TEXT("p.Chaos.ImmPhys.Joint.AngleTolerance"), ChaosImmediate_Joint_AngleTolerance, TEXT("AngleTolerance."));
+FAutoConsoleVariableRef CVarChaosImmPhysJointNumShockPropagationIterations(TEXT("p.Chaos.ImmPhys.Joint.NumShockPropagationIterations"), ChaosImmediate_Joint_NumShockPropagationIterations, TEXT("How many iterations to run shock propagation for"));
+FAutoConsoleVariableRef CVarChaosImmPhysJointSolvePositionLast(TEXT("p.Chaos.ImmPhys.Joint.SolvePositionLast"), ChaosImmediate_Joint_SolvePositionLast, TEXT("Should we solve joints in position-then-rotation order (false) rotation-then-position order (true, default)"));
 FAutoConsoleVariableRef CVarChaosImmPhysJointEnableTwistLimits(TEXT("p.Chaos.ImmPhys.Joint.EnableTwistLimits"), ChaosImmediate_Joint_EnableTwistLimits, TEXT("EnableTwistLimits."));
 FAutoConsoleVariableRef CVarChaosImmPhysJointEnableSwingLimits(TEXT("p.Chaos.ImmPhys.Joint.EnableSwingLimits"), ChaosImmediate_Joint_EnableSwingLimits, TEXT("EnableSwingLimits."));
 FAutoConsoleVariableRef CVarChaosImmPhysJointEnableDrives(TEXT("p.Chaos.ImmPhys.Joint.EnableDrives"), ChaosImmediate_Joint_EnableDrives, TEXT("EnableDrives."));
 FAutoConsoleVariableRef CVarChaosImmPhysJointLinearProjection(TEXT("p.Chaos.ImmPhys.Joint.LinearProjection"), ChaosImmediate_Joint_LinearProjection, TEXT("6Dof joint projection amount override (if >= 0)."));
 FAutoConsoleVariableRef CVarChaosImmPhysJointAngularProjection(TEXT("p.Chaos.ImmPhys.Joint.AngularProjection"), ChaosImmediate_Joint_AngularProjection, TEXT("6Dof joint projection amount override (if >= 0)."));
+FAutoConsoleVariableRef CVarChaosImmPhysJointShockPropagation(TEXT("p.Chaos.ImmPhys.Joint.ShockPropagation"), ChaosImmediate_Joint_ShockPropagation, TEXT("6Dof joint shock propagation override (if >= 0)."));
 FAutoConsoleVariableRef CVarChaosImmPhysJointStiffness(TEXT("p.Chaos.ImmPhys.Joint.Stiffness"), ChaosImmediate_Joint_Stiffness, TEXT("6Dof joint stiffness override (if > 0)."));
 FAutoConsoleVariableRef CVarChaosImmPhysJointSoftLinearStiffness(TEXT("p.Chaos.ImmPhys.Joint.SoftLinearStiffness"), ChaosImmediate_Joint_SoftLinearStiffness, TEXT("6Dof joint soft linear stiffness override (if > 0)."));
 FAutoConsoleVariableRef CVarChaosImmPhysJointSoftTwistStiffness(TEXT("p.Chaos.ImmPhys.Joint.SoftTwistStiffness"), ChaosImmediate_Joint_SoftTwistStiffness, TEXT("6Dof joint SoftTwist stiffness override (if > 0)."));
@@ -135,6 +141,13 @@ FAutoConsoleVariableRef CVarChaosImmPhysJointMaxInertiaRatio(TEXT("p.Chaos.ImmPh
 int32 ChaosImmediate_SolverType = (int32)Chaos::EConstraintSolverType::StandardPbd;
 FAutoConsoleVariableRef CVarChaosImmPhysSolverType(TEXT("p.Chaos.ImmPhys.SolverType"), ChaosImmediate_SolverType, TEXT("0 = None; 1 = GbfPbd; 2 = Pbd; 3 = QuasiPbd"));
 
+// Whether to use the linear joint solver which is significantly faster than the non-linear one but less accurate. 
+// In princple this choice is independent of the SolverType, but in practice you might want to use the non-Linear Joint Solver when you are using the 
+// Stanbard PBD Solver algorithm. This is because the Linear Joint Solver implementation does not implmenent the Projection phase which is part of the 
+// stabilization of the standard PBD solver. Ymmv.
+bool bChaosImmediate_Joint_UseLinearSolver = true;
+FAutoConsoleVariableRef CVarChaosImmPhysJointUseCachedSolver(TEXT("p.Chaos.ImmPhys.Joint.UseLinearSolver"), bChaosImmediate_Joint_UseLinearSolver, TEXT("Use linear version of joint solver. (default is true"));
+
 //
 // end remove when finished
 //
@@ -780,6 +793,8 @@ namespace ImmediatePhysics_Chaos
 			JointsSettings.AngleTolerance = ChaosImmediate_Joint_AngleTolerance;
 			JointsSettings.MinParentMassRatio = ChaosImmediate_Joint_MinParentMassRatio;
 			JointsSettings.MaxInertiaRatio = ChaosImmediate_Joint_MaxInertiaRatio;
+			JointsSettings.NumShockPropagationIterations = ChaosImmediate_Joint_NumShockPropagationIterations;
+			JointsSettings.bSolvePositionLast = ChaosImmediate_Joint_SolvePositionLast != 0;
 			JointsSettings.bEnableTwistLimits = ChaosImmediate_Joint_EnableTwistLimits != 0;
 			JointsSettings.bEnableSwingLimits = ChaosImmediate_Joint_EnableSwingLimits != 0;
 			JointsSettings.bEnableDrives = ChaosImmediate_Joint_EnableDrives != 0;
@@ -788,6 +803,7 @@ namespace ImmediatePhysics_Chaos
 			JointsSettings.SwingStiffnessOverride = ChaosImmediate_Joint_Stiffness;
 			JointsSettings.LinearProjectionOverride = ChaosImmediate_Joint_LinearProjection;
 			JointsSettings.AngularProjectionOverride = ChaosImmediate_Joint_AngularProjection;
+			JointsSettings.ShockPropagationOverride = ChaosImmediate_Joint_ShockPropagation;
 			JointsSettings.SoftLinearStiffnessOverride = ChaosImmediate_Joint_SoftLinearStiffness;
 			JointsSettings.SoftTwistStiffnessOverride = ChaosImmediate_Joint_SoftTwistStiffness;
 			JointsSettings.SoftTwistDampingOverride = ChaosImmediate_Joint_SoftTwistDamping;
@@ -797,6 +813,7 @@ namespace ImmediatePhysics_Chaos
 			JointsSettings.LinearDriveDampingOverride = ChaosImmediate_Joint_LinearDriveDamping;
 			JointsSettings.AngularDriveStiffnessOverride = ChaosImmediate_Joint_AngularDriveStiffness;
 			JointsSettings.AngularDriveDampingOverride = ChaosImmediate_Joint_AngularDriveDamping;
+			JointsSettings.bUseLinearSolver = bChaosImmediate_Joint_UseLinearSolver;
 			Implementation->Joints.SetSettings(JointsSettings);
 
 			Implementation->Collisions.SetRestitutionEnabled(ChaosImmediate_Collision_RestitutionEnabled != 0);
diff --git a/Engine/Source/Runtime/Experimental/Chaos/Private/Chaos/ChaosDebugDraw.cpp b/Engine/Source/Runtime/Experimental/Chaos/Private/Chaos/ChaosDebugDraw.cpp
--- a/Engine/Source/Runtime/Experimental/Chaos/Private/Chaos/ChaosDebugDraw.cpp
+++ b/Engine/Source/Runtime/Experimental/Chaos/Private/Chaos/ChaosDebugDraw.cpp
@@ -815,11 +815,11 @@ namespace Chaos
 				// Are we within the region of interest?
 				const FReal Particle0Size = Contact.GetParticle0()->HasBounds() ? 0.5f * Contact.GetParticle0()->LocalBounds().Extents().Size() : TNumericLimits<FReal>::Max();
 				const FReal Particle1Size = Contact.GetParticle1()->HasBounds() ? 0.5f * Contact.GetParticle1()->LocalBounds().Extents().Size() : TNumericLimits<FReal>::Max();
-				if (!FDebugDrawQueue::GetInstance().IsInRegionOfInterest(WorldActorTransform0.GetLocation(), Particle0Size))
+				if (!FDebugDrawQueue::GetInstance().IsInRegionOfInterest(SpaceTransform.TransformPositionNoScale(WorldActorTransform0.GetLocation()), Particle0Size))
 				{
 					return;
 				}
-				if (!FDebugDrawQueue::GetInstance().IsInRegionOfInterest(WorldActorTransform1.GetLocation(), Particle1Size))
+				if (!FDebugDrawQueue::GetInstance().IsInRegionOfInterest(SpaceTransform.TransformPositionNoScale(WorldActorTransform1.GetLocation()), Particle1Size))
 				{
 					return;
 				}
@@ -872,12 +872,12 @@ namespace Chaos
 					if ((Settings.PushOutScale > 0) && !ManifoldPoint.NetPushOut.IsNearlyZero())
 					{
 						FColor Color = (ColorScale * PushOutImpusleColor).ToFColor(false);
-						FDebugDrawQueue::GetInstance().DrawDebugLine(WorldPointPlaneLocation, WorldPointPlaneLocation + Settings.DrawScale * Settings.PushOutScale * ManifoldPoint.NetPushOut, Color, false, KINDA_SMALL_NUMBER, Settings.DrawPriority, Settings.LineThickness);
+						FDebugDrawQueue::GetInstance().DrawDebugLine(WorldPointPlaneLocation, WorldPointPlaneLocation + Settings.DrawScale * Settings.PushOutScale * SpaceTransform.TransformVectorNoScale(ManifoldPoint.NetPushOut), Color, false, KINDA_SMALL_NUMBER, Settings.DrawPriority, Settings.LineThickness);
 					}
 					if ((Settings.ImpulseScale > 0) && !ManifoldPoint.NetImpulse.IsNearlyZero())
 					{
 						FColor Color = (ColorScale * ImpulseColor).ToFColor(false);
-						FDebugDrawQueue::GetInstance().DrawDebugLine(WorldPointPlaneLocation, WorldPointPlaneLocation + Settings.DrawScale * Settings.ImpulseScale * ManifoldPoint.NetImpulse, Color, false, KINDA_SMALL_NUMBER, Settings.DrawPriority, Settings.LineThickness);
+						FDebugDrawQueue::GetInstance().DrawDebugLine(WorldPointPlaneLocation, WorldPointPlaneLocation + Settings.DrawScale * Settings.ImpulseScale * SpaceTransform.TransformVectorNoScale(ManifoldPoint.NetImpulse), Color, false, KINDA_SMALL_NUMBER, Settings.DrawPriority, Settings.LineThickness);
 					}
 
 					// Manifold plane and normal
diff --git a/Engine/Source/Runtime/Experimental/Chaos/Private/Chaos/Joint/JointConstraintsCVars.cpp b/Engine/Source/Runtime/Experimental/Chaos/Private/Chaos/Joint/JointConstraintsCVars.cpp
--- a/Engine/Source/Runtime/Experimental/Chaos/Private/Chaos/Joint/JointConstraintsCVars.cpp
+++ b/Engine/Source/Runtime/Experimental/Chaos/Private/Chaos/Joint/JointConstraintsCVars.cpp
@@ -29,5 +29,3 @@ FAutoConsoleVariableRef CVarChaosJointLinearVelocityThresholdToApplyRestitution(
 float Chaos_Joint_AngularVelocityThresholdToApplyRestitution = 1e-2f;
 FAutoConsoleVariableRef CVarChaosJointAngularVelocityThresholdToApplyRestitution(TEXT("p.Chaos.Joint.AngularVelocityThresholdToApplyRestitution"), Chaos_Joint_AngularVelocityThresholdToApplyRestitution, TEXT("Apply restitution only if initial velocity is higher than this threshold (used in Quasipbd)"));
 
-bool bChaos_Joint_UseCachedSolver = true;
-FAutoConsoleVariableRef CVarChaosJointUseCachedSolver(TEXT("p.Chaos.Joint.UseCachedSolver"), bChaos_Joint_UseCachedSolver, TEXT("Use cached version of joint solver. (default is false"));
diff --git a/Engine/Source/Runtime/Experimental/Chaos/Private/Chaos/Joint/PBDJointCachedSolverGaussSeidel.cpp b/Engine/Source/Runtime/Experimental/Chaos/Private/Chaos/Joint/PBDJointCachedSolverGaussSeidel.cpp
--- a/Engine/Source/Runtime/Experimental/Chaos/Private/Chaos/Joint/PBDJointCachedSolverGaussSeidel.cpp
+++ b/Engine/Source/Runtime/Experimental/Chaos/Private/Chaos/Joint/PBDJointCachedSolverGaussSeidel.cpp
@@ -90,6 +90,50 @@ void FPBDJointCachedSolver::Update(
 	//UpdateIsActive();
 }
 
+void FPBDJointCachedSolver::UpdateMass0()
+{
+	// @todo(chaos): this needs to recache all the mass-dependent state on the axis data etc
+	if ((ConditionedInvMs[0] > 0) && (InvMScales[0] > 0))
+	{
+		InvMs[0] = InvMScales[0] * ConditionedInvMs[0];
+		InvIs[0] = Utilities::ComputeWorldSpaceInertia(Q(0), InvMScales[0] * ConditionedInvILs[0]);
+	}
+	else
+	{
+		InvMs[0] = 0;
+		InvIs[0] = FMatrix33(0);
+	}
+}
+
+void FPBDJointCachedSolver::UpdateMass1()
+{
+	// @todo(chaos): this needs to recache all the mass-dependent state on the axis data etc
+	if ((ConditionedInvMs[1] > 0) && (InvMScales[1] > 0))
+	{
+		InvMs[1] = InvMScales[1] * ConditionedInvMs[1];
+		InvIs[1] = Utilities::ComputeWorldSpaceInertia(Q(1), InvMScales[1] * ConditionedInvILs[1]);
+	}
+	else
+	{
+		InvMs[1] = 0;
+		InvIs[1] = FMatrix33(0);
+	}
+}
+
+void FPBDJointCachedSolver::SetInvMassScales(const FReal InvMScale0, const FReal InvMScale1)
+{
+	if (InvMScales[0] != InvMScale0)
+	{
+		InvMScales[0] = InvMScale0;
+		UpdateMass0();
+	}
+	if (InvMScales[1] != InvMScale1)
+	{
+		InvMScales[1] = InvMScale1;
+		UpdateMass1();
+	}
+}
+
 void FPBDJointCachedSolver::EnableProjection()
 {
 	Body0().SetInvMScale(0);
@@ -115,12 +159,11 @@ void FPBDJointCachedSolver::Init(
 	SolverBodies[0].SetInvMScale(JointSettings.ParentInvMassScale);
 	SolverBodies[1].SetInvMScale(FReal(1));
 
-	FVec3 ConditionedInvILs[2];
-	FPBDJointUtilities::ConditionInverseMassAndInertia(Body0().InvM(), Body1().InvM(), Body0().InvILocal(), Body1().InvILocal(),
-		SolverSettings.MinParentMassRatio, SolverSettings.MaxInertiaRatio, ConditionedInvMs[0], ConditionedInvMs[1], ConditionedInvILs[0], ConditionedInvILs[1]);
-
-	ConditionedInvIs[0] = Utilities::ComputeWorldSpaceInertia(Q(0), ConditionedInvILs[0]);
-	ConditionedInvIs[1] = Utilities::ComputeWorldSpaceInertia(Q(1), ConditionedInvILs[1]);
+	InvMScales[0] = FReal(1);
+	InvMScales[1] = FReal(1);
+	FPBDJointUtilities::ConditionInverseMassAndInertia(Body0().InvM(), Body1().InvM(), Body0().InvILocal(), Body1().InvILocal(), SolverSettings.MinParentMassRatio, SolverSettings.MaxInertiaRatio, ConditionedInvMs[0], ConditionedInvMs[1], ConditionedInvILs[0], ConditionedInvILs[1]);
+	UpdateMass0();
+	UpdateMass1();
 
 	NetLinearImpulse = FVec3(0);
 	NetAngularImpulse = FVec3(0);
@@ -176,11 +219,22 @@ void FPBDJointCachedSolver::ApplyConstraints(
 	NumActiveConstraints = 0;
 	SolverStiffness = InSolverStiffness;
 
-	ApplyPositionConstraints(Dt);
-	ApplyRotationConstraints(Dt);
+	if (SolverSettings.bSolvePositionLast)
+	{
+		ApplyRotationConstraints(Dt);
+		ApplyPositionConstraints(Dt);
 
-	ApplyPositionDrives(Dt);
-	ApplyRotationDrives(Dt);
+		ApplyRotationDrives(Dt);
+		ApplyPositionDrives(Dt);
+	}
+	else
+	{
+		ApplyPositionConstraints(Dt);
+		ApplyRotationConstraints(Dt);
+
+		ApplyPositionDrives(Dt);
+		ApplyRotationDrives(Dt);
+	}
 
 	//UpdateIsActive();
 }
@@ -199,8 +253,16 @@ void FPBDJointCachedSolver::ApplyVelocityConstraints(
 	// (including restitution) are also enforced. This also prevents any position
 	// errors from the previous frame getting converted into energy.
 
-	ApplyLinearVelocityConstraints();
-	ApplyAngularVelocityConstraints();
+	if (SolverSettings.bSolvePositionLast)
+	{
+		ApplyAngularVelocityConstraints();
+		ApplyLinearVelocityConstraints();
+	}
+	else
+	{
+		ApplyLinearVelocityConstraints();
+		ApplyAngularVelocityConstraints();
+	}
 
 	// @todo(chaos): We can also apply velocity drives here rather than in the Pbd pass
 }
diff --git a/Engine/Source/Runtime/Experimental/Chaos/Private/Chaos/Joint/PBDJointSolverGaussSeidel.cpp b/Engine/Source/Runtime/Experimental/Chaos/Private/Chaos/Joint/PBDJointSolverGaussSeidel.cpp
--- a/Engine/Source/Runtime/Experimental/Chaos/Private/Chaos/Joint/PBDJointSolverGaussSeidel.cpp
+++ b/Engine/Source/Runtime/Experimental/Chaos/Private/Chaos/Joint/PBDJointSolverGaussSeidel.cpp
@@ -128,10 +128,11 @@ namespace Chaos
 		SolverBodies[0].SetInvMScale(JointSettings.ParentInvMassScale);
 		SolverBodies[1].SetInvMScale(FReal(1));
 
-		FVec3 ConditionedInvILs[2];
+		InvMScales[0] = FReal(1);
+		InvMScales[1] = FReal(1);
 		FPBDJointUtilities::ConditionInverseMassAndInertia(Body0().InvM(), Body1().InvM(), Body0().InvILocal(), Body1().InvILocal(), SolverSettings.MinParentMassRatio, SolverSettings.MaxInertiaRatio, ConditionedInvMs[0], ConditionedInvMs[1], ConditionedInvILs[0], ConditionedInvILs[1]);
-		ConditionedInvIs[0] = Utilities::ComputeWorldSpaceInertia(Q(0), ConditionedInvILs[0]);
-		ConditionedInvIs[1] = Utilities::ComputeWorldSpaceInertia(Q(1), ConditionedInvILs[1]);
+		UpdateMass0();
+		UpdateMass1();
 
 		NetLinearImpulse = FVec3(0);
 		NetAngularImpulse = FVec3(0);
@@ -194,6 +195,47 @@ namespace Chaos
 		UpdateIsActive(Dt, SolverSettings, JointSettings);
 	}
 
+	void FPBDJointSolver::UpdateMass0()
+	{
+		if ((ConditionedInvMs[0] > 0) && (InvMScales[0] > 0))
+		{
+			InvMs[0] = InvMScales[0] * ConditionedInvMs[0];
+			InvIs[0] = Utilities::ComputeWorldSpaceInertia(Q(0), InvMScales[0] * ConditionedInvILs[0]);
+		}
+		else
+		{
+			InvMs[0] = 0;
+			InvIs[0] = FMatrix33(0);
+		}
+	}
+
+	void FPBDJointSolver::UpdateMass1()
+	{
+		if ((ConditionedInvMs[1] > 0) && (InvMScales[1] > 0))
+		{
+			InvMs[1] = InvMScales[1] * ConditionedInvMs[1];
+			InvIs[1] = Utilities::ComputeWorldSpaceInertia(Q(1), InvMScales[1] * ConditionedInvILs[1]);
+		}
+		else
+		{
+			InvMs[1] = 0;
+			InvIs[1] = FMatrix33(0);
+		}
+	}
+
+	void FPBDJointSolver::SetInvMassScales(const FReal InvMScale0, const FReal InvMScale1)
+	{
+		if (InvMScales[0] != InvMScale0)
+		{
+			InvMScales[0] = InvMScale0;
+			UpdateMass0();
+		}
+		if (InvMScales[1] != InvMScale1)
+		{
+			InvMScales[1] = InvMScale1;
+			UpdateMass1();
+		}
+	}
 
 	void FPBDJointSolver::ApplyConstraints(
 		const FReal Dt,
@@ -204,11 +246,22 @@ namespace Chaos
 		NumActiveConstraints = 0;
 		SolverStiffness = InSolverStiffness;
 
-		ApplyPositionConstraints(Dt, SolverSettings, JointSettings);
-		ApplyRotationConstraints(Dt, SolverSettings, JointSettings);
-	
-		ApplyPositionDrives(Dt, SolverSettings, JointSettings);
-		ApplyRotationDrives(Dt, SolverSettings, JointSettings);
+		if (SolverSettings.bSolvePositionLast)
+		{
+			ApplyRotationConstraints(Dt, SolverSettings, JointSettings);
+			ApplyPositionConstraints(Dt, SolverSettings, JointSettings);
+
+			ApplyRotationDrives(Dt, SolverSettings, JointSettings);
+			ApplyPositionDrives(Dt, SolverSettings, JointSettings);
+		}
+		else
+		{
+			ApplyPositionConstraints(Dt, SolverSettings, JointSettings);
+			ApplyRotationConstraints(Dt, SolverSettings, JointSettings);
+
+			ApplyPositionDrives(Dt, SolverSettings, JointSettings);
+			ApplyRotationDrives(Dt, SolverSettings, JointSettings);
+		}
 
 		UpdateIsActive(Dt, SolverSettings, JointSettings);
 	}
@@ -228,8 +281,16 @@ namespace Chaos
 		// (including restitution) are also enforced. This also prevents any position
 		// errors from the previous frame getting converted into energy.
 
-		ApplyLinearVelocityConstraints(Dt, SolverSettings, JointSettings);
-		ApplyAngularVelocityConstraints(Dt, SolverSettings, JointSettings);
+		if (SolverSettings.bSolvePositionLast)
+		{
+			ApplyAngularVelocityConstraints(Dt, SolverSettings, JointSettings);
+			ApplyLinearVelocityConstraints(Dt, SolverSettings, JointSettings);
+		}
+		else
+		{
+			ApplyLinearVelocityConstraints(Dt, SolverSettings, JointSettings);
+			ApplyAngularVelocityConstraints(Dt, SolverSettings, JointSettings);
+		}
 
 		// @todo(chaos): We can also apply velocity drives here rather than in the Pbd pass
 	}
@@ -242,7 +303,7 @@ namespace Chaos
 		const FPBDJointSettings& JointSettings)
 	{
 		// @todo(chaos): We need to handle parent/child being the other way round
-		if (!Body1().IsDynamic())
+		if (!IsDynamic(1))
 		{
 			// If child is kinematic, return. 
 			return;
@@ -550,11 +611,11 @@ namespace Chaos
 		if (bLinearLocked[0] && bLinearLocked[1] && bLinearLocked[2])
 		{
 			// Hard point constraint (most common case)
-			if (!Body0().IsDynamic())
+			if (!IsDynamic(0))
 			{
 				ApplyPointPositionConstraintKD(0, 1, Dt, SolverSettings, JointSettings);
 			}
-			else if (!Body1().IsDynamic())
+			else if (!IsDynamic(1))
 			{
 				ApplyPointPositionConstraintKD(1, 0, Dt, SolverSettings, JointSettings);
 			}
@@ -1096,11 +1157,11 @@ namespace Chaos
 		}
 		else
 		{
-			if (Body0().IsDynamic())
+			if (IsDynamic(0))
 			{
 				Body(0).ApplyRotationDelta(DR0);
 			}
-			if (Body1().IsDynamic())
+			if (IsDynamic(1))
 			{
 				Body(1).ApplyRotationDelta(DR1);
 			}
@@ -1310,11 +1371,11 @@ namespace Chaos
 		const FReal Angle,
 		const int32 AngularHardLambdaIndex)
 	{
-		if (!Body0().IsDynamic())
+		if (!IsDynamic(0))
 		{
 			ApplyRotationConstraintKD(0, 1, JointStiffness, Axis, Angle, AngularHardLambdaIndex);
 		}
-		else if (!Body1().IsDynamic())
+		else if (!IsDynamic(1))
 		{
 			ApplyRotationConstraintKD(1, 0, JointStiffness, Axis, -Angle, AngularHardLambdaIndex);
 		}
@@ -1340,8 +1401,8 @@ namespace Chaos
 		const FReal AngVelTarget,
 		FReal& Lambda)
 	{
-		check(!Body(KIndex).IsDynamic());
-		check(Body(DIndex).IsDynamic());
+		check(!IsDynamic(KIndex));
+		check(IsDynamic(DIndex));
 
 		if (bRealTypeCompatibleWithISPC && bChaos_Joint_ISPC_Enabled)
 		{
@@ -1396,8 +1457,8 @@ namespace Chaos
 		const FReal AngVelTarget,
 		FReal& Lambda)
 	{
-		check(Body(0).IsDynamic());
-		check(Body(1).IsDynamic());
+		check(IsDynamic(0));
+		check(IsDynamic(1));
 
 		if (bRealTypeCompatibleWithISPC && bChaos_Joint_ISPC_Enabled)
 		{
@@ -1455,11 +1516,11 @@ namespace Chaos
 		const FReal AngVelTarget,
 		FReal& Lambda)
 	{
-		if (!Body0().IsDynamic())
+		if (!IsDynamic(0))
 		{
 			ApplyRotationConstraintSoftKD(0, 1, Dt, JointStiffness, JointDamping, bAccelerationMode, Axis, Angle, AngVelTarget, Lambda);
 		}
-		else if (!Body1().IsDynamic())
+		else if (!IsDynamic(1))
 		{
 			ApplyRotationConstraintSoftKD(1, 0, Dt, JointStiffness, JointDamping, bAccelerationMode, Axis, -Angle, -AngVelTarget, Lambda);
 		}
@@ -1933,8 +1994,8 @@ namespace Chaos
 		const FPBDJointSolverSettings& SolverSettings,
 		const FPBDJointSettings& JointSettings)
 	{
-		check(!Body(KIndex).IsDynamic());
-		check(Body(DIndex).IsDynamic());
+		check(!IsDynamic(KIndex));
+		check(IsDynamic(DIndex));
 
 		FReal Stiffness = SolverStiffness * FPBDJointUtilities::GetLinearStiffness(SolverSettings, JointSettings);
 		const FVec3 CX = ConnectorXs[DIndex] - ConnectorXs[KIndex];
@@ -1976,8 +2037,8 @@ namespace Chaos
 		const FPBDJointSolverSettings& SolverSettings,
 		const FPBDJointSettings& JointSettings)
 	{
-		check(Body(0).IsDynamic());
-		check(Body(1).IsDynamic());
+		check(IsDynamic(0));
+		check(IsDynamic(1));
 
 		FReal Stiffness = SolverStiffness * FPBDJointUtilities::GetLinearStiffness(SolverSettings, JointSettings);
 		const FVec3 CX = ConnectorXs[1] - ConnectorXs[0];
diff --git a/Engine/Source/Runtime/Experimental/Chaos/Private/Chaos/PBDJointConstraintUtilities.cpp b/Engine/Source/Runtime/Experimental/Chaos/Private/Chaos/PBDJointConstraintUtilities.cpp
--- a/Engine/Source/Runtime/Experimental/Chaos/Private/Chaos/PBDJointConstraintUtilities.cpp
+++ b/Engine/Source/Runtime/Experimental/Chaos/Private/Chaos/PBDJointConstraintUtilities.cpp
@@ -601,6 +601,24 @@ namespace Chaos
 		return JointSettings.AngularDriveForceMode == EJointForceMode::Acceleration;
 	}
 
+	FReal FPBDJointUtilities::GetShockPropagationInvMassScale(
+		EConstraintSolverType SolverType,
+		int32 It,
+		int32 NumIts,
+		const FPBDJointSolverSettings& SolverSettings,
+		const FPBDJointSettings& JointSettings)
+	{
+		// Shock propagation is only enabled for the last iteration, and only for the QPBD solver
+		const bool bEnableShockPropagation = JointSettings.bProjectionEnabled && (It >= NumIts - SolverSettings.NumShockPropagationIterations) && (SolverType == EConstraintSolverType::QuasiPbd);
+
+		if (bEnableShockPropagation)
+		{
+			return (SolverSettings.ShockPropagationOverride >= 0.0f) ? SolverSettings.ShockPropagationOverride : JointSettings.ShockPropagation;
+		}
+
+		return FReal(1);
+	}
+
 
 	FVec3 FPBDJointUtilities::ConditionInertia(const FVec3& InI, const FReal MaxRatio)
 	{
diff --git a/Engine/Source/Runtime/Experimental/Chaos/Private/Chaos/PBDJointConstraints.cpp b/Engine/Source/Runtime/Experimental/Chaos/Private/Chaos/PBDJointConstraints.cpp
--- a/Engine/Source/Runtime/Experimental/Chaos/Private/Chaos/PBDJointConstraints.cpp
+++ b/Engine/Source/Runtime/Experimental/Chaos/Private/Chaos/PBDJointConstraints.cpp
@@ -187,6 +187,7 @@ namespace Chaos
 		: Stiffness(1)
 		, LinearProjection(0)
 		, AngularProjection(0)
+		, ShockPropagation(0)
 		, ParentInvMassScale(1)
 		, bCollisionEnabled(true)
 		, bProjectionEnabled(false)
@@ -356,6 +357,9 @@ namespace Chaos
 		, MinSolverStiffness(1)
 		, MaxSolverStiffness(1)
 		, NumIterationsAtMaxSolverStiffness(1)
+		, NumShockPropagationIterations(0)
+		, bUseLinearSolver(true)
+		, bSolvePositionLast(true)
 		, bEnableTwistLimits(true)
 		, bEnableSwingLimits(true)
 		, bEnableDrives(true)
@@ -364,6 +368,7 @@ namespace Chaos
 		, SwingStiffnessOverride(-1)
 		, LinearProjectionOverride(-1)
 		, AngularProjectionOverride(-1)
+		, ShockPropagationOverride(-1)
 		, LinearDriveStiffnessOverride(-1)
 		, LinearDriveDampingOverride(-1)
 		, AngularDriveStiffnessOverride(-1)
@@ -757,7 +762,7 @@ namespace Chaos
 			bJointsDirty = false;
 		}
 
-		if (bChaos_Joint_UseCachedSolver)
+		if (Settings.bUseLinearSolver)
 		{
 			CachedConstraintSolvers.SetNum(NumConstraints());
 		}
@@ -771,7 +776,7 @@ namespace Chaos
 	{
 		SCOPE_CYCLE_COUNTER(STAT_Joints_UnprepareTick);
 
-		if (bChaos_Joint_UseCachedSolver)
+		if (Settings.bUseLinearSolver)
 		{
 			CachedConstraintSolvers.Empty();
 		}
@@ -864,7 +869,7 @@ namespace Chaos
 		FSolverBody* Body0 = SolverData.GetBodyContainer().FindOrAdd(Particle0);
 		FSolverBody* Body1 = SolverData.GetBodyContainer().FindOrAdd(Particle1);
 
-		if (bChaos_Joint_UseCachedSolver)
+		if (Settings.bUseLinearSolver)
 		{
 			FPBDJointCachedSolver& Solver = CachedConstraintSolvers[ConstraintIndex];
 			Solver.Init(
@@ -921,7 +926,7 @@ namespace Chaos
 				int32 Index0, Index1;
 				GetConstrainedParticleIndices(ConstraintIndex, Index0, Index1);
 
-				if (bChaos_Joint_UseCachedSolver)
+				if (Settings.bUseLinearSolver)
 				{
 					FPBDJointCachedSolver& Solver = CachedConstraintSolvers[ConstraintIndex];
 					// NOTE: LinearImpulse/AngularImpulse in the solver are not really impulses - they are mass-weighted position/rotation delta, or (impulse x dt).
@@ -1024,7 +1029,7 @@ namespace Chaos
 			for (int32 ConstraintIndex : SolverData.GetConstraintIndices(ContainerId))
 			{
 				const FPBDJointSettings& JointSettings = ConstraintSettings[ConstraintIndex];
-				if (bChaos_Joint_UseCachedSolver)
+				if (Settings.bUseLinearSolver)
 				{
 					if (JointSettings.bProjectionEnabled)
 					{
@@ -1083,7 +1088,7 @@ namespace Chaos
 		}
 
 		// check valid particle and solver state
-		if (bChaos_Joint_UseCachedSolver)
+		if (Settings.bUseLinearSolver)
 		{
 			const FPBDJointCachedSolver& Solver = CachedConstraintSolvers[ConstraintIndex];
 			if ((Particle0->Sleeping() && Particle1->Sleeping())
@@ -1121,7 +1126,7 @@ namespace Chaos
 		UE_LOG(LogChaosJoint, VeryVerbose, TEXT("Solve Joint Constraint %d %s %s (dt = %f; it = %d / %d)"), ConstraintIndex, *Constraint[0]->ToString(), *Constraint[1]->ToString(), Dt, It, NumIts);
 
 		const FPBDJointSettings& JointSettings = ConstraintSettings[ConstraintIndex];
-		if (bChaos_Joint_UseCachedSolver)
+		if (Settings.bUseLinearSolver)
 		{
 			FPBDJointCachedSolver& Solver = CachedConstraintSolvers[ConstraintIndex];
 
@@ -1142,6 +1147,11 @@ namespace Chaos
 			const bool bWasActive = Solver.GetIsActive();
 			Solver.Update(Dt, Settings, JointSettings);
 
+			// Set parent inverser mass scale based on current shock propagation state
+			// @todo(chaos): needs to handle the parent/child being in opposite order
+			const FReal ShockPropagationInvMassScale = FPBDJointUtilities::GetShockPropagationInvMassScale(SolverType, It, NumIts, Settings, JointSettings);
+			Solver.SetInvMassScales(ShockPropagationInvMassScale, FReal(1));
+
 			if (!bWasActive && !Solver.GetIsActive() && bChaos_Joint_EarlyOut_Enabled)
 			{
 				return false;
@@ -1198,6 +1208,11 @@ namespace Chaos
 			const bool bWasActive = Solver.GetIsActive();
 			Solver.Update(Dt, Settings, JointSettings);
 
+			// Set parent inverser mass scale based on current shock propagation state
+			// @todo(chaos): needs to handle the parent/child being in opposite order
+			const FReal ShockPropagationInvMassScale = FPBDJointUtilities::GetShockPropagationInvMassScale(SolverType, It, NumIts, Settings, JointSettings);
+			Solver.SetInvMassScales(ShockPropagationInvMassScale, FReal(1));
+
 			if (!bWasActive && !Solver.GetIsActive() && bChaos_Joint_EarlyOut_Enabled)
 			{
 				return false;
@@ -1246,7 +1261,7 @@ namespace Chaos
 		UE_LOG(LogChaosJoint, VeryVerbose, TEXT("Project Joint Constraint %d %s %s (dt = %f; it = %d / %d)"), ConstraintIndex, *Constraint[0]->ToString(), *Constraint[1]->ToString(), Dt, It, NumIts);
 
 		const FPBDJointSettings& JointSettings = ConstraintSettings[ConstraintIndex];
-		if (bChaos_Joint_UseCachedSolver)
+		if (Settings.bUseLinearSolver)
 		{
 			FPBDJointCachedSolver& Solver = CachedConstraintSolvers[ConstraintIndex];
 
@@ -1267,6 +1282,11 @@ namespace Chaos
 			const bool bWasActive = Solver.GetIsActive();
 			Solver.Update(Dt, Settings, JointSettings);
 			
+			// Set parent inverser mass scale based on current shock propagation state
+			// @todo(chaos): needs to handle the parent/child being in opposite order
+			const FReal ShockPropagationInvMassScale = FPBDJointUtilities::GetShockPropagationInvMassScale(SolverType, It, NumIts, Settings, JointSettings);
+			Solver.SetInvMassScales(ShockPropagationInvMassScale, FReal(1));
+
 			// For quasipbd, we may still need to stabilize velocities and solve for restitution.
 			if (!bWasActive && !Solver.GetIsActive() && SolverType != EConstraintSolverType::QuasiPbd && bChaos_Joint_EarlyOut_Enabled)
 			{
@@ -1319,7 +1339,12 @@ namespace Chaos
 			// If we were solved last iteration and nothing has changed since, we are done
 			const bool bWasActive = Solver.GetIsActive();
 			Solver.Update(Dt, Settings, JointSettings);
-			
+
+			// Set parent inverser mass scale based on current shock propagation state
+			// @todo(chaos): needs to handle the parent/child being in opposite order
+			const FReal ShockPropagationInvMassScale = FPBDJointUtilities::GetShockPropagationInvMassScale(SolverType, It, NumIts, Settings, JointSettings);
+			Solver.SetInvMassScales(ShockPropagationInvMassScale, FReal(1));
+
 			// For quasipbd, we may still need to stabilize velocities and solve for restitution.
 			if (!bWasActive && !Solver.GetIsActive() && SolverType != EConstraintSolverType::QuasiPbd && bChaos_Joint_EarlyOut_Enabled)
 			{
@@ -1415,8 +1440,8 @@ namespace Chaos
 			}
 		}
 
-		FConstraintSolverBody& Body0 = bChaos_Joint_UseCachedSolver ? CachedConstraintSolvers[ConstraintIndex].Body(0) : ConstraintSolvers[ConstraintIndex].Body(0);
-		FConstraintSolverBody& Body1 = bChaos_Joint_UseCachedSolver ? CachedConstraintSolvers[ConstraintIndex].Body(1) : ConstraintSolvers[ConstraintIndex].Body(1);
+		FConstraintSolverBody& Body0 = Settings.bUseLinearSolver ? CachedConstraintSolvers[ConstraintIndex].Body(0) : ConstraintSolvers[ConstraintIndex].Body(0);
+		FConstraintSolverBody& Body1 = Settings.bUseLinearSolver ? CachedConstraintSolvers[ConstraintIndex].Body(1) : ConstraintSolvers[ConstraintIndex].Body(1);
 
 		const FTransformPair& ConstraintFramesLocal = JointSettings.ConnectorTransforms;
 		FTransformPair ConstraintFramesGlobal(ConstraintFramesLocal[Index0] * FRigidTransform3(Body0.ActorP(), Body0.ActorQ()), ConstraintFramesLocal[Index1] * FRigidTransform3(Body1.ActorP(), Body1.ActorQ()));
diff --git a/Engine/Source/Runtime/Experimental/Chaos/Private/PBDRigidsSolver.cpp b/Engine/Source/Runtime/Experimental/Chaos/Private/PBDRigidsSolver.cpp
--- a/Engine/Source/Runtime/Experimental/Chaos/Private/PBDRigidsSolver.cpp
+++ b/Engine/Source/Runtime/Experimental/Chaos/Private/PBDRigidsSolver.cpp
@@ -179,6 +179,10 @@ namespace Chaos
 		int32 ChaosSolver_SolverType = (int32)Chaos::EConstraintSolverType::QuasiPbd;
 		FAutoConsoleVariableRef CVarChaosSolverSolverType(TEXT("p.Chaos.Solver.SolverType"), ChaosSolver_SolverType, TEXT("0 = None; 1 = GbfPbd; 2 = Pbd; 3 = QuasiPbd"));
 
+		// Joint solver mode (linear vs non-linear)
+		bool bChaosSolverJointUseLinearSolver = true;
+		FAutoConsoleVariableRef CVarChaosSolverJointUseCachedSolver(TEXT("p.Chaos.Solver.Joint.UseLinearSolver"), bChaosSolverJointUseLinearSolver, TEXT("Use linear version of joint solver. (default is true"));
+
 		// Iteration count cvars
 		// These override the engine config if >= 0
 
@@ -236,10 +240,15 @@ namespace Chaos
 		float ChaosSolverJointMinSolverStiffness = 1.0f;
 		float ChaosSolverJointMaxSolverStiffness = 1.0f;
 		int32 ChaosSolverJointNumIterationsAtMaxSolverStiffness = 1;
+		bool bChaosSolverJointSolvePositionFirst = false;
+		int32 ChaosSolverJointNumShockProagationIterations = 0;
+		FRealSingle ChaosSolverJointShockPropagation = -1.0f;
 		FAutoConsoleVariableRef CVarChaosSolverJointMinSolverStiffness(TEXT("p.Chaos.Solver.Joint.MinSolverStiffness"), ChaosSolverJointMinSolverStiffness, TEXT("Solver stiffness on first iteration, increases each iteration toward MaxSolverStiffness."));
 		FAutoConsoleVariableRef CVarChaosSolverJointMaxSolverStiffness(TEXT("p.Chaos.Solver.Joint.MaxSolverStiffness"), ChaosSolverJointMaxSolverStiffness, TEXT("Solver stiffness on last iteration, increases each iteration from MinSolverStiffness."));
 		FAutoConsoleVariableRef CVarChaosSolverJointNumIterationsAtMaxSolverStiffness(TEXT("p.Chaos.Solver.Joint.NumIterationsAtMaxSolverStiffness"), ChaosSolverJointNumIterationsAtMaxSolverStiffness, TEXT("How many iterations we want at MaxSolverStiffness."));
-
+		FAutoConsoleVariableRef CVarChaosSolverJointSolvePositionFirst(TEXT("p.Chaos.Solver.Joint.SolvePositionFirst"), bChaosSolverJointSolvePositionFirst, TEXT("Should we solve joints in position-then-rotation order (true) rotation-then-position order (false, default)"));
+		FAutoConsoleVariableRef CVarChaosSolverJointNumShockPropagationIterations(TEXT("p.Chaos.Solver.Joint.NumShockPropagationIterations"), ChaosSolverJointNumShockProagationIterations, TEXT("How many iterations to enable SHockProagation for."));
+		FAutoConsoleVariableRef CVarChaosSolverJointShockPropagation(TEXT("p.Chaos.Solver.Joint.ShockPropagation"), ChaosSolverJointShockPropagation, TEXT("6Dof joint shock propagation override (if >= 0)."));
 
 		int32 ChaosVisualDebuggerEnable = 1;
 		FAutoConsoleVariableRef CVarChaosVisualDebuggerEnable(TEXT("p.Chaos.VisualDebuggerEnable"), ChaosVisualDebuggerEnable, TEXT("Enable/Disable pushing/saving data to the visual debugger"));
@@ -807,6 +816,10 @@ namespace Chaos
 		JointsSettings.AngleTolerance = ChaosSolverJointAngleTolerance;
 		JointsSettings.MinParentMassRatio = ChaosSolverJointMinParentMassRatio;
 		JointsSettings.MaxInertiaRatio = ChaosSolverJointMaxInertiaRatio;
+		JointsSettings.bSolvePositionLast = bChaosSolverJointSolvePositionFirst;
+		JointsSettings.NumShockPropagationIterations = ChaosSolverJointNumShockProagationIterations;
+		JointsSettings.ShockPropagationOverride = ChaosSolverJointShockPropagation;
+		JointsSettings.bUseLinearSolver = bChaosSolverJointUseLinearSolver;
 		MEvolution->GetJointConstraints().SetSettings(JointsSettings);
 
 		// Apply CVAR overrides if set
diff --git a/Engine/Source/Runtime/Experimental/Chaos/Public/Chaos/Joint/PBDJointCachedSolverGaussSeidel.h b/Engine/Source/Runtime/Experimental/Chaos/Public/Chaos/Joint/PBDJointCachedSolverGaussSeidel.h
--- a/Engine/Source/Runtime/Experimental/Chaos/Public/Chaos/Joint/PBDJointCachedSolverGaussSeidel.h
+++ b/Engine/Source/Runtime/Experimental/Chaos/Public/Chaos/Joint/PBDJointCachedSolverGaussSeidel.h
@@ -172,12 +172,17 @@ struct FPointConstraintDatas
 
 		inline FReal InvM(int32 BodyIndex) const
 		{
-			return ConditionedInvMs[BodyIndex];
+			return InvMs[BodyIndex];
 		}
 
 		inline FMatrix33 InvI(int32 BodyIndex) const
 		{
-			return ConditionedInvIs[BodyIndex];
+			return InvIs[BodyIndex];
+		}
+
+		inline bool IsDynamic(int32 BodyIndex) const
+		{
+			return (InvM(BodyIndex) > 0);
 		}
 
 		inline const FVec3& GetNetLinearImpulse() const
@@ -244,6 +249,10 @@ struct FPointConstraintDatas
 
 		void EnableProjection();
 
+		void SetInvMassScales(
+			const FReal InvMScale0,
+			const FReal InvMScale1);
+
 	private:
 
 		// Initialize state dependent on particle positions
@@ -253,6 +262,9 @@ struct FPointConstraintDatas
 		void UpdateDerivedState(const int32 BodyIndex);
 		void UpdateDerivedState();
 
+		void UpdateMass0();
+		void UpdateMass1();
+
 		// Check to see if this constraint still needs further solved
 		// @todo(chaos): the term "active" is used inconsistently with the meaning elsewhere. Active should
 		// mean "contributed impusles". Should use "solved" rather than "active"
@@ -545,8 +557,11 @@ struct FPointConstraintDatas
 		FRotation3 InitConnectorRs[MaxConstrainedBodies];	// World-space joint connector rotations
 
 		// Conditioned InvM and InvI
+		FReal InvMScales[MaxConstrainedBodies];
 		FReal ConditionedInvMs[MaxConstrainedBodies];
-		FMatrix33 ConditionedInvIs[MaxConstrainedBodies];
+		FVec3 ConditionedInvILs[MaxConstrainedBodies];		// Local-space
+		FReal InvMs[MaxConstrainedBodies];
+		FMatrix33 InvIs[MaxConstrainedBodies];				// World-space
 
 		// Accumulated Impulse and AngularImpulse (Impulse * Dt since they are mass multiplied position corrections)
 		FVec3 NetLinearImpulse;
diff --git a/Engine/Source/Runtime/Experimental/Chaos/Public/Chaos/Joint/PBDJointSolverGaussSeidel.h b/Engine/Source/Runtime/Experimental/Chaos/Public/Chaos/Joint/PBDJointSolverGaussSeidel.h
--- a/Engine/Source/Runtime/Experimental/Chaos/Public/Chaos/Joint/PBDJointSolverGaussSeidel.h
+++ b/Engine/Source/Runtime/Experimental/Chaos/Public/Chaos/Joint/PBDJointSolverGaussSeidel.h
@@ -104,12 +104,17 @@ namespace Chaos
 
 		inline FReal InvM(int32 BodyIndex) const
 		{
-			return ConditionedInvMs[BodyIndex];
+			return InvMs[BodyIndex];
 		}
 
 		inline FMatrix33 InvI(int32 BodyIndex) const
 		{
-			return ConditionedInvIs[BodyIndex];
+			return InvIs[BodyIndex];
+		}
+
+		inline bool IsDynamic(int32 BodyIndex) const
+		{
+			return (InvM(BodyIndex) > 0);
 		}
 
 		inline const FVec3& GetNetLinearImpulse() const
@@ -175,6 +180,10 @@ namespace Chaos
 			const FPBDJointSolverSettings& SolverSettings,
 			const FPBDJointSettings& JointSettings);
 
+		void SetInvMassScales(
+			const FReal InvMScale0, 
+			const FReal InvMScale1);
+
 	private:
 
 		// Initialize state dependent on particle positions
@@ -184,6 +193,9 @@ namespace Chaos
 		void UpdateDerivedState(const int32 BodyIndex);
 		void UpdateDerivedState();
 
+		void UpdateMass0();
+		void UpdateMass1();
+
 		// Check to see if this constraint still needs further solved
 		// @todo(chaos): the term "active" is used inconsistently with the meaning elsewhere. Active should
 		// mean "contributed impusles". Should use "solved" rather than "active"
@@ -706,8 +718,11 @@ namespace Chaos
 		FRotation3 InitConnectorRs[MaxConstrainedBodies];	// World-space joint connector rotations
 
 		// Conditioned InvM and InvI
+		FReal InvMScales[MaxConstrainedBodies];
 		FReal ConditionedInvMs[MaxConstrainedBodies];
-		FMatrix33 ConditionedInvIs[MaxConstrainedBodies];
+		FVec3 ConditionedInvILs[MaxConstrainedBodies];		// Local-space
+		FReal InvMs[MaxConstrainedBodies];
+		FMatrix33 InvIs[MaxConstrainedBodies];				// World-space
 
 		// Accumulated Impulse and AngularImpulse (Impulse * Dt since they are mass multiplied position corrections)
 		FVec3 NetLinearImpulse;
diff --git a/Engine/Source/Runtime/Experimental/Chaos/Public/Chaos/JointProperties.inl b/Engine/Source/Runtime/Experimental/Chaos/Public/Chaos/JointProperties.inl
--- a/Engine/Source/Runtime/Experimental/Chaos/Public/Chaos/JointProperties.inl
+++ b/Engine/Source/Runtime/Experimental/Chaos/Public/Chaos/JointProperties.inl
@@ -6,6 +6,7 @@ CHAOS_INNER_JOINT_PROPERTY(JointSettings, CollisionEnabled, bCollisionEnabled, b
 CHAOS_INNER_JOINT_PROPERTY(JointSettings, ProjectionEnabled, bProjectionEnabled, bool)
 CHAOS_INNER_JOINT_PROPERTY(JointSettings, ProjectionLinearAlpha, LinearProjection, FReal)
 CHAOS_INNER_JOINT_PROPERTY(JointSettings, ProjectionAngularAlpha, AngularProjection, FReal)
+CHAOS_INNER_JOINT_PROPERTY(JointSettings, ShockPropagationAlpha, ShockPropagation, FReal)
 CHAOS_INNER_JOINT_PROPERTY(JointSettings, ParentInvMassScale, ParentInvMassScale, FReal)
 CHAOS_INNER_JOINT_PROPERTY(JointSettings, LinearBreakForce, LinearBreakForce, FReal)
 CHAOS_INNER_JOINT_PROPERTY(JointSettings, LinearPlasticityLimit, LinearPlasticityLimit, FReal)
diff --git a/Engine/Source/Runtime/Experimental/Chaos/Public/Chaos/PBDJointConstraintTypes.h b/Engine/Source/Runtime/Experimental/Chaos/Public/Chaos/PBDJointConstraintTypes.h
--- a/Engine/Source/Runtime/Experimental/Chaos/Public/Chaos/PBDJointConstraintTypes.h
+++ b/Engine/Source/Runtime/Experimental/Chaos/Public/Chaos/PBDJointConstraintTypes.h
@@ -129,13 +129,14 @@ namespace Chaos
 		FTransformPair ConnectorTransforms;
 
 		FReal Stiffness;
-		FReal LinearProjection;
-		FReal AngularProjection;
+		FReal LinearProjection;			// @chaos(todo): remove - old PBD solver only
+		FReal AngularProjection;		// @chaos(todo): remove - old PBD solver only
+		FReal ShockPropagation;
 		FReal ParentInvMassScale;
 
 		bool bCollisionEnabled;
-		bool bProjectionEnabled;
-		bool bSoftProjectionEnabled;
+		bool bProjectionEnabled;		// @chaos(todo): remove - old PBD solver only
+		bool bSoftProjectionEnabled;	// @chaos(todo): remove - old PBD solver only
 
 		TVector<EJointMotionType, 3> LinearMotionTypes;
 		FReal LinearLimit;
@@ -218,8 +219,16 @@ namespace Chaos
 		FReal MinSolverStiffness;
 		FReal MaxSolverStiffness;
 		int32 NumIterationsAtMaxSolverStiffness;
+		int32 NumShockPropagationIterations;
 
-		// @todo(ccaulfield): remove these TEMP overrides for testing
+		// Whether to use the linear or non-linear joint solver
+		bool bUseLinearSolver;
+
+		// Whether to solve rotation then position limits (true), or vice versa
+		// Solving position last leads to less separation at the joints when limits are being forced
+		bool bSolvePositionLast;
+
+		// @todo(chaos): remove these TEMP overrides for testing
 		bool bEnableTwistLimits;
 		bool bEnableSwingLimits;
 		bool bEnableDrives;
@@ -228,6 +237,7 @@ namespace Chaos
 		FReal SwingStiffnessOverride;
 		FReal LinearProjectionOverride;
 		FReal AngularProjectionOverride;
+		FReal ShockPropagationOverride;
 		FReal LinearDriveStiffnessOverride;
 		FReal LinearDriveDampingOverride;
 		FReal AngularDriveStiffnessOverride;
diff --git a/Engine/Source/Runtime/Experimental/Chaos/Public/Chaos/PBDJointConstraintUtilities.h b/Engine/Source/Runtime/Experimental/Chaos/Public/Chaos/PBDJointConstraintUtilities.h
--- a/Engine/Source/Runtime/Experimental/Chaos/Public/Chaos/PBDJointConstraintUtilities.h
+++ b/Engine/Source/Runtime/Experimental/Chaos/Public/Chaos/PBDJointConstraintUtilities.h
@@ -5,6 +5,7 @@
 
 #include "Chaos/ParticleHandleFwd.h"
 #include "Chaos/PBDJointConstraintTypes.h"
+#include "Chaos/Collision/CollisionApplyType.h"
 
 
 namespace Chaos
@@ -266,6 +267,13 @@ namespace Chaos
 			const FPBDJointSolverSettings& SolverSettings,
 			const FPBDJointSettings& JointSettings);
 
+		static FReal GetShockPropagationInvMassScale(
+			EConstraintSolverType SolverType,
+			int32 It,
+			int32 NumIts,
+			const FPBDJointSolverSettings& SolverSettings,
+			const FPBDJointSettings& JointSettings);
+
 		static FVec3 GetSphereLimitedPositionError(const FVec3& CX, const FReal Radius);
 		static FVec3 GetCylinderLimitedPositionError(const FVec3& CX, const FVec3& Axis, const FReal Limit, const EJointMotionType AxisMotion);
 		static FVec3 GetLineLimitedPositionError(const FVec3& CX, const FVec3& Axis, const FReal Limit, const EJointMotionType AxisMotion);
diff --git a/Engine/Source/Runtime/PhysicsCore/Private/ChaosEngineInterface.cpp b/Engine/Source/Runtime/PhysicsCore/Private/ChaosEngineInterface.cpp
--- a/Engine/Source/Runtime/PhysicsCore/Private/ChaosEngineInterface.cpp
+++ b/Engine/Source/Runtime/PhysicsCore/Private/ChaosEngineInterface.cpp
@@ -1289,6 +1289,17 @@ void FChaosEngineInterface::SetProjectionEnabled_AssumesLocked(const FPhysicsCon
 	}
 }
 
+void FChaosEngineInterface::SetShockPropagationAlpha_AssumesLocked(const FPhysicsConstraintHandle& InConstraintRef, float InShockPropagationAlpha)
+{
+	if (InConstraintRef.IsValid() && InConstraintRef.Constraint->IsType(Chaos::EConstraintType::JointConstraintType))
+	{
+		if (Chaos::FJointConstraint* Constraint = static_cast<Chaos::FJointConstraint*>(InConstraintRef.Constraint))
+		{
+			Constraint->SetShockPropagationAlpha(InShockPropagationAlpha);
+		}
+	}
+}
+
 void FChaosEngineInterface::SetParentDominates_AssumesLocked(const FPhysicsConstraintHandle& InConstraintRef,bool bInParentDominates)
 {
 	if (InConstraintRef.IsValid() && InConstraintRef.Constraint->IsType(Chaos::EConstraintType::JointConstraintType))
diff --git a/Engine/Source/Runtime/PhysicsCore/Public/Chaos/ChaosEngineInterface.h b/Engine/Source/Runtime/PhysicsCore/Public/Chaos/ChaosEngineInterface.h
--- a/Engine/Source/Runtime/PhysicsCore/Public/Chaos/ChaosEngineInterface.h
+++ b/Engine/Source/Runtime/PhysicsCore/Public/Chaos/ChaosEngineInterface.h
@@ -497,7 +497,8 @@ public:
 
 	static void SetCanVisualize(const FPhysicsConstraintHandle& InConstraintRef,bool bInCanVisualize);
 	static void SetCollisionEnabled(const FPhysicsConstraintHandle& InConstraintRef,bool bInCollisionEnabled);
-	static void SetProjectionEnabled_AssumesLocked(const FPhysicsConstraintHandle& InConstraintRef,bool bInProjectionEnabled,float InLinearAlpah = 1.0f,float InAngularAlpha = 0.0f);
+	static void SetProjectionEnabled_AssumesLocked(const FPhysicsConstraintHandle& InConstraintRef,bool bInProjectionEnabled,float InLinearAlpha = 1.0f,float InAngularAlpha = 0.0f);
+	static void SetShockPropagationAlpha_AssumesLocked(const FPhysicsConstraintHandle& InConstraintRef, float InShockPropagationAlpha);
 	static void SetParentDominates_AssumesLocked(const FPhysicsConstraintHandle& InConstraintRef,bool bInParentDominates);
 	static void SetBreakForces_AssumesLocked(const FPhysicsConstraintHandle& InConstraintRef,float InLinearBreakForce,float InAngularBreakForce);
 	static void SetPlasticityLimits_AssumesLocked(const FPhysicsConstraintHandle& InConstraintRef, float InLinearPlasticityLimit, float InAngularPlasticityLimit, EConstraintPlasticityType InLinearPlasticityType);
