Force integrating previously ignored change
Source CL 18912815
Desc:

Fix Lira crash : node edges were not removed correctly

#rnx

diff --git a/Engine/Source/Runtime/Experimental/Chaos/Private/Chaos/Island/IslandGraph.cpp b/Engine/Source/Runtime/Experimental/Chaos/Private/Chaos/Island/IslandGraph.cpp
--- a/Engine/Source/Runtime/Experimental/Chaos/Private/Chaos/Island/IslandGraph.cpp
+++ b/Engine/Source/Runtime/Experimental/Chaos/Private/Chaos/Island/IslandGraph.cpp
@@ -138,7 +138,7 @@ void FIslandGraph<NodeType, EdgeType, IslandType>::RemoveNode(const NodeType& No
 			FGraphNode& GraphNode = GraphNodes[NodeIndex];
 			
 			// If only one node and zero edges, we invalidate the node island
-			if(GraphNode.bValidNode  && GraphIslands.IsValidIndex(GraphNode.IslandIndex))
+			if(GraphNode.NodeEdges.Num() == 0  && GraphIslands.IsValidIndex(GraphNode.IslandIndex))
 			{
 				if(!GraphNode.bStationaryNode)
 				{
@@ -742,7 +742,7 @@ void FIslandGraph<NodeType, EdgeType, IslandType>::UpdateGraph()
 			// Update of the sleeping flag on the island
 			if(GraphIslands.IsValidIndex(GraphNode.IslandIndex) && !GraphNode.bStationaryNode)
 			{
-				if(GraphNode.bValidNode)
+				if(GraphNode.NodeEdges.Num() == 0)
 				{
 					GraphIslands[GraphNode.IslandIndex].bIsSleeping = false;
 				}
diff --git a/Engine/Source/Runtime/Experimental/Chaos/Private/Chaos/Island/IslandManager.cpp b/Engine/Source/Runtime/Experimental/Chaos/Private/Chaos/Island/IslandManager.cpp
--- a/Engine/Source/Runtime/Experimental/Chaos/Private/Chaos/Island/IslandManager.cpp
+++ b/Engine/Source/Runtime/Experimental/Chaos/Private/Chaos/Island/IslandManager.cpp
@@ -500,7 +500,7 @@ void FPBDIslandManager::RemoveParticle(FGeometryParticleHandle* ParticleHandle)
 			{
 				const FGraphNode& GraphNode = IslandGraph->GraphNodes[*NodeIndex];
 				// If the node is valid : we need to also remove it from its own island 
-				if(GraphNode.bValidNode && IslandSolvers.IsValidIndex(GraphNode.IslandIndex))
+				if(GraphNode.NodeEdges.Num() == 0 && IslandSolvers.IsValidIndex(GraphNode.IslandIndex))
 				{
 					IslandSolvers[GraphNode.IslandIndex]->RemoveParticle(ParticleHandle);
 				}
