Chaos : fixing geometry collection component not receiving notification if all particles were disabled or sleeping preventing removal of some pieces

diff --git a/Engine/Source/Runtime/Experimental/Chaos/Public/Chaos/PBDRigidsSOAs.h b/Engine/Source/Runtime/Experimental/Chaos/Public/Chaos/PBDRigidsSOAs.h
--- a/Engine/Source/Runtime/Experimental/Chaos/Public/Chaos/PBDRigidsSOAs.h
+++ b/Engine/Source/Runtime/Experimental/Chaos/Public/Chaos/PBDRigidsSOAs.h
@@ -326,6 +326,7 @@ public:
 			{
 				if (Particle->GetParticleType() == Chaos::EParticleType::GeometryCollection)
 				{
+					FellAsleepGeometryCollectionArray.Remove(PBDRigid);
 					bGeometryCollectionDirty = true;
 				}
 				else // clustered
@@ -370,6 +371,7 @@ public:
 			{
 				if (Particle->GetParticleType() == Chaos::EParticleType::GeometryCollection)
 				{
+					FellAsleepGeometryCollectionArray.Remove(PBDRigid);
 					bGeometryCollectionDirty = true;
 					return;
 				}
@@ -465,6 +467,7 @@ public:
 					{
 						if (Particle->GetParticleType() == Chaos::EParticleType::GeometryCollection)
 						{
+							FellAsleepGeometryCollectionArray.Remove(PBDRigid);
 							bGeometryCollectionDirty = true;
 							return true;
 						}
@@ -529,6 +532,7 @@ public:
 					{
 						if (Particle->GetParticleType() == Chaos::EParticleType::GeometryCollection)
 						{
+							FellAsleepGeometryCollectionArray.Insert(PBDRigid);
 							bGeometryCollectionDirty = true;
 						}
 					}
@@ -1063,6 +1067,7 @@ private:
 				{&StaticGeometryCollectionArray},
 				{&KinematicGeometryCollectionArray},
 				{&DynamicGeometryCollectionArray},
+				{&FellAsleepGeometryCollectionArray.GetArray()},
 				{&TransientDirtyMapArray.GetArray()}
 			};
 			DirtyParticlesView = MakeParticleView(MoveTemp(TmpArray));
@@ -1204,6 +1209,7 @@ private:
 	TArray<TPBDGeometryCollectionParticleHandle<FReal, 3>*> KinematicGeometryCollectionArray;
 	TArray<TPBDGeometryCollectionParticleHandle<FReal, 3>*> SleepingGeometryCollectionArray;
 	TArray<TPBDGeometryCollectionParticleHandle<FReal, 3>*> DynamicGeometryCollectionArray;
+	TParticleMapArray<FPBDRigidParticleHandle> FellAsleepGeometryCollectionArray;
 	bool bGeometryCollectionDirty;
 
 	//Utility structures for maintaining an Active particles view
