Chaos
- Capsule-Capsule manifold fixes
- Sphere-Capsule manifold fixes

diff --git a/Engine/Source/Runtime/Experimental/Chaos/Private/Chaos/CollisionOneShotManifolds.cpp b/Engine/Source/Runtime/Experimental/Chaos/Private/Chaos/CollisionOneShotManifolds.cpp
--- a/Engine/Source/Runtime/Experimental/Chaos/Private/Chaos/CollisionOneShotManifolds.cpp
+++ b/Engine/Source/Runtime/Experimental/Chaos/Private/Chaos/CollisionOneShotManifolds.cpp
@@ -24,6 +24,18 @@ namespace Chaos
 	FRealSingle Chaos_Collision_Manifold_SphereCapsuleSizeThreshold = 0.5f;
 	FAutoConsoleVariableRef CVarChaos_Manifold_SphereCapsuleSizeThreshold(TEXT("p.Chaos.Collision.Manifold.SphereCapsuleSizeThreshold"), Chaos_Collision_Manifold_SphereCapsuleSizeThreshold, TEXT(""));
 
+	// When two capsule axes are aligned more than this, generate a multi-point manifold
+	FRealSingle Chaos_Collision_Manifold_CapsuleAxisAlignedThreshold = 0.8f;	// About 30deg
+	FAutoConsoleVariableRef CVarChaos_Manifold_CapsuleAxisAlignedThreshold(TEXT("p.Chaos.Collision.Manifold.CapsuleAxisAlignedThreshold"), Chaos_Collision_Manifold_CapsuleAxisAlignedThreshold, TEXT(""));
+
+	// When two capsules penetrate by more than this fraction of the radius, generate a multi-point manifold
+	FRealSingle Chaos_Collision_Manifold_CapsuleDeepPenetrationFraction = 0.05f;
+	FAutoConsoleVariableRef CVarChaos_Manifold_CapsuleDeepPenetrationFraction(TEXT("p.Chaos.Collision.Manifold.CapsuleDeepPenetrationFraction"), Chaos_Collision_Manifold_CapsuleDeepPenetrationFraction, TEXT(""));
+
+	// When two capsule lie on top of each other in an X, the extra manifold points are this fraction of the radius from the primary contact point
+	FRealSingle Chaos_Collision_Manifold_CapsuleRadialContactFraction = 0.25f;
+	FAutoConsoleVariableRef CVarChaos_Manifold_CapsuleRadialContactFraction(TEXT("p.Chaos.Collision.Manifold.CapsuleRadialContactFraction"), Chaos_Collision_Manifold_CapsuleRadialContactFraction, TEXT(""));
+
 	FRealSingle Chaos_Collision_Manifold_PlaneContactNormalEpsilon = 0.001f;
 	FAutoConsoleVariableRef CVarChaos_Manifold_PlaneContactNormalEpsilon(TEXT("p.Chaos.Collision.Manifold.PlaneContactNormalEpsilon"), Chaos_Collision_Manifold_PlaneContactNormalEpsilon, TEXT("Normal tolerance used to distinguish face contacts from edge-edge contacts"));
 
diff --git a/Engine/Source/Runtime/Experimental/Chaos/Private/Chaos/CollisionOneShotManifoldsMiscShapes.cpp b/Engine/Source/Runtime/Experimental/Chaos/Private/Chaos/CollisionOneShotManifoldsMiscShapes.cpp
--- a/Engine/Source/Runtime/Experimental/Chaos/Private/Chaos/CollisionOneShotManifoldsMiscShapes.cpp
+++ b/Engine/Source/Runtime/Experimental/Chaos/Private/Chaos/CollisionOneShotManifoldsMiscShapes.cpp
@@ -26,6 +26,9 @@ namespace Chaos
 {
 	extern FRealSingle Chaos_Collision_EdgePrunePlaneDistance;
 	extern FRealSingle Chaos_Collision_Manifold_SphereCapsuleSizeThreshold;
+	extern FRealSingle Chaos_Collision_Manifold_CapsuleAxisAlignedThreshold;
+	extern FRealSingle Chaos_Collision_Manifold_CapsuleDeepPenetrationFraction;
+	extern FRealSingle Chaos_Collision_Manifold_CapsuleRadialContactFraction;
 
 	namespace Collisions
 	{
@@ -153,19 +156,28 @@ namespace Chaos
 						return;
 					}
 
+					// If the contact is deep, there's a high chance that pushing one end out will push the other deeper and we also need more contacts.
+					// Note: we only consider the radius of the dynamic object(s) when deciding what "deep" means because the extra contacts are only
+					// to prevent excessive rotation from the single contact we have so far, and only the dynamic objects will rotate.
+					const FReal DeepRadiusFraction = Chaos_Collision_Manifold_CapsuleDeepPenetrationFraction;
+					const bool bIsDeep = NearPhi < -DeepRadiusFraction * Capsule.GetRadius();
+					if (!bIsDeep)
+					{
+						return;
+					}
+
 					// Now add the two end caps
 					// Calculate the vector orthogonal to the capsule axis that gives the nearest points on the capsule cyclinder to the sphere
 					// The initial length will be proportional to the sine of the angle between the axis and the delta position and will approach
 					// zero when the capsule is end-on to the sphere, in which case we won't add the end caps.
 					constexpr FReal EndCapSinAngleThreshold = FReal(0.35);	// about 20deg
 					constexpr FReal EndCapDistanceThreshold = FReal(0.2);	// fraction
-					FVec3 CapsuleOrthogonal = FVec3::CrossProduct(Capsule.GetAxis(), FVec3::CrossProduct(Capsule.GetAxis(), SpherePos - Capsule.GetCenter()));
+					FVec3 CapsuleOrthogonal = FVec3::CrossProduct(Capsule.GetAxis(), FVec3::CrossProduct(Capsule.GetAxis(), NearPosDir));
 					const FReal CapsuleOrthogonalLenSq = CapsuleOrthogonal.SizeSquared();
 					if (CapsuleOrthogonalLenSq > FMath::Square(EndCapSinAngleThreshold))
 					{
 						// Orthogonal must point towards the sphere, but currently depends on the relative axis orientation
-						const FReal CapsuleOrthogonalLen = FMath::Sqrt(CapsuleOrthogonalLenSq);
-						CapsuleOrthogonal = CapsuleOrthogonal / CapsuleOrthogonalLen;
+						CapsuleOrthogonal = CapsuleOrthogonal * FMath::InvSqrt(CapsuleOrthogonalLenSq);
 						if (FVec3::DotProduct(CapsuleOrthogonal, SpherePos - Capsule.GetCenter()) < FReal(0))
 						{
 							CapsuleOrthogonal = -CapsuleOrthogonal;
@@ -268,10 +280,9 @@ namespace Chaos
 			}
 		}
 
+		// @todo(chaos): this will be faster if we transform into the space of one of the capsules
 		void ConstructCapsuleCapsuleOneShotManifold(const FCapsule& CapsuleA, const FRigidTransform3& CapsuleATransform, const FCapsule& CapsuleB, const FRigidTransform3& CapsuleBTransform, const FReal Dt, FPBDCollisionConstraint& Constraint)
 		{
-			const FReal AxisDotMinimum = 0.707f; // If the axes are off by more than this, just a single manifold point will be generated
-
 			// We only build one shot manifolds once
 			// All convexes are pre-scaled, or wrapped in TImplicitObjectScaled
 			ensure(CapsuleATransform.GetScale3D() == FVec3(1.0f, 1.0f, 1.0f));
@@ -280,97 +291,233 @@ namespace Chaos
 			// @todo(chaos): support manifold maintenance
 			Constraint.ResetActiveManifoldContacts();
 
-			FVec3 CapsuleADirection(CapsuleATransform.TransformVector(CapsuleA.GetSegment().GetAxis()));
-			const FVec3 CapsuleBDirection(CapsuleBTransform.TransformVector(CapsuleB.GetSegment().GetAxis()));
-
-			FReal ADotB = FVec3::DotProduct(CapsuleADirection, CapsuleBDirection);
+			FVec3 AAxis(CapsuleATransform.TransformVector(CapsuleA.GetSegment().GetAxis()));
+			const FVec3 BAxis(CapsuleBTransform.TransformVector(CapsuleB.GetSegment().GetAxis()));
 
 			const FReal AHalfLen = CapsuleA.GetHeight() / 2.0f;
 			const FReal BHalfLen = CapsuleB.GetHeight() / 2.0f;
 
-			if (FMath::Abs(ADotB) < AxisDotMinimum || AHalfLen < KINDA_SMALL_NUMBER || BHalfLen < KINDA_SMALL_NUMBER)
+			// Used in a few places below where we need to use the smaller/larger capsule, but always a dynamic one
+			const FReal ADynamicRadius = FConstGenericParticleHandle(Constraint.GetParticle0())->IsDynamic() ? CapsuleA.GetRadius() : TNumericLimits<FReal>::Max();
+			const FReal BDynamicRadius = FConstGenericParticleHandle(Constraint.GetParticle1())->IsDynamic() ? CapsuleB.GetRadius() : TNumericLimits<FReal>::Max();
+
+			// Make both capsules point in the same general direction
+			FReal ADotB = FVec3::DotProduct(AAxis, BAxis);
+			if (ADotB < 0)
 			{
-				FContactPoint ContactPoint = CapsuleCapsuleContactPoint(CapsuleA, CapsuleATransform, CapsuleB, CapsuleBTransform, FReal(0));
-				if (ContactPoint.Phi < Constraint.GetCullDistance())
-				{
-					Constraint.AddOneshotManifoldContact(ContactPoint);
-				}
-				return;
+				ADotB = -ADotB;
+				AAxis = -AAxis;
 			}
-			
-			FVector P1, P2;
+
+			// Get the closest points on the two line segments. This is used to generate the closest contact point
+			// which is always added to the manifold (if within CullDistance). We may also add other points.
+			FVector AClosest, BClosest;
 			const FVector ACenter = CapsuleATransform.TransformPosition(CapsuleA.GetCenter());
 			const FVector BCenter = CapsuleBTransform.TransformPosition(CapsuleB.GetCenter());
 			FMath::SegmentDistToSegmentSafe(
-				ACenter + AHalfLen * CapsuleADirection, 
-				ACenter - AHalfLen * CapsuleADirection, 
-				BCenter + BHalfLen * CapsuleBDirection, 
-				BCenter - BHalfLen * CapsuleBDirection, 
-				P1, 
-				P2);
-
-			FVec3 Delta = P2 - P1;
-			FReal DeltaLen = Delta.Size();
+				ACenter + AHalfLen * AAxis,
+				ACenter - AHalfLen * AAxis,
+				BCenter + BHalfLen * BAxis,
+				BCenter - BHalfLen * BAxis,
+				AClosest,
+				BClosest);
+
+			FVec3 ClosestDelta = BClosest - AClosest;
+			FReal ClosestDeltaLen = ClosestDelta.Size();
+
+			// Stop now if we are beyond the cull distance
+			const FReal ClosestPhi = ClosestDeltaLen - (CapsuleA.GetRadius() + CapsuleB.GetRadius());
+			if (ClosestPhi > Constraint.GetCullDistance())
+			{
+				return;
+			}
 
-			if (DeltaLen < KINDA_SMALL_NUMBER)
+			// Calculate the normal from the two closest points. Handle exact axis overlaps.
+			FVec3 ClosestNormal;
+			if (ClosestDeltaLen > KINDA_SMALL_NUMBER)
+			{
+				ClosestNormal = -ClosestDelta / ClosestDeltaLen;
+			}
+			else
+			{
+				// Center axes exactly intersect. We'll fake a result that pops the capsules out along the Z axis, with the smaller capsule going up
+				ClosestNormal = (ADynamicRadius <= BDynamicRadius) ? FVec3(0, 0, 1) : FVec3(0, 0, -1);
+			}
+			const FVec3 ClosestLocationA = AClosest - ClosestNormal * CapsuleA.GetRadius();
+			const FVec3 ClosestLocationB = BClosest + ClosestNormal * CapsuleB.GetRadius();
+
+			// We always add the closest point to the manifold
+			// We may also add 2 more points generated from the end cap positions of the smaller capsule
+			FContactPoint ClosestContactPoint;
+			ClosestContactPoint.ShapeContactPoints[0] = CapsuleATransform.InverseTransformPosition(ClosestLocationA);
+			ClosestContactPoint.ShapeContactPoints[1] = CapsuleBTransform.InverseTransformPosition(ClosestLocationB);
+			ClosestContactPoint.ShapeContactNormal = CapsuleBTransform.InverseTransformVector(ClosestNormal);
+			ClosestContactPoint.Phi = ClosestPhi;
+			ClosestContactPoint.FaceIndex = INDEX_NONE;
+			ClosestContactPoint.ContactType = EContactPointType::VertexPlane;
+			Constraint.AddOneshotManifoldContact(ClosestContactPoint);
+
+			// We don't generate manifold points within this fraction (of segment length) distance
+			constexpr FReal TDeltaThreshold = FReal(0.2);		// fraction
+
+			// If the nearest cylinder normal is parallel to the other axis within this tolerance, we stick with 1 manifold point
+			constexpr FReal SinAngleThreshold = FReal(0.35);	// about 20deg (this would be an endcap-versus-cylinderwall collision at >70 degs)
+
+			// If the capsules are in an X configuration, this controls the distance of the manifold points from the closest point
+			const FReal RadialContactFraction = Chaos_Collision_Manifold_CapsuleRadialContactFraction;
+
+			// Calculate the line segment times for the nearest point calculate above
+			// NOTE: TA and TB will be in [-1, 1]
+			const FReal TA = FVec3::DotProduct(AClosest - ACenter, AAxis) / AHalfLen;
+			const FReal TB = FVec3::DotProduct(BClosest - BCenter, BAxis) / BHalfLen;
+
+			// If we have an end-end contact with no segment overlap, stick with the single point manifold
+			// This is when we have two capsules laid end to end (as opposed to side-by-side)
+			// NOTE: This test only works because we made the axes point in the same direction above
+			if ((TA < FReal(-1) + TDeltaThreshold) && (TB > FReal(1) - TDeltaThreshold))
 			{
-				FContactPoint ContactPoint = CapsuleCapsuleContactPoint(CapsuleA, CapsuleATransform, CapsuleB, CapsuleBTransform, FReal(0));
-				if (ContactPoint.Phi < Constraint.GetCullDistance())
-				{
-					Constraint.AddOneshotManifoldContact(ContactPoint);
-				}
 				return;
 			}
-			
-			// Make both capsules point in the same general direction
-			if (ADotB < 0)
+			if ((TB < FReal(-1) + TDeltaThreshold) && (TA > FReal(1) - TDeltaThreshold))
 			{
-				ADotB = -ADotB;
-				CapsuleADirection = -CapsuleADirection;
+				return;
 			}
 
-			// Now project A points onto B segment
-			const FReal ProjA1OntoB = FVec3::DotProduct(ACenter - BCenter - AHalfLen * CapsuleADirection, CapsuleBDirection);
-			const FReal ProjA2OntoB = FVec3::DotProduct(ACenter - BCenter + AHalfLen * CapsuleADirection, CapsuleBDirection);
-
-			const FReal Clipped1Coord = FMath::Max(ProjA1OntoB, -BHalfLen); // 1D coordinates
-			const FReal Clipped2Coord = FMath::Min(ProjA2OntoB, BHalfLen);
-			if (Clipped1Coord > Clipped2Coord) // No overlap
+			// If the axes are closely aligned, we definitely want more contact points (e.g., capsule lying on top of another).
+			// Also if the contact is deep, there's a high chance that pushing one end out will push the other deeper and we also need more contacts.
+			// Note: we only consider the radius of the dynamic object(s) when deciding what "deep" means because the extra contacts are only
+			// to prevent excessive rotation from the single contact we have so far, and only the dynamic objects will rotate.
+			const FReal AxisDotMinimum = Chaos_Collision_Manifold_CapsuleAxisAlignedThreshold;
+			const FReal DeepRadiusFraction = Chaos_Collision_Manifold_CapsuleDeepPenetrationFraction;
+			const FReal MinDynamicRadius = FMath::Min(ADynamicRadius, BDynamicRadius);
+			const bool bAreAligned = ADotB > AxisDotMinimum;
+			const bool bIsDeep = ClosestPhi < -DeepRadiusFraction * MinDynamicRadius;
+			if (!bAreAligned && !bIsDeep)
 			{
-				FContactPoint ContactPoint = CapsuleCapsuleContactPoint(CapsuleA, CapsuleATransform, CapsuleB, CapsuleBTransform, FReal(0));
-				if (ContactPoint.Phi < Constraint.GetCullDistance())
-				{
-					Constraint.AddOneshotManifoldContact(ContactPoint);
-				}
 				return;
 			}
 
-			//FReal NewPhi = DeltaLen - (CapsuleA.GetRadius() + CapsuleB.GetRadius());
-			FVec3 Dir = Delta / DeltaLen;
-			FVec3 Normal = -Dir;
+			// Lambda: Create a contact point between a point on the cylinder of FirstCapsule at FirstT, with the nearest point on SecondCapsule
+			auto MakeCapsuleSegmentContact = [](
+				const FReal FirstT,
+				const FVec3& FirstCenter,
+				const FVec3& FirstAxis,
+				const FReal FirstHalfLen,
+				const FReal FirstRadius,
+				const FRigidTransform3& FirstTransform,
+				const FVec3& SecondCenter,
+				const FVec3& SecondAxis,
+				const FReal SecondHalfLen,
+				const FReal SecondRadius,
+				const FRigidTransform3& SecondTransform,
+				const FVec3& Orthogonal,
+				const FReal CullDistance,
+				const bool bSwap) -> FContactPoint
+			{
+				FContactPoint ContactPoint;
 
-			FContactPoint ContactPoint;
-			ContactPoint.ShapeContactNormal = CapsuleBTransform.InverseTransformVector(Normal);
+				const FVec3 FirstContactPos = FirstCenter + (FirstT * FirstHalfLen) * FirstAxis + Orthogonal * FirstRadius;
+				const FVec3 SecondSegmentPos = FMath::ClosestPointOnLine(SecondCenter - SecondHalfLen * SecondAxis, SecondCenter + SecondHalfLen * SecondAxis, FirstContactPos);
+				const FReal SecondSegmentDist = (FirstContactPos - SecondSegmentPos).Size();
+				const FVec3 SecondSegmentDir = (FirstContactPos - SecondSegmentPos) / SecondSegmentDist;
+				const FVec3 SecondContactPos = SecondSegmentPos + SecondRadius * SecondSegmentDir;
+				const FReal ContactPhi = SecondSegmentDist - SecondRadius;
 
-			auto AddManifoldPoint = [&](FReal ClippedCoord)
-			{
-				FVec3 LocationB = ClippedCoord * CapsuleBDirection + BCenter + Normal * (CapsuleB.GetRadius());
-				const FReal ProjCentreAOntoB = FVec3::DotProduct(ACenter - BCenter, CapsuleBDirection);
-				// Note location A is calculated by rotation (effectively) instead of the usual plane clipping
-				FVec3 LocationA = (ClippedCoord - ProjCentreAOntoB) * CapsuleADirection + ACenter - Normal * (CapsuleA.GetRadius());
+				if (ContactPhi < CullDistance)
+				{
+					if (!bSwap)
+					{
+						ContactPoint.ShapeContactPoints[0] = FirstTransform.InverseTransformPositionNoScale(FirstContactPos);
+						ContactPoint.ShapeContactPoints[1] = SecondTransform.InverseTransformPositionNoScale(SecondContactPos);
+						ContactPoint.ShapeContactNormal = SecondTransform.InverseTransformVectorNoScale(SecondSegmentDir);
+					}
+					else
+					{
+						ContactPoint.ShapeContactPoints[0] = SecondTransform.InverseTransformPositionNoScale(SecondContactPos);
+						ContactPoint.ShapeContactPoints[1] = FirstTransform.InverseTransformPositionNoScale(FirstContactPos);
+						ContactPoint.ShapeContactNormal = -FirstTransform.InverseTransformVectorNoScale(SecondSegmentDir);
+					}
+					ContactPoint.Phi = ContactPhi;
+					ContactPoint.FaceIndex = INDEX_NONE;
+					ContactPoint.ContactType = EContactPointType::VertexPlane;
+				}
 
-				const FReal Phi = FVec3::DotProduct(LocationA - LocationB, Normal);
-				if (Phi < Constraint.GetCullDistance())
+				return ContactPoint;
+			};
+
+			// Lambda: Add up to 2 more contacts from the cylindrical surface on FirstCylinder, if they are not too close to the existing contact.
+			// The point locations depend on cylinder alignment.
+			auto MakeCapsuleEndPointContacts = [&Constraint, TDeltaThreshold, SinAngleThreshold, RadialContactFraction, &MakeCapsuleSegmentContact](
+				const FReal FirstT,
+				const FVec3& FirstCenter,
+				const FVec3& FirstAxis,
+				const FReal FirstHalfLen,
+				const FReal FirstRadius,
+				const FRigidTransform3& FirstTransform,
+				const FVec3& SecondCenter,
+				const FVec3& SecondAxis,
+				const FReal SecondHalfLen,
+				const FReal SecondRadius,
+				const FRigidTransform3& SecondTransform,
+				const FVec3& ClosestDir,
+				const FReal FirstAxisDotSecondAxis,
+				const bool bSwap) -> void
+			{
+				// Orthogonal: the vector from a point on FirstCapsule's axis to its cylinder surface, in the direction of SecondCapsule
+				FVec3 Orthogonal = FVec3::CrossProduct(FirstAxis, FVec3::CrossProduct(FirstAxis, ClosestDir));
+				const FReal OrthogonalLenSq = Orthogonal.SizeSquared();
+				if (OrthogonalLenSq > FMath::Square(SinAngleThreshold))
 				{
-					ContactPoint.ShapeContactPoints[0] = CapsuleATransform.InverseTransformPosition(LocationA);
-					ContactPoint.ShapeContactPoints[1] = CapsuleBTransform.InverseTransformPosition(LocationB);
-					ContactPoint.Phi = Phi;
-					Constraint.AddOneshotManifoldContact(ContactPoint);
+					Orthogonal = Orthogonal * FMath::InvSqrt(OrthogonalLenSq);
+					if (FVec3::DotProduct(Orthogonal, SecondCenter - FirstCenter) < FReal(0))
+					{
+						Orthogonal = -Orthogonal;
+					}
+
+					// Clip the FirstCapsule's end points to be within the line segment of SecondCapsule
+					// This is to restrict the extra contacts to the overlapping line segment (e.g, when capsules are lying partly on top of each other)
+					const FReal ProjectedLen = FReal(2) * FirstHalfLen * FirstAxisDotSecondAxis;
+					const FReal ClippedTMin = FVec3::DotProduct((SecondCenter - SecondHalfLen * SecondAxis) - (FirstCenter + FirstHalfLen * FirstAxis), SecondAxis) / ProjectedLen;
+					const FReal ClippedTMax = FVec3::DotProduct((SecondCenter + SecondHalfLen * SecondAxis) - (FirstCenter - FirstHalfLen * FirstAxis), SecondAxis) / ProjectedLen;
+
+					// Clip the FirstCapsules end points to be within some laterial distance od the SecondCapsule's axis
+					// This restricts the contacts to be at a useful location when line segments are perpendicular to each other 
+					// (e.g., when the capsules are on top of each other but in a cross)
+					// As we get more perpendicular, move limits closer to radius fraction
+					const FReal MaxDeltaTRadial = RadialContactFraction * (SecondRadius / FirstHalfLen);
+					const FReal RadialClippedTMax = FMath::Lerp(FReal(MaxDeltaTRadial), FReal(1), FirstAxisDotSecondAxis);
+
+					const FReal TMin = FMath::Max3(FReal(-1), ClippedTMin, -RadialClippedTMax);
+					const FReal TMax = FMath::Min3(FReal(1), ClippedTMax, RadialClippedTMax);
+
+					if (TMin < FirstT - TDeltaThreshold)
+					{
+						FContactPoint EndContact0 = MakeCapsuleSegmentContact(TMin, FirstCenter, FirstAxis, FirstHalfLen, FirstRadius, FirstTransform, SecondCenter, SecondAxis, SecondHalfLen, SecondRadius, SecondTransform, Orthogonal, Constraint.GetCullDistance(), bSwap);
+						if (EndContact0.Phi < Constraint.GetCullDistance())
+						{
+							Constraint.AddOneshotManifoldContact(EndContact0);
+						}
+					}
+					if (TMax > FirstT + TDeltaThreshold)
+					{
+						FContactPoint EndContact1 = MakeCapsuleSegmentContact(TMax, FirstCenter, FirstAxis, FirstHalfLen, FirstRadius, FirstTransform, SecondCenter, SecondAxis, SecondHalfLen, SecondRadius, SecondTransform, Orthogonal, Constraint.GetCullDistance(), bSwap);
+						if (EndContact1.Phi < Constraint.GetCullDistance())
+						{
+							Constraint.AddOneshotManifoldContact(EndContact1);
+						}
+					}
 				}
 			};
 
-			AddManifoldPoint(Clipped1Coord);
-			AddManifoldPoint(Clipped2Coord);
+			// Generate the extra manifold points
+			if (ADynamicRadius <= BDynamicRadius)
+			{
+				MakeCapsuleEndPointContacts(TA, ACenter, AAxis, AHalfLen, CapsuleA.GetRadius(), CapsuleATransform, BCenter, BAxis, BHalfLen, CapsuleB.GetRadius(), CapsuleBTransform, ClosestNormal, ADotB, false);
+			}
+			else
+			{
+				MakeCapsuleEndPointContacts(TB, BCenter, BAxis, BHalfLen, CapsuleB.GetRadius(), CapsuleBTransform, ACenter, AAxis, AHalfLen, CapsuleA.GetRadius(), CapsuleATransform, ClosestNormal, ADotB, true);
+			}
 		}
 
 		template <typename TriMeshType>
diff --git a/Engine/Source/Runtime/Experimental/Chaos/Public/Chaos/Collision/ContactPoint.h b/Engine/Source/Runtime/Experimental/Chaos/Public/Chaos/Collision/ContactPoint.h
--- a/Engine/Source/Runtime/Experimental/Chaos/Public/Chaos/Collision/ContactPoint.h
+++ b/Engine/Source/Runtime/Experimental/Chaos/Public/Chaos/Collision/ContactPoint.h
@@ -54,6 +54,10 @@ namespace Chaos
 		bool IsSet() const { return (Phi != TNumericLimits<FReal>::Max()); }
 
 		// Switch the shape indices. For use when calling a collision detection method which takes shape types in the opposite order to what you want.
+		// @todo(chaos): remove or fix this function
+		// WARNING: this function can no longer be used in isolation as it could when we were calculating world-space contact data. For this to
+		// work correctly, the normal must either already be in the space of the first shape, or will need to be transformed after.
+		// Alternatively we could start using EContactPointType to indicate normal ownership
 		FContactPoint& SwapShapes()
 		{
 			if (IsSet())
