Middle phase optimization for the sweep tests.

diff --git a/Engine/Source/Runtime/Experimental/Chaos/Public/Chaos/AABBTree.h b/Engine/Source/Runtime/Experimental/Chaos/Public/Chaos/AABBTree.h
--- a/Engine/Source/Runtime/Experimental/Chaos/Public/Chaos/AABBTree.h
+++ b/Engine/Source/Runtime/Experimental/Chaos/Public/Chaos/AABBTree.h
@@ -2,6 +2,7 @@
 #pragma once
 
 #include "Chaos/AABB.h"
+#include "Chaos/AABBVectorized.h"
 #include "Chaos/AABBTreeDirtyGridUtils.h"
 #include "Chaos/Defines.h"
 #include "Chaos/GeometryParticles.h"
@@ -12,6 +13,7 @@
 #include "Chaos/BoundingVolume.h"
 #include "ProfilingDebugging/CsvProfiler.h"
 #include "ChaosStats.h"
+#include "Math/VectorRegister.h"
 
 CSV_DECLARE_CATEGORY_EXTERN(ChaosPhysicsTimers);
 
@@ -150,6 +152,12 @@ struct TAABBTreeIntersectionHelper
 		check(false);
 		return true;
 	}
+	static bool IntersectsVectorized(const VectorRegister4Float& Start, FQueryFastData& QueryFastData, VectorRegister4Float& TOI, VectorRegister4Float& OutPosition,
+		const FAABBVectorized& Bounds, const FAABBVectorized& QueryBounds, const VectorRegister4Float& QueryHalfExtents, const VectorRegister4Float& Dir, const VectorRegister4Float InvDir, const bool bParallel[3])
+	{
+		check(false);
+		return true;
+	}
 };
 
 template<>
@@ -160,6 +168,13 @@ struct TAABBTreeIntersectionHelper<FQueryFastData, EAABBQueryType::Raycast>
 	{
 		return Bounds.RaycastFast(Start, Dir, InvDir, bParallel, QueryFastData.CurrentLength, QueryFastData.InvCurrentLength, TOI, OutPosition);
 	}
+
+	FORCEINLINE_DEBUGGABLE static bool IntersectsVectorized(const VectorRegister4Float& Start, FQueryFastData& QueryFastData, VectorRegister4Float& TOI, VectorRegister4Float& OutPosition,
+		const FAABBVectorized& Bounds, const FAABBVectorized& QueryBounds, const VectorRegister4Float& QueryHalfExtents, const VectorRegister4Float& Dir, const VectorRegister4Float InvDir, const bool bParallel[3])
+	{
+		check(false);
+		return true;
+	}
 };
 
 template <>
@@ -171,6 +186,16 @@ struct TAABBTreeIntersectionHelper<FQueryFastData, EAABBQueryType::Sweep>
 		FAABB3 SweepBounds(Bounds.Min() - QueryHalfExtents, Bounds.Max() + QueryHalfExtents);
 		return SweepBounds.RaycastFast(Start, Dir, InvDir, bParallel, QueryFastData.CurrentLength, QueryFastData.InvCurrentLength, TOI, OutPosition);
 	}
+
+	FORCEINLINE_DEBUGGABLE static bool IntersectsVectorized(const VectorRegister4Float& Start, FQueryFastData& QueryFastData, VectorRegister4Float& TOI, VectorRegister4Float& OutPosition,
+		const FAABBVectorized& Bounds, const FAABBVectorized& QueryBounds, const VectorRegister4Float& QueryHalfExtents, const VectorRegister4Float& Dir, const VectorRegister4Float InvDir, const bool bParallel[3])
+	{
+		VectorRegister4Float CurrentLength = MakeVectorRegisterFloatFromDouble(VectorLoadDouble1(&QueryFastData.CurrentLength));
+		VectorRegister4Float InvCurrentLength = MakeVectorRegisterFloatFromDouble(VectorLoadDouble1(&QueryFastData.InvCurrentLength));
+		FAABBVectorized SweepBounds(VectorSubtract(Bounds.GetMin(), QueryHalfExtents), VectorAdd(Bounds.GetMax(), QueryHalfExtents));
+		return SweepBounds.RaycastFast(Start, Dir, InvDir, bParallel, CurrentLength, InvCurrentLength, TOI, OutPosition);
+	}
+
 };
 
 template <>
@@ -310,8 +335,11 @@ struct TAABBTreeLeafArray : public TBoundsWrapperHelper<TPayloadType, T, bComput
 	FORCEINLINE_DEBUGGABLE bool RaycastSweepImp(const TVec3<T>& Start, TQueryFastData& QueryFastData, const TVec3<T>& QueryHalfExtents, TSQVisitor& Visitor, const TVec3<T>& Dir, const TVec3<T> InvDir, const bool bParallel[3]) const
 	{
 		PHYSICS_CSV_CUSTOM_VERY_EXPENSIVE(PhysicsCounters, MaxLeafSize, Elems.Num(), ECsvCustomStatOp::Max);
-		FVec3 TmpPosition;
-		FReal TOI;
+
+		const VectorRegister4Float DirSimd = MakeVectorRegisterFloatFromDouble(MakeVectorRegisterDouble(Dir.X, Dir.Y, Dir.Z, 0.0));
+		const VectorRegister4Float InvDirSimd = MakeVectorRegisterFloatFromDouble(MakeVectorRegisterDouble(InvDir.X, InvDir.Y, InvDir.Z, 0.0));
+		const VectorRegister4Float QueryHalfExtentsSimd = MakeVectorRegisterFloatFromDouble(MakeVectorRegisterDouble(QueryHalfExtents.X, QueryHalfExtents.Y, QueryHalfExtents.Z, 0.0));
+		
 		for (const auto& Elem : Elems)
 		{
 			if (PrePreFilterHelper(Elem.Payload, Visitor))
@@ -320,14 +348,39 @@ struct TAABBTreeLeafArray : public TBoundsWrapperHelper<TPayloadType, T, bComput
 			}
 
 			const FAABB3 InstanceBounds(Elem.Bounds.Min(), Elem.Bounds.Max());
-			if (TAABBTreeIntersectionHelper<TQueryFastData, bSweep ? EAABBQueryType::Sweep :
-					EAABBQueryType::Raycast>::Intersects(Start, QueryFastData, TOI, TmpPosition, InstanceBounds, FAABB3(), QueryHalfExtents, Dir, InvDir, bParallel))
+			if (bSweep == true)
 			{
-				TSpatialVisitorData<TPayloadType> VisitData(Elem.Payload, true, InstanceBounds);
-				const bool bContinue = (bSweep && Visitor.VisitSweep(VisitData, QueryFastData)) || (!bSweep &&  Visitor.VisitRaycast(VisitData, QueryFastData));
-				if (!bContinue)
+				const VectorRegister4Float Min = MakeVectorRegisterFloatFromDouble(MakeVectorRegisterDouble(InstanceBounds.Min().X, InstanceBounds.Min().Y, InstanceBounds.Min().Z, 0.0));
+				const VectorRegister4Float Max = MakeVectorRegisterFloatFromDouble(MakeVectorRegisterDouble(InstanceBounds.Max().X, InstanceBounds.Max().Y, InstanceBounds.Max().Z, 0.0));
+				const FAABBVectorized InstanceBoundsSimd(Min, Max);
+				VectorRegister4Float TmpPosition;
+				VectorRegister4Float TOI;
+				VectorRegister4Float StartSimd = MakeVectorRegisterFloatFromDouble(MakeVectorRegisterDouble(Start.X, Start.Y, Start.Z, 0.0));
+
+				if (TAABBTreeIntersectionHelper<TQueryFastData, EAABBQueryType::Sweep>::IntersectsVectorized(
+					StartSimd, QueryFastData, TOI, TmpPosition, InstanceBoundsSimd, FAABBVectorized(), QueryHalfExtentsSimd, DirSimd, InvDirSimd, bParallel))
 				{
-					return false;
+					TSpatialVisitorData<TPayloadType> VisitData(Elem.Payload, true, InstanceBounds);
+					const bool bContinue = (bSweep && Visitor.VisitSweep(VisitData, QueryFastData)) || (!bSweep && Visitor.VisitRaycast(VisitData, QueryFastData));
+					if (!bContinue)
+					{
+						return false;
+					}
+				}
+			}
+			else 
+			{
+				FVec3 TmpPosition;
+				FReal TOI;
+				if (TAABBTreeIntersectionHelper<TQueryFastData,
+					EAABBQueryType::Raycast>::Intersects(Start, QueryFastData, TOI, TmpPosition, InstanceBounds, FAABB3(), QueryHalfExtents, Dir, InvDir, bParallel))
+				{
+					TSpatialVisitorData<TPayloadType> VisitData(Elem.Payload, true, InstanceBounds);
+					const bool bContinue = (bSweep && Visitor.VisitSweep(VisitData, QueryFastData)) || (!bSweep && Visitor.VisitRaycast(VisitData, QueryFastData));
+					if (!bContinue)
+					{
+						return false;
+					}
 				}
 			}
 		}
diff --git a/Engine/Source/Runtime/Experimental/Chaos/Public/Chaos/AABBVectorized.h b/Engine/Source/Runtime/Experimental/Chaos/Public/Chaos/AABBVectorized.h
new file mode 100644
--- /dev/null
+++ b/Engine/Source/Runtime/Experimental/Chaos/Public/Chaos/AABBVectorized.h
@@ -0,0 +1,80 @@
+// Copyright Epic Games, Inc. All Rights Reserved.
+#pragma once
+
+#include "Chaos/Core.h"
+#include "ChaosArchive.h"
+#include "Math/VectorRegister.h"
+#include "Math/UnrealMathVectorConstants.h"
+#include "Chaos/VectorUtility.h"
+
+
+namespace Chaos
+{
+
+class FAABBVectorized
+{
+public:
+
+	FORCEINLINE FAABBVectorized()
+	{
+		Min = GlobalVectorConstants::BigNumber;
+		Max = VectorNegate(GlobalVectorConstants::BigNumber);
+	}
+
+	FORCEINLINE FAABBVectorized(const VectorRegister4Float& Min, const VectorRegister4Float& Max)
+		: Min(Min)
+		, Max(Max)
+	{
+	}
+
+	FORCEINLINE const VectorRegister4Float& GetMin() const { return Min; }
+	FORCEINLINE const VectorRegister4Float& GetMax() const { return Max; }
+
+	FORCEINLINE_DEBUGGABLE bool RaycastFast(const VectorRegister4Float& StartPoint, const VectorRegister4Float& Dir, const VectorRegister4Float& InvDir, const bool* bParallel, const VectorRegister4Float& Length, const VectorRegister4Float& InvLength, VectorRegister4Float& OutEntryTime, VectorRegister4Float& OutExitTime) const
+	{
+		const VectorRegister4Float StartToMin = VectorSubtract(Min, StartPoint);
+		const VectorRegister4Float StartToMax = VectorSubtract(Max, StartPoint);
+		const VectorRegister4Float Parallel = MakeVectorRegisterFloat(~uint32(0) * static_cast<uint32>(bParallel[0]), ~uint32(0) * static_cast<uint32>(bParallel[1]), ~uint32(0) * static_cast<uint32>(bParallel[2]), 0);
+		const VectorRegister4Float StarToMinGTZero = VectorCompareGT(StartToMin, VectorZero());
+		const VectorRegister4Float ZeroGTStarToMax = VectorCompareGT(VectorZero(), StartToMax);
+		VectorRegister4Float IsFalse = VectorBitwiseAnd(VectorBitwiseOr(StarToMinGTZero, ZeroGTStarToMax), Parallel);
+
+		if (VectorMaskBits(IsFalse))
+		{
+			return false;	//parallel and outside
+		}
+
+		const VectorRegister4Float StartToMinInvDir = VectorMultiply(StartToMin, InvDir);
+		const VectorRegister4Float StartToMaxInvDir = VectorMultiply(StartToMax, InvDir);
+
+		VectorRegister4Float Time1 = VectorSelect(Parallel, VectorZero(), StartToMinInvDir);
+		VectorRegister4Float Time2 = VectorSelect(Parallel, Length, StartToMaxInvDir);
+		const VectorRegister4Float Time1GTTime2 = VectorCompareGT(Time1, Time2);
+
+		// Maybe could be optimized with a xor
+		const VectorRegister4Float TmpTime1 = VectorSelect(Time1GTTime2, Time2, Time1);
+		Time2 = VectorSelect(Time1GTTime2, Time1, Time2);
+		Time1 = TmpTime1;
+
+		VectorRegister4Float LatestStartTime = VectorMax(Time1, VectorSwizzle(Time1, 1, 2, 0, 3));
+		LatestStartTime = VectorMax(LatestStartTime, VectorSwizzle(Time1, 2, 0, 1, 3));
+
+		VectorRegister4Float EarliestEndTime = VectorMin(Time2, VectorSwizzle(Time2, 1, 2, 0, 3));
+		EarliestEndTime = VectorMin(EarliestEndTime, VectorSwizzle(Time2, 2, 0, 1, 3));
+
+		//Outside of slab before entering another
+		IsFalse = VectorCompareGT(LatestStartTime, EarliestEndTime);
+
+		if (VectorMaskBits(IsFalse))
+		{
+			return false; 
+		}
+		OutEntryTime = LatestStartTime;
+		OutExitTime = EarliestEndTime;
+		return true;
+	}
+
+private:
+	VectorRegister4Float Min, Max;
+};
+}
\ No newline at end of file
