Fix sweep with scaled rotated triangle meshes

diff --git a/Engine/Source/Runtime/Experimental/Chaos/Private/Chaos/TriangleMeshImplicitObject.cpp b/Engine/Source/Runtime/Experimental/Chaos/Private/Chaos/TriangleMeshImplicitObject.cpp
--- a/Engine/Source/Runtime/Experimental/Chaos/Private/Chaos/TriangleMeshImplicitObject.cpp
+++ b/Engine/Source/Runtime/Experimental/Chaos/Private/Chaos/TriangleMeshImplicitObject.cpp
@@ -951,9 +951,7 @@ struct FTriangleMeshSweepVisitor
 			}
 		}
 
-		const auto& WorldScaleQueryGeom = ScaleGeomIntoWorldHelper(QueryGeom, TriMeshScale);
-
-		if(GJKRaycast2<FReal>(Tri, WorldScaleQueryGeom, ScaledStartTM, ScaledDirNormalized, LengthScale * CurData.CurrentLength, Time, HitPosition, HitNormal, Thickness, bComputeMTD))
+		if(GJKRaycast2<FReal>(Tri, QueryGeom, ScaledStartTM, ScaledDirNormalized, LengthScale * CurData.CurrentLength, Time, HitPosition, HitNormal, Thickness, bComputeMTD))
 		{
 			// Time is world scale, OutTime is local scale.
 			if(Time < LengthScale * OutTime)
@@ -1017,6 +1015,24 @@ void ComputeScaledSweepInputs(FVec3 TriMeshScale, const FRigidTransform3& StartT
 	OutScaledStartTM = FRigidTransform3(StartTM.GetLocation() * TriMeshScale, StartTM.GetRotation());
 }
 
+FVec3 SafeInvScale(const FVec3& Scale)
+{
+	constexpr FReal MinMagnitude = 1e-6f; // consistent with ImplicitObjectScaled::SetScale
+	FVec3 InvScale;
+	for (int Axis = 0; Axis < 3; ++Axis)
+	{
+		if (FMath::Abs(Scale[Axis]) < MinMagnitude)
+		{
+			InvScale[Axis] = 1 / MinMagnitude;
+		}
+		else
+		{
+			InvScale[Axis] = 1 / Scale[Axis];
+		}
+	}
+	return InvScale;
+}
+
 template <typename QueryGeomType>
 bool FTriangleMeshImplicitObject::SweepGeomImp(const QueryGeomType& QueryGeom, const FRigidTransform3& StartTM, const FVec3& Dir, const FReal Length,
 	FReal& OutTime, FVec3& OutPosition, FVec3& OutNormal, int32& OutFaceIndex, const FReal Thickness, const bool bComputeMTD, FVec3 TriMeshScale) const
@@ -1035,10 +1051,9 @@ bool FTriangleMeshImplicitObject::SweepGeomImp(const QueryGeomType& QueryGeom, c
 		using VisitorType = FTriangleMeshSweepVisitor<QueryGeomType,decltype(Elements[0][0])>;
 		VisitorType SQVisitor(*this,Elements, QueryGeom,StartTM,Dir,ScaledDirNormalized,LengthScale,ScaledStartTM,Thickness,bComputeMTD, TriMeshScale, CullsBackFaceRaycastCode);
 
-
 		const FAABB3 QueryBounds = QueryGeom.BoundingBox().TransformedAABB(FRigidTransform3(FVec3::ZeroVector,StartTM.GetRotation()));
-		const FVec3 StartPoint = StartTM.TransformPositionNoScale(QueryBounds.Center());
-		const FVec3 Inflation = QueryBounds.Extents() * 0.5 + FVec3(Thickness);
+		const FVec3 StartPoint = StartTM.TransformPosition(QueryBounds.Center());
+		const FVec3 Inflation = QueryBounds.Extents() * SafeInvScale(TriMeshScale) * 0.5 + FVec3(Thickness);
 		BVH.template Sweep<VisitorType>(StartPoint,Dir,Length,Inflation,SQVisitor);
 
 		if(SQVisitor.OutTime <= Length)
@@ -1062,24 +1077,24 @@ bool FTriangleMeshImplicitObject::SweepGeomImp(const QueryGeomType& QueryGeom, c
 	return bHit;
 }
 
-bool FTriangleMeshImplicitObject::SweepGeom(const TSphere<FReal,3>& QueryGeom, const FRigidTransform3& StartTM, const FVec3& Dir, const FReal Length, FReal& OutTime, FVec3& OutPosition, FVec3& OutNormal, int32& OutFaceIndex, const FReal Thickness, const bool bComputeMTD) const
+bool FTriangleMeshImplicitObject::SweepGeom(const TSphere<FReal,3>& QueryGeom, const FRigidTransform3& StartTM, const FVec3& Dir, const FReal Length, FReal& OutTime, FVec3& OutPosition, FVec3& OutNormal, int32& OutFaceIndex, const FReal Thickness, const bool bComputeMTD, FVec3 TriMeshScale) const
 {
-	return SweepGeomImp(QueryGeom, StartTM, Dir, Length, OutTime, OutPosition, OutNormal, OutFaceIndex, Thickness, bComputeMTD);
+	return SweepGeomImp(QueryGeom, StartTM, Dir, Length, OutTime, OutPosition, OutNormal, OutFaceIndex, Thickness, bComputeMTD, TriMeshScale);
 }
 
-bool FTriangleMeshImplicitObject::SweepGeom(const TBox<FReal, 3>& QueryGeom, const FRigidTransform3& StartTM, const FVec3& Dir, const FReal Length, FReal& OutTime, FVec3& OutPosition, FVec3& OutNormal, int32& OutFaceIndex, const FReal Thickness, const bool bComputeMTD) const
+bool FTriangleMeshImplicitObject::SweepGeom(const TBox<FReal, 3>& QueryGeom, const FRigidTransform3& StartTM, const FVec3& Dir, const FReal Length, FReal& OutTime, FVec3& OutPosition, FVec3& OutNormal, int32& OutFaceIndex, const FReal Thickness, const bool bComputeMTD, FVec3 TriMeshScale) const
 {
-	return SweepGeomImp(QueryGeom, StartTM, Dir, Length, OutTime, OutPosition, OutNormal, OutFaceIndex, Thickness, bComputeMTD);
+	return SweepGeomImp(QueryGeom, StartTM, Dir, Length, OutTime, OutPosition, OutNormal, OutFaceIndex, Thickness, bComputeMTD, TriMeshScale);
 }
 
-bool FTriangleMeshImplicitObject::SweepGeom(const FCapsule& QueryGeom, const FRigidTransform3& StartTM, const FVec3& Dir, const FReal Length, FReal& OutTime, FVec3& OutPosition, FVec3& OutNormal, int32& OutFaceIndex, const FReal Thickness, const bool bComputeMTD) const
+bool FTriangleMeshImplicitObject::SweepGeom(const FCapsule& QueryGeom, const FRigidTransform3& StartTM, const FVec3& Dir, const FReal Length, FReal& OutTime, FVec3& OutPosition, FVec3& OutNormal, int32& OutFaceIndex, const FReal Thickness, const bool bComputeMTD, FVec3 TriMeshScale) const
 {
-	return SweepGeomImp(QueryGeom, StartTM, Dir, Length, OutTime, OutPosition, OutNormal, OutFaceIndex, Thickness, bComputeMTD);
+	return SweepGeomImp(QueryGeom, StartTM, Dir, Length, OutTime, OutPosition, OutNormal, OutFaceIndex, Thickness, bComputeMTD, TriMeshScale);
 }
 
-bool FTriangleMeshImplicitObject::SweepGeom(const FConvex& QueryGeom, const FRigidTransform3& StartTM, const FVec3& Dir, const FReal Length, FReal& OutTime, FVec3& OutPosition, FVec3& OutNormal, int32& OutFaceIndex, const FReal Thickness, const bool bComputeMTD) const
+bool FTriangleMeshImplicitObject::SweepGeom(const FConvex& QueryGeom, const FRigidTransform3& StartTM, const FVec3& Dir, const FReal Length, FReal& OutTime, FVec3& OutPosition, FVec3& OutNormal, int32& OutFaceIndex, const FReal Thickness, const bool bComputeMTD, FVec3 TriMeshScale) const
 {
-	return SweepGeomImp(QueryGeom, StartTM, Dir, Length, OutTime, OutPosition, OutNormal, OutFaceIndex, Thickness, bComputeMTD);
+	return SweepGeomImp(QueryGeom, StartTM, Dir, Length, OutTime, OutPosition, OutNormal, OutFaceIndex, Thickness, bComputeMTD, TriMeshScale);
 }
 
 bool FTriangleMeshImplicitObject::SweepGeom(const TImplicitObjectScaled<TSphere<FReal, 3>>& QueryGeom, const FRigidTransform3& StartTM, const FVec3& Dir, const FReal Length, FReal& OutTime, FVec3& OutPosition, FVec3& OutNormal, int32& OutFaceIndex, const FReal Thickness, const bool bComputeMTD, FVec3 TriMeshScale) const
diff --git a/Engine/Source/Runtime/Experimental/Chaos/Public/Chaos/ImplicitObjectScaled.h b/Engine/Source/Runtime/Experimental/Chaos/Public/Chaos/ImplicitObjectScaled.h
--- a/Engine/Source/Runtime/Experimental/Chaos/Public/Chaos/ImplicitObjectScaled.h
+++ b/Engine/Source/Runtime/Experimental/Chaos/Public/Chaos/ImplicitObjectScaled.h
@@ -650,11 +650,9 @@ public:
 			TVector<T, d> UnscaledNormal;
 			T UnscaledTime;
 
-			auto ScaledB = MakeScaledHelper(B, MInvScale);
-
 			TRigidTransform<T, d> BToATMNoScale(BToATM.GetLocation() * MInvScale, BToATM.GetRotation());
 			
-			if (MObject->SweepGeom(ScaledB, BToATMNoScale, UnscaledDir, UnscaledLength, UnscaledTime, UnscaledPosition, UnscaledNormal, OutFaceIndex, Thickness, bComputeMTD, MScale))
+			if (MObject->SweepGeom(B, BToATMNoScale, UnscaledDir, UnscaledLength, UnscaledTime, UnscaledPosition, UnscaledNormal, OutFaceIndex, Thickness, bComputeMTD, MScale))
 			{
 				const T NewTime = LengthScaleInv * UnscaledTime;
 				//We double check that NewTime < Length because of potential precision issues. When that happens we always keep the shortest hit first
diff --git a/Engine/Source/Runtime/Experimental/Chaos/Public/Chaos/TriangleMeshImplicitObject.h b/Engine/Source/Runtime/Experimental/Chaos/Public/Chaos/TriangleMeshImplicitObject.h
--- a/Engine/Source/Runtime/Experimental/Chaos/Public/Chaos/TriangleMeshImplicitObject.h
+++ b/Engine/Source/Runtime/Experimental/Chaos/Public/Chaos/TriangleMeshImplicitObject.h
@@ -218,10 +218,10 @@ namespace Chaos
 		bool OverlapGeom(const TImplicitObjectScaled<FCapsule>& QueryGeom, const FRigidTransform3& QueryTM, const FReal Thickness, FMTDInfo* OutMTD = nullptr, FVec3 TriMeshScale = FVec3(1.0f)) const;
 		bool OverlapGeom(const TImplicitObjectScaled<FConvex>& QueryGeom, const FRigidTransform3& QueryTM, const FReal Thickness, FMTDInfo* OutMTD = nullptr, FVec3 TriMeshScale = FVec3(1.0f)) const;
 
-		bool SweepGeom(const TSphere<FReal, 3>& QueryGeom, const FRigidTransform3& StartTM, const FVec3& Dir, const FReal Length, FReal& OutTime, FVec3& OutPosition, FVec3& OutNormal, int32& OutFaceIndex, const FReal Thickness = 0, const bool bComputeMTD = false) const;
-		bool SweepGeom(const TBox<FReal, 3>& QueryGeom, const FRigidTransform3& StartTM, const FVec3& Dir, const FReal Length, FReal& OutTime, FVec3& OutPosition, FVec3& OutNormal, int32& OutFaceIndex, const FReal Thickness = 0, const bool bComputeMTD = false) const;
-		bool SweepGeom(const FCapsule& QueryGeom, const FRigidTransform3& StartTM, const FVec3& Dir, const FReal Length, FReal& OutTime, FVec3& OutPosition, FVec3& OutNormal, int32& OutFaceIndex, const FReal Thickness = 0, const bool bComputeMTD = false) const;
-		bool SweepGeom(const FConvex& QueryGeom, const FRigidTransform3& StartTM, const FVec3& Dir, const FReal Length, FReal& OutTime, FVec3& OutPosition, FVec3& OutNormal, int32& OutFaceIndex, const FReal Thickness = 0, const bool bComputeMTD = false) const;
+		bool SweepGeom(const TSphere<FReal, 3>& QueryGeom, const FRigidTransform3& StartTM, const FVec3& Dir, const FReal Length, FReal& OutTime, FVec3& OutPosition, FVec3& OutNormal, int32& OutFaceIndex, const FReal Thickness = 0, const bool bComputeMTD = false, FVec3 TriMeshScale = FVec3(1.0f)) const;
+		bool SweepGeom(const TBox<FReal, 3>& QueryGeom, const FRigidTransform3& StartTM, const FVec3& Dir, const FReal Length, FReal& OutTime, FVec3& OutPosition, FVec3& OutNormal, int32& OutFaceIndex, const FReal Thickness = 0, const bool bComputeMTD = false, FVec3 TriMeshScale = FVec3(1.0f)) const;
+		bool SweepGeom(const FCapsule& QueryGeom, const FRigidTransform3& StartTM, const FVec3& Dir, const FReal Length, FReal& OutTime, FVec3& OutPosition, FVec3& OutNormal, int32& OutFaceIndex, const FReal Thickness = 0, const bool bComputeMTD = false, FVec3 TriMeshScale = FVec3(1.0f)) const;
+		bool SweepGeom(const FConvex& QueryGeom, const FRigidTransform3& StartTM, const FVec3& Dir, const FReal Length, FReal& OutTime, FVec3& OutPosition, FVec3& OutNormal, int32& OutFaceIndex, const FReal Thickness = 0, const bool bComputeMTD = false, FVec3 TriMeshScale = FVec3(1.0f)) const;
 
 		bool SweepGeom(const TImplicitObjectScaled<TSphere<FReal, 3>>& QueryGeom, const FRigidTransform3& StartTM, const FVec3& Dir, const FReal Length, FReal& OutTime, FVec3& OutPosition, FVec3& OutNormal, int32& OutFaceIndex, const FReal Thickness = 0, const bool bComputeMTD = false, FVec3 TriMeshScale = FVec3(1.0f)) const;
 		bool SweepGeom(const TImplicitObjectScaled<TBox<FReal, 3>>& QueryGeom, const FRigidTransform3& StartTM, const FVec3& Dir, const FReal Length, FReal& OutTime, FVec3& OutPosition, FVec3& OutNormal, int32& OutFaceIndex, const FReal Thickness = 0, const bool bComputeMTD = false, FVec3 TriMeshScale = FVec3(1.0f)) const;
