#ue5 - LWC: Added *ToInt32() versions to Trunc/Round/Floor/Ceil to make those able to be explicit and cleaner to replace where needed. Added missing TruncToInt64(double), RoundToInt(double), RoundToInt64(double).

Added overloads for *ToFloat(double) methods to TUnrealPlatformMathSSE4Base to avoid ambiguous overloads (and make it easier for us to update later when do another pass at int64 versions).

Resolved ambiguous Max() and Min() calls when mixing int32/int64, float/double, etc by using macro to promote to the largest precision type.

diff --git a/Engine/Source/Editor/MovieSceneTools/Private/TrackEditors/AttachTrackEditor.cpp b/Engine/Source/Editor/MovieSceneTools/Private/TrackEditors/AttachTrackEditor.cpp
--- a/Engine/Source/Editor/MovieSceneTools/Private/TrackEditors/AttachTrackEditor.cpp
+++ b/Engine/Source/Editor/MovieSceneTools/Private/TrackEditors/AttachTrackEditor.cpp
@@ -1044,7 +1044,7 @@ void F3DAttachTrackEditor::CompensateChildTrack(const TRange<FFrameNumber>& InAt
 		FFrameRate TickResolution = GetSequencer()->GetFocusedTickResolution();
 		FFrameRate DisplayRate = GetSequencer()->GetFocusedDisplayRate();
 		for (FFrameNumber FrameItr = InAttachRange.GetLowerBoundValue(); FrameItr < InAttachRange.GetUpperBoundValue(); 
-			FrameItr += FMath::RoundToInt(TickResolution.AsDecimal() / DisplayRate.AsDecimal()))
+			FrameItr += FMath::RoundToInt32(TickResolution.AsDecimal() / DisplayRate.AsDecimal()))
 		{
 			ResizeAndAddKey(FrameItr, Channels.Num(), TransformMap, &KeyTimesToCompensate);
 			for (FMovieSceneDoubleValue& DoubleVal : TransformMap[FrameItr])
diff --git a/Engine/Source/Editor/Persona/Private/AnimTimeline/AnimModel.cpp b/Engine/Source/Editor/Persona/Private/AnimTimeline/AnimModel.cpp
--- a/Engine/Source/Editor/Persona/Private/AnimTimeline/AnimModel.cpp
+++ b/Engine/Source/Editor/Persona/Private/AnimTimeline/AnimModel.cpp
@@ -66,13 +66,13 @@ double FAnimModel::GetFrameRate() const
 
 int32 FAnimModel::GetTickResolution() const
 {
-	return FMath::RoundToInt((double)GetDefault<UPersonaOptions>()->TimelineScrubSnapValue * GetFrameRate());
+	return FMath::RoundToInt32((double)GetDefault<UPersonaOptions>()->TimelineScrubSnapValue * GetFrameRate());
 }
 
 TRange<FFrameNumber> FAnimModel::GetPlaybackRange() const
 {
 	const int32 Resolution = GetTickResolution();
-	return TRange<FFrameNumber>(FFrameNumber(FMath::RoundToInt(PlaybackRange.GetLowerBoundValue() * (double)Resolution)), FFrameNumber(FMath::RoundToInt(PlaybackRange.GetUpperBoundValue() * (double)Resolution)));
+	return TRange<FFrameNumber>(FFrameNumber(FMath::RoundToInt32(PlaybackRange.GetLowerBoundValue() * (double)Resolution)), FFrameNumber(FMath::RoundToInt32(PlaybackRange.GetUpperBoundValue() * (double)Resolution)));
 }
 
 FFrameNumber FAnimModel::GetScrubPosition() const
@@ -82,7 +82,7 @@ FFrameNumber FAnimModel::GetScrubPosition() const
 		UDebugSkelMeshComponent* PreviewMeshComponent = WeakPreviewScene.Pin()->GetPreviewMeshComponent();
 		if(PreviewMeshComponent && PreviewMeshComponent->IsPreviewOn())
 		{
-			return FFrameNumber(FMath::RoundToInt(PreviewMeshComponent->PreviewInstance->GetCurrentTime() * (double)GetTickResolution()));
+			return FFrameNumber(FMath::RoundToInt32(PreviewMeshComponent->PreviewInstance->GetCurrentTime() * (double)GetTickResolution()));
 		}
 	}
 
diff --git a/Engine/Source/Editor/UnrealEd/Private/WorldPartition/WorldPartitionMiniMapBuilder.cpp b/Engine/Source/Editor/UnrealEd/Private/WorldPartition/WorldPartitionMiniMapBuilder.cpp
--- a/Engine/Source/Editor/UnrealEd/Private/WorldPartition/WorldPartitionMiniMapBuilder.cpp
+++ b/Engine/Source/Editor/UnrealEd/Private/WorldPartition/WorldPartitionMiniMapBuilder.cpp
@@ -119,8 +119,8 @@ bool UWorldPartitionMiniMapBuilder::PostRun(UWorld* World, FPackageSourceControl
 	SourceBlocks.Reserve(MiniMapTiles.Num());
 	SourceImageData.Reserve(MiniMapTiles.Num());
 
-	FTileCoordinates Min = { FMath::FloorToInt(EditorBounds.Min.X / IterativeCellSize), FMath::FloorToInt(EditorBounds.Min.Y / IterativeCellSize) };
-	FTileCoordinates Max = { FMath::FloorToInt(EditorBounds.Max.X / IterativeCellSize), FMath::FloorToInt(EditorBounds.Max.Y / IterativeCellSize) };
+	FTileCoordinates Min = { FMath::FloorToInt32(EditorBounds.Min.X / IterativeCellSize), FMath::FloorToInt32(EditorBounds.Min.Y / IterativeCellSize) };
+	FTileCoordinates Max = { FMath::FloorToInt32(EditorBounds.Max.X / IterativeCellSize), FMath::FloorToInt32(EditorBounds.Max.Y / IterativeCellSize) };
 
 	FVector2D				Offset;
 	ETextureSourceFormat	Format = TSF_BGRA8;
diff --git a/Engine/Source/Runtime/Core/Private/IO/IoStore.cpp b/Engine/Source/Runtime/Core/Private/IO/IoStore.cpp
--- a/Engine/Source/Runtime/Core/Private/IO/IoStore.cpp
+++ b/Engine/Source/Runtime/Core/Private/IO/IoStore.cpp
@@ -771,7 +771,7 @@ public:
 		// https://en.wikipedia.org/wiki/Perfect_hash_function
 		TRACE_CPUPROFILER_EVENT_SCOPE(TocGeneratePerfectHashes);
 		uint32 ChunkCount = TocResource.ChunkIds.Num();
-		uint32 SeedCount = FMath::Max(1, FMath::RoundToInt(ChunkCount / 2.0));
+		uint32 SeedCount = FMath::Max(1, FMath::RoundToInt32(ChunkCount / 2.0));
 		check(TocResource.ChunkOffsetLengths.Num() == ChunkCount);
 		
 		TArray<FIoChunkId> OutTocChunkIds;
diff --git a/Engine/Source/Runtime/Core/Private/Math/UnrealMath.cpp b/Engine/Source/Runtime/Core/Private/Math/UnrealMath.cpp
--- a/Engine/Source/Runtime/Core/Private/Math/UnrealMath.cpp
+++ b/Engine/Source/Runtime/Core/Private/Math/UnrealMath.cpp
@@ -1768,11 +1768,11 @@ static bool ComputeProjectedSphereShaft(
 			FVector::FReal X = (Dot3(P,Axis) / P.W + 1.0f * AxisSign) / 2.0f * AxisSign;
 			if(FMath::IsNegative(Nxa) ^ FMath::IsNegative(AxisSign))
 			{
-				InOutMaxX = FMath::Min<int32>(FMath::CeilToInt(ViewSizeX * X + ViewX),InOutMaxX);
+				InOutMaxX = FMath::Min(FMath::CeilToInt32(ViewSizeX * X + ViewX),InOutMaxX);
 			}
 			else
 			{
-				InOutMinX = FMath::Max<int32>(FMath::FloorToInt(ViewSizeX * X + ViewX),InOutMinX);
+				InOutMinX = FMath::Max(FMath::FloorToInt32(ViewSizeX * X + ViewX),InOutMinX);
 			}
 		}
 
@@ -1784,11 +1784,11 @@ static bool ComputeProjectedSphereShaft(
 			FVector::FReal X = (Dot3(P,Axis) / P.W + 1.0f * AxisSign) / 2.0f * AxisSign;
 			if(FMath::IsNegative(Nxb) ^ FMath::IsNegative(AxisSign))
 			{
-				InOutMaxX = FMath::Min<int32>(FMath::CeilToInt(ViewSizeX * X + ViewX),InOutMaxX);
+				InOutMaxX = FMath::Min(FMath::CeilToInt32(ViewSizeX * X + ViewX),InOutMaxX);
 			}
 			else
 			{
-				InOutMinX = FMath::Max<int32>(FMath::FloorToInt(ViewSizeX * X + ViewX),InOutMinX);
+				InOutMinX = FMath::Max(FMath::FloorToInt32(ViewSizeX * X + ViewX),InOutMinX);
 			}
 		}
 	}
diff --git a/Engine/Source/Runtime/Core/Private/Tests/Math/UnrealMathTest.cpp b/Engine/Source/Runtime/Core/Private/Tests/Math/UnrealMathTest.cpp
--- a/Engine/Source/Runtime/Core/Private/Tests/Math/UnrealMathTest.cpp
+++ b/Engine/Source/Runtime/Core/Private/Tests/Math/UnrealMathTest.cpp
@@ -3069,6 +3069,11 @@ bool FVectorRegisterAbstractionTest::RunTest(const FString& Parameters)
 		// Expected to generate errors (float->int truncation)
 		I = FMath::Clamp(I1, F2, I3);
 #endif
+
+		int64 Big1 = 0, Big2 = 0, Big3 = 0;
+		Big3 = FMath::Max(Big1, Big2);
+		Big3 = FMath::Max(Big1, I2);
+		Big3 = FMath::Max(I1, Big2);
 	}
 
 	// Sin, Cos, Tan tests
diff --git a/Engine/Source/Runtime/Core/Private/Tests/Misc/TimecodeTest.cpp b/Engine/Source/Runtime/Core/Private/Tests/Misc/TimecodeTest.cpp
--- a/Engine/Source/Runtime/Core/Private/Tests/Misc/TimecodeTest.cpp
+++ b/Engine/Source/Runtime/Core/Private/Tests/Misc/TimecodeTest.cpp
@@ -107,7 +107,7 @@ bool FTimecodeTest::RunTest(const FString& Parameters)
 					// Do they have the same hours, minutes, seconds
 					// To test this, we start from the number of events (FrameIndex) from which we got our timecode first
 					// Timecode is just a label and doesn't necessarily reflect real time especially with 23.976 FrameRate like
-					const int32 TotalSeconds = FMath::FloorToInt((double)FrameIndex * FrameRate.AsInterval());
+					const int32 TotalSeconds = FMath::FloorToInt32((double)FrameIndex * FrameRate.AsInterval());
 					const int32 FrameHours = (TotalSeconds / (60 * 60));
 					const int32 FrameMinutes = ((TotalSeconds % (60 * 60)) / 60);
 					const int32 FrameSeconds = ((TotalSeconds % (60 * 60)) % 60);
diff --git a/Engine/Source/Runtime/Core/Public/GenericPlatform/GenericPlatformMath.h b/Engine/Source/Runtime/Core/Public/GenericPlatform/GenericPlatformMath.h
--- a/Engine/Source/Runtime/Core/Public/GenericPlatform/GenericPlatformMath.h
+++ b/Engine/Source/Runtime/Core/Public/GenericPlatform/GenericPlatformMath.h
@@ -197,15 +197,21 @@ struct FGenericPlatformMath
 	 * @param F		Floating point value to convert
 	 * @return		Truncated integer.
 	 */
-	static CONSTEXPR FORCEINLINE int32 TruncToInt(float F)
+	static CONSTEXPR FORCEINLINE int32 TruncToInt32(float F)
 	{
 		return (int32)F;
 	}
-	static CONSTEXPR FORCEINLINE int32 TruncToInt(double F)
+	static CONSTEXPR FORCEINLINE int32 TruncToInt32(double F)
 	{
 		return (int32)F;
 	}
-	RESOLVE_FLOAT_TO_TYPE_AMBIGUITY(TruncToInt, int32);
+	static CONSTEXPR FORCEINLINE int64 TruncToInt64(double F)
+	{
+		return (int64)F;
+	}
+
+	static CONSTEXPR FORCEINLINE int32 TruncToInt(float F) { return TruncToInt32(F); }
+	static CONSTEXPR FORCEINLINE int64 TruncToInt(double F) { return TruncToInt64(F); }
 
 	/**
 	 * Converts a float to an integer value with truncation towards zero.
@@ -239,26 +245,28 @@ struct FGenericPlatformMath
 	 * @param F		Floating point value to convert
 	 * @return		An integer less or equal to 'F'.
 	 */
-	static FORCEINLINE int32 FloorToInt(float F)
+	static FORCEINLINE int32 FloorToInt32(float F)
 	{
-		int32 I = (int32)F;
-		I -= ( (float)I > F );
+		int32 I = TruncToInt32(F);
+		I -= ((float)I > F);
 		return I;
 	}
-	static FORCEINLINE int32 FloorToInt(double F)
+	static FORCEINLINE int32 FloorToInt32(double F)
 	{
-		int32 I = (int32)F;
-		I -= ((float)I > F);
+		int32 I = TruncToInt32(F);
+		I -= ((double)I > F);
 		return I;
 	}
-	RESOLVE_FLOAT_TO_TYPE_AMBIGUITY(FloorToInt, int32);
-
 	static FORCEINLINE int64 FloorToInt64(double F)
 	{
-		int64 I = (int64)F;
+		int64 I = TruncToInt64(F);
 		I -= ((double)I > F);
 		return I;
 	}
+
+	static FORCEINLINE int32 FloorToInt(float F) { return FloorToInt32(F); }
+	static FORCEINLINE int64 FloorToInt(double F) { return FloorToInt64(F); }
+	
 	
 	/**
 	* Converts a float to the nearest less or equal integer.
@@ -292,10 +300,21 @@ struct FGenericPlatformMath
 	 * @param F		Floating point value to convert
 	 * @return		The nearest integer to 'F'.
 	 */
-	static FORCEINLINE int32 RoundToInt(float F)
+	static FORCEINLINE int32 RoundToInt32(float F)
+	{
+		return FloorToInt32(F + 0.5f);
+	}
+	static FORCEINLINE int32 RoundToInt32(double F)
 	{
-		return FloorToInt(F + 0.5f);
+		return FloorToInt32(F + 0.5);
 	}
+	static FORCEINLINE int64 RoundToInt64(double F)
+	{
+		return FloorToInt64(F + 0.5);
+	}
+	
+	static FORCEINLINE int32 RoundToInt(float F) { return RoundToInt32(F); }
+	static FORCEINLINE int64 RoundToInt(double F) { return RoundToInt64(F); }
 
 	/**
 	* Converts a float to the nearest integer. Rounds up when the fraction is .5
@@ -329,29 +348,28 @@ struct FGenericPlatformMath
 	* @param F		Floating point value to convert
 	* @return		An integer greater or equal to 'F'.
 	*/
-	static FORCEINLINE int32 CeilToInt(float F)
+	static FORCEINLINE int32 CeilToInt32(float F)
 	{
-		int32 I = (int32)F;
-		I += ( (float)I < F );
+		int32 I = TruncToInt32(F);
+		I += ((float)I < F);
 		return I;
 	}
-
-	static FORCEINLINE int32 CeilToInt(double F)
+	static FORCEINLINE int32 CeilToInt32(double F)
 	{
-		int32 I = (int32)F;
-		I += ((float)I < F);
+		int32 I = TruncToInt32(F);
+		I += ((double)I < F);
 		return I;
 	}
-
-	RESOLVE_FLOAT_TO_TYPE_AMBIGUITY(CeilToInt, int32);
-
 	static FORCEINLINE int64 CeilToInt64(double F)
 	{
-		int64 I = (int64)F;
+		int64 I = TruncToInt64(F);
 		I += ((double)I < F);
 		return I;
 	}
 
+	static FORCEINLINE int32 CeilToInt(float F) { return CeilToInt32(F); }
+	static FORCEINLINE int64 CeilToInt(double F) { return CeilToInt64(F); }
+
 	/**
 	* Converts a float to the nearest greater or equal integer.
 	* @param F		Floating point value to convert
@@ -916,11 +934,9 @@ struct FGenericPlatformMath
 		return (A<=B) ? A : B;
 	}
 
-	// LWC_TODO: Simplify code refactoring to avoid double/float mismatches. Should not ship?
-	static CONSTEXPR FORCEINLINE double Max(const double A, const float B) { return Max<double>(A, B); }
-	static CONSTEXPR FORCEINLINE double Max(const float A, const double B) { return Max<double>(A, B); }
-	static CONSTEXPR FORCEINLINE double Min(const double A, const float B) { return Min<double>(A, B); }
-	static CONSTEXPR FORCEINLINE double Min(const float A, const double B) { return Min<double>(A, B); }
+	// Allow mixing of types to promote to highest precision type
+	MIX_TYPES_2_ARGS_CONSTEXPR(Max);
+	MIX_TYPES_2_ARGS_CONSTEXPR(Min);
 
 	/**
 	* Min of Array
diff --git a/Engine/Source/Runtime/Core/Public/Math/IntRect.h b/Engine/Source/Runtime/Core/Public/Math/IntRect.h
--- a/Engine/Source/Runtime/Core/Public/Math/IntRect.h
+++ b/Engine/Source/Runtime/Core/Public/Math/IntRect.h
@@ -315,8 +315,9 @@ public:
 
 FORCEINLINE FIntRect FIntRect::Scale( float Fraction ) const
 {
-	FVector2D Min2D = FVector2D((float)Min.X, (float)Min.Y) * Fraction;
-	FVector2D Max2D = FVector2D((float)Max.X, (float)Max.Y) * Fraction;
+	using Vec2D = UE::Math::TVector2<float>;
+	const Vec2D Min2D = Vec2D((float)Min.X, (float)Min.Y) * Fraction;
+	const Vec2D Max2D = Vec2D((float)Max.X, (float)Max.Y) * Fraction;
 
 	return FIntRect(FMath::FloorToInt(Min2D.X), FMath::FloorToInt(Min2D.Y), FMath::CeilToInt(Max2D.X), FMath::CeilToInt(Max2D.Y));
 }
diff --git a/Engine/Source/Runtime/Core/Public/Math/UnrealMathUtility.h b/Engine/Source/Runtime/Core/Public/Math/UnrealMathUtility.h
--- a/Engine/Source/Runtime/Core/Public/Math/UnrealMathUtility.h
+++ b/Engine/Source/Runtime/Core/Public/Math/UnrealMathUtility.h
@@ -209,6 +209,9 @@ struct FMath : public FPlatformMath
 		return Min + RandHelper64(Range);
 	}
 
+	UE_NODISCARD static FORCEINLINE int64 RandRange(int32 Min, int64 Max) { return RandRange((int64)Min, Max); }
+	UE_NODISCARD static FORCEINLINE int64 RandRange(int64 Min, int32 Max) { return RandRange(Min, (int64)Max); }
+
 	/** Util to generate a random number in a range. Overloaded to distinguish from int32 version, where passing a float is typically a mistake. */
 	UE_NODISCARD static FORCEINLINE float RandRange(float InMin, float InMax)
 	{
diff --git a/Engine/Source/Runtime/Core/Public/Math/UnrealPlatformMathSSE.h b/Engine/Source/Runtime/Core/Public/Math/UnrealPlatformMathSSE.h
--- a/Engine/Source/Runtime/Core/Public/Math/UnrealPlatformMathSSE.h
+++ b/Engine/Source/Runtime/Core/Public/Math/UnrealPlatformMathSSE.h
@@ -66,12 +66,12 @@ namespace SSE
 		return InvSqrt(InValue);
 	}
 
-	static FORCEINLINE int32 TruncToInt(float F)
+	static FORCEINLINE int32 TruncToInt32(float F)
 	{
 		return _mm_cvtt_ss2si(_mm_set_ss(F));
 	}
 
-	static FORCEINLINE int32 TruncToInt(double InValue)
+	static FORCEINLINE int32 TruncToInt32(double InValue)
 	{
 		return _mm_cvttsd_si32(_mm_set_sd(InValue));
 	}
@@ -81,14 +81,14 @@ namespace SSE
 		return _mm_cvttsd_si64(_mm_set_sd(InValue));
 	}
 
-	static FORCEINLINE int32 FloorToInt(float F)
+	static FORCEINLINE int32 FloorToInt32(float F)
 	{
 		// Note: unlike the Generic solution and the SSE4 float solution, we implement FloorToInt using a rounding instruction, rather than implementing RoundToInt using a floor instruction.  
 		// We therefore need to do the same times-2 transform (with a slighly different formula) that RoundToInt does; see the note on RoundToInt
 		return _mm_cvt_ss2si(_mm_set_ss(F + F - 0.5f)) >> 1;
 	}
 
-	static FORCEINLINE int32 FloorToInt(double InValue)
+	static FORCEINLINE int32 FloorToInt32(double InValue)
 	{
 		return _mm_cvtsd_si32(_mm_set_sd(InValue + InValue - 0.5)) >> 1;
 	}
@@ -98,7 +98,7 @@ namespace SSE
 		return _mm_cvtsd_si64(_mm_set_sd(InValue + InValue - 0.5)) >> 1;
 	}
 
-	static FORCEINLINE int32 RoundToInt(float F)
+	static FORCEINLINE int32 RoundToInt32(float F)
 	{
 		// Note: the times-2 is to remove the rounding-to-nearest-even-number behavior that mm_cvt_ss2si uses when the fraction is .5
 		// The formula we uses causes the round instruction to always be applied to a an odd integer when the original value was 0.5, and eliminates the rounding-to-nearest-even-number behavior
@@ -108,7 +108,7 @@ namespace SSE
 		return _mm_cvt_ss2si(_mm_set_ss(F + F + 0.5f)) >> 1;
 	}
 
-	static FORCEINLINE int32 RoundToInt(double InValue)
+	static FORCEINLINE int32 RoundToInt32(double InValue)
 	{
 		return _mm_cvtsd_si32(_mm_set_sd(InValue + InValue + 0.5)) >> 1;
 	}
@@ -118,14 +118,14 @@ namespace SSE
 		return _mm_cvtsd_si64(_mm_set_sd(InValue + InValue + 0.5)) >> 1;
 	}
 
-	static FORCEINLINE int32 CeilToInt(float F)
+	static FORCEINLINE int32 CeilToInt32(float F)
 	{
 		// Note: unlike the Generic solution and the SSE4 float solution, we implement CeilToInt using a rounding instruction, rather than a dedicated ceil instruction
 		// We therefore need to do the same times-2 transform (with a slighly different formula) that RoundToInt does; see the note on RoundToInt
 		return -(_mm_cvt_ss2si(_mm_set_ss(-0.5f - (F + F))) >> 1);
 	}
 
-	static FORCEINLINE int32 CeilToInt(double InValue)
+	static FORCEINLINE int32 CeilToInt32(double InValue)
 	{
 		return -(_mm_cvtsd_si32(_mm_set_sd(-0.5 - (InValue + InValue))) >> 1);
 	}
@@ -226,41 +226,63 @@ template<class Base>
 struct TUnrealPlatformMathSSEBase : public Base
 {
 	template<typename T>
-	static FORCEINLINE int32 TruncToInt(T F)
+	static FORCEINLINE int32 TruncToInt32(T F)
 	{
-		return UE4::SSE::TruncToInt(F);
+		return UE4::SSE::TruncToInt32(F);
 	}
 
-	template<typename T>
-	static FORCEINLINE int32 RoundToInt(T F)
+	static FORCEINLINE int64 TruncToInt64(double F)
 	{
-		return UE4::SSE::RoundToInt(F);
+		return UE4::SSE::TruncToInt64(F);
 	}
 
 	template<typename T>
-	static FORCEINLINE int32 FloorToInt(T F)
+	static FORCEINLINE int32 RoundToInt32(T F)
+	{
+		return UE4::SSE::RoundToInt32(F);
+	}
+
+	static FORCEINLINE int64 RoundToInt64(double F)
 	{
-		return UE4::SSE::FloorToInt(F);
+		return UE4::SSE::RoundToInt64(F);
 	}
 
 	template<typename T>
-	static FORCEINLINE int64 FloorToInt64(T F)
+	static FORCEINLINE int32 FloorToInt32(T F)
+	{
+		return UE4::SSE::FloorToInt32(F);
+	}
+
+	static FORCEINLINE int64 FloorToInt64(double F)
 	{
 		return UE4::SSE::FloorToInt64(F);
 	}
 
 	template<typename T>
-	static FORCEINLINE int32 CeilToInt(T F)
+	static FORCEINLINE int32 CeilToInt32(T F)
 	{
-		return UE4::SSE::CeilToInt(F);
+		return UE4::SSE::CeilToInt32(F);
 	}
 
-	template<typename T>
-	static FORCEINLINE int64 CeilToInt64(T F)
+	static FORCEINLINE int64 CeilToInt64(double F)
 	{
 		return UE4::SSE::CeilToInt64(F);
 	}
 
+	//
+	// Wrappers for overloads in the base, required since calls declared in base struct won't redirect back to this class
+	//
+
+	static FORCEINLINE int32 TruncToInt(float F) { return TruncToInt32(F); }
+	static FORCEINLINE int64 TruncToInt(double F) { return TruncToInt64(F); }
+	static FORCEINLINE int32 FloorToInt(float F) { return FloorToInt32(F); }
+	static FORCEINLINE int64 FloorToInt(double F) { return FloorToInt64(F); }
+	static FORCEINLINE int32 RoundToInt(float F) { return RoundToInt32(F); }
+	static FORCEINLINE int64 RoundToInt(double F) { return RoundToInt64(F); }
+	static FORCEINLINE int32 CeilToInt(float F) { return CeilToInt32(F); }
+	static FORCEINLINE int64 CeilToInt(double F) { return CeilToInt64(F); }
+
+
 	template<typename T>
 	static FORCEINLINE T InvSqrt(T F)
 	{
diff --git a/Engine/Source/Runtime/Core/Public/Math/UnrealPlatformMathSSE4.h b/Engine/Source/Runtime/Core/Public/Math/UnrealPlatformMathSSE4.h
--- a/Engine/Source/Runtime/Core/Public/Math/UnrealPlatformMathSSE4.h
+++ b/Engine/Source/Runtime/Core/Public/Math/UnrealPlatformMathSSE4.h
@@ -67,6 +67,9 @@ template<class Base>
 struct TUnrealPlatformMathSSE4Base : public TUnrealPlatformMathSSEBase<Base>
 {
 #if UNREALPLATFORMMATH_SSE4_1_ENABLED
+
+	// Truncate
+
 	static FORCEINLINE float TruncToFloat(float F)
 	{
 		return UE4::SSE4::TruncToFloat(F);
@@ -76,6 +79,9 @@ struct TUnrealPlatformMathSSE4Base : public TUnrealPlatformMathSSEBase<Base>
 	{
 		return UE4::SSE4::TruncToDouble(F);
 	}
+
+	// Round
+
 	static FORCEINLINE float RoundToFloat(float F)
 	{
 		return UE4::SSE4::RoundToFloat(F);
@@ -86,6 +92,8 @@ struct TUnrealPlatformMathSSE4Base : public TUnrealPlatformMathSSEBase<Base>
 		return UE4::SSE4::RoundToDouble(F);
 	}
 
+	// Floor
+
 	static FORCEINLINE float FloorToFloat(float F)
 	{
 		return UE4::SSE4::FloorToFloat(F);
@@ -96,6 +104,8 @@ struct TUnrealPlatformMathSSE4Base : public TUnrealPlatformMathSSEBase<Base>
 		return UE4::SSE4::FloorToDouble(F);
 	}
 
+	// Ceil
+
 	static FORCEINLINE float CeilToFloat(float F)
 	{
 		return UE4::SSE4::CeilToFloat(F);
@@ -105,5 +115,16 @@ struct TUnrealPlatformMathSSE4Base : public TUnrealPlatformMathSSEBase<Base>
 	{
 		return UE4::SSE4::CeilToDouble(F);
 	}
+
+
+	//
+	// Wrappers for overloads in the base, required since calls declared in base struct won't redirect back to this class
+	//
+
+	static FORCEINLINE double TruncToFloat(double F) { return TruncToDouble(F); }
+	static FORCEINLINE double RoundToFloat(double F) { return RoundToDouble(F); }
+	static FORCEINLINE double FloorToFloat(double F) { return FloorToDouble(F); }
+	static FORCEINLINE double CeilToFloat(double F) { return CeilToDouble(F); }
+
 #endif // UNREALPLATFORMMATH_SSE4_ENABLED
 };
diff --git a/Engine/Source/Runtime/Core/Public/Math/Vector.h b/Engine/Source/Runtime/Core/Public/Math/Vector.h
--- a/Engine/Source/Runtime/Core/Public/Math/Vector.h
+++ b/Engine/Source/Runtime/Core/Public/Math/Vector.h
@@ -2597,9 +2597,9 @@ inline FVector FMath::VRand()
 
 
 FORCEINLINE FIntVector::FIntVector(FVector InVector)
-    : X(FMath::TruncToInt(InVector.X))
-    , Y(FMath::TruncToInt(InVector.Y))
-    , Z(FMath::TruncToInt(InVector.Z))
+	: X(FMath::TruncToInt32(InVector.X))
+	, Y(FMath::TruncToInt32(InVector.Y))
+	, Z(FMath::TruncToInt32(InVector.Z))
 {
 }
 
diff --git a/Engine/Source/Runtime/Core/Public/Templates/ResolveTypeAmbiguity.h b/Engine/Source/Runtime/Core/Public/Templates/ResolveTypeAmbiguity.h
--- a/Engine/Source/Runtime/Core/Public/Templates/ResolveTypeAmbiguity.h
+++ b/Engine/Source/Runtime/Core/Public/Templates/ResolveTypeAmbiguity.h
@@ -41,28 +41,33 @@
 
 
 // Mixing any types with any other type results in same type as "X * Y", which is promoted to the type of the result of mixed arithmetic between the types
-#define MIX_TYPES_2_ARGS(Func) \
+#define MIX_TYPES_2_ARGS_ACTUAL(Func, OptionalMarkup) \
 	template<typename Arg1, typename Arg2, TEMPLATE_REQUIRES(TNot<TIsSame<Arg1, Arg2>>::Value)> \
-	static FORCEINLINE auto Func(Arg1 X, Arg2 Y) -> decltype(X * Y) \
+	static OptionalMarkup FORCEINLINE auto Func(Arg1 X, Arg2 Y) -> decltype(X * Y) \
 	{ \
 		using ArgType = decltype(X * Y); \
 		return Func((ArgType)X, (ArgType)Y); \
 	}
 
+#define MIX_TYPES_2_ARGS(Func)				MIX_TYPES_2_ARGS_ACTUAL(Func,)
+#define MIX_TYPES_2_ARGS_CONSTEXPR(Func)	MIX_TYPES_2_ARGS_ACTUAL(Func, CONSTEXPR)
+
 // Mixing any types with any other type results in same type as "X * Y", which is promoted to the type of the result of mixed arithmetic between the types
-#define MIX_TYPES_3_ARGS(Func) \
+#define MIX_TYPES_3_ARGS_ACTUAL(Func, OptionalMarkup) \
 	template<typename Arg1, typename Arg2, typename Arg3, TEMPLATE_REQUIRES(TOr< \
 																				TNot<TIsSame<Arg1, Arg2>>, \
 																				TNot<TIsSame<Arg2, Arg3>>, \
 																				TNot<TIsSame<Arg1, Arg3>> \
 																				>::Value \
 																			)> \
-	static FORCEINLINE auto Func(Arg1 X, Arg2 Y, Arg3 Z) -> decltype(X * Y * Z) \
+	static OptionalMarkup FORCEINLINE auto Func(Arg1 X, Arg2 Y, Arg3 Z) -> decltype(X * Y * Z) \
 	{ \
 		using ArgType = decltype(X * Y * Z); \
 		return Func((ArgType)X, (ArgType)Y, (ArgType)Z); \
 	}
 
+#define MIX_TYPES_3_ARGS(Func)				MIX_TYPES_3_ARGS_ACTUAL(Func,)
+#define MIX_TYPES_3_ARGS_CONSTEXPR(Func)	MIX_TYPES_3_ARGS_ACTUAL(Func, CONSTEXPR)
 
 
 // Cannot know whether the type of e.g. Fmod(int, int) should be float or double.
diff --git a/Engine/Source/Runtime/Engine/Private/TickTaskManager.cpp b/Engine/Source/Runtime/Engine/Private/TickTaskManager.cpp
--- a/Engine/Source/Runtime/Engine/Private/TickTaskManager.cpp
+++ b/Engine/Source/Runtime/Engine/Private/TickTaskManager.cpp
@@ -2009,14 +2009,14 @@ void FTickFunction::QueueTickFunction(FTickTaskSequencer& TTS, const struct FTic
 					}
 					else
 					{
-						MaxPrerequisiteTickGroup =  FMath::Max<ETickingGroup>(MaxPrerequisiteTickGroup, Prereq->InternalData->ActualStartTickGroup);
+						MaxPrerequisiteTickGroup =  FMath::Max<ETickingGroup>(MaxPrerequisiteTickGroup, Prereq->InternalData->ActualStartTickGroup.GetValue());
 						TaskPrerequisites.Add(Prereq->GetCompletionHandle());
 					}
 				}
 			}
 
 			// tick group is the max of the prerequisites, the current tick group, and the desired tick group
-			ETickingGroup MyActualTickGroup =  FMath::Max<ETickingGroup>(MaxPrerequisiteTickGroup, FMath::Max<ETickingGroup>(TickGroup,TickContext.TickGroup));
+			ETickingGroup MyActualTickGroup =  FMath::Max<ETickingGroup>(MaxPrerequisiteTickGroup, FMath::Max<ETickingGroup>(TickGroup.GetValue(),TickContext.TickGroup));
 			if (MyActualTickGroup != TickGroup)
 			{
 				// if the tick was "demoted", make sure it ends up in an ordinary tick group.
@@ -2097,7 +2097,7 @@ void FTickFunction::QueueTickFunctionParallel(const struct FTickContext& TickCon
 						}
 						else
 						{
-							MaxPrerequisiteTickGroup = FMath::Max<ETickingGroup>(MaxPrerequisiteTickGroup, Prereq->InternalData->ActualStartTickGroup);
+							MaxPrerequisiteTickGroup = FMath::Max<ETickingGroup>(MaxPrerequisiteTickGroup, Prereq->InternalData->ActualStartTickGroup.GetValue());
 							TaskPrerequisites.Add(Prereq->GetCompletionHandle());
 						}
 					}
diff --git a/Engine/Source/Runtime/Engine/Public/WorldPartition/WorldPartitionStreamingSource.h b/Engine/Source/Runtime/Engine/Public/WorldPartition/WorldPartitionStreamingSource.h
--- a/Engine/Source/Runtime/Engine/Public/WorldPartition/WorldPartitionStreamingSource.h
+++ b/Engine/Source/Runtime/Engine/Public/WorldPartition/WorldPartitionStreamingSource.h
@@ -71,7 +71,7 @@ public:
 			return Segments;
 		}
 
-		const int32 SegmentCount = FMath::Max(4, FMath::CeilToInt(64 * Angle / 360.f));
+		const int32 SegmentCount = FMath::Max(4, FMath::CeilToInt32(64 * (float)Angle / 360.f));
 		const FReal AngleStep = Angle / FReal(SegmentCount);
 		const FRotator ShapeRotation = FRotationMatrix::MakeFromX(Axis).Rotator();
 		const FVector ScaledAxis = FVector::ForwardVector * Radius;
diff --git a/Engine/Source/Runtime/Experimental/Chaos/Private/Chaos/Framework/PhysicsSolverBase.cpp b/Engine/Source/Runtime/Experimental/Chaos/Private/Chaos/Framework/PhysicsSolverBase.cpp
--- a/Engine/Source/Runtime/Experimental/Chaos/Private/Chaos/Framework/PhysicsSolverBase.cpp
+++ b/Engine/Source/Runtime/Experimental/Chaos/Private/Chaos/Framework/PhysicsSolverBase.cpp
@@ -347,13 +347,13 @@ namespace Chaos
 			{
 
 				InternalDt = AsyncDt;
-				NumSteps = FMath::FloorToInt(AccumulatedTime / InternalDt);
+				NumSteps = FMath::FloorToInt32(AccumulatedTime / InternalDt);
 				AccumulatedTime -= InternalDt * static_cast<FReal>(NumSteps);
 			}
 		}
 		else if (bSubstepping && InDt > 0)
 		{
-			NumSteps = FMath::CeilToInt(DtWithPause / MMaxDeltaTime);
+			NumSteps = FMath::CeilToInt32(DtWithPause / MMaxDeltaTime);
 			if (NumSteps > MMaxSubSteps)
 			{
 				// Hitting this case means we're losing time, given the constraints of MaxSteps and MaxDt we can't
diff --git a/Engine/Source/Runtime/Experimental/Chaos/Private/Chaos/HeightField.cpp b/Engine/Source/Runtime/Experimental/Chaos/Private/Chaos/HeightField.cpp
--- a/Engine/Source/Runtime/Experimental/Chaos/Private/Chaos/HeightField.cpp
+++ b/Engine/Source/Runtime/Experimental/Chaos/Private/Chaos/HeightField.cpp
@@ -1226,7 +1226,7 @@ namespace Chaos
 							ExpandAxis = 0;
 						}
 						FReal ExpandSize = HalfExtents3D[ExpandAxis];
-						int32 Steps = FMath::TruncToInt(FMath::RoundFromZero(ExpandSize / FMath::Abs(GeomData.Scale[ExpandAxis])));
+						int32 Steps = FMath::TruncToInt32(FMath::RoundFromZero(ExpandSize / FMath::Abs(GeomData.Scale[ExpandAxis])));
 
 						Expand(Coord, ThickenDir, Steps);
 						Expand(Coord, -ThickenDir, Steps);
diff --git a/Engine/Source/Runtime/Experimental/Chaos/Private/Chaos/Island/IslandManager.cpp b/Engine/Source/Runtime/Experimental/Chaos/Private/Chaos/Island/IslandManager.cpp
--- a/Engine/Source/Runtime/Experimental/Chaos/Private/Chaos/Island/IslandManager.cpp
+++ b/Engine/Source/Runtime/Experimental/Chaos/Private/Chaos/Island/IslandManager.cpp
@@ -317,7 +317,7 @@ void FPBDIslandManager::InitializeGroups()
 	// @todo(chaos): is the number of worker threads a good indicator of how many threads we get in the solver loop? (Currently uses ParallelFor)
 	// Check for use of the "-onethread" command line arg, and physics threading disabled (GetNumWorkerThreads() is not affected by these)
 	const int32 NumWorkerThreads = (FApp::ShouldUseThreadingForPerformance() && !GSingleThreadedPhysics) ? FTaskGraphInterface::Get().GetNumWorkerThreads() : 0;
-	const int32 MaxIslandGroups = FMath::Max(1, FMath::CeilToInt(FReal(NumWorkerThreads) * GChaosSolverIslandGroupsMultiplier));
+	const int32 MaxIslandGroups = FMath::Max(1, FMath::CeilToInt32(FReal(NumWorkerThreads) * GChaosSolverIslandGroupsMultiplier));
 
 	IslandGroups.SetNum(MaxIslandGroups, false);
 	
diff --git a/Engine/Source/Runtime/Experimental/Chaos/Private/Chaos/PBDCollisionConstraintUtil.cpp b/Engine/Source/Runtime/Experimental/Chaos/Private/Chaos/PBDCollisionConstraintUtil.cpp
--- a/Engine/Source/Runtime/Experimental/Chaos/Private/Chaos/PBDCollisionConstraintUtil.cpp
+++ b/Engine/Source/Runtime/Experimental/Chaos/Private/Chaos/PBDCollisionConstraintUtil.cpp
@@ -35,15 +35,15 @@ namespace Chaos
 		{
 			PrincipalAxisLength = Extent.Z;
 		}
-		int32 NumberOfCells = FMath::CeilToInt(PrincipalAxisLength / CellSize);
+		int32 NumberOfCells = FMath::CeilToInt32(PrincipalAxisLength / CellSize);
 		check(NumberOfCells > 0);
 
 		CellSize = PrincipalAxisLength / (FReal)NumberOfCells;
 		FReal CellSizeInv = 1 / CellSize;
 
-		int32 NumberOfCellsX = FMath::CeilToInt(Extent.X * CellSizeInv) + 1;
-		int32 NumberOfCellsY = FMath::CeilToInt(Extent.Y * CellSizeInv) + 1;
-		int32 NumberOfCellsZ = FMath::CeilToInt(Extent.Z * CellSizeInv) + 1;
+		int32 NumberOfCellsX = FMath::CeilToInt32(Extent.X * CellSizeInv) + 1;
+		int32 NumberOfCellsY = FMath::CeilToInt32(Extent.Y * CellSizeInv) + 1;
+		int32 NumberOfCellsZ = FMath::CeilToInt32(Extent.Z * CellSizeInv) + 1;
 
 		// Create a Hash Table, but only store the buckets with constraint(s) as a map
 		// HashTableMap<BucketIdx, ConstraintIdx>
@@ -83,15 +83,15 @@ namespace Chaos
 		{
 			PrincipalAxisLength = Extent.Z;
 		}
-		int32 NumberOfCells = FMath::CeilToInt(PrincipalAxisLength / CellSize);
+		int32 NumberOfCells = FMath::CeilToInt32(PrincipalAxisLength / CellSize);
 		check(NumberOfCells > 0);
 
 		CellSize = PrincipalAxisLength / (FReal)NumberOfCells;
 		FReal CellSizeInv = 1 / CellSize;
 
-		int32 NumberOfCellsX = FMath::CeilToInt(Extent.X * CellSizeInv) + 1;
-		int32 NumberOfCellsY = FMath::CeilToInt(Extent.Y * CellSizeInv) + 1;
-		int32 NumberOfCellsZ = FMath::CeilToInt(Extent.Z * CellSizeInv) + 1;
+		int32 NumberOfCellsX = FMath::CeilToInt32(Extent.X * CellSizeInv) + 1;
+		int32 NumberOfCellsY = FMath::CeilToInt32(Extent.Y * CellSizeInv) + 1;
+		int32 NumberOfCellsZ = FMath::CeilToInt32(Extent.Z * CellSizeInv) + 1;
 
 		// Create a Hash Table, but only store the buckets with constraint(s) as a map
 		// HashTableMap<BucketIdx, ConstraintIdx>
@@ -131,15 +131,15 @@ namespace Chaos
 		{
 			PrincipalAxisLength = Extent.Z;
 		}
-		int32 NumberOfCells = FMath::CeilToInt(PrincipalAxisLength / CellSize);
+		int32 NumberOfCells = FMath::CeilToInt32(PrincipalAxisLength / CellSize);
 		check(NumberOfCells > 0);
 
 		CellSize = PrincipalAxisLength / (FReal)NumberOfCells;
 		FReal CellSizeInv = 1 / CellSize;
 
-		int32 NumberOfCellsX = FMath::CeilToInt(Extent.X * CellSizeInv) + 1;
-		int32 NumberOfCellsY = FMath::CeilToInt(Extent.Y * CellSizeInv) + 1;
-		int32 NumberOfCellsZ = FMath::CeilToInt(Extent.Z * CellSizeInv) + 1;
+		int32 NumberOfCellsX = FMath::CeilToInt32(Extent.X * CellSizeInv) + 1;
+		int32 NumberOfCellsY = FMath::CeilToInt32(Extent.Y * CellSizeInv) + 1;
+		int32 NumberOfCellsZ = FMath::CeilToInt32(Extent.Z * CellSizeInv) + 1;
 
 		// Create a Hash Table, but only store the buckets with constraint(s) as a map
 		// HashTableMap<BucketIdx, ConstraintIdx>
@@ -150,9 +150,9 @@ namespace Chaos
 		{
 			CellSize = PrincipalAxisLength / 1000;
 			CellSizeInv = 1 / CellSize;
-			NumberOfCellsX = FMath::CeilToInt(Extent.X * CellSizeInv) + 1;
-			NumberOfCellsY = FMath::CeilToInt(Extent.Y * CellSizeInv) + 1;
-			NumberOfCellsZ = FMath::CeilToInt(Extent.Z * CellSizeInv) + 1;
+			NumberOfCellsX = FMath::CeilToInt32(Extent.X * CellSizeInv) + 1;
+			NumberOfCellsY = FMath::CeilToInt32(Extent.Y * CellSizeInv) + 1;
+			NumberOfCellsZ = FMath::CeilToInt32(Extent.Z * CellSizeInv) + 1;
 			NumberOfCellsXY = NumberOfCellsX * NumberOfCellsY;
 			NumberOfCellsXYZ = NumberOfCellsXY * NumberOfCellsZ;
 		}
@@ -189,15 +189,15 @@ namespace Chaos
 		{
 			PrincipalAxisLength = Extent.Z;
 		}
-		int32 NumberOfCells = FMath::CeilToInt(PrincipalAxisLength / CellSize);
+		int32 NumberOfCells = FMath::CeilToInt32(PrincipalAxisLength / CellSize);
 		check(NumberOfCells > 0);
 
 		CellSize = PrincipalAxisLength / (FReal)NumberOfCells;
 		FReal CellSizeInv = 1 / CellSize;
 
-		int32 NumberOfCellsX = FMath::CeilToInt(Extent.X * CellSizeInv) + 1;
-		int32 NumberOfCellsY = FMath::CeilToInt(Extent.Y * CellSizeInv) + 1;
-		int32 NumberOfCellsZ = FMath::CeilToInt(Extent.Z * CellSizeInv) + 1;
+		int32 NumberOfCellsX = FMath::CeilToInt32(Extent.X * CellSizeInv) + 1;
+		int32 NumberOfCellsY = FMath::CeilToInt32(Extent.Y * CellSizeInv) + 1;
+		int32 NumberOfCellsZ = FMath::CeilToInt32(Extent.Z * CellSizeInv) + 1;
 
 		// Create a Hash Table, but only store the buckets with constraint(s) as a map
 		// HashTableMap<BucketIdx, ConstraintIdx>
@@ -237,15 +237,15 @@ namespace Chaos
 		{
 			PrincipalAxisLength = Extent.Z;
 		}
-		int32 NumberOfCells = FMath::CeilToInt(PrincipalAxisLength / CellSize);
+		int32 NumberOfCells = FMath::CeilToInt32(PrincipalAxisLength / CellSize);
 		check(NumberOfCells > 0);
 
 		CellSize = PrincipalAxisLength / (FReal)NumberOfCells;
 		FReal CellSizeInv = 1 / CellSize;
 
-		int32 NumberOfCellsX = FMath::CeilToInt(Extent.X * CellSizeInv) + 1;
-		int32 NumberOfCellsY = FMath::CeilToInt(Extent.Y * CellSizeInv) + 1;
-		int32 NumberOfCellsZ = FMath::CeilToInt(Extent.Z * CellSizeInv) + 1;
+		int32 NumberOfCellsX = FMath::CeilToInt32(Extent.X * CellSizeInv) + 1;
+		int32 NumberOfCellsY = FMath::CeilToInt32(Extent.Y * CellSizeInv) + 1;
+		int32 NumberOfCellsZ = FMath::CeilToInt32(Extent.Z * CellSizeInv) + 1;
 
 		// Create a Hash Table, but only store the buckets with constraint(s) as a map
 		// HashTableMap<BucketIdx, ConstraintIdx>
@@ -285,15 +285,15 @@ namespace Chaos
 		{
 			PrincipalAxisLength = Extent.Z;
 		}
-		int32 NumberOfCells = FMath::CeilToInt(PrincipalAxisLength / CellSize);
+		int32 NumberOfCells = FMath::CeilToInt32(PrincipalAxisLength / CellSize);
 		check(NumberOfCells > 0);
 
 		CellSize = PrincipalAxisLength / (FReal)NumberOfCells;
 		FReal CellSizeInv = 1 / CellSize;
 
-		int32 NumberOfCellsX = FMath::CeilToInt(Extent.X * CellSizeInv) + 1;
-		int32 NumberOfCellsY = FMath::CeilToInt(Extent.Y * CellSizeInv) + 1;
-		int32 NumberOfCellsZ = FMath::CeilToInt(Extent.Z * CellSizeInv) + 1;
+		int32 NumberOfCellsX = FMath::CeilToInt32(Extent.X * CellSizeInv) + 1;
+		int32 NumberOfCellsY = FMath::CeilToInt32(Extent.Y * CellSizeInv) + 1;
+		int32 NumberOfCellsZ = FMath::CeilToInt32(Extent.Z * CellSizeInv) + 1;
 
 		// Create a Hash Table, but only store the buckets with constraint(s) as a map
 		// HashTableMap<BucketIdx, ConstraintIdx>
diff --git a/Engine/Source/Runtime/Experimental/Chaos/Private/Chaos/PBDRigidsEvolutionGBF.cpp b/Engine/Source/Runtime/Experimental/Chaos/Private/Chaos/PBDRigidsEvolutionGBF.cpp
--- a/Engine/Source/Runtime/Experimental/Chaos/Private/Chaos/PBDRigidsEvolutionGBF.cpp
+++ b/Engine/Source/Runtime/Experimental/Chaos/Private/Chaos/PBDRigidsEvolutionGBF.cpp
@@ -163,7 +163,7 @@ void SerializeToDisk(TEvolution& Evolution)
 void FPBDRigidsEvolutionGBF::Advance(const FReal Dt,const FReal MaxStepDt,const int32 MaxSteps)
 {
 	// Determine how many steps we would like to take
-	int32 NumSteps = FMath::CeilToInt(Dt / MaxStepDt);
+	int32 NumSteps = FMath::CeilToInt32(Dt / MaxStepDt);
 	if (NumSteps > 0)
 	{
 		PrepareTick();
diff --git a/Engine/Source/Runtime/Experimental/Chaos/Private/Chaos/SpatialHash.cpp b/Engine/Source/Runtime/Experimental/Chaos/Private/Chaos/SpatialHash.cpp
--- a/Engine/Source/Runtime/Experimental/Chaos/Private/Chaos/SpatialHash.cpp
+++ b/Engine/Source/Runtime/Experimental/Chaos/Private/Chaos/SpatialHash.cpp
@@ -30,13 +30,13 @@ namespace Chaos
 		}
 
 		T PrincipalAxisLength = Extents[MBoundingBox.LargestAxis()];
-		int32 NumberOfCellsOnPrincipalAxis = FMath::CeilToInt(PrincipalAxisLength / MCellSize);
+		int32 NumberOfCellsOnPrincipalAxis = FMath::CeilToInt32(PrincipalAxisLength / MCellSize);
 		MCellSize = PrincipalAxisLength / (T)NumberOfCellsOnPrincipalAxis;
 		T CellSizeInv = (T)1.0 / MCellSize;
 
-		MNumberOfCellsX = FMath::CeilToInt(Extents[0] * CellSizeInv) + 1;
-		MNumberOfCellsY = FMath::CeilToInt(Extents[1] * CellSizeInv) + 1;
-		MNumberOfCellsZ = FMath::CeilToInt(Extents[2] * CellSizeInv) + 1;
+		MNumberOfCellsX = FMath::CeilToInt32(Extents[0] * CellSizeInv) + 1;
+		MNumberOfCellsY = FMath::CeilToInt32(Extents[1] * CellSizeInv) + 1;
+		MNumberOfCellsZ = FMath::CeilToInt32(Extents[2] * CellSizeInv) + 1;
 
 		for (int32 IdxParticle = 0; IdxParticle < MParticles.Num(); ++IdxParticle)
 		{
@@ -75,13 +75,13 @@ namespace Chaos
 		MCellSize = SmallestAxisLength / (T)10.0;
 
 		T PrincipalAxisLength = Extents[MBoundingBox.LargestAxis()];
-		int32 NumberOfCellsOnPrincipalAxis = FMath::CeilToInt(PrincipalAxisLength / MCellSize);
+		int32 NumberOfCellsOnPrincipalAxis = FMath::CeilToInt32(PrincipalAxisLength / MCellSize);
 		MCellSize = PrincipalAxisLength / (T)NumberOfCellsOnPrincipalAxis;
 		T CellSizeInv = (T)1.0 / MCellSize;
 
-		MNumberOfCellsX = FMath::CeilToInt(Extents[0] * CellSizeInv) + 1;
-		MNumberOfCellsY = FMath::CeilToInt(Extents[1] * CellSizeInv) + 1;
-		MNumberOfCellsZ = FMath::CeilToInt(Extents[2] * CellSizeInv) + 1;
+		MNumberOfCellsX = FMath::CeilToInt32(Extents[0] * CellSizeInv) + 1;
+		MNumberOfCellsY = FMath::CeilToInt32(Extents[1] * CellSizeInv) + 1;
+		MNumberOfCellsZ = FMath::CeilToInt32(Extents[2] * CellSizeInv) + 1;
 
 		for (int32 IdxParticle = 0; IdxParticle < MParticles.Num(); ++IdxParticle)
 		{
diff --git a/Engine/Source/Runtime/Experimental/Chaos/Private/PhysicsProxy/GeometryCollectionPhysicsProxy.cpp b/Engine/Source/Runtime/Experimental/Chaos/Private/PhysicsProxy/GeometryCollectionPhysicsProxy.cpp
--- a/Engine/Source/Runtime/Experimental/Chaos/Private/PhysicsProxy/GeometryCollectionPhysicsProxy.cpp
+++ b/Engine/Source/Runtime/Experimental/Chaos/Private/PhysicsProxy/GeometryCollectionPhysicsProxy.cpp
@@ -964,7 +964,7 @@ FGeometryCollectionPhysicsProxy::BuildClusters(
 	{
 		const Chaos::FReal NumCollisionParticles = static_cast<Chaos::FReal>(ClusterCreationParameters.CollisionParticles->Size());
 		const int32 ClampedCollisionParticlesSize = 
-			FMath::TruncToInt(Chaos::FReal(FMath::Max(0, FMath::Min(NumCollisionParticles * CollisionParticlesPerObjectFraction, NumCollisionParticles))));
+			FMath::TruncToInt32(Chaos::FReal(FMath::Max(0, FMath::Min(NumCollisionParticles * CollisionParticlesPerObjectFraction, NumCollisionParticles))));
 		ClusterCreationParameters.CollisionParticles->Resize(ClampedCollisionParticlesSize);
 	}
 	TArray<Chaos::TPBDRigidParticleHandle<Chaos::FReal, 3>*> ChildHandlesCopy(ChildHandles);
@@ -2437,8 +2437,8 @@ void FGeometryCollectionPhysicsProxy::InitializeSharedCollisionStructures(
 								MassSpaceParticles,
 								*UnionMesh,
 								InstanceBoundingBox,
-								FMath::FloorToInt(MinResolution),
-								FMath::FloorToInt(MaxResolution),
+								FMath::FloorToInt32(MinResolution),
+								FMath::FloorToInt32(MaxResolution),
 								SizeSpecificData.CollisionShapesData[0].CollisionObjectReductionPercentage,
 								SizeSpecificData.CollisionShapesData[0].CollisionType));
 						// Fall back on sphere if level set rasterization failed.
