do not create 0-vertex convex shapes for geometry collections

diff --git a/Engine/Source/Runtime/Experimental/Chaos/Private/GeometryCollection/GeometryCollectionConvexUtility.cpp b/Engine/Source/Runtime/Experimental/Chaos/Private/GeometryCollection/GeometryCollectionConvexUtility.cpp
--- a/Engine/Source/Runtime/Experimental/Chaos/Private/GeometryCollection/GeometryCollectionConvexUtility.cpp
+++ b/Engine/Source/Runtime/Experimental/Chaos/Private/GeometryCollection/GeometryCollectionConvexUtility.cpp
@@ -1162,11 +1162,19 @@ FGeometryCollectionConvexUtility::FGeometryCollectionConvexData FGeometryCollect
 	GeometryCollection->Resize(Convexes.Num(), "Convex");
 	ConvexHull = MoveTemp(Convexes);
 
-	// clear all null hulls
+	// clear all null and empty hulls
 	TArray<int32> EmptyConvex;
 	for (int32 ConvexIdx = 0; ConvexIdx < ConvexHull.Num(); ConvexIdx++)
 	{
-		if (!ConvexHull[ConvexIdx].IsValid())
+		if (ConvexHull[ConvexIdx].IsValid())
+		{
+			if (ConvexHull[ConvexIdx]->NumVertices() == 0)
+			{
+				ConvexHull[ConvexIdx].Reset();
+				EmptyConvex.Add(ConvexIdx);
+			}
+		}
+		else // (!ConvexHull[ConvexIdx].IsValid())
 		{
 			EmptyConvex.Add(ConvexIdx);
 		}
