Chaos Cloth Optimization - Optimize FTriangleMesh::GetPointNormals

Updated ispc code on behalf of jeff.rous. He reports that this ispc change is 2.5x faster than before. At the same time he was doing that, I was also trying to optimize things and got the non-ispc path to be about 2x faster with this small change.

diff --git a/Engine/Source/Runtime/Core/Public/Math/Vector.isph b/Engine/Source/Runtime/Core/Public/Math/Vector.isph
--- a/Engine/Source/Runtime/Core/Public/Math/Vector.isph
+++ b/Engine/Source/Runtime/Core/Public/Math/Vector.isph
@@ -1529,9 +1529,15 @@ inline FVector3f operator*(const float F, const FVector3f &A)
 	return Result;
 }
 
-inline uniform FVector operator*(const uniform FVector &A, const uniform FReal F)
+inline uniform FVector3d operator*(const uniform FVector3d &A, const uniform double F)
 {
-	const uniform FVector FVec = {{F,F,F}};
+	const uniform FVector3d FVec = {{F,F,F}};
+	return A * FVec;
+}
+
+inline uniform FVector3f operator*(const uniform FVector3f &A, const uniform float F)
+{
+	const uniform FVector3f FVec = {{F,F,F}};
 	return A * FVec;
 }
 
@@ -3802,20 +3808,37 @@ inline float VectorSize(const varying FVector3f &A)
 	return sqrt(VectorDot(A, A));
 }
 
-inline uniform FVector VectorGetSafeNormal(const uniform FVector &Vector)
+inline uniform FVector3d VectorGetSafeNormal(const uniform FVector3d &Vector)
 {
-	const uniform FReal SquareSum = VectorSizeSquared(Vector);
+	const uniform double SquareSum = VectorSizeSquared(Vector);
 
 	// Not sure if it's safe to add tolerance in there. Might introduce too many errors
-	if(SquareSum == 1.f)
+	if(SquareSum == 1.)
 	{
 		return Vector;
 	}
-	else if(SquareSum < SMALL_NUMBER)
+	else if(SquareSum < DOUBLE_SMALL_NUMBER)
 	{
 		return ZeroVector;
 	}
-	const uniform FReal Scale = InvSqrt(SquareSum);
+	const uniform double Scale = InvSqrt(SquareSum);
+	return Vector * Scale;
+}
+
+inline uniform FVector3f VectorGetSafeNormal(const uniform FVector3f &Vector)
+{
+	const uniform float SquareSum = VectorSizeSquared(Vector);
+
+	// Not sure if it's safe to add tolerance in there. Might introduce too many errors
+	if(SquareSum == 1.f)
+	{
+		return Vector;
+	}
+	else if(SquareSum < FLOAT_SMALL_NUMBER)
+	{
+		return FloatZeroVector;
+	}
+	const uniform float Scale = InvSqrt(SquareSum);
 	return Vector * Scale;
 }
 
diff --git a/Engine/Source/Runtime/Experimental/Chaos/Private/Chaos/TriangleMesh.cpp b/Engine/Source/Runtime/Experimental/Chaos/Private/Chaos/TriangleMesh.cpp
--- a/Engine/Source/Runtime/Experimental/Chaos/Private/Chaos/TriangleMesh.cpp
+++ b/Engine/Source/Runtime/Experimental/Chaos/Private/Chaos/TriangleMesh.cpp
@@ -256,6 +256,7 @@ TArray<TVec4<int32>> FTriangleMesh::GetUniqueAdjacentElements() const
 template <typename T>
 void FTriangleMesh::GetFaceNormals(TArray<TVec3<T>>& Normals, const TConstArrayView<TVec3<T>>& Points, const bool ReturnEmptyOnError) const
 {
+	TRACE_CPUPROFILER_EVENT_SCOPE(FTriangleMesh_GetFaceNormals);
 	Normals.Reset(MElements.Num());
 	if (ReturnEmptyOnError)
 	{
@@ -339,12 +340,11 @@ template <typename T>
 void FTriangleMesh::GetPointNormals(TArrayView<TVec3<T>> PointNormals, const TConstArrayView<TVec3<T>>& FaceNormals, const bool bUseGlobalArray) const
 {
 	check(MPointToTriangleMap.Num() != 0);
-
+	TRACE_CPUPROFILER_EVENT_SCOPE(FTriangleMesh_GetPointNormals);
 	for (int32 Element = 0; Element < MNumIndices; ++Element)  // Iterate points with local indexes
 	{
 		const int32 NormalIndex = bUseGlobalArray ? LocalToGlobal(Element) : Element;  // Select whether the points normal indices match the points indices or start at 0
-		TVec3<T>& Normal = PointNormals[NormalIndex];
-		Normal = TVec3<T>(0.);
+		TVec3<T> Normal((T)0);
 		const TArray<int32>& TriangleMap = MPointToTriangleMap[Element];  // Access MPointToTriangleMap with local index
 		for (int32 k = 0; k < TriangleMap.Num(); ++k)
 		{
@@ -353,7 +353,7 @@ void FTriangleMesh::GetPointNormals(TArrayView<TVec3<T>> PointNormals, const TCo
 				Normal += FaceNormals[TriangleMap[k]];
 			}
 		}
-		Normal = Normal.GetSafeNormal();
+		PointNormals[NormalIndex] = Normal.GetSafeNormal();
 	}
 }
 template CHAOS_API void FTriangleMesh::GetPointNormals<FRealDouble>(TArrayView<TVec3<FRealDouble>>, const TConstArrayView<TVec3<FRealDouble>>&, const bool) const;
@@ -380,8 +380,7 @@ CHAOS_API void FTriangleMesh::GetPointNormals<FRealSingle>(TArrayView<TVec3<FRea
 		for (int32 Element = 0; Element < MNumIndices; ++Element)  // Iterate points with local indexes
 		{
 			const int32 NormalIndex = bUseGlobalArray ? LocalToGlobal(Element) : Element;  // Select whether the points normal indices match the points indices or start at 0
-			TVec3<FRealSingle>& Normal = PointNormals[NormalIndex];
-			Normal = TVec3<FRealSingle>(0.);
+			TVec3<FRealSingle> Normal(0.f);
 			const TArray<int32>& TriangleMap = MPointToTriangleMap[Element];  // Access MPointToTriangleMap with local index
 			for (int32 k = 0; k < TriangleMap.Num(); ++k)
 			{
@@ -390,7 +389,7 @@ CHAOS_API void FTriangleMesh::GetPointNormals<FRealSingle>(TArrayView<TVec3<FRea
 					Normal += FaceNormals[TriangleMap[k]];
 				}
 			}
-			Normal = Normal.GetSafeNormal();
+			PointNormals[NormalIndex] = Normal.GetSafeNormal();
 		}
 	}
 }
diff --git a/Engine/Source/Runtime/Experimental/Chaos/Private/Chaos/TriangleMesh.ispc b/Engine/Source/Runtime/Experimental/Chaos/Private/Chaos/TriangleMesh.ispc
--- a/Engine/Source/Runtime/Experimental/Chaos/Private/Chaos/TriangleMesh.ispc
+++ b/Engine/Source/Runtime/Experimental/Chaos/Private/Chaos/TriangleMesh.ispc
@@ -39,24 +39,27 @@ export void GetPointNormals(uniform FVector3f PointNormals[],
 							const uniform int32 NumFaceNormals,
 							const uniform int32 NumIndices)
 {
-	foreach(Element = 0 ... NumIndices)
+	#pragma unroll 4
+	for(uniform int32 Element = 0; Element < NumIndices; Element++)
 	{
-		varying FVector3f Normal = FloatZeroVector;
+		uniform FVector4f Normal = FloatZero;
 
-		#pragma ignore warning(perf)
-		const varying TArrayInt TriangleMap = PointToTriangleMap[Element];  // Access MPointToTriangleMap with local index
+		const uniform TArrayInt TriangleMap = PointToTriangleMap[Element];  // Access MPointToTriangleMap with local index
 
-		for (int32 k = 0; k < TriangleMap.ArrayNum; ++k)
+		#pragma unroll 2
+		for (uniform int32 k = 0; k < TriangleMap.ArrayNum; ++k)
 		{
-			#pragma ignore warning(perf)
-			const int32 Index = TriangleMap.Data[k];
+			const uniform int32 Index = TriangleMap.Data[k];
 
 			if (Index >= 0 && Index < NumFaceNormals)
 			{
-				Normal = Normal + VectorGather(&FaceNormals[Index]);
+				const uniform FVector3f FaceNormal = FaceNormals[Index];
+
+				Normal = Normal + SetVector4(FaceNormal, FLOAT_ZERO);
 			}
 		}
 
-		VectorStore(&PointNormals[extract(Element, 0) + StartIndex], VectorGetSafeNormal(Normal));
+		uniform FVector3f SafeNormal = VectorGetSafeNormal(SetVector(Normal));
+		PointNormals[Element + StartIndex] = SafeNormal;
 	}
 }
