Interchange Framework - Adding UInterChangeFactoryBaseNode class with "subpath" attribute, allowing to customize the imported package locations. Also removing unnecessary constructors.

diff --git a/Engine/Plugins/Experimental/Interchange/Runtime/Source/FactoryNodes/Public/InterchangeActorFactoryNode.h b/Engine/Plugins/Experimental/Interchange/Runtime/Source/FactoryNodes/Public/InterchangeActorFactoryNode.h
--- a/Engine/Plugins/Experimental/Interchange/Runtime/Source/FactoryNodes/Public/InterchangeActorFactoryNode.h
+++ b/Engine/Plugins/Experimental/Interchange/Runtime/Source/FactoryNodes/Public/InterchangeActorFactoryNode.h
@@ -2,12 +2,12 @@
 
 #pragma once
 
-#include "Nodes/InterchangeBaseNode.h"
+#include "Nodes/InterchangeFactoryBaseNode.h"
 
 #include "InterchangeActorFactoryNode.generated.h"
 
 UCLASS(BlueprintType, Experimental)
-class INTERCHANGEFACTORYNODES_API UInterchangeActorFactoryNode : public UInterchangeBaseNode
+class INTERCHANGEFACTORYNODES_API UInterchangeActorFactoryNode : public UInterchangeFactoryBaseNode
 {
 	GENERATED_BODY()
 
diff --git a/Engine/Plugins/Experimental/Interchange/Runtime/Source/FactoryNodes/Public/InterchangeMaterialFactoryNode.h b/Engine/Plugins/Experimental/Interchange/Runtime/Source/FactoryNodes/Public/InterchangeMaterialFactoryNode.h
--- a/Engine/Plugins/Experimental/Interchange/Runtime/Source/FactoryNodes/Public/InterchangeMaterialFactoryNode.h
+++ b/Engine/Plugins/Experimental/Interchange/Runtime/Source/FactoryNodes/Public/InterchangeMaterialFactoryNode.h
@@ -3,7 +3,7 @@
 #pragma once
 
 #include "CoreMinimal.h"
-#include "Nodes/InterchangeBaseNode.h"
+#include "Nodes/InterchangeFactoryBaseNode.h"
 #include "InterchangeMaterialDefinitions.h"
 #include "InterchangeShaderGraphNode.h"
 
@@ -14,7 +14,7 @@
 #include "InterchangeMaterialFactoryNode.generated.h"
 
 UCLASS(Abstract, Experimental)
-class INTERCHANGEFACTORYNODES_API UInterchangeBaseMaterialFactoryNode : public UInterchangeBaseNode
+class INTERCHANGEFACTORYNODES_API UInterchangeBaseMaterialFactoryNode : public UInterchangeFactoryBaseNode
 {
 	GENERATED_BODY()
 
diff --git a/Engine/Plugins/Experimental/Interchange/Runtime/Source/FactoryNodes/Public/InterchangePhysicsAssetFactoryNode.h b/Engine/Plugins/Experimental/Interchange/Runtime/Source/FactoryNodes/Public/InterchangePhysicsAssetFactoryNode.h
--- a/Engine/Plugins/Experimental/Interchange/Runtime/Source/FactoryNodes/Public/InterchangePhysicsAssetFactoryNode.h
+++ b/Engine/Plugins/Experimental/Interchange/Runtime/Source/FactoryNodes/Public/InterchangePhysicsAssetFactoryNode.h
@@ -3,7 +3,7 @@
 #pragma once
 
 #include "CoreMinimal.h"
-#include "Nodes/InterchangeBaseNode.h"
+#include "Nodes/InterchangeFactoryBaseNode.h"
 
 #if WITH_ENGINE
 #include "Animation/Skeleton.h"
@@ -13,19 +13,11 @@
 #include "InterchangePhysicsAssetFactoryNode.generated.h"
 
 UCLASS(BlueprintType, Experimental)
-class INTERCHANGEFACTORYNODES_API UInterchangePhysicsAssetFactoryNode : public UInterchangeBaseNode
+class INTERCHANGEFACTORYNODES_API UInterchangePhysicsAssetFactoryNode : public UInterchangeFactoryBaseNode
 {
 	GENERATED_BODY()
 
 public:
-	UInterchangePhysicsAssetFactoryNode()
-	:UInterchangeBaseNode()
-	{
-#if WITH_ENGINE
-		AssetClass = nullptr;
-#endif
-	}
-
 	/**
 	 * Initialize node data
 	 * @param: UniqueID - The uniqueId for this node
diff --git a/Engine/Plugins/Experimental/Interchange/Runtime/Source/FactoryNodes/Public/InterchangeSkeletalMeshFactoryNode.h b/Engine/Plugins/Experimental/Interchange/Runtime/Source/FactoryNodes/Public/InterchangeSkeletalMeshFactoryNode.h
--- a/Engine/Plugins/Experimental/Interchange/Runtime/Source/FactoryNodes/Public/InterchangeSkeletalMeshFactoryNode.h
+++ b/Engine/Plugins/Experimental/Interchange/Runtime/Source/FactoryNodes/Public/InterchangeSkeletalMeshFactoryNode.h
@@ -3,7 +3,7 @@
 #pragma once
 
 #include "CoreMinimal.h"
-#include "Nodes/InterchangeBaseNode.h"
+#include "Nodes/InterchangeFactoryBaseNode.h"
 
 #if WITH_ENGINE
 #include "Engine/SkeletalMesh.h"
@@ -26,7 +26,7 @@ namespace UE
 }//ns UE
 
 UCLASS(BlueprintType, Experimental)
-class INTERCHANGEFACTORYNODES_API UInterchangeSkeletalMeshFactoryNode : public UInterchangeBaseNode
+class INTERCHANGEFACTORYNODES_API UInterchangeSkeletalMeshFactoryNode : public UInterchangeFactoryBaseNode
 {
 	GENERATED_BODY()
 
diff --git a/Engine/Plugins/Experimental/Interchange/Runtime/Source/FactoryNodes/Public/InterchangeSkeletonFactoryNode.h b/Engine/Plugins/Experimental/Interchange/Runtime/Source/FactoryNodes/Public/InterchangeSkeletonFactoryNode.h
--- a/Engine/Plugins/Experimental/Interchange/Runtime/Source/FactoryNodes/Public/InterchangeSkeletonFactoryNode.h
+++ b/Engine/Plugins/Experimental/Interchange/Runtime/Source/FactoryNodes/Public/InterchangeSkeletonFactoryNode.h
@@ -3,7 +3,7 @@
 #pragma once
 
 #include "CoreMinimal.h"
-#include "Nodes/InterchangeBaseNode.h"
+#include "Nodes/InterchangeFactoryBaseNode.h"
 
 #if WITH_ENGINE
 #include "Animation/Skeleton.h"
@@ -12,19 +12,11 @@
 #include "InterchangeSkeletonFactoryNode.generated.h"
 
 UCLASS(BlueprintType, Experimental)
-class INTERCHANGEFACTORYNODES_API UInterchangeSkeletonFactoryNode : public UInterchangeBaseNode
+class INTERCHANGEFACTORYNODES_API UInterchangeSkeletonFactoryNode : public UInterchangeFactoryBaseNode
 {
 	GENERATED_BODY()
 
 public:
-	UInterchangeSkeletonFactoryNode()
-	:UInterchangeBaseNode()
-	{
-#if WITH_ENGINE
-		AssetClass = nullptr;
-#endif
-	}
-
 	/**
 	 * Initialize node data
 	 * @param: UniqueID - The uniqueId for this node
diff --git a/Engine/Plugins/Experimental/Interchange/Runtime/Source/FactoryNodes/Public/InterchangeStaticMeshFactoryNode.h b/Engine/Plugins/Experimental/Interchange/Runtime/Source/FactoryNodes/Public/InterchangeStaticMeshFactoryNode.h
--- a/Engine/Plugins/Experimental/Interchange/Runtime/Source/FactoryNodes/Public/InterchangeStaticMeshFactoryNode.h
+++ b/Engine/Plugins/Experimental/Interchange/Runtime/Source/FactoryNodes/Public/InterchangeStaticMeshFactoryNode.h
@@ -3,7 +3,7 @@
 #pragma once
 
 #include "CoreMinimal.h"
-#include "Nodes/InterchangeBaseNode.h"
+#include "Nodes/InterchangeFactoryBaseNode.h"
 
 #if WITH_ENGINE
 #include "Engine/StaticMesh.h"
@@ -25,7 +25,7 @@ namespace UE
 
 
 UCLASS(BlueprintType, Experimental)
-class INTERCHANGEFACTORYNODES_API UInterchangeStaticMeshFactoryNode : public UInterchangeBaseNode
+class INTERCHANGEFACTORYNODES_API UInterchangeStaticMeshFactoryNode : public UInterchangeFactoryBaseNode
 {
 	GENERATED_BODY()
 
diff --git a/Engine/Plugins/Experimental/Interchange/Runtime/Source/FactoryNodes/Public/InterchangeTextureFactoryNode.h b/Engine/Plugins/Experimental/Interchange/Runtime/Source/FactoryNodes/Public/InterchangeTextureFactoryNode.h
--- a/Engine/Plugins/Experimental/Interchange/Runtime/Source/FactoryNodes/Public/InterchangeTextureFactoryNode.h
+++ b/Engine/Plugins/Experimental/Interchange/Runtime/Source/FactoryNodes/Public/InterchangeTextureFactoryNode.h
@@ -3,7 +3,7 @@
 #pragma once
 
 #include "CoreMinimal.h"
-#include "Nodes/InterchangeBaseNode.h"
+#include "Nodes/InterchangeFactoryBaseNode.h"
 
 #if WITH_ENGINE
 #include "Engine/Texture.h"
@@ -14,19 +14,11 @@
 
 
 UCLASS(BlueprintType, Experimental)
-class INTERCHANGEFACTORYNODES_API UInterchangeTextureFactoryNode : public UInterchangeBaseNode
+class INTERCHANGEFACTORYNODES_API UInterchangeTextureFactoryNode : public UInterchangeFactoryBaseNode
 {
 	GENERATED_BODY()
 
 public:
-	UInterchangeTextureFactoryNode()
-		:UInterchangeBaseNode()
-	{
-#if WITH_ENGINE
-		AssetClass = nullptr;
-#endif
-	}
-
 	/**
 	 * Initialize node data
 	 * @param: UniqueID - The uniqueId for this node
diff --git a/Engine/Source/Runtime/Experimental/Interchange/Core/Private/Nodes/InterchangeFactoryBaseNode.cpp b/Engine/Source/Runtime/Experimental/Interchange/Core/Private/Nodes/InterchangeFactoryBaseNode.cpp
new file mode 100644
--- /dev/null
+++ b/Engine/Source/Runtime/Experimental/Interchange/Core/Private/Nodes/InterchangeFactoryBaseNode.cpp
@@ -0,0 +1,33 @@
+// Copyright Epic Games, Inc. All Rights Reserved.
+#include "Nodes/InterchangeFactoryBaseNode.h"
+
+#include "CoreMinimal.h"
+#include "Types/AttributeStorage.h"
+#include "UObject/Class.h"
+#include "UObject/Object.h"
+#include "UObject/ObjectMacros.h"
+
+FString UInterchangeFactoryBaseNode::GetKeyDisplayName(const UE::Interchange::FAttributeKey& NodeAttributeKey) const
+{
+	FString KeyDisplayName = NodeAttributeKey.ToString();
+	if (NodeAttributeKey == Macro_CustomSubPathKey)
+	{
+		KeyDisplayName = TEXT("Import Sub-Path");
+	}
+	else
+	{
+		KeyDisplayName = Super::GetKeyDisplayName(NodeAttributeKey);
+	}
+
+	return KeyDisplayName;
+}
+
+bool UInterchangeFactoryBaseNode::GetCustomSubPath(FString& AttributeValue) const
+{
+	IMPLEMENT_NODE_ATTRIBUTE_GETTER(SubPath, FString)
+}
+
+bool UInterchangeFactoryBaseNode::SetCustomSubPath(const FString& AttributeValue)
+{
+	IMPLEMENT_NODE_ATTRIBUTE_SETTER_NODELEGATE(SubPath, FString)
+}
diff --git a/Engine/Source/Runtime/Experimental/Interchange/Core/Public/Nodes/InterchangeFactoryBaseNode.h b/Engine/Source/Runtime/Experimental/Interchange/Core/Public/Nodes/InterchangeFactoryBaseNode.h
new file mode 100644
--- /dev/null
+++ b/Engine/Source/Runtime/Experimental/Interchange/Core/Public/Nodes/InterchangeFactoryBaseNode.h
@@ -0,0 +1,44 @@
+// Copyright Epic Games, Inc. All Rights Reserved.
+
+#pragma once
+
+#include "CoreMinimal.h"
+#include "Nodes/InterchangeBaseNode.h"
+#include "Nodes/InterchangeBaseNodeUtilities.h"
+#include "Types/AttributeStorage.h"
+#include "UObject/Class.h"
+#include "UObject/Object.h"
+#include "UObject/ObjectMacros.h"
+
+#include "InterchangeFactoryBaseNode.generated.h"
+
+/**
+ * This struct is used to store and retrieve key value attributes. The attributes are store in a generic FAttributeStorage which serialize the value in a TArray64<uint8>
+ * See UE::Interchange::EAttributeTypes to know the supported template types
+ * This is an abstract class. This is the base class of the interchange node graph format, all class in this format should derive from this class
+ */
+UCLASS(BlueprintType, Experimental)
+class INTERCHANGECORE_API UInterchangeFactoryBaseNode : public UInterchangeBaseNode
+{
+	GENERATED_BODY()
+
+public:
+	virtual FString GetKeyDisplayName(const UE::Interchange::FAttributeKey& NodeAttributeKey) const override;
+
+	/**
+	 * Return the custom sub-path under PackageBasePath, where the assets will be created.
+	 */
+	UFUNCTION(BlueprintCallable, Category = "Interchange | Node")
+	bool GetCustomSubPath(FString& AttributeValue) const;
+
+	/**
+	 * Set the custom sub-path under PackageBasePath where the assets will be created.
+	 */
+	UFUNCTION(BlueprintCallable, Category = "Interchange | Node")
+	bool SetCustomSubPath(const FString& AttributeValue);
+
+private:
+
+	const UE::Interchange::FAttributeKey Macro_CustomSubPathKey = UE::Interchange::FAttributeKey(TEXT("SubPath"));
+
+};
\ No newline at end of file
diff --git a/Engine/Source/Runtime/Experimental/Interchange/Engine/Private/Tasks/InterchangeTaskCreateAsset.cpp b/Engine/Source/Runtime/Experimental/Interchange/Engine/Private/Tasks/InterchangeTaskCreateAsset.cpp
--- a/Engine/Source/Runtime/Experimental/Interchange/Engine/Private/Tasks/InterchangeTaskCreateAsset.cpp
+++ b/Engine/Source/Runtime/Experimental/Interchange/Engine/Private/Tasks/InterchangeTaskCreateAsset.cpp
@@ -19,6 +19,7 @@
 #include "UObject/ObjectMacros.h"
 #include "UObject/WeakObjectPtrTemplates.h"
 #include "Nodes/InterchangeBaseNode.h"
+#include "Nodes/InterchangeFactoryBaseNode.h"
 
 namespace UE
 {
@@ -39,7 +40,16 @@ namespace UE
 				FString SanitazedPackageBasePath = PackageBasePath;
 				SanitizeObjectPath(SanitazedPackageBasePath);
 
-				OutPackageName = FPaths::Combine(*SanitazedPackageBasePath, *OutAssetName);
+				FString SubPath;
+				if (const UInterchangeFactoryBaseNode* FactoryBaseNode = Cast<UInterchangeFactoryBaseNode>(Node))
+				{
+					if(FactoryBaseNode->GetCustomSubPath(SubPath))
+					{
+						SanitizeObjectPath(SubPath);
+					}
+				}
+
+				OutPackageName = FPaths::Combine(*SanitazedPackageBasePath, *SubPath, *OutAssetName);
 			}
 
 			UObject* GetExistingObjectFromAssetImportData(TSharedPtr<UE::Interchange::FImportAsyncHelper, ESPMode::ThreadSafe> AsyncHelper, UInterchangeBaseNode* Node)
