Use a union-based representation for types which want to treat its members as both individual named fields and as an array, to avoid strict aliasing issues.

diff --git a/Engine/Source/Runtime/Core/Public/Math/IntPoint.h b/Engine/Source/Runtime/Core/Public/Math/IntPoint.h
--- a/Engine/Source/Runtime/Core/Public/Math/IntPoint.h
+++ b/Engine/Source/Runtime/Core/Public/Math/IntPoint.h
@@ -16,11 +16,20 @@
  */
 struct FIntPoint
 {
-	/** Holds the point's x-coordinate. */
-	int32 X;
-	
-	/** Holds the point's y-coordinate. */
-	int32 Y;
+	union
+	{
+		struct
+		{
+			/** Holds the point's x-coordinate. */
+			int32 X;
+
+			/** Holds the point's y-coordinate. */
+			int32 Y;
+		};
+
+		UE_DEPRECATED(all, "For internal use only")
+		int32 XY[2];
+	};
 
 public:
 
@@ -373,13 +382,17 @@ FORCEINLINE FIntPoint::FIntPoint(EForceInit)
 
 FORCEINLINE const int32& FIntPoint::operator()(int32 PointIndex) const
 {
-	return (&X)[PointIndex];
+PRAGMA_DISABLE_DEPRECATION_WARNINGS
+	return XY[PointIndex];
+PRAGMA_ENABLE_DEPRECATION_WARNINGS
 }
 
 
 FORCEINLINE int32& FIntPoint::operator()(int32 PointIndex)
 {
-	return (&X)[PointIndex];
+PRAGMA_DISABLE_DEPRECATION_WARNINGS
+	return XY[PointIndex];
+PRAGMA_ENABLE_DEPRECATION_WARNINGS
 }
 
 
diff --git a/Engine/Source/Runtime/Core/Public/Math/IntVector.h b/Engine/Source/Runtime/Core/Public/Math/IntVector.h
--- a/Engine/Source/Runtime/Core/Public/Math/IntVector.h
+++ b/Engine/Source/Runtime/Core/Public/Math/IntVector.h
@@ -13,14 +13,23 @@
  */
 struct FIntVector
 {
-	/** Holds the point's x-coordinate. */
-	int32 X;
-	
-	/** Holds the point's y-coordinate. */
-	int32 Y;
+	union
+	{
+		struct
+		{
+			/** Holds the point's x-coordinate. */
+			int32 X;
+
+			/** Holds the point's y-coordinate. */
+			int32 Y;
+
+			/**  Holds the point's z-coordinate. */
+			int32 Z;
+		};
 
-	/**  Holds the point's z-coordinate. */
-	int32 Z;
+		UE_DEPRECATED(all, "For internal use only")
+		int32 XYZ[3];
+	};
 
 public:
 
@@ -344,25 +353,33 @@ FORCEINLINE FIntVector::FIntVector( EForceInit )
 
 FORCEINLINE const int32& FIntVector::operator()( int32 ComponentIndex ) const
 {
-	return (&X)[ComponentIndex];
+PRAGMA_DISABLE_DEPRECATION_WARNINGS
+	return XYZ[ComponentIndex];
+PRAGMA_ENABLE_DEPRECATION_WARNINGS
 }
 
 
 FORCEINLINE int32& FIntVector::operator()( int32 ComponentIndex )
 {
-	return (&X)[ComponentIndex];
+PRAGMA_DISABLE_DEPRECATION_WARNINGS
+	return XYZ[ComponentIndex];
+PRAGMA_ENABLE_DEPRECATION_WARNINGS
 }
 
 
 FORCEINLINE const int32& FIntVector::operator[]( int32 ComponentIndex ) const
 {
-	return (&X)[ComponentIndex];
+PRAGMA_DISABLE_DEPRECATION_WARNINGS
+	return XYZ[ComponentIndex];
+PRAGMA_ENABLE_DEPRECATION_WARNINGS
 }
 
 
 FORCEINLINE int32& FIntVector::operator[]( int32 ComponentIndex )
 {
-	return (&X)[ComponentIndex];
+PRAGMA_DISABLE_DEPRECATION_WARNINGS
+	return XYZ[ComponentIndex];
+PRAGMA_ENABLE_DEPRECATION_WARNINGS
 }
 
 FORCEINLINE bool FIntVector::operator==( const FIntVector& Other ) const
@@ -529,7 +546,16 @@ FORCEINLINE uint32 GetTypeHash(const FIntVector& Vector)
 
 struct FIntVector2
 {
-	int32 X, Y;
+	union
+	{
+		struct
+		{
+			int32 X, Y;
+		};
+
+		UE_DEPRECATED(all, "For internal use only")
+		int32 XY[2];
+	};
 
 	FORCEINLINE FIntVector2()
 	{
@@ -555,13 +581,17 @@ struct FIntVector2
 
 	FORCEINLINE const int32& operator[](int32 ComponentIndex) const
 	{
-		return (&X)[ComponentIndex];
+PRAGMA_DISABLE_DEPRECATION_WARNINGS
+		return XY[ComponentIndex];
+PRAGMA_ENABLE_DEPRECATION_WARNINGS
 	}
 
 
 	FORCEINLINE int32& operator[](int32 ComponentIndex)
 	{
-		return (&X)[ComponentIndex];
+PRAGMA_DISABLE_DEPRECATION_WARNINGS
+		return XY[ComponentIndex];
+PRAGMA_ENABLE_DEPRECATION_WARNINGS
 	}
 
 	FORCEINLINE bool operator==(const FIntVector2& Other) const
@@ -578,7 +608,16 @@ struct FIntVector2
 
 struct FUintVector2
 {
-	uint32 X, Y;
+	union
+	{
+		struct
+		{
+			uint32 X, Y;
+		};
+
+		UE_DEPRECATED(all, "For internal use only")
+		uint32 XY[2];
+	};
 
 	FORCEINLINE FUintVector2()
 	{
@@ -604,13 +643,17 @@ struct FUintVector2
 
 	FORCEINLINE const uint32& operator[](int32 ComponentIndex) const
 	{
-		return (&X)[ComponentIndex];
+PRAGMA_DISABLE_DEPRECATION_WARNINGS
+		return XY[ComponentIndex];
+PRAGMA_ENABLE_DEPRECATION_WARNINGS
 	}
 
 
 	FORCEINLINE uint32& operator[](int32 ComponentIndex)
 	{
-		return (&X)[ComponentIndex];
+PRAGMA_DISABLE_DEPRECATION_WARNINGS
+		return XY[ComponentIndex];
+PRAGMA_ENABLE_DEPRECATION_WARNINGS
 	}
 
 	FORCEINLINE bool operator==(const FUintVector2& Other) const
@@ -627,7 +670,16 @@ struct FUintVector2
 
 struct FIntVector4
 {
-	int32 X, Y, Z, W;
+	union
+	{
+		struct
+		{
+			int32 X, Y, Z, W;
+		};
+
+		UE_DEPRECATED(all, "For internal use only")
+		int32 XYZW[4];
+	};
 
 	FORCEINLINE FIntVector4()
 	{
@@ -828,25 +880,33 @@ struct FIntVector4
 
 FORCEINLINE const int32& FIntVector4::operator()(int32 ComponentIndex) const
 {
-	return (&X)[ComponentIndex];
+PRAGMA_DISABLE_DEPRECATION_WARNINGS
+	return XYZW[ComponentIndex];
+PRAGMA_ENABLE_DEPRECATION_WARNINGS
 }
 
 
 FORCEINLINE int32& FIntVector4::operator()(int32 ComponentIndex)
 {
-	return (&X)[ComponentIndex];
+PRAGMA_DISABLE_DEPRECATION_WARNINGS
+	return XYZW[ComponentIndex];
+PRAGMA_ENABLE_DEPRECATION_WARNINGS
 }
 
 
 FORCEINLINE const int32& FIntVector4::operator[](int32 ComponentIndex) const
 {
-	return (&X)[ComponentIndex];
+PRAGMA_DISABLE_DEPRECATION_WARNINGS
+	return XYZW[ComponentIndex];
+PRAGMA_ENABLE_DEPRECATION_WARNINGS
 }
 
 
 FORCEINLINE int32& FIntVector4::operator[](int32 ComponentIndex)
 {
-	return (&X)[ComponentIndex];
+PRAGMA_DISABLE_DEPRECATION_WARNINGS
+	return XYZW[ComponentIndex];
+PRAGMA_ENABLE_DEPRECATION_WARNINGS
 }
 
 FORCEINLINE bool FIntVector4::operator==(const FIntVector4& Other) const
@@ -966,7 +1026,16 @@ FORCEINLINE FIntVector4 FIntVector4::operator^(int32 Value) const
 
 struct FUintVector4
 {
-	uint32 X, Y, Z, W;
+	union
+	{
+		struct
+		{
+			uint32 X, Y, Z, W;
+		};
+
+		UE_DEPRECATED(all, "For internal use only")
+		uint32 XYZW[4];
+	};
 
 	FORCEINLINE FUintVector4()
 	{
@@ -998,13 +1067,17 @@ struct FUintVector4
 
 	FORCEINLINE const uint32& operator[](int32 ComponentIndex) const
 	{
-		return (&X)[ComponentIndex];
+PRAGMA_DISABLE_DEPRECATION_WARNINGS
+		return XYZW[ComponentIndex];
+PRAGMA_ENABLE_DEPRECATION_WARNINGS
 	}
 
 
 	FORCEINLINE uint32& operator[](int32 ComponentIndex)
 	{
-		return (&X)[ComponentIndex];
+PRAGMA_DISABLE_DEPRECATION_WARNINGS
+		return XYZW[ComponentIndex];
+PRAGMA_ENABLE_DEPRECATION_WARNINGS
 	}
 
 	FORCEINLINE bool operator==(const FUintVector4& Other) const
diff --git a/Engine/Source/Runtime/Core/Public/Math/Vector.h b/Engine/Source/Runtime/Core/Public/Math/Vector.h
--- a/Engine/Source/Runtime/Core/Public/Math/Vector.h
+++ b/Engine/Source/Runtime/Core/Public/Math/Vector.h
@@ -53,14 +53,23 @@ struct TVector
 public:
 	using FReal = T;
 
-    /** Vector's X component. */
-    T X;
+	union
+	{
+		struct
+		{
+			/** Vector's X component. */
+			T X;
 
-    /** Vector's Y component. */
-    T Y;
+			/** Vector's Y component. */
+			T Y;
 
-    /** Vector's Z component. */
-    T Z;
+			/** Vector's Z component. */
+			T Z;
+		};
+
+		UE_DEPRECATED(all, "For internal use only")
+		T XYZ[3];
+	};
 
 	/** A zero vector (0,0,0) */
 	CORE_API static const TVector<T> ZeroVector;
@@ -1537,14 +1546,18 @@ template<typename T>
 FORCEINLINE T& TVector<T>::operator[](int32 Index)
 {
     checkSlow(Index >= 0 && Index < 3);
-    return (&X)[Index];
+PRAGMA_DISABLE_DEPRECATION_WARNINGS
+    return XYZ[Index];
+PRAGMA_ENABLE_DEPRECATION_WARNINGS
 }
 
 template<typename T>
 FORCEINLINE T TVector<T>::operator[](int32 Index)const
 {
     checkSlow(Index >= 0 && Index < 3);
-    return (&X)[Index];
+PRAGMA_DISABLE_DEPRECATION_WARNINGS
+    return XYZ[Index];
+PRAGMA_ENABLE_DEPRECATION_WARNINGS
 }
 
 template<typename T>
@@ -1837,13 +1850,17 @@ FORCEINLINE void TVector<T>::AddBounded(const TVector<T>& V, T Radius)
 template<typename T>
 FORCEINLINE T& TVector<T>::Component(int32 Index)
 {
-    return (&X)[Index];
+PRAGMA_DISABLE_DEPRECATION_WARNINGS
+    return XYZ[Index];
+PRAGMA_ENABLE_DEPRECATION_WARNINGS
 }
 
 template<typename T>
 FORCEINLINE T TVector<T>::Component(int32 Index) const
 {
-    return (&X)[Index];
+PRAGMA_DISABLE_DEPRECATION_WARNINGS
+    return XYZ[Index];
+PRAGMA_ENABLE_DEPRECATION_WARNINGS
 }
 
 template<typename T>
diff --git a/Engine/Source/Runtime/Core/Public/Math/Vector2D.h b/Engine/Source/Runtime/Core/Public/Math/Vector2D.h
--- a/Engine/Source/Runtime/Core/Public/Math/Vector2D.h
+++ b/Engine/Source/Runtime/Core/Public/Math/Vector2D.h
@@ -35,11 +35,20 @@ struct TVector2
 public:
 	using FReal = T;
 	
-	/** Vector's X component. */
-	T X;
+	union
+	{
+		struct
+		{
+			/** Vector's X component. */
+			T X;
+
+			/** Vector's Y component. */
+			T Y;
+		};
 
-	/** Vector's Y component. */
-	T Y;
+		UE_DEPRECATED(all, "For internal use only")
+		T XY[2];
+	};
 
 	/** Global 2D zero vector constant (0,0) */
 	CORE_API static const TVector2<T> ZeroVector;
@@ -1100,13 +1109,17 @@ FORCEINLINE bool TVector2<T>::IsZero() const
 template<typename T>
 FORCEINLINE T& TVector2<T>::Component(int32 Index)
 {
-	return (&X)[Index];
+PRAGMA_DISABLE_DEPRECATION_WARNINGS
+	return XY[Index];
+PRAGMA_ENABLE_DEPRECATION_WARNINGS
 }
 
 template<typename T>
 FORCEINLINE T TVector2<T>::Component(int32 Index) const
 {
-	return (&X)[Index];
+PRAGMA_DISABLE_DEPRECATION_WARNINGS
+	return XY[Index];
+PRAGMA_ENABLE_DEPRECATION_WARNINGS
 }
 
 template<typename T>
diff --git a/Engine/Source/Runtime/Core/Public/Math/Vector4.h b/Engine/Source/Runtime/Core/Public/Math/Vector4.h
--- a/Engine/Source/Runtime/Core/Public/Math/Vector4.h
+++ b/Engine/Source/Runtime/Core/Public/Math/Vector4.h
@@ -31,17 +31,26 @@ struct alignas(16) TVector4
 public:
 	using FReal = T;
 
-	/** The vector's X-component. */
-	T X;
+	union
+	{
+		struct
+		{
+			/** The vector's X-component. */
+			T X;
+
+			/** The vector's Y-component. */
+			T Y;
 
-	/** The vector's Y-component. */
-	T Y;
+			/** The vector's Z-component. */
+			T Z;
 
-	/** The vector's Z-component. */
-	T Z;
+			/** The vector's W-component. */
+			T W;
+		};
 
-	/** The vector's W-component. */
-	T W;
+		UE_DEPRECATED(all, "For internal use only")
+		T XYZW[4];
+	};
 
 public:
 
@@ -607,14 +616,18 @@ FORCEINLINE TVector4<T>::TVector4(const FIntVector4& InVector)
 template<typename T>
 FORCEINLINE T& TVector4<T>::operator[](int32 ComponentIndex)
 {
-	return (&X)[ ComponentIndex ];
+PRAGMA_DISABLE_DEPRECATION_WARNINGS
+	return XYZW[ ComponentIndex ];
+PRAGMA_ENABLE_DEPRECATION_WARNINGS
 }
 
 
 template<typename T>
 FORCEINLINE T TVector4<T>::operator[](int32 ComponentIndex) const
 {
-	return (&X)[ ComponentIndex ];
+PRAGMA_DISABLE_DEPRECATION_WARNINGS
+	return XYZW[ ComponentIndex ];
+PRAGMA_ENABLE_DEPRECATION_WARNINGS
 }
 
 
@@ -691,13 +704,17 @@ FORCEINLINE TVector4<T> TVector4<T>::operator^(const TVector4<T>& V) const
 template<typename T>
 FORCEINLINE T& TVector4<T>::Component(int32 Index)
 {
-	return (&X)[Index];
+PRAGMA_DISABLE_DEPRECATION_WARNINGS
+	return XYZW[Index];
+PRAGMA_ENABLE_DEPRECATION_WARNINGS
 }
 
 template<typename T>
 FORCEINLINE const T& TVector4<T>::Component(int32 Index) const
 {
-	return (&X)[Index];
+PRAGMA_DISABLE_DEPRECATION_WARNINGS
+	return XYZW[Index];
+PRAGMA_ENABLE_DEPRECATION_WARNINGS
 }
 
 template<typename T>
diff --git a/Engine/Source/Runtime/Datasmith/CADKernel/Public/CADKernel/Math/Point.h b/Engine/Source/Runtime/Datasmith/CADKernel/Public/CADKernel/Math/Point.h
--- a/Engine/Source/Runtime/Datasmith/CADKernel/Public/CADKernel/Math/Point.h
+++ b/Engine/Source/Runtime/Datasmith/CADKernel/Public/CADKernel/Math/Point.h
@@ -13,9 +13,18 @@ class FPoint2D;
 class CADKERNEL_API FPoint
 {
 public:
-	double X;
-	double Y;
-	double Z;
+	union
+	{
+		struct
+		{
+			double X;
+			double Y;
+			double Z;
+		};
+
+		UE_DEPRECATED(all, "For internal use only")
+		double XYZ[3];
+	};
 
 public:
 
@@ -48,12 +57,16 @@ public:
 
 	constexpr double& operator[](int32 Index)
 	{
-		return *(&X + Index);
+PRAGMA_DISABLE_DEPRECATION_WARNINGS
+		return XYZ[Index];
+PRAGMA_ENABLE_DEPRECATION_WARNINGS
 	}
 
 	constexpr double operator[](int32 Index) const
 	{
-		return *(&X + Index);
+PRAGMA_DISABLE_DEPRECATION_WARNINGS
+		return XYZ[Index];
+PRAGMA_ENABLE_DEPRECATION_WARNINGS
 	}
 
 	bool operator==(const FPoint& Point) const
@@ -319,10 +332,19 @@ public:
 class CADKERNEL_API FPointH
 {
 public:
-	double X;
-	double Y;
-	double Z;
-	double W;
+	union
+	{
+		struct
+		{
+			double X;
+			double Y;
+			double Z;
+			double W;
+		};
+
+		UE_DEPRECATED(all, "For internal use only")
+		double XYZW[4];
+	};
 
 public:
 	/** A zero point (0,0,0)*/
@@ -356,12 +378,16 @@ public:
 
 	constexpr double& operator[](int32 Index)
 	{
-		return *(&X + Index);
+PRAGMA_DISABLE_DEPRECATION_WARNINGS
+		return XYZW[Index];
+PRAGMA_ENABLE_DEPRECATION_WARNINGS
 	}
 
 	constexpr double operator[](int32 Index) const
 	{
-		return *(&X + Index);
+PRAGMA_DISABLE_DEPRECATION_WARNINGS
+		return XYZW[Index];
+PRAGMA_ENABLE_DEPRECATION_WARNINGS
 	}
 
 	FPointH operator+(const FPointH& Point) const
@@ -632,9 +658,18 @@ typedef FPoint2D FSurfacicTolerance;
 class CADKERNEL_API FFPoint
 {
 public:
-	float X;
-	float Y;
-	float Z;
+	union
+	{
+		struct
+		{
+			float X;
+			float Y;
+			float Z;
+		};
+
+		UE_DEPRECATED(all, "For internal use only")
+		float XYZ[3];
+	};
 
 	/** A zero point (0,0,0)*/
 	static const FFPoint ZeroPoint;
@@ -657,12 +692,16 @@ public:
 
 	constexpr float& operator[](int32 Index)
 	{
-		return *(&X + Index);
+PRAGMA_DISABLE_DEPRECATION_WARNINGS
+		return XYZ[Index];
+PRAGMA_ENABLE_DEPRECATION_WARNINGS
 	}
 
 	constexpr float operator[](int32 Index) const
 	{
-		return *(&X + Index);
+PRAGMA_DISABLE_DEPRECATION_WARNINGS
+		return XYZ[Index];
+PRAGMA_ENABLE_DEPRECATION_WARNINGS
 	}
 
 	FFPoint& operator=(const FPoint& Point)
diff --git a/Engine/Source/Runtime/Engine/Public/AnimationCompression.h b/Engine/Source/Runtime/Engine/Public/AnimationCompression.h
--- a/Engine/Source/Runtime/Engine/Public/AnimationCompression.h
+++ b/Engine/Source/Runtime/Engine/Public/AnimationCompression.h
@@ -219,9 +219,18 @@ public:
 class FQuatFloat96NoW
 {
 public:
-	float X;
-	float Y;
-	float Z;
+	union
+	{
+		struct
+		{
+			float X;
+			float Y;
+			float Z;
+		};
+
+		UE_DEPRECATED(all, "For internal use only")
+		float XYZ[3];
+	};
 
 	FQuatFloat96NoW()
 	{}
@@ -266,7 +275,9 @@ public:
 
 	void ToQuat(FQuat4f& Out) const
 	{
-		Out = ToQuat<true>(&X);
+PRAGMA_DISABLE_DEPRECATION_WARNINGS
+		Out = ToQuat<true>(XYZ);
+PRAGMA_ENABLE_DEPRECATION_WARNINGS
 	}
 
 	friend FArchive& operator<<(FArchive& Ar, FQuatFloat96NoW& Quat)
diff --git a/Engine/Source/Runtime/Experimental/Chaos/Public/Chaos/AABBTree.h b/Engine/Source/Runtime/Experimental/Chaos/Public/Chaos/AABBTree.h
--- a/Engine/Source/Runtime/Experimental/Chaos/Public/Chaos/AABBTree.h
+++ b/Engine/Source/Runtime/Experimental/Chaos/Public/Chaos/AABBTree.h
@@ -2802,6 +2802,8 @@ private:
 			const TVec3<T> ElemCenter = Elem.Bounds.Center();
 			
 			// This was changed to work around a code generation issue on some platforms, don't change it without testing results of ElemCenter computation.
+			// 
+			// NOTE: This needs review.  TVec3<T>::operator[] should now cope with the strict aliasing violation which caused the original codegen breakage.
 			T CenterVal = ElemCenter[0];
 			if (MaxAxis == 1)
 			{
diff --git a/Engine/Source/Runtime/Experimental/Chaos/Public/Chaos/Triangle.h b/Engine/Source/Runtime/Experimental/Chaos/Public/Chaos/Triangle.h
--- a/Engine/Source/Runtime/Experimental/Chaos/Public/Chaos/Triangle.h
+++ b/Engine/Source/Runtime/Experimental/Chaos/Public/Chaos/Triangle.h
@@ -18,51 +18,49 @@ namespace Chaos
 		}
 
 		FTriangle(const FVec3& InA, const FVec3& InB, const FVec3& InC)
-			: A(InA)
-			, B(InB)
-			, C(InC)
+			: ABC{ InA, InB, InC }
 		{
 		}
 
 		FORCEINLINE FVec3& operator[](uint32 InIndex)
 		{
 			checkSlow(InIndex < 3);
-			return (&A)[InIndex];
+			return ABC[InIndex];
 		}
 
 		FORCEINLINE const FVec3& operator[](uint32 InIndex) const
 		{
 			checkSlow(InIndex < 3);
-			return (&A)[InIndex];
+			return ABC[InIndex];
 		}
 
 		FORCEINLINE const FVec3& GetVertex(const int32 InIndex) const
 		{
 			checkSlow(InIndex < 3);
-			return (&A)[InIndex];
+			return ABC[InIndex];
 		}
 
 		FORCEINLINE FVec3 GetNormal() const
 		{
-			return FVec3::CrossProduct(B - A, C - A).GetSafeNormal();
+			return FVec3::CrossProduct(ABC[1] - ABC[0], ABC[2] - ABC[0]).GetSafeNormal();
 		}
 
 		FORCEINLINE TPlane<FReal, 3> GetPlane() const
 		{
-			return TPlane<FReal, 3>(A, GetNormal());
+			return TPlane<FReal, 3>(ABC[0], GetNormal());
 		}
 
 		// Face index is ignored since we only have one face
 		// Used for manifold generation
 		FORCEINLINE TPlaneConcrete<FReal, 3> GetPlane(int32 FaceIndex) const
 		{
-			return TPlaneConcrete < FReal, 3> (A, GetNormal());
+			return TPlaneConcrete < FReal, 3> (ABC[0], GetNormal());
 		}
 
 		FORCEINLINE void GetPlaneNX(const int32 FaceIndex, FVec3& OutN, FVec3& OutX) const
 		{
 			OutN = GetNormal();
-			OutX = A;
+			OutX = ABC[0];
 		}
 
 		// Get the nearest point on an edge and the edge vertices
@@ -73,12 +71,12 @@ namespace Chaos
 			FReal ClosestDistanceSq = TNumericLimits<FReal>::Max();
 
 			int32 PlaneVerticesNum = 3;
-
-			FVec3 P0 = C;
+			
+			FVec3 P0 = ABC[2];
 			for (int32 PlaneVertexIndex = 0; PlaneVertexIndex < PlaneVerticesNum; ++PlaneVertexIndex)
 			{
 				const TVector<FReal, 3>& P1 = GetVertex(PlaneVertexIndex);
-
+				
 				const FVec3 EdgePosition = FMath::ClosestPointOnLine(P0, P1, Position);
 				const FReal EdgeDistanceSq = (EdgePosition - Position).SizeSquared();
 
@@ -158,61 +156,61 @@ namespace Chaos
 		FORCEINLINE FReal PhiWithNormal(const FVec3& InSamplePoint, FVec3& OutNormal) const
 		{
 			OutNormal = GetNormal();
-			FVec3 ClosestPoint = FindClosestPointOnTriangle(GetPlane(), A, B, C, InSamplePoint);
+			FVec3 ClosestPoint = FindClosestPointOnTriangle(GetPlane(), ABC[0], ABC[1], ABC[2], InSamplePoint);
 			return FVec3::DotProduct((InSamplePoint - ClosestPoint), OutNormal);
 		}
 
 		FORCEINLINE FVec3 Support(const FVec3& Direction, const FReal Thickness, int32& VertexIndex) const
 		{
-			const FReal DotA = FVec3::DotProduct(A, Direction);
-			const FReal DotB = FVec3::DotProduct(B, Direction);
-			const FReal DotC = FVec3::DotProduct(C, Direction);
+			const FReal DotA = FVec3::DotProduct(ABC[0], Direction);
+			const FReal DotB = FVec3::DotProduct(ABC[1], Direction);
+			const FReal DotC = FVec3::DotProduct(ABC[2], Direction);
 
 			if(DotA >= DotB && DotA >= DotC)
 			{
 				VertexIndex = 0;
 				if(Thickness != 0)
 				{
-					return A + Direction.GetUnsafeNormal() * Thickness;
+					return ABC[0] + Direction.GetUnsafeNormal() * Thickness;
 				}
-				return A;
+				return ABC[0];
 			}
 			else if(DotB >= DotA && DotB >= DotC)
 			{
 				VertexIndex = 1;
 				if(Thickness != 0)
 				{
-					return B + Direction.GetUnsafeNormal() * Thickness;
+					return ABC[1] + Direction.GetUnsafeNormal() * Thickness;
 				}
-				return B;
+				return ABC[1];
 			}
 			VertexIndex = 2;
 			if(Thickness != 0)
 			{
-				return C + Direction.GetUnsafeNormal() * Thickness;
+				return ABC[2] + Direction.GetUnsafeNormal() * Thickness;
 			}
-			return C;
+			return ABC[2];
 		}
 
 		FORCEINLINE_DEBUGGABLE FVec3 SupportCore(const FVec3& Direction, const FReal InMargin, FReal* OutSupportDelta,int32& VertexIndex) const
 		{
 			// Note: assumes margin == 0
-			const FReal DotA = FVec3::DotProduct(A, Direction);
-			const FReal DotB = FVec3::DotProduct(B, Direction);
-			const FReal DotC = FVec3::DotProduct(C, Direction);
+			const FReal DotA = FVec3::DotProduct(ABC[0], Direction);
+			const FReal DotB = FVec3::DotProduct(ABC[1], Direction);
+			const FReal DotC = FVec3::DotProduct(ABC[2], Direction);
 
 			if (DotA >= DotB && DotA >= DotC)
 			{
 				VertexIndex = 0;
-				return A;
+				return ABC[0];
 			}
 			else if (DotB >= DotA && DotB >= DotC)
 			{
 				VertexIndex = 1;
-				return B;
+				return ABC[1];
 			}
 			VertexIndex = 2;
-			return C;
+			return ABC[2];
 		}
 
 		FORCEINLINE FVec3 SupportCoreScaled(const FVec3& Direction, FReal InMargin, const FVec3& Scale, FReal* OutSupportDelta, int32& VertexIndex) const
@@ -237,7 +235,7 @@ namespace Chaos
 
 		FORCEINLINE bool Overlap(const FVec3& Point, const FReal Thickness) const
 		{
-			const FVec3 ClosestPoint = FindClosestPointOnTriangle(GetPlane(), A, B, C, Point);
+			const FVec3 ClosestPoint = FindClosestPointOnTriangle(GetPlane(), ABC[0], ABC[1], ABC[2], Point);
 			const FReal AdjustedThickness = FMath::Max(Thickness, KINDA_SMALL_NUMBER);
 			return (Point - ClosestPoint).SizeSquared() <= (AdjustedThickness * AdjustedThickness);
 		}
@@ -251,14 +249,12 @@ namespace Chaos
 
 		friend FChaosArchive& operator<<(FChaosArchive& Ar, FTriangle& Value);
 
-		FVec3 A;
-		FVec3 B;
-		FVec3 C;
+		FVec3 ABC[3];
 	};
 
 	inline FChaosArchive& operator<<(FChaosArchive& Ar, FTriangle& Value)
 	{
-		Ar << Value.A << Value.B << Value.C;
+		Ar << Value.ABC[0] << Value.ABC[1] << Value.ABC[2];
 		return Ar;
 	}
 
diff --git a/Engine/Source/Runtime/Experimental/ChaosCore/Public/Chaos/Vector.h b/Engine/Source/Runtime/Experimental/ChaosCore/Public/Chaos/Vector.h
--- a/Engine/Source/Runtime/Experimental/ChaosCore/Public/Chaos/Vector.h
+++ b/Engine/Source/Runtime/Experimental/ChaosCore/Public/Chaos/Vector.h
@@ -1064,8 +1064,10 @@ namespace Chaos
 			return !FMath::IsFinite(X) || !FMath::IsFinite(Y) || !FMath::IsFinite(Z);
 		}
 
-		FORCEINLINE T operator[](int32 Idx) const { return (static_cast<const T*>(&X))[Idx]; }
-		FORCEINLINE T& operator[](int32 Idx) { return (static_cast<T*>(&X))[Idx]; }
+PRAGMA_DISABLE_DEPRECATION_WARNINGS
+		FORCEINLINE T operator[](int32 Idx) const { return XYZ[Idx]; }
+		FORCEINLINE T& operator[](int32 Idx) { return XYZ[Idx]; }
+PRAGMA_ENABLE_DEPRECATION_WARNINGS
 
 		FORCEINLINE TVector<T, 3> operator-() const { return {-X, -Y, -Z}; }
 		FORCEINLINE TVector<T, 3> operator*(const TVector<T, 3>& Other) const { return {X * Other.X, Y * Other.Y, Z * Other.Z}; }
@@ -1120,9 +1122,18 @@ namespace Chaos
 		}
 #endif
 
-		T X;
-		T Y;
-		T Z;
+		union
+		{
+			struct
+			{
+				T X;
+				T Y;
+				T Z;
+			};
+
+			UE_DEPRECATED(all, "For internal use only")
+			T XYZ[3];
+		};
 	};
 	template<class T>
 	inline TVector<T, 3> operator*(const T S, const TVector<T, 3>& V)
@@ -1187,8 +1198,10 @@ namespace Chaos
 			return *this;
 		}
 
-		FORCEINLINE FElement operator[](const int32 Idx) const { return (static_cast<const FElement*>(&X))[Idx]; }
-		FORCEINLINE FElement& operator[](const int32 Idx) { return (static_cast<FElement*>(&X))[Idx]; }
+PRAGMA_DISABLE_DEPRECATION_WARNINGS
+		FORCEINLINE FElement operator[](const int32 Idx) const { return XY[Idx]; }
+		FORCEINLINE FElement& operator[](const int32 Idx) { return XY[Idx]; }
+PRAGMA_ENABLE_DEPRECATION_WARNINGS
 
 		FORCEINLINE TVector<FElement, 2> operator-() const { return {-X, -Y}; }
 		FORCEINLINE TVector<FElement, 2> operator*(const TVector<FElement, 2>& Other) const { return {X * Other.X, Y * Other.Y}; }
@@ -1237,8 +1250,17 @@ namespace Chaos
 
 
 	private:
-		FElement X;
-		FElement Y;
+		union
+		{
+			struct
+			{
+				FElement X;
+				FElement Y;
+			};
+
+			UE_DEPRECATED(all, "For internal use only")
+			FElement XY[2];
+		};
 	};
 
 	template<class T>
diff --git a/Engine/Source/Runtime/GeometryCore/Public/IndexTypes.h b/Engine/Source/Runtime/GeometryCore/Public/IndexTypes.h
--- a/Engine/Source/Runtime/GeometryCore/Public/IndexTypes.h
+++ b/Engine/Source/Runtime/GeometryCore/Public/IndexTypes.h
@@ -25,8 +25,16 @@ namespace Geometry
  */
 struct FIndex2i
 {
-	int A{ IndexConstants::InvalidID };
-	int B{ IndexConstants::InvalidID };
+	union
+	{
+		struct
+		{
+			int A, B;
+		};
+
+		UE_DEPRECATED(all, "For internal use only")
+		int AB[2] = { IndexConstants::InvalidID, IndexConstants::InvalidID };
+	};
 
 	constexpr FIndex2i() = default;
 	constexpr FIndex2i(int ValA, int ValB)
@@ -48,11 +56,15 @@ struct FIndex2i
 
 	int& operator[](int Idx)
 	{
-		return (&A)[Idx];
+PRAGMA_DISABLE_DEPRECATION_WARNINGS
+		return AB[Idx];
+PRAGMA_ENABLE_DEPRECATION_WARNINGS
 	}
 	const int& operator[](int Idx) const
 	{
-		return (&A)[Idx];
+PRAGMA_DISABLE_DEPRECATION_WARNINGS
+		return AB[Idx];
+PRAGMA_ENABLE_DEPRECATION_WARNINGS
 	}
 
 	inline bool operator==(const FIndex2i& Other) const
@@ -138,9 +150,18 @@ FORCEINLINE uint32 GetTypeHash(const FIndex2i& Index)
  */
 struct FIndex3i
 {
-	int A{ IndexConstants::InvalidID };
-	int B{ IndexConstants::InvalidID };
-	int C{ IndexConstants::InvalidID };
+	union
+	{
+		struct
+		{
+			int A;
+			int B;
+			int C;
+		};
+
+		UE_DEPRECATED(all, "For internal use only")
+		int ABC[3] = { IndexConstants::InvalidID, IndexConstants::InvalidID, IndexConstants::InvalidID };
+	};
 
 	constexpr FIndex3i() = default;
 	constexpr FIndex3i(int ValA, int ValB, int ValC)
@@ -162,11 +183,15 @@ struct FIndex3i
 
 	int& operator[](int Idx)
 	{
-		return (&A)[Idx];
+PRAGMA_DISABLE_DEPRECATION_WARNINGS
+		return ABC[Idx];
+PRAGMA_ENABLE_DEPRECATION_WARNINGS
 	}
 	const int& operator[](int Idx) const
 	{
-		return (&A)[Idx];
+PRAGMA_DISABLE_DEPRECATION_WARNINGS
+		return ABC[Idx];
+PRAGMA_ENABLE_DEPRECATION_WARNINGS
 	}
 
 	bool operator==(const FIndex3i& Other) const
@@ -258,7 +283,16 @@ FORCEINLINE uint32 GetTypeHash(const FIndex3i& Index)
 
 struct FIndex4i
 {
-	int A, B, C, D;
+	union
+	{
+		struct
+		{
+			int A, B, C, D;
+		};
+
+		UE_DEPRECATED(all, "For internal use only")
+		int ABCD[4];
+	};
 
 	FIndex4i()
 	{
@@ -286,11 +320,15 @@ struct FIndex4i
 
 	int& operator[](int Idx)
 	{
-		return (&A)[Idx];
+PRAGMA_DISABLE_DEPRECATION_WARNINGS
+		return ABCD[Idx];
+PRAGMA_ENABLE_DEPRECATION_WARNINGS
 	}
 	const int& operator[](int Idx) const
 	{
-		return (&A)[Idx];
+PRAGMA_DISABLE_DEPRECATION_WARNINGS
+		return ABCD[Idx];
+PRAGMA_ENABLE_DEPRECATION_WARNINGS
 	}
 
 	bool operator==(const FIndex4i& Other) const
diff --git a/Engine/Source/Runtime/GeometryCore/Public/IntVectorTypes.h b/Engine/Source/Runtime/GeometryCore/Public/IntVectorTypes.h
--- a/Engine/Source/Runtime/GeometryCore/Public/IntVectorTypes.h
+++ b/Engine/Source/Runtime/GeometryCore/Public/IntVectorTypes.h
@@ -17,7 +17,16 @@ namespace Geometry {
 */
 struct FVector2i
 {
-	int32 X{}, Y{};
+	union
+	{
+		struct
+		{
+			int32 X, Y;
+		};
+
+		UE_DEPRECATED(all, "For internal use only")
+		int32 XY[2] = { {}, {} };
+	};
 
 	constexpr FVector2i()
 		: X(0), Y(0)
@@ -34,17 +43,23 @@ struct FVector2i
 	{
 	}
 
+PRAGMA_DISABLE_DEPRECATION_WARNINGS
 	constexpr FVector2i(const FVector2i& Vec) = default;
 	constexpr FVector2i& operator=(const FVector2i& V2) = default;
+//PRAGMA_ENABLE_DEPRECATION_WARNINGS
 
 	explicit constexpr operator const int32*() const
 	{
-		return &X;
+PRAGMA_DISABLE_DEPRECATION_WARNINGS
+		return XY;
+PRAGMA_ENABLE_DEPRECATION_WARNINGS
 	};
 
 	explicit constexpr operator int32*()
 	{
-		return &X;
+PRAGMA_DISABLE_DEPRECATION_WARNINGS
+		return XY;
+PRAGMA_ENABLE_DEPRECATION_WARNINGS
 	}
 
 	explicit operator FVector2f() const
@@ -90,11 +105,15 @@ struct FVector2i
 
 	constexpr int32& operator[](int Idx)
 	{
-		return (&X)[Idx];
+PRAGMA_DISABLE_DEPRECATION_WARNINGS
+		return XY[Idx];
+PRAGMA_ENABLE_DEPRECATION_WARNINGS
 	}
 	constexpr const int32& operator[](int Idx) const
 	{
-		return (&X)[Idx];
+PRAGMA_DISABLE_DEPRECATION_WARNINGS
+		return XY[Idx];
+PRAGMA_ENABLE_DEPRECATION_WARNINGS
 	}
 
 	constexpr int32 SquaredLength() const
@@ -221,7 +240,16 @@ FORCEINLINE uint32 GetTypeHash(const FVector2i& Vector)
 */
 struct FVector3i
 {
-	int32 X{}, Y{}, Z{};
+	union
+	{
+		struct
+		{
+			int32 X, Y, Z;
+		};
+
+		UE_DEPRECATED(all, "For internal use only")
+		int32 XYZ[3] = { {}, {}, {} };
+	};
 
 	constexpr FVector3i()
 		: X(0), Y(0), Z(0)
@@ -238,16 +266,22 @@ struct FVector3i
 	{
 	}
 
+PRAGMA_DISABLE_DEPRECATION_WARNINGS
 	constexpr FVector3i(const FVector3i& Vec) = default;
 	FVector3i& operator=(const FVector3i& V2) = default;
+PRAGMA_ENABLE_DEPRECATION_WARNINGS
 
 	explicit constexpr operator const int32*() const
 	{
-		return &X;
+PRAGMA_DISABLE_DEPRECATION_WARNINGS
+		return XYZ;
+PRAGMA_ENABLE_DEPRECATION_WARNINGS
 	};
 	explicit constexpr operator int32*()
 	{
-		return &X;
+PRAGMA_DISABLE_DEPRECATION_WARNINGS
+		return XYZ;
+PRAGMA_ENABLE_DEPRECATION_WARNINGS
 	}
 
 	explicit operator FVector3f() const
@@ -296,11 +330,15 @@ struct FVector3i
 	
 	int32& operator[](int Idx)
 	{
-		return (&X)[Idx];
+PRAGMA_DISABLE_DEPRECATION_WARNINGS
+		return XYZ[Idx];
+PRAGMA_ENABLE_DEPRECATION_WARNINGS
 	}
 	const int32& operator[](int Idx) const
 	{
-		return (&X)[Idx];
+PRAGMA_DISABLE_DEPRECATION_WARNINGS
+		return XYZ[Idx];
+PRAGMA_ENABLE_DEPRECATION_WARNINGS
 	}
 
 	constexpr int32 SquaredLength() const
diff --git a/Engine/Source/Runtime/GeometryCore/Public/Quaternion.h b/Engine/Source/Runtime/GeometryCore/Public/Quaternion.h
--- a/Engine/Source/Runtime/GeometryCore/Public/Quaternion.h
+++ b/Engine/Source/Runtime/GeometryCore/Public/Quaternion.h
@@ -20,11 +20,20 @@ using namespace UE::Math;
 template<typename RealType>
 struct TQuaternion
 {
-	// note: in Wm5 version, this is a 4-element arraY stored in order (w,x,y,z).
-	RealType X;
-	RealType Y;
-	RealType Z;
-	RealType W;
+	union
+	{
+		struct
+		{
+			// note: in Wm5 version, this is a 4-element arraY stored in order (w,x,y,z).
+			RealType X;
+			RealType Y;
+			RealType Z;
+			RealType W;
+		};
+
+		UE_DEPRECATED(all, "For internal use only")
+		RealType XYZW[4];
+	};
 
 	TQuaternion();
 	TQuaternion(RealType X, RealType Y, RealType Z, RealType W);
@@ -46,8 +55,18 @@ struct TQuaternion
 	static TQuaternion<RealType> Zero() { return TQuaternion<RealType>((RealType)0, (RealType)0, (RealType)0, (RealType)0); }
 	static TQuaternion<RealType> Identity() { return TQuaternion<RealType>((RealType)0, (RealType)0, (RealType)0, (RealType)1); }
 
-	RealType & operator[](int i) { return (&X)[i]; }
-	const RealType & operator[](int i) const { return (&X)[i]; }
+	RealType& operator[](int i)
+	{
+PRAGMA_DISABLE_DEPRECATION_WARNINGS
+		return XYZW[i];
+PRAGMA_ENABLE_DEPRECATION_WARNINGS
+	}
+	const RealType& operator[](int i) const
+	{
+PRAGMA_DISABLE_DEPRECATION_WARNINGS
+		return XYZW[i];
+PRAGMA_ENABLE_DEPRECATION_WARNINGS
+	}
 
 	bool EpsilonEqual(const TQuaternion<RealType>& Other, RealType Epsilon) const;
 
