Fix NaN warnings when components of Dir are close to 0.

diff --git a/Engine/Source/Runtime/Experimental/Chaos/Public/Chaos/ISpatialAcceleration.h b/Engine/Source/Runtime/Experimental/Chaos/Public/Chaos/ISpatialAcceleration.h
--- a/Engine/Source/Runtime/Experimental/Chaos/Public/Chaos/ISpatialAcceleration.h
+++ b/Engine/Source/Runtime/Experimental/Chaos/Public/Chaos/ISpatialAcceleration.h
@@ -13,8 +13,8 @@ struct CHAOS_API FQueryFastData
 {
 	FQueryFastData(const FVec3& InDir, const FReal InLength)
 		: Dir(InDir)
-		, InvDir( (InDir[0] == 0) ? 0 : 1 / Dir[0], (InDir[1] == 0) ? 0 : 1 / Dir[1], (InDir[2] == 0) ? 0 : 1 / Dir[2])
-		, bParallel{ InDir[0] == 0, InDir[1] == 0, InDir[2] == 0 }
+		, InvDir( (FMath::Abs(InDir[0]) < SMALL_NUMBER) ? 0 : 1 / Dir[0], (FMath::Abs(InDir[1]) < SMALL_NUMBER) ? 0 : 1 / Dir[1], (FMath::Abs(InDir[2]) < SMALL_NUMBER) ? 0 : 1 / Dir[2])
+		, bParallel{ FMath::Abs(InDir[0]) < SMALL_NUMBER, FMath::Abs(InDir[1]) < SMALL_NUMBER, FMath::Abs(InDir[2]) < SMALL_NUMBER }
 	{
 		CHAOS_ENSURE(InLength != 0.0f);
 		SetLength(InLength);
