Chaos: Fix triangle register construction prior to calling GJK for heightfields. Fixes missing second triangle in each cell leaving holes in the heightfield

diff --git a/Engine/Source/Runtime/Experimental/Chaos/Private/Chaos/HeightField.cpp b/Engine/Source/Runtime/Experimental/Chaos/Private/Chaos/HeightField.cpp
--- a/Engine/Source/Runtime/Experimental/Chaos/Private/Chaos/HeightField.cpp
+++ b/Engine/Source/Runtime/Experimental/Chaos/Private/Chaos/HeightField.cpp
@@ -379,7 +379,7 @@ namespace Chaos
 				}
 
 				//Convert into local space of A to get better precision
-				VectorRegister4Float AReg = MakeVectorRegisterFloatFromDouble(MakeVectorRegister(A.X, B.Y, A.Z, 0.0));
+				VectorRegister4Float AReg = MakeVectorRegisterFloatFromDouble(MakeVectorRegister(A.X, A.Y, A.Z, 0.0));
 				VectorRegister4Float BReg = MakeVectorRegisterFloatFromDouble(MakeVectorRegister(B.X, B.Y, B.Z, 0.0));
 				VectorRegister4Float CReg = MakeVectorRegisterFloatFromDouble(MakeVectorRegister(C.X, C.Y, C.Z, 0.0));
 
