Improve Joint linear stability and fix validdatas issue + shock propagation for linear solver + diagonal point constraint

diff --git a/Engine/Source/Runtime/Experimental/Chaos/Private/Chaos/Joint/JointConstraintsCVars.cpp b/Engine/Source/Runtime/Experimental/Chaos/Private/Chaos/Joint/JointConstraintsCVars.cpp
--- a/Engine/Source/Runtime/Experimental/Chaos/Private/Chaos/Joint/JointConstraintsCVars.cpp
+++ b/Engine/Source/Runtime/Experimental/Chaos/Private/Chaos/Joint/JointConstraintsCVars.cpp
@@ -29,3 +29,6 @@ FAutoConsoleVariableRef CVarChaosJointLinearVelocityThresholdToApplyRestitution(
 float Chaos_Joint_AngularVelocityThresholdToApplyRestitution = 1e-2f;
 FAutoConsoleVariableRef CVarChaosJointAngularVelocityThresholdToApplyRestitution(TEXT("p.Chaos.Joint.AngularVelocityThresholdToApplyRestitution"), Chaos_Joint_AngularVelocityThresholdToApplyRestitution, TEXT("Apply restitution only if initial velocity is higher than this threshold (used in Quasipbd)"));
 
+bool bChaos_Joint_StablePointConstraint = false;
+FAutoConsoleVariableRef CVarChaosJointStablePointConstraint(TEXT("p.Chaos.Joint.StablePointConstraint"), bChaos_Joint_StablePointConstraint, TEXT("Improve stability for the point constraint (increase computation time)"));
+
diff --git a/Engine/Source/Runtime/Experimental/Chaos/Private/Chaos/Joint/JointConstraintsCVars.h b/Engine/Source/Runtime/Experimental/Chaos/Private/Chaos/Joint/JointConstraintsCVars.h
--- a/Engine/Source/Runtime/Experimental/Chaos/Private/Chaos/Joint/JointConstraintsCVars.h
+++ b/Engine/Source/Runtime/Experimental/Chaos/Private/Chaos/Joint/JointConstraintsCVars.h
@@ -28,3 +28,5 @@ extern float Chaos_Joint_AngularVelocityThresholdToApplyRestitution;
 
 extern bool bChaos_Joint_UseCachedSolver;
 
+extern bool bChaos_Joint_StablePointConstraint;
+
diff --git a/Engine/Source/Runtime/Experimental/Chaos/Private/Chaos/Joint/PBDJointCachedSolverGaussSeidel.cpp b/Engine/Source/Runtime/Experimental/Chaos/Private/Chaos/Joint/PBDJointCachedSolverGaussSeidel.cpp
--- a/Engine/Source/Runtime/Experimental/Chaos/Private/Chaos/Joint/PBDJointCachedSolverGaussSeidel.cpp
+++ b/Engine/Source/Runtime/Experimental/Chaos/Private/Chaos/Joint/PBDJointCachedSolverGaussSeidel.cpp
@@ -22,14 +22,11 @@ void FPBDJointCachedSolver::InitDerivedState()
 	InitConnectorRs[0] = R(0) * LocalConnectorXs[0].GetRotation();
 	InitConnectorRs[1] = R(1) * LocalConnectorXs[1].GetRotation();
 	InitConnectorRs[1].EnforceShortestArcWith(InitConnectorRs[0]);
-
-	ConnectorXs[0] = InitConnectorXs[0];
-	ConnectorRs[0] = InitConnectorRs[0];
-
-	ConnectorXs[1] = InitConnectorXs[1];
-	ConnectorRs[1] = InitConnectorRs[1];
 	
-	UpdateDerivedState();
+	ComputeBodyState(0);
+	ComputeBodyState(1);
+
+	ConnectorRs[1].EnforceShortestArcWith(ConnectorRs[0]);
 }
 
 void FPBDJointCachedSolver::ComputeBodyState(const int32 BodyIndex)
@@ -96,7 +93,7 @@ void FPBDJointCachedSolver::UpdateMass0()
 	if ((ConditionedInvMs[0] > 0) && (InvMScales[0] > 0))
 	{
 		InvMs[0] = InvMScales[0] * ConditionedInvMs[0];
-		InvIs[0] = Utilities::ComputeWorldSpaceInertia(Q(0), InvMScales[0] * ConditionedInvILs[0]);
+		InvIs[0] = Utilities::ComputeWorldSpaceInertia(CurrentQs[0], InvMScales[0] * ConditionedInvILs[0]);
 	}
 	else
 	{
@@ -111,7 +108,7 @@ void FPBDJointCachedSolver::UpdateMass1()
 	if ((ConditionedInvMs[1] > 0) && (InvMScales[1] > 0))
 	{
 		InvMs[1] = InvMScales[1] * ConditionedInvMs[1];
-		InvIs[1] = Utilities::ComputeWorldSpaceInertia(Q(1), InvMScales[1] * ConditionedInvILs[1]);
+		InvIs[1] = Utilities::ComputeWorldSpaceInertia(CurrentQs[1], InvMScales[1] * ConditionedInvILs[1]);
 	}
 	else
 	{
@@ -120,17 +117,46 @@ void FPBDJointCachedSolver::UpdateMass1()
 	}
 }
 
-void FPBDJointCachedSolver::SetInvMassScales(const FReal InvMScale0, const FReal InvMScale1)
+void FPBDJointCachedSolver::SetInvMassScales(const FReal InvMScale0, const FReal InvMScale1, const FReal Dt)
 {
+	bool bNeedsUpdate = false;
 	if (InvMScales[0] != InvMScale0)
 	{
 		InvMScales[0] = InvMScale0;
 		UpdateMass0();
+		bNeedsUpdate = true;
 	}
 	if (InvMScales[1] != InvMScale1)
 	{
 		InvMScales[1] = InvMScale1;
 		UpdateMass1();
+		bNeedsUpdate = true;
+	}
+	if(bNeedsUpdate)
+	{
+		if(PointConstraints.bValidDatas)
+		{
+			InitPointPositionMass();
+		}
+		for(int32 ConstraintIndex = 0; ConstraintIndex < 3; ++ConstraintIndex)
+		{
+			if(PositionConstraints.bValidDatas[ConstraintIndex])
+			{
+				InitPositionDatasMass(PositionConstraints, ConstraintIndex, Dt);
+			}
+			if(RotationConstraints.bValidDatas[ConstraintIndex])
+			{
+				InitRotationDatasMass(RotationConstraints, ConstraintIndex, Dt);
+			}
+			if(PositionDrives.bValidDatas[ConstraintIndex])
+			{
+				InitPositionDatasMass(PositionDrives, ConstraintIndex, Dt);
+			}
+			if(RotationDrives.bValidDatas[ConstraintIndex])
+			{
+				InitRotationDatasMass(RotationDrives, ConstraintIndex, Dt);
+			}
+		}
 	}
 }
 
@@ -162,9 +188,7 @@ void FPBDJointCachedSolver::Init(
 	InvMScales[0] = FReal(1);
 	InvMScales[1] = FReal(1);
 	FPBDJointUtilities::ConditionInverseMassAndInertia(Body0().InvM(), Body1().InvM(), Body0().InvILocal(), Body1().InvILocal(), SolverSettings.MinParentMassRatio, SolverSettings.MaxInertiaRatio, ConditionedInvMs[0], ConditionedInvMs[1], ConditionedInvILs[0], ConditionedInvILs[1]);
-	UpdateMass0();
-	UpdateMass1();
-
+	
 	NetLinearImpulse = FVec3(0);
 	NetAngularImpulse = FVec3(0);
 
@@ -189,6 +213,9 @@ void FPBDJointCachedSolver::Init(
 
 	InitDerivedState();
 
+	UpdateMass0();
+	UpdateMass1();
+
 	// Cache all the informations for the position and rotation constraints
 	InitPositionConstraints(Dt, SolverSettings, JointSettings);
 	InitRotationConstraints(Dt, SolverSettings, JointSettings);
@@ -307,6 +334,7 @@ void FPBDJointCachedSolver::InitPositionConstraints(
 	PositionConstraints.bValidDatas[0] = false;
 	PositionConstraints.bValidDatas[1] = false;
 	PositionConstraints.bValidDatas[2] = false;
+	PointConstraints.bValidDatas = false;
 
 	TVec3<bool> bLinearLocked, bLinearLimited;
 	if(!ExtractLinearMotion(JointSettings, bLinearLocked, bLinearLimited))
@@ -383,6 +411,21 @@ void FPBDJointCachedSolver::InitPositionConstraints(
 	}
 }
 
+void FPBDJointCachedSolver::InitPositionDatasMass(
+	FAxisConstraintDatas& PositionDatas, 
+	const int32 ConstraintIndex,
+	const FReal Dt)
+{
+	const FVec3 AngularAxis0 = FVec3::CrossProduct(PositionDatas.ConstraintArms[ConstraintIndex][0], PositionDatas.ConstraintAxis[ConstraintIndex]);
+	const FVec3 AngularAxis1 = FVec3::CrossProduct(PositionDatas.ConstraintArms[ConstraintIndex][1], PositionDatas.ConstraintAxis[ConstraintIndex]);
+	const FVec3 IA0 = Utilities::Multiply(InvI(0), AngularAxis0);
+	const FVec3 IA1 = Utilities::Multiply(InvI(1), AngularAxis1);
+	const FReal II0 = FVec3::DotProduct(AngularAxis0, IA0);
+	const FReal II1 = FVec3::DotProduct(AngularAxis1, IA1);
+
+	PositionDatas.UpdateMass(ConstraintIndex, IA0, IA1, InvM(0) + II0 + InvM(1) + II1, Dt);
+}
+
 void FPBDJointCachedSolver::InitPositionConstraintDatas(
 	const int32 ConstraintIndex,
 	const FVec3& ConstraintAxis,
@@ -403,35 +446,59 @@ void FPBDJointCachedSolver::InitPositionConstraintDatas(
 	const FVec3 ConstraintArm0 = ConnectorXs[1] - LocalDelta * LocalAxis - P(0);
 	const FVec3 ConstraintArm1 = ConnectorXs[1] - P(1);
 
-	const FVec3 AngularAxis0 = FVec3::CrossProduct(ConstraintArm0, LocalAxis);
-	const FVec3 AngularAxis1 = FVec3::CrossProduct(ConstraintArm1, LocalAxis);
-	const FVec3 IA0 = Utilities::Multiply(InvI(0), AngularAxis0);
-	const FVec3 IA1 = Utilities::Multiply(InvI(1), AngularAxis1);
-	const FReal II0 = FVec3::DotProduct(AngularAxis0, IA0);
-	const FReal II1 = FVec3::DotProduct(AngularAxis1, IA1);
-
-	PositionConstraints.UpdateDatas(ConstraintIndex, LocalAxis, LocalDelta,
-		IA0, IA1, InvM(0) + II0 + InvM(1) + II1, ConstraintRestitution, Dt, true, ConstraintArm0, ConstraintArm1);
+	PositionConstraints.UpdateDatas(ConstraintIndex, LocalAxis, LocalDelta, ConstraintRestitution, true, ConstraintArm0, ConstraintArm1);
 
 	const FVec3 CV0 = V(0) + FVec3::CrossProduct(W(0), ConstraintArm0);
 	const FVec3 CV1 = V(1) + FVec3::CrossProduct(W(1), ConstraintArm1);
 	const FVec3 CV = CV1 - CV0;
 
 	InitConstraintAxisLinearVelocities[ConstraintIndex] = FVec3::DotProduct(CV, LocalAxis);
+
+	InitPositionDatasMass(PositionConstraints, ConstraintIndex, Dt);
+}
+
+void FPBDJointCachedSolver::InitPointPositionMass()
+{
+	FMatrix33 IM = FMatrix33::Zero;
+	if(bChaos_Joint_StablePointConstraint)
+	{
+		if(Body(0).IsDynamic() && InvM(0) > 0.0)
+		{
+			IM += Utilities::ComputeJointFactorMatrix(PointConstraints.ConstraintArms[0], InvI(0), InvM(0));
+		}
+		if(Body(1).IsDynamic() && InvM(1) > 0.0)
+		{
+			IM += Utilities::ComputeJointFactorMatrix(PointConstraints.ConstraintArms[1], InvI(1), InvM(1));
+		}
+		IM = IM.Inverse();
+	}
+	else
+	{
+		FVec3 ID = FVec3::ZeroVector;
+		if(Body(0).IsDynamic() && InvM(0) > 0.0)
+		{
+			ID += Utilities::ComputeDiagonalJointFactorMatrix(PointConstraints.ConstraintArms[0], InvI(0), InvM(0));
+		}
+		if(Body(1).IsDynamic() && InvM(1) > 0.0)
+		{
+			ID += Utilities::ComputeDiagonalJointFactorMatrix(PointConstraints.ConstraintArms[1], InvI(1), InvM(1));
+		}
+		IM.SetAt(0,0, (FMath::Abs(ID.X) > SMALL_NUMBER) ? 1.0 / ID.X : 0.0);
+		IM.SetAt(1,1, (FMath::Abs(ID.Y) > SMALL_NUMBER) ? 1.0 / ID.Y : 0.0);
+		IM.SetAt(2,2, (FMath::Abs(ID.Z) > SMALL_NUMBER) ? 1.0 / ID.Z : 0.0);
+	}
+	PointConstraints.ConstraintIM = IM;
+	PointConstraints.ConstraintLambda = FVec3::ZeroVector;
 }
 
 void FPBDJointCachedSolver::InitPointPositionConstraint()
 {
 	const FVec3 ConstraintArm0 = ConnectorXs[0] - P(0);
 	const FVec3 ConstraintArm1 = ConnectorXs[1] - P(1);
+	
+	PointConstraints.UpdateDatas(ConnectorXs[1] - ConnectorXs[0], ConstraintArm0, ConstraintArm1);
 
-	const FMatrix33 IM = Body(0).IsDynamic() && Body(1).IsDynamic() ?
-		Utilities::ComputeJointFactorMatrix(ConstraintArm0, InvI(0), InvM(0)) +
-		Utilities::ComputeJointFactorMatrix(ConstraintArm1, InvI(1), InvM(1)) : Body(0).IsDynamic() ?
-		Utilities::ComputeJointFactorMatrix(ConstraintArm0, InvI(0), InvM(0)) :
-		Utilities::ComputeJointFactorMatrix(ConstraintArm1, InvI(1), InvM(1));
-
-	PointConstraints.UpdateDatas(ConnectorXs[1] - ConnectorXs[0], IM.Inverse(), ConstraintArm0, ConstraintArm1);
+	InitPointPositionMass();
 }
 
 void FPBDJointCachedSolver::InitSphericalPositionConstraint(
@@ -498,16 +565,34 @@ void FPBDJointCachedSolver::ApplyPositionConstraints(
 	}
 }
 
+FVec3 FPBDJointCachedSolver::ComputePointCX() 
+{
+	if(bChaos_Joint_StablePointConstraint)
+	{
+		// Non-linear version
+		UpdateDerivedState();
+		return  ConnectorXs[1] - ConnectorXs[0];
+	}
+	else
+	{
+		// Linear version
+		return PointConstraints.ConstraintCX + Body(1).DP()  - Body(0).DP() +
+			FVec3::CrossProduct(Body(1).DQ(), PointConstraints.ConstraintArms[1]) -
+				FVec3::CrossProduct(Body(0).DQ(), PointConstraints.ConstraintArms[0]) ;
+	}
+}
+
 void FPBDJointCachedSolver::ApplyPointPositionConstraint()
 {
-	const FVec3 CX = PointConstraints.ConstraintCX + Body(1).DP()  - Body(0).DP() +
-		FVec3::CrossProduct(Body(1).DQ(), PointConstraints.ConstraintArms[1]) -
-			FVec3::CrossProduct(Body(0).DQ(), PointConstraints.ConstraintArms[0]) ;
+	const FVec3 CX = ComputePointCX();
 
 	if (CX.SizeSquared() > PositionTolerance * PositionTolerance)
 	{
 		// Calculate constraint correction
-		const FVec3 DX = SolverStiffness * PointConstraints.ConstraintStiffness * Utilities::Multiply(PointConstraints.ConstraintIM, CX);
+		const FVec3 DX = (bChaos_Joint_StablePointConstraint) ? SolverStiffness * PointConstraints.ConstraintStiffness * Utilities::Multiply(PointConstraints.ConstraintIM, CX) :
+																SolverStiffness * PointConstraints.ConstraintStiffness * PointConstraints.ConstraintIM.GetDiagonal() * CX;
+
+		ChaosSolverEnsure(!DX.ContainsNaN());
 
 		// Apply constraint correction
 		if( Body(0).IsDynamic())
@@ -576,14 +661,14 @@ void FPBDJointCachedSolver::SolvePositionConstraintSoft(
 	const FReal TargetVel)
 {
 	FReal VelDt = 0;
-	if (PositionConstraints.ConstraintDamping[ConstraintIndex] > KINDA_SMALL_NUMBER)
+	if (PositionConstraints.ConstraintSoftDamping[ConstraintIndex] > KINDA_SMALL_NUMBER)
 	{
 		const FVec3 V0Dt = FVec3::CalculateVelocity(InitConnectorXs[0], ConnectorXs[0]+Body(0).DP() + FVec3::CrossProduct(Body(0).DQ(), PositionConstraints.ConstraintArms[ConstraintIndex][0]), 1.0f);
 		const FVec3 V1Dt = FVec3::CalculateVelocity(InitConnectorXs[1], ConnectorXs[1]+Body(1).DP() + FVec3::CrossProduct(Body(1).DQ(), PositionConstraints.ConstraintArms[ConstraintIndex][1]), 1.0f);
 		VelDt = TargetVel * Dt + FVec3::DotProduct(V0Dt - V1Dt, PositionConstraints.ConstraintAxis[ConstraintIndex] );
 	}
 
-	const FReal DeltaLambda = SolverStiffness * (PositionConstraints.ConstraintSoftStiffness[ConstraintIndex] * DeltaConstraint - PositionConstraints.ConstraintDamping[ConstraintIndex] * VelDt - PositionConstraints.ConstraintLambda[ConstraintIndex]) /
+	const FReal DeltaLambda = SolverStiffness * (PositionConstraints.ConstraintSoftStiffness[ConstraintIndex] * DeltaConstraint - PositionConstraints.ConstraintSoftDamping[ConstraintIndex] * VelDt - PositionConstraints.ConstraintLambda[ConstraintIndex]) /
 		PositionConstraints.ConstraintSoftIM[ConstraintIndex];
 	PositionConstraints.ConstraintLambda[ConstraintIndex] += DeltaLambda;
 
@@ -596,6 +681,7 @@ void FPBDJointCachedSolver::ApplyAxisPositionConstraint(
 	const FVec3 CX = Body(1).DP()  - Body(0).DP() +
 		FVec3::CrossProduct(Body(1).DQ(), PositionConstraints.ConstraintArms[ConstraintIndex][1]) -
 			FVec3::CrossProduct(Body(0).DQ(), PositionConstraints.ConstraintArms[ConstraintIndex][0]) ;
+	
 	FReal DeltaPosition = PositionConstraints.ConstraintCX[ConstraintIndex] + FVec3::DotProduct(CX, PositionConstraints.ConstraintAxis[ConstraintIndex]);
 
 	bool NeedsSolve = false;
@@ -811,7 +897,13 @@ FPBDJointUtilities::GetSoftTwistStiffness(SolverSettings, JointSettings),
 	{
 		if (bAngularLimited[S1] && bAngularLimited[S2])
 		{
-			InitConeConstraint(JointSettings, Dt);
+			// When using non linear solver, the cone swing direction could change at each iteration
+			// stabilizing the solver. In the linear case we need to constraint along the 2 directions
+			// for better stability
+			InitSwingConstraint(JointSettings, SolverSettings, Dt, EJointAngularConstraintIndex::Swing1);
+			InitSwingConstraint(JointSettings, SolverSettings, Dt, EJointAngularConstraintIndex::Swing2);
+
+			//InitConeConstraint(JointSettings, Dt);
 		}
 		else if (bAngularLimited[S1] && bAngularLocked[S2])
 		{
@@ -862,6 +954,19 @@ FPBDJointUtilities::GetSoftTwistStiffness(SolverSettings, JointSettings),
 	}
 }
 
+void FPBDJointCachedSolver::InitRotationDatasMass(
+		FAxisConstraintDatas& RotationDatas,
+		const int32 ConstraintIndex,
+		const FReal Dt)
+{
+	const FVec3 IA0 = Utilities::Multiply(InvI(0), RotationDatas.ConstraintAxis[ConstraintIndex]);
+	const FVec3 IA1 = Utilities::Multiply(InvI(1), RotationDatas.ConstraintAxis[ConstraintIndex]);
+	const FReal II0 = FVec3::DotProduct(RotationDatas.ConstraintAxis[ConstraintIndex], IA0);
+	const FReal II1 = FVec3::DotProduct(RotationDatas.ConstraintAxis[ConstraintIndex], IA1);
+
+	RotationDatas.UpdateMass(ConstraintIndex,  IA0, IA1, II0 + II1, Dt);
+}
+
 void FPBDJointCachedSolver::InitRotationConstraintDatas(
 		const FPBDJointSettings& JointSettings,
 		const int32 ConstraintIndex,
@@ -873,18 +978,15 @@ void FPBDJointCachedSolver::InitRotationConstraintDatas(
 {
 	const FVec3 LocalAxis = (ConstraintAngle < 0.0f) ? -ConstraintAxis : ConstraintAxis;
 	const FReal LocalAngle = (ConstraintAngle < 0.0f) ? -ConstraintAngle : ConstraintAngle;
-	
-	const FVec3 IA0 = Utilities::Multiply(InvI(0), LocalAxis);
-	const FVec3 IA1 = Utilities::Multiply(InvI(1), LocalAxis);
-	const FReal II0 = FVec3::DotProduct(LocalAxis, IA0);
-	const FReal II1 = FVec3::DotProduct(LocalAxis, IA1);
 
-	RotationConstraints.UpdateDatas(ConstraintIndex, LocalAxis, LocalAngle, IA0, IA1, II0 + II1, ConstraintRestitution, Dt, bCheckLimit);
+	RotationConstraints.UpdateDatas(ConstraintIndex, LocalAxis, LocalAngle, ConstraintRestitution, bCheckLimit);
 
 	RotationConstraints.ConstraintLimits[ConstraintIndex] = FMath::Max(
 		JointSettings.AngularLimits[ConstraintIndex] - GetAngularConstraintPadding(ConstraintIndex), (FReal)0.);
 
 	InitConstraintAxisAngularVelocities[ConstraintIndex] = FVec3::DotProduct(W(1) - W(0), LocalAxis);
+
+	InitRotationDatasMass(RotationConstraints, ConstraintIndex, Dt);
 }
 
 void FPBDJointCachedSolver::CorrectAxisAngleConstraint(
@@ -920,7 +1022,9 @@ void FPBDJointCachedSolver::InitTwistConstraint(
 	FReal TwistAngle;
 	FPBDJointUtilities::GetTwistAxisAngle(ConnectorRs[0], ConnectorRs[1], TwistAxis, TwistAngle);
 
-	InitRotationConstraintDatas( JointSettings, (int32)EJointAngularConstraintIndex::Twist, TwistAxis, TwistAngle, JointSettings.TwistRestitution, Dt, true);
+	// Project the angle directly to avoid checking the limits during the solve.
+	CorrectAxisAngleConstraint(JointSettings,(int32)EJointAngularConstraintIndex::Twist, TwistAxis, TwistAngle);
+	InitRotationConstraintDatas( JointSettings, (int32)EJointAngularConstraintIndex::Twist, TwistAxis, TwistAngle, JointSettings.TwistRestitution, Dt, false);
 }
 
 void FPBDJointCachedSolver::InitConeConstraint(
@@ -929,6 +1033,7 @@ void FPBDJointCachedSolver::InitConeConstraint(
 {
 	FVec3 SwingAxisLocal;
 	FReal SwingAngle = 0.0f;
+
 	FPBDJointUtilities::GetEllipticalConeAxisErrorLocal(ConnectorRs[0], ConnectorRs[1], 0.0, 0.0, SwingAxisLocal, SwingAngle);
 	SwingAxisLocal.SafeNormalize();
 
@@ -955,6 +1060,7 @@ void FPBDJointCachedSolver::InitDualConeSwingConstraint(
 	const FPBDJointSettings& JointSettings,
 	const FReal Dt,
 	const EJointAngularConstraintIndex SwingConstraintIndex)
+	
 {
 	FVec3 SwingAxis;
 	FReal SwingAngle;
@@ -1054,7 +1160,7 @@ void FPBDJointCachedSolver::SolveRotationConstraintSoft(
 {
 	// Damping angular velocity
 	FReal AngVelDt = 0;
-	if (RotationConstraints.ConstraintDamping[ConstraintIndex] > KINDA_SMALL_NUMBER)
+	if (RotationConstraints.ConstraintSoftDamping[ConstraintIndex] > KINDA_SMALL_NUMBER)
 	{
 		const FVec3 W0Dt = FRotation3::CalculateAngularVelocity(InitConnectorRs[0], ConnectorRs[0], 1.0f);
 		const FVec3 W1Dt = FRotation3::CalculateAngularVelocity(InitConnectorRs[1], ConnectorRs[1], 1.0f);
@@ -1062,7 +1168,7 @@ void FPBDJointCachedSolver::SolveRotationConstraintSoft(
 	}
 
 	const FReal DeltaLambda = SolverStiffness * (RotationConstraints.ConstraintSoftStiffness[ConstraintIndex] * DeltaConstraint -
-		RotationConstraints.ConstraintDamping[ConstraintIndex] * AngVelDt - RotationConstraints.ConstraintLambda[ConstraintIndex]) /
+		RotationConstraints.ConstraintSoftDamping[ConstraintIndex] * AngVelDt - RotationConstraints.ConstraintLambda[ConstraintIndex]) /
 		RotationConstraints.ConstraintSoftIM[ConstraintIndex];
 	RotationConstraints.ConstraintLambda[ConstraintIndex] += DeltaLambda;
 
@@ -1088,10 +1194,11 @@ void FPBDJointCachedSolver::ApplyRotationConstraint(
 		{
 			DeltaAngle += RotationConstraints.ConstraintLimits[ConstraintIndex];
 			NeedsSolve = true;
+			
 		}
 	}
 
-	if (!RotationConstraints.bLimitsCheck[ConstraintIndex] || (RotationConstraints.bLimitsCheck[ConstraintIndex] && NeedsSolve && FMath::Abs(DeltaAngle) > AngleTolerance))
+	if (!RotationConstraints.bLimitsCheck[ConstraintIndex]|| (RotationConstraints.bLimitsCheck[ConstraintIndex] && NeedsSolve && FMath::Abs(DeltaAngle) > AngleTolerance))
 	{
 		if (RotationConstraints.bSoftLimit[ConstraintIndex])
 		{
@@ -1150,7 +1257,7 @@ void FPBDJointCachedSolver::ApplyAngularVelocityConstraint(const int32 Constrain
 	{
 		FReal TargetVel = 0.0f;
 		if (RotationConstraints.ConstraintRestitution[ConstraintIndex] != 0.0f)
-		{
+		{ 
 			const FReal InitVel = InitConstraintAxisAngularVelocities[ConstraintIndex];
 			TargetVel = InitVel > Chaos_Joint_AngularVelocityThresholdToApplyRestitution ?
 				-RotationConstraints.ConstraintRestitution[ConstraintIndex] * InitVel : 0.0f;
@@ -1197,7 +1304,7 @@ void FPBDJointCachedSolver::InitPositionDrives(
 						FPBDJointUtilities::GetLinearDriveDamping(SolverSettings, JointSettings, AxisIndex), 0.0f);
 					const FVec3 Axis = R0M.GetAxis(AxisIndex);
 				
-					if ((FMath::Abs(FVec3::DotProduct(CX,Axis)) > PositionTolerance) || (PositionDrives.ConstraintDamping[AxisIndex] > 0.0f))
+					if ((FMath::Abs(FVec3::DotProduct(CX,Axis)) > PositionTolerance) || (PositionDrives.ConstraintSoftDamping[AxisIndex] > 0.0f))
 					{
 						InitAxisPositionDrive(AxisIndex, Axis, CX, VTarget, Dt);
 					}
@@ -1217,16 +1324,11 @@ void FPBDJointCachedSolver::InitAxisPositionDrive(
 	const FVec3 ConstraintArm0 = ConnectorXs[0] - P(0);
 	const FVec3 ConstraintArm1 = ConnectorXs[1] - P(1);
 
-	const FVec3 AngularAxis0 = FVec3::CrossProduct(ConstraintArm0, ConstraintAxis);
-	const FVec3 AngularAxis1 = FVec3::CrossProduct(ConstraintArm1, ConstraintAxis);
-	const FVec3 IA0 = Utilities::Multiply(InvI(0), AngularAxis0);
-	const FVec3 IA1 = Utilities::Multiply(InvI(1), AngularAxis1);
-	const FReal II0 = FVec3::DotProduct(AngularAxis0, IA0);
-	const FReal II1 = FVec3::DotProduct(AngularAxis1, IA1);
-
 	PositionDrives.UpdateDatas(ConstraintIndex, ConstraintAxis, FVec3::DotProduct(DeltaPosition, ConstraintAxis),
-		IA0, IA1, InvM(0) + II0 + InvM(1) + II1, 0.0f, Dt, true, ConstraintArm0, ConstraintArm1,
+		0.0f,  true, ConstraintArm0, ConstraintArm1,
 		FVec3::DotProduct(DeltaVelocity, ConstraintAxis));
+
+	InitPositionDatasMass(PositionDrives, ConstraintIndex, Dt);
 }
 
 /** APPLY POSITION  DRIVES *********************************************************************************/
@@ -1253,7 +1355,7 @@ void FPBDJointCachedSolver::ApplyAxisPositionDrive(
 	const FReal DeltaPos = PositionDrives.ConstraintCX[ConstraintIndex] + FVec3::DotProduct(Delta1 - Delta0, PositionDrives.ConstraintAxis[ConstraintIndex]);
 
 	FReal VelDt = 0;
-	if (PositionDrives.ConstraintDamping[ConstraintIndex] > KINDA_SMALL_NUMBER)
+	if (PositionDrives.ConstraintSoftDamping[ConstraintIndex] > KINDA_SMALL_NUMBER)
 	{
 		const FVec3 V0Dt = FVec3::CalculateVelocity(InitConnectorXs[0], ConnectorXs[0]+ Delta0, 1.0f);
 		const FVec3 V1Dt = FVec3::CalculateVelocity(InitConnectorXs[1], ConnectorXs[1]+ Delta1, 1.0f);
@@ -1261,7 +1363,7 @@ void FPBDJointCachedSolver::ApplyAxisPositionDrive(
 	}
 
 	const FReal DeltaLambda = SolverStiffness * (PositionDrives.ConstraintSoftStiffness[ConstraintIndex] * DeltaPos -
-		PositionDrives.ConstraintDamping[ConstraintIndex] * VelDt - PositionDrives.ConstraintLambda[ConstraintIndex]) /
+		PositionDrives.ConstraintSoftDamping[ConstraintIndex] * VelDt - PositionDrives.ConstraintLambda[ConstraintIndex]) /
 		PositionDrives.ConstraintSoftIM[ConstraintIndex];
 	PositionDrives.ConstraintLambda[ConstraintIndex] += DeltaLambda;
 
@@ -1345,13 +1447,9 @@ void FPBDJointCachedSolver::InitRotationConstraintDrive(
 			const FReal Dt,
 			const FReal DeltaAngle)
 {
-	const FVec3 IA0 = Utilities::Multiply(InvI(0), ConstraintAxis);
-	const FVec3 IA1 = Utilities::Multiply(InvI(1), ConstraintAxis);
-	const FReal II0 = FVec3::DotProduct(ConstraintAxis, IA0);
-	const FReal II1 = FVec3::DotProduct(ConstraintAxis, IA1);
+	RotationDrives.UpdateDatas(ConstraintIndex, ConstraintAxis, DeltaAngle,0.0f);
 
-	RotationDrives.UpdateDatas(ConstraintIndex, ConstraintAxis, DeltaAngle,
-		IA0, IA1, II0 + II1, 0.0f, Dt);
+	InitRotationDatasMass(RotationDrives, ConstraintIndex, Dt);
 }
 
 void FPBDJointCachedSolver::InitSwingTwistDrives(
@@ -1385,21 +1483,21 @@ void FPBDJointCachedSolver::InitSwingTwistDrives(
 
 	RotationDrives.bAccelerationMode = FPBDJointUtilities::GetDriveAccelerationMode(SolverSettings, JointSettings);
 
-	const bool bUseTwistDrive = bTwistDriveEnabled && (((FMath::Abs(DTwistAngle) > AngleTolerance) && (RotationDrives.ConstraintSoftStiffness[TW] > 0.0f)) || (RotationDrives.ConstraintDamping[TW]  > 0.0f));
+	const bool bUseTwistDrive = bTwistDriveEnabled && (((FMath::Abs(DTwistAngle) > AngleTolerance) && (RotationDrives.ConstraintSoftStiffness[TW] > 0.0f)) || (RotationDrives.ConstraintSoftDamping[TW]  > 0.0f));
 	if (bUseTwistDrive)
 	{
 		InitRotationConstraintDrive(TW, ConnectorRs[1] * FJointConstants::TwistAxis(), Dt, DTwistAngle);
 		RotationDrives.ConstraintVX[TW] = JointSettings.AngularDriveVelocityTarget[TW];
 	}
 
-	const bool bUseSwing1Drive = bSwing1DriveEnabled && (((FMath::Abs(DSwing1Angle) > AngleTolerance) && (RotationDrives.ConstraintSoftStiffness[S1] > 0.0f)) || (RotationDrives.ConstraintDamping[S1] > 0.0f));
+	const bool bUseSwing1Drive = bSwing1DriveEnabled && (((FMath::Abs(DSwing1Angle) > AngleTolerance) && (RotationDrives.ConstraintSoftStiffness[S1] > 0.0f)) || (RotationDrives.ConstraintSoftDamping[S1] > 0.0f));
 	if (bUseSwing1Drive)
 	{
 		InitRotationConstraintDrive(S1, ConnectorRs[1] * FJointConstants::Swing1Axis(),  Dt, DSwing1Angle);
 		RotationDrives.ConstraintVX[S1] = JointSettings.AngularDriveVelocityTarget[S1];
 	}
 
-	const bool bUseSwing2Drive = bSwing2DriveEnabled && (((FMath::Abs(DSwing2Angle) > AngleTolerance) && (RotationDrives.ConstraintSoftStiffness[S2] > 0.0f)) || (RotationDrives.ConstraintDamping[S2] > 0.0f));
+	const bool bUseSwing2Drive = bSwing2DriveEnabled && (((FMath::Abs(DSwing2Angle) > AngleTolerance) && (RotationDrives.ConstraintSoftStiffness[S2] > 0.0f)) || (RotationDrives.ConstraintSoftDamping[S2] > 0.0f));
 	if (bUseSwing2Drive)
 	{
 		InitRotationConstraintDrive(S2, ConnectorRs[1] * FJointConstants::Swing2Axis(),  Dt, DSwing2Angle);
@@ -1420,7 +1518,7 @@ void FPBDJointCachedSolver::InitSLerpDrive(
 	RotationDrives.bAccelerationMode = FPBDJointUtilities::GetDriveAccelerationMode(SolverSettings, JointSettings);
 
 	// If damping is enabled, we need to apply the drive about all 3 axes, but without damping we can just drive along the axis of error
-	if (RotationDrives.ConstraintDamping[0]  > 0.0f)
+	if (RotationDrives.ConstraintSoftDamping[0]  > 0.0f)
 	{
 		// NOTE: Slerp target velocity only works properly if we have a stiffness of zero.
 		FVec3 Axes[3] = { FVec3(1, 0, 0), FVec3(0, 1, 0), FVec3(0, 0, 1) };
@@ -1495,7 +1593,7 @@ void FPBDJointCachedSolver::ApplyAxisRotationDrive(
 
 	// Damping angular velocity
 	FReal AngVelDt = 0;
-	if (RotationDrives.ConstraintDamping[ConstraintIndex] > KINDA_SMALL_NUMBER)
+	if (RotationDrives.ConstraintSoftDamping[ConstraintIndex] > KINDA_SMALL_NUMBER)
 	{
 		const FVec3 W0Dt = FRotation3::CalculateAngularVelocity(InitConnectorRs[0], ConnectorRs[0], 1.0f);
 		const FVec3 W1Dt = FRotation3::CalculateAngularVelocity(InitConnectorRs[1], ConnectorRs[1], 1.0f);
@@ -1503,7 +1601,7 @@ void FPBDJointCachedSolver::ApplyAxisRotationDrive(
 	}
 
 	const FReal DeltaLambda = SolverStiffness * (RotationDrives.ConstraintSoftStiffness[ConstraintIndex] * DeltaConstraint -
-		RotationDrives.ConstraintDamping[ConstraintIndex] * AngVelDt - RotationDrives.ConstraintLambda[ConstraintIndex]) /
+		RotationDrives.ConstraintSoftDamping[ConstraintIndex] * AngVelDt - RotationDrives.ConstraintLambda[ConstraintIndex]) /
 		RotationDrives.ConstraintSoftIM[ConstraintIndex];
 	RotationDrives.ConstraintLambda[ConstraintIndex] += DeltaLambda;
 
@@ -1605,9 +1703,11 @@ void FAxisConstraintDatas::InitDatas(
 	const FReal HardStiffness)
 {
 	bSoftLimit[ConstraintIndex] = bHasSoftLimits;
-	ConstraintSoftStiffness[ConstraintIndex] = SoftStiffness;
 	ConstraintHardStiffness[ConstraintIndex] = HardStiffness;
-	ConstraintDamping[ConstraintIndex] = bHasSoftLimits ? SoftDamping : 0.0;
+	ConstraintSoftStiffness[ConstraintIndex] = SoftStiffness;
+	ConstraintSoftDamping[ConstraintIndex] = SoftDamping;
+	SettingsSoftStiffness[ConstraintIndex] = SoftStiffness;
+	SettingsSoftDamping[ConstraintIndex] = SoftDamping;
 	bValidDatas[ConstraintIndex] = false;
 	bLimitsCheck[ConstraintIndex] = true;
 	ConstraintLambda = FVec3::Zero();
@@ -1619,11 +1719,7 @@ void FAxisConstraintDatas::UpdateDatas(
 	const int32 ConstraintIndex,
 	const FVec3& DatasAxis,
 	const FReal DatasCX,
-	const FVec3& DatasIA0,
-	const FVec3& DatasIA1,
-	const FReal DatasIM,
 	const FReal DatasRestitution,
-	const FReal Dt,
 	const bool bCheckLimit,
 	const FVec3& DatasArm0 ,
 	const FVec3& DatasArm1 ,
@@ -1635,19 +1731,29 @@ void FAxisConstraintDatas::UpdateDatas(
 	ConstraintCX[ConstraintIndex] = DatasCX;
 	ConstraintVX[ConstraintIndex] = DatasVX;
 	ConstraintAxis[ConstraintIndex] = DatasAxis;
-	ConstraintDRAxis[ConstraintIndex][0] = DatasIA0;
-	ConstraintDRAxis[ConstraintIndex][1] = -DatasIA1;
-	ConstraintHardIM[ConstraintIndex] = DatasIM;
 	ConstraintRestitution[ConstraintIndex] = DatasRestitution;
 	ConstraintArms[ConstraintIndex][0] = DatasArm0;
 	ConstraintArms[ConstraintIndex][1] = DatasArm1;
+}
+
+void FAxisConstraintDatas::UpdateMass(
+		const int32 ConstraintIndex,
+		const FVec3& DatasIA0,
+		const FVec3& DatasIA1,
+		const FReal DatasIM,
+		const FReal Dt)
+{
+	ConstraintDRAxis[ConstraintIndex][0] = DatasIA0;
+	ConstraintDRAxis[ConstraintIndex][1] = -DatasIA1;
+	ConstraintHardIM[ConstraintIndex] = DatasIM;
+	ConstraintLambda = FVec3::Zero();
 
 	if(bSoftLimit[ConstraintIndex])
 	{
 		const FReal SpringMassScale = (bAccelerationMode) ? (FReal)1 / (ConstraintHardIM[ConstraintIndex]) : (FReal)1;
-		ConstraintSoftStiffness[ConstraintIndex] = SpringMassScale * ConstraintSoftStiffness[ConstraintIndex] * Dt * Dt;
-		ConstraintDamping[ConstraintIndex] = SpringMassScale * ConstraintDamping[ConstraintIndex] * Dt;
-		ConstraintSoftIM[ConstraintIndex] = (ConstraintSoftStiffness[ConstraintIndex] + ConstraintDamping[ConstraintIndex]) * ConstraintHardIM[ConstraintIndex] + (FReal)1;
+		ConstraintSoftStiffness[ConstraintIndex] = SpringMassScale * SettingsSoftStiffness[ConstraintIndex] * Dt * Dt;
+		ConstraintSoftDamping[ConstraintIndex] = SpringMassScale * SettingsSoftDamping[ConstraintIndex] * Dt;
+		ConstraintSoftIM[ConstraintIndex] = (ConstraintSoftStiffness[ConstraintIndex] + ConstraintSoftDamping[ConstraintIndex]) * ConstraintHardIM[ConstraintIndex] + (FReal)1;
 	}
 }
 
@@ -1661,13 +1767,11 @@ void FPointConstraintDatas::InitDatas(
 
 void FPointConstraintDatas::UpdateDatas(
 	const FVec3& DatasCX,
-	const FMatrix33& DatasIM,
 	const FVec3& DatasArm0,
 	const FVec3& DatasArm1)
 {
 	bValidDatas= true;
 	ConstraintCX = DatasCX;
-	ConstraintIM = DatasIM;
 	ConstraintArms[0] = DatasArm0;
 	ConstraintArms[1] = DatasArm1;
 }
diff --git a/Engine/Source/Runtime/Experimental/Chaos/Private/Chaos/PBDJointConstraints.cpp b/Engine/Source/Runtime/Experimental/Chaos/Private/Chaos/PBDJointConstraints.cpp
--- a/Engine/Source/Runtime/Experimental/Chaos/Private/Chaos/PBDJointConstraints.cpp
+++ b/Engine/Source/Runtime/Experimental/Chaos/Private/Chaos/PBDJointConstraints.cpp
@@ -1150,7 +1150,7 @@ namespace Chaos
 			// Set parent inverser mass scale based on current shock propagation state
 			// @todo(chaos): needs to handle the parent/child being in opposite order
 			const FReal ShockPropagationInvMassScale = FPBDJointUtilities::GetShockPropagationInvMassScale(SolverType, It, NumIts, Settings, JointSettings);
-			Solver.SetInvMassScales(ShockPropagationInvMassScale, FReal(1));
+			Solver.SetInvMassScales(ShockPropagationInvMassScale, FReal(1), Dt);
 
 			if (!bWasActive && !Solver.GetIsActive() && bChaos_Joint_EarlyOut_Enabled)
 			{
@@ -1285,7 +1285,7 @@ namespace Chaos
 			// Set parent inverser mass scale based on current shock propagation state
 			// @todo(chaos): needs to handle the parent/child being in opposite order
 			const FReal ShockPropagationInvMassScale = FPBDJointUtilities::GetShockPropagationInvMassScale(SolverType, It, NumIts, Settings, JointSettings);
-			Solver.SetInvMassScales(ShockPropagationInvMassScale, FReal(1));
+			Solver.SetInvMassScales(ShockPropagationInvMassScale, FReal(1), Dt);
 
 			// For quasipbd, we may still need to stabilize velocities and solve for restitution.
 			if (!bWasActive && !Solver.GetIsActive() && SolverType != EConstraintSolverType::QuasiPbd && bChaos_Joint_EarlyOut_Enabled)
diff --git a/Engine/Source/Runtime/Experimental/Chaos/Public/Chaos/Joint/PBDJointCachedSolverGaussSeidel.h b/Engine/Source/Runtime/Experimental/Chaos/Public/Chaos/Joint/PBDJointCachedSolverGaussSeidel.h
--- a/Engine/Source/Runtime/Experimental/Chaos/Public/Chaos/Joint/PBDJointCachedSolverGaussSeidel.h
+++ b/Engine/Source/Runtime/Experimental/Chaos/Public/Chaos/Joint/PBDJointCachedSolverGaussSeidel.h
@@ -29,19 +29,27 @@ struct FAxisConstraintDatas
 		const int32 ConstraintIndex,
 		const FVec3& DatasAxis,
 		const FReal DatasCX,
-		const FVec3& DatasIA0,
-		const FVec3& DatasIA1,
-		const FReal DatasIM,
 		const FReal DatasRestitution,
-		const FReal Dt,
 		const bool bCheckLimit = true,
 		const FVec3& DatasArm0 = FVec3::Zero(),
 		const FVec3& DatasArm1 = FVec3::Zero(),
 		const FReal DatasVX = 0.0);
+
+	/** Update the mass dependent datas */
+	void UpdateMass(
+		const int32 ConstraintIndex,
+		const FVec3& DatasIA0,
+		const FVec3& DatasIA1,
+		const FReal DatasIM,
+		const FReal Dt);
 	
-	FVec3 ConstraintSoftStiffness;
 	FVec3 ConstraintHardStiffness;
-	FVec3 ConstraintDamping;
+	FVec3 ConstraintSoftStiffness;
+	FVec3 ConstraintSoftDamping;
+	
+	FVec3 SettingsSoftDamping;
+	FVec3 SettingsSoftStiffness;
+	
 	FVec3 ConstraintArms[3][2];
 	FVec3 ConstraintAxis[3];
 	FVec3 ConstraintLimits;
@@ -71,11 +79,10 @@ struct FPointConstraintDatas
 	/** Update the point joint datas with axis, limits, arms... */
 	void UpdateDatas(
 		const FVec3& DatasCX,
-		const FMatrix33& DatasIM,
 		const FVec3& DatasArm0,
 		const FVec3& DatasArm1);
 	
-	bool bValidDatas;
+	bool bValidDatas = false;
 	
 	FReal ConstraintStiffness;
 	FVec3 ConstraintArms[2];
@@ -251,7 +258,8 @@ struct FPointConstraintDatas
 
 		void SetInvMassScales(
 			const FReal InvMScale0,
-			const FReal InvMScale1);
+			const FReal InvMScale1,
+			const FReal Dt);
 
 	private:
 
@@ -288,6 +296,8 @@ struct FPointConstraintDatas
 			const FPBDJointSettings& JointSettings);
 		
 		void InitPointPositionConstraint();
+
+		void InitPointPositionMass();
 		
 		void InitPositionConstraintDatas(
 			const int32 ConstraintIndex,
@@ -298,6 +308,11 @@ struct FPointConstraintDatas
 			const FReal ConstraintLimit,
 			const EJointMotionType JointType);
 
+		void InitPositionDatasMass(
+			FAxisConstraintDatas& PositionDatas,
+			const int32 ConstraintIndex,
+			const FReal Dt);
+
 		void InitSphericalPositionConstraint(
 			const FPBDJointSettings& JointSettings,
 			const FReal Dt);
@@ -322,6 +337,8 @@ struct FPointConstraintDatas
 		
 		void ApplyPointPositionConstraint();
 		
+		FVec3 ComputePointCX();
+		
 		void ApplyAxisPositionConstraint(
 			const int32 ConstraintIndex,
 			const FReal Dt);
@@ -374,6 +391,11 @@ struct FPointConstraintDatas
 			const FReal ConstraintRestitution,
 			const FReal Dt,
 			const bool bCheckLimit);
+
+		void InitRotationDatasMass(
+			FAxisConstraintDatas& RotationDatas, 
+			const int32 ConstraintIndex,
+			const FReal Dt);
 		
 		void InitTwistConstraint(
 			const FPBDJointSettings& JointSettings,
diff --git a/Engine/Source/Runtime/Experimental/Chaos/Public/Chaos/Utilities.h b/Engine/Source/Runtime/Experimental/Chaos/Public/Chaos/Utilities.h
--- a/Engine/Source/Runtime/Experimental/Chaos/Public/Chaos/Utilities.h
+++ b/Engine/Source/Runtime/Experimental/Chaos/Public/Chaos/Utilities.h
@@ -304,6 +304,21 @@ namespace Chaos
 				-V[1] * (V[2] * M.M[0][0] - V[0] * M.M[2][0]) + V[0] * (V[2] * M.M[1][0] - V[0] * M.M[2][1]),
 				-V[1] * (-V[1] * M.M[0][0] + V[0] * M.M[1][0]) + V[0] * (-V[1] * M.M[1][0] + V[0] * M.M[1][1]) + Im);
 		}
+		
+		/**
+		 * Calculate the matrix diagonal that maps a constraint position error to constraint position and rotation corrections.
+		*/
+		template<class T>
+		TVec3<T> ComputeDiagonalJointFactorMatrix(const TVec3<T>& V, const PMatrix<T, 3, 3>& M, const T& Im)
+		{
+			// Rigid objects rotational contribution to the impulse.
+			// Vx*M*VxT+Im
+			check(Im > FLT_MIN);
+			return TVec3<T>(
+				-V[2] * (-V[2] * M.M[1][1] + V[1] * M.M[2][1]) + V[1] * (-V[2] * M.M[2][1] + V[1] * M.M[2][2]) + Im,
+				 V[2] * (V[2] * M.M[0][0] - V[0] * M.M[2][0]) - V[0] * (V[2] * M.M[2][0] - V[0] * M.M[2][2]) + Im,
+				-V[1] * (-V[1] * M.M[0][0] + V[0] * M.M[1][0]) + V[0] * (-V[1] * M.M[1][0] + V[0] * M.M[1][1]) + Im);
+		}
 
 		/**
 		 * Detects intersections between 2D line segments, returns intersection results as times along each line segment - these times can be
