Chaos - fix heightfield missing collisions (over pruning because it wasn't updated to use the new triangle collision)

diff --git a/Engine/Source/Runtime/Experimental/Chaos/Private/Chaos/CollisionOneShotManifolds.cpp b/Engine/Source/Runtime/Experimental/Chaos/Private/Chaos/CollisionOneShotManifolds.cpp
--- a/Engine/Source/Runtime/Experimental/Chaos/Private/Chaos/CollisionOneShotManifolds.cpp
+++ b/Engine/Source/Runtime/Experimental/Chaos/Private/Chaos/CollisionOneShotManifolds.cpp
@@ -1539,6 +1539,20 @@ namespace Chaos
 				const FTriangle& Triangle,
 				const FReal CullDistance,
 				TCArray<FContactPoint, 4>& OutContactPoints);
+
+		template
+			void ConstructPlanarConvexTriangleOneShotManifold(
+				const TImplicitObjectScaled<FImplicitBox3>& Convex,
+				const FTriangle& Triangle,
+				const FReal CullDistance,
+				TCArray<FContactPoint, 4>& OutContactPoints);
+
+		template
+			void ConstructPlanarConvexTriangleOneShotManifold(
+				const TImplicitObjectInstanced<FImplicitBox3>& Convex,
+				const FTriangle& Triangle,
+				const FReal CullDistance,
+				TCArray<FContactPoint, 4>& OutContactPoints);
 	}
 }
 
diff --git a/Engine/Source/Runtime/Experimental/Chaos/Private/Chaos/HeightField.cpp b/Engine/Source/Runtime/Experimental/Chaos/Private/Chaos/HeightField.cpp
--- a/Engine/Source/Runtime/Experimental/Chaos/Private/Chaos/HeightField.cpp
+++ b/Engine/Source/Runtime/Experimental/Chaos/Private/Chaos/HeightField.cpp
@@ -1385,7 +1385,7 @@ namespace Chaos
 	}
 
 	template <typename GeomType>
-	bool FHeightField::ContactManifoldImp(const GeomType& QueryGeom, const FRigidTransform3& QueryTM, const FReal Thickness, TArray<FContactPoint>& ContactPoints) const
+	bool FHeightField::ContactManifoldNonPlanarConvexImp(const GeomType& QueryGeom, const FRigidTransform3& QueryTM, const FReal Thickness, TArray<FContactPoint>& ContactPoints) const
 	{
 		auto OverlapTriangle = [&](const FVec3& A, const FVec3& B, const FVec3& C,
 			FPBDCollisionConstraint& Constraint)
@@ -1404,9 +1404,9 @@ namespace Chaos
 			int32 PointIndex = 0;
 			bool done = false;
 			const FReal SamePointErrorMarginSqr = 0.01f;
-			
+
 			int32 ContactPointsNum = ContactPoints.Num();
-			for( ; PointIndex < ContactPointsNum; PointIndex++ )
+			for (; PointIndex < ContactPointsNum; PointIndex++)
 			{
 				FVec3 DiffVector = ContactPoint.ShapeContactPoints[1] - ContactPoints[PointIndex].ShapeContactPoints[1];
 				// Check if point is the same (or close)
@@ -1428,8 +1428,8 @@ namespace Chaos
 			{
 				ContactPoints.Add(ContactPoint);
 			}
-		};		
-		
+		};
+
 		bool bResult = false;
 		FAABB3 QueryBounds = QueryGeom.BoundingBox();
 		QueryBounds.Thicken(Thickness);
@@ -1468,6 +1468,7 @@ namespace Chaos
 				// First Triangle
 				{
 					Constraint.ResetManifold();
+					Constraint.GetGJKWarmStartData().Reset();
 					OverlapTriangle(Points[0], Points[1], Points[3], Constraint);
 					for (FManifoldPoint& ManifoldPoint : Constraint.GetManifoldPoints())
 					{
@@ -1478,6 +1479,7 @@ namespace Chaos
 				// Second Triangle
 				{
 					Constraint.ResetManifold();
+					Constraint.GetGJKWarmStartData().Reset();
 					OverlapTriangle(Points[0], Points[3], Points[2], Constraint);
 					for (FManifoldPoint& ManifoldPoint : Constraint.GetManifoldPoints())
 					{
@@ -1515,6 +1517,142 @@ namespace Chaos
 		return true;
 	}
 
+	template <typename GeomType>
+	bool FHeightField::ContactManifoldPlanarConvexImp(const GeomType& QueryGeom, const FRigidTransform3& QueryTM, const FReal Thickness, TArray<FContactPoint>& ContactPoints) const
+	{
+		auto OverlapTriangle = [&](const FVec3& A, const FVec3& B, const FVec3& C, TCArray<FContactPoint, 4>& TriangleContactPoints)
+		{
+			// Create triangle in query space
+			FTriangle TriangleConvex(
+				QueryTM.InverseTransformPositionNoScale(A), 
+				QueryTM.InverseTransformPositionNoScale(B), 
+				QueryTM.InverseTransformPositionNoScale(C));
+
+			Collisions::ConstructPlanarConvexTriangleOneShotManifold(QueryGeom, TriangleConvex, Thickness, TriangleContactPoints);
+
+			// Convert back to shape-local space
+			for (FContactPoint& ContactPoint : TriangleContactPoints)
+			{
+				ContactPoint.ShapeContactPoints[1] = QueryTM.TransformPositionNoScale(ContactPoint.ShapeContactPoints[1]);
+				ContactPoint.ShapeContactNormal = QueryTM.TransformVectorNoScale(ContactPoint.ShapeContactNormal);
+			}
+		};
+
+		auto InsertSorted = [&](const FContactPoint& ContactPoint)
+		{
+			int32 PointIndex = 0;
+			bool done = false;
+			const FReal SamePointErrorMarginSqr = 0.01f;
+			
+			int32 ContactPointsNum = ContactPoints.Num();
+			for( ; PointIndex < ContactPointsNum; PointIndex++ )
+			{
+				FVec3 DiffVector = ContactPoint.ShapeContactPoints[1] - ContactPoints[PointIndex].ShapeContactPoints[1];
+				// Check if point is the same (or close)
+				if (DiffVector.SizeSquared() < SamePointErrorMarginSqr)
+				{
+					done = true;
+					break;
+				}
+
+				if (ContactPoint.Phi < ContactPoints[PointIndex].Phi)
+				{
+					ContactPoints.Insert(ContactPoint, PointIndex);
+					done = true;
+					break;
+				}
+			}
+
+			if (!done)
+			{
+				ContactPoints.Add(ContactPoint);
+			}
+		};		
+		
+		bool bResult = false;
+		FAABB3 QueryBounds = QueryGeom.BoundingBox();
+		QueryBounds.Thicken(Thickness);
+		QueryBounds = QueryBounds.TransformedAABB(QueryTM);
+
+		FBounds2D FlatQueryBounds;
+		FlatQueryBounds.Min = FVec2(QueryBounds.Min()[0], QueryBounds.Min()[1]);
+		FlatQueryBounds.Max = FVec2(QueryBounds.Max()[0], QueryBounds.Max()[1]);
+
+		TArray<TVec2<int32>> Intersections;
+		FVec3 Points[4];
+
+		GetGridIntersections(FlatQueryBounds, Intersections);
+
+		FReal MinContactPhi = FLT_MAX;
+		FVec3 LocalContactLocation, LocalContactNormal;
+		TCArray<FContactPoint, 4> TriangleContactPoints;
+
+		for (const TVec2<int32>& Cell : Intersections)
+		{
+			const int32 SingleIndex = Cell[1] * GeomData.NumCols + Cell[0];
+			const int32 CellIndex = Cell[1] * (GeomData.NumCols - 1) + Cell[0];
+
+			// Check for holes and skip checking if we'll never collide
+			if (GeomData.MaterialIndices.IsValidIndex(CellIndex) && GeomData.MaterialIndices[CellIndex] == TNumericLimits<uint8>::Max())
+			{
+				continue;
+			}
+
+			// The triangle is solid so proceed to test it
+			FAABB3 CellBounds;
+			GeomData.GetPointsAndBoundsScaled(SingleIndex, Points, CellBounds);
+			if (CellBounds.Intersects(QueryBounds))
+			{
+				// First Triangle
+				{
+					TriangleContactPoints.Reset();
+					OverlapTriangle(Points[0], Points[1], Points[3], TriangleContactPoints);
+					for (FContactPoint& ContactPoint : TriangleContactPoints)
+					{
+						ContactPoint.FaceIndex = CellIndex * 2;
+						InsertSorted(ContactPoint);
+					}
+				}
+				// Second Triangle
+				{
+					TriangleContactPoints.Reset();
+					OverlapTriangle(Points[0], Points[3], Points[2], TriangleContactPoints);
+					for (FContactPoint& ContactPoint : TriangleContactPoints)
+					{
+						ContactPoint.FaceIndex = CellIndex * 2 + 1;
+						InsertSorted(ContactPoint);
+					}
+				}
+			}
+		}
+
+		// Remove edge contacts that are "hidden" by face contacts
+		// EdgePruneDistance should be some fraction of the convex margin...
+		const FReal EdgePruneDistance = Chaos_Collision_EdgePrunePlaneDistance;
+		Collisions::PruneEdgeContactPointsOrdered(ContactPoints, EdgePruneDistance);
+
+		// Remove all points (except for the deepest one, and ones with phis similar to it)
+		const FReal CullMargin = 0.1f;
+		int32 NewContactPointCount = ContactPoints.Num() > 0 ? 1 : 0;
+		for (int32 Index = 1; Index < ContactPoints.Num(); Index++)
+		{
+			if (ContactPoints[Index].Phi < 0 || ContactPoints[Index].Phi - ContactPoints[0].Phi < CullMargin)
+			{
+				NewContactPointCount++;
+			}
+			else
+			{
+				break;
+			}
+		}
+		ContactPoints.SetNum(NewContactPointCount, false);
+
+		// Reduce to only 4 contact points from here
+		Collisions::ReduceManifoldContactPointsTriangeMesh(ContactPoints);
+
+		return true;
+	}
+
 	template <typename GeomType>
 	bool FHeightField::GJKContactPointImp(const GeomType& QueryGeom, const FRigidTransform3& QueryTM, const FReal Thickness,
 		FVec3& ContactLocation, FVec3& ContactNormal, FReal& ContactPhi) const
@@ -1700,7 +1838,7 @@ namespace Chaos
 
 	bool FHeightField::ContactManifold(const TBox<FReal, 3>& QueryGeom, const FRigidTransform3& QueryTM, const FReal Thickness, TArray<FContactPoint>& ContactPoints) const
 	{
-		return ContactManifoldImp(QueryGeom, QueryTM, Thickness, ContactPoints);
+		return ContactManifoldPlanarConvexImp(QueryGeom, QueryTM, Thickness, ContactPoints);
 	}
 
 	/*bool FHeightField::ContactManifold(const TSphere<FReal, 3>& QueryGeom, const FRigidTransform3& QueryTM, const FReal Thickness, TArray<FContactPoint>& ContactPoints) const
@@ -1710,27 +1848,27 @@ namespace Chaos
 
 	bool FHeightField::ContactManifold(const FCapsule& QueryGeom, const FRigidTransform3& QueryTM, const FReal Thickness, TArray<FContactPoint>& ContactPoints) const
 	{
-		return ContactManifoldImp(QueryGeom, QueryTM, Thickness, ContactPoints);
+		return ContactManifoldNonPlanarConvexImp(QueryGeom, QueryTM, Thickness, ContactPoints);
 	}
 
 	bool FHeightField::ContactManifold(const FConvex& QueryGeom, const FRigidTransform3& QueryTM, const FReal Thickness, TArray<FContactPoint>& ContactPoints) const
 	{
-		return ContactManifoldImp(QueryGeom, QueryTM, Thickness, ContactPoints);
+		return ContactManifoldPlanarConvexImp(QueryGeom, QueryTM, Thickness, ContactPoints);
 	}
 
 	bool FHeightField::ContactManifold(const TImplicitObjectScaled<TBox<FReal, 3>>& QueryGeom, const FRigidTransform3& QueryTM, const FReal Thickness, TArray<FContactPoint>& ContactPoints) const
 	{
-		return ContactManifoldImp(QueryGeom, QueryTM, Thickness, ContactPoints);
+		return ContactManifoldPlanarConvexImp(QueryGeom, QueryTM, Thickness, ContactPoints);
 	}
 
 	bool FHeightField::ContactManifold(const TImplicitObjectScaled<FCapsule>& QueryGeom, const FRigidTransform3& QueryTM, const FReal Thickness, TArray<FContactPoint>& ContactPoints) const
 	{
-		return ContactManifoldImp(QueryGeom, QueryTM, Thickness, ContactPoints);
+		return ContactManifoldNonPlanarConvexImp(QueryGeom, QueryTM, Thickness, ContactPoints);
 	}
 
 	bool FHeightField::ContactManifold(const TImplicitObjectScaled<FConvex>& QueryGeom, const FRigidTransform3& QueryTM, const FReal Thickness, TArray<FContactPoint>& ContactPoints) const
 	{
-		return ContactManifoldImp(QueryGeom, QueryTM, Thickness, ContactPoints);
+		return ContactManifoldPlanarConvexImp(QueryGeom, QueryTM, Thickness, ContactPoints);
 	}
 
 	template <typename QueryGeomType>
diff --git a/Engine/Source/Runtime/Experimental/Chaos/Public/Chaos/HeightField.h b/Engine/Source/Runtime/Experimental/Chaos/Public/Chaos/HeightField.h
--- a/Engine/Source/Runtime/Experimental/Chaos/Public/Chaos/HeightField.h
+++ b/Engine/Source/Runtime/Experimental/Chaos/Public/Chaos/HeightField.h
@@ -644,8 +644,14 @@ namespace Chaos
 		template <typename GeomType>
 		bool GJKContactPointImp(const GeomType& QueryGeom, const FRigidTransform3& QueryTM, const FReal Thickness, FVec3& ContactLocation, FVec3& ContactNormal, FReal& ContactPhi) const;
 
+		// @todo(chaos): remove ContactManifoldNonPlanarConvexImp and ContactManifoldPlanarConvexImp and move collision detection code out of this class.
+		// Instead, provide an API for generating triangles in a query bounds in some space.
+
+		template <typename GeomType>
+		bool ContactManifoldNonPlanarConvexImp(const GeomType& QueryGeom, const FRigidTransform3& QueryTM, const FReal Thickness, TArray<FContactPoint>& ContactPoints) const;
+
 		template <typename GeomType>
-		bool ContactManifoldImp(const GeomType& QueryGeom, const FRigidTransform3& QueryTM, const FReal Thickness, TArray<FContactPoint>& ContactPoints) const;
+		bool ContactManifoldPlanarConvexImp(const GeomType& QueryGeom, const FRigidTransform3& QueryTM, const FReal Thickness, TArray<FContactPoint>& ContactPoints) const;
 
 	};
 }
