When importing is complete, ensure that all imported objects and their subobjects no longer have the Async internal flag set. This will be the case if an Interchange factory created subobjects as inners of the asset object on a worker thread.

diff --git a/Engine/Source/Runtime/Experimental/Interchange/Engine/Private/Tasks/InterchangeTaskCompletion.cpp b/Engine/Source/Runtime/Experimental/Interchange/Engine/Private/Tasks/InterchangeTaskCompletion.cpp
--- a/Engine/Source/Runtime/Experimental/Interchange/Engine/Private/Tasks/InterchangeTaskCompletion.cpp
+++ b/Engine/Source/Runtime/Experimental/Interchange/Engine/Private/Tasks/InterchangeTaskCompletion.cpp
@@ -14,6 +14,7 @@
 #include "Templates/SharedPointer.h"
 #include "UObject/GarbageCollection.h"
 #include "UObject/Object.h"
+#include "UObject/UObjectHash.h"
 #include "UObject/ObjectMacros.h"
 #include "UObject/WeakObjectPtrTemplates.h"
 
@@ -81,9 +82,18 @@ void UE::Interchange::FTaskPreCompletion::DoTask(ENamedThreads::Type CurrentThre
 				Message->DestinationAssetName = ImportedObject->GetPathName();
 				Message->AssetType = ImportedObject->GetClass();
 
-				//Clear any async flag from the created asset
+				//Clear any async flag from the created asset and all its subobjects
 				const EInternalObjectFlags AsyncFlags = EInternalObjectFlags::Async | EInternalObjectFlags::AsyncLoading;
 				ImportedObject->ClearInternalFlags(AsyncFlags);
+
+				TArray<UObject*> ImportedSubobjects;
+				const bool bIncludeNestedObjects = true;
+				GetObjectsWithOuter(ImportedObject, ImportedSubobjects, bIncludeNestedObjects);
+				for (UObject* ImportedSubobject : ImportedSubobjects)
+				{
+					ImportedSubobject->ClearInternalFlags(AsyncFlags);
+				}
+
 				//Make sure the package is dirty
 				ImportedObject->MarkPackageDirty();
 
