Changed a check to a checkf to shed more light on a Chaos assert

diff --git a/Engine/Source/Runtime/Experimental/Chaos/Private/Chaos/Island/IslandManager.cpp b/Engine/Source/Runtime/Experimental/Chaos/Private/Chaos/Island/IslandManager.cpp
--- a/Engine/Source/Runtime/Experimental/Chaos/Private/Chaos/Island/IslandManager.cpp
+++ b/Engine/Source/Runtime/Experimental/Chaos/Private/Chaos/Island/IslandManager.cpp
@@ -697,7 +697,7 @@ void FPBDIslandManager::BuildGroups(const int32 NumContainers)
 	{
 		if( FPBDIslandSolver* IslandSolver = IslandSolvers[SortedIndex].Get())
 		{
-			check(GroupIndex < NumGroups);
+			checkf(GroupIndex < NumGroups, TEXT("Index %d is less than NumGroups %d"), GroupIndex, NumGroups);
 			IslandGroups[GroupIndex]->AddIsland(IslandSolver);
 			IslandGroups[GroupIndex]->NumParticles() += IslandSolver->NumParticles();
 			IslandGroups[GroupIndex]->NumConstraints() += IslandSolver->NumConstraints();
