Chaos : Fix collision events not swapping info for swapped order contact pairs (mirrors 19012674).

diff --git a/Engine/Source/Runtime/Engine/Private/PhysicsEngine/Experimental/PhysScene_Chaos.cpp b/Engine/Source/Runtime/Engine/Private/PhysicsEngine/Experimental/PhysScene_Chaos.cpp
--- a/Engine/Source/Runtime/Engine/Private/PhysicsEngine/Experimental/PhysScene_Chaos.cpp
+++ b/Engine/Source/Runtime/Engine/Private/PhysicsEngine/Experimental/PhysScene_Chaos.cpp
@@ -781,6 +781,12 @@ void FPhysScene_Chaos::HandleCollisionEvents(const Chaos::FCollisionEventData& E
 								
 								NewContact.PhysMaterial[0] = InternalMat1 ? FPhysicsUserData::Get<UPhysicalMaterial>(InternalMat1->UserData) : nullptr;
 								NewContact.PhysMaterial[1] = InternalMat2 ? FPhysicsUserData::Get<UPhysicalMaterial>(InternalMat2->UserData) : nullptr;
+
+
+								if (bSwapOrder)
+								{
+									NotifyInfo.RigidCollisionData.SwapContactOrders();
+								}
 							}
 						}
 					}
diff --git a/Engine/Source/Runtime/Experimental/ChaosSolverEngine/Private/Chaos/ChaosGameplayEventDispatcher.cpp b/Engine/Source/Runtime/Experimental/ChaosSolverEngine/Private/Chaos/ChaosGameplayEventDispatcher.cpp
--- a/Engine/Source/Runtime/Experimental/ChaosSolverEngine/Private/Chaos/ChaosGameplayEventDispatcher.cpp
+++ b/Engine/Source/Runtime/Experimental/ChaosSolverEngine/Private/Chaos/ChaosGameplayEventDispatcher.cpp
@@ -372,6 +372,12 @@ void UChaosGameplayEventDispatcher::HandleCollisionEvents(const Chaos::FCollisio
 										NewContact.ContactPosition = CollisionDataItem.Location;
 										NewContact.ContactPenetration = CollisionDataItem.PenetrationDepth;
 										// NewContact.PhysMaterial[1] UPhysicalMaterial required here
+
+
+										if (bSwapOrder)
+										{
+											NotifyInfo.RigidCollisionData.SwapContactOrders();
+										}
 									}
 
 								}
