Chaos : fix CCD Convex Convex  swept constraint generation not accounting for the local shapes transform
Fix the case where Particle0 is static.
Fix the stickiness issue for CCD objects

modified from 18963998 by cedric.caillaud

diff --git a/Engine/Source/Runtime/Experimental/Chaos/Private/Chaos/CCDUtilities.cpp b/Engine/Source/Runtime/Experimental/Chaos/Private/Chaos/CCDUtilities.cpp
--- a/Engine/Source/Runtime/Experimental/Chaos/Private/Chaos/CCDUtilities.cpp
+++ b/Engine/Source/Runtime/Experimental/Chaos/Private/Chaos/CCDUtilities.cpp
@@ -41,7 +41,7 @@ namespace Chaos
 
 			return 0;
 		}
-		return TNumericLimits<FReal>::Max();
+		return 0;
 	}
 
 	int32 FCCDConstraint::GetFastMovingKinematicIndex(const FPBDCollisionConstraint* Constraint, const FVec3 Displacements[]) const
@@ -433,6 +433,14 @@ namespace Chaos
 							*	When resweeping, compute estimated TOI instead of accurate TOI since updated TOI might need to be updated again.
 							*/
 
+							// We need to update cached position at TOI=1 in collision constraint.
+							FPBDCollisionConstraint* SweptConstraint = SortedCCDConstraints[i]->SweptConstraint;
+							FRigidTransform3 ShapeWorldTransform0 = SweptConstraint->GetShapeWorldTransform0();
+							FRigidTransform3 ShapeWorldTransform1 = SweptConstraint->GetShapeWorldTransform1();
+							ShapeWorldTransform0.SetTranslation(FConstGenericParticleHandle(SweptConstraint->GetParticle0())->P());
+							ShapeWorldTransform1.SetTranslation(FConstGenericParticleHandle(SweptConstraint->GetParticle1())->P());
+
+							SweptConstraint->SetShapeWorldTransforms(ShapeWorldTransform0, ShapeWorldTransform1);
 							const bool bUpdated = Collisions::UpdateConstraintFromGeometrySwept<ECollisionUpdateType::Deepest>(*(AttachedCCDConstraint->SweptConstraint), RigidTransforms[0], RigidTransforms[1], RestDt);
 							if (bUpdated)
 							{
@@ -454,19 +462,6 @@ namespace Chaos
 				}
 			}
 		}
-
-		// We need to update the world-space contact points at the final locations
-		for (int32 i = ConstraintStart; i < ConstraintEnd; i++)
-		{
-			FPBDCollisionConstraint* Constraint = SortedCCDConstraints[i]->SweptConstraint;
-			FRigidTransform3 ShapeWorldTransform0 = Constraint->GetShapeWorldTransform0();
-			FRigidTransform3 ShapeWorldTransform1 = Constraint->GetShapeWorldTransform1();
-			ShapeWorldTransform0.SetTranslation(FConstGenericParticleHandle(Constraint->GetParticle0())->P());
-			ShapeWorldTransform1.SetTranslation(FConstGenericParticleHandle(Constraint->GetParticle1())->P());
-
-			Constraint->SetShapeWorldTransforms(ShapeWorldTransform0, ShapeWorldTransform1);
-			Constraint->UpdateManifoldContacts();
-		}
 	}
 
 
@@ -559,24 +554,17 @@ namespace Chaos
 
 	void FCCDManager::UpdateSweptConstraints(const FReal Dt, FCollisionConstraintAllocator *CollisionAllocator)
 	{
+		// We need to update the world-space contact points at the final locations
+		// @todo(chaos): parallelize this code
 		for (FPBDCollisionConstraint* SweptConstraint : SweptConstraints)
 		{
-			FRigidTransform3 RigidTransforms[2];
-			for (int32 i = 0; i < 2; i++)
-			{
-				FGenericParticleHandle Particle = FGenericParticleHandle(SweptConstraint->GetParticle(i));
-				const bool IsStatic = Particle->ObjectState() == EObjectStateType::Static;
-				if (IsStatic)
-				{
-					RigidTransforms[i] = FRigidTransform3(Particle->X(), Particle->R());
-				}
-				else
-				{
-					RigidTransforms[i] = FRigidTransform3(Particle->P(), Particle->Q());
-				}
-			}
-			SweptConstraint->ResetManifold();
-			Collisions::UpdateConstraintFromGeometry<ECollisionUpdateType::Deepest>(*(SweptConstraint), RigidTransforms[0], RigidTransforms[1], Dt);
+			FRigidTransform3 ShapeWorldTransform0 = SweptConstraint->GetShapeWorldTransform0();
+			FRigidTransform3 ShapeWorldTransform1 = SweptConstraint->GetShapeWorldTransform1();
+			ShapeWorldTransform0.SetTranslation(FConstGenericParticleHandle(SweptConstraint->GetParticle0())->P());
+			ShapeWorldTransform1.SetTranslation(FConstGenericParticleHandle(SweptConstraint->GetParticle1())->P());
+
+			SweptConstraint->SetShapeWorldTransforms(ShapeWorldTransform0, ShapeWorldTransform1);
+			SweptConstraint->UpdateManifoldContacts();
 
 			// @todo(zhenglin): Removing constraints that has Phi larger than CullDistance could reduce the island sizes in the normal solve. But I could not get this to work...
 			// if (SweptConstraint->GetPhi() > SweptConstraint->GetCullDistance())
diff --git a/Engine/Source/Runtime/Experimental/Chaos/Private/Chaos/CollisionResolution.cpp b/Engine/Source/Runtime/Experimental/Chaos/Private/Chaos/CollisionResolution.cpp
--- a/Engine/Source/Runtime/Experimental/Chaos/Private/Chaos/CollisionResolution.cpp
+++ b/Engine/Source/Runtime/Experimental/Chaos/Private/Chaos/CollisionResolution.cpp
@@ -235,12 +235,6 @@ namespace Chaos
 				return;
 			}
 
-			TPBDRigidParticleHandle<FReal, 3>* RigidParticle = Particle->CastToRigidParticle();
-			if (!CHAOS_ENSURE(RigidParticle))
-			{
-				return;
-			}
-
 			const FReal MinBounds = Particle->Geometry()->BoundingBox().Extents().Min();
 			const FReal Depth = -SweptConstraint.GetPhi();
 			const FReal MaxDepth = MinBounds * CCDAllowedDepthBoundsScale;
@@ -248,7 +242,23 @@ namespace Chaos
 			if (TOI <= 0.0f)
 			{
 				// Initial overlap. Phi is correct already.
-				SweptConstraint.TimeOfImpact = 0.0f;
+				TPBDRigidParticleHandle<FReal, 3> *Rigid0 = SweptConstraint.GetParticle0()->CastToRigidParticle();
+				TPBDRigidParticleHandle<FReal, 3> *Rigid1 = SweptConstraint.GetParticle1()->CastToRigidParticle();
+				const FRigidTransform3& ShapeWorldTransform1 = SweptConstraint.GetShapeWorldTransform1();
+				const FVec3 Normal = ShapeWorldTransform1.TransformVectorNoScale(SweptConstraint.GetClosestManifoldPoint()->ContactPoint.ShapeContactNormal);
+				const FVec3 V0 = Rigid0 != nullptr ? Rigid0->V() : FVec3(0.f);
+				const FVec3 V1 = Rigid1 != nullptr ? Rigid1->V() : FVec3(0.f);
+				const FReal NormalV = FVec3::DotProduct(V0 - V1, Normal);
+				// The two objects are already separating. This constraint needs not to be considered for CCD if the normal velocity is positive.
+				// If TOI < 0 (initial penetration case), CCD will not depenetrate the objects since CCD is only a velocity solve. Depenetration will be left for the main solve.
+				if (NormalV >= 0) 
+				{
+					SweptConstraint.TimeOfImpact = std::numeric_limits<FReal>::max();
+				}
+				else
+				{
+					SweptConstraint.TimeOfImpact = 0.0f;
+				}
 			}
 			else if (TOI > 0.0f && TOI <= 1.0f)
 			{
@@ -1337,8 +1347,8 @@ namespace Chaos
 		{
 			CONDITIONAL_SCOPE_CYCLE_COUNTER(STAT_Collisions_UpdateGenericConvexConvexConstraintSwept, ConstraintsDetailedStats);
 			FReal TOI = TNumericLimits<FReal>::Max();
-			const FRigidTransform3 EndWorldTransform0(FGenericParticleHandle(Particle0)->P(), FGenericParticleHandle(Particle0)->Q());
-			const FRigidTransform3 EndWorldTransform1(FGenericParticleHandle(Particle1)->P(), FGenericParticleHandle(Particle1)->Q());
+			const FRigidTransform3& EndWorldTransform0 = Constraint.GetShapeWorldTransform0();
+			const FRigidTransform3& EndWorldTransform1 = Constraint.GetShapeWorldTransform1();
 			UpdateContactPointNoCull(Constraint, GenericConvexConvexContactPointSwept(Implicit0, StartWorldTransform0, EndWorldTransform0, Implicit1, StartWorldTransform1, EndWorldTransform1, Dir, Length, TOI), Dt);
 			SetSweptConstraintTOI(Particle0, TOI, Length, Dir, Constraint);
 		}
