Chaos - Capsule-Convex manifold

diff --git a/Engine/Source/Runtime/Experimental/Chaos/Private/Chaos/CollisionOneShotManifolds.cpp b/Engine/Source/Runtime/Experimental/Chaos/Private/Chaos/CollisionOneShotManifolds.cpp
--- a/Engine/Source/Runtime/Experimental/Chaos/Private/Chaos/CollisionOneShotManifolds.cpp
+++ b/Engine/Source/Runtime/Experimental/Chaos/Private/Chaos/CollisionOneShotManifolds.cpp
@@ -45,6 +45,11 @@ namespace Chaos
 	FRealSingle Chaos_Collision_Manifold_EdgeContactNormalThreshold = 0.9f;	// About 25deg
 	FAutoConsoleVariableRef CVarChaos_Manifold_EdgeContactNormalThreshold(TEXT("p.Chaos.Collision.Manifold.EdgeNormalThreshold"), Chaos_Collision_Manifold_EdgeContactNormalThreshold, TEXT(""));
 
+	// We use a smaller margin for triangle to convex. Ideally we would have none at all in this case
+	// LWC_TODO: This needs to be a larger value for float builds (probably 1)
+	FRealSingle Chaos_Collision_Manifold_TriangleConvexMarginMultiplier = 0.5f;
+	FAutoConsoleVariableRef CVarChaos_Manifold_TriangleConvexMarginMultipler(TEXT("p.Chaos.Collision.Manifold.TriangleConvexMarginMultiplier"), Chaos_Collision_Manifold_TriangleConvexMarginMultiplier, TEXT(""));
+
 	FRealSingle Chaos_Collision_Manifold_CullDistanceMarginMultiplier = 1.0f;
 	FAutoConsoleVariableRef CVarChaosCollisioConvexManifoldCullDistanceMarginMultiplier(TEXT("p.Chaos.Collision.Manifold.CullDistanceMarginMultiplier"), Chaos_Collision_Manifold_CullDistanceMarginMultiplier, TEXT(""));
 
@@ -61,12 +66,13 @@ namespace Chaos
 	FAutoConsoleVariableRef CVarChaos_Collision_EnableManifoldInject(TEXT("p.Chaos.Collision.EnableManifoldGJKInject"), bChaos_Collision_EnableManifoldGJKInject, TEXT(""));
 
 	// See GJKContactPointMargin for comments on why these matter
+	// LWC_TODO: These needs to be a larger values for float builds (1.e-3f)
 	FRealSingle Chaos_Collision_GJKEpsilon = 1.e-6f;
 	FRealSingle Chaos_Collision_EPAEpsilon = 1.e-6f;
 	FAutoConsoleVariableRef CVarChaos_Collision_GJKEpsilon(TEXT("p.Chaos.Collision.GJKEpsilon"), Chaos_Collision_GJKEpsilon, TEXT(""));
 	FAutoConsoleVariableRef CVarChaos_Collision_EPAEpsilon(TEXT("p.Chaos.Collision.EPAEpsilon"), Chaos_Collision_EPAEpsilon, TEXT(""));
 
-	// Whether to prune spurious edge contacts in trimesh nd heighfield contacts (pretty much essential)
+	// Whether to prune spurious edge contacts in trimesh and heighfield contacts (pretty much essential)
 	bool bChaos_Collision_EnableEdgePrune = true;
 	FRealSingle Chaos_Collision_EdgePrunePlaneDistance = 3.0;
 	FAutoConsoleVariableRef CVarChaos_Collision_EnableEdgePrune(TEXT("p.Chaos.Collision.EnableEdgePrune"), bChaos_Collision_EnableEdgePrune, TEXT(""));
@@ -491,7 +497,6 @@ namespace Chaos
 			const int32 RefPlaneIndex,
 			const int32 OtherPlaneIndex,
 			const FVec3& RefPlaneNormal,
-			const FVec3& RefPlanePos,
 			FVec3* VertexBuffer1,
 			FVec3* VertexBuffer2,
 			uint32& ContactPointCount,	// InOut
@@ -527,6 +532,64 @@ namespace Chaos
 			return VertexBuffer1;
 		}
 
+		template <typename ConvexImplicitType>
+		FVec3* GenerateLineConvexManifoldClippedVerticesSameSpace(
+			const ConvexImplicitType& Convex,
+			const TSegment<FReal>& Segment,
+			const int32 ConvexPlaneIndex,
+			const FVec3& PlaneNormal,
+			FVec3* VertexBuffer,
+			int32& ContactPointCount,	// InOut
+			const int32 MaxContactPointCount,
+			const FReal PlaneTolerance)
+		{
+			check(MaxContactPointCount >= 2);
+
+			const FReal DistanceEpsilon = FReal(KINDA_SMALL_NUMBER);
+
+			// Populate the clipped vertices by the line segment vertices
+			ContactPointCount = 2;
+			VertexBuffer[0] = Segment.GetX1();
+			VertexBuffer[1] = Segment.GetX2();
+
+			// Now clip against all edge planes that belong to the convex face
+			// Note winding order matters here, and we have to handle negative scales
+			const FReal WindingOrder = Convex.GetWindingOrder();
+			const int32 ConvexFaceVerticesNum = Convex.NumPlaneVertices(ConvexPlaneIndex);
+			int32 ClippingPlaneCount = ConvexFaceVerticesNum;
+			FVec3 PrevPoint = Convex.GetVertex(Convex.GetPlaneVertex(ConvexPlaneIndex, ClippingPlaneCount - 1));
+			for (int32 ClippingPlaneIndex = 0; (ClippingPlaneIndex < ClippingPlaneCount) && (ContactPointCount > 1); ++ClippingPlaneIndex)
+			{
+				const FVec3 CurrentPoint = Convex.GetVertex(Convex.GetPlaneVertex(ConvexPlaneIndex, ClippingPlaneIndex));
+				const FVec3 ClippingPlaneNormal = WindingOrder * FVec3::CrossProduct(PlaneNormal, PrevPoint - CurrentPoint).GetSafeNormal();
+				PrevPoint = CurrentPoint;
+
+				// Distance to plane (+ve outside, -ve inside)
+				const FReal Dist0 = FVec3::DotProduct(VertexBuffer[0] - CurrentPoint, ClippingPlaneNormal);
+				const FReal Dist1 = FVec3::DotProduct(VertexBuffer[1] - CurrentPoint, ClippingPlaneNormal);
+				if ((Dist0 > PlaneTolerance) && (Dist1 > PlaneTolerance))
+				{
+					// Both points are outside, so no contacts
+					ContactPointCount = 0;
+					break;
+				}
+				else if ((Dist0 > PlaneTolerance) && (Dist1 < 0))
+				{
+					// First point is outside, second inside
+					const FReal Alpha = Dist1 / (Dist1 - Dist0);
+					VertexBuffer[0] = VertexBuffer[1] + Alpha * (VertexBuffer[0] - VertexBuffer[1]);
+				}
+				else if ((Dist1 > PlaneTolerance) && (Dist0 < 0))
+				{
+					// Second point is outside, first inside
+					const FReal Alpha = Dist0 / (Dist0 - Dist1);
+					VertexBuffer[1] = VertexBuffer[0] + Alpha * (VertexBuffer[1] - VertexBuffer[0]);
+				}
+			}
+
+			return VertexBuffer;
+		}
+
 		// Use GJK to find the closest points (or shallowest penetrating points) on two convex shapes usingthe specified margin
 		// @todo(chaos): dedupe from GJKContactPoint in CollisionResolution.cpp
 		template <typename GeometryA, typename GeometryB>
@@ -567,10 +630,10 @@ namespace Chaos
 			return Contact;
 		}
 
-		// GJK contact point between a convex and a triangle
-		// This assumes that the triangle has been transformed into convex space
-		template <typename GeometryA>
-		FContactPoint GJKContactPointConvexTriangle(const GeometryA& A, const FTriangle& B, FReal MarginA, FReal& OutMaxMarginDelta, int32& VertexIndexA, int32& VertexIndexB)
+		// GJK contact point between two GJKShape wrapped implicits
+		// This assumes that the both shapes are in the same space
+		template <typename GJKShapeA, typename GJKShapeB>
+		FContactPoint GJKContactPointSameSpace(const GJKShapeA& A, const GJKShapeB& B, FReal& OutMaxMarginDelta, int32& VertexIndexA, int32& VertexIndexB)
 		{
 			SCOPE_CYCLE_COUNTER_MANIFOLD_GJK();
 
@@ -583,10 +646,7 @@ namespace Chaos
 			const FReal GJKEpsilon = Chaos_Collision_GJKEpsilon;
 			const FReal EPAEpsilon = Chaos_Collision_EPAEpsilon;
 
-			const TGJKCoreShape<GeometryA> AWithMargin(A, MarginA);
-			const TGJKShape<FTriangle> BNoMargin(B);
-
-			if (GJKPenetrationSameSpace(AWithMargin, BNoMargin, Penetration, ClosestA, ClosestB, Normal, VertexIndexA, VertexIndexB, OutMaxMarginDelta, GJKEpsilon, EPAEpsilon))
+			if (GJKPenetrationSameSpace(A, B, Penetration, ClosestA, ClosestB, Normal, VertexIndexA, VertexIndexB, OutMaxMarginDelta, GJKEpsilon, EPAEpsilon))
 			{
 				Contact.ShapeContactPoints[0] = ClosestA;
 				Contact.ShapeContactPoints[1] = ClosestB;
@@ -872,8 +932,6 @@ namespace Chaos
 			// For edge-edge contacts, we find the edges involved and project the contact onto the edges (if necessary - see below)
 			if (!bIsPlaneContact)
 			{
-				SCOPE_CYCLE_COUNTER_MANIFOLD_ADDEDGEEDGE();
-
 				if (ForceOneShotManifoldEdgeEdgeCaseZeroCullDistance && GJKContactPoint.Phi > 0)
 				{
 					return;
@@ -942,7 +1000,6 @@ namespace Chaos
 			FVec3 ClippedVertices2[MaxContactPointCount];
 
 			{
-				SCOPE_CYCLE_COUNTER_MANIFOLD_CLIP();
 				if (ReferenceFaceConvex1)
 				{
 					RefConvexTM = &Convex1Transform;
@@ -955,7 +1012,6 @@ namespace Chaos
 						MostOpposingPlaneIndexConvex1,
 						MostOpposingPlaneIndexConvex2,
 						RefPlaneNormal,
-						RefPlanePosition,
 						ClippedVertices1,
 						ClippedVertices2,
 						ContactPointCount,
@@ -973,7 +1029,6 @@ namespace Chaos
 						MostOpposingPlaneIndexConvex2,
 						MostOpposingPlaneIndexConvex1,
 						RefPlaneNormal,
-						RefPlanePosition,
 						ClippedVertices1,
 						ClippedVertices2,
 						ContactPointCount,
@@ -992,8 +1047,6 @@ namespace Chaos
 			// Reduce number of contacts to the maximum allowed
 			if (ContactPointCount > 4)
 			{
-				SCOPE_CYCLE_COUNTER_MANIFOLD_REDUCE();
-
 				FRotation3 RotateSeperationToZ = FRotation3::FromRotatedVector(RefPlaneNormal, FVec3(0.0f, 0.0f, 1.0f));
 				for (uint32 ContactPointIndex = 0; ContactPointIndex < ContactPointCount; ++ContactPointIndex)
 				{
@@ -1013,7 +1066,6 @@ namespace Chaos
 
 			// Generate the contact points from the clipped vertices
 			{
-				SCOPE_CYCLE_COUNTER_MANIFOLD_ADDFACEVERTEX();
 				for (uint32 ContactPointIndex = 0; ContactPointIndex < ContactPointCount; ++ContactPointIndex)
 				{
 					FContactPoint ContactPoint;
@@ -1037,6 +1089,191 @@ namespace Chaos
 		}
 
 
+		template <typename ConvexType>
+		void ConstructCapsuleConvexOneShotManifold(
+			const FImplicitCapsule3& Capsule, 
+			const FRigidTransform3& CapsuleTransform,
+			const ConvexType& Convex,
+			const FRigidTransform3& ConvexTransform,
+			const FReal CullDistance,
+			TCArray<FContactPoint, 4>& OutContactPoints)
+		{
+			SCOPE_CYCLE_COUNTER_MANIFOLD();
+			check(OutContactPoints.Num() == 0);
+			using FLineSegment3 = TSegment<FReal>;
+
+			// Transform Capsule into Convex space
+			// We only keep the segment at this point - we'll add the radius back later if we create contacts
+			const FRigidTransform3 CapsuleToConvexTransform = CapsuleTransform.GetRelativeTransformNoScale(ConvexTransform);
+			const FReal CapsuleRadius = Capsule.GetRadius();
+			const FReal CapsuleHeight = Capsule.GetHeight();
+			const FLineSegment3 CapsuleAxisSegment = FLineSegment3(
+				CapsuleToConvexTransform.TransformPositionNoScale(Capsule.GetX1()),
+				CapsuleToConvexTransform.TransformVectorNoScale(Capsule.GetAxis()),
+				CapsuleHeight);
+
+			// No margins for the convex or the capsule for this GJK. We'll correct for the radius later...
+			// NOTE: We are passing the convex first into GJK, but the output should put the capsule first
+			FReal UnusedMaxMarginDelta = FReal(0);
+			int32 VertexIndexA = INDEX_NONE, VertexIndexB = INDEX_NONE;
+			const TGJKShape<FLineSegment3> GJKSegment(CapsuleAxisSegment);
+			const TGJKShape<ConvexType> GJKConvex(Convex);
+			FContactPoint GJKContactPoint = GJKContactPointSameSpace(GJKConvex, GJKSegment, UnusedMaxMarginDelta, VertexIndexA, VertexIndexB);
+
+			// Map contact back onto capsule
+			GJKContactPoint.ShapeContactPoints[1] = GJKContactPoint.ShapeContactPoints[1] + CapsuleRadius * GJKContactPoint.ShapeContactNormal;
+			GJKContactPoint.Phi = GJKContactPoint.Phi - CapsuleRadius;
+
+			PHYSICS_CSV_CUSTOM_EXPENSIVE(PhysicsCounters, NumManifoldsGJKCalled, 1, ECsvCustomStatOp::Accumulate);
+
+			if (GJKContactPoint.Phi > CullDistance)
+			{
+				PHYSICS_CSV_CUSTOM_EXPENSIVE(PhysicsCounters, NumManifoldsGJKCulled, 1, ECsvCustomStatOp::Accumulate);
+				return;
+			}
+
+			PHYSICS_CSV_CUSTOM_EXPENSIVE(PhysicsCounters, NumManifoldsCreated, 1, ECsvCustomStatOp::Accumulate);
+
+			// Is this a cylinder or end cap contact?
+			const FReal CapsuleEndCapAxisNormalThreshold = FReal(0.01);	// About sin(0.5 deg)
+			const FReal CapsuleAxisDotNormal = FVec3::DotProduct(CapsuleAxisSegment.GetAxis(), GJKContactPoint.ShapeContactNormal);
+			const bool bIsCapsuleCylinderContact = (FMath::Abs(CapsuleAxisDotNormal) < CapsuleEndCapAxisNormalThreshold);
+
+			// Find the best plane on the convex (most opposing direction)
+			const FReal PlaneSearchDistance = FReal(0);	// use the miniumum (see SelectContactPlane)
+			const int32 ConvexPlaneIndex = SelectContactPlane(Convex, GJKContactPoint.ShapeContactPoints[0], GJKContactPoint.ShapeContactNormal, PlaneSearchDistance, VertexIndexA);
+			FVec3 ConvexPlanePosition, ConvexPlaneNormal;
+			Convex.GetPlaneNX(ConvexPlaneIndex, ConvexPlaneNormal, ConvexPlanePosition);
+			const FReal ConvexPlaneNormalDotContactNormal = FVec3::DotProduct(ConvexPlaneNormal, GJKContactPoint.ShapeContactNormal);
+			const FReal PlaneContactNormalEpsilon = Chaos_Collision_Manifold_PlaneContactNormalEpsilon;
+			bool bIsConvexPlaneContact = FMath::IsNearlyEqual(ConvexPlaneNormalDotContactNormal, FReal(-1), PlaneContactNormalEpsilon);
+
+			if (bIsCapsuleCylinderContact)
+			{
+				// Cylinder plane most opposing GJK result direction
+				const FVec3 CapsulePlaneNormal = -(GJKContactPoint.ShapeContactNormal - CapsuleAxisDotNormal * CapsuleAxisSegment.GetAxis()).GetSafeNormal();
+
+				// If we have a cylinder contact we generate 2 contact points from the clipped line segment
+				// NOTE: There is an edge case where GJK returns a cylinder normal but both segment positions are
+				// outside the box. This happens if you drop a capsule onto a box so that the segment exactly touches
+				// a box face. It also happens if we mis-identify the contact as a cylinder contact because of
+				// the normal check tolerance. Either way, we fall through to just using the GJK result
+				const FReal PlaneTolerance = FReal(1e-2);
+				const int32 MaxContactPointCount = 2;
+				int32 ContactPointCount = 0;
+				FVec3 ClippedVertices[MaxContactPointCount];
+				GenerateLineConvexManifoldClippedVerticesSameSpace(
+					Convex,
+					CapsuleAxisSegment,
+					ConvexPlaneIndex,
+					CapsulePlaneNormal,
+					ClippedVertices,
+					ContactPointCount,
+					MaxContactPointCount,
+					PlaneTolerance);
+
+				if (ContactPointCount > 0)
+				{
+					// Divisor used below to project the capsule contact onto the convex plane
+					// See FMath::RayPlaneIntersection for math, but it is separated here to avoid excessive dot products
+					const FReal ConvexDistanceDivisor = FVec3::DotProduct(CapsulePlaneNormal, ConvexPlaneNormal);
+					if (FMath::Abs(ConvexDistanceDivisor) > KINDA_SMALL_NUMBER)
+					{
+						const FReal ConvexDistanceMultiplier = FReal(1) / ConvexDistanceDivisor;
+
+						for (int32 ContactPointIndex = 0; ContactPointIndex < ContactPointCount; ++ContactPointIndex)
+						{
+							// Capsule points are on cylinder surface, so easy to calculate
+							const FVec3 CapsuleContactPoint = ClippedVertices[ContactPointIndex] - CapsuleRadius * CapsulePlaneNormal;
+
+							// Project the capsule contacts onto the convex plane
+							const FReal ConvexContactPointDistance = FVec3::DotProduct((ConvexPlanePosition - CapsuleContactPoint), ConvexPlaneNormal) * ConvexDistanceMultiplier;
+							const FVec3 ConvexContactPoint = CapsuleContactPoint + ConvexContactPointDistance * CapsulePlaneNormal;
+
+							FContactPoint& ContactPoint = OutContactPoints[OutContactPoints.Add()];
+							ContactPoint.ShapeContactPoints[0] = CapsuleToConvexTransform.InverseTransformPositionNoScale(CapsuleContactPoint);
+							ContactPoint.ShapeContactPoints[1] = ConvexContactPoint;
+							ContactPoint.ShapeContactNormal = CapsulePlaneNormal;
+							ContactPoint.Phi = FVec3::DotProduct(CapsuleContactPoint - ConvexContactPoint, CapsulePlaneNormal);
+							ContactPoint.FaceIndex = INDEX_NONE;
+							ContactPoint.ContactType = EContactPointType::PlaneVertex;
+						}
+
+						return;
+					}
+				}
+			}
+			else if (bIsConvexPlaneContact)
+			{
+				// If we get here we are creating a convex face versus end cap collision
+				// The end cap collision gets added at the end (the GJK result), but we also add the clipped capsule cylinder points as well
+				// but only if the cylinder is not standing vertically on the face
+				const FReal CapsuleCylinderTolerance = FReal(0.707);	// about 45 deg
+				const FReal CapsuleAxisDotConvexNormal = FVec3::DotProduct(ConvexPlaneNormal, CapsuleAxisSegment.GetAxis());
+				if (CapsuleAxisDotConvexNormal < CapsuleCylinderTolerance)
+				{
+					// The line segment on the surface of the cylinder closest to the contact
+					const FVec3 CapsuleCylinderNormal = (ConvexPlaneNormal - CapsuleAxisDotConvexNormal * CapsuleAxisSegment.GetAxis()).GetUnsafeNormal();
+					const FLineSegment3 CapsuleCylinderSegment = FLineSegment3(CapsuleAxisSegment.GetX1() - CapsuleRadius * CapsuleCylinderNormal, CapsuleAxisSegment.GetAxis(), CapsuleHeight);
+
+					// Clip the cylinder segment to the convex face
+					const FReal PlaneTolerance = FReal(1e-2);
+					const int32 MaxContactPointCount = 2;
+					int32 ContactPointCount = 0;
+					FVec3 ClippedVertices[MaxContactPointCount];
+					GenerateLineConvexManifoldClippedVerticesSameSpace(
+						Convex,
+						CapsuleCylinderSegment,
+						ConvexPlaneIndex,
+						ConvexPlaneNormal,
+						ClippedVertices,
+						ContactPointCount,
+						MaxContactPointCount,
+						PlaneTolerance);
+
+					const FReal PointDistanceToleranceSq = FMath::Square(FReal(0.1) * CapsuleRadius);
+
+					// Add the vertices if not clipped away
+					for (int32 ContactPointIndex = 0; ContactPointIndex < ContactPointCount; ++ContactPointIndex)
+					{
+						// Don't add the point if it's very close to the GJK result because we're adding that too below
+						const FVec3 CapsuleContactPoint = ClippedVertices[ContactPointIndex];
+						const FReal CapsuleContactGJKDistanceSq = (CapsuleContactPoint - GJKContactPoint.ShapeContactPoints[1]).SizeSquared();
+						if (CapsuleContactGJKDistanceSq > PointDistanceToleranceSq)
+						{
+							const FVec3 ConvexContactPoint = CapsuleContactPoint - FVec3::DotProduct(CapsuleContactPoint - ConvexPlanePosition, ConvexPlaneNormal) * ConvexPlaneNormal;
+
+							FContactPoint& ContactPoint = OutContactPoints[OutContactPoints.Add()];
+							ContactPoint.ShapeContactPoints[0] = CapsuleToConvexTransform.InverseTransformPositionNoScale(CapsuleContactPoint);
+							ContactPoint.ShapeContactPoints[1] = ConvexContactPoint;
+							ContactPoint.ShapeContactNormal = ConvexPlaneNormal;
+							ContactPoint.Phi = FVec3::DotProduct(CapsuleContactPoint - ConvexContactPoint, ConvexPlaneNormal);
+							ContactPoint.FaceIndex = INDEX_NONE;
+							ContactPoint.ContactType = EContactPointType::VertexPlane;
+						}
+					}
+				}
+
+				// @todo(chaos): handle the case where we have a convex on the end cap by adding a couple stabilizing points around the main contact
+			}
+		
+			// We have an end-cap collision, so add the GJK point
+			// NOTE: We are not using convex margins, so the GJK result is already on both shape surfaces
+			// NOTE: we must swap the shape orders and transform the capsule data back into its space
+			{
+				FContactPoint ContactPoint;
+				ContactPoint.ShapeContactPoints[0] = CapsuleToConvexTransform.InverseTransformPositionNoScale(GJKContactPoint.ShapeContactPoints[1]);
+				ContactPoint.ShapeContactPoints[1] = GJKContactPoint.ShapeContactPoints[0];
+				ContactPoint.ShapeContactNormal = -GJKContactPoint.ShapeContactNormal;
+				ContactPoint.Phi = GJKContactPoint.Phi;
+				ContactPoint.FaceIndex = INDEX_NONE;
+				ContactPoint.ContactType = EContactPointType::EdgeEdge;
+				OutContactPoints.Add(ContactPoint);
+				return;
+			}
+		}
+
+		// @todo(chaos): don't use GJK/EPA for triangle-convex so we can avoid using margains and avoid errors related to misidentifying edge contacts
 		template <typename ConvexType>
 		void ConstructPlanarConvexTriangleOneShotManifold(const ConvexType& Convex, const FTriangle& Triangle, const FReal CullDistance, TCArray<FContactPoint, 4>& OutContactPoints)
 		{
@@ -1045,10 +1282,13 @@ namespace Chaos
 
 			// Find the deepest penetration. This is used to determine the planes and points to use for the manifold
 			// MaxMarginDelta is an upper bound on the distance from the contact on the rounded core shape to the actual shape surface. 
-			const FReal ConvexMargin = Convex.GetMargin();
+			const FReal MarginMultiplier = Chaos_Collision_Manifold_TriangleConvexMarginMultiplier;
+			const FReal ConvexMargin = MarginMultiplier * Convex.GetMargin();
 			FReal MaxMarginDelta = FReal(0);
 			int32 VertexIndexA = INDEX_NONE, VertexIndexB = INDEX_NONE;
-			FContactPoint GJKContactPoint = GJKContactPointConvexTriangle(Convex, Triangle, ConvexMargin, MaxMarginDelta, VertexIndexA, VertexIndexB);
+			const TGJKCoreShape<ConvexType> GJKConvex(Convex, ConvexMargin);
+			const TGJKShape<FTriangle> GJKTriangle(Triangle);
+			FContactPoint GJKContactPoint = GJKContactPointSameSpace(GJKConvex, GJKTriangle, MaxMarginDelta, VertexIndexA, VertexIndexB);
 
 			PHYSICS_CSV_CUSTOM_EXPENSIVE(PhysicsCounters, NumManifoldsGJKCalled, 1, ECsvCustomStatOp::Accumulate);
 
@@ -1083,7 +1323,6 @@ namespace Chaos
 			bool bIsPlaneContact = FMath::IsNearlyEqual(ConvexPlaneNormalDotContactNormal, FReal(-1), PlaneContactNormalEpsilon) || FMath::IsNearlyEqual(TrianglePlaneNormalDotContactNormal, FReal(1), PlaneContactNormalEpsilon);
 			if (!bIsPlaneContact)
 			{
-				SCOPE_CYCLE_COUNTER_MANIFOLD_ADDEDGEEDGE();
 				if (ForceOneShotManifoldEdgeEdgeCaseZeroCullDistance && GJKContactPoint.Phi > 0)
 				{
 					return;
@@ -1161,7 +1400,6 @@ namespace Chaos
 					0,
 					ConvexPlaneIndex,
 					TrianglePlaneNormal,
-					TrianglePlanePosition,
 					ClippedVertices1,
 					ClippedVertices2,
 					(uint32&)ContactPointCount,
@@ -1180,7 +1418,6 @@ namespace Chaos
 					ConvexPlaneIndex,
 					0,
 					ConvexPlaneNormal,
-					ConvexPlanePosition,
 					ClippedVertices1,
 					ClippedVertices2,
 					(uint32&)ContactPointCount,
@@ -1198,8 +1435,6 @@ namespace Chaos
 			// Reduce number of contacts to the maximum allowed
 			if (ContactPointCount > 4)
 			{
-				SCOPE_CYCLE_COUNTER_MANIFOLD_REDUCE();
-
 				FRotation3 RotateSeperationToZ = FRotation3::FromRotatedVector(RefPlaneNormal, FVec3(0.0f, 0.0f, 1.0f));
 				for (int32 ContactPointIndex = 0; ContactPointIndex < ContactPointCount; ++ContactPointIndex)
 				{
@@ -1220,7 +1455,6 @@ namespace Chaos
 			if (ContactType == EContactPointType::VertexPlane)
 			{
 				// Triangle face contact (clipped vertices are convex vertices)
-				SCOPE_CYCLE_COUNTER_MANIFOLD_ADDFACEVERTEX();
 				for (int32 ContactPointIndex = 0; ContactPointIndex < ContactPointCount; ++ContactPointIndex)
 				{
 					FContactPoint& ContactPoint = OutContactPoints[OutContactPoints.Add()];
@@ -1237,7 +1471,6 @@ namespace Chaos
 			else
 			{
 				// Convex face contact (clipped vertices are triangle vertices)
-				SCOPE_CYCLE_COUNTER_MANIFOLD_ADDFACEVERTEX();
 				for (int32 ContactPointIndex = 0; ContactPointIndex < ContactPointCount; ++ContactPointIndex)
 				{
 					FContactPoint& ContactPoint = OutContactPoints[OutContactPoints.Add()];
@@ -1402,6 +1635,107 @@ namespace Chaos
 			const FRigidTransform3& Convex2Transform, //world
 			const FReal Dt,
 			FPBDCollisionConstraint& Constraint);
+	
+
+		template
+		void ConstructCapsuleConvexOneShotManifold(
+			const FImplicitCapsule3& Capsule,
+			const FRigidTransform3& CapsuleTransform,
+			const FImplicitConvex3& Convex,
+			const FRigidTransform3& ConvexTransform,
+			const FReal CullDistance,
+			TCArray<FContactPoint, 4>& OutContactPoints);
+
+		template
+		void ConstructCapsuleConvexOneShotManifold(
+			const FImplicitCapsule3& Capsule,
+			const FRigidTransform3& CapsuleTransform,
+			const TImplicitObjectInstanced<FImplicitConvex3>& Convex,
+			const FRigidTransform3& ConvexTransform,
+			const FReal CullDistance,
+			TCArray<FContactPoint, 4>& OutContactPoints);
+
+		template
+		void ConstructCapsuleConvexOneShotManifold(
+			const FImplicitCapsule3& Capsule,
+			const FRigidTransform3& CapsuleTransform,
+			const TImplicitObjectScaled<FImplicitConvex3>& Convex,
+			const FRigidTransform3& ConvexTransform,
+			const FReal CullDistance,
+			TCArray<FContactPoint, 4>& OutContactPoints);
+
+		template
+		void ConstructCapsuleConvexOneShotManifold(
+			const FImplicitCapsule3& Capsule,
+			const FRigidTransform3& CapsuleTransform,
+			const FImplicitBox3& Convex,
+			const FRigidTransform3& ConvexTransform,
+			const FReal CullDistance,
+			TCArray<FContactPoint, 4>& OutContactPoints);
+
+
+		template
+		void ConstructPlanarConvexTriangleOneShotManifold(
+			const FImplicitConvex3& Convex, 
+			const FTriangle& Triangle, 
+			const FReal CullDistance, 
+			TCArray<FContactPoint, 4>& OutContactPoints);
+
+		template
+		void ConstructPlanarConvexTriangleOneShotManifold(
+			const TImplicitObjectInstanced<FImplicitConvex3>& Convex,
+			const FTriangle& Triangle,
+			const FReal CullDistance,
+			TCArray<FContactPoint, 4>& OutContactPoints);
+
+		template
+		void ConstructPlanarConvexTriangleOneShotManifold(
+			const TImplicitObjectScaled<FImplicitConvex3>& Convex,
+			const FTriangle& Triangle,
+			const FReal CullDistance,
+			TCArray<FContactPoint, 4>& OutContactPoints);
+
+		template
+		void ConstructPlanarConvexTriangleOneShotManifold(
+			const FImplicitBox3& Convex,
+			const FTriangle& Triangle,
+			const FReal CullDistance,
+			TCArray<FContactPoint, 4>& OutContactPoints);
+
+		template
+		void ConstructPlanarConvexTriangleOneShotManifold(
+			const TImplicitObjectScaled<FImplicitBox3>& Convex,
+			const FTriangle& Triangle,
+			const FReal CullDistance,
+			TCArray<FContactPoint, 4>& OutContactPoints);
+
+		template
+		void ConstructPlanarConvexTriangleOneShotManifold(
+			const TImplicitObjectInstanced<FImplicitBox3>& Convex,
+			const FTriangle& Triangle,
+			const FReal CullDistance,
+			TCArray<FContactPoint, 4>& OutContactPoints);
+
+
+		// @todo(chaos): all ConstructConvexConvexOneShotManifold taking a triangle are still used in FHeightField::ContactManifoldImp and the 
+		// TriMesh version, but they should changed to use ConstructPlanarConvexTriangleOneShotManifold
+		template
+		void ConstructConvexConvexOneShotManifold<class FCapsule, FTriangle>(
+			const FCapsule& Implicit1,
+			const FRigidTransform3& Convex1Transform, //world
+			const FTriangle& Implicit2,
+			const FRigidTransform3& Convex2Transform, //world
+			const FReal Dt,
+			FPBDCollisionConstraint& Constraint);
+
+		template
+		void ConstructConvexConvexOneShotManifold<TImplicitObjectScaled<class FCapsule, 1>, class FTriangle>(
+			const TImplicitObjectScaled<class FCapsule, 1>& Implicit1,
+			const FRigidTransform3& Convex1Transform, //world
+			const FTriangle& Implicit2,
+			const FRigidTransform3& Convex2Transform, //world
+			const FReal Dt,
+			FPBDCollisionConstraint& Constraint);
 
 		template
 		void ConstructConvexConvexOneShotManifold<FImplicitBox3, FTriangle>(
@@ -1492,67 +1826,6 @@ namespace Chaos
 			const FRigidTransform3& Convex2Transform, //world
 			const FReal Dt,
 			FPBDCollisionConstraint& Constraint);
-
-		template
-			void ConstructConvexConvexOneShotManifold<class FCapsule, FTriangle>(
-				const FCapsule& Implicit1,
-				const FRigidTransform3& Convex1Transform, //world
-				const FTriangle& Implicit2,
-				const FRigidTransform3& Convex2Transform, //world
-				const FReal Dt,
-				FPBDCollisionConstraint& Constraint);
-
-		template
-			void ConstructConvexConvexOneShotManifold<TImplicitObjectScaled<class FCapsule, 1>, class FTriangle>(
-				const TImplicitObjectScaled<class FCapsule, 1>& Implicit1,
-				const FRigidTransform3& Convex1Transform, //world
-				const FTriangle& Implicit2,
-				const FRigidTransform3& Convex2Transform, //world
-				const FReal Dt,
-				FPBDCollisionConstraint& Constraint);
-	
-	
-		template
-			void ConstructPlanarConvexTriangleOneShotManifold(
-				const FImplicitConvex3& Convex, 
-				const FTriangle& Triangle, 
-				const FReal CullDistance, 
-				TCArray<FContactPoint, 4>& OutContactPoints);
-
-		template
-			void ConstructPlanarConvexTriangleOneShotManifold(
-				const TImplicitObjectInstanced<FImplicitConvex3>& Convex,
-				const FTriangle& Triangle,
-				const FReal CullDistance,
-				TCArray<FContactPoint, 4>& OutContactPoints);
-
-		template
-			void ConstructPlanarConvexTriangleOneShotManifold(
-				const TImplicitObjectScaled<FImplicitConvex3>& Convex,
-				const FTriangle& Triangle,
-				const FReal CullDistance,
-				TCArray<FContactPoint, 4>& OutContactPoints);
-
-		template
-			void ConstructPlanarConvexTriangleOneShotManifold(
-				const FImplicitBox3& Convex,
-				const FTriangle& Triangle,
-				const FReal CullDistance,
-				TCArray<FContactPoint, 4>& OutContactPoints);
-
-		template
-			void ConstructPlanarConvexTriangleOneShotManifold(
-				const TImplicitObjectScaled<FImplicitBox3>& Convex,
-				const FTriangle& Triangle,
-				const FReal CullDistance,
-				TCArray<FContactPoint, 4>& OutContactPoints);
-
-		template
-			void ConstructPlanarConvexTriangleOneShotManifold(
-				const TImplicitObjectInstanced<FImplicitBox3>& Convex,
-				const FTriangle& Triangle,
-				const FReal CullDistance,
-				TCArray<FContactPoint, 4>& OutContactPoints);
 	}
 }
 
diff --git a/Engine/Source/Runtime/Experimental/Chaos/Private/Chaos/CollisionOneShotManifolds.h b/Engine/Source/Runtime/Experimental/Chaos/Private/Chaos/CollisionOneShotManifolds.h
--- a/Engine/Source/Runtime/Experimental/Chaos/Private/Chaos/CollisionOneShotManifolds.h
+++ b/Engine/Source/Runtime/Experimental/Chaos/Private/Chaos/CollisionOneShotManifolds.h
@@ -35,6 +35,15 @@ namespace Chaos
 			const FReal Dt,
 			FPBDCollisionConstraint& Constraint);
 
+		template <typename ConvexType>
+		void ConstructCapsuleConvexOneShotManifold(
+			const FImplicitCapsule3& Capsule,
+			const FRigidTransform3& CapsuleTransform,
+			const ConvexType& Convex,
+			const FRigidTransform3& ConvexTransform,
+			const FReal CullDistance,
+			TCArray<FContactPoint, 4>& OutContactPoints);
+
 		template <typename ConvexType>
 		void ConstructPlanarConvexTriangleOneShotManifold(
 			const ConvexType& Convex, 
diff --git a/Engine/Source/Runtime/Experimental/Chaos/Private/Chaos/CollisionResolution.cpp b/Engine/Source/Runtime/Experimental/Chaos/Private/Chaos/CollisionResolution.cpp
--- a/Engine/Source/Runtime/Experimental/Chaos/Private/Chaos/CollisionResolution.cpp
+++ b/Engine/Source/Runtime/Experimental/Chaos/Private/Chaos/CollisionResolution.cpp
@@ -1077,7 +1077,14 @@ namespace Chaos
 		{
 			if (Constraint.GetUseManifold())
 			{
-				ConstructConvexConvexOneShotManifold<FCapsule, FImplicitBox3>(A, ATransform, B, BTransform, Dt, Constraint);
+				TCArray<FContactPoint, 4> ContactPoints;
+				ConstructCapsuleConvexOneShotManifold<FImplicitBox3>(A, ATransform, B, BTransform, Constraint.GetCullDistance(), ContactPoints);
+
+				Constraint.ResetActiveManifoldContacts();
+				for (const FContactPoint& ContactPoint : ContactPoints)
+				{
+					Constraint.AddOneshotManifoldContact(ContactPoint);
+				}
 			}
 			else
 			{
@@ -1130,7 +1137,14 @@ namespace Chaos
 			{
 				if (Constraint.GetUseManifold())
 				{
-					ConstructConvexConvexOneShotManifold<FCapsule, ConvexType>(A, ATransform, B, BTransform, Dt, Constraint);
+					TCArray<FContactPoint, 4> ContactPoints;
+					ConstructCapsuleConvexOneShotManifold<ConvexType>(A, ATransform, B, BTransform, Constraint.GetCullDistance(), ContactPoints);
+
+					Constraint.ResetActiveManifoldContacts();
+					for (const FContactPoint& ContactPoint : ContactPoints)
+					{
+						Constraint.AddOneshotManifoldContact(ContactPoint);
+					}
 				}
 				else
 				{
diff --git a/Engine/Source/Runtime/Experimental/Chaos/Public/Chaos/PBDRigidsEvolutionGBF.h b/Engine/Source/Runtime/Experimental/Chaos/Public/Chaos/PBDRigidsEvolutionGBF.h
--- a/Engine/Source/Runtime/Experimental/Chaos/Public/Chaos/PBDRigidsEvolutionGBF.h
+++ b/Engine/Source/Runtime/Experimental/Chaos/Public/Chaos/PBDRigidsEvolutionGBF.h
@@ -98,8 +98,8 @@ namespace Chaos
 		static constexpr int32 DefaultNumCollisionPairIterations = 1;
 		static constexpr int32 DefaultNumPushOutIterations = 1;
 		static constexpr int32 DefaultNumCollisionPushOutPairIterations = 1;
-		static constexpr FRealSingle DefaultCollisionMarginFraction = 0.1f;
-		static constexpr FRealSingle DefaultCollisionMarginMax = 100.0f;
+		static constexpr FRealSingle DefaultCollisionMarginFraction = 0.05f;
+		static constexpr FRealSingle DefaultCollisionMarginMax = 10.0f;
 		static constexpr FRealSingle DefaultCollisionCullDistance = 3.0f;
 		static constexpr FRealSingle DefaultCollisionMaxPushOutVelocity = 1000.0f;
 		static constexpr int32 DefaultNumJointPairIterations = 1;
diff --git a/Engine/Source/Runtime/Experimental/Chaos/Public/Chaos/Segment.h b/Engine/Source/Runtime/Experimental/Chaos/Public/Chaos/Segment.h
--- a/Engine/Source/Runtime/Experimental/Chaos/Public/Chaos/Segment.h
+++ b/Engine/Source/Runtime/Experimental/Chaos/Public/Chaos/Segment.h
@@ -8,11 +8,16 @@ namespace Chaos
 	{
 	public:
 		TSegment() {}
-		TSegment(const TVec3<T>& x1, const TVec3<T>& x2)
-			: MPoint(x1)
-			, MAxis(x2 - x1)
+		TSegment(const TVec3<T>& X1, const TVec3<T>& X2)
+			: MPoint(X1)
+			, MAxis(X2 - X1)
 			, MLength(MAxis.SafeNormalize())
 		{ }
+		TSegment(const TVec3<T>& X1, const TVec3<T>& Axis, const T Length)
+			: MPoint(X1)
+			, MAxis(Axis)
+			, MLength(Length)
+		{ }
 
 		FORCEINLINE bool IsConvex() const { return true; }
 
