Chaos : Fields
- use parent child matrices when mapping cluster children into the field sample set.

diff --git a/Engine/Source/Runtime/Experimental/Chaos/Public/PhysicsProxy/FieldSystemProxyHelper.h b/Engine/Source/Runtime/Experimental/Chaos/Public/PhysicsProxy/FieldSystemProxyHelper.h
--- a/Engine/Source/Runtime/Experimental/Chaos/Public/PhysicsProxy/FieldSystemProxyHelper.h
+++ b/Engine/Source/Runtime/Experimental/Chaos/Public/PhysicsProxy/FieldSystemProxyHelper.h
@@ -76,8 +76,28 @@ namespace Chaos
 			ExecutionDatas.SampleIndices.SetNum(FilteredHandles.Num(),false);
 			InsideHandles.SetNum(FilteredHandles.Num(),false);
 
-			auto FillExecutionDatas = [&ExecutionDatas,&FieldCommand,&InsideHandles](const FVec3& SamplePosition, Chaos::FGeometryParticleHandle* ParticleHandle, int32& HandleIndex)
+			auto FillExecutionDatas = [&ExecutionDatas,&FieldCommand,&InsideHandles](FVec3 SamplePosition, Chaos::FGeometryParticleHandle* ParticleHandle, int32& HandleIndex)
 			{
+				if (FPBDRigidClusteredParticleHandle* ClusterHandle = ParticleHandle->CastToClustered())
+				{
+					// Disabled clustered particles that are driven by a parent, contain particle 
+					// positions in local space. The field system requires the transformation of 
+					// the disabled child particles into world space.
+					if (ClusterHandle->Disabled() == true)
+					{
+						if (FPBDRigidClusteredParticleHandle* ParentHandle = ClusterHandle->Parent())
+						{
+							if (ParentHandle->Disabled() == false)
+							{
+								const FRigidTransform3 ParentWorldTM(ParentHandle->P(), ParentHandle->Q());
+								const FRigidTransform3 ChildFrame = ClusterHandle->ChildToParent() * ParentWorldTM;
+								SamplePosition = ChildFrame.GetTranslation();
+							}
+						}
+					}
+				}
+
+
 				if (FieldCommand.BoundingBox.IsInside(SamplePosition))
 				{
 					ExecutionDatas.SamplePositions[HandleIndex] = SamplePosition;
