Fix chaos unit test + point constraint in corotated space + update of the angular velcoity for the drive damping

diff --git a/Engine/Source/Runtime/Experimental/Chaos/Private/Chaos/Joint/PBDJointCachedSolverGaussSeidel.cpp b/Engine/Source/Runtime/Experimental/Chaos/Private/Chaos/Joint/PBDJointCachedSolverGaussSeidel.cpp
--- a/Engine/Source/Runtime/Experimental/Chaos/Private/Chaos/Joint/PBDJointCachedSolverGaussSeidel.cpp
+++ b/Engine/Source/Runtime/Experimental/Chaos/Private/Chaos/Joint/PBDJointCachedSolverGaussSeidel.cpp
@@ -27,6 +27,9 @@ void FPBDJointCachedSolver::InitDerivedState()
 	ComputeBodyState(1);
 
 	ConnectorRs[1].EnforceShortestArcWith(ConnectorRs[0]);
+
+	ConnectorWDts[0] = FRotation3::CalculateAngularVelocity(InitConnectorRs[0], ConnectorRs[0], 1.0f);
+	ConnectorWDts[1] = FRotation3::CalculateAngularVelocity(InitConnectorRs[1], ConnectorRs[1], 1.0f);
 }
 
 void FPBDJointCachedSolver::ComputeBodyState(const int32 BodyIndex)
@@ -357,7 +360,7 @@ void FPBDJointCachedSolver::InitPositionConstraints(
 	if (bLinearLocked[0] && bLinearLocked[1] && bLinearLocked[2])
 	{
 		// Hard point constraint (most common case)
-		InitPointPositionConstraint();
+		InitPointPositionConstraint(Dt);
 	}
 	else if (bLinearLimited[0] && bLinearLimited[1] && bLinearLimited[2])
 	{
@@ -460,45 +463,57 @@ void FPBDJointCachedSolver::InitPositionConstraintDatas(
 void FPBDJointCachedSolver::InitPointPositionMass()
 {
 	FMatrix33 IM = FMatrix33::Zero;
-	if(bChaos_Joint_StablePointConstraint)
+	if(InvM(0) > 0.0)
 	{
-		if(Body(0).IsDynamic() && InvM(0) > 0.0)
-		{
-			IM += Utilities::ComputeJointFactorMatrix(PointConstraints.ConstraintArms[0], InvI(0), InvM(0));
-		}
-		if(Body(1).IsDynamic() && InvM(1) > 0.0)
-		{
-			IM += Utilities::ComputeJointFactorMatrix(PointConstraints.ConstraintArms[1], InvI(1), InvM(1));
-		}
-		IM = IM.Inverse();
+		IM += Utilities::ComputeJointFactorMatrix(PointConstraints.ConstraintArms[0], InvI(0), InvM(0));
 	}
-	else
+	if(InvM(1) > 0.0)
 	{
-		FVec3 ID = FVec3::ZeroVector;
-		if(Body(0).IsDynamic() && InvM(0) > 0.0)
-		{
-			ID += Utilities::ComputeDiagonalJointFactorMatrix(PointConstraints.ConstraintArms[0], InvI(0), InvM(0));
-		}
-		if(Body(1).IsDynamic() && InvM(1) > 0.0)
-		{
-			ID += Utilities::ComputeDiagonalJointFactorMatrix(PointConstraints.ConstraintArms[1], InvI(1), InvM(1));
-		}
-		IM.SetAt(0,0, (FMath::Abs(ID.X) > SMALL_NUMBER) ? 1.0 / ID.X : 0.0);
-		IM.SetAt(1,1, (FMath::Abs(ID.Y) > SMALL_NUMBER) ? 1.0 / ID.Y : 0.0);
-		IM.SetAt(2,2, (FMath::Abs(ID.Z) > SMALL_NUMBER) ? 1.0 / ID.Z : 0.0);
+		IM += Utilities::ComputeJointFactorMatrix(PointConstraints.ConstraintArms[1], InvI(1), InvM(1));
 	}
+	IM = IM.Inverse();
 	PointConstraints.ConstraintIM = IM;
 	PointConstraints.ConstraintLambda = FVec3::ZeroVector;
 }
 
-void FPBDJointCachedSolver::InitPointPositionConstraint()
+void FPBDJointCachedSolver::InitPointPositionConstraint(const FReal Dt)
 {
-	const FVec3 ConstraintArm0 = ConnectorXs[0] - P(0);
-	const FVec3 ConstraintArm1 = ConnectorXs[1] - P(1);
-	
-	PointConstraints.UpdateDatas(ConnectorXs[1] - ConnectorXs[0], ConstraintArm0, ConstraintArm1);
+	if(bChaos_Joint_StablePointConstraint)
+	{
+		const FVec3 ConstraintArm0 = ConnectorXs[0] - P(0);
+		const FVec3 ConstraintArm1 = ConnectorXs[1] - P(1);
+		
+		PointConstraints.UpdateDatas(ConnectorXs[1] - ConnectorXs[0], ConstraintArm0, ConstraintArm1);
+		
+		InitPointPositionMass();
+	}
+	else
+	{
+		const FVec3 ConstraintArm0 = ConnectorXs[0] - P(0);
+		const FVec3 ConstraintArm1 = ConnectorXs[1] - P(1);
 
-	InitPointPositionMass();
+		const FMatrix33 R0M = ConnectorRs[0].ToMatrix(); 
+		const FVec3 CX = ConnectorRs[0].Inverse().RotateVector(ConnectorXs[1] - ConnectorXs[0]);
+
+		for(int32 ConstraintIndex = 0; ConstraintIndex < 3; ++ConstraintIndex)
+		{
+			const FVec3 ConstraintAxis = R0M.GetColumn(ConstraintIndex);
+			
+			PositionConstraints.ConstraintLimits[ConstraintIndex] = 0.0;
+			PositionConstraints.MotionType[ConstraintIndex] = EJointMotionType::Locked;
+
+			PositionConstraints.UpdateDatas(ConstraintIndex, ConstraintAxis, CX[ConstraintIndex],
+				0.0, false, ConstraintArm0, ConstraintArm1);
+
+			const FVec3 CV0 = V(0) + FVec3::CrossProduct(W(0), ConstraintArm0);
+			const FVec3 CV1 = V(1) + FVec3::CrossProduct(W(1), ConstraintArm1);
+			const FVec3 CV = CV1 - CV0;
+
+			InitConstraintAxisLinearVelocities[ConstraintIndex] = FVec3::DotProduct(CV, ConstraintAxis);
+
+			InitPositionDatasMass(PositionConstraints, ConstraintIndex, Dt);
+		}
+	}
 }
 
 void FPBDJointCachedSolver::InitSphericalPositionConstraint(
@@ -565,9 +580,9 @@ void FPBDJointCachedSolver::ApplyPositionConstraints(
 	}
 }
 
-FVec3 FPBDJointCachedSolver::ComputePointCX() 
+FVec3 FPBDJointCachedSolver::ComputePointCX(const bool bLinearCompute) 
 {
-	if(bChaos_Joint_StablePointConstraint)
+	if(!bLinearCompute)
 	{
 		// Non-linear version
 		UpdateDerivedState();
@@ -584,13 +599,12 @@ FVec3 FPBDJointCachedSolver::ComputePointCX()
 
 void FPBDJointCachedSolver::ApplyPointPositionConstraint()
 {
-	const FVec3 CX = ComputePointCX();
+	const FVec3 CX = ComputePointCX(false);
 
 	if (CX.SizeSquared() > PositionTolerance * PositionTolerance)
 	{
 		// Calculate constraint correction
-		const FVec3 DX = (bChaos_Joint_StablePointConstraint) ? SolverStiffness * PointConstraints.ConstraintStiffness * Utilities::Multiply(PointConstraints.ConstraintIM, CX) :
-																SolverStiffness * PointConstraints.ConstraintStiffness * PointConstraints.ConstraintIM.GetDiagonal() * CX;
+		const FVec3 DX = SolverStiffness * PointConstraints.ConstraintStiffness * Utilities::Multiply(PointConstraints.ConstraintIM, CX);
 
 		ChaosSolverEnsure(!DX.ContainsNaN());
 
@@ -1162,8 +1176,8 @@ void FPBDJointCachedSolver::SolveRotationConstraintSoft(
 	FReal AngVelDt = 0;
 	if (RotationConstraints.ConstraintSoftDamping[ConstraintIndex] > KINDA_SMALL_NUMBER)
 	{
-		const FVec3 W0Dt = FRotation3::CalculateAngularVelocity(InitConnectorRs[0], ConnectorRs[0], 1.0f);
-		const FVec3 W1Dt = FRotation3::CalculateAngularVelocity(InitConnectorRs[1], ConnectorRs[1], 1.0f);
+		const FVec3 W0Dt = FVec3(Body(0).DQ()) + ConnectorWDts[0];
+		const FVec3 W1Dt = FVec3(Body(1).DQ()) + ConnectorWDts[1];
 		AngVelDt = TargetVel * Dt + FVec3::DotProduct(RotationConstraints.ConstraintAxis[ConstraintIndex] , W0Dt - W1Dt);
 	}
 
@@ -1595,8 +1609,8 @@ void FPBDJointCachedSolver::ApplyAxisRotationDrive(
 	FReal AngVelDt = 0;
 	if (RotationDrives.ConstraintSoftDamping[ConstraintIndex] > KINDA_SMALL_NUMBER)
 	{
-		const FVec3 W0Dt = FRotation3::CalculateAngularVelocity(InitConnectorRs[0], ConnectorRs[0], 1.0f);
-		const FVec3 W1Dt = FRotation3::CalculateAngularVelocity(InitConnectorRs[1], ConnectorRs[1], 1.0f);
+		const FVec3 W0Dt = FVec3(Body(0).DQ()) + ConnectorWDts[0];
+		const FVec3 W1Dt = FVec3(Body(1).DQ()) + ConnectorWDts[1];
 		AngVelDt = RotationDrives.ConstraintVX[ConstraintIndex] * Dt + FVec3::DotProduct(RotationDrives.ConstraintAxis[ConstraintIndex] , W0Dt - W1Dt);
 	}
 
@@ -1680,8 +1694,8 @@ void FPBDJointCachedSolver::CalculateAngularConstraintPadding(
 		SetAngularConstraintPadding(ConstraintIndex, 0.0f);
 
 		// Calculate the velocity we want to match
-		const FVec3 W0Dt = FRotation3::CalculateAngularVelocity(InitConnectorRs[0], ConnectorRs[0], 1.0f);
-		const FVec3 W1Dt = FRotation3::CalculateAngularVelocity(InitConnectorRs[1], ConnectorRs[1], 1.0f);
+		const FVec3 W0Dt = FVec3(Body(0).DQ()) + ConnectorWDts[0];
+		const FVec3 W1Dt = FVec3(Body(1).DQ()) + ConnectorWDts[1];
 		const FReal AxisWDt = FVec3::DotProduct(W1Dt - W0Dt, RotationConstraints.ConstraintAxis[(int32)ConstraintIndex]);
 
 		// Calculate the padding to apply to the constraint that will result in the
diff --git a/Engine/Source/Runtime/Experimental/Chaos/Public/Chaos/Joint/PBDJointCachedSolverGaussSeidel.h b/Engine/Source/Runtime/Experimental/Chaos/Public/Chaos/Joint/PBDJointCachedSolverGaussSeidel.h
--- a/Engine/Source/Runtime/Experimental/Chaos/Public/Chaos/Joint/PBDJointCachedSolverGaussSeidel.h
+++ b/Engine/Source/Runtime/Experimental/Chaos/Public/Chaos/Joint/PBDJointCachedSolverGaussSeidel.h
@@ -295,7 +295,7 @@ struct FPointConstraintDatas
 			const FPBDJointSolverSettings& SolverSettings,
 			const FPBDJointSettings& JointSettings);
 		
-		void InitPointPositionConstraint();
+		void InitPointPositionConstraint(const FReal Dt);
 
 		void InitPointPositionMass();
 		
@@ -337,7 +337,7 @@ struct FPointConstraintDatas
 		
 		void ApplyPointPositionConstraint();
 		
-		FVec3 ComputePointCX();
+		FVec3 ComputePointCX(const bool bLinearCompute);
 		
 		void ApplyAxisPositionConstraint(
 			const int32 ConstraintIndex,
@@ -573,6 +573,7 @@ struct FPointConstraintDatas
 		// World-space constraint settings
 		FVec3 ConnectorXs[MaxConstrainedBodies];			// World-space joint connector positions
 		FRotation3 ConnectorRs[MaxConstrainedBodies];		// World-space joint connector rotations
+		FVec3 ConnectorWDts[MaxConstrainedBodies];			// World-space joint connector angular velocities * dt
 
 		// XPBD Initial iteration world-space body state
 		FVec3 InitConnectorXs[MaxConstrainedBodies];		// World-space joint connector positions
