Chaos - fix particle bounds for compound shapes with child transforms

diff --git a/Engine/Source/Runtime/Experimental/Chaos/Private/Chaos/GeometryParticles.cpp b/Engine/Source/Runtime/Experimental/Chaos/Private/Chaos/GeometryParticles.cpp
--- a/Engine/Source/Runtime/Experimental/Chaos/Private/Chaos/GeometryParticles.cpp
+++ b/Engine/Source/Runtime/Experimental/Chaos/Private/Chaos/GeometryParticles.cpp
@@ -203,16 +203,18 @@ namespace Chaos
 
 	void FPerShapeData::UpdateWorldSpaceState(const FRigidTransform3& WorldTransform, const FVec3& BoundsExpansion)
 	{
-		const FImplicitObject* LeafGeometry = GetLeafGeometry();
+		FRigidTransform3 LeafWorldTransform = WorldTransform;
 		if (bHasCachedLeafInfo)
 		{
 			FPerShapeDataCachedLeafInfo& LeafInfo = *static_cast<FPerShapeDataCachedLeafInfo*>(this);
-			LeafInfo.LeafWorldTransform = FRigidTransform3::MultiplyNoScale(FRigidTransform3(GetLeafRelativeTransform()), WorldTransform);
+			LeafWorldTransform = FRigidTransform3::MultiplyNoScale(FRigidTransform3(GetLeafRelativeTransform()), WorldTransform);
+			LeafInfo.LeafWorldTransform = LeafWorldTransform;
 		}
 
+		const FImplicitObject* LeafGeometry = GetLeafGeometry();
 		if ((LeafGeometry != nullptr) && LeafGeometry->HasBoundingBox())
 		{
-			WorldSpaceInflatedShapeBounds = LeafGeometry->CalculateTransformedBounds(WorldTransform).ThickenSymmetrically(BoundsExpansion);;
+			WorldSpaceInflatedShapeBounds = LeafGeometry->CalculateTransformedBounds(LeafWorldTransform).ThickenSymmetrically(BoundsExpansion);;
 		}
 	}
 
