10% improvement to the Solver time for simple shapes by reducing cache misses

diff --git a/Engine/Source/Runtime/Experimental/Chaos/Private/Chaos/CollisionResolution.cpp b/Engine/Source/Runtime/Experimental/Chaos/Private/Chaos/CollisionResolution.cpp
--- a/Engine/Source/Runtime/Experimental/Chaos/Private/Chaos/CollisionResolution.cpp
+++ b/Engine/Source/Runtime/Experimental/Chaos/Private/Chaos/CollisionResolution.cpp
@@ -1636,8 +1636,9 @@ namespace Chaos
 		template<ECollisionUpdateType UpdateType>
 		inline void UpdateConstraintFromGeometryImpl(FPBDCollisionConstraint& Constraint, const FRigidTransform3& WorldTransform0, const FRigidTransform3& WorldTransform1, const FReal Dt)
 		{
-			CONDITIONAL_SCOPE_CYCLE_COUNTER(STAT_Collisions_UpdateConstraintFromGeometryInternal, ConstraintsDetailedStats);
-			INC_DWORD_STAT(STAT_ChaosCollisionCounter_NumContactUpdates);
+			// These add an extra 2% overhead in the simple shape cases
+			//CONDITIONAL_SCOPE_CYCLE_COUNTER(STAT_Collisions_UpdateConstraintFromGeometryInternal, ConstraintsDetailedStats);
+			//INC_DWORD_STAT(STAT_ChaosCollisionCounter_NumContactUpdates);
 
 			// @todo(chaos): remove
 			//const FVec3 OriginalContactPositionLocal0 = WorldTransform0.InverseTransformPosition(Constraint.Manifold.Location);
diff --git a/Engine/Source/Runtime/Experimental/Chaos/Public/Chaos/Collision/CollisionConstraintAllocator.h b/Engine/Source/Runtime/Experimental/Chaos/Public/Chaos/Collision/CollisionConstraintAllocator.h
--- a/Engine/Source/Runtime/Experimental/Chaos/Public/Chaos/Collision/CollisionConstraintAllocator.h
+++ b/Engine/Source/Runtime/Experimental/Chaos/Public/Chaos/Collision/CollisionConstraintAllocator.h
@@ -163,11 +163,11 @@ namespace Chaos
 		 * @note Nothing outside of thie allocator should hold a pointer to the midphase, or any constraints 
 		 * it creates for more than the duration of the tick. Except the IslandManager :| 
 		*/
-		FParticlePairMidPhase* GetParticlePairMidPhase(FGeometryParticleHandle* Particle0, FGeometryParticleHandle* Particle1)
+		FParticlePairMidPhase* GetParticlePairMidPhase(FGeometryParticleHandle* Particle0, FGeometryParticleHandle* Particle1, FGeometryParticleHandle* SearchParticlePerformanceHint)
 		{
 			// NOTE: Called from the CollisionDetection parellel-for loop.
 
-			FParticlePairMidPhase* MidPhase = FindParticlePairMidPhaseImpl(Particle0, Particle1);
+			FParticlePairMidPhase* MidPhase = FindParticlePairMidPhaseImpl(Particle0, Particle1, SearchParticlePerformanceHint);
 			if (MidPhase == nullptr)
 			{
 				MidPhase = CreateParticlePairMidPhase(Particle0, Particle1);
@@ -179,9 +179,9 @@ namespace Chaos
 		/**
 		 * @brief Return a midphase for a particle pair only if it already exists
 		*/
-		FParticlePairMidPhase* FindParticlePairMidPhase(FGeometryParticleHandle* Particle0, FGeometryParticleHandle* Particle1)
+		FParticlePairMidPhase* FindParticlePairMidPhase(FGeometryParticleHandle* Particle0, FGeometryParticleHandle* Particle1, FGeometryParticleHandle* SearchParticlePerformanceHint)
 		{
-			return FindParticlePairMidPhaseImpl(Particle0, Particle1);
+			return FindParticlePairMidPhaseImpl(Particle0, Particle1, SearchParticlePerformanceHint);
 		}
 
 		/**
@@ -234,7 +234,7 @@ namespace Chaos
 					Swap(Particle0, Particle1);
 				}
 
-				FParticlePairMidPhase* MidPhase = GetParticlePairMidPhase(Particle0, Particle1);
+				FParticlePairMidPhase* MidPhase = GetParticlePairMidPhase(Particle0, Particle1, SourceConstraint.Particle[0]);
 
 				// We may be adding multiple constraints for the same particle pair, so we need
 				// to make sure the map is up to date in the case where we just created a new MidPhase
@@ -314,10 +314,10 @@ namespace Chaos
 			PruneExpiredMidPhases();
 		}
 
-		FParticlePairMidPhase* FindParticlePairMidPhaseImpl(FGeometryParticleHandle* Particle0, FGeometryParticleHandle* Particle1)
+		FParticlePairMidPhase* FindParticlePairMidPhaseImpl(FGeometryParticleHandle* Particle0, FGeometryParticleHandle* Particle1, FGeometryParticleHandle* SearchParticle)
 		{
 			// Find the existing midphase from one of the particle's lists of midphases
-			FGeometryParticleHandle* SearchParticle = (Particle0->ParticleCollisions().Num() < Particle1->ParticleCollisions().Num()) ? Particle0 : Particle1;
+//			FGeometryParticleHandle* SearchParticle2 = (Particle0->ParticleCollisions().Num() <= Particle1->ParticleCollisions().Num()) ? Particle0 : Particle1;
 
 			const FCollisionParticlePairKey Key = FCollisionParticlePairKey(Particle0, Particle1);
 			return SearchParticle->ParticleCollisions().FindMidPhase(Key.GetKey());
@@ -597,4 +597,4 @@ namespace Chaos
 	}
 
 
-}
\ No newline at end of file
+}
diff --git a/Engine/Source/Runtime/Experimental/Chaos/Public/Chaos/Collision/NarrowPhase.h b/Engine/Source/Runtime/Experimental/Chaos/Public/Chaos/Collision/NarrowPhase.h
--- a/Engine/Source/Runtime/Experimental/Chaos/Public/Chaos/Collision/NarrowPhase.h
+++ b/Engine/Source/Runtime/Experimental/Chaos/Public/Chaos/Collision/NarrowPhase.h
@@ -34,9 +34,9 @@ namespace Chaos
 		FReal GetBoundsVelocityInflation() const { return BoundsVelocityInflation; }
 		void SetBoundsVelocityInflation(const FReal InBoundsVelocityInflation) { BoundsVelocityInflation = InBoundsVelocityInflation; }
 
-		void GenerateCollisions(FReal Dt, TGeometryParticleHandle<FReal, 3>* Particle0, TGeometryParticleHandle<FReal, 3>* Particle1, const bool bForceDisableCCD)
+		void GenerateCollisions(FReal Dt, TGeometryParticleHandle<FReal, 3>* Particle0, TGeometryParticleHandle<FReal, 3>* Particle1, TGeometryParticleHandle<FReal, 3>* SearchParticlePerfHint, const bool bForceDisableCCD)
 		{
-			FParticlePairMidPhase* MidPhase = CollisionAllocator->GetParticlePairMidPhase(Particle0, Particle1);
+			FParticlePairMidPhase* MidPhase = CollisionAllocator->GetParticlePairMidPhase(Particle0, Particle1, SearchParticlePerfHint);
 			if (MidPhase != nullptr)
 			{
 				Context.bForceDisableCCD = bForceDisableCCD;
@@ -44,9 +44,9 @@ namespace Chaos
 			}
 		}
 
-		FParticlePairMidPhase* GetParticlePairMidPhase(TGeometryParticleHandle<FReal, 3>* Particle0, TGeometryParticleHandle<FReal, 3>* Particle1)
+		FParticlePairMidPhase* GetParticlePairMidPhase(TGeometryParticleHandle<FReal, 3>* Particle0, TGeometryParticleHandle<FReal, 3>* Particle1, TGeometryParticleHandle<FReal, 3>* SearchParticlePerfHint )
 		{
-			return CollisionAllocator->GetParticlePairMidPhase(Particle0, Particle1);
+			return CollisionAllocator->GetParticlePairMidPhase(Particle0, Particle1, SearchParticlePerfHint);
 		}
 
 		// Use this function if a Mid phase pair is already allocated
diff --git a/Engine/Source/Runtime/Experimental/Chaos/Public/Chaos/Collision/ParticlePairBroadPhase.h b/Engine/Source/Runtime/Experimental/Chaos/Public/Chaos/Collision/ParticlePairBroadPhase.h
--- a/Engine/Source/Runtime/Experimental/Chaos/Public/Chaos/Collision/ParticlePairBroadPhase.h
+++ b/Engine/Source/Runtime/Experimental/Chaos/Public/Chaos/Collision/ParticlePairBroadPhase.h
@@ -47,7 +47,7 @@ namespace Chaos
 
 					if ((ParticleA != nullptr) && (ParticleB != nullptr))
 					{
-						ProduceOverlaps(Dt, NarrowPhase, ParticleA, ParticleB);
+						ProduceOverlaps(Dt, NarrowPhase, ParticleA, ParticleB, ParticleA);
 					}
 				}
 			}
@@ -62,7 +62,7 @@ namespace Chaos
 						{
 							if (ParticleB != nullptr)
 							{
-								ProduceOverlaps(Dt, NarrowPhase, ParticleA, ParticleB);
+								ProduceOverlaps(Dt, NarrowPhase, ParticleA, ParticleB, ParticleA);
 							}
 						}
 					}
@@ -75,13 +75,14 @@ namespace Chaos
 			FReal Dt,
 			FNarrowPhase& NarrowPhase,
 			FParticleHandle* ParticleA,
-			FParticleHandle* ParticleB)
+			FParticleHandle* ParticleB,
+			FParticleHandle* SearchParticle)
 		{
 			const FAABB3& Box0 = ParticleA->WorldSpaceInflatedBounds();
 			const FAABB3& Box1 = ParticleB->WorldSpaceInflatedBounds();
 			if (Box0.Intersects(Box1))
 			{
-				NarrowPhase.GenerateCollisions(Dt, ParticleA, ParticleB, true);
+				NarrowPhase.GenerateCollisions(Dt, ParticleA, ParticleB, SearchParticle, true);
 			}
 		}
 
diff --git a/Engine/Source/Runtime/Experimental/Chaos/Public/Chaos/Collision/SpatialAccelerationBroadPhase.h b/Engine/Source/Runtime/Experimental/Chaos/Public/Chaos/Collision/SpatialAccelerationBroadPhase.h
--- a/Engine/Source/Runtime/Experimental/Chaos/Public/Chaos/Collision/SpatialAccelerationBroadPhase.h
+++ b/Engine/Source/Runtime/Experimental/Chaos/Public/Chaos/Collision/SpatialAccelerationBroadPhase.h
@@ -238,6 +238,7 @@ namespace Chaos
 				const int32 NumPotentials = PotentialIntersections.Num();
 				int32 NumIntoNarrowPhase = 0;
 				TArray<FParticlePairMidPhase*> MidPhasePairs;
+				MidPhasePairs.Reserve(NumPotentials);
 				for (int32 i = 0; i < NumPotentials; ++i)
 				{
 					auto& Particle2 = *PotentialIntersections[i].GetGeometryParticleHandle_PhysicsThread();
@@ -374,19 +375,27 @@ namespace Chaos
 						Swap(ParticleA, ParticleB);
 					}
 
-					FParticlePairMidPhase* MidPhase = NarrowPhase.GetParticlePairMidPhase(ParticleA, ParticleB);
+					FParticlePairMidPhase* MidPhase = NarrowPhase.GetParticlePairMidPhase(ParticleA, ParticleB, Particle1.Handle());
 					if (MidPhase)
 					{
 						MidPhasePairs.Add(MidPhase);
 					}
 				}
 
+				const int32  PrefetchLookahead = 4;
+
+				for (int32 Index = 0; Index < MidPhasePairs.Num() && Index < PrefetchLookahead; Index++)
+				{
+					MidPhasePairs[Index]->CachePrefetch();
+				}
+
+
 				for (int32 Index = 0; Index < MidPhasePairs.Num(); Index++)
 				{
 					// Prefetch next pair
-					if (Index != MidPhasePairs.Num() - 1)
+					if (Index + PrefetchLookahead  < MidPhasePairs.Num())
 					{
-						MidPhasePairs[Index + 1]->CachePrefetch();
+						MidPhasePairs[Index + PrefetchLookahead]->CachePrefetch();
 					}
 
 					FParticlePairMidPhase* MidPhasePair = MidPhasePairs[Index];
