Re-implement suspension hardstop using manifold injection to create a frictionless collision constraint.

diff --git a/Engine/Plugins/Experimental/ChaosVehiclesPlugin/Source/ChaosVehicles/Private/ChaosWheeledVehicleMovementComponent.cpp b/Engine/Plugins/Experimental/ChaosVehiclesPlugin/Source/ChaosVehicles/Private/ChaosWheeledVehicleMovementComponent.cpp
--- a/Engine/Plugins/Experimental/ChaosVehiclesPlugin/Source/ChaosVehicles/Private/ChaosWheeledVehicleMovementComponent.cpp
+++ b/Engine/Plugins/Experimental/ChaosVehiclesPlugin/Source/ChaosVehicles/Private/ChaosWheeledVehicleMovementComponent.cpp
@@ -597,7 +597,7 @@ void UChaosWheeledVehicleSimulation::ApplySuspensionForces(float DeltaTime)
 
 							Chaos::FPhysicsSolver* Solver = Proxy->GetSolver<Chaos::FPhysicsSolver>();
 
-							Solver->SetSuspensionTargetOnPhysicsThread(Constraint, TargetPos, PWheel.InContact());
+							Solver->SetSuspensionTarget(Constraint, TargetPos, HitResult.ImpactNormal, PWheel.InContact());
 						}
 					}
 				}
diff --git a/Engine/Source/Programs/HeadlessChaos/Private/HeadlessChaosTestConstraints.cpp b/Engine/Source/Programs/HeadlessChaos/Private/HeadlessChaosTestConstraints.cpp
--- a/Engine/Source/Programs/HeadlessChaos/Private/HeadlessChaosTestConstraints.cpp
+++ b/Engine/Source/Programs/HeadlessChaos/Private/HeadlessChaosTestConstraints.cpp
@@ -145,6 +145,7 @@ namespace ChaosTest {
 		SuspensionSettings.SpringStiffness = 0.0f;		// the spring has no effect
 		SuspensionSettings.SpringDamping = 0.0f;
 		SuspensionSettings.Axis = FVec3(0.0f, 0.0f, 1.0f);
+		SuspensionSettings.Normal = FVec3(0.0f, 0.0f, 1.0f);
 
 		{
 			FParticleUniqueIndicesMultithreaded UniqueIndices;
@@ -289,6 +290,7 @@ namespace ChaosTest {
 		SuspensionSettings.SpringDamping = 0.5f;
 		SuspensionSettings.Target = FVec3(0,0,9);
 		SuspensionSettings.Axis = FVec3(0.0f, 0.0f, 1.0f);
+		SuspensionSettings.Normal = FVec3(0.0f, 0.0f, 1.0f);
 
 		TArray<FVec3> SusLocalOffset;
 		SusLocalOffset.Push(FVector(0, 0, -1));
diff --git a/Engine/Source/Runtime/Experimental/Chaos/Private/Chaos/PBDSuspensionConstraints.cpp b/Engine/Source/Runtime/Experimental/Chaos/Private/Chaos/PBDSuspensionConstraints.cpp
--- a/Engine/Source/Runtime/Experimental/Chaos/Private/Chaos/PBDSuspensionConstraints.cpp
+++ b/Engine/Source/Runtime/Experimental/Chaos/Private/Chaos/PBDSuspensionConstraints.cpp
@@ -2,6 +2,30 @@
 
 #include "Chaos/PBDSuspensionConstraints.h"
 #include "Chaos/Evolution/SolverDatas.h"
+#include "Chaos/DebugDrawQueue.h"
+
+// Private includes
+#include "Chaos/Collision/PBDCollisionSolver.h"
+
+bool bChaos_Suspension_Spring_Enabled = true;
+FAutoConsoleVariableRef CVarChaosSuspensionSpringEnabled(TEXT("p.Chaos.Suspension.Spring.Enabled"), bChaos_Suspension_Spring_Enabled, TEXT("Enable/Disable Spring part of suspension constraint"));
+
+bool bChaos_Suspension_Hardstop_Enabled = true;
+FAutoConsoleVariableRef CVarChaosSuspensionHardstopEnabled(TEXT("p.Chaos.Suspension.Hardstop.Enabled"), bChaos_Suspension_Hardstop_Enabled, TEXT("Enable/Disable Hardstop part of suspension constraint"));
+
+bool bChaos_Suspension_VelocitySolve = true;
+FAutoConsoleVariableRef CVarChaosSuspensionVelocitySolve(TEXT("p.Chaos.Suspension.VelocitySolve"), bChaos_Suspension_VelocitySolve, TEXT("Enable/Disable VelocitySolve"));
+
+float Chaos_Suspension_MaxPushoutVelocity = 50.f;
+FAutoConsoleVariableRef CVarChaosSuspensionMaxPushoutVelocity(TEXT("p.Chaos.Suspension.MaxPushoutVelocity"), Chaos_Suspension_MaxPushoutVelocity, TEXT("Chaos Suspension Max Pushout Velocity Value"));
+
+float Chaos_Suspension_MaxPushout = 5.f;
+FAutoConsoleVariableRef CVarChaosSuspensionMaxPushout(TEXT("p.Chaos.Suspension.MaxPushout"), Chaos_Suspension_MaxPushout, TEXT("Chaos Suspension Max Pushout Value"));
+
+#if CHAOS_DEBUG_DRAW
+bool bChaos_Suspension_DebugDraw_Hardstop = false;
+FAutoConsoleVariableRef CVarChaosSuspensionDebugDrawHardstop(TEXT("p.Chaos.Suspension.DebugDraw.Hardstop"), bChaos_Suspension_DebugDraw_Hardstop, TEXT("Debug draw suspension hardstop manifold"));
+#endif
 
 namespace Chaos
 {
@@ -45,6 +69,9 @@ namespace Chaos
 		ConstraintResults.AddDefaulted();
 		ConstraintEnabledStates.Add(true); // Note: assumes always enabled on creation
 		ConstraintSolverBodies.Add(nullptr);
+		StaticCollisionBodies.Add(FSolverBody());
+		CollisionSolvers.Add( new FPBDCollisionSolver() );
+
 		Handles.Add(HandleAllocator.AllocHandle(this, NewIndex));
 		return Handles[NewIndex];
 	}
@@ -65,6 +92,8 @@ namespace Chaos
 			Handles[ConstraintIndex] = nullptr;
 		}
 
+		delete CollisionSolvers[ConstraintIndex];
+
 		// Swap the last constraint into the gap to keep the array packed
 		ConstrainedParticles.RemoveAtSwap(ConstraintIndex);
 		SuspensionLocalOffset.RemoveAtSwap(ConstraintIndex);
@@ -72,6 +101,9 @@ namespace Chaos
 		ConstraintResults.RemoveAtSwap(ConstraintIndex);
 		ConstraintEnabledStates.RemoveAtSwap(ConstraintIndex);
 		ConstraintSolverBodies.RemoveAtSwap(ConstraintIndex);
+		CollisionSolvers.RemoveAtSwap(ConstraintIndex);
+		StaticCollisionBodies.RemoveAtSwap(ConstraintIndex);
+
 		Handles.RemoveAtSwap(ConstraintIndex);
 
 		// Update the handle for the constraint that was moved
@@ -93,6 +125,88 @@ namespace Chaos
 		ConstraintSolverBodies[ConstraintIndex] = SolverData.GetBodyContainer().FindOrAdd(ConstrainedParticles[ConstraintIndex]);
 
 		ConstraintResults[ConstraintIndex].Reset();
+
+		// Hard-stop Collision Manifold Generation
+		{
+			check(CollisionSolvers[ConstraintIndex] != nullptr);
+			check(ConstraintSolverBodies[ConstraintIndex] != nullptr);
+
+			FPBDCollisionSolver* Solver = CollisionSolvers[ConstraintIndex];
+			Solver->Reset(); // clear previous manifolds
+
+			FSolverBody* Body0 = ConstraintSolverBodies[ConstraintIndex];	// vehicle chassis			
+			FSolverBody* Body1 = &StaticCollisionBodies[ConstraintIndex];	// Spoofed terrain
+
+			const FPBDSuspensionSettings& Setting = ConstraintSettings[ConstraintIndex];
+
+			if (Body0->IsDynamic() && Setting.Enabled)
+			{
+				const FVec3& T = Setting.Target;
+
+				// \todo(chaos): we can cache the CoM-relative connector once per frame rather than recalculate per iteration
+				// (we should not be accessing particle state in the solver methods, although this one actually is ok because it only uses frame constrants)
+				const FGenericParticleHandle Particle = ConstrainedParticles[ConstraintIndex];
+				const FVec3& SuspensionActorOffset = SuspensionLocalOffset[ConstraintIndex];
+				const FVec3 SuspensionCoMOffset = Particle->RotationOfMass().UnrotateVector(SuspensionActorOffset - Particle->CenterOfMass());
+				const FVec3 SuspensionCoMAxis = Particle->RotationOfMass().UnrotateVector(Setting.Axis);
+			
+				const FRotation3 BodyQ = Body0->CorrectedQ();
+				const FVec3 BodyP = Body0->CorrectedP();
+				const FVec3 WorldSpaceX = BodyQ.RotateVector(SuspensionCoMOffset) + BodyP;
+				FVec3 AxisWorld = BodyQ.RotateVector(SuspensionCoMAxis);
+				FReal Distance = FVec3::DotProduct(WorldSpaceX - T, AxisWorld);
+
+				// if the suspension is NOT compressed right down to the hard-stop then we won't need a manifold, should be the case 99% time
+				if (Distance > Setting.MinLength)
+				{
+					return;
+				}
+
+				FReal WorldContactDeltaNormal = Setting.MinLength - Distance; // the distance we need to correct by
+				const FVec3 WorldArm = BodyQ.RotateVector(SuspensionCoMOffset);
+
+				// position the spoofed terrain at the same position at body0
+				FVec3 PosBody1 = T + Distance*AxisWorld;
+				Body1->SetP(PosBody1);
+				Body1->SetX(PosBody1);
+				check(Body1->InvM() == 0);
+
+				Solver->SetSolverBodies(Body0, Body1);
+				Solver->SetNumManifoldPoints(1);
+				Solver->SetFriction(0, 0, 0);
+
+#if CHAOS_DEBUG_DRAW
+				if (bChaos_Suspension_DebugDraw_Hardstop)
+				{
+					FVec3 Body0Center = WorldArm + BodyP;
+					FVec3 Body1Center = PosBody1;
+					FReal Radius = 30;
+					FDebugDrawQueue::GetInstance().DrawDebugCircle(Body0Center, Radius, 60, FColor::Yellow, false, -1.0f, 0, 3, FVector(1, 0, 0), FVector(0, 1, 0), false);
+					FDebugDrawQueue::GetInstance().DrawDebugCircle(Body1Center, Radius, 60, FColor::Green, false, -1.0f, 0, 3, FVector(1, 0, 0), FVector(0, 1, 0), false);
+					FString TextOut = FString::Format(TEXT("{0}"), { WorldContactDeltaNormal });
+					FDebugDrawQueue::GetInstance().DrawDebugString(Body0Center + FVec3(0, 50, 50), TextOut, nullptr, FColor::White, -1.f, true, 1.0f);
+				}
+#endif
+
+				// inject a manifold for our suspension Hard-stop - behaves like a regular friction-less collision, prevents the vehicle chassis from ever hitting the ground
+				Solver->SetManifoldPoint(
+					0,										// ManifoldIndex
+					FSolverReal(Dt),						// Delta Time
+					FSolverReal(0.0f),						// Restitution,
+					FSolverReal(0.1f),						// RestitutionVelocityThreshold,
+					FSolverVec3(WorldArm),					// RelativeContactPosition0,
+					FSolverVec3(0, 0, 0),					// RelativeContactPosition1,
+					FSolverVec3(AxisWorld),					// WorldContactNormal
+					FSolverVec3(0, 0, 0),					// WorldContactTangentU,
+					FSolverVec3(0, 0, 0),					// WorldContactTangentV,
+					FSolverReal(-WorldContactDeltaNormal),	// WorldContactDeltaNormal
+					FSolverReal(0),							// WorldContactDeltaTangentU,
+					FSolverReal(0)							// WorldContactDeltaTangentV
+					);
+
+			}
+		}
+
 	}
 
 	void FPBDSuspensionConstraints::ScatterOutput(FReal Dt, FPBDIslandSolverData& SolverData)
@@ -105,15 +219,48 @@ namespace Chaos
 
 	bool FPBDSuspensionConstraints::ApplyPhase1Serial(const FReal Dt, const int32 It, const int32 NumIts, FPBDIslandSolverData& SolverData)
 	{
-		for (int32 ConstraintIndex : SolverData.GetConstraintIndices(ContainerId))
+		if (bChaos_Suspension_Hardstop_Enabled)
 		{
-			ApplySingle(Dt, ConstraintIndex);
+			// Suspension Hardstop
+			for (FPBDCollisionSolver* CollisionSolver : CollisionSolvers)
+			{
+				if (CollisionSolver->NumManifoldPoints() > 0)
+				{
+					FReal MaxPushoutValue = FMath::Min(Chaos_Suspension_MaxPushout, Chaos_Suspension_MaxPushoutVelocity * Dt);
+					CollisionSolver->SolvePositionNoFriction(FSolverReal(Dt), FSolverReal(MaxPushoutValue));
+				}
+			}
+		}
+
+		if (bChaos_Suspension_Spring_Enabled)
+		{
+			// Suspension Spring
+			for (int32 ConstraintIndex : SolverData.GetConstraintIndices(ContainerId))
+			{
+				ApplySingle(Dt, ConstraintIndex);
+			}
 		}
 
 		// @todo(chaos): early iteration termination in FPBDSuspensionConstraints
 		return true;
 	}
 
+	bool FPBDSuspensionConstraints::ApplyPhase2Serial(const FReal Dt, const int32 It, const int32 NumIts, FPBDIslandSolverData& SolverData)
+	{
+		bool bNeedsAnotherIteration = false;
+
+		if (bChaos_Suspension_Hardstop_Enabled && bChaos_Suspension_VelocitySolve)
+		{
+			// Suspension Hardstop
+			for (FPBDCollisionSolver* CollisionSolver : CollisionSolvers)
+			{
+				bNeedsAnotherIteration |= CollisionSolver->SolveVelocity(FSolverReal(Dt), false);
+			}
+		}
+
+		return bNeedsAnotherIteration;
+	}
+
 	void FPBDSuspensionConstraints::ApplySingle(const FReal Dt, int32 ConstraintIndex)
 	{
 		check(ConstraintSolverBodies[ConstraintIndex] != nullptr);
@@ -180,17 +327,17 @@ namespace Chaos
 			{
 				// target point distance is less at min compression limit 
 				// - apply distance constraint to try keep a valid min limit
-				FVec3 Ts = WorldSpaceX + AxisWorld * (Setting.MinLength - Distance);
-				DX = (Ts - WorldSpaceX) * Setting.HardstopStiffness;
+				//FVec3 Ts = WorldSpaceX + AxisWorld * (Setting.MinLength - Distance);
+				//DX = (Ts - WorldSpaceX) * Setting.HardstopStiffness;
 
 				Distance = Setting.MinLength;
 
-				if (PointVelocityAlongAxis < 0.0f)
-				{
-					const FVec3 SpringVelocity = PointVelocityAlongAxis * AxisWorld;
-					DX -= SpringVelocity * Setting.HardstopVelocityCompensation;
-					PointVelocityAlongAxis = 0.0f; //this Dx will cancel velocity, so don't pass PointVelocityAlongAxis on to suspension force calculation 
-				}
+				//if (PointVelocityAlongAxis < 0.0f)
+				//{
+				//	const FVec3 SpringVelocity = PointVelocityAlongAxis * AxisWorld;
+				//	DX -= SpringVelocity * Setting.HardstopVelocityCompensation;
+				//	PointVelocityAlongAxis = 0.0f; //this Dx will cancel velocity, so don't pass PointVelocityAlongAxis on to suspension force calculation 
+				//}
 			}
 
 			{
diff --git a/Engine/Source/Runtime/Experimental/Chaos/Private/PBDRigidsSolver.cpp b/Engine/Source/Runtime/Experimental/Chaos/Private/PBDRigidsSolver.cpp
--- a/Engine/Source/Runtime/Experimental/Chaos/Private/PBDRigidsSolver.cpp
+++ b/Engine/Source/Runtime/Experimental/Chaos/Private/PBDRigidsSolver.cpp
@@ -694,12 +694,12 @@ namespace Chaos
 			});
 	}
 
-	void FPBDRigidsSolver::SetSuspensionTargetOnPhysicsThread(Chaos::FSuspensionConstraint* GTConstraint, const FVector& TargetPos, bool Enabled)
+	void FPBDRigidsSolver::SetSuspensionTarget(Chaos::FSuspensionConstraint* GTConstraint, const FVector& TargetPos, const FVector& Normal, bool Enabled)
 	{
 		EnsureIsInPhysicsThreadContext();
 		FSuspensionConstraintPhysicsProxy* SuspensionProxy = GTConstraint->GetProxy<FSuspensionConstraintPhysicsProxy>();
 		check(SuspensionProxy);
-		SuspensionProxy->UpdateTargetOnPhysicsThread(this, TargetPos, Enabled);
+		SuspensionProxy->UpdateTargetOnPhysicsThread(this, TargetPos, Normal, Enabled);
 	}
 
 	CHAOS_API int32 RewindCaptureNumFrames = -1;
diff --git a/Engine/Source/Runtime/Experimental/Chaos/Private/PhysicsProxy/SuspensionConstraintProxy.cpp b/Engine/Source/Runtime/Experimental/Chaos/Private/PhysicsProxy/SuspensionConstraintProxy.cpp
--- a/Engine/Source/Runtime/Experimental/Chaos/Private/PhysicsProxy/SuspensionConstraintProxy.cpp
+++ b/Engine/Source/Runtime/Experimental/Chaos/Private/PhysicsProxy/SuspensionConstraintProxy.cpp
@@ -106,13 +106,13 @@ void FSuspensionConstraintPhysicsProxy::DestroyOnPhysicsThread(FPBDRigidsSolver*
 	}
 }
 
-
-void FSuspensionConstraintPhysicsProxy::UpdateTargetOnPhysicsThread(FPBDRigidsSolver* RBDSolver, const FVector& TargetPos, bool Enabled)
+void FSuspensionConstraintPhysicsProxy::UpdateTargetOnPhysicsThread(FPBDRigidsSolver* RBDSolver, const FVector& TargetPos, const FVector& Normal, bool Enabled)
 {
 	if (Constraint_PT)
 	{
 		auto& SuspensionConstraints = RBDSolver->GetSuspensionConstraints();
 		SuspensionConstraints.GetSettings(Constraint_PT->GetConstraintIndex()).Target = TargetPos;
+		SuspensionConstraints.GetSettings(Constraint_PT->GetConstraintIndex()).Normal = Normal;
 		SuspensionConstraints.GetSettings(Constraint_PT->GetConstraintIndex()).Enabled = Enabled;
 	}
 }
diff --git a/Engine/Source/Runtime/Experimental/Chaos/Public/Chaos/PBDSuspensionConstraintData.h b/Engine/Source/Runtime/Experimental/Chaos/Public/Chaos/PBDSuspensionConstraintData.h
--- a/Engine/Source/Runtime/Experimental/Chaos/Public/Chaos/PBDSuspensionConstraintData.h
+++ b/Engine/Source/Runtime/Experimental/Chaos/Public/Chaos/PBDSuspensionConstraintData.h
@@ -23,6 +23,7 @@ namespace Chaos
 		MinLength = static_cast<uint64_t>(1) << 8,
 		MaxLength = static_cast<uint64_t>(1) << 9,
 		Axis = static_cast<uint64_t>(1) << 10,
+		Normal = static_cast<uint64_t>(1) << 11,
 
 		DummyFlag
 	};
diff --git a/Engine/Source/Runtime/Experimental/Chaos/Public/Chaos/PBDSuspensionConstraintTypes.h b/Engine/Source/Runtime/Experimental/Chaos/Public/Chaos/PBDSuspensionConstraintTypes.h
--- a/Engine/Source/Runtime/Experimental/Chaos/Public/Chaos/PBDSuspensionConstraintTypes.h
+++ b/Engine/Source/Runtime/Experimental/Chaos/Public/Chaos/PBDSuspensionConstraintTypes.h
@@ -38,11 +38,12 @@ namespace Chaos
 			, MinLength(0)
 			, MaxLength(0)
 			, Axis(0,0,1)
-			, Target(0,0,0)
+			, Target(0, 0, 0)
+			, Normal(0, 0, 0)
 		{
 		}
 
-		FPBDSuspensionSettings(bool InEnabled, FReal InHardstopStiffness, FReal InHardstopVelocityCompensation, FReal InSpringPreload, FReal InSpringStiffness, FReal InDamping, FReal InMinLength, FReal InMaxLength, const FVec3& InAxis, const FVec3& InTarget)
+		FPBDSuspensionSettings(bool InEnabled, FReal InHardstopStiffness, FReal InHardstopVelocityCompensation, FReal InSpringPreload, FReal InSpringStiffness, FReal InDamping, FReal InMinLength, FReal InMaxLength, const FVec3& InAxis, const FVec3& InTarget, const FVec3& InNormal)
 			: Enabled(InEnabled)
 			, HardstopStiffness(InHardstopStiffness)
 			, HardstopVelocityCompensation(InHardstopVelocityCompensation)
@@ -53,6 +54,7 @@ namespace Chaos
 			, MaxLength(InMaxLength)
 			, Axis(InAxis)
 			, Target(InTarget)
+			, Normal(InNormal)
 		{
 		}
 
@@ -66,6 +68,7 @@ namespace Chaos
 		FReal MaxLength;
 		FVec3 Axis;
 		FVec3 Target;
+		FVec3 Normal;
 	};
 
 	class CHAOS_API FPBDSuspensionResults
diff --git a/Engine/Source/Runtime/Experimental/Chaos/Public/Chaos/PBDSuspensionConstraints.h b/Engine/Source/Runtime/Experimental/Chaos/Public/Chaos/PBDSuspensionConstraints.h
--- a/Engine/Source/Runtime/Experimental/Chaos/Public/Chaos/PBDSuspensionConstraints.h
+++ b/Engine/Source/Runtime/Experimental/Chaos/Public/Chaos/PBDSuspensionConstraints.h
@@ -17,6 +17,7 @@ namespace Chaos
 {
 	class FPBDSuspensionConstraints;
 	class FPBDIslandSolverData;
+	class FPBDCollisionSolver;
 
 	class CHAOS_API FPBDSuspensionConstraintHandle : public TIndexedContainerConstraintHandle<FPBDSuspensionConstraints>
 	{
@@ -236,7 +237,7 @@ namespace Chaos
 		void ScatterOutput(FReal Dt, FPBDIslandSolverData& SolverData);
 
 		bool ApplyPhase1Serial(const FReal Dt, const int32 It, const int32 NumIts, FPBDIslandSolverData& SolverData);
-		bool ApplyPhase2Serial(const FReal Dt, const int32 It, const int32 NumIts, FPBDIslandSolverData& SolverData) { return false; }
+		bool ApplyPhase2Serial(const FReal Dt, const int32 It, const int32 NumIts, FPBDIslandSolverData& SolverData);
 
 	protected:
 		using Base::GetConstraintIndex;
@@ -259,6 +260,9 @@ namespace Chaos
 
 		FHandles Handles;
 		FConstraintHandleAllocator HandleAllocator;
+
+		TArray<FPBDCollisionSolver*> CollisionSolvers;
+		TArray<FSolverBody> StaticCollisionBodies;
 	};
 }
 
diff --git a/Engine/Source/Runtime/Experimental/Chaos/Public/Chaos/SuspensionProperties.inl b/Engine/Source/Runtime/Experimental/Chaos/Public/Chaos/SuspensionProperties.inl
--- a/Engine/Source/Runtime/Experimental/Chaos/Public/Chaos/SuspensionProperties.inl
+++ b/Engine/Source/Runtime/Experimental/Chaos/Public/Chaos/SuspensionProperties.inl
@@ -12,6 +12,7 @@ CHAOS_INNER_SUSP_PROPERTY(SuspensionSettings, SpringDamping, FReal);
 CHAOS_INNER_SUSP_PROPERTY(SuspensionSettings, MinLength, FReal);
 CHAOS_INNER_SUSP_PROPERTY(SuspensionSettings, MaxLength, FReal);
 CHAOS_INNER_SUSP_PROPERTY(SuspensionSettings, Axis, FVec3);
+CHAOS_INNER_SUSP_PROPERTY(SuspensionSettings, Normal, FVec3)
 CHAOS_INNER_SUSP_PROPERTY(SuspensionParticleProxy, ParticleProxy, FSingleParticlePhysicsProxy*)
 CHAOS_INNER_SUSP_PROPERTY(SuspensionLocation, Location, FVec3)
 
diff --git a/Engine/Source/Runtime/Experimental/Chaos/Public/PBDRigidsSolver.h b/Engine/Source/Runtime/Experimental/Chaos/Public/PBDRigidsSolver.h
--- a/Engine/Source/Runtime/Experimental/Chaos/Public/PBDRigidsSolver.h
+++ b/Engine/Source/Runtime/Experimental/Chaos/Public/PBDRigidsSolver.h
@@ -219,7 +219,7 @@ namespace Chaos
 
 		FPBDSuspensionConstraints& GetSuspensionConstraints() { return MEvolution->GetSuspensionConstraints(); }
 		const FPBDSuspensionConstraints& GetSuspensionConstraints() const { return MEvolution->GetSuspensionConstraints(); }
-		void SetSuspensionTargetOnPhysicsThread(Chaos::FSuspensionConstraint* GTConstraint, const FVector& TargetPos, bool Enabled);
+		void SetSuspensionTarget(Chaos::FSuspensionConstraint* GTConstraint, const FVector& TargetPos, const FVector& Normal, bool Enabled);
 
 		void EnableRewindCapture(int32 NumFrames, bool InUseCollisionResimCache, TUniquePtr<IRewindCallback>&& RewindCallback = TUniquePtr<IRewindCallback>());
 
diff --git a/Engine/Source/Runtime/Experimental/Chaos/Public/PhysicsProxy/SuspensionConstraintProxy.h b/Engine/Source/Runtime/Experimental/Chaos/Public/PhysicsProxy/SuspensionConstraintProxy.h
--- a/Engine/Source/Runtime/Experimental/Chaos/Public/PhysicsProxy/SuspensionConstraintProxy.h
+++ b/Engine/Source/Runtime/Experimental/Chaos/Public/PhysicsProxy/SuspensionConstraintProxy.h
@@ -46,7 +46,7 @@ public:
 	void DestroyOnGameThread();
 	void DestroyOnPhysicsThread(FPBDRigidsSolver* InSolver);
 
-	void UpdateTargetOnPhysicsThread(FPBDRigidsSolver* InSolver, const FVector& TargetPos, bool Enabled);
+	void UpdateTargetOnPhysicsThread(FPBDRigidsSolver* InSolver, const FVector& TargetPos, const FVector& Normal, bool Enabled);
 
 
 	//
