Chaos - Sphere-Convex manifold

diff --git a/Engine/Source/Runtime/Experimental/Chaos/Private/Chaos/CollisionOneShotManifoldsMiscShapes.cpp b/Engine/Source/Runtime/Experimental/Chaos/Private/Chaos/CollisionOneShotManifoldsMiscShapes.cpp
--- a/Engine/Source/Runtime/Experimental/Chaos/Private/Chaos/CollisionOneShotManifoldsMiscShapes.cpp
+++ b/Engine/Source/Runtime/Experimental/Chaos/Private/Chaos/CollisionOneShotManifoldsMiscShapes.cpp
@@ -30,6 +30,9 @@ namespace Chaos
 	extern FRealSingle Chaos_Collision_Manifold_CapsuleDeepPenetrationFraction;
 	extern FRealSingle Chaos_Collision_Manifold_CapsuleRadialContactFraction;
 
+	extern FRealSingle Chaos_Collision_GJKEpsilon;
+	extern FRealSingle Chaos_Collision_EPAEpsilon;
+
 	namespace Collisions
 	{
 		void ConstructSphereSphereOneShotManifold(
@@ -231,6 +234,97 @@ namespace Chaos
 			}
 		}
 
+		template<typename ConvexType>
+		void ConstructSphereConvexManifoldImpl(const FImplicitSphere3& Sphere, const ConvexType& Convex, const FRigidTransform3& SphereToConvexTransform, const FReal CullDistance, TCArray<FContactPoint, 4>& ContactPoints)
+		{
+			// Transform the sphere into convex space
+			const FVec3 SpherePos = SphereToConvexTransform.TransformPositionNoScale(Sphere.GetCenter());
+			const FReal SphereRadius = Sphere.GetRadius();
+
+			// No margins for the convex, but treat the sphere as a point with a margin
+			FGJKSphereShape GJKSphere(SpherePos, SphereRadius);
+			TGJKShape<ConvexType> GJKConvex(Convex);
+
+			// GJK and EPA tolerances. See comments in GJKContactPointMargin
+			const FReal GJKEpsilon = Chaos_Collision_GJKEpsilon;
+			const FReal EPAEpsilon = Chaos_Collision_EPAEpsilon;
+			FReal ClosestPenetration;
+			int32 ClosestVertexIndexSphere = INDEX_NONE, ClosestVertexIndexConvex = INDEX_NONE;
+			FReal ClosestSupportMaxDelta = FReal(0);
+
+			// Primary contact
+			// NOTE: swapped contact point order to match desired output order (Sphere, Convex)
+			FContactPoint ClosestContactPoint;
+			GJKPenetrationSameSpace(
+				GJKConvex, 
+				GJKSphere, 
+				ClosestPenetration,
+				ClosestContactPoint.ShapeContactPoints[1], 
+				ClosestContactPoint.ShapeContactPoints[0], 
+				ClosestContactPoint.ShapeContactNormal, 
+				ClosestVertexIndexConvex,
+				ClosestVertexIndexSphere,
+				ClosestSupportMaxDelta,
+				GJKEpsilon, 
+				EPAEpsilon);
+
+			// Stop now if beyond cull distance
+			const FReal ClosestPhi = -ClosestPenetration;
+			if (ClosestPhi > CullDistance)
+			{
+				return;
+			}
+
+			// We always use the primary contact so add it to the output now
+			ClosestContactPoint.ShapeContactPoints[0] = SphereToConvexTransform.InverseTransformPositionNoScale(ClosestContactPoint.ShapeContactPoints[0]);
+			ClosestContactPoint.Phi = ClosestPhi;
+			ClosestContactPoint.FaceIndex = INDEX_NONE;
+			ClosestContactPoint.ContactType = EContactPointType::Unknown;
+			ContactPoints.Add(ClosestContactPoint);
+
+			// If the sphere is "large" compared to the convex add more points
+			const FReal SpheerConvexManifoldSizeThreshold = FReal(1);
+			const FReal ConvexSize = Convex.BoundingBox().Extents().GetAbsMax();
+			if (SphereRadius > SpheerConvexManifoldSizeThreshold * ConvexSize)
+			{
+				// Find the convex plane to use - the one most opposing the primary contact normal
+				const int32 ConvexPlaneIndex = Convex.GetMostOpposingPlane(-ClosestContactPoint.ShapeContactNormal);
+				if (ConvexPlaneIndex != INDEX_NONE)
+				{
+					FVec3 ConvexPlanePosition, ConvexPlaneNormal;
+					Convex.GetPlaneNX(ConvexPlaneIndex, ConvexPlaneNormal, ConvexPlanePosition);
+
+					// Project the face verts onto the sphere along the normal and generate speculative contacts
+					// We actually just take a third of the points, chosen arbitrarily. This may not be the best choice for convexes where
+					// most of the face verts are close to each other with a few outliers. 
+					// @todo(chaos): a better option would be to build a triangle of contacts around the primary contact, with the verts projected into the convex face
+					const int32 NumConvexPlaneVertices = Convex.NumPlaneVertices(ConvexPlaneIndex);
+					const int32 PlaneVertexStride = FMath::Max(1, NumConvexPlaneVertices / 3);
+					for (int32 PlaneVertexIndex = 0; PlaneVertexIndex < NumConvexPlaneVertices; PlaneVertexIndex += PlaneVertexStride)
+					{
+						const FVec3 ConvexPlaneVertex = Convex.GetVertex(Convex.GetPlaneVertex(ConvexPlaneIndex, PlaneVertexIndex));
+						const FReal ConvexContactDistance = Utilities::RaySphereIntersectionDistance(ConvexPlaneVertex, ClosestContactPoint.ShapeContactNormal, SpherePos, SphereRadius);
+						if (ConvexContactDistance < CullDistance)
+						{
+							FContactPoint& ConvexContactPoint = ContactPoints[ContactPoints.Add()];
+							ConvexContactPoint.ShapeContactPoints[0] = SphereToConvexTransform.InverseTransformPositionNoScale(ConvexPlaneVertex + ClosestContactPoint.ShapeContactNormal * ConvexContactDistance);
+							ConvexContactPoint.ShapeContactPoints[1] = ConvexPlaneVertex;
+							ConvexContactPoint.ShapeContactNormal = ClosestContactPoint.ShapeContactNormal;
+							ConvexContactPoint.Phi = ConvexContactDistance;
+							ConvexContactPoint.FaceIndex = INDEX_NONE;
+							ConvexContactPoint.ContactType = EContactPointType::VertexPlane;
+
+							if (ContactPoints.IsFull())
+							{
+								break;
+							}
+						}
+					}
+				}
+			}
+		}
+
+
 		void ConstructSphereConvexManifold(const TSphere<FReal, 3>& Sphere, const FRigidTransform3& SphereTransform, const FImplicitObject3& Convex, const FRigidTransform3& ConvexTransform, const FReal Dt, FPBDCollisionConstraint& Constraint)
 		{
 			// We only build one shot manifolds once
@@ -238,11 +332,33 @@ namespace Chaos
 			ensure(SphereTransform.GetScale3D() == FVec3(1.0f, 1.0f, 1.0f));
 			ensure(ConvexTransform.GetScale3D() == FVec3(1.0f, 1.0f, 1.0f));
 
-			// @todo(chaos): support manifold maintenance
-			Constraint.ResetActiveManifoldContacts();
+			const FRigidTransform3 SphereToConvexTransform = SphereTransform.GetRelativeTransformNoScale(ConvexTransform);
 
-			FContactPoint ContactPoint = SphereConvexContactPoint(Sphere, SphereTransform, Convex, ConvexTransform);
-			if (ContactPoint.Phi < Constraint.GetCullDistance())
+			TCArray<FContactPoint, 4> ContactPoints;
+			if (const FImplicitBox3* RawBox = Convex.template GetObject<FImplicitBox3>())
+			{
+				ConstructSphereConvexManifoldImpl(Sphere, *RawBox, SphereToConvexTransform, Constraint.GetCullDistance(), ContactPoints);
+			}
+			else if (const TImplicitObjectScaled<FImplicitConvex3>* ScaledConvex = Convex.template GetObject<TImplicitObjectScaled<FImplicitConvex3>>())
+			{
+				ConstructSphereConvexManifoldImpl(Sphere, *ScaledConvex, SphereToConvexTransform, Constraint.GetCullDistance(), ContactPoints);
+			}
+			else if (const TImplicitObjectInstanced<FImplicitConvex3>* InstancedConvex = Convex.template GetObject<TImplicitObjectInstanced<FImplicitConvex3>>())
+			{
+				ConstructSphereConvexManifoldImpl(Sphere, *InstancedConvex, SphereToConvexTransform, Constraint.GetCullDistance(), ContactPoints);
+			}
+			else if (const FImplicitConvex3* RawConvex = Convex.template GetObject<FImplicitConvex3>())
+			{
+				ConstructSphereConvexManifoldImpl(Sphere, *RawConvex, SphereToConvexTransform, Constraint.GetCullDistance(), ContactPoints);
+			}
+			else
+			{
+				check(false);
+			}
+
+			// Add the points to the constraint
+			Constraint.ResetActiveManifoldContacts();
+			for (FContactPoint& ContactPoint : ContactPoints)
 			{
 				Constraint.AddOneshotManifoldContact(ContactPoint);
 			}
diff --git a/Engine/Source/Runtime/Experimental/Chaos/Public/Chaos/GJKShape.h b/Engine/Source/Runtime/Experimental/Chaos/Public/Chaos/GJKShape.h
--- a/Engine/Source/Runtime/Experimental/Chaos/Public/Chaos/GJKShape.h
+++ b/Engine/Source/Runtime/Experimental/Chaos/Public/Chaos/GJKShape.h
@@ -88,6 +88,34 @@ namespace Chaos
 		const FReal Margin;
 	};
 
+	/**
+	 * @brief A sphere with minimal API for use in GJK/EPA
+	 * Equivalent to TGJKCoreShape<FImplicitSphere3> without any indirection.
+	*/
+	class FGJKSphereShape
+	{
+	public:
+		FGJKSphereShape(const FVec3 InPos, const FReal InRadius) 
+			: Pos(InPos), Radius(InRadius) 
+		{
+		}
+
+		inline const FVec3& SupportCore(const FVec3& Direction, const FReal InMargin, FReal* MaxMarginDelta, int32& VertexIndex) const
+		{
+			VertexIndex = 0;
+			return Pos;
+		}
+
+		inline FReal GetMargin() const
+		{
+			return Radius;
+		}
+
+	private:
+		FVec3 Pos;
+		FReal Radius;
+	};
+
 
 	// Utility for creating TGJKShape objects using template parameter deduction
 	template<typename T_SHAPE>
diff --git a/Engine/Source/Runtime/Experimental/Chaos/Public/Chaos/Utilities.h b/Engine/Source/Runtime/Experimental/Chaos/Public/Chaos/Utilities.h
--- a/Engine/Source/Runtime/Experimental/Chaos/Public/Chaos/Utilities.h
+++ b/Engine/Source/Runtime/Experimental/Chaos/Public/Chaos/Utilities.h
@@ -578,29 +578,49 @@ namespace Chaos
 			C2 = P2 + T * D2;
 		}
 
-		inline FReal ClosestTimeOnLineSegment(const FVector& Point, const FVector& StartPoint, const FVector& EndPoint)
+		template<typename T>
+		inline T ClosestTimeOnLineSegment(const TVec3<T>& Point, const TVec3<T>& StartPoint, const TVec3<T>& EndPoint)
 		{
-			const FVector Segment = EndPoint - StartPoint;
-			const FVector VectToPoint = Point - StartPoint;
+			const TVec3<T> Segment = EndPoint - StartPoint;
+			const TVec3<T> VectToPoint = Point - StartPoint;
 
 			// See if closest point is before StartPoint
-			const FVector::FReal Dot1 = FVec3::DotProduct(VectToPoint, Segment);
+			const T Dot1 = TVec3<T>::DotProduct(VectToPoint, Segment);
 			if (Dot1 <= 0)
 			{
-				return FReal(0);
+				return T(0);
 			}
 
 			// See if closest point is beyond EndPoint
-			const FVector::FReal Dot2 = FVec3::DotProduct(Segment, Segment);
+			const T Dot2 = TVec3<T>::DotProduct(Segment, Segment);
 			if (Dot2 <= Dot1)
 			{
-				return FReal(1);
+				return T(1);
 			}
 
 			// Closest Point is within segment
 			return Dot1 / Dot2;
 		}
 
+		/**
+		 * @brief The distance from Start to the sphere along the vector Dir. Returns numeric max if no intersection.
+		*/
+		template<typename T>
+		inline T RaySphereIntersectionDistance(const TVec3<T>& RayStart, const TVec3<T>& RayDir, const TVec3<T>& SpherePos, const T SphereRadius)
+		{
+			const TVec3<T> EO = SpherePos - RayStart;
+			const T V = TVec3<T>::DotProduct(RayDir, EO);
+			const T Disc = SphereRadius * SphereRadius - (TVec3<T>::DotProduct(EO, EO) - V * V);
+			if (Disc >= 0)
+			{
+				return (V - FMath::Sqrt(Disc));
+			}
+			else
+			{
+				return TNumericLimits<T>::Max();
+			}
+		}
+
 
 	} // namespace Utilities
 } // namespace Chaos
\ No newline at end of file
