Split Gather into parallelizable and non paralleizalbe functions.

diff --git a/Engine/Source/Runtime/Experimental/Chaos/Private/Chaos/Collision/PBDCollisionConstraintHandle.cpp b/Engine/Source/Runtime/Experimental/Chaos/Private/Chaos/Collision/PBDCollisionConstraintHandle.cpp
--- a/Engine/Source/Runtime/Experimental/Chaos/Private/Chaos/Collision/PBDCollisionConstraintHandle.cpp
+++ b/Engine/Source/Runtime/Experimental/Chaos/Private/Chaos/Collision/PBDCollisionConstraintHandle.cpp
@@ -57,8 +57,17 @@ namespace Chaos
 		return { GetContact().GetParticle0(), GetContact().GetParticle1() };
 	}
 
-	void FPBDCollisionConstraintHandle::GatherInput(FReal Dt, const int32 Particle0Level, const int32 Particle1Level, FPBDIslandSolverData& SolverData)
+	void FPBDCollisionConstraintHandle::PreGatherInput(FPBDIslandSolverData& SolverData)
 	{
+		ConcreteContainer()->PreGatherInput(GetContact(), SolverData);
+	}
+
+	void FPBDCollisionConstraintHandle::GatherInput(FReal Dt, const int32 Particle0Level, const int32 Particle1Level, FPBDIslandSolverData& SolverData, const bool bNeedsPreGather)
+	{
+		if (bNeedsPreGather)
+		{
+			PreGatherInput(SolverData);
+		}
 		ConcreteContainer()->GatherInput(Dt, GetContact(), Particle0Level, Particle1Level, SolverData);
 	}
 
diff --git a/Engine/Source/Runtime/Experimental/Chaos/Private/Chaos/Collision/SolverCollisionContainer.cpp b/Engine/Source/Runtime/Experimental/Chaos/Private/Chaos/Collision/SolverCollisionContainer.cpp
--- a/Engine/Source/Runtime/Experimental/Chaos/Private/Chaos/Collision/SolverCollisionContainer.cpp
+++ b/Engine/Source/Runtime/Experimental/Chaos/Private/Chaos/Collision/SolverCollisionContainer.cpp
@@ -55,6 +55,15 @@ namespace Chaos
 
 		bool IsIncrementalManifold() const { return bIsIncremental; }
 
+		void PreGatherInput(
+			FPBDCollisionConstraint& InConstraint,
+			FSolverBodyContainer& SolverBodyContainer)
+		{
+			Constraint = &InConstraint;
+			SolverBodyContainer.FindOrAdd(Constraint->GetParticle0());
+			SolverBodyContainer.FindOrAdd(Constraint->GetParticle1());
+		}
+
 		/**
 		 * @brief Initialize the data required for the solver, and bind to the bodies
 		*/
@@ -65,7 +74,6 @@ namespace Chaos
 			const int32 Particle1Level,
 			FSolverBodyContainer& SolverBodyContainer)
 		{
-			Constraint = &InConstraint;
 			bIsIncremental = Constraint->GetUseIncrementalCollisionDetection();
 
 			// Find the solver bodies for the particles we constrain. This will add them to the container
@@ -110,8 +118,8 @@ namespace Chaos
 
 		void GatherManifoldPoints(
 			const FReal InDt,
-			FSolverBody* Body0,
-			FSolverBody* Body1)
+			const FSolverBody* Body0,
+			const FSolverBody* Body1)
 		{
 			FSolverReal Dt = FSolverReal(InDt);
 
@@ -199,8 +207,8 @@ namespace Chaos
 
 		void GatherManifoldPointsVectorRegister(
 			const FReal InDt,
-			FSolverBody* Body0,
-			FSolverBody* Body1)
+			const FSolverBody* Body0,
+			const FSolverBody* Body1)
 		{
 			// Attempt to make this LWC independent. Actually this isn't easy because we explicitly use MakeVectorRegisterFloatFromDouble
 
@@ -381,19 +389,33 @@ namespace Chaos
 		CollisionSolvers.SetNum(MaxCollisions, false);
 	}
 
-	void FPBDCollisionSolverContainer::AddConstraintSolver(FReal Dt, FPBDCollisionConstraint& Constraint, const int32 Particle0Level, const int32 Particle1Level, FSolverBodyContainer& SolverBodyContainer, int32& ConstraintIndex)
+	void FPBDCollisionSolverContainer::PreAddConstraintSolver(FPBDCollisionConstraint& Constraint, FSolverBodyContainer& SolverBodyContainer, int32& ConstraintIndex)
 	{
 		// This container is required to allocate pointers that are valid for the whole tick,
 		// so we cannot allow the container to resize during the tick. See Reset()
 		check(ConstraintIndex < CollisionSolvers.Num());
+		Constraint.SetSolverBodyContainerIndex(ConstraintIndex);
 
-		auto& CollisionSolver = CollisionSolvers[ConstraintIndex];
+		FPBDCollisionSolverAdapter& CollisionSolver = CollisionSolvers[ConstraintIndex];
 		CollisionSolver.GetSolver().Reset();
 
+		CollisionSolver.PreGatherInput(Constraint, SolverBodyContainer);
+
+		++ConstraintIndex;
+	}
+
+	void FPBDCollisionSolverContainer::AddConstraintSolver(FReal Dt, FPBDCollisionConstraint& Constraint, const int32 Particle0Level, const int32 Particle1Level, FSolverBodyContainer& SolverBodyContainer)
+	{
+		// This container is required to allocate pointers that are valid for the whole tick,
+		// so we cannot allow the container to resize during the tick. See Reset()
+		int32 ConstraintIndex = Constraint.GetSolverBodyContainerIndex();
+		check(ConstraintIndex < CollisionSolvers.Num());
+
+		FPBDCollisionSolverAdapter& CollisionSolver = CollisionSolvers[ConstraintIndex];
+
 		CollisionSolver.GatherInput(Dt, Constraint, Particle0Level, Particle1Level, SolverBodyContainer);
 
 		bRequiresIncrementalCollisionDetection |= CollisionSolver.IsIncrementalManifold();
-		++ConstraintIndex;
 	}
 
 	void FPBDCollisionSolverContainer::UpdatePositionShockPropagation(const FReal Dt, const int32 It, const int32 NumIts, const int32 BeginIndex, const int32 EndIndex)
diff --git a/Engine/Source/Runtime/Experimental/Chaos/Private/Chaos/Framework/Parallel.cpp b/Engine/Source/Runtime/Experimental/Chaos/Private/Chaos/Framework/Parallel.cpp
--- a/Engine/Source/Runtime/Experimental/Chaos/Private/Chaos/Framework/Parallel.cpp
+++ b/Engine/Source/Runtime/Experimental/Chaos/Private/Chaos/Framework/Parallel.cpp
@@ -14,13 +14,22 @@ namespace Chaos
 	CHAOS_API bool bDisablePhysicsParallelFor = false;
 	CHAOS_API bool bDisableParticleParallelFor = false;
 	CHAOS_API bool bDisableCollisionParallelFor = false;
+	CHAOS_API int32 InnerParallelForBatchSize = 4;
 
 	FAutoConsoleVariableRef CVarDisablePhysicsParallelFor(TEXT("p.Chaos.DisablePhysicsParallelFor"), bDisablePhysicsParallelFor, TEXT("Disable parallel execution in Chaos Evolution"));
 	FAutoConsoleVariableRef CVarDisableParticleParallelFor(TEXT("p.Chaos.DisableParticleParallelFor"), bDisableParticleParallelFor, TEXT("Disable parallel execution for Chaos Particles (Collisions, "));
 	FAutoConsoleVariableRef CVarDisableCollisionParallelFor(TEXT("p.Chaos.DisableCollisionParallelFor"), bDisableCollisionParallelFor, TEXT("Disable parallel execution for Chaos Collisions (also disabled by DisableParticleParallelFor)"));
+	FAutoConsoleVariableRef CVarInnerPhysicsBatchSize(TEXT("p.Chaos.InnerParallelForBatchSize"), InnerParallelForBatchSize, TEXT("Set the batch size threshold for inner parallel fors"));
 #endif
 }
 
+void Chaos::InnerPhysicsParallelFor(int32 Num, TFunctionRef<void(int32)> InCallable, bool bForceSingleThreaded)
+{
+	int32 NumWorkers = int32(LowLevelTasks::FScheduler::Get().GetNumWorkers());
+	int32 BatchSize = Num / NumWorkers;
+	PhysicsParallelFor(Num, InCallable, (BatchSize > InnerParallelForBatchSize) ? bForceSingleThreaded : true);
+}
+
 void Chaos::PhysicsParallelFor(int32 InNum, TFunctionRef<void(int32)> InCallable, bool bForceSingleThreaded)
 {
 	TRACE_CPUPROFILER_EVENT_SCOPE(Chaos_PhysicsParallelFor);
diff --git a/Engine/Source/Runtime/Experimental/Chaos/Private/Chaos/PBDCollisionConstraints.cpp b/Engine/Source/Runtime/Experimental/Chaos/Private/Chaos/PBDCollisionConstraints.cpp
--- a/Engine/Source/Runtime/Experimental/Chaos/Private/Chaos/PBDCollisionConstraints.cpp
+++ b/Engine/Source/Runtime/Experimental/Chaos/Private/Chaos/PBDCollisionConstraints.cpp
@@ -376,6 +376,15 @@ namespace Chaos
 		return SolverData.GetConstraintContainer<FPBDCollisionSolverContainer>(ContainerId);
 	}
 
+	void FPBDCollisionConstraints::PreGatherInput(FPBDCollisionConstraint& Constraint, FPBDIslandSolverData& SolverData)
+	{
+		if (SolverType == EConstraintSolverType::QuasiPbd)
+		{
+			FPBDCollisionSolverContainer& SolverContainer = GetConstraintSolverContainer(SolverData);
+			SolverContainer.PreAddConstraintSolver(Constraint, SolverData.GetBodyContainer(), SolverData.GetConstraintIndex(ContainerId));
+		}
+	}
+
 	void FPBDCollisionConstraints::GatherInput(const FReal Dt, FPBDCollisionConstraint& Constraint, const int32 Particle0Level, const int32 Particle1Level, FPBDIslandSolverData& SolverData)
 	{
 		if (SolverType == EConstraintSolverType::QuasiPbd)
@@ -384,7 +393,7 @@ namespace Chaos
 			check(Constraint.IsEnabled());
 
 			FPBDCollisionSolverContainer& SolverContainer = GetConstraintSolverContainer(SolverData);
-			SolverContainer.AddConstraintSolver(Dt, Constraint, Particle0Level, Particle1Level, SolverData.GetBodyContainer(), SolverData.GetConstraintIndex(ContainerId));
+			SolverContainer.AddConstraintSolver(Dt, Constraint, Particle0Level, Particle1Level, SolverData.GetBodyContainer());
 		}
 		else
 		{
diff --git a/Engine/Source/Runtime/Experimental/Chaos/Private/Chaos/PBDConstraintRule.cpp b/Engine/Source/Runtime/Experimental/Chaos/Private/Chaos/PBDConstraintRule.cpp
--- a/Engine/Source/Runtime/Experimental/Chaos/Private/Chaos/PBDConstraintRule.cpp
+++ b/Engine/Source/Runtime/Experimental/Chaos/Private/Chaos/PBDConstraintRule.cpp
@@ -27,7 +27,9 @@ namespace Chaos
 	/** console variable to disable the colors computation on each island */
 	int32 ChaosDisableIslandColors = true;
 	FAutoConsoleVariableRef CVarChaosDisableIslandColors(TEXT("p.Chaos.Islands.DisableColors"), ChaosDisableIslandColors, TEXT(""));
-
+	
+	DECLARE_CYCLE_STAT(TEXT("FPBDConstraintRule::Gather"), STAT_Constraints_Gather, STATGROUP_Chaos);
+	DECLARE_CYCLE_STAT(TEXT("FPBDConstraintRule::Color_Gather"), STAT_Constraints_Color_Gather, STATGROUP_Chaos);
 
 	template<class ConstraintType>
 	TSimpleConstraintRule<ConstraintType>::TSimpleConstraintRule(int32 InPriority, FConstraints& InConstraints)
@@ -149,6 +151,7 @@ namespace Chaos
 	template<class ConstraintType>
 	void TPBDConstraintIslandRule<ConstraintType>::GatherSolverInput(const FReal Dt, int32 GroupIndex)
 	{
+		SCOPE_CYCLE_COUNTER(STAT_Constraints_Gather);
 		if(FPBDIslandGroup* IslandGroup = ConstraintGraph->GetIslandGroup(GroupIndex))
 		{
 			// This will reset the number of constraints inside the solver datas. For now we keep this function since according
@@ -158,7 +161,7 @@ namespace Chaos
 			Constraints.SetNumIslandConstraints(IslandGroup->ConstraintCount(Constraints.GetContainerId()), *IslandGroup);
 			IslandGroup->InitConstraintIndex(Constraints.GetContainerId());
 			
-			for (auto& IslandSolver : IslandGroup->GetIslands() )
+			for (FPBDIslandSolver* IslandSolver : IslandGroup->GetIslands())
 			{
 				if(!IslandSolver->IsSleeping() && !IslandSolver->IsUsingCache())
 				{
@@ -239,6 +242,7 @@ namespace Chaos
 	template<class ConstraintType>
 	void TPBDConstraintColorRule<ConstraintType>::GatherSolverInput(const FReal Dt, int32 GroupIndex)
 	{
+		SCOPE_CYCLE_COUNTER(STAT_Constraints_Color_Gather);
 		FPBDConstraintGraph* LocalGraph = ConstraintGraph;
 		auto GetParticleLevel = [LocalGraph](TGeometryParticleHandle<FReal, 3>* ConstrainedParticle) -> int32
 		{
@@ -261,7 +265,7 @@ namespace Chaos
 				IslandConstraintSets.Reset();
 				int32 ConstraintSetBegin = 0, ConstraintSetEnd = 0;
 				
-				for (auto& IslandSolver : IslandGroup->GetIslands() )
+				for (FPBDIslandSolver* IslandSolver : IslandGroup->GetIslands())
 				{
 					ConstraintSetBegin = ConstraintSetEnd;
 
@@ -296,10 +300,19 @@ namespace Chaos
 									
 									if (OffsetEnd != OffsetBegin)
 									{
+										for (int32 ConstraintIndex = OffsetBegin; ConstraintIndex < OffsetEnd; ++ConstraintIndex)
+										{
+											FConstraintContainerHandle* Constraint = SortedConstraints[ConstraintIndex];
+											if (Constraint->IsEnabled())
+											{
+												Constraint->PreGatherInput(*IslandGroup);
+											}
+										}
 										// Calculate the range of indices for this color as a set of independent contacts
 										TPair<int32, int32> ColorConstrainSet(ConstraintSetEnd, ConstraintSetEnd);
-										for(int32 ConstraintIndex = OffsetBegin; ConstraintIndex < OffsetEnd; ++ConstraintIndex)
+										InnerPhysicsParallelFor(OffsetEnd - OffsetBegin, [&](int32 Index)
 										{
+											int32 ConstraintIndex = Index + OffsetBegin;
 											FConstraintContainerHandle* Constraint = SortedConstraints[ConstraintIndex];
 											if (Constraint->IsEnabled())
 											{
@@ -308,11 +321,17 @@ namespace Chaos
 												const TVector<TGeometryParticleHandle<FReal, 3>*, 2> ConstrainedParticles = Constraint->GetConstrainedParticles();
 												const int32 Particle0Level = GetParticleLevel(ConstrainedParticles[0]);
 												const int32 Particle1Level = GetParticleLevel(ConstrainedParticles[1]);
-									
+
 												// Note we are building the SolverBodies as we go, in the order that we visit them. Each constraint
 												// references two bodies, so we won't strictly be accessing only in cache order, but it's about as good as it can be.
-												Constraint->GatherInput(Dt, Particle0Level, Particle1Level, *IslandGroup);
-									
+												Constraint->GatherInput(Dt, Particle0Level, Particle1Level, *IslandGroup, false);
+											}
+										});
+										for (int32 ConstraintIndex = OffsetBegin; ConstraintIndex < OffsetEnd; ++ConstraintIndex)
+										{
+											FConstraintContainerHandle* Constraint = SortedConstraints[ConstraintIndex];
+											if (Constraint->IsEnabled())
+											{
 												// Update the current constraint set of this color
 												ColorConstrainSet.Value = ++ConstraintSetEnd;
 											}
@@ -594,4 +613,4 @@ namespace Chaos
 	template class TPBDConstraintIslandRule<FPBDSuspensionConstraints>;
 	template class TPBDConstraintIslandRule<FPBDRigidDynamicSpringConstraints>;
 	template class TPBDConstraintIslandRule<FPBDRigidSpringConstraints>;
-}
\ No newline at end of file
+}
diff --git a/Engine/Source/Runtime/Experimental/Chaos/Public/Chaos/Collision/PBDCollisionConstraint.h b/Engine/Source/Runtime/Experimental/Chaos/Public/Chaos/Collision/PBDCollisionConstraint.h
--- a/Engine/Source/Runtime/Experimental/Chaos/Public/Chaos/Collision/PBDCollisionConstraint.h
+++ b/Engine/Source/Runtime/Experimental/Chaos/Public/Chaos/Collision/PBDCollisionConstraint.h
@@ -457,6 +457,12 @@ namespace Chaos
 			SolverBodies[1] = InSolverBody1;
 		}
 
+		int32 GetSolverBodyContainerIndex() const { return ConstraintIndex; }
+		void SetSolverBodyContainerIndex(int32 InConstraintIndex)
+		{
+			ConstraintIndex = InConstraintIndex;
+		}
+
 		/**
 		 * @brief Whether this constraint was fully restored from a previous tick, and the manifold should be reused as-is
 		*/
@@ -648,6 +654,9 @@ namespace Chaos
 		// These are only needed here while we still have the legacy solvers (not QuasiPBD)
 		FSolverBody* SolverBodies[2];
 
+		// Stores the index into the solver container for this constraint
+		int32 ConstraintIndex;
+
 		// Simplex data from the last call to GJK, used to warm-start GJK
 		FGJKSimplexData GJKWarmStartData;
 
diff --git a/Engine/Source/Runtime/Experimental/Chaos/Public/Chaos/Collision/PBDCollisionConstraintHandle.h b/Engine/Source/Runtime/Experimental/Chaos/Public/Chaos/Collision/PBDCollisionConstraintHandle.h
--- a/Engine/Source/Runtime/Experimental/Chaos/Public/Chaos/Collision/PBDCollisionConstraintHandle.h
+++ b/Engine/Source/Runtime/Experimental/Chaos/Public/Chaos/Collision/PBDCollisionConstraintHandle.h
@@ -87,7 +87,8 @@ namespace Chaos
 
 		TVector<TGeometryParticleHandle<FReal, 3>*, 2> GetConstrainedParticles();
 
-		void GatherInput(FReal Dt, const int32 Particle0Level, const int32 Particle1Level, FPBDIslandSolverData& SolverData);
+		void PreGatherInput(FPBDIslandSolverData& SolverData);
+		void GatherInput(FReal Dt, const int32 Particle0Level, const int32 Particle1Level, FPBDIslandSolverData& SolverData, const bool bNeedsPreGather = true);
 
 		FSolverBody* GetSolverBody0();
 		FSolverBody* GetSolverBody1();
diff --git a/Engine/Source/Runtime/Experimental/Chaos/Public/Chaos/Collision/SolverCollisionContainer.h b/Engine/Source/Runtime/Experimental/Chaos/Public/Chaos/Collision/SolverCollisionContainer.h
--- a/Engine/Source/Runtime/Experimental/Chaos/Public/Chaos/Collision/SolverCollisionContainer.h
+++ b/Engine/Source/Runtime/Experimental/Chaos/Public/Chaos/Collision/SolverCollisionContainer.h
@@ -41,7 +41,8 @@ namespace Chaos
 		void SetNum(const int32 MaxCollisions);
 
 		// Add a solver constraint for the specified constraint and gather the required solver data
-		void AddConstraintSolver(FReal Dt, FPBDCollisionConstraint& Constraint, const int32 Particle0Level, const int32 Particle1Level, FSolverBodyContainer& SolverBodyContainer, int32& ConstraintIndex);
+		void PreAddConstraintSolver(FPBDCollisionConstraint& Constraint, FSolverBodyContainer& SolverBodyContainer, int32& ConstraintIndex);
+		void AddConstraintSolver(FReal Dt, FPBDCollisionConstraint& Constraint, const int32 Particle0Level, const int32 Particle1Level, FSolverBodyContainer& SolverBodyContainer);
 
 		bool SolvePositionSerial(const FReal Dt, const int32 It, const int32 NumIts, const int32 BeginIndex, const int32 EndIndex);
 		bool SolvePositionParallel(const FReal Dt, const int32 It, const int32 NumIts, const int32 BeginIndex, const int32 EndIndex);
diff --git a/Engine/Source/Runtime/Experimental/Chaos/Public/Chaos/Framework/Parallel.h b/Engine/Source/Runtime/Experimental/Chaos/Public/Chaos/Framework/Parallel.h
--- a/Engine/Source/Runtime/Experimental/Chaos/Public/Chaos/Framework/Parallel.h
+++ b/Engine/Source/Runtime/Experimental/Chaos/Public/Chaos/Framework/Parallel.h
@@ -7,6 +7,7 @@
 namespace Chaos
 {
 	void CHAOS_API PhysicsParallelFor(int32 InNum, TFunctionRef<void(int32)> InCallable, bool bForceSingleThreaded = false);
+	void CHAOS_API InnerPhysicsParallelFor(int32 InNum, TFunctionRef<void(int32)> InCallable, bool bForceSingleThreaded = false);
 	//void CHAOS_API PhysicsParallelFor_RecursiveDivide(int32 InNum, TFunctionRef<void(int32)> InCallable, bool bForceSingleThreaded = false);
 
 
diff --git a/Engine/Source/Runtime/Experimental/Chaos/Public/Chaos/PBDCollisionConstraints.h b/Engine/Source/Runtime/Experimental/Chaos/Public/Chaos/PBDCollisionConstraints.h
--- a/Engine/Source/Runtime/Experimental/Chaos/Public/Chaos/PBDCollisionConstraints.h
+++ b/Engine/Source/Runtime/Experimental/Chaos/Public/Chaos/PBDCollisionConstraints.h
@@ -165,6 +165,7 @@ public:
 	//
 
 	void SetNumIslandConstraints(const int32 NumIslandConstraints, FPBDIslandSolverData& SolverData);
+	void PreGatherInput(FPBDCollisionConstraint& Constraint, FPBDIslandSolverData& SolverData);
 	void GatherInput(const FReal Dt, FPBDCollisionConstraint& Constraint, const int32 Particle0Level, const int32 Particle1Level, FPBDIslandSolverData& SolverData);
 
 	bool ApplyPhase1Serial(const FReal Dt, const int32 It, const int32 NumIts, FPBDIslandSolverData& SolverData);
