Missing non-const API to get at shared geometry ptr.

diff --git a/Engine/Source/Runtime/Experimental/Chaos/Public/PhysicsProxy/SingleParticlePhysicsProxy.h b/Engine/Source/Runtime/Experimental/Chaos/Public/PhysicsProxy/SingleParticlePhysicsProxy.h
--- a/Engine/Source/Runtime/Experimental/Chaos/Public/PhysicsProxy/SingleParticlePhysicsProxy.h
+++ b/Engine/Source/Runtime/Experimental/Chaos/Public/PhysicsProxy/SingleParticlePhysicsProxy.h
@@ -1123,10 +1123,10 @@ public:
 		SetGeometry(TSharedPtr<FImplicitObject, ESPMode::ThreadSafe>(RawGeometry));
 	}
 
-	void SetGeometry(TSharedPtr<FImplicitObject, ESPMode::ThreadSafe> SharedGeometry)
+	void SetGeometry(TSharedPtr<const FImplicitObject, ESPMode::ThreadSafe> SharedGeometry)
 	{
 		VerifyContext();
-		GetParticle_LowLevel()->SetGeometry(SharedGeometry);
+		GetParticle_LowLevel()->SetGeometry(ConstCastSharedPtr<FImplicitObject, const FImplicitObject, ESPMode::ThreadSafe>(SharedGeometry));
 	}
 
 	//Note: this must be called after setting geometry. This API seems bad. Should probably be part of setting geometry
