Fix updatesyncstate crash in the island manager

diff --git a/Engine/Source/Runtime/Experimental/Chaos/Private/Chaos/Island/IslandManager.cpp b/Engine/Source/Runtime/Experimental/Chaos/Private/Chaos/Island/IslandManager.cpp
--- a/Engine/Source/Runtime/Experimental/Chaos/Private/Chaos/Island/IslandManager.cpp
+++ b/Engine/Source/Runtime/Experimental/Chaos/Private/Chaos/Island/IslandManager.cpp
@@ -173,8 +173,7 @@ inline void UpdateSyncState(FPBDIslandSolver* IslandSolver, FPBDRigidsSOAs& Part
 		IslandSolver->SetNeedsResim(bNeedsResim);
 	}
 }
-
-/** Add all the graph particle and constraints to the solver islands*/
+	/** Add all the graph particle and constraints to the solver islands*/
 inline void PopulateIslands(FPBDIslandManager::GraphType* IslandGraph)
 {
 	auto AddNodeToIsland = [IslandGraph](const int32 IslandIndex, FPBDIslandManager::GraphType::FGraphNode& GraphNode)
@@ -191,22 +190,36 @@ inline void PopulateIslands(FPBDIslandManager::GraphType* IslandGraph)
 	TSet<int32> NodeIslands;
 	for(auto& GraphNode : IslandGraph->GraphNodes)
 	{
+		NodeIslands = GraphNode.NodeIslands;
+		GraphNode.NodeIslands.Reset();
+
+		// First we re-add the island indices coming from a sleeping islands since we are not adding them here.
+		for (auto& IslandIndex : NodeIslands)
+		{
+			if (IslandGraph->GraphIslands.IsValidIndex(IslandIndex))
+			{
+				if (IslandGraph->GraphIslands[IslandIndex].IslandItem->IsSleeping())
+				{
+					GraphNode.NodeIslands.Add(IslandIndex);
+				}
+			}
+		}
 		// If the node is valid : only one island
 		if( GraphNode.bValidNode)
 		{
 			AddNodeToIsland(GraphNode.IslandIndex, GraphNode);
+			GraphNode.NodeIslands.Add(GraphNode.IslandIndex);
 		}
 		else
 		{
 			// A particle could belong to several islands when static/kinematic (not valid)
 			// First we compute the unique particle set of islands
-			NodeIslands.Reset();
 			for( auto& NodeEdge : GraphNode.NodeEdges)
 			{
-				NodeIslands.Add(IslandGraph->GraphEdges[NodeEdge].IslandIndex);
+				GraphNode.NodeIslands.Add(IslandGraph->GraphEdges[NodeEdge].IslandIndex);
 			}
 			// Loop over the set of islands and add the particle to the solver island
-			for(auto& NodeIsland : NodeIslands)
+			for(auto& NodeIsland : GraphNode.NodeIslands)
 			{
 				AddNodeToIsland(NodeIsland, GraphNode);
 			}
@@ -499,27 +512,28 @@ void FPBDIslandManager::RemoveParticle(FGeometryParticleHandle* ParticleHandle)
 			if (IslandGraph->GraphNodes.IsValidIndex(*NodeIndex))
 			{
 				const FGraphNode& GraphNode = IslandGraph->GraphNodes[*NodeIndex];
-				// If the node is valid : we need to also remove it from its own island 
-				if(GraphNode.NodeEdges.Num() == 0 && IslandSolvers.IsValidIndex(GraphNode.IslandIndex))
+
+				// We loop over all the connected edges to find all the islands in which the particle
+				// is (static/kinematic particles could belong to several islands)
+				// And remove the particle from the solver island. It will allow the solver
+				// islands to be updated directly and not at the next sync.
+				for (const int32& EdgeIndex : GraphNode.NodeEdges)
 				{
-					IslandSolvers[GraphNode.IslandIndex]->RemoveParticle(ParticleHandle);
+					const int32 IslandIndex = IslandGraph->GraphEdges[EdgeIndex].IslandIndex;
+					if (IslandSolvers.IsValidIndex(IslandIndex))
+					{
+						//We need to remove as well the constraint from the island solver since the edges are removed from the graph
+						IslandSolvers[IslandIndex]->RemoveConstraint(IslandGraph->GraphEdges[EdgeIndex].EdgeItem);
+						IslandGraph->GraphEdges[EdgeIndex].EdgeItem->SetConstraintGraphIndex(INDEX_NONE);
+					}
 				}
-				else
+
+				// If the node is valid : we need to also remove it from its own island 
+				for (const int32& IslandIndex : GraphNode.NodeIslands)
 				{
-					// We loop over all the connected edges to find all the islands in which the particle
-					// is (static/kinematic particles could belong to several islands)
-					// And remove the particle from the solver island. It will allow the solver
-					// islands to be updated directly and not at the next sync.
-					for (const int32& EdgeIndex : GraphNode.NodeEdges)
+					if (IslandSolvers.IsValidIndex(IslandIndex))
 					{
-						const int32 IslandIndex = IslandGraph->GraphEdges[EdgeIndex].IslandIndex;
-						if (IslandSolvers.IsValidIndex(IslandIndex))
-						{
-							IslandSolvers[IslandIndex]->RemoveParticle(ParticleHandle);
-
-							//We need to remove as well the constraint from the island solver since the edges are removed from the graph
-							IslandSolvers[IslandIndex]->RemoveConstraint(IslandGraph->GraphEdges[EdgeIndex].EdgeItem);
-						}
+						IslandSolvers[IslandIndex]->RemoveParticle(ParticleHandle);
 					}
 				}
 			}
diff --git a/Engine/Source/Runtime/Experimental/Chaos/Public/Chaos/Island/IslandGraph.h b/Engine/Source/Runtime/Experimental/Chaos/Public/Chaos/Island/IslandGraph.h
--- a/Engine/Source/Runtime/Experimental/Chaos/Public/Chaos/Island/IslandGraph.h
+++ b/Engine/Source/Runtime/Experimental/Chaos/Public/Chaos/Island/IslandGraph.h
@@ -130,6 +130,9 @@ public:
 
 		/** Node Island Index (for static/kinematic particles could belong to several islands) */
 		int32 IslandIndex = INDEX_NONE;
+		
+		/** List of islands in which the node is referenced */
+		TSet<int32> NodeIslands;
 
 		/** Check if a node is valid (checked for graph partitioning) */
 		bool bValidNode = true;
