Chaos Cloth - Add global point damping to offer a faster alternative to the current local damping formulation. Saves 40us on the pre-iteration step of the Owen asset in Test builds.
- Rename the existing Damping parameter to LocalDamping, set it disabled by default, and add a new more accurate damping config migration path for the NvCloth assets.
- Tidy up a couple of mismatched LWC conversions.
- Fix the cloth position initializations by adding a reset as it was impeding the damping tests.

diff --git a/Engine/Plugins/Experimental/ChaosCloth/Source/Chaos/Private/ChaosCloth/ChaosClothConfig.cpp b/Engine/Plugins/Experimental/ChaosCloth/Source/Chaos/Private/ChaosCloth/ChaosClothConfig.cpp
--- a/Engine/Plugins/Experimental/ChaosCloth/Source/Chaos/Private/ChaosCloth/ChaosClothConfig.cpp
+++ b/Engine/Plugins/Experimental/ChaosCloth/Source/Chaos/Private/ChaosCloth/ChaosClothConfig.cpp
@@ -7,6 +7,7 @@
 #include "ClothingSimulationInteractor.h"
 #include "UObject/PhysicsObjectVersion.h"
 #include "UObject/FortniteMainBranchObjectVersion.h"
+#include "UObject/UE5ReleaseStreamObjectVersion.h"
 
 // Legacy parameters not yet migrated to Chaos parameters:
 //  VerticalConstraintConfig.CompressionLimit
@@ -64,8 +65,14 @@ void UChaosClothConfig::MigrateFrom(const FClothConfig_Legacy& ClothConfig)
 	Drag.Low = Drag.High = bUsePointBasedWindModel ? 0.07f  : ClothConfig.WindDragCoefficient;  // Only Accurate wind uses the WindDragCoefficient
 	Lift.Low = Lift.High = bUsePointBasedWindModel ? 0.035f : ClothConfig.WindLiftCoefficient;  // Only Accurate wind uses the WindLiftCoefficient
 
+	// Apply legacy damping calculations, see FClothingSimulationNv::ApplyClothConfig()
+	constexpr float DampStiffnesssFreq = 10.0f;
+	constexpr float PrecalcLog2 = 0.69314718f;
 	const float Damping = (ClothConfig.Damping.X + ClothConfig.Damping.Y + ClothConfig.Damping.Z) / 3.f;
-	DampingCoefficient = FMath::Clamp(Damping * Damping * 0.7f, 0.f, 1.f);  // Nv Cloth seems to have a different damping formulation.
+	const float DampStiffFreqRatio = DampStiffnesssFreq / ClothConfig.StiffnessFrequency;
+	const float ExpDamp = DampStiffFreqRatio * FMath::Log2(1.0f - Damping);
+	const float AdjustedDamping = 1.0f - FMath::Exp(ExpDamp * PrecalcLog2);
+	DampingCoefficient = FMath::Clamp(AdjustedDamping, 0.f, 1.f);
 
 	CollisionThickness = FMath::Clamp(ClothConfig.CollisionThickness, 0.f, 1000.f);
 	SelfCollisionThickness = FMath::Clamp(ClothConfig.SelfCollisionRadius, 0.f, 1000.f);
@@ -114,6 +121,7 @@ void UChaosClothConfig::Serialize(FArchive& Ar)
 	Ar.UsingCustomVersion(FChaosClothConfigCustomVersion::GUID);
 	Ar.UsingCustomVersion(FPhysicsObjectVersion::GUID);
 	Ar.UsingCustomVersion(FFortniteMainBranchObjectVersion::GUID);
+	Ar.UsingCustomVersion(FUE5ReleaseStreamObjectVersion::GUID);
 }
 
 void UChaosClothConfig::PostLoad()
@@ -124,6 +132,7 @@ void UChaosClothConfig::PostLoad()
 	const int32 ChaosClothConfigCustomVersion = GetLinkerCustomVersion(FChaosClothConfigCustomVersion::GUID);
 	const int32 PhysicsObjectVersion = GetLinkerCustomVersion(FPhysicsObjectVersion::GUID);
 	const int32 FortniteMainBranchObjectVersion = GetLinkerCustomVersion(FFortniteMainBranchObjectVersion::GUID);
+	const int32 UE5ReleaseStreamObjectVersion = GetLinkerCustomVersion(FUE5ReleaseStreamObjectVersion::GUID);
 
 	if (ChaosClothConfigCustomVersion < FChaosClothConfigCustomVersion::UpdateDragDefault)
 	{
@@ -180,6 +189,13 @@ void UChaosClothConfig::PostLoad()
 		BendingStiffnessWeighted.Low = BendingStiffnessWeighted.High = BendingStiffness_DEPRECATED;
 		AreaStiffnessWeighted.Low = AreaStiffnessWeighted.High = AreaStiffness_DEPRECATED;
 	}
+
+	if (UE5ReleaseStreamObjectVersion < FUE5ReleaseStreamObjectVersion::ChaosClothFasterDamping)
+	{
+		// Note: The previous damping has been renamed LocalDamping to make space for a faster but more primitive global point damping.
+		LocalDampingCoefficient = DampingCoefficient;
+		DampingCoefficient = 0.f;
+	}
 #endif  // #if WITH_EDITORONLY_DATA
 }
 
diff --git a/Engine/Plugins/Experimental/ChaosCloth/Source/Chaos/Private/ChaosCloth/ChaosClothingSimulation.cpp b/Engine/Plugins/Experimental/ChaosCloth/Source/Chaos/Private/ChaosCloth/ChaosClothingSimulation.cpp
--- a/Engine/Plugins/Experimental/ChaosCloth/Source/Chaos/Private/ChaosCloth/ChaosClothingSimulation.cpp
+++ b/Engine/Plugins/Experimental/ChaosCloth/Source/Chaos/Private/ChaosCloth/ChaosClothingSimulation.cpp
@@ -415,6 +415,7 @@ void FClothingSimulation::CreateActor(USkeletalMeshComponent* InOwnerComponent,
 		TVec2<FRealSingle>(ClothConfig->Lift.Low, ClothConfig->Lift.High),  // Animatable
 		ClothConfig->bUsePointBasedWindModel,
 		ClothConfig->DampingCoefficient,
+		ClothConfig->LocalDampingCoefficient,
 		ClothConfig->CollisionThickness,
 		ClothConfig->FrictionCoefficient,
 		ClothConfig->bUseCCD,
@@ -522,11 +523,11 @@ void FClothingSimulation::Simulate(IClothingSimulationContext* InContext)
 
 #if !(UE_BUILD_SHIPPING || UE_BUILD_TEST)
 		const bool bNeedsReset = (ClothingSimulationConsole::Command && ClothingSimulationConsole::Command->MustReset(ResetCount)) ||
-			(Context->TeleportMode == EClothingTeleportMode::TeleportAndReset);
+			Context->TeleportMode == EClothingTeleportMode::TeleportAndReset || PrevSimulationTime == 0.f;
 #else
-		const bool bNeedsReset = (Context->TeleportMode == EClothingTeleportMode::TeleportAndReset);
+		const bool bNeedsReset = Context->TeleportMode == EClothingTeleportMode::TeleportAndReset || PrevSimulationTime == 0.f;
 #endif
-		const bool bNeedsTeleport = (Context->TeleportMode > EClothingTeleportMode::None) || !PrevSimulationTime;
+		const bool bNeedsTeleport = (Context->TeleportMode > EClothingTeleportMode::None);
 		bIsTeleported = bNeedsTeleport;
 
 		// Update Solver animatable parameters
@@ -827,6 +828,7 @@ void FClothingSimulation::RefreshClothConfig(const IClothingSimulationContext* I
 			TVec2<FRealSingle>(ClothConfig->Lift.Low, ClothConfig->Lift.High),  // Animatable
 			ClothConfig->bUsePointBasedWindModel,
 			ClothConfig->DampingCoefficient,
+			ClothConfig->LocalDampingCoefficient,
 			ClothConfig->CollisionThickness,
 			ClothConfig->FrictionCoefficient,
 			ClothConfig->bUseCCD,
diff --git a/Engine/Plugins/Experimental/ChaosCloth/Source/Chaos/Private/ChaosCloth/ChaosClothingSimulationCloth.cpp b/Engine/Plugins/Experimental/ChaosCloth/Source/Chaos/Private/ChaosCloth/ChaosClothingSimulationCloth.cpp
--- a/Engine/Plugins/Experimental/ChaosCloth/Source/Chaos/Private/ChaosCloth/ChaosClothingSimulationCloth.cpp
+++ b/Engine/Plugins/Experimental/ChaosCloth/Source/Chaos/Private/ChaosCloth/ChaosClothingSimulationCloth.cpp
@@ -121,7 +121,7 @@ void FClothingSimulationCloth::FLODData::Add(FClothingSimulationSolver* Solver,
 				DisabledCollisionElements.Emplace(TVec2<int32>(Element, Index));
 			}
 		}
-		ClothConstraints.SetSelfCollisionConstraints(TriangleMesh, MoveTemp(DisabledCollisionElements), (FReal)Cloth->SelfCollisionThickness);
+		ClothConstraints.SetSelfCollisionConstraints(TriangleMesh, MoveTemp(DisabledCollisionElements), (Softs::FSolverReal)Cloth->SelfCollisionThickness);
 	}
 
 	// Edge constraints
@@ -138,7 +138,7 @@ void FClothingSimulationCloth::FLODData::Add(FClothingSimulationSolver* Solver,
 		if (Cloth->bUseBendingElements)
 		{
 			TArray<Chaos::TVec4<int32>> BendingElements = TriangleMesh.GetUniqueAdjacentElements();
-			ClothConstraints.SetBendingConstraints(MoveTemp(BendingElements), Cloth->BendingStiffness[0]);  // TODO: Add support for weight maps
+			ClothConstraints.SetBendingConstraints(MoveTemp(BendingElements), (Softs::FSolverReal)Cloth->BendingStiffness[0]);  // TODO: Add support for weight maps
 		}
 		else
 		{
@@ -189,12 +189,12 @@ void FClothingSimulationCloth::FLODData::Add(FClothingSimulationSolver* Solver,
 				}
 			}
 
-			ClothConstraints.SetVolumeConstraints(DoubleBendingConstraints, (FReal)Cloth->VolumeStiffness);
+			ClothConstraints.SetVolumeConstraints(DoubleBendingConstraints, (Softs::FSolverReal)Cloth->VolumeStiffness);
 		}
 		else
 		{
 			TArray<Chaos::TVec3<int32>> SurfaceConstraints = SurfaceElements;
-			ClothConstraints.SetVolumeConstraints(MoveTemp(SurfaceConstraints), (FReal)Cloth->VolumeStiffness);
+			ClothConstraints.SetVolumeConstraints(MoveTemp(SurfaceConstraints), (Softs::FSolverReal)Cloth->VolumeStiffness);
 		}
 	}
 
@@ -204,13 +204,13 @@ void FClothingSimulationCloth::FLODData::Add(FClothingSimulationSolver* Solver,
 	{
 		const TConstArrayView<FRealSingle>& TetherScaleMultipliers = WeightMaps[(int32)EChaosWeightMapTarget::TetherScale];
 
-		const FReal MeshScale = Cloth->Mesh->GetReferenceBoneTransform().GetScale3D().GetMax();
+		const Softs::FSolverReal MeshScale(Cloth->Mesh->GetReferenceBoneTransform().GetScale3D().GetMax());
 
 		ClothConstraints.SetLongRangeConstraints(
 			Tethers,
 			TetherStiffnessMultipliers,
 			TetherScaleMultipliers,
-			FVec2((FReal)Cloth->TetherScale[0], (FReal)Cloth->TetherScale[1]) * MeshScale,
+			Softs::FSolverVec2((Softs::FSolverReal)Cloth->TetherScale[0], (Softs::FSolverReal)Cloth->TetherScale[1]) * MeshScale,
 			Cloth->bUseXPBDConstraints);
 	}
 
@@ -239,7 +239,7 @@ void FClothingSimulationCloth::FLODData::Add(FClothingSimulationSolver* Solver,
 	// Shape target constraint
 	if (Cloth->ShapeTargetStiffness > (FRealSingle)0.)
 	{
-		ClothConstraints.SetShapeTargetConstraints((FReal)Cloth->ShapeTargetStiffness);
+		ClothConstraints.SetShapeTargetConstraints((Softs::FSolverReal)Cloth->ShapeTargetStiffness);
 	}
 
 	// Commit rules to solver
@@ -253,7 +253,7 @@ void FClothingSimulationCloth::FLODData::Add(FClothingSimulationSolver* Solver,
 	NumDynammicParticles = 0;
 	for (int32 Index = 0; Index < NumParticles; ++Index)
 	{
-		if (InvMasses[Index] == 0.f)
+		if (InvMasses[Index] == (Softs::FSolverReal)0.)
 		{
 			++NumKinenamicParticles;
 		}
@@ -363,6 +363,7 @@ FClothingSimulationCloth::FClothingSimulationCloth(
 	const TVec2<FRealSingle>& InLift,
 	bool bInUseLegacyWind,
 	FRealSingle InDampingCoefficient,
+	FRealSingle InLocalDampingCoefficient,
 	FRealSingle InCollisionThickness,
 	FRealSingle InFrictionCoefficient,
 	bool bInUseCCD,
@@ -403,6 +404,7 @@ FClothingSimulationCloth::FClothingSimulationCloth(
 	, AirDensity(1.225e-6f)  // Set by clothing interactor
 	, bUseLegacyWind(bInUseLegacyWind)
 	, DampingCoefficient(InDampingCoefficient)
+	, LocalDampingCoefficient(InLocalDampingCoefficient)
 	, CollisionThickness(InCollisionThickness)
 	, FrictionCoefficient(InFrictionCoefficient)
 	, bUseCCD(bInUseCCD)
@@ -808,7 +810,7 @@ void FClothingSimulationCloth::Update(FClothingSimulationSolver* Solver)
 		Solver->SetWindVelocity(GroupId, WindVelocity + Solver->GetWindVelocity());
 
 		// Update general solver properties
-		Solver->SetProperties(GroupId, DampingCoefficient, CollisionThickness, FrictionCoefficient);
+		Solver->SetProperties(GroupId, DampingCoefficient, LocalDampingCoefficient, CollisionThickness, FrictionCoefficient);
 
 		// Update use of continuous collision detection
 		Solver->SetUseCCD(GroupId, bUseCCD);
diff --git a/Engine/Plugins/Experimental/ChaosCloth/Source/Chaos/Private/ChaosCloth/ChaosClothingSimulationSolver.cpp b/Engine/Plugins/Experimental/ChaosCloth/Source/Chaos/Private/ChaosCloth/ChaosClothingSimulationSolver.cpp
--- a/Engine/Plugins/Experimental/ChaosCloth/Source/Chaos/Private/ChaosCloth/ChaosClothingSimulationSolver.cpp
+++ b/Engine/Plugins/Experimental/ChaosCloth/Source/Chaos/Private/ChaosCloth/ChaosClothingSimulationSolver.cpp
@@ -71,6 +71,7 @@ namespace ClothingSimulationSolverDefault
 	static const FRealSingle CollisionThickness = 1.2f;
 	static const FRealSingle FrictionCoefficient = 0.2f;
 	static const FRealSingle DampingCoefficient = 0.01f;
+	static const FRealSingle LocalDampingCoefficient = 0.f;
 }
 
 namespace ClothingSimulationSolverConstant
@@ -104,7 +105,8 @@ FClothingSimulationSolver::FClothingSimulationSolver()
 			(Softs::FSolverReal)ClothingSimulationSolverDefault::CollisionThickness,
 			(Softs::FSolverReal)ClothingSimulationSolverDefault::SelfCollisionThickness,
 			(Softs::FSolverReal)ClothingSimulationSolverDefault::FrictionCoefficient,
-			(Softs::FSolverReal)ClothingSimulationSolverDefault::DampingCoefficient));
+			(Softs::FSolverReal)ClothingSimulationSolverDefault::DampingCoefficient,
+			(Softs::FSolverReal)ClothingSimulationSolverDefault::LocalDampingCoefficient));
 
 	// Add simulation groups arrays
 	Evolution->AddArray(&PreSimulationTransforms);
@@ -653,9 +655,10 @@ void FClothingSimulationSolver::ParticleMassClampAndEnslave(int32 Offset, int32
 	}
 }
 
-void FClothingSimulationSolver::SetProperties(uint32 GroupId, FRealSingle DampingCoefficient, FRealSingle CollisionThickness, FRealSingle FrictionCoefficient)
+void FClothingSimulationSolver::SetProperties(uint32 GroupId, FRealSingle DampingCoefficient, FRealSingle LocalDampingCoefficient, FRealSingle CollisionThickness, FRealSingle FrictionCoefficient)
 {
 	Evolution->SetDamping(DampingCoefficient, GroupId);
+	Evolution->SetLocalDamping(LocalDampingCoefficient, GroupId);
 	Evolution->SetCollisionThickness(CollisionThickness, GroupId);
 	Evolution->SetCoefficientOfFriction(FrictionCoefficient, GroupId);
 }
diff --git a/Engine/Plugins/Experimental/ChaosCloth/Source/Chaos/Public/ChaosCloth/ChaosClothConfig.h b/Engine/Plugins/Experimental/ChaosCloth/Source/Chaos/Public/ChaosCloth/ChaosClothConfig.h
--- a/Engine/Plugins/Experimental/ChaosCloth/Source/Chaos/Public/ChaosCloth/ChaosClothConfig.h
+++ b/Engine/Plugins/Experimental/ChaosCloth/Source/Chaos/Public/ChaosCloth/ChaosClothConfig.h
@@ -219,10 +219,22 @@ public:
 	UPROPERTY(EditAnywhere, Category = "Collision Properties")
 	bool bUseLegacyBackstop = false;
 
-	/** The amount of damping applied to the cloth velocities. */
+	/**
+	 * The amount of global damping applied to the cloth velocities, also known as point damping.
+	 * Point damping improves simulation stability, but can also cause an overall slow-down effect and therefore is best left to very small percentage amounts.
+	 */
 	UPROPERTY(EditAnywhere, Category = "Environmental Properties", meta = (UIMin = "0", UIMax = "1", ClampMin = "0", ClampMax = "1"))
 	float DampingCoefficient = 0.01f;
 
+	/**
+	 * The amount of local damping applied to the cloth velocities.
+	 * This type of damping only damps individual deviations of the particles velocities from the global motion.
+	 * It makes the cloth deformations more cohesive and reduces jitter without affecting the overall movement.
+	 * It can also produce synchronization artifacts where part of the cloth mesh are disconnected (which might well be desirable, or not), and is more expensive than global damping.
+	 */
+	UPROPERTY(EditAnywhere, Category = "Environmental Properties", meta = (UIMin = "0", UIMax = "1", ClampMin = "0", ClampMax = "1"))
+	float LocalDampingCoefficient = 0.f;
+
 	/**
 	 * This parameter is automatically set by the migration code. It can be overridden here to use the old deprecated "Legacy" wind model in order to preserve behavior with previous versions of the engine.
 	 * The old wind model is not an accurate aerodynamic model and as such should be avoided. Being point based, it doesn't take into account the surface area that gets hit by the wind.
diff --git a/Engine/Plugins/Experimental/ChaosCloth/Source/Chaos/Public/ChaosCloth/ChaosClothingSimulationCloth.h b/Engine/Plugins/Experimental/ChaosCloth/Source/Chaos/Public/ChaosCloth/ChaosClothingSimulationCloth.h
--- a/Engine/Plugins/Experimental/ChaosCloth/Source/Chaos/Public/ChaosCloth/ChaosClothingSimulationCloth.h
+++ b/Engine/Plugins/Experimental/ChaosCloth/Source/Chaos/Public/ChaosCloth/ChaosClothingSimulationCloth.h
@@ -57,6 +57,7 @@ namespace Chaos
 			const TVec2<FRealSingle>& InLift,
 			bool bInUseLegacyWind,
 			FRealSingle InDampingCoefficient,
+			FRealSingle InLocalDampingCoefficient,
 			FRealSingle InCollisionThickness,
 			FRealSingle InFrictionCoefficient,
 			bool bInUseCCD,
@@ -216,6 +217,7 @@ namespace Chaos
 		FRealSingle AirDensity;
 		bool bUseLegacyWind;
 		FRealSingle DampingCoefficient;
+		FRealSingle LocalDampingCoefficient;
 		FRealSingle CollisionThickness;
 		FRealSingle FrictionCoefficient;
 		bool bUseCCD;
diff --git a/Engine/Plugins/Experimental/ChaosCloth/Source/Chaos/Public/ChaosCloth/ChaosClothingSimulationSolver.h b/Engine/Plugins/Experimental/ChaosCloth/Source/Chaos/Public/ChaosCloth/ChaosClothingSimulationSolver.h
--- a/Engine/Plugins/Experimental/ChaosCloth/Source/Chaos/Public/ChaosCloth/ChaosClothingSimulationSolver.h
+++ b/Engine/Plugins/Experimental/ChaosCloth/Source/Chaos/Public/ChaosCloth/ChaosClothingSimulationSolver.h
@@ -86,7 +86,12 @@ namespace Chaos
 			FRealSingle FictitiousAngularScale);
 
 		// Set general cloth simulation properties.
-		void SetProperties(uint32 GroupId, FRealSingle DampingCoefficient, FRealSingle CollisionThickness, FRealSingle FrictionCoefficient);
+		void SetProperties(
+			uint32 GroupId,
+			FRealSingle DampingCoefficient,
+			FRealSingle MotionDampingCoefficient,
+			FRealSingle CollisionThickness,
+			FRealSingle FrictionCoefficient);
 
 		// Set whether to use continuous collision detection.
 		void SetUseCCD(uint32 GroupId, bool bUseCCD);
diff --git a/Engine/Source/Runtime/Core/Public/UObject/UE5ReleaseStreamObjectVersion.h b/Engine/Source/Runtime/Core/Public/UObject/UE5ReleaseStreamObjectVersion.h
--- a/Engine/Source/Runtime/Core/Public/UObject/UE5ReleaseStreamObjectVersion.h
+++ b/Engine/Source/Runtime/Core/Public/UObject/UE5ReleaseStreamObjectVersion.h
@@ -105,6 +105,9 @@ struct CORE_API FUE5ReleaseStreamObjectVersion
 		// Refresh geometry collections that had not already generated convex bodies.
 		GeometryCollectionConvexDefaults,
 
+		// Add faster damping calculations to the cloth simulation and rename previous Damping parameter to LocalDamping.
+		ChaosClothFasterDamping,
+
 		// -----<new versions can be added above this line>-------------------------------------------------
 		VersionPlusOne,
 		LatestVersion = VersionPlusOne - 1
diff --git a/Engine/Source/Runtime/Experimental/Chaos/Private/Chaos/PBDEvolution.cpp b/Engine/Source/Runtime/Experimental/Chaos/Private/Chaos/PBDEvolution.cpp
--- a/Engine/Source/Runtime/Experimental/Chaos/Private/Chaos/PBDEvolution.cpp
+++ b/Engine/Source/Runtime/Experimental/Chaos/Private/Chaos/PBDEvolution.cpp
@@ -59,6 +59,7 @@ void FPBDEvolution::AddGroups(int32 NumGroups)
 		MGroupSelfCollisionThicknesses[GroupId] = MSelfCollisionThickness;
 		MGroupCoefficientOfFrictions[GroupId] = MCoefficientOfFriction;
 		MGroupDampings[GroupId] = MDamping;
+		MGroupLocalDampings[GroupId] = MLocalDamping;
 		MGroupUseCCDs[GroupId]  = false;
 	}
 }
@@ -69,9 +70,16 @@ void FPBDEvolution::ResetGroups()
 	AddGroups(1);  // Add default group
 }
 
-FPBDEvolution::FPBDEvolution(FSolverParticles&& InParticles, FSolverRigidParticles&& InGeometryParticles, TArray<TVec3<int32>>&& CollisionTriangles,
-    int32 NumIterations, FSolverReal CollisionThickness, FSolverReal SelfCollisionThickness, FSolverReal CoefficientOfFriction, FSolverReal Damping)
-    : MParticles(MoveTemp(InParticles))
+FPBDEvolution::FPBDEvolution(
+	FSolverParticles&& InParticles,
+	FSolverRigidParticles&& InGeometryParticles,
+	TArray<TVec3<int32>>&& CollisionTriangles,
+	int32 NumIterations, FSolverReal CollisionThickness,
+	FSolverReal SelfCollisionThickness,
+	FSolverReal CoefficientOfFriction,
+	FSolverReal Damping,
+	FSolverReal LocalDamping)
+	: MParticles(MoveTemp(InParticles))
 	, MParticlesActiveView(MParticles)
 	, MCollisionParticles(MoveTemp(InGeometryParticles))
 	, MCollisionParticlesActiveView(MCollisionParticles)
@@ -83,6 +91,7 @@ FPBDEvolution::FPBDEvolution(FSolverParticles&& InParticles, FSolverRigidParticl
 	, MSelfCollisionThickness(SelfCollisionThickness)
 	, MCoefficientOfFriction(CoefficientOfFriction)
 	, MDamping(Damping)
+	, MLocalDamping(LocalDamping)
 	, MTime(0)
 	, MSmoothDt(1.f / 30.f)  // Initialize filtered timestep at 30fps
 {
@@ -94,6 +103,7 @@ FPBDEvolution::FPBDEvolution(FSolverParticles&& InParticles, FSolverRigidParticl
 	TArrayCollection::AddArray(&MGroupSelfCollisionThicknesses);
 	TArrayCollection::AddArray(&MGroupCoefficientOfFrictions);
 	TArrayCollection::AddArray(&MGroupDampings);
+	TArrayCollection::AddArray(&MGroupLocalDampings);
 	TArrayCollection::AddArray(&MGroupUseCCDs);
 	AddGroups(1);  // Add default group
 
@@ -208,7 +218,7 @@ void FPBDEvolution::PreIterationUpdate(
 		VelocityField.UpdateForces(MParticles, Dt);  // Update force per surface element
 	}
 
-	FPerParticleDampVelocity DampVelocityRule(MGroupDampings[ParticleGroupId]);
+	FPerParticleDampVelocity DampVelocityRule(MGroupLocalDampings[ParticleGroupId]);
 	if (bDampVelocityRule)
 	{
 		TRACE_CPUPROFILER_EVENT_SCOPE(ChaosPBDVelocityDampUpdateState);
@@ -220,9 +230,30 @@ void FPBDEvolution::PreIterationUpdate(
 		TRACE_CPUPROFILER_EVENT_SCOPE(ChaosClothSolverIntegrate);
 		SCOPE_CYCLE_COUNTER(STAT_ChaosClothSolverIntegrate);
 
+		constexpr FSolverReal DampingFrequency = (FSolverReal)60.;  // The damping value is the percentage of velocity removed per frame when running at 60Hz
+		const FSolverReal Damping = FMath::Clamp(MGroupDampings[ParticleGroupId], (FSolverReal)0., (FSolverReal)1.);
+		FSolverReal DampingPowDt;
+		FSolverReal DampingIntegrated;
+		if (Damping > (FSolverReal)1. - (FSolverReal)KINDA_SMALL_NUMBER)
+		{
+			DampingIntegrated = DampingPowDt = (FSolverReal)0.;
+		}
+		else if (Damping > (FSolverReal)SMALL_NUMBER)
+		{
+			const FSolverReal LogValueByFrequency = FMath::Loge((FSolverReal)1. - Damping) * DampingFrequency;
+
+			DampingPowDt = FMath::Exp(LogValueByFrequency * Dt);  // DampingPowDt = FMath::Pow(OneMinusDamping, Dt * DampingFrequency);
+			DampingIntegrated = (DampingPowDt - (FSolverReal)1.) / LogValueByFrequency;
+		}
+		else
+		{
+			DampingPowDt = (FSolverReal)1.;
+			DampingIntegrated = Dt;
+		}
+
 		const int32 RangeSize = Range - Offset;
 		PhysicsParallelFor(RangeSize,
-			[this, &Offset, &ForceRule, &Gravity, &VelocityField, &DampVelocityRule, Dt](int32 i)
+			[this, &Offset, &ForceRule, &Gravity, &VelocityField, &DampVelocityRule, DampingPowDt, DampingIntegrated, Dt](int32 i)
 			{
 				const int32 Index = Offset + i;
 				if (MParticles.InvM(Index) != (FSolverReal)0.)  // Process dynamic particles
@@ -251,8 +282,10 @@ void FPBDEvolution::PreIterationUpdate(
 						DampVelocityRule.ApplyFast(MParticles, Dt, Index);
 					}
 
-					// Euler Step
-					MParticles.P(Index) = MParticles.X(Index) + MParticles.V(Index) * Dt;
+					// Euler Step with point damping integration
+					MParticles.P(Index) = MParticles.X(Index) + MParticles.V(Index) * DampingIntegrated;
+
+					MParticles.V(Index) *= DampingPowDt;
 				}
 				else  // Process kinematic particles
 				{
@@ -299,7 +332,7 @@ void FPBDEvolution::AdvanceOneTimeStep(const FSolverReal Dt, const bool bSmoothD
 
 				if (MGroupVelocityFields[ParticleGroupId].IsActive())
 				{
-					if (MGroupDampings[ParticleGroupId] > (FSolverReal)0.)
+					if (MGroupLocalDampings[ParticleGroupId] > (FSolverReal)0.)
 					{
 						if (MGroupForceRules[ParticleGroupId])  // VeloctiyFields, Damping, Forces  // Damping?????
 						{
@@ -324,7 +357,7 @@ void FPBDEvolution::AdvanceOneTimeStep(const FSolverReal Dt, const bool bSmoothD
 				}
 				else   // No Velocity Fields
 				{
-					if (MGroupDampings[ParticleGroupId] > (FSolverReal)0.)
+					if (MGroupLocalDampings[ParticleGroupId] > (FSolverReal)0.)
 					{
 						if (MGroupForceRules[ParticleGroupId])  // VeloctiyFields, Damping, Forces
 						{
diff --git a/Engine/Source/Runtime/Experimental/Chaos/Public/Chaos/PBDEvolution.h b/Engine/Source/Runtime/Experimental/Chaos/Public/Chaos/PBDEvolution.h
--- a/Engine/Source/Runtime/Experimental/Chaos/Public/Chaos/PBDEvolution.h
+++ b/Engine/Source/Runtime/Experimental/Chaos/Public/Chaos/PBDEvolution.h
@@ -15,8 +15,17 @@ namespace Chaos::Softs
 class CHAOS_API FPBDEvolution : public TArrayCollection
 {
  public:
-	// TODO(mlentine): Init particles from some type of input
-	FPBDEvolution(FSolverParticles&& InParticles, FSolverRigidParticles&& InGeometryParticles, TArray<TVec3<int32>>&& CollisionTriangles, int32 NumIterations = 1, FSolverReal CollisionThickness = 0, FSolverReal SelfCollisionsThickness = 0, FSolverReal CoefficientOfFriction = 0, FSolverReal Damping = 0.04f);
+	// TODO: Tidy up this constructor (and update Headless Chaos)
+	FPBDEvolution(
+		FSolverParticles&& InParticles,
+		FSolverRigidParticles&& InGeometryParticles,
+		TArray<TVec3<int32>>&& CollisionTriangles,
+		int32 NumIterations = 1,
+		FSolverReal CollisionThickness = (FSolverReal)0.,
+		FSolverReal SelfCollisionsThickness = (FSolverReal)0.,
+		FSolverReal CoefficientOfFriction = (FSolverReal)0.,
+		FSolverReal Damping = (FSolverReal)0.04,
+		FSolverReal LocalDamping = (FSolverReal)0.);
 	~FPBDEvolution() {}
 
 	// Advance one time step. Filter the input time step if specified.
@@ -108,6 +117,9 @@ class CHAOS_API FPBDEvolution : public TArrayCollection
 	FSolverReal GetDamping(const uint32 GroupId = 0) const { check(GroupId < TArrayCollection::Size()); return MGroupDampings[GroupId]; }
 	void SetDamping(const FSolverReal Damping, const uint32 GroupId = 0) { check(GroupId < TArrayCollection::Size()); MGroupDampings[GroupId] = Damping; }
 
+	FSolverReal GetLocalDamping(const uint32 GroupId = 0) const { check(GroupId < TArrayCollection::Size()); return MGroupLocalDampings[GroupId]; }
+	void SetLocalDamping(const FSolverReal LocalDamping, const uint32 GroupId = 0) { check(GroupId < TArrayCollection::Size()); MGroupLocalDampings[GroupId] = LocalDamping; }
+
 	bool GetUseCCD(const uint32 GroupId = 0) const { check(GroupId < TArrayCollection::Size()); return MGroupUseCCDs[GroupId]; }
 	void SetUseCCD(const bool bUseCCD, const uint32 GroupId = 0) { check(GroupId < TArrayCollection::Size()); MGroupUseCCDs[GroupId] = bUseCCD; }
 
@@ -149,6 +161,7 @@ private:
 	TArrayCollectionArray<FSolverReal> MGroupSelfCollisionThicknesses;
 	TArrayCollectionArray<FSolverReal> MGroupCoefficientOfFrictions;
 	TArrayCollectionArray<FSolverReal> MGroupDampings;
+	TArrayCollectionArray<FSolverReal> MGroupLocalDampings;
 	TArrayCollectionArray<bool> MGroupUseCCDs;
 	
 	TArray<TFunction<void(FSolverParticles&, const FSolverReal)>> MConstraintInits;
@@ -165,6 +178,7 @@ private:
 	FSolverReal MSelfCollisionThickness;
 	FSolverReal MCoefficientOfFriction;
 	FSolverReal MDamping;
+	FSolverReal MLocalDamping;
 	FSolverReal MTime;
 	FSolverReal MSmoothDt;
 };
