Fixed precision checking logic in TUniformGrid when we use doubles as the underlying type

diff --git a/Engine/Source/Runtime/Experimental/Chaos/Public/Chaos/UniformGrid.h b/Engine/Source/Runtime/Experimental/Chaos/Public/Chaos/UniformGrid.h
--- a/Engine/Source/Runtime/Experimental/Chaos/Public/Chaos/UniformGrid.h
+++ b/Engine/Source/Runtime/Experimental/Chaos/Public/Chaos/UniformGrid.h
@@ -11,6 +11,17 @@ class TArrayND;
 template<class T, int d>
 class TArrayFaceND;
 
+template<typename T>
+struct TGridPrecisionLimit
+{
+	static_assert(sizeof(T) == 0, "Unsupported grid floating point type");
+};
+
+// Precision limits where our floating point precision breaks down to 1 (each value above this adds more than 1 to the number)
+// Note the F32 limit is actually higher than this point (should be 8.388e6) but kept at 1e7 for legacy reasons
+template<> struct TGridPrecisionLimit<FRealSingle> { static constexpr FRealSingle value = 1e7; };
+template<> struct TGridPrecisionLimit<FRealDouble> { static constexpr FRealDouble value = 4.5035e15; };
+
 template<class T, int d>
 class CHAOS_API TUniformGridBase
 {
@@ -68,7 +79,7 @@ class CHAOS_API TUniformGridBase
 			const TVector<T, d> MinToDXRatio = MMinCorner / MDx;
 			for (int32 Axis = 0; Axis < d; ++Axis)
 			{
-				ensure(FMath::Abs(MinToDXRatio[Axis]) < 1e7); //make sure we have the precision we need
+				ensure(FMath::Abs(MinToDXRatio[Axis]) < TGridPrecisionLimit<T>::value); //make sure we have the precision we need
 			}
 		}
 	}
