Implement custom batch for Chaos. Currently only used by Gather. Shaves about 30% of the time off of gather.

diff --git a/Engine/Source/Runtime/Experimental/Chaos/Private/Chaos/Framework/Parallel.cpp b/Engine/Source/Runtime/Experimental/Chaos/Private/Chaos/Framework/Parallel.cpp
--- a/Engine/Source/Runtime/Experimental/Chaos/Private/Chaos/Framework/Parallel.cpp
+++ b/Engine/Source/Runtime/Experimental/Chaos/Private/Chaos/Framework/Parallel.cpp
@@ -9,7 +9,9 @@ using namespace Chaos;
 namespace Chaos
 {
 	CHAOS_API int32 GSingleThreadedPhysics = 0;
-	CHAOS_API int32 InnerParallelForBatchSize = 4;
+	CHAOS_API int32 InnerParallelForBatchSize = 0;
+	CHAOS_API int32 MinRangeBatchSize = 10;
+	CHAOS_API int32 MaxRangeBatchWorkers = 100;
 	
 #if !UE_BUILD_SHIPPING
 	CHAOS_API bool bDisablePhysicsParallelFor = false;
@@ -20,16 +22,25 @@ namespace Chaos
 	FAutoConsoleVariableRef CVarDisableParticleParallelFor(TEXT("p.Chaos.DisableParticleParallelFor"), bDisableParticleParallelFor, TEXT("Disable parallel execution for Chaos Particles (Collisions, "));
 	FAutoConsoleVariableRef CVarDisableCollisionParallelFor(TEXT("p.Chaos.DisableCollisionParallelFor"), bDisableCollisionParallelFor, TEXT("Disable parallel execution for Chaos Collisions (also disabled by DisableParticleParallelFor)"));
 	FAutoConsoleVariableRef CVarInnerPhysicsBatchSize(TEXT("p.Chaos.InnerParallelForBatchSize"), InnerParallelForBatchSize, TEXT("Set the batch size threshold for inner parallel fors"));
+	FAutoConsoleVariableRef CVarMinRangeBatchSize(TEXT("p.Chaos.MinRangeBatchSize"), MinRangeBatchSize, TEXT("Set the min range batch size for parallel for"));
+	FAutoConsoleVariableRef CVarMaxRangeBatchWorkers(TEXT("p.Chaos.MaxRangeBatchWorkers"), MaxRangeBatchWorkers, TEXT("Set the max range batch num workers for parallel for"));
 #endif
 }
 
 void Chaos::InnerPhysicsParallelFor(int32 Num, TFunctionRef<void(int32)> InCallable, bool bForceSingleThreaded)
 {
 	int32 NumWorkers = int32(LowLevelTasks::FScheduler::Get().GetNumWorkers());
-	int32 BatchSize = Num / NumWorkers;
+	int32 BatchSize = FMath::DivideAndRoundUp<int32>(Num, NumWorkers);
 	PhysicsParallelFor(Num, InCallable, (BatchSize > InnerParallelForBatchSize) ? bForceSingleThreaded : true);
 }
 
+void Chaos::InnerPhysicsParallelForRange(int32 InNum, TFunctionRef<void(int32, int32)> InCallable, bool bForceSingleThreaded)
+{
+	int32 NumWorkers = int32(LowLevelTasks::FScheduler::Get().GetNumWorkers());
+	int32 BatchSize = FMath::DivideAndRoundUp<int32>(InNum, NumWorkers);
+	PhysicsParallelForRange(InNum, InCallable, (BatchSize > InnerParallelForBatchSize) ? bForceSingleThreaded : true);
+}
+
 void Chaos::PhysicsParallelFor(int32 InNum, TFunctionRef<void(int32)> InCallable, bool bForceSingleThreaded)
 {
 	TRACE_CPUPROFILER_EVENT_SCOPE(Chaos_PhysicsParallelFor);
@@ -55,6 +66,72 @@ void Chaos::PhysicsParallelFor(int32 InNum, TFunctionRef<void(int32)> InCallable
 	::ParallelFor(InNum, PassThrough, !!GSingleThreadedPhysics || bDisablePhysicsParallelFor || bForceSingleThreaded);
 }
 
+void Chaos::PhysicsParallelForRange(int32 InNum, TFunctionRef<void(int32, int32)> InCallable, bool bForceSingleThreaded)
+{
+	TRACE_CPUPROFILER_EVENT_SCOPE(Chaos_PhysicsParallelFor);
+	using namespace Chaos;
+
+	// Passthrough for now, except with global flag to disable parallel
+#if PHYSICS_THREAD_CONTEXT
+	const bool bIsInPhysicsSimContext = IsInPhysicsThreadContext();
+	const bool bIsInGameThreadContext = IsInGameThreadContext();
+#else
+	const bool bIsInPhysicsSimContext = false;
+	const bool bIsInGameThreadContext = false;
+#endif
+
+	//calculate the number of workers
+	int32 NumWorkers = int32(LowLevelTasks::FScheduler::Get().GetNumWorkers());
+	if (!LowLevelTasks::FScheduler::Get().IsWorkerThread())
+	{
+		NumWorkers++; //named threads help with the work
+	}
+	NumWorkers = FMath::Min(NumWorkers, InNum);
+	NumWorkers = FMath::Min(NumWorkers, MaxRangeBatchWorkers);
+	check(NumWorkers > 0);
+	int32 BatchSize = FMath::DivideAndRoundUp<int32>(InNum, NumWorkers);
+	// @todo(mlentine): Find a better batch size in this case
+	if (InNum < MinRangeBatchSize)
+	{
+		NumWorkers = 1;
+		BatchSize = InNum;
+	}
+	else
+	{
+		while (BatchSize < MinRangeBatchSize && NumWorkers > 1)
+		{
+			NumWorkers /= 2;
+			BatchSize = FMath::DivideAndRoundUp<int32>(InNum, NumWorkers);
+		}
+	}
+	TArray<int32> RangeIndex;
+	RangeIndex.Add(0);
+	for (int32 i = 1; i <= NumWorkers; i++)
+	{
+		int32 PrevEnd = RangeIndex[i - 1];
+		int32 NextEnd = FMath::Min(BatchSize + RangeIndex[i - 1], InNum);
+		if (NextEnd != PrevEnd)
+		{
+			RangeIndex.Add(NextEnd);
+		}
+		else
+		{
+			break;
+		}
+	}
+
+	auto PassThrough = [InCallable, &RangeIndex, bIsInPhysicsSimContext, bIsInGameThreadContext](int32 ThreadId)
+	{
+#if PHYSICS_THREAD_CONTEXT
+		FPhysicsThreadContextScope PTScope(bIsInPhysicsSimContext);
+		FGameThreadContextScope GTScope(bIsInGameThreadContext);
+#endif
+		InCallable(RangeIndex[ThreadId], RangeIndex[ThreadId + 1]);
+	};
+	::ParallelFor(RangeIndex.Num() - 1, PassThrough, !!GSingleThreadedPhysics || bDisablePhysicsParallelFor || bForceSingleThreaded);
+}
+
+
 //class FRecursiveDivideTask
 //{
 //	TFuture<void> ThisFuture;
diff --git a/Engine/Source/Runtime/Experimental/Chaos/Private/Chaos/PBDConstraintRule.cpp b/Engine/Source/Runtime/Experimental/Chaos/Private/Chaos/PBDConstraintRule.cpp
--- a/Engine/Source/Runtime/Experimental/Chaos/Private/Chaos/PBDConstraintRule.cpp
+++ b/Engine/Source/Runtime/Experimental/Chaos/Private/Chaos/PBDConstraintRule.cpp
@@ -310,21 +310,24 @@ namespace Chaos
 										}
 										// Calculate the range of indices for this color as a set of independent contacts
 										TPair<int32, int32> ColorConstrainSet(ConstraintSetEnd, ConstraintSetEnd);
-										InnerPhysicsParallelFor(OffsetEnd - OffsetBegin, [&](int32 Index)
+										InnerPhysicsParallelForRange(OffsetEnd - OffsetBegin, [&](int32 StartIndex, int32 EndIndex)
 										{
-											int32 ConstraintIndex = Index + OffsetBegin;
-											FConstraintContainerHandle* Constraint = SortedConstraints[ConstraintIndex];
-											if (Constraint->IsEnabled())
+											for (int32 Index = StartIndex; Index < EndIndex; ++Index)
 											{
-												// Levels that should be assigned to the bodies for shock propagation
-												// @todo(chaos): optimize the lookup
-												const TVector<TGeometryParticleHandle<FReal, 3>*, 2> ConstrainedParticles = Constraint->GetConstrainedParticles();
-												const int32 Particle0Level = GetParticleLevel(ConstrainedParticles[0]);
-												const int32 Particle1Level = GetParticleLevel(ConstrainedParticles[1]);
-
-												// Note we are building the SolverBodies as we go, in the order that we visit them. Each constraint
-												// references two bodies, so we won't strictly be accessing only in cache order, but it's about as good as it can be.
-												Constraint->GatherInput(Dt, Particle0Level, Particle1Level, *IslandGroup, false);
+												int32 ConstraintIndex = Index + OffsetBegin;
+												FConstraintContainerHandle* Constraint = SortedConstraints[ConstraintIndex];
+												if (Constraint->IsEnabled())
+												{
+													// Levels that should be assigned to the bodies for shock propagation
+													// @todo(chaos): optimize the lookup
+													const TVector<TGeometryParticleHandle<FReal, 3>*, 2> ConstrainedParticles = Constraint->GetConstrainedParticles();
+													const int32 Particle0Level = GetParticleLevel(ConstrainedParticles[0]);
+													const int32 Particle1Level = GetParticleLevel(ConstrainedParticles[1]);
+
+													// Note we are building the SolverBodies as we go, in the order that we visit them. Each constraint
+													// references two bodies, so we won't strictly be accessing only in cache order, but it's about as good as it can be.
+													Constraint->GatherInput(Dt, Particle0Level, Particle1Level, *IslandGroup, false);
+												}
 											}
 										});
 										for (int32 ConstraintIndex = OffsetBegin; ConstraintIndex < OffsetEnd; ++ConstraintIndex)
diff --git a/Engine/Source/Runtime/Experimental/Chaos/Public/Chaos/Framework/Parallel.h b/Engine/Source/Runtime/Experimental/Chaos/Public/Chaos/Framework/Parallel.h
--- a/Engine/Source/Runtime/Experimental/Chaos/Public/Chaos/Framework/Parallel.h
+++ b/Engine/Source/Runtime/Experimental/Chaos/Public/Chaos/Framework/Parallel.h
@@ -6,7 +6,9 @@
 
 namespace Chaos
 {
+	void CHAOS_API PhysicsParallelForRange(int32 InNum, TFunctionRef<void(int32, int32)> InCallable, bool bForceSingleThreaded = false);
 	void CHAOS_API PhysicsParallelFor(int32 InNum, TFunctionRef<void(int32)> InCallable, bool bForceSingleThreaded = false);
+	void CHAOS_API InnerPhysicsParallelForRange(int32 InNum, TFunctionRef<void(int32, int32)> InCallable, bool bForceSingleThreaded = false);
 	void CHAOS_API InnerPhysicsParallelFor(int32 InNum, TFunctionRef<void(int32)> InCallable, bool bForceSingleThreaded = false);
 	//void CHAOS_API PhysicsParallelFor_RecursiveDivide(int32 InNum, TFunctionRef<void(int32)> InCallable, bool bForceSingleThreaded = false);
 
