Chaos - RBAN solver fixes
- more NaN detection in the solver
- fix shock propagation when second body is kinematic
- disable linear solver for RBAN because no projection support yet
- update body and constraint state every iteration with RBAN
- disable incremental manifolds for RBAN
- make contact reuse into a solver-specific setting and  disable it for RBAN solver
- fix shape transforms for per-iteration collision detection
- don't use pair iterations for QuasiPBD
- RBAN projection now obeys the "solve position last" option

diff --git a/Engine/Source/Runtime/Engine/Private/PhysicsEngine/ImmediatePhysics/ImmediatePhysicsChaos/ImmediatePhysicsSimulation_Chaos.cpp b/Engine/Source/Runtime/Engine/Private/PhysicsEngine/ImmediatePhysics/ImmediatePhysicsChaos/ImmediatePhysicsSimulation_Chaos.cpp
--- a/Engine/Source/Runtime/Engine/Private/PhysicsEngine/ImmediatePhysics/ImmediatePhysicsChaos/ImmediatePhysicsSimulation_Chaos.cpp
+++ b/Engine/Source/Runtime/Engine/Private/PhysicsEngine/ImmediatePhysics/ImmediatePhysicsChaos/ImmediatePhysicsSimulation_Chaos.cpp
@@ -364,6 +364,7 @@ namespace ImmediatePhysics_Chaos
 		Implementation->NarrowPhase.GetContext().bFilteringEnabled = false;
 		Implementation->NarrowPhase.GetContext().bDeferUpdate = true;
 		Implementation->NarrowPhase.GetContext().bAllowManifolds = false;
+		Implementation->NarrowPhase.GetContext().bAllowManifoldReuse = false;
 	}
 
 	FSimulation::~FSimulation()
@@ -813,7 +814,7 @@ namespace ImmediatePhysics_Chaos
 			JointsSettings.LinearDriveDampingOverride = ChaosImmediate_Joint_LinearDriveDamping;
 			JointsSettings.AngularDriveStiffnessOverride = ChaosImmediate_Joint_AngularDriveStiffness;
 			JointsSettings.AngularDriveDampingOverride = ChaosImmediate_Joint_AngularDriveDamping;
-			JointsSettings.bUseLinearSolver = bChaosImmediate_Joint_UseLinearSolver;
+			JointsSettings.bUseLinearSolver = bChaosImmediate_Joint_UseLinearSolver && ((EConstraintSolverType)ChaosImmediate_SolverType == EConstraintSolverType::QuasiPbd);
 			Implementation->Joints.SetSettings(JointsSettings);
 
 			Implementation->Collisions.SetRestitutionEnabled(ChaosImmediate_Collision_RestitutionEnabled != 0);
@@ -826,7 +827,9 @@ namespace ImmediatePhysics_Chaos
 			Implementation->NarrowPhase.SetBoundsExpansion(ChaosImmediate_Collision_CullDistance);
 			Implementation->NarrowPhase.GetContext().bDeferUpdate = (ChaosImmediate_Collision_DeferNarrowPhase != 0);
 			Implementation->NarrowPhase.GetContext().bAllowManifolds = (ChaosImmediate_Collision_UseManifolds != 0);
+			Implementation->NarrowPhase.GetContext().bAllowManifoldReuse = false;
 
+			Implementation->Evolution.SetSolverType((EConstraintSolverType)ChaosImmediate_SolverType);
 			Implementation->Collisions.SetSolverType((EConstraintSolverType)ChaosImmediate_SolverType);
 			Implementation->Joints.SetSolverType((EConstraintSolverType)ChaosImmediate_SolverType);
 
diff --git a/Engine/Source/Runtime/Experimental/Chaos/Private/Chaos/Collision/PBDCollisionConstraint.cpp b/Engine/Source/Runtime/Experimental/Chaos/Private/Chaos/Collision/PBDCollisionConstraint.cpp
--- a/Engine/Source/Runtime/Experimental/Chaos/Private/Chaos/Collision/PBDCollisionConstraint.cpp
+++ b/Engine/Source/Runtime/Experimental/Chaos/Private/Chaos/Collision/PBDCollisionConstraint.cpp
@@ -245,8 +245,12 @@ namespace Chaos
 
 		CullDistance = InCullDistance;
 
+		// Are we allowing manifolds? If manifolds are enabled, we will build a one-shot manifold
+		// if supported by the shape pair, otherwise an incremental manifold will be created and
+		// we call collision detection every iteration to add new points (this is expensive).
+		// NOTE: bUseIncrementalManifold will get set to false later if we add a one-shot manifold
 		Flags.bUseManifold = bInUseManifold;
-		Flags.bUseIncrementalManifold = true;	// This will get changed later if we call AddOneShotManifoldContact
+		Flags.bUseIncrementalManifold = bInUseManifold;
 
 		const FReal Margin0 = GetImplicit0()->GetMargin();
 		const FReal Margin1 = GetImplicit1()->GetMargin();
@@ -511,10 +515,13 @@ namespace Chaos
 	{
 		if (ManifoldPoints.IsFull())
 		{
+			// @todo(chaos): we should remove a contact here if we try to add a new point
+			// For now just update the existing ones to select the deepest
+			UpdateManifoldContacts();
 			return;
 		}
 
-		if (Flags.bUseManifold)
+		if (Flags.bUseIncrementalManifold)
 		{
 			// See if the manifold point already exists
 			int32 ManifoldPointIndex = FindManifoldPoint(ContactPoint);
@@ -529,10 +536,9 @@ namespace Chaos
 				ManifoldPointIndex = AddManifoldPoint(ContactPoint);
 			}
 
-			if (ManifoldPoints[ManifoldPointIndex].ContactPoint.Phi < GetPhi())
-			{
-				ClosestManifoldPointIndex = ManifoldPointIndex;
-			}
+			// If collision detection did its job, this contact is the deepest
+			// NOTE: other contact Phis will be out of date at the current iteration's transforms
+			ClosestManifoldPointIndex = ManifoldPointIndex;
 		}
 		else 
 		{
@@ -544,8 +550,6 @@ namespace Chaos
 
 			ClosestManifoldPointIndex = 0;
 		}
-
-		Flags.bUseIncrementalManifold = true;
 	}
 
 	void FPBDCollisionConstraint::ResetManifold()
diff --git a/Engine/Source/Runtime/Experimental/Chaos/Private/Chaos/Collision/ParticlePairMidPhase.cpp b/Engine/Source/Runtime/Experimental/Chaos/Private/Chaos/Collision/ParticlePairMidPhase.cpp
--- a/Engine/Source/Runtime/Experimental/Chaos/Private/Chaos/Collision/ParticlePairMidPhase.cpp
+++ b/Engine/Source/Runtime/Experimental/Chaos/Private/Chaos/Collision/ParticlePairMidPhase.cpp
@@ -21,9 +21,6 @@ namespace Chaos
 		bool bChaos_Collision_MidPhase_EnableBoundsChecks = true;
 		FAutoConsoleVariableRef CVarChaos_Collision_EnableBoundsChecks(TEXT("p.Chaos.Collision.EnableBoundsChecks"), bChaos_Collision_MidPhase_EnableBoundsChecks, TEXT(""));
 
-		bool bChaos_Collision_EnableManifoldUpdate = true;
-		FAutoConsoleVariableRef CVarChaos_Collision_EnableManifoldUpdate(TEXT("p.Chaos.Collision.EnableManifoldUpdate"), bChaos_Collision_EnableManifoldUpdate, TEXT(""));
-
 		Chaos::FRealSingle Chaos_Collision_CullDistanceScaleInverseSize = 0.01f;	// 100cm
 		Chaos::FRealSingle Chaos_Collision_MinCullDistanceScale = 1.0f;
 		FAutoConsoleVariableRef CVarChaos_Collision_CullDistanceReferenceSize(TEXT("p.Chaos.Collision.CullDistanceReferenceSize"), Chaos_Collision_CullDistanceScaleInverseSize, TEXT(""));
@@ -121,6 +118,7 @@ namespace Chaos
 		const bool bIsCapsule1 = (ImplicitType1 == ImplicitObjectType::Capsule);
 		const bool bIsTriangle0 = (ImplicitType0 == ImplicitObjectType::TriangleMesh) || (ImplicitType0 == ImplicitObjectType::HeightField);
 		const bool bIsTriangle1 = (ImplicitType1 == ImplicitObjectType::TriangleMesh) || (ImplicitType1 == ImplicitObjectType::HeightField);
+		const bool bIsLevelSet = ((ShapePairType == EContactShapesType::LevelSetLevelSet) || (ShapePairType == EContactShapesType::Unknown));
 
 		const bool bAllowBoundsChecked = bChaos_Collision_MidPhase_EnableBoundsChecks && bHasBounds0 && bHasBounds1;
 		Flags.bEnableAABBCheck = bAllowBoundsChecked && !(bIsSphere0 && bIsSphere1);	// No AABB test if both are spheres
@@ -133,7 +131,8 @@ namespace Chaos
 		}
 
 		// Do not try to reuse manifold points for capsules or spheres (against anything)
-		Flags.bEnableManifoldUpdate = bChaos_Collision_EnableManifoldUpdate && !bIsSphere0 && !bIsSphere1 && !bIsCapsule0 && !bIsCapsule1 && !bIsTriangle0 && !bIsTriangle1;
+		// NOTE: This can also be disabled for all shape types by the solver (see GenerateCollisionImpl and the Context)
+		Flags.bEnableManifoldUpdate = !bIsSphere0 && !bIsSphere1 && !bIsCapsule0 && !bIsCapsule1 && !bIsTriangle0 && !bIsTriangle1 && !bIsLevelSet;
 	}
 
 	FSingleShapePairCollisionDetector::~FSingleShapePairCollisionDetector()
@@ -218,30 +217,30 @@ namespace Chaos
 
 	int32 FSingleShapePairCollisionDetector::GenerateCollision(
 		const FReal CullDistance,
-		const bool bUseManifold,
-		const FReal Dt)
+		const FReal Dt,
+		FCollisionContext& Context)
 	{
 		if (DoBoundsOverlap(CullDistance))
 		{
-			return GenerateCollisionImpl(CullDistance, bUseManifold, Dt);
+			return GenerateCollisionImpl(CullDistance, Dt, Context);
 		}
 		return 0;
 	}
 
 	int32 FSingleShapePairCollisionDetector::GenerateCollisionCCD(
 		const FReal CullDistance,
-		const bool bUseManifold,
-		const FReal Dt)
+		const FReal Dt,
+		FCollisionContext& Context)
 	{
-		return GenerateCollisionCCDImpl(CullDistance, bUseManifold, Dt);
+		return GenerateCollisionCCDImpl(CullDistance, Dt, Context);
 	}
 
-	void FSingleShapePairCollisionDetector::CreateConstraint(const FReal CullDistance, const bool bUseManifold)
+	void FSingleShapePairCollisionDetector::CreateConstraint(const FReal CullDistance, FCollisionContext& Context)
 	{
 		PHYSICS_CSV_SCOPED_EXPENSIVE(PhysicsVerbose, NarrowPhase_CreateConstraint);
 		check(Constraint == nullptr);
 
-		Constraint = CreateShapePairConstraint(Particle0, Shape0, Particle1, Shape1, CullDistance, ShapePairType, bUseManifold);
+		Constraint = CreateShapePairConstraint(Particle0, Shape0, Particle1, Shape1, CullDistance, ShapePairType, Context.bAllowManifolds);
 
 		Constraint->GetContainerCookie().MidPhase = &MidPhase;
 		Constraint->GetContainerCookie().bIsMultiShapePair = false;
@@ -250,9 +249,9 @@ namespace Chaos
 	}
 
 	int32 FSingleShapePairCollisionDetector::GenerateCollisionImpl(
-		const FReal CullDistance,
-		const bool bUseManifold,
-		const FReal Dt)
+		const FReal CullDistance, 
+		const FReal Dt,
+		FCollisionContext& Context)
 	{
 		if (Constraint == nullptr)
 		{
@@ -260,7 +259,7 @@ namespace Chaos
 			// other, we never allocate a constraint for them. Once they overlap, we reuse the constraint
 			// until the owing particles are not overlapping. i.e., we keep the constraint even if
 			// the shape pairs stop overlapping, reusing it if they start overlapping again.
-			CreateConstraint(CullDistance, bUseManifold);
+			CreateConstraint(CullDistance, Context);
 		}
 
 		if (Constraint != nullptr)
@@ -291,7 +290,7 @@ namespace Chaos
 			}
 
 			bool bWasManifoldRestored = false;
-			if (Flags.bEnableManifoldUpdate && bWasUpdatedLastTick)
+			if (Context.bAllowManifoldReuse && Flags.bEnableManifoldUpdate && bWasUpdatedLastTick)
 			{
 				// Update the existing manifold. We can re-use as-is if none of the points have moved much and the bodies have not moved much
 				// NOTE: this can succeed in "restoring" even if we have no manifold points
@@ -308,12 +307,17 @@ namespace Chaos
 			{
 				// We will be updating the manifold, if only partially, so update the restore comparison transforms
 				Constraint->SetLastShapeWorldTransforms(ShapeWorldTransform0, ShapeWorldTransform1);
-				// Run the narrow phase
-				Collisions::UpdateConstraint(*Constraint.Get(), ShapeWorldTransform0, ShapeWorldTransform1, Dt);
+
+				if (!Context.bDeferUpdate)
+				{
+					// Run the narrow phase
+					Collisions::UpdateConstraint(*Constraint.Get(), ShapeWorldTransform0, ShapeWorldTransform1, Dt);
+				}
 			}
 
 			// If we have a valid contact, add it to the active list
-			if (Constraint->GetPhi() <= CullDistance)
+			// We also add it to the active list if collision detection is deferred (which is if per-iteration collision detection is enabled like with RBAN)
+			if ((Constraint->GetPhi() <= CullDistance) || Context.bDeferUpdate)
 			{
 				if (MidPhase.GetCollisionAllocator().ActivateConstraint(Constraint.Get()))
 				{
@@ -327,14 +331,14 @@ namespace Chaos
 	}
 
 	int32 FSingleShapePairCollisionDetector::GenerateCollisionCCDImpl(
-		const FReal CullDistance,
-		const bool bUseManifold,
-		const FReal Dt)
+		const FReal CullDistance, 
+		const FReal Dt,
+		FCollisionContext& Context)
 	{
 		if (Constraint == nullptr)
 		{
 			// Lazy creation of the constraint. 
-			CreateConstraint(CullDistance, bUseManifold);
+			CreateConstraint(CullDistance, Context);
 		}
 
 		if (Constraint != nullptr)
@@ -453,7 +457,6 @@ namespace Chaos
 
 	int32 FMultiShapePairCollisionDetector::GenerateCollisions(
 		const FReal CullDistance,
-		const bool bUseManifold,
 		const FReal Dt,
 		FCollisionContext& Context)
 	{
@@ -853,29 +856,26 @@ namespace Chaos
 
 		// Enable CCD?
 		const bool bUseCCD = Flags.bIsCCD && ShouldEnableCCD(Dt);
-
-		// Allow Manifolds?
-		const bool bUseManifold = Context.bAllowManifolds;
-
+		
 		// Run collision detection on all potentially colliding shape pairs
 		NumActiveConstraints = 0;
 		if (bUseCCD) // Generate CCD constraints as long as AABBs overlap
 		{
 			for (FSingleShapePairCollisionDetector& ShapePair : ShapePairDetectors)
 			{
-				NumActiveConstraints += ShapePair.GenerateCollisionCCD(CullDistance, bUseManifold, Dt);
+				NumActiveConstraints += ShapePair.GenerateCollisionCCD(CullDistance, Dt, Context);
 			}
 		}
 		else
 		{
 			for (FSingleShapePairCollisionDetector& ShapePair : ShapePairDetectors)
 			{
-				NumActiveConstraints += ShapePair.GenerateCollision(CullDistance, bUseManifold, Dt);
+				NumActiveConstraints += ShapePair.GenerateCollision(CullDistance, Dt, Context);
 			}
 		}
 		for (FMultiShapePairCollisionDetector& MultiShapePair : MultiShapePairDetectors)
 		{
-			NumActiveConstraints += MultiShapePair.GenerateCollisions(CullDistance, bUseManifold, Dt, Context);
+			NumActiveConstraints += MultiShapePair.GenerateCollisions(CullDistance, Dt, Context);
 		}
 
 		LastUsedEpoch = CollisionAllocator->GetCurrentEpoch();
diff --git a/Engine/Source/Runtime/Experimental/Chaos/Private/Chaos/CollisionResolution.cpp b/Engine/Source/Runtime/Experimental/Chaos/Private/Chaos/CollisionResolution.cpp
--- a/Engine/Source/Runtime/Experimental/Chaos/Private/Chaos/CollisionResolution.cpp
+++ b/Engine/Source/Runtime/Experimental/Chaos/Private/Chaos/CollisionResolution.cpp
@@ -2191,6 +2191,8 @@ namespace Chaos
 		{
 			const FRigidTransform3 WorldTransform0 = Constraint.GetShapeRelativeTransform0() * ParticleWorldTransform0;
 			const FRigidTransform3 WorldTransform1 = Constraint.GetShapeRelativeTransform1() * ParticleWorldTransform1;
+			Constraint.SetShapeWorldTransforms(WorldTransform0, WorldTransform1);
+
 			UpdateConstraintFromGeometryImpl<UpdateType>(Constraint, WorldTransform0, WorldTransform1, Dt);
 		}
 
@@ -2199,6 +2201,8 @@ namespace Chaos
 		{
 			const FRigidTransform3 WorldTransform0 = Constraint.GetShapeRelativeTransform0() * ParticleWorldTransform0;
 			const FRigidTransform3 WorldTransform1 = Constraint.GetShapeRelativeTransform1() * ParticleWorldTransform1;
+			Constraint.SetShapeWorldTransforms(WorldTransform0, WorldTransform1);
+
 			return UpdateConstraintFromGeometrySweptImpl<UpdateType>(Constraint, WorldTransform0, WorldTransform1, Dt);
 		}
 
diff --git a/Engine/Source/Runtime/Experimental/Chaos/Private/Chaos/Evolution/PBDMinEvolution.cpp b/Engine/Source/Runtime/Experimental/Chaos/Private/Chaos/Evolution/PBDMinEvolution.cpp
--- a/Engine/Source/Runtime/Experimental/Chaos/Private/Chaos/Evolution/PBDMinEvolution.cpp
+++ b/Engine/Source/Runtime/Experimental/Chaos/Private/Chaos/Evolution/PBDMinEvolution.cpp
@@ -135,6 +135,7 @@ namespace Chaos
 		, CollisionDetector(InCollisionDetector)
 		, ParticlePrevXs(InPrevX)
 		, ParticlePrevRs(InPrevR)
+		, SolverType(EConstraintSolverType::QuasiPbd)
 		, NumApplyIterations(0)
 		, NumApplyPushOutIterations(0)
 		, BoundsExtension(InBoundsExtension)
@@ -516,7 +517,22 @@ namespace Chaos
 	{
 		SCOPE_CYCLE_COUNTER(STAT_MinEvolution_UpdateVelocites);
 
-		SolverData.GetBodyContainer().SetImplicitVelocities(Dt);
+		// @todo(chaos): clean this up - the two solvers calculate implicit velocity differently because 
+		// QPBD accumulates transform deltas and the StandardPBD applies transform changes directly
+		if (SolverType == EConstraintSolverType::StandardPbd)
+		{
+			for (FSolverBodyAdapter& SolverBody : SolverData.GetBodyContainer().GetBodies())
+			{
+				const FVec3 V = FVec3::CalculateVelocity(SolverBody.GetSolverBody().X(), SolverBody.GetSolverBody().P(), Dt);
+				const FVec3 W = FRotation3::CalculateAngularVelocity(SolverBody.GetSolverBody().R(), SolverBody.GetSolverBody().Q(), Dt);
+				SolverBody.GetSolverBody().SetV(V);
+				SolverBody.GetSolverBody().SetW(W);
+			}
+		}
+		else
+		{ 
+			SolverData.GetBodyContainer().SetImplicitVelocities(Dt);
+		}
 	}
 
 	void FPBDMinEvolution::ApplyConstraintsPhase2(FReal Dt)
diff --git a/Engine/Source/Runtime/Experimental/Chaos/Private/Chaos/Evolution/SolverBodyContainer.cpp b/Engine/Source/Runtime/Experimental/Chaos/Private/Chaos/Evolution/SolverBodyContainer.cpp
--- a/Engine/Source/Runtime/Experimental/Chaos/Private/Chaos/Evolution/SolverBodyContainer.cpp
+++ b/Engine/Source/Runtime/Experimental/Chaos/Private/Chaos/Evolution/SolverBodyContainer.cpp
@@ -79,4 +79,20 @@ namespace Chaos
 		}
 	}
 
+	void FSolverBodyContainer::ApplyCorrections()
+	{
+		for (FSolverBodyAdapter& SolverBody : SolverBodies)
+		{
+			SolverBody.GetSolverBody().ApplyCorrections();
+		}
+	}
+
+	void FSolverBodyContainer::UpdateRotationDependentState()
+	{
+		for (FSolverBodyAdapter& SolverBody : SolverBodies)
+		{
+			SolverBody.GetSolverBody().UpdateRotationDependentState();
+		}
+	}
+
 }
\ No newline at end of file
diff --git a/Engine/Source/Runtime/Experimental/Chaos/Private/Chaos/Joint/PBDJointCachedSolverGaussSeidel.cpp b/Engine/Source/Runtime/Experimental/Chaos/Private/Chaos/Joint/PBDJointCachedSolverGaussSeidel.cpp
--- a/Engine/Source/Runtime/Experimental/Chaos/Private/Chaos/Joint/PBDJointCachedSolverGaussSeidel.cpp
+++ b/Engine/Source/Runtime/Experimental/Chaos/Private/Chaos/Joint/PBDJointCachedSolverGaussSeidel.cpp
@@ -606,8 +606,6 @@ void FPBDJointCachedSolver::ApplyPointPositionConstraint()
 		// Calculate constraint correction
 		const FVec3 DX = SolverStiffness * PointConstraints.ConstraintStiffness * Utilities::Multiply(PointConstraints.ConstraintIM, CX);
 
-		ChaosSolverEnsure(!DX.ContainsNaN());
-
 		// Apply constraint correction
 		if( Body(0).IsDynamic())
 		{
diff --git a/Engine/Source/Runtime/Experimental/Chaos/Private/Chaos/Joint/PBDJointSolverGaussSeidel.cpp b/Engine/Source/Runtime/Experimental/Chaos/Private/Chaos/Joint/PBDJointSolverGaussSeidel.cpp
--- a/Engine/Source/Runtime/Experimental/Chaos/Private/Chaos/Joint/PBDJointSolverGaussSeidel.cpp
+++ b/Engine/Source/Runtime/Experimental/Chaos/Private/Chaos/Joint/PBDJointSolverGaussSeidel.cpp
@@ -184,7 +184,6 @@ namespace Chaos
 		SolverBodies[1].Reset();
 	}
 
-
 	void FPBDJointSolver::Update(
 		const FReal Dt,
 		const FPBDJointSolverSettings& SolverSettings,
@@ -195,6 +194,16 @@ namespace Chaos
 		UpdateIsActive(Dt, SolverSettings, JointSettings);
 	}
 
+	void FPBDJointSolver::UpdateMasses(
+		const FReal InvMScale0,
+		const FReal InvMScale1)
+	{
+		InvMScales[0] = InvMScale0;
+		InvMScales[1] = InvMScale1;
+		UpdateMass0();
+		UpdateMass1();
+	}
+
 	void FPBDJointSolver::UpdateMass0()
 	{
 		if ((ConditionedInvMs[0] > 0) && (InvMScales[0] > 0))
@@ -295,40 +304,20 @@ namespace Chaos
 		// @todo(chaos): We can also apply velocity drives here rather than in the Pbd pass
 	}
 
-
-	void FPBDJointSolver::ApplyProjections(
+	void FPBDJointSolver::ApplyPositionProjection(
 		const FReal Dt,
-		const FReal InSolverStiffness,
 		const FPBDJointSolverSettings& SolverSettings,
-		const FPBDJointSettings& JointSettings)
+		const FPBDJointSettings& JointSettings,
+		FVec3& DP1,
+		FVec3& DR1)
 	{
-		// @todo(chaos): We need to handle parent/child being the other way round
-		if (!IsDynamic(1))
-		{
-			// If child is kinematic, return. 
-			return;
-		}
-
-		SolverStiffness = InSolverStiffness;
-
-		const FReal LinearProjection = FPBDJointUtilities::GetLinearProjection(SolverSettings, JointSettings);
-		const FReal AngularProjection = FPBDJointUtilities::GetAngularProjection(SolverSettings, JointSettings);
-
-		FVec3 DP1 = FVec3(0);
-		FVec3 DR1 = FVec3(0);
-
 		// Position Projection
+		const FReal LinearProjection = FPBDJointUtilities::GetLinearProjection(SolverSettings, JointSettings);
 		const bool bLinearSoft = FPBDJointUtilities::GetSoftLinearLimitEnabled(SolverSettings, JointSettings);
 		const bool bLinearProjectionEnabled = (bLinearSoft && JointSettings.bSoftProjectionEnabled) || (!bLinearSoft && JointSettings.bProjectionEnabled);
 		const TVec3<EJointMotionType>& LinearMotion = JointSettings.LinearMotionTypes;
-		const bool bLinearLocked =
-			(LinearMotion[0] == EJointMotionType::Locked)
-			&& (LinearMotion[1] == EJointMotionType::Locked)
-			&& (LinearMotion[2] == EJointMotionType::Locked);
-		const bool bLinearLimited =
-			(LinearMotion[0] == EJointMotionType::Limited)
-			&& (LinearMotion[1] == EJointMotionType::Limited)
-			&& (LinearMotion[2] == EJointMotionType::Limited);
+		const bool bLinearLocked = (LinearMotion[0] == EJointMotionType::Locked) && (LinearMotion[1] == EJointMotionType::Locked) && (LinearMotion[2] == EJointMotionType::Locked);
+		const bool bLinearLimited = (LinearMotion[0] == EJointMotionType::Limited) && (LinearMotion[1] == EJointMotionType::Limited) && (LinearMotion[2] == EJointMotionType::Limited);
 		if (bLinearProjectionEnabled && (LinearProjection > 0))
 		{
 			if (bLinearLocked)
@@ -341,6 +330,18 @@ namespace Chaos
 			}
 			// @todo(ccaulfield): support mixed linear projection
 		}
+	}
+
+	void FPBDJointSolver::ApplyRotationProjection(
+		const FReal Dt,
+		const FPBDJointSolverSettings& SolverSettings,
+		const FPBDJointSettings& JointSettings,
+		FVec3& DP1,
+		FVec3& DR1)
+	{
+		const FReal AngularProjection = FPBDJointUtilities::GetAngularProjection(SolverSettings, JointSettings);
+		const TVec3<EJointMotionType>& LinearMotion = JointSettings.LinearMotionTypes;
+		const bool bLinearLocked = (LinearMotion[0] == EJointMotionType::Locked) && (LinearMotion[1] == EJointMotionType::Locked) && (LinearMotion[2] == EJointMotionType::Locked);
 
 		// Twist projection
 		const bool bTwistSoft = FPBDJointUtilities::GetSoftTwistLimitEnabled(SolverSettings, JointSettings);
@@ -396,10 +397,46 @@ namespace Chaos
 				ApplySingleLockedSwingProjection(Dt, SolverSettings, JointSettings, EJointAngularConstraintIndex::Swing2, AngularProjection, bLinearLocked, DP1, DR1);
 			}
 		}
+	}
+
+	void FPBDJointSolver::ApplyProjections(
+		const FReal Dt,
+		const FReal InSolverStiffness,
+		const FPBDJointSolverSettings& SolverSettings,
+		const FPBDJointSettings& JointSettings)
+	{
+		// @todo(chaos): We need to handle parent/child being the other way round
+		if (!IsDynamic(1))
+		{
+			// If child is kinematic, return. 
+			return;
+		}
+
+		SolverStiffness = InSolverStiffness;
+
+
+		FVec3 DP1 = FVec3(0);
+		FVec3 DR1 = FVec3(0);
+
+		if (SolverSettings.bSolvePositionLast)
+		{
+			ApplyRotationProjection(Dt, SolverSettings, JointSettings, DP1, DR1);
+			ApplyPositionProjection(Dt, SolverSettings, JointSettings, DP1, DR1);
+		}
+		else
+		{
+			ApplyPositionProjection(Dt, SolverSettings, JointSettings, DP1, DR1);
+			ApplyRotationProjection(Dt, SolverSettings, JointSettings, DP1, DR1);
+		}
 
 		// Final position fixup
+		const TVec3<EJointMotionType>& LinearMotion = JointSettings.LinearMotionTypes;
+		const bool bLinearLocked = (LinearMotion[0] == EJointMotionType::Locked) && (LinearMotion[1] == EJointMotionType::Locked) && (LinearMotion[2] == EJointMotionType::Locked);
 		if (bLinearLocked)
 		{
+			const FReal LinearProjection = FPBDJointUtilities::GetLinearProjection(SolverSettings, JointSettings);
+			const bool bLinearSoft = FPBDJointUtilities::GetSoftLinearLimitEnabled(SolverSettings, JointSettings);
+			const bool bLinearProjectionEnabled = (bLinearSoft && JointSettings.bSoftProjectionEnabled) || (!bLinearSoft && JointSettings.bProjectionEnabled);
 			if (bLinearProjectionEnabled && (LinearProjection > 0))
 			{
 				ApplyTranslateProjection(Dt, SolverSettings, JointSettings, LinearProjection, DP1, DR1);
diff --git a/Engine/Source/Runtime/Experimental/Chaos/Private/Chaos/PBDCollisionConstraints.cpp b/Engine/Source/Runtime/Experimental/Chaos/Private/Chaos/PBDCollisionConstraints.cpp
--- a/Engine/Source/Runtime/Experimental/Chaos/Private/Chaos/PBDCollisionConstraints.cpp
+++ b/Engine/Source/Runtime/Experimental/Chaos/Private/Chaos/PBDCollisionConstraints.cpp
@@ -632,7 +632,7 @@ namespace Chaos
 	bool FPBDCollisionConstraints::LegacyApplyPhase2Serial(const FReal Dt, const int32 Iterations, const int32 NumIterations, const int32 BeginIndex, const int32 EndIndex, FPBDIslandSolverData& SolverData)
 	{
 		bool bNeedsAnotherIteration = false;
-		if (MApplyPairIterations > 0)
+		if (MApplyPushOutPairIterations > 0)
 		{
 			const Collisions::FContactParticleParameters ParticleParameters = GetContactParticleParameters(Dt);
 			const Collisions::FContactIterationParameters IterationParameters = GetContactIterationParameters(Dt, Iterations, NumIterations, MApplyPushOutPairIterations, bNeedsAnotherIteration);
diff --git a/Engine/Source/Runtime/Experimental/Chaos/Private/Chaos/PBDCollisionConstraintsContact.cpp b/Engine/Source/Runtime/Experimental/Chaos/Private/Chaos/PBDCollisionConstraintsContact.cpp
--- a/Engine/Source/Runtime/Experimental/Chaos/Private/Chaos/PBDCollisionConstraintsContact.cpp
+++ b/Engine/Source/Runtime/Experimental/Chaos/Private/Chaos/PBDCollisionConstraintsContact.cpp
@@ -182,14 +182,16 @@ namespace Chaos
 					FVec3 DP0 = Body0.InvM() * DX;
 					FVec3 DR0 = Utilities::Multiply(FMatrix33(Body0.InvI()), FVec3::CrossProduct(VectorToPoint0, DX));
 					Body0.ApplyTransformDelta(DP0, DR0);
-					//Body0.UpdateRotationDependentState();
+					Body0.ApplyCorrections();
+					Body0.UpdateRotationDependentState();
 				}
 				if (Body1.IsDynamic())
 				{
 					FVec3 DP1 = Body1.InvM() * -DX;
 					FVec3 DR1 = Utilities::Multiply(FMatrix33(Body1.InvI()), FVec3::CrossProduct(VectorToPoint1, -DX));
 					Body1.ApplyTransformDelta(DP1, DR1);
-					//Body1.UpdateRotationDependentState();
+					Body1.ApplyCorrections();
+					Body1.UpdateRotationDependentState();
 				}
 
 				return DX;
@@ -205,7 +207,7 @@ namespace Chaos
 			for (int32 PairIt = 0; PairIt < IterationParameters.NumPairIterations; ++PairIt)
 			{
 				// Update the contact information based on current particles' positions
-				if (!Constraint.GetUseManifold() || (Constraint.GetManifoldPoints().Num() == 0))
+				if (Constraint.GetUseIncrementalCollisionDetection() || (Constraint.GetManifoldPoints().Num() == 0))
 				{
 					Collisions::Update(Constraint, IterationParameters.Dt);
 				}
diff --git a/Engine/Source/Runtime/Experimental/Chaos/Private/Chaos/PBDJointConstraintUtilities.cpp b/Engine/Source/Runtime/Experimental/Chaos/Private/Chaos/PBDJointConstraintUtilities.cpp
--- a/Engine/Source/Runtime/Experimental/Chaos/Private/Chaos/PBDJointConstraintUtilities.cpp
+++ b/Engine/Source/Runtime/Experimental/Chaos/Private/Chaos/PBDJointConstraintUtilities.cpp
@@ -602,20 +602,14 @@ namespace Chaos
 	}
 
 	FReal FPBDJointUtilities::GetShockPropagationInvMassScale(
-		EConstraintSolverType SolverType,
-		int32 It,
-		int32 NumIts,
 		const FPBDJointSolverSettings& SolverSettings,
 		const FPBDJointSettings& JointSettings)
 	{
-		// Shock propagation is only enabled for the last iteration, and only for the QPBD solver
-		const bool bEnableShockPropagation = JointSettings.bProjectionEnabled && (It >= NumIts - SolverSettings.NumShockPropagationIterations) && (SolverType == EConstraintSolverType::QuasiPbd);
-
-		if (bEnableShockPropagation)
+		if (JointSettings.bProjectionEnabled)
 		{
-			return (SolverSettings.ShockPropagationOverride >= 0.0f) ? SolverSettings.ShockPropagationOverride : JointSettings.ShockPropagation;
+			// ShockProagation setting is a alpha. For an alpha of 0 we want an invmass scale of 1, and vice-versa
+			return (SolverSettings.ShockPropagationOverride >= FReal(0)) ? (FReal(1) - SolverSettings.ShockPropagationOverride) : (FReal(1) - JointSettings.ShockPropagation);
 		}
-
 		return FReal(1);
 	}
 
diff --git a/Engine/Source/Runtime/Experimental/Chaos/Private/Chaos/PBDJointConstraints.cpp b/Engine/Source/Runtime/Experimental/Chaos/Private/Chaos/PBDJointConstraints.cpp
--- a/Engine/Source/Runtime/Experimental/Chaos/Private/Chaos/PBDJointConstraints.cpp
+++ b/Engine/Source/Runtime/Experimental/Chaos/Private/Chaos/PBDJointConstraints.cpp
@@ -994,11 +994,12 @@ namespace Chaos
 		SCOPE_CYCLE_COUNTER(STAT_Joints_Apply);
 
 		int32 NumActive = 0;
-		if (Settings.ApplyPairIterations > 0)
+		const int32 NumPairIts = (SolverType == EConstraintSolverType::QuasiPbd) ? 1 : Settings.ApplyPairIterations;
+		if (NumPairIts > 0)
 		{
 			for (int32 ConstraintIndex : SolverData.GetConstraintIndices(ContainerId))
 			{
-				NumActive += ApplyPhase1Single(Dt, ConstraintIndex, Settings.ApplyPairIterations, It, NumIts);
+				NumActive += ApplyPhase1Single(Dt, ConstraintIndex, NumPairIts, It, NumIts);
 			}
 		}
 		return (NumActive > 0);
@@ -1010,12 +1011,13 @@ namespace Chaos
 		SCOPE_CYCLE_COUNTER(STAT_Joints_ApplyPushOut);
 
 		int32 NumActive = 0;
-		if (Settings.ApplyPairIterations > 0)
+		const int32 NumPairIts = (SolverType == EConstraintSolverType::QuasiPbd) ? 1 : Settings.ApplyPushOutPairIterations;
+		if (NumPairIts > 0)
 		{
 			// UpdateConstraintProjection(It, NumIts, SolverData);
 			for (int32 ConstraintIndex : SolverData.GetConstraintIndices(ContainerId))
 			{
-				NumActive += ApplyPhase2Single(Dt, ConstraintIndex, Settings.ApplyPairIterations, It, NumIts);
+				NumActive += ApplyPhase2Single(Dt, ConstraintIndex, NumPairIts, It, NumIts);
 			}
 		}
 		return (NumActive > 0);
@@ -1113,6 +1115,21 @@ namespace Chaos
 		return true;
 	}
 
+	// @todo(chaos): ShockPropagation needs to handle the parent/child being in opposite order
+	FReal FPBDJointConstraints::CalculateShockPropagationInvMassScale(const FConstraintSolverBody& Body0, const FConstraintSolverBody& Body1, const FPBDJointSettings& JointSettings, const int32 It, const int32 NumIts) const
+	{
+		// Shock propagation is only enabled for the last iteration, and only for the QPBD solver.
+		// The standard PBD solver runs projection in the second solver phase which is mostly the same thing.
+		if ((It >= (NumIts - Settings.NumShockPropagationIterations)) && (SolverType == EConstraintSolverType::QuasiPbd))
+		{
+			if (Body0.IsDynamic() && Body1.IsDynamic())
+			{
+				return FPBDJointUtilities::GetShockPropagationInvMassScale(Settings, JointSettings);
+			}
+		}
+		return FReal(1);
+	}
+
 	// This position solver iterates over each of the inner constraints (position, twist, swing) and solves them independently.
 	// This will converge slowly in some cases, particularly where resolving angular constraints violates position constraints and vice versa.
 	bool FPBDJointConstraints::ApplyPhase1Single(const FReal Dt, const int32 ConstraintIndex, const int32 NumPairIts, const int32 It, const int32 NumIts)
@@ -1147,9 +1164,8 @@ namespace Chaos
 			const bool bWasActive = Solver.GetIsActive();
 			Solver.Update(Dt, Settings, JointSettings);
 
-			// Set parent inverser mass scale based on current shock propagation state
-			// @todo(chaos): needs to handle the parent/child being in opposite order
-			const FReal ShockPropagationInvMassScale = FPBDJointUtilities::GetShockPropagationInvMassScale(SolverType, It, NumIts, Settings, JointSettings);
+			// Set parent inverse mass scale based on current shock propagation state
+			const FReal ShockPropagationInvMassScale = CalculateShockPropagationInvMassScale(Solver.Body0(), Solver.Body1(), JointSettings, It, NumIts);
 			Solver.SetInvMassScales(ShockPropagationInvMassScale, FReal(1), Dt);
 
 			if (!bWasActive && !Solver.GetIsActive() && bChaos_Joint_EarlyOut_Enabled)
@@ -1208,24 +1224,43 @@ namespace Chaos
 			const bool bWasActive = Solver.GetIsActive();
 			Solver.Update(Dt, Settings, JointSettings);
 
-			// Set parent inverser mass scale based on current shock propagation state
-			// @todo(chaos): needs to handle the parent/child being in opposite order
-			const FReal ShockPropagationInvMassScale = FPBDJointUtilities::GetShockPropagationInvMassScale(SolverType, It, NumIts, Settings, JointSettings);
-			Solver.SetInvMassScales(ShockPropagationInvMassScale, FReal(1));
-
 			if (!bWasActive && !Solver.GetIsActive() && bChaos_Joint_EarlyOut_Enabled)
 			{
 				return false;
 			}
 
+			// Set parent inverse mass scale based on current shock propagation state
+			const FReal ShockPropagationInvMassScale = CalculateShockPropagationInvMassScale(Solver.Body0(), Solver.Body1(), JointSettings, It, NumIts);
+
 			const FReal IterationStiffness = CalculateIterationStiffness(It, NumIts);
 			for (int32 PairIt = 0; PairIt < NumPairIts; ++PairIt)
 			{
 				UE_LOG(LogChaosJoint, VeryVerbose, TEXT("  Pair Iteration %d / %d"), PairIt, NumPairIts);
 
+				if (SolverType == EConstraintSolverType::StandardPbd)
+				{
+					Solver.UpdateMasses(ShockPropagationInvMassScale, FReal(1));
+				}
+
+				Solver.SetInvMassScales(ShockPropagationInvMassScale, FReal(1));
+
 				// This is the same position solver for all SolverTypes (which makes it wrong for the GbfPbd solver which should be a pbd-esque velocity solve)
 				Solver.ApplyConstraints(Dt, IterationStiffness, Settings, JointSettings);
 
+				if (SolverType == EConstraintSolverType::StandardPbd)
+				{
+					if (Solver.Body0().IsDynamic())
+					{
+						Solver.Body0().SolverBody().ApplyCorrections();
+						Solver.Body0().UpdateRotationDependentState();
+					}
+					if (Solver.Body1().IsDynamic())
+					{
+						Solver.Body1().SolverBody().ApplyCorrections();
+						Solver.Body1().UpdateRotationDependentState();
+					}
+				}
+
 				if (!Solver.GetIsActive() && bChaos_Joint_EarlyOut_Enabled)
 				{
 					break;
@@ -1282,9 +1317,8 @@ namespace Chaos
 			const bool bWasActive = Solver.GetIsActive();
 			Solver.Update(Dt, Settings, JointSettings);
 			
-			// Set parent inverser mass scale based on current shock propagation state
-			// @todo(chaos): needs to handle the parent/child being in opposite order
-			const FReal ShockPropagationInvMassScale = FPBDJointUtilities::GetShockPropagationInvMassScale(SolverType, It, NumIts, Settings, JointSettings);
+			// Set parent inverse mass scale based on current shock propagation state
+			const FReal ShockPropagationInvMassScale = CalculateShockPropagationInvMassScale(Solver.Body0(), Solver.Body1(), JointSettings, It, NumIts);
 			Solver.SetInvMassScales(ShockPropagationInvMassScale, FReal(1), Dt);
 
 			// For quasipbd, we may still need to stabilize velocities and solve for restitution.
@@ -1340,17 +1374,15 @@ namespace Chaos
 			const bool bWasActive = Solver.GetIsActive();
 			Solver.Update(Dt, Settings, JointSettings);
 
-			// Set parent inverser mass scale based on current shock propagation state
-			// @todo(chaos): needs to handle the parent/child being in opposite order
-			const FReal ShockPropagationInvMassScale = FPBDJointUtilities::GetShockPropagationInvMassScale(SolverType, It, NumIts, Settings, JointSettings);
-			Solver.SetInvMassScales(ShockPropagationInvMassScale, FReal(1));
-
 			// For quasipbd, we may still need to stabilize velocities and solve for restitution.
 			if (!bWasActive && !Solver.GetIsActive() && SolverType != EConstraintSolverType::QuasiPbd && bChaos_Joint_EarlyOut_Enabled)
 			{
 				return false;
 			}
 
+			// Set parent inverse mass scale based on current shock propagation state
+			const FReal ShockPropagationInvMassScale = CalculateShockPropagationInvMassScale(Solver.Body0(), Solver.Body1(), JointSettings, It, NumIts);
+
 			const FReal IterationStiffness = CalculateIterationStiffness(It, NumIts);
 			for (int32 PairIt = 0; PairIt < NumPairIts; ++PairIt)
 			{
@@ -1360,9 +1392,11 @@ namespace Chaos
 					break;
 				case EConstraintSolverType::GbfPbd:
 				case EConstraintSolverType::StandardPbd:
+					Solver.UpdateMasses(ShockPropagationInvMassScale, FReal(1));
 					Solver.ApplyProjections(Dt, IterationStiffness, Settings, JointSettings);
 					break;
 				case EConstraintSolverType::QuasiPbd:
+					Solver.SetInvMassScales(ShockPropagationInvMassScale, FReal(1));
 					Solver.ApplyVelocityConstraints(Dt, IterationStiffness, Settings, JointSettings);
 					break;
 				}
diff --git a/Engine/Source/Runtime/Experimental/Chaos/Private/PBDRigidsSolver.cpp b/Engine/Source/Runtime/Experimental/Chaos/Private/PBDRigidsSolver.cpp
--- a/Engine/Source/Runtime/Experimental/Chaos/Private/PBDRigidsSolver.cpp
+++ b/Engine/Source/Runtime/Experimental/Chaos/Private/PBDRigidsSolver.cpp
@@ -235,6 +235,9 @@ namespace Chaos
 		int32 ChaosSolverCollisionUseManifolds = 1;
 		FAutoConsoleVariableRef CVarChaosSolverCollisionUseManifolds(TEXT("p.Chaos.Solver.Collision.UseManifolds"), ChaosSolverCollisionUseManifolds, TEXT("Enable/Disable use of manifolds in collision."));
 
+		// Allow manifolds to be reused between ticks
+		int32 ChaosSolverCollisionAllowManifoldUpdate = 1;
+		FAutoConsoleVariableRef CVarChaosSolverCollisionAllowManifoldUpdate(TEXT("p.Chaos.Solver.Collision.AllowManifoldUpdate"), ChaosSolverCollisionAllowManifoldUpdate, TEXT("Enable/Disable reuse of manifolds between ticks (for small movement)."));
 
 		// Joint cvars
 		float ChaosSolverJointMinSolverStiffness = 1.0f;
@@ -807,6 +810,8 @@ namespace Chaos
 	{
 		MEvolution->GetCollisionDetector().GetNarrowPhase().GetContext().bDeferUpdate = (ChaosSolverCollisionDeferNarrowPhase != 0);
 		MEvolution->GetCollisionDetector().GetNarrowPhase().GetContext().bAllowManifolds = (ChaosSolverCollisionUseManifolds != 0);
+		MEvolution->GetCollisionDetector().GetNarrowPhase().GetContext().bAllowManifoldReuse = (ChaosSolverCollisionAllowManifoldUpdate != 0);
+
 		
 		FPBDJointSolverSettings JointsSettings = MEvolution->GetJointConstraints().GetSettings();
 		JointsSettings.MinSolverStiffness = ChaosSolverJointMinSolverStiffness;
diff --git a/Engine/Source/Runtime/Experimental/Chaos/Public/Chaos/Collision/CollisionContext.h b/Engine/Source/Runtime/Experimental/Chaos/Public/Chaos/Collision/CollisionContext.h
--- a/Engine/Source/Runtime/Experimental/Chaos/Public/Chaos/Collision/CollisionContext.h
+++ b/Engine/Source/Runtime/Experimental/Chaos/Public/Chaos/Collision/CollisionContext.h
@@ -17,6 +17,7 @@ namespace Chaos
 			: bFilteringEnabled(true)
 			, bDeferUpdate(false)
 			, bAllowManifolds(false)
+			, bAllowManifoldReuse(false)
 			, bForceDisableCCD(false)
 			, CollisionAllocator(nullptr)
 		{
@@ -30,9 +31,12 @@ namespace Chaos
 		// which could be undesirable in some cases (destruction?).
 		bool bDeferUpdate;
 
-		// Whether to use manifolds where supported [default: false]
+		// Whether to use one-shot manifolds where supported [default: false]
 		bool bAllowManifolds;
 
+		// Whether we can reuse manifolds between frames if contacts have not moved far [default: false]
+		bool bAllowManifoldReuse;
+
 		// Force disable CCD
 		bool bForceDisableCCD;
 
diff --git a/Engine/Source/Runtime/Experimental/Chaos/Public/Chaos/Collision/PBDCollisionConstraint.h b/Engine/Source/Runtime/Experimental/Chaos/Public/Chaos/Collision/PBDCollisionConstraint.h
--- a/Engine/Source/Runtime/Experimental/Chaos/Public/Chaos/Collision/PBDCollisionConstraint.h
+++ b/Engine/Source/Runtime/Experimental/Chaos/Public/Chaos/Collision/PBDCollisionConstraint.h
@@ -370,8 +370,14 @@ namespace Chaos
 		FReal GetCullDistance() const { return CullDistance; }
 		void SetCullDistance(FReal InCullDistance) { CullDistance = InCullDistance; }
 
+		// Whether we are using manifolds (either one-shot or incremental)
 		bool GetUseManifold() const { return Flags.bUseManifold; }
-		
+
+		// Whether we can use incremental manifolds (updated each iteration)
+		bool GetUseIncrementalManifold() const { return Flags.bUseIncrementalManifold; }
+
+		// Whether we run collision detection every iteration (true if we are not using one shot manifolds)
+		// NOTE: This is initially set based on whether are allowing incremental manifolds
 		bool GetUseIncrementalCollisionDetection() const { return !Flags.bUseManifold || Flags.bUseIncrementalManifold; }
 
 		/**
@@ -401,6 +407,9 @@ namespace Chaos
 				{
 					ClosestManifoldPointIndex = ManifoldPointIndex;
 				}
+
+				// If we can use a one-shot manifold with this shape pair, disable the incremental flag
+				// to prevent calling collision detection again
 				Flags.bUseIncrementalManifold = false;
 			}
 		}
diff --git a/Engine/Source/Runtime/Experimental/Chaos/Public/Chaos/Collision/ParticlePairMidPhase.h b/Engine/Source/Runtime/Experimental/Chaos/Public/Chaos/Collision/ParticlePairMidPhase.h
--- a/Engine/Source/Runtime/Experimental/Chaos/Public/Chaos/Collision/ParticlePairMidPhase.h
+++ b/Engine/Source/Runtime/Experimental/Chaos/Public/Chaos/Collision/ParticlePairMidPhase.h
@@ -67,8 +67,8 @@ namespace Chaos
 		*/
 		int32 GenerateCollision(
 			const FReal CullDistance,
-			const bool bUseManifold,
-			const FReal Dt);
+			const FReal Dt,
+			FCollisionContext& Context);
 
 		/**
 		 * @brief Generate a SweptConstraint as long as AABBs overlap
@@ -76,8 +76,8 @@ namespace Chaos
 		*/
 		int32 GenerateCollisionCCD(
 			const FReal CullDistance,
-			const bool bUseManifold,
-			const FReal Dt);
+			const FReal Dt,
+			FCollisionContext& Context);
 
 		/**
 		 * @brief Reactivate the constraint
@@ -95,13 +95,13 @@ namespace Chaos
 	private:
 		int32 GenerateCollisionImpl(
 			const FReal CullDistance, 
-			const bool bUseManifold,
-			const FReal Dt);
+			const FReal Dt,
+			FCollisionContext& Context);
 			
 		int32 GenerateCollisionCCDImpl(
 			const FReal CullDistance, 
-			const bool bUseManifold,
-			const FReal Dt);
+			const FReal Dt,
+			FCollisionContext& Context);
 
 		/**
 		 * @brief Whether the two shapes are separated by less than CullDistance (i.e., we should run the narrow phase).
@@ -112,7 +112,7 @@ namespace Chaos
 		/**
 		 * @brief Create a constraint
 		*/
-		void CreateConstraint(const FReal CullDistance, const bool bUseManifold);
+		void CreateConstraint(const FReal CullDistance, FCollisionContext& Context);
 
 		FParticlePairMidPhase& MidPhase;
 		TUniquePtr<FPBDCollisionConstraint> Constraint;
@@ -163,7 +163,6 @@ namespace Chaos
 		*/
 		int32 GenerateCollisions(
 			const FReal CullDistance,
-			const bool bUseManifold,
 			const FReal Dt,
 			FCollisionContext& Context);
 
diff --git a/Engine/Source/Runtime/Experimental/Chaos/Public/Chaos/Evolution/PBDMinEvolution.h b/Engine/Source/Runtime/Experimental/Chaos/Public/Chaos/Evolution/PBDMinEvolution.h
--- a/Engine/Source/Runtime/Experimental/Chaos/Public/Chaos/Evolution/PBDMinEvolution.h
+++ b/Engine/Source/Runtime/Experimental/Chaos/Public/Chaos/Evolution/PBDMinEvolution.h
@@ -3,6 +3,7 @@
 
 #include "Chaos/Core.h"
 #include "Chaos/ArrayCollectionArray.h"
+#include "Chaos/Collision/CollisionApplyType.h"
 #include "Chaos/Evolution/SimulationSpace.h"
 #include "Chaos/Evolution/SolverDatas.h"
 #include "Chaos/ParticleHandleFwd.h"
@@ -39,6 +40,11 @@ namespace Chaos
 		void Advance(const FReal StepDt, const int32 NumSteps, const FReal RewindDt);
 		void AdvanceOneTimeStep(const FReal Dt, const FReal StepFraction);
 
+		void SetSolverType(const EConstraintSolverType InSolverType)
+		{
+			SolverType = InSolverType;
+		}
+
 		void SetNumIterations(const int32 NumIts)
 		{
 			NumApplyIterations = NumIts;
@@ -123,6 +129,7 @@ namespace Chaos
 		TArray<FSimpleConstraintRule*> PrioritizedConstraintRules;
 		FPBDIslandSolverData SolverData;
 
+		EConstraintSolverType SolverType;
 		int32 NumApplyIterations;
 		int32 NumApplyPushOutIterations;
 		FReal BoundsExtension;
diff --git a/Engine/Source/Runtime/Experimental/Chaos/Public/Chaos/Evolution/SolverBody.h b/Engine/Source/Runtime/Experimental/Chaos/Public/Chaos/Evolution/SolverBody.h
--- a/Engine/Source/Runtime/Experimental/Chaos/Public/Chaos/Evolution/SolverBody.h
+++ b/Engine/Source/Runtime/Experimental/Chaos/Public/Chaos/Evolution/SolverBody.h
@@ -5,16 +5,30 @@
 #include "Math/Quat.h"
 
 
-#if DO_CHECK && !(UE_BUILD_SHIPPING || UE_BUILD_TEST)
-#define ChaosSolverEnsure( condition )		ensure( condition )
-#else
-#define ChaosSolverEnsure(...)
+// Set to 1 to enable extra NaN tests in the constraint solver
+// NOTE: These will cause slowness!
+#ifndef CHAOS_CONSTRAINTSOLVER_NAN_DIAGNOSTIC
+#define CHAOS_CONSTRAINTSOLVER_NAN_DIAGNOSTIC ((DO_CHECK && UE_BUILD_DEBUG) || ENABLE_NAN_DIAGNOSTIC)
 #endif
 
-
 // Set to 1 to force single precision in the constraint solver, even if the default numeric type is double
 #define CHAOS_CONSTRAINTSOLVER_LOWPRECISION 1
 
+
+// Add some NaN checking when CHAOS_CONSTRAINTSOLVER_NAN_DIAGNOSTIC is defined
+#if CHAOS_CONSTRAINTSOLVER_NAN_DIAGNOSTIC
+inline void ChaosSolverCheckNaN(const Chaos::FRealSingle& V) { ensure(!FMath::IsNaN(V)); }
+inline void ChaosSolverCheckNaN(const Chaos::TVec3<Chaos::FRealSingle>& V) { ensure(!V.ContainsNaN()); }
+inline void ChaosSolverCheckNaN(const Chaos::TRotation3<Chaos::FRealSingle>& V) { ensure(!V.ContainsNaN()); }
+
+inline void ChaosSolverCheckNaN(const Chaos::FRealDouble& V) { ensure(!FMath::IsNaN(V)); }
+inline void ChaosSolverCheckNaN(const Chaos::TVec3<Chaos::FRealDouble>& V) { ensure(!V.ContainsNaN()); }
+inline void ChaosSolverCheckNaN(const Chaos::TRotation3<Chaos::FRealDouble>& V) { ensure(!V.ContainsNaN()); }
+#else
+#define ChaosSolverCheckNaN(...)
+#endif
+
+
 namespace Chaos
 {
 	// Set the math types used by the constraint solvers when single precision is acceptable.
@@ -31,12 +45,11 @@ namespace Chaos
 	using FSolverVec3 = TVec3<FSolverReal>;
 	using FSolverMatrix33 = TMatrix33<FSolverReal>;
 
-
 	class FSolverBody;
 	class FSolverBodyContainer;
 
 
-	// A pair of pointers so solver bodies
+	// A pair of pointers to solver bodies
 	// @note Pointers are only valid for the Constraint Solving phase of the tick
 	using FSolverBodyPtrPair = TVector<FSolverBody*, 2>;
 
@@ -49,7 +62,7 @@ namespace Chaos
 	 * This avoids the sqrt which is a massively dominating cost especially with doubles
 	 * when we do not have a fast reciproical sqrt (AVX2)
 	 *
-	 * This uses the first order Pade apprimant instead of a Taylor expansion
+	 * This uses the first order Pade approximation instead of a Taylor expansion
 	 * to get more accurate results for small quaternion deltas (i.e., where
 	 * Q.SizeSquared() is already near 1). 
 	 *
@@ -100,8 +113,8 @@ namespace Chaos
 			if (IsDynamic() && (Dt != FReal(0)))
 			{
 				const FSolverReal InvDt = FSolverReal(1) / FSolverReal(Dt);
-				State.V += FVec3(State.DP * InvDt);
-				SetW(State.W + State.DQ * InvDt);
+				SetV(State.V + FVec3(State.DP * InvDt));
+				SetW(State.W + FVec3(State.DQ * InvDt));
 			}
 		}
 
@@ -157,7 +170,7 @@ namespace Chaos
 		inline const FRotation3& R() const { return State.R; }
 		inline void SetR(const FRotation3& InR)
 		{
-			ChaosSolverEnsure(!InR.ContainsNaN());
+			ChaosSolverCheckNaN(InR);
 			State.R = InR;
 		}
 
@@ -176,8 +189,8 @@ namespace Chaos
 		*/
 		inline const FRotation3& Q() const { return State.Q; }
 		inline void SetQ(const FRotation3& InQ)
-		{
-			ChaosSolverEnsure(!InQ.ContainsNaN());
+		{ 
+			ChaosSolverCheckNaN(InQ);
 			State.Q = InQ;
 		}
 
@@ -185,7 +198,11 @@ namespace Chaos
 		 * @brief World-space center of mass velocity
 		*/
 		inline const FVec3& V() const { return State.V; }
-		inline void SetV(const FVec3& InV) { State.V = InV; }
+		inline void SetV(const FVec3& InV)
+		{ 
+			ChaosSolverCheckNaN(InV);
+			State.V = InV;
+		}
 
 		/**
 		 * @brief World-space center of mass angular velocity
@@ -193,7 +210,7 @@ namespace Chaos
 		inline const FVec3& W() const { return State.W; }
 		inline void SetW(const FVec3& InW)
 		{
-			ChaosSolverEnsure(!InW.ContainsNaN());
+			ChaosSolverCheckNaN(InW);
 			State.W = InW;
 		}
 
@@ -223,7 +240,7 @@ namespace Chaos
 		 * @brief World-space rotation after applying the net correction DQ()
 		 * @note Calculated on demand from Q() and DQ() (requires quaternion multiply and normalization)
 		*/
-		inline FRotation3 CorrectedQ() const { return IsDynamic() ? FRotation3::IntegrateRotationWithAngularVelocity(State.Q, FVec3(State.DQ), FReal(1)) : State.Q; }
+		inline FRotation3 CorrectedQ() const { return (IsDynamic() && !State.DQ.IsZero()) ? FRotation3::IntegrateRotationWithAngularVelocity(State.Q, FVec3(State.DQ), FReal(1)) : State.Q; }
 
 		/**
 		 * @brief Apply the accumulated position and rotation corrections to the predicted P and Q
@@ -276,6 +293,7 @@ namespace Chaos
 		*/
 		inline void ApplyPositionDelta(const FSolverVec3& DP)
 		{
+			ChaosSolverCheckNaN(DP);
 			State.DP += DP;
 		}
 
@@ -284,6 +302,7 @@ namespace Chaos
 		*/
 		inline void ApplyRotationDelta(const FSolverVec3& DR)
 		{
+			ChaosSolverCheckNaN(DR);
 			State.DQ += DR;
 		}
 
@@ -301,6 +320,7 @@ namespace Chaos
 		*/
 		inline void ApplyLinearVelocityDelta(const FSolverVec3& DV)
 		{
+			ChaosSolverCheckNaN(DV);
 			State.V += FVec3(DV);
 		}
 
@@ -309,6 +329,7 @@ namespace Chaos
 		*/
 		inline void ApplyAngularVelocityDelta(const FSolverVec3& DW)
 		{
+			ChaosSolverCheckNaN(DW);
 			SetW(State.W + DW);
 		}
 
diff --git a/Engine/Source/Runtime/Experimental/Chaos/Public/Chaos/Evolution/SolverBodyContainer.h b/Engine/Source/Runtime/Experimental/Chaos/Public/Chaos/Evolution/SolverBodyContainer.h
--- a/Engine/Source/Runtime/Experimental/Chaos/Public/Chaos/Evolution/SolverBodyContainer.h
+++ b/Engine/Source/Runtime/Experimental/Chaos/Public/Chaos/Evolution/SolverBodyContainer.h
@@ -173,6 +173,12 @@ namespace Chaos
 		// Recalculate the velocities of all bodies based on their transform deltas
 		void SetImplicitVelocities(FReal Dt);
 
+		// Apply accumulated transform deltas to the body transforms
+		void ApplyCorrections();
+
+		// Can be called after ApplyCorrections to update inertia to match the current transform
+		void UpdateRotationDependentState();
+
 	private:
 		int32 AddParticle(FGenericParticleHandle InParticle);
 
diff --git a/Engine/Source/Runtime/Experimental/Chaos/Public/Chaos/Joint/PBDJointSolverGaussSeidel.h b/Engine/Source/Runtime/Experimental/Chaos/Public/Chaos/Joint/PBDJointSolverGaussSeidel.h
--- a/Engine/Source/Runtime/Experimental/Chaos/Public/Chaos/Joint/PBDJointSolverGaussSeidel.h
+++ b/Engine/Source/Runtime/Experimental/Chaos/Public/Chaos/Joint/PBDJointSolverGaussSeidel.h
@@ -158,6 +158,10 @@ namespace Chaos
 			const FPBDJointSolverSettings& SolverSettings,
 			const FPBDJointSettings& JointSettings);
 
+		void UpdateMasses(
+			const FReal InvMassScale0,
+			const FReal InvMassScale1);
+
 		// Run the position solve for the constraints
 		void ApplyConstraints(
 			const FReal Dt,
@@ -436,6 +440,19 @@ namespace Chaos
 			const FReal DeltaPos,
 			const FReal DeltaVel);
 
+		void ApplyPositionProjection(
+			const FReal Dt,
+			const FPBDJointSolverSettings& SolverSettings,
+			const FPBDJointSettings& JointSettings,
+			FVec3& DP1,
+			FVec3& DR1);
+
+		void ApplyRotationProjection(
+			const FReal Dt,
+			const FPBDJointSolverSettings& SolverSettings,
+			const FPBDJointSettings& JointSettings,
+			FVec3& DP1,
+			FVec3& DR1);
 
 		void ApplyPointProjection(
 			const FReal Dt,
diff --git a/Engine/Source/Runtime/Experimental/Chaos/Public/Chaos/PBDJointConstraintUtilities.h b/Engine/Source/Runtime/Experimental/Chaos/Public/Chaos/PBDJointConstraintUtilities.h
--- a/Engine/Source/Runtime/Experimental/Chaos/Public/Chaos/PBDJointConstraintUtilities.h
+++ b/Engine/Source/Runtime/Experimental/Chaos/Public/Chaos/PBDJointConstraintUtilities.h
@@ -268,9 +268,6 @@ namespace Chaos
 			const FPBDJointSettings& JointSettings);
 
 		static FReal GetShockPropagationInvMassScale(
-			EConstraintSolverType SolverType,
-			int32 It,
-			int32 NumIts,
 			const FPBDJointSolverSettings& SolverSettings,
 			const FPBDJointSettings& JointSettings);
 
diff --git a/Engine/Source/Runtime/Experimental/Chaos/Public/Chaos/PBDJointConstraints.h b/Engine/Source/Runtime/Experimental/Chaos/Public/Chaos/PBDJointConstraints.h
--- a/Engine/Source/Runtime/Experimental/Chaos/Public/Chaos/PBDJointConstraints.h
+++ b/Engine/Source/Runtime/Experimental/Chaos/Public/Chaos/PBDJointConstraints.h
@@ -294,6 +294,7 @@ namespace Chaos
 		friend class FPBDJointConstraintHandle;
 
 		FReal CalculateIterationStiffness(int32 It, int32 NumIts) const;
+		FReal CalculateShockPropagationInvMassScale(const FConstraintSolverBody& Body0, const FConstraintSolverBody& Body1, const FPBDJointSettings& JointSettings, const int32 It, const int32 NumIts) const;
 
 		void GetConstrainedParticleIndices(const int32 ConstraintIndex, int32& Index0, int32& Index1) const;
 		void CalculateConstraintSpace(int32 ConstraintIndex, FVec3& OutX0, FMatrix33& OutR0, FVec3& OutX1, FMatrix33& OutR1) const;
