Fix cylindrical joint and slerp drive instabilities + improve linear constraints setup + fix switching joint axis when forces is applied

diff --git a/Engine/Source/Runtime/Experimental/Chaos/Private/Chaos/Joint/JointConstraintsCVars.cpp b/Engine/Source/Runtime/Experimental/Chaos/Private/Chaos/Joint/JointConstraintsCVars.cpp
--- a/Engine/Source/Runtime/Experimental/Chaos/Private/Chaos/Joint/JointConstraintsCVars.cpp
+++ b/Engine/Source/Runtime/Experimental/Chaos/Private/Chaos/Joint/JointConstraintsCVars.cpp
@@ -29,6 +29,3 @@ FAutoConsoleVariableRef CVarChaosJointLinearVelocityThresholdToApplyRestitution(
 float Chaos_Joint_AngularVelocityThresholdToApplyRestitution = 1e-2f;
 FAutoConsoleVariableRef CVarChaosJointAngularVelocityThresholdToApplyRestitution(TEXT("p.Chaos.Joint.AngularVelocityThresholdToApplyRestitution"), Chaos_Joint_AngularVelocityThresholdToApplyRestitution, TEXT("Apply restitution only if initial velocity is higher than this threshold (used in Quasipbd)"));
 
-bool bChaos_Joint_StablePointConstraint = false;
-FAutoConsoleVariableRef CVarChaosJointStablePointConstraint(TEXT("p.Chaos.Joint.StablePointConstraint"), bChaos_Joint_StablePointConstraint, TEXT("Improve stability for the point constraint (increase computation time)"));
-
diff --git a/Engine/Source/Runtime/Experimental/Chaos/Private/Chaos/Joint/JointConstraintsCVars.h b/Engine/Source/Runtime/Experimental/Chaos/Private/Chaos/Joint/JointConstraintsCVars.h
--- a/Engine/Source/Runtime/Experimental/Chaos/Private/Chaos/Joint/JointConstraintsCVars.h
+++ b/Engine/Source/Runtime/Experimental/Chaos/Private/Chaos/Joint/JointConstraintsCVars.h
@@ -27,6 +27,3 @@ extern float Chaos_Joint_LinearVelocityThresholdToApplyRestitution;
 extern float Chaos_Joint_AngularVelocityThresholdToApplyRestitution;
 
 extern bool bChaos_Joint_UseCachedSolver;
-
-extern bool bChaos_Joint_StablePointConstraint;
-
diff --git a/Engine/Source/Runtime/Experimental/Chaos/Private/Chaos/Joint/PBDJointCachedSolverGaussSeidel.cpp b/Engine/Source/Runtime/Experimental/Chaos/Private/Chaos/Joint/PBDJointCachedSolverGaussSeidel.cpp
--- a/Engine/Source/Runtime/Experimental/Chaos/Private/Chaos/Joint/PBDJointCachedSolverGaussSeidel.cpp
+++ b/Engine/Source/Runtime/Experimental/Chaos/Private/Chaos/Joint/PBDJointCachedSolverGaussSeidel.cpp
@@ -137,10 +137,6 @@ void FPBDJointCachedSolver::SetInvMassScales(const FReal InvMScale0, const FReal
 	}
 	if(bNeedsUpdate)
 	{
-		if(PointConstraints.bValidDatas)
-		{
-			InitPointPositionMass();
-		}
 		for(int32 ConstraintIndex = 0; ConstraintIndex < 3; ++ConstraintIndex)
 		{
 			if(PositionConstraints.bValidDatas[ConstraintIndex])
@@ -337,7 +333,6 @@ void FPBDJointCachedSolver::InitPositionConstraints(
 	PositionConstraints.bValidDatas[0] = false;
 	PositionConstraints.bValidDatas[1] = false;
 	PositionConstraints.bValidDatas[2] = false;
-	PointConstraints.bValidDatas = false;
 
 	TVec3<bool> bLinearLocked, bLinearLimited;
 	if(!ExtractLinearMotion(JointSettings, bLinearLocked, bLinearLimited))
@@ -353,162 +348,58 @@ void FPBDJointCachedSolver::InitPositionConstraints(
 			FPBDJointUtilities::GetSoftLinearDamping(SolverSettings, JointSettings),
 			FPBDJointUtilities::GetLinearStiffness(SolverSettings, JointSettings));
 	}
-	PointConstraints.InitDatas(FPBDJointUtilities::GetLinearStiffness(SolverSettings, JointSettings));
 
 	const TVec3<EJointMotionType>& LinearMotion = JointSettings.LinearMotionTypes;
 
-	if (bLinearLocked[0] && bLinearLocked[1] && bLinearLocked[2])
-	{
-		// Hard point constraint (most common case)
-		InitPointPositionConstraint(Dt);
-	}
-	else if (bLinearLimited[0] && bLinearLimited[1] && bLinearLimited[2])
-	{
-		// Spherical constraint
-		InitSphericalPositionConstraint(JointSettings, Dt);
-	}
-	else if (bLinearLocked[1] && bLinearLocked[2] && !bLinearLocked[0])
-	{
-		// Line constraint along X axis
-		InitCylindricalPositionConstraint(JointSettings, Dt, 0, LinearMotion[0], EJointMotionType::Locked);
-	}
-	else if (bLinearLocked[0] && bLinearLocked[2] && !bLinearLocked[1])
-	{
-		// Line constraint along Y axis
-		InitCylindricalPositionConstraint(JointSettings, Dt, 1, LinearMotion[1], EJointMotionType::Locked);
-	}
-	else if (bLinearLocked[0] && bLinearLocked[1] && !bLinearLocked[2])
-	{
-		// Line constraint along Z axis
-		InitCylindricalPositionConstraint(JointSettings, Dt, 2, LinearMotion[2], EJointMotionType::Locked);
-	}
-	else if (bLinearLimited[1] && bLinearLimited[2] && !bLinearLimited[0])
-	{
-		// Cylindrical constraint along X axis
-		InitCylindricalPositionConstraint(JointSettings, Dt, 0, LinearMotion[0], EJointMotionType::Limited);
-	}
-	else if (bLinearLimited[0] && bLinearLimited[2] && !bLinearLimited[1])
-	{
-		// Cylindrical constraint along Y axis
-		InitCylindricalPositionConstraint(JointSettings, Dt, 1, LinearMotion[1], EJointMotionType::Limited);
-	}
-	else if (bLinearLimited[0] && bLinearLimited[1] && !bLinearLimited[2])
-	{
-		// Cylindrical constraint along Z axis
-		InitCylindricalPositionConstraint(JointSettings, Dt, 2, LinearMotion[2], EJointMotionType::Limited);
-	}
-	else if (bLinearLocked[0] || bLinearLimited[0])
-	{
-		// Planar constraint along X axis
-		InitPlanarPositionConstraint(JointSettings, Dt, 0, LinearMotion[0]);
-	}
-	else if (bLinearLocked[1] || bLinearLimited[1])
-	{
-		// Planar constraint along Y axis
-		InitPlanarPositionConstraint(JointSettings, Dt, 1, LinearMotion[1]);
-	}
-	else if (bLinearLocked[2] || bLinearLimited[2])
-	{
-		// Planar constraint along Z axis
-		InitPlanarPositionConstraint(JointSettings, Dt, 2, LinearMotion[2]);
-	}
-}
-
-void FPBDJointCachedSolver::InitPositionDatasMass(
-	FAxisConstraintDatas& PositionDatas, 
-	const int32 ConstraintIndex,
-	const FReal Dt)
-{
-	const FVec3 AngularAxis0 = FVec3::CrossProduct(PositionDatas.ConstraintArms[ConstraintIndex][0], PositionDatas.ConstraintAxis[ConstraintIndex]);
-	const FVec3 AngularAxis1 = FVec3::CrossProduct(PositionDatas.ConstraintArms[ConstraintIndex][1], PositionDatas.ConstraintAxis[ConstraintIndex]);
-	const FVec3 IA0 = Utilities::Multiply(InvI(0), AngularAxis0);
-	const FVec3 IA1 = Utilities::Multiply(InvI(1), AngularAxis1);
-	const FReal II0 = FVec3::DotProduct(AngularAxis0, IA0);
-	const FReal II1 = FVec3::DotProduct(AngularAxis1, IA1);
-
-	PositionDatas.UpdateMass(ConstraintIndex, IA0, IA1, InvM(0) + II0 + InvM(1) + II1, Dt);
-}
-
-void FPBDJointCachedSolver::InitPositionConstraintDatas(
-	const int32 ConstraintIndex,
-	const FVec3& ConstraintAxis,
-	const FReal& ConstraintDelta,
-	const FReal ConstraintRestitution,
-	const FReal Dt,
-	const FReal ConstraintLimit,
-	const EJointMotionType JointType)
-{
-	const FVec3 LocalAxis = (ConstraintDelta < 0.0f) ? -ConstraintAxis : ConstraintAxis;
-	const FReal LocalDelta = (ConstraintDelta < 0.0f) ? -ConstraintDelta : ConstraintDelta;
-
-	PositionConstraints.ConstraintLimits[ConstraintIndex] = (JointType == EJointMotionType::Locked) ? 0.0f :
-		FMath::Max(ConstraintLimit - GetLinearConstraintPadding(ConstraintIndex), (FReal)0.);
-	PositionConstraints.MotionType[ConstraintIndex] = JointType;
-
-	// This uses ConnectorXs[1] rather than ConnectorXs[0] for stability reasons. It is more stable this way when ConnectorXs[0] and ConnectorXs[1] are far apart 
-	const FVec3 ConstraintArm0 = ConnectorXs[1] - LocalDelta * LocalAxis - P(0);
-	const FVec3 ConstraintArm1 = ConnectorXs[1] - P(1);
-
-	PositionConstraints.UpdateDatas(ConstraintIndex, LocalAxis, LocalDelta, ConstraintRestitution, true, ConstraintArm0, ConstraintArm1);
-
-	const FVec3 CV0 = V(0) + FVec3::CrossProduct(W(0), ConstraintArm0);
-	const FVec3 CV1 = V(1) + FVec3::CrossProduct(W(1), ConstraintArm1);
-	const FVec3 CV = CV1 - CV0;
 
-	InitConstraintAxisLinearVelocities[ConstraintIndex] = FVec3::DotProduct(CV, LocalAxis);
+	const FVec3 ConstraintArm0Limited = ConnectorXs[1] - P(0);
+	const FVec3 ConstraintArm1Limited = ConnectorXs[1] - P(1);
+	
+	FVec3 ConstraintArm0Locked = ConstraintArm0Limited;
+	FVec3 ConstraintArm1Locked = ConstraintArm1Limited;
 
-	InitPositionDatasMass(PositionConstraints, ConstraintIndex, Dt);
-}
+	const FVec3 DX = ConnectorXs[1] - ConnectorXs[0];
+	const FMatrix33 R0M = ConnectorRs[0].ToMatrix(); 
+	FVec3 CX = FVec3::ZeroVector;
 
-void FPBDJointCachedSolver::InitPointPositionMass()
-{
-	FMatrix33 IM = FMatrix33::Zero;
-	if(InvM(0) > 0.0)
-	{
-		IM += Utilities::ComputeJointFactorMatrix(PointConstraints.ConstraintArms[0], InvI(0), InvM(0));
-	}
-	if(InvM(1) > 0.0)
+	for(int32 ConstraintIndex = 0; ConstraintIndex < 3; ++ConstraintIndex)
 	{
-		IM += Utilities::ComputeJointFactorMatrix(PointConstraints.ConstraintArms[1], InvI(1), InvM(1));
-	}
-	IM = IM.Inverse();
-	PointConstraints.ConstraintIM = IM;
-	PointConstraints.ConstraintLambda = FVec3::ZeroVector;
-}
+		if(bLinearLocked[ConstraintIndex] || bLinearLimited[ConstraintIndex])
+		{
+			const FVec3 ConstraintAxis = R0M.GetAxis(ConstraintIndex);
+			CX[ConstraintIndex] = FVec3::DotProduct(DX, ConstraintAxis);
 
-void FPBDJointCachedSolver::InitPointPositionConstraint(const FReal Dt)
-{
-	if(bChaos_Joint_StablePointConstraint)
-	{
-		const FVec3 ConstraintArm0 = ConnectorXs[0] - P(0);
-		const FVec3 ConstraintArm1 = ConnectorXs[1] - P(1);
-		
-		PointConstraints.UpdateDatas(ConnectorXs[1] - ConnectorXs[0], ConstraintArm0, ConstraintArm1);
-		
-		InitPointPositionMass();
+			if(bLinearLocked[ConstraintIndex])
+			{
+				ConstraintArm0Locked -= ConstraintAxis * CX[ConstraintIndex];
+			}
+		}
 	}
-	else
+	
+	for(int32 ConstraintIndex = 0; ConstraintIndex < 3; ++ConstraintIndex)
 	{
-		const FVec3 ConstraintArm0 = ConnectorXs[0] - P(0);
-		const FVec3 ConstraintArm1 = ConnectorXs[1] - P(1);
-
-		const FMatrix33 R0M = ConnectorRs[0].ToMatrix(); 
-		const FVec3 CX = ConnectorRs[0].Inverse().RotateVector(ConnectorXs[1] - ConnectorXs[0]);
-
-		for(int32 ConstraintIndex = 0; ConstraintIndex < 3; ++ConstraintIndex)
+		if(bLinearLocked[ConstraintIndex] || bLinearLimited[ConstraintIndex])
 		{
-			const FVec3 ConstraintAxis = R0M.GetColumn(ConstraintIndex);
+			const FVec3 ConstraintAxis = R0M.GetAxis(ConstraintIndex);
+			PositionConstraints.MotionType[ConstraintIndex] = LinearMotion[ConstraintIndex];
 			
-			PositionConstraints.ConstraintLimits[ConstraintIndex] = 0.0;
-			PositionConstraints.MotionType[ConstraintIndex] = EJointMotionType::Locked;
-
-			PositionConstraints.UpdateDatas(ConstraintIndex, ConstraintAxis, CX[ConstraintIndex],
-				0.0, false, ConstraintArm0, ConstraintArm1);
-
-			const FVec3 CV0 = V(0) + FVec3::CrossProduct(W(0), ConstraintArm0);
-			const FVec3 CV1 = V(1) + FVec3::CrossProduct(W(1), ConstraintArm1);
+			if(bLinearLocked[ConstraintIndex])
+			{
+				PositionConstraints.ConstraintLimits[ConstraintIndex] = 0.0f;
+				PositionConstraints.UpdateDatas(ConstraintIndex, ConstraintAxis, CX[ConstraintIndex],
+					0.0, false, ConstraintArm0Locked, ConstraintArm1Locked);
+			}
+			else if(bLinearLimited[ConstraintIndex])
+			{
+				PositionConstraints.ConstraintLimits[ConstraintIndex] =
+					FMath::Max(JointSettings.LinearLimit - GetLinearConstraintPadding(ConstraintIndex), (FReal)0.);
+				PositionConstraints.UpdateDatas(ConstraintIndex, ConstraintAxis, CX[ConstraintIndex],
+					JointSettings.LinearRestitution, true, ConstraintArm0Limited, ConstraintArm1Limited);
+			}
+			const FVec3 CV0 = V(0) + FVec3::CrossProduct(W(0), PositionConstraints.ConstraintArms[ConstraintIndex][0]);
+			const FVec3 CV1 = V(1) + FVec3::CrossProduct(W(1), PositionConstraints.ConstraintArms[ConstraintIndex][1]);
 			const FVec3 CV = CV1 - CV0;
-
+			
 			InitConstraintAxisLinearVelocities[ConstraintIndex] = FVec3::DotProduct(CV, ConstraintAxis);
 
 			InitPositionDatasMass(PositionConstraints, ConstraintIndex, Dt);
@@ -516,47 +407,19 @@ void FPBDJointCachedSolver::InitPointPositionConstraint(const FReal Dt)
 	}
 }
 
-void FPBDJointCachedSolver::InitSphericalPositionConstraint(
-	const FPBDJointSettings& JointSettings,
+void FPBDJointCachedSolver::InitPositionDatasMass(
+	FAxisConstraintDatas& PositionDatas, 
+	const int32 ConstraintIndex,
 	const FReal Dt)
 {
-	FVec3 Axis;
-	FReal Delta;
-	FPBDJointUtilities::GetSphericalAxisDelta(ConnectorXs[0], ConnectorXs[1], Axis, Delta);
-
-	InitPositionConstraintDatas(SphericalPositionConstraintIndex, Axis,  Delta, JointSettings.LinearRestitution, Dt, JointSettings.LinearLimit, EJointMotionType::Limited);
-}
-
-void FPBDJointCachedSolver::InitCylindricalPositionConstraint(
-	const FPBDJointSettings& JointSettings,
-	const FReal Dt,
-	const int32 AxisIndex,
-	const EJointMotionType AxialMotion,
-	const EJointMotionType RadialMotion)
-{
-	FVec3 MainAxis, RadialAxis;
-	FReal AxialDelta, RadialDelta;
-	FPBDJointUtilities::GetCylindricalAxesDeltas(ConnectorRs[0], ConnectorXs[0], ConnectorXs[1],
-		AxisIndex, MainAxis, AxialDelta, RadialAxis, RadialDelta);
-
-	InitPositionConstraintDatas((int32)EJointCylindricalPositionConstraintType::Axial, MainAxis, AxialDelta,  JointSettings.LinearRestitution, Dt,
-		JointSettings.LinearLimit, AxialMotion);   
-	InitPositionConstraintDatas((int32)EJointCylindricalPositionConstraintType::Radial, RadialAxis, RadialDelta, JointSettings.LinearRestitution, Dt,
-		JointSettings.LinearLimit, RadialMotion);
-}
-
-void FPBDJointCachedSolver::InitPlanarPositionConstraint(
-	const FPBDJointSettings& JointSettings,
-	const FReal Dt,
-	const int32 AxisIndex,
-	const EJointMotionType AxialMotion)
-{
-	FVec3 Axis;
-	FReal Delta;
-	FPBDJointUtilities::GetPlanarAxisDelta(ConnectorRs[0], ConnectorXs[0], ConnectorXs[1], AxisIndex, Axis, Delta);
+	const FVec3 AngularAxis0 = FVec3::CrossProduct(PositionDatas.ConstraintArms[ConstraintIndex][0], PositionDatas.ConstraintAxis[ConstraintIndex]);
+	const FVec3 AngularAxis1 = FVec3::CrossProduct(PositionDatas.ConstraintArms[ConstraintIndex][1], PositionDatas.ConstraintAxis[ConstraintIndex]);
+	const FVec3 IA0 = Utilities::Multiply(InvI(0), AngularAxis0);
+	const FVec3 IA1 = Utilities::Multiply(InvI(1), AngularAxis1);
+	const FReal II0 = FVec3::DotProduct(AngularAxis0, IA0);
+	const FReal II1 = FVec3::DotProduct(AngularAxis1, IA1);
 
-	InitPositionConstraintDatas(PlanarPositionConstraintIndex, Axis, Delta, JointSettings.LinearRestitution, Dt,
-		JointSettings.LinearLimit, AxialMotion);
+	PositionDatas.UpdateMass(ConstraintIndex, IA0, IA1, InvM(0) + II0 + InvM(1) + II1, Dt);
 }
 
 /** APPLY POSITION CONSTRAINT *****************************************************************************************/
@@ -564,89 +427,33 @@ void FPBDJointCachedSolver::InitPlanarPositionConstraint(
 void FPBDJointCachedSolver::ApplyPositionConstraints(
 	const FReal Dt)
 {
-	if(PointConstraints.bValidDatas)
-	{
-		ApplyPointPositionConstraint();
-	}
-	else
-	{
-		for(int32 ConstraintIndex = 0; ConstraintIndex < 3; ++ConstraintIndex)
-		{
-			if(PositionConstraints.bValidDatas[ConstraintIndex])
-			{
-				ApplyAxisPositionConstraint(ConstraintIndex, Dt);
-			}
-		}
-	}
-}
-
-FVec3 FPBDJointCachedSolver::ComputePointCX(const bool bLinearCompute) 
-{
-	if(!bLinearCompute)
-	{
-		// Non-linear version
-		UpdateDerivedState();
-		return  ConnectorXs[1] - ConnectorXs[0];
-	}
-	else
-	{
-		// Linear version
-		return PointConstraints.ConstraintCX + Body(1).DP()  - Body(0).DP() +
-			FVec3::CrossProduct(Body(1).DQ(), PointConstraints.ConstraintArms[1]) -
-				FVec3::CrossProduct(Body(0).DQ(), PointConstraints.ConstraintArms[0]) ;
-	}
-}
-
-void FPBDJointCachedSolver::ApplyPointPositionConstraint()
-{
-	const FVec3 CX = ComputePointCX(false);
-
-	if (CX.SizeSquared() > PositionTolerance * PositionTolerance)
+	for(int32 ConstraintIndex = 0; ConstraintIndex < 3; ++ConstraintIndex)
 	{
-		// Calculate constraint correction
-		const FVec3 DX = SolverStiffness * PointConstraints.ConstraintStiffness * Utilities::Multiply(PointConstraints.ConstraintIM, CX);
-
-		// Apply constraint correction
-		if( Body(0).IsDynamic())
-		{
-			const FVec3 DP0 = InvM(0) * DX;
-			const FVec3 DR0 = Utilities::Multiply(InvI(0), FVec3::CrossProduct(PointConstraints.ConstraintArms[0], DX));
-    	
-			ApplyPositionDelta(0, DP0);
-			ApplyRotationDelta(0, DR0);
-		}
-		if( Body(1).IsDynamic())
+		if(PositionConstraints.bValidDatas[ConstraintIndex])
 		{
-			const FVec3 DP1 = -InvM(1) * DX;
-			const FVec3 DR1 = Utilities::Multiply(InvI(1), FVec3::CrossProduct(PointConstraints.ConstraintArms[1], -DX));
-
-			ApplyPositionDelta(1, DP1);
-			ApplyRotationDelta(1, DR1);
+			ApplyAxisPositionConstraint(ConstraintIndex, Dt);
 		}
-
-		NetLinearImpulse += DX;
-		PointConstraints.ConstraintLambda += DX;
-		++NumActiveConstraints;
 	}
 }
-
+	
 void FPBDJointCachedSolver::SolvePositionConstraintDelta(
 	const int32 ConstraintIndex, 
-	const FReal DeltaLambda)
+	const FReal DeltaLambda,
+	const FAxisConstraintDatas& ConstraintDatas)
 {
-	const FVec3 DX = PositionConstraints.ConstraintAxis[ConstraintIndex] * DeltaLambda;
+	const FVec3 DX = ConstraintDatas.ConstraintAxis[ConstraintIndex] * DeltaLambda;
 
 	if(Body(0).IsDynamic())
 	{
 		const FVec3 DP0 = InvM(0) * DX;
-		const FVec3 DR0 = PositionConstraints.ConstraintDRAxis[ConstraintIndex][0] * DeltaLambda;
+		const FVec3 DR0 = ConstraintDatas.ConstraintDRAxis[ConstraintIndex][0] * DeltaLambda;
 		ApplyPositionDelta(0,DP0);
 		ApplyRotationDelta(0,DR0);
 	}
 	if(Body(1).IsDynamic())
 	{
 		const FVec3 DP1 = -InvM(1) * DX;
-		const FVec3 DR1 = PositionConstraints.ConstraintDRAxis[ConstraintIndex][1] * DeltaLambda;
+		const FVec3 DR1 = ConstraintDatas.ConstraintDRAxis[ConstraintIndex][1] * DeltaLambda;
 		ApplyPositionDelta(1,DP1);
 		ApplyRotationDelta(1,DR1);
 	}
@@ -663,7 +470,7 @@ void FPBDJointCachedSolver::SolvePositionConstraintHard(
 		PositionConstraints.ConstraintHardIM[ConstraintIndex];
 
 	PositionConstraints.ConstraintLambda[ConstraintIndex] += DeltaLambda;
-	SolvePositionConstraintDelta(ConstraintIndex, DeltaLambda);
+	SolvePositionConstraintDelta(ConstraintIndex, DeltaLambda, PositionConstraints);
 }
 
 void FPBDJointCachedSolver::SolvePositionConstraintSoft(
@@ -684,7 +491,7 @@ void FPBDJointCachedSolver::SolvePositionConstraintSoft(
 		PositionConstraints.ConstraintSoftIM[ConstraintIndex];
 	PositionConstraints.ConstraintLambda[ConstraintIndex] += DeltaLambda;
 
-	SolvePositionConstraintDelta(ConstraintIndex, DeltaLambda);
+	SolvePositionConstraintDelta(ConstraintIndex, DeltaLambda, PositionConstraints);
 }
 
 void FPBDJointCachedSolver::ApplyAxisPositionConstraint(
@@ -731,44 +538,11 @@ void FPBDJointCachedSolver::ApplyAxisPositionConstraint(
 
 void FPBDJointCachedSolver::ApplyLinearVelocityConstraints()
 {
-	if(PointConstraints.bValidDatas)
-	{
-		ApplyPointVelocityConstraint();
-	}
-	else
-	{
-		for(int32 ConstraintIndex = 0; ConstraintIndex < 3; ++ConstraintIndex)
-		{
-			if(PositionConstraints.bValidDatas[ConstraintIndex])
-			{
-				ApplyAxisVelocityConstraint(ConstraintIndex);
-			}
-		}
-	}
-}
-
-void FPBDJointCachedSolver::ApplyPointVelocityConstraint()
-{
-	if (!NetLinearImpulse.IsNearlyZero())
+	for(int32 ConstraintIndex = 0; ConstraintIndex < 3; ++ConstraintIndex)
 	{
-		const FVec3 CV0 = V(0) + FVec3::CrossProduct(W(0), PointConstraints.ConstraintArms[0]);
-		const FVec3 CV1 = V(1) + FVec3::CrossProduct(W(1), PointConstraints.ConstraintArms[1]);
-    
-		const FVec3 CV = CV1 - CV0;
-		const FVec3 MDV = SolverStiffness * PointConstraints.ConstraintStiffness * Utilities::Multiply(PointConstraints.ConstraintIM, CV);
-
-		if(Body0().IsDynamic())
-		{
-			const FVec3 DV0 = InvM(0) * MDV;
-			const FVec3 DW0 = Utilities::Multiply(InvI(0), FVec3::CrossProduct(PointConstraints.ConstraintArms[0], MDV));
-			Body(0).ApplyVelocityDelta(DV0, DW0);
-		}
-
-		if(Body1().IsDynamic())
+		if(PositionConstraints.bValidDatas[ConstraintIndex])
 		{
-			const FVec3 DV1 = -InvM(1) * MDV;
-			const FVec3 DW1 = Utilities::Multiply(InvI(1), FVec3::CrossProduct(PointConstraints.ConstraintArms[1], -MDV));
-			Body(1).ApplyVelocityDelta(DV1, DW1);
+			ApplyAxisVelocityConstraint(ConstraintIndex);
 		}
 	}
 }
@@ -804,7 +578,7 @@ void FPBDJointCachedSolver::SolveLinearVelocityConstraint(
 
 void FPBDJointCachedSolver::ApplyAxisVelocityConstraint(const int32 ConstraintIndex)
 {
-	if(!NetLinearImpulse.IsNearlyZero() && PositionConstraints.ConstraintLambda[ConstraintIndex] > SMALL_NUMBER)
+	if(!NetLinearImpulse.IsNearlyZero() && (FMath::Abs(PositionConstraints.ConstraintLambda[ConstraintIndex]) > SMALL_NUMBER) )
 	{
 		FReal TargetVel = 0.0f;
 		if (PositionConstraints.MotionType[ConstraintIndex] == EJointMotionType::Limited && PositionConstraints.ConstraintRestitution[ConstraintIndex] != 0.0f)
@@ -1033,7 +807,7 @@ void FPBDJointCachedSolver::InitTwistConstraint(
 	FVec3 TwistAxis;
 	FReal TwistAngle;
 	FPBDJointUtilities::GetTwistAxisAngle(ConnectorRs[0], ConnectorRs[1], TwistAxis, TwistAngle);
-
+	
 	// Project the angle directly to avoid checking the limits during the solve.
 	CorrectAxisAngleConstraint(JointSettings,(int32)EJointAngularConstraintIndex::Twist, TwistAxis, TwistAngle);
 	InitRotationConstraintDatas( JointSettings, (int32)EJointAngularConstraintIndex::Twist, TwistAxis, TwistAngle, JointSettings.TwistRestitution, Dt, false);
@@ -1058,12 +832,17 @@ void FPBDJointCachedSolver::InitSingleLockedSwingConstraint(
 	const FReal Dt,
 	const EJointAngularConstraintIndex SwingConstraintIndex)
 {
-	// NOTE: SwingAxis is not normalized in this mode. It has length Sin(SwingAngle).
-	// Likewise, the SwingAngle is actually Sin(SwingAngle)
+	 //NOTE: SwingAxis is not normalized in this mode. It has length Sin(SwingAngle).
+	 //Likewise, the SwingAngle is actually Sin(SwingAngle)
+	// FVec3 SwingAxis;
+	// FReal SwingAngle;
+	// FPBDJointUtilities::GetLockedSwingAxisAngle(ConnectorRs[0], ConnectorRs[1], SwingConstraintIndex, SwingAxis, SwingAngle);
+	//SwingAxis.SafeNormalize();
+
+	// Using the locked swing axis angle results in potential axis switching since this axis is the result of OtherSwing x TwistAxis
 	FVec3 SwingAxis;
 	FReal SwingAngle;
-	FPBDJointUtilities::GetLockedSwingAxisAngle(ConnectorRs[0], ConnectorRs[1], SwingConstraintIndex, SwingAxis, SwingAngle);
-	SwingAxis.SafeNormalize();
+	FPBDJointUtilities::GetSwingAxisAngle(ConnectorRs[0], ConnectorRs[1], 0.0, SwingConstraintIndex, SwingAxis, SwingAngle);
 
 	InitRotationConstraintDatas(JointSettings, (int32)SwingConstraintIndex, SwingAxis, SwingAngle, 0.0, Dt, false);
 }
@@ -1134,19 +913,24 @@ void FPBDJointCachedSolver::ApplyRotationConstraints(
 
 void FPBDJointCachedSolver::SolveRotationConstraintDelta(
 		const int32 ConstraintIndex, 
-		const FReal DeltaLambda)
+		const FReal DeltaLambda,
+		const bool bIsSoftConstraint,
+		const FAxisConstraintDatas& ConstraintDatas)
 {
+	const FVec3 DeltaImpulse =  ConstraintDatas.ConstraintAxis[ConstraintIndex]  * DeltaLambda;
 	if(Body(0).IsDynamic())
 	{
-		const FVec3 DR0 = RotationConstraints.ConstraintDRAxis[ConstraintIndex][0] * DeltaLambda;
+		const FVec3 DR0 = !bIsSoftConstraint ? ConstraintDatas.ConstraintDRAxis[ConstraintIndex][0] * DeltaLambda :
+			DeltaImpulse *  (FVec3::DotProduct(ConstraintDatas.ConstraintAxis[ConstraintIndex], ConstraintDatas.ConstraintDRAxis[ConstraintIndex][0]));
 		ApplyRotationDelta(0, DR0);
 	}
 	if(Body(1).IsDynamic())
 	{
-		const FVec3 DR1 = RotationConstraints.ConstraintDRAxis[ConstraintIndex][1] * DeltaLambda;
+		const FVec3 DR1 = !bIsSoftConstraint ? ConstraintDatas.ConstraintDRAxis[ConstraintIndex][1] * DeltaLambda :
+			DeltaImpulse * ( FVec3::DotProduct(ConstraintDatas.ConstraintAxis[ConstraintIndex], ConstraintDatas.ConstraintDRAxis[ConstraintIndex][1]));
 		ApplyRotationDelta(1, DR1);
 	}
-	NetAngularImpulse += RotationConstraints.ConstraintAxis[ConstraintIndex]  * DeltaLambda;
+	NetAngularImpulse += DeltaImpulse;
 	++NumActiveConstraints;
 }
 
@@ -1161,7 +945,7 @@ void FPBDJointCachedSolver::SolveRotationConstraintHard(
 	{
 		RotationConstraints.ConstraintLambda[ConstraintIndex] += DeltaLambda;
 	}
-	SolveRotationConstraintDelta(ConstraintIndex, DeltaLambda);
+	SolveRotationConstraintDelta(ConstraintIndex, DeltaLambda, false, RotationConstraints);
 }
 
 void FPBDJointCachedSolver::SolveRotationConstraintSoft(
@@ -1184,7 +968,7 @@ void FPBDJointCachedSolver::SolveRotationConstraintSoft(
 		RotationConstraints.ConstraintSoftIM[ConstraintIndex];
 	RotationConstraints.ConstraintLambda[ConstraintIndex] += DeltaLambda;
 
-	SolveRotationConstraintDelta(ConstraintIndex, DeltaLambda);
+	SolveRotationConstraintDelta(ConstraintIndex, DeltaLambda,false, RotationConstraints);
 }
 
 void FPBDJointCachedSolver::ApplyRotationConstraint(
@@ -1265,7 +1049,7 @@ void FPBDJointCachedSolver::SolveAngularVelocityConstraint(
 
 void FPBDJointCachedSolver::ApplyAngularVelocityConstraint(const int32 ConstraintIndex)
 {
-	if(!NetAngularImpulse.IsNearlyZero() && RotationConstraints.ConstraintLambda[ConstraintIndex] > SMALL_NUMBER)
+	if(!NetAngularImpulse.IsNearlyZero() && (FMath::Abs(RotationConstraints.ConstraintLambda[ConstraintIndex]) > SMALL_NUMBER))
 	{
 		FReal TargetVel = 0.0f;
 		if (RotationConstraints.ConstraintRestitution[ConstraintIndex] != 0.0f)
@@ -1356,6 +1140,7 @@ void FPBDJointCachedSolver::ApplyPositionDrives(
 		}
 	}
 }
+	
 
 void FPBDJointCachedSolver::ApplyAxisPositionDrive(
 		const int32 ConstraintIndex,
@@ -1379,25 +1164,7 @@ void FPBDJointCachedSolver::ApplyAxisPositionDrive(
 		PositionDrives.ConstraintSoftIM[ConstraintIndex];
 	PositionDrives.ConstraintLambda[ConstraintIndex] += DeltaLambda;
 
-	const FVec3 DX = PositionDrives.ConstraintAxis[ConstraintIndex] * DeltaLambda;
-
-	if(Body(0).IsDynamic())
-	{
-		const FVec3 DP0 = InvM(0) * DX;
-		const FVec3 DR0 = PositionDrives.ConstraintDRAxis[ConstraintIndex][0] * DeltaLambda;
-		ApplyPositionDelta(0,DP0);
-		ApplyRotationDelta(0,DR0);
-	}
-	if(Body(1).IsDynamic())
-	{
-		const FVec3 DP1 = -InvM(1) * DX;
-		const FVec3 DR1 = PositionDrives.ConstraintDRAxis[ConstraintIndex][1] * DeltaLambda;
-		ApplyPositionDelta(1,DP1);
-		ApplyRotationDelta(1,DR1);
-	}
-
-	NetLinearImpulse += DX;
-	++NumActiveConstraints;
+	SolvePositionConstraintDelta(ConstraintIndex, DeltaLambda, PositionDrives);
 }
 
 /** INIT ROTATION DRIVES *********************************************************************************/
@@ -1576,7 +1343,9 @@ void FPBDJointCachedSolver::InitSLerpDrive(
 
 			if (FMath::Abs(SLerpAngle) > AngleTolerance)
 			{
-				InitRotationConstraintDrive(0, SLerpAxis, Dt, -SLerpAngle);
+				FReal AngVelTarget = (JointSettings.AngularDriveDamping > FReal(0)) ? FVec3::DotProduct(SLerpAxis, ConnectorRs[0] * JointSettings.AngularDriveVelocityTarget) : 0.0f;
+				InitRotationConstraintDrive((int32)EJointAngularConstraintIndex::Swing1, SLerpAxis, Dt, -SLerpAngle);
+				RotationDrives.ConstraintVX[(int32)EJointAngularConstraintIndex::Swing1] = AngVelTarget;
 			}
 		}
 	}
@@ -1617,18 +1386,7 @@ void FPBDJointCachedSolver::ApplyAxisRotationDrive(
 		RotationDrives.ConstraintSoftIM[ConstraintIndex];
 	RotationDrives.ConstraintLambda[ConstraintIndex] += DeltaLambda;
 
-	if(Body(0).IsDynamic())
-	{
-		const FVec3 DR0 = RotationDrives.ConstraintDRAxis[ConstraintIndex][0] * DeltaLambda;
-		ApplyRotationDelta(0, DR0);
-	}
-	if(Body(1).IsDynamic())
-	{
-		const FVec3 DR1 = RotationDrives.ConstraintDRAxis[ConstraintIndex][1] * DeltaLambda;
-		ApplyRotationDelta(1, DR1);
-	}
-	NetAngularImpulse += RotationDrives.ConstraintAxis[ConstraintIndex]  * DeltaLambda;
-	++NumActiveConstraints;
+	SolveRotationConstraintDelta(ConstraintIndex, DeltaLambda,true, RotationDrives);
 }
 
  // Joint utilities
@@ -1768,25 +1526,6 @@ void FAxisConstraintDatas::UpdateMass(
 		ConstraintSoftIM[ConstraintIndex] = (ConstraintSoftStiffness[ConstraintIndex] + ConstraintSoftDamping[ConstraintIndex]) * ConstraintHardIM[ConstraintIndex] + (FReal)1;
 	}
 }
-
-void FPointConstraintDatas::InitDatas(
-	const FReal PointStiffness)
-{
-	ConstraintStiffness = PointStiffness;
-	bValidDatas= false;
-	ConstraintLambda = FVec3::Zero();
-}
-
-void FPointConstraintDatas::UpdateDatas(
-	const FVec3& DatasCX,
-	const FVec3& DatasArm0,
-	const FVec3& DatasArm1)
-{
-	bValidDatas= true;
-	ConstraintCX = DatasCX;
-	ConstraintArms[0] = DatasArm0;
-	ConstraintArms[1] = DatasArm1;
-}
 	
 }
 
diff --git a/Engine/Source/Runtime/Experimental/Chaos/Public/Chaos/Joint/PBDJointCachedSolverGaussSeidel.h b/Engine/Source/Runtime/Experimental/Chaos/Public/Chaos/Joint/PBDJointCachedSolverGaussSeidel.h
--- a/Engine/Source/Runtime/Experimental/Chaos/Public/Chaos/Joint/PBDJointCachedSolverGaussSeidel.h
+++ b/Engine/Source/Runtime/Experimental/Chaos/Public/Chaos/Joint/PBDJointCachedSolverGaussSeidel.h
@@ -68,30 +68,7 @@ struct FAxisConstraintDatas
 	bool bSoftLimit[3];
 	EJointMotionType MotionType[3];
 };
-
-/** Cached point joint datas that will be used during the apply */
-struct FPointConstraintDatas
-{
-	/** Init the point joint datas with stiffness / damping */
-	void InitDatas(
-		const FReal PointStiffness);
-
-	/** Update the point joint datas with axis, limits, arms... */
-	void UpdateDatas(
-		const FVec3& DatasCX,
-		const FVec3& DatasArm0,
-		const FVec3& DatasArm1);
-	
-	bool bValidDatas = false;
 	
-	FReal ConstraintStiffness;
-	FVec3 ConstraintArms[2];
-	
-	FVec3 ConstraintCX;
-	FMatrix33 ConstraintIM;
-	FVec3 ConstraintLambda;
-};
-
 	/**
 	 * Calculate new positions and rotations for a pair of bodies connected by a joint.
 	 *
@@ -294,58 +271,25 @@ struct FPointConstraintDatas
 			const FReal Dt,
 			const FPBDJointSolverSettings& SolverSettings,
 			const FPBDJointSettings& JointSettings);
-		
-		void InitPointPositionConstraint(const FReal Dt);
-
-		void InitPointPositionMass();
-		
-		void InitPositionConstraintDatas(
-			const int32 ConstraintIndex,
-			const FVec3& ConstraintAxis,
-			const FReal& ConstraintDelta,
-			const FReal ConstraintRestitution,
-			const FReal Dt,
-			const FReal ConstraintLimit,
-			const EJointMotionType JointType);
 
 		void InitPositionDatasMass(
 			FAxisConstraintDatas& PositionDatas,
 			const int32 ConstraintIndex,
 			const FReal Dt);
-
-		void InitSphericalPositionConstraint(
-			const FPBDJointSettings& JointSettings,
-			const FReal Dt);
-
-		void InitCylindricalPositionConstraint(
-			const FPBDJointSettings& JointSettings,
-			const FReal Dt,
-			const int32 AxisIndex,
-			const EJointMotionType AxialMotion,
-			const EJointMotionType RadialMotion);
-
-		void InitPlanarPositionConstraint(
-			const FPBDJointSettings& JointSettings,
-			const FReal Dt,
-			const int32 AxisIndex,
-			const EJointMotionType AxialMotion);
 		
 		/** Apply Position constraints */
 
 		void ApplyPositionConstraints(
 			const FReal Dt);
 		
-		void ApplyPointPositionConstraint();
-		
-		FVec3 ComputePointCX(const bool bLinearCompute);
-		
 		void ApplyAxisPositionConstraint(
 			const int32 ConstraintIndex,
 			const FReal Dt);
 		
 		void SolvePositionConstraintDelta(
 			const int32 ConstraintIndex, 
-			const FReal DeltaLambda);
+			const FReal DeltaLambda,
+			const FAxisConstraintDatas& ConstraintDatas);
 
 		void SolvePositionConstraintHard(
 			const int32 ConstraintIndex,
@@ -361,8 +305,6 @@ struct FPointConstraintDatas
 		
 		void ApplyLinearVelocityConstraints();
 
-		void ApplyPointVelocityConstraint();
-
 		void ApplyAxisVelocityConstraint(
 			const int32 ConstraintIndex);
 
@@ -441,7 +383,9 @@ struct FPointConstraintDatas
 
 		void SolveRotationConstraintDelta(
 			const int32 ConstraintIndex, 
-			const FReal DeltaLambda);
+			const FReal DeltaLambda,
+			const bool bIsSoftConstraint,
+			const FAxisConstraintDatas& ConstraintDatas);
 
 		void SolveRotationConstraintHard(
 			const int32 ConstraintIndex,
@@ -612,8 +556,6 @@ struct FPointConstraintDatas
 		int32 NumActiveConstraints;					// The number of active constraints and drives in the last iteration (-1 initial value)
 		bool bIsActive;								// Whether any constraints actually moved any bodies in last iteration
 		
-		FPointConstraintDatas PointConstraints;
-		
 		FAxisConstraintDatas PositionConstraints;
 		FAxisConstraintDatas RotationConstraints;
 
