Allow physics solver to run on GT while still going wide.
Makes it easier to profile physics perf

diff --git a/Engine/Source/Runtime/Experimental/Chaos/Private/Chaos/Framework/PhysicsSolverBase.cpp b/Engine/Source/Runtime/Experimental/Chaos/Private/Chaos/Framework/PhysicsSolverBase.cpp
--- a/Engine/Source/Runtime/Experimental/Chaos/Private/Chaos/Framework/PhysicsSolverBase.cpp
+++ b/Engine/Source/Runtime/Experimental/Chaos/Private/Chaos/Framework/PhysicsSolverBase.cpp
@@ -67,6 +67,9 @@ namespace Chaos
 		ENamedThreads::HighTaskPriority // if we don't have hi pri threads, then use normal priority threads at high task priority instead
 	);
 
+	int32 PhysicsRunsOnGT = 0;
+	FAutoConsoleVariableRef CVarPhysicsRunsOnGT(TEXT("p.PhysicsRunsOnGT"), PhysicsRunsOnGT, TEXT("If true the physics thread runs on the game thread, but will still go wide on tasks like collision detection"));
+
 	ENamedThreads::Type FPhysicsSolverProcessPushDataTask::GetDesiredThread()
 	{
 		return CPrio_FPhysicsTickTask.Get();
@@ -74,7 +77,7 @@ namespace Chaos
 
 	ENamedThreads::Type FPhysicsSolverAdvanceTask::GetDesiredThread()
 	{
-		return CPrio_FPhysicsTickTask.Get();
+		return PhysicsRunsOnGT == 0 ? CPrio_FPhysicsTickTask.Get() : ENamedThreads::GameThread;
 	}
 
 	void FPhysicsSolverProcessPushDataTask::ProcessPushData()
@@ -177,6 +180,12 @@ namespace Chaos
 #endif
 	{
 		UE_LOG(LogChaos, Log, TEXT("FPhysicsSolverBase::AsyncDt:%f"), IsUsingAsyncResults() ? AsyncDt : -1);
+
+		//If user is running with -PhysicsRunsOnGT override the cvar (doing it here to avoid parsing every time task is scheduled)
+		if(FParse::Param(FCommandLine::Get(), TEXT("PhysicsRunsOnGT")))
+		{
+			PhysicsRunsOnGT = 1;
+		}
 	}
 
 	void FPhysicsSolverBase::EnableAsyncMode(FReal FixedDt)
