Remove ImplicitShapeMap to save memory. Collision detection passes shapes down with implicits and stores on constraints.

diff --git a/Engine/Source/Programs/HeadlessChaos/Private/HeadlessChaosTestShapePadding.cpp b/Engine/Source/Programs/HeadlessChaos/Private/HeadlessChaosTestShapePadding.cpp
--- a/Engine/Source/Programs/HeadlessChaos/Private/HeadlessChaosTestShapePadding.cpp
+++ b/Engine/Source/Programs/HeadlessChaos/Private/HeadlessChaosTestShapePadding.cpp
@@ -90,10 +90,12 @@ namespace ChaosTest {
 		TUniquePtr<FPBDCollisionConstraint> Constraint = FPBDCollisionConstraint::Make(
 			Box0,
 			Box0->Geometry().Get(),
+			Box0->ShapesArray()[0].Get(),
 			nullptr,
 			FRigidTransform3(),
 			Box1,
 			Box1->Geometry().Get(),
+			Box1->ShapesArray()[0].Get(),
 			nullptr,
 			FRigidTransform3(),
 			FLT_MAX,
@@ -200,10 +202,12 @@ namespace ChaosTest {
 		TUniquePtr<FPBDCollisionConstraint> Constraint = FPBDCollisionConstraint::Make(
 			Box0,
 			Box0->Geometry().Get(),
+			Box0->ShapesArray()[0].Get(),
 			nullptr,
 			FRigidTransform3(),
 			Box1,
 			Box1->Geometry().Get(),
+			Box1->ShapesArray()[0].Get(),
 			nullptr,
 			FRigidTransform3(),
 			FLT_MAX,
diff --git a/Engine/Source/Runtime/Experimental/Chaos/Private/Chaos/ChaosDebugDraw.cpp b/Engine/Source/Runtime/Experimental/Chaos/Private/Chaos/ChaosDebugDraw.cpp
--- a/Engine/Source/Runtime/Experimental/Chaos/Private/Chaos/ChaosDebugDraw.cpp
+++ b/Engine/Source/Runtime/Experimental/Chaos/Private/Chaos/ChaosDebugDraw.cpp
@@ -259,11 +259,12 @@ namespace Chaos
 		//
 		//
 
-		void DrawShapesImpl(const FGeometryParticleHandle* Particle, const FRigidTransform3& ShapeTransform, const FImplicitObject* Shape, const FReal Margin, const FColor& Color, const FChaosDebugDrawSettings& Settings);
+		void DrawShapesImpl(const FGeometryParticleHandle* Particle, const FRigidTransform3& ShapeTransform, const FImplicitObject* Implicit, const FShapeOrShapesArray& Shapes, const FReal Margin, const FColor& Color, const FChaosDebugDrawSettings& Settings);
 
-		void DrawShape(const FRigidTransform3& ShapeTransform, const FImplicitObject* Shape, const FColor& Color, const FChaosDebugDrawSettings* Settings)
+		void DrawShape(const FRigidTransform3& ShapeTransform, const FImplicitObject* Implicit, const FShapeOrShapesArray& Shapes, const FColor& Color, const FChaosDebugDrawSettings* Settings)
 		{
-			DrawShapesImpl(nullptr, ShapeTransform, Shape, 0.0f, Color, GetChaosDebugDrawSettings(Settings));
+			// Note: At the time of commenting this function does nothing as DrawShapesImpl does not handle null particle.
+			DrawShapesImpl(nullptr, ShapeTransform, Implicit, Shapes, 0.0f, Color, GetChaosDebugDrawSettings(Settings));
 		}
 
 		void DrawShapesConvexImpl(const TGeometryParticleHandle<FReal, 3>* Particle, const FRigidTransform3& ShapeTransform, const FConvex* Shape, const FReal InMargin, const FColor& Color, const FChaosDebugDrawSettings& Settings)
@@ -367,9 +368,9 @@ namespace Chaos
 		}
 
 		template <bool bInstanced>
-		void DrawShapesScaledImpl(const FGeometryParticleHandle* Particle, const FRigidTransform3& ShapeTransform, const FImplicitObject* Shape, const FReal Margin, const FColor& Color, const FChaosDebugDrawSettings& Settings)
+		void DrawShapesScaledImpl(const FGeometryParticleHandle* Particle, const FRigidTransform3& ShapeTransform, const FImplicitObject* Implicit, const FShapeOrShapesArray& Shapes, const FReal Margin, const FColor& Color, const FChaosDebugDrawSettings& Settings)
 		{
-			const EImplicitObjectType PackedType = Shape->GetType();
+			const EImplicitObjectType PackedType = Implicit->GetType();
 			const EImplicitObjectType InnerType = GetInnerType(PackedType);
 			CHAOS_CHECK(IsScaled(PackedType));
 			CHAOS_CHECK(IsInstanced(PackedType) == bInstanced);
@@ -391,19 +392,19 @@ namespace Chaos
 				break;
 			case ImplicitObjectType::LevelSet:
 			{
-				const TImplicitObjectScaled<FLevelSet, bInstanced>* Scaled = Shape->template GetObject<TImplicitObjectScaled<FLevelSet, bInstanced>>();
+				const TImplicitObjectScaled<FLevelSet, bInstanced>* Scaled = Implicit->template GetObject<TImplicitObjectScaled<FLevelSet, bInstanced>>();
 				// even though thhe levelset is scaled, the debugdraw uses the collisionParticles  that are pre-scaled
 				// so no need to pass the scaled transform and just extract the wrapped LevelSet
-				DrawShapesImpl(Particle, ShapeTransform, Scaled->GetUnscaledObject(), Scaled->GetMargin(), Color, Settings);
+				DrawShapesImpl(Particle, ShapeTransform, Scaled->GetUnscaledObject(), Shapes, Scaled->GetMargin(), Color, Settings);
 				break;
 			}
 			case ImplicitObjectType::Unknown:
 				break;
 			case ImplicitObjectType::Convex:
 			{
-				const TImplicitObjectScaled<FConvex, bInstanced>* Scaled = Shape->template GetObject<TImplicitObjectScaled<FConvex, bInstanced>>();
+				const TImplicitObjectScaled<FConvex, bInstanced>* Scaled = Implicit->template GetObject<TImplicitObjectScaled<FConvex, bInstanced>>();
 				ScaleTM.SetScale3D(Scaled->GetScale());
-				DrawShapesImpl(Particle, ScaleTM * ShapeTransform, Scaled->GetUnscaledObject(), Scaled->GetMargin(), Color, Settings);
+				DrawShapesImpl(Particle, ScaleTM * ShapeTransform, Scaled->GetUnscaledObject(), Shapes, Scaled->GetMargin(), Color, Settings);
 				break;
 			}
 			case ImplicitObjectType::TaperedCylinder:
@@ -412,16 +413,16 @@ namespace Chaos
 				break;
 			case ImplicitObjectType::TriangleMesh:
 			{
-				const TImplicitObjectScaled<FTriangleMeshImplicitObject, bInstanced>* Scaled = Shape->template GetObject<TImplicitObjectScaled<FTriangleMeshImplicitObject, bInstanced>>();
+				const TImplicitObjectScaled<FTriangleMeshImplicitObject, bInstanced>* Scaled = Implicit->template GetObject<TImplicitObjectScaled<FTriangleMeshImplicitObject, bInstanced>>();
 				ScaleTM.SetScale3D(Scaled->GetScale());
-				DrawShapesImpl(Particle, ScaleTM * ShapeTransform, Scaled->GetUnscaledObject(), 0.0f, Color, Settings);
+				DrawShapesImpl(Particle, ScaleTM * ShapeTransform, Scaled->GetUnscaledObject(), Shapes, 0.0f, Color, Settings);
 				break;
 			}
 			case ImplicitObjectType::HeightField:
 			{
-				const TImplicitObjectScaled<FHeightField, bInstanced>* Scaled = Shape->template GetObject<TImplicitObjectScaled<FHeightField, bInstanced>>();
+				const TImplicitObjectScaled<FHeightField, bInstanced>* Scaled = Implicit->template GetObject<TImplicitObjectScaled<FHeightField, bInstanced>>();
 				ScaleTM.SetScale3D(Scaled->GetScale());
-				DrawShapesImpl(Particle, ScaleTM * ShapeTransform, Scaled->GetUnscaledObject(), 0.0f, Color, Settings);
+				DrawShapesImpl(Particle, ScaleTM * ShapeTransform, Scaled->GetUnscaledObject(), Shapes, 0.0f, Color, Settings);
 				break;
 			}
 			default:
@@ -429,9 +430,9 @@ namespace Chaos
 			}
 		}
 
-		void DrawShapesInstancedImpl(const FGeometryParticleHandle* Particle, const FRigidTransform3& ShapeTransform, const FImplicitObject* Shape, const FReal Margin, const FColor& Color, const FChaosDebugDrawSettings& Settings)
+		void DrawShapesInstancedImpl(const FGeometryParticleHandle* Particle, const FRigidTransform3& ShapeTransform, const FImplicitObject* Implicit, const FShapeOrShapesArray& Shapes, const FReal Margin, const FColor& Color, const FChaosDebugDrawSettings& Settings)
 		{
-			const EImplicitObjectType PackedType = Shape->GetType();
+			const EImplicitObjectType PackedType = Implicit->GetType();
 			const EImplicitObjectType InnerType = GetInnerType(PackedType);
 			CHAOS_CHECK(IsScaled(PackedType) == false);
 			CHAOS_CHECK(IsInstanced(PackedType));
@@ -456,8 +457,8 @@ namespace Chaos
 				break;
 			case ImplicitObjectType::Convex:
 			{
-				const TImplicitObjectInstanced<FConvex>* Instanced = Shape->template GetObject<TImplicitObjectInstanced<FConvex>>();
-				DrawShapesImpl(Particle, ShapeTransform, Instanced->GetInstancedObject(), Instanced->GetMargin(), Color, Settings);
+				const TImplicitObjectInstanced<FConvex>* Instanced = Implicit->template GetObject<TImplicitObjectInstanced<FConvex>>();
+				DrawShapesImpl(Particle, ShapeTransform, Instanced->GetInstancedObject(), Shapes, Instanced->GetMargin(), Color, Settings);
 				break;
 			}
 			case ImplicitObjectType::TaperedCylinder:
@@ -466,14 +467,14 @@ namespace Chaos
 				break;
 			case ImplicitObjectType::TriangleMesh:
 			{
-				const TImplicitObjectInstanced<FTriangleMeshImplicitObject>* Scaled = Shape->template GetObject<TImplicitObjectInstanced<FTriangleMeshImplicitObject>>();
-				DrawShapesImpl(Particle, ShapeTransform, Scaled->GetInstancedObject(), 0.0f, Color, Settings);
+				const TImplicitObjectInstanced<FTriangleMeshImplicitObject>* Scaled = Implicit->template GetObject<TImplicitObjectInstanced<FTriangleMeshImplicitObject>>();
+				DrawShapesImpl(Particle, ShapeTransform, Scaled->GetInstancedObject(), Shapes, 0.0f, Color, Settings);
 				break;
 			}
 			case ImplicitObjectType::HeightField:
 			{
-				const TImplicitObjectInstanced<FHeightField>* Scaled = Shape->template GetObject<TImplicitObjectInstanced<FHeightField>>();
-				DrawShapesImpl(Particle, ShapeTransform, Scaled->GetInstancedObject(), 0.0f, Color, Settings);
+				const TImplicitObjectInstanced<FHeightField>* Scaled = Implicit->template GetObject<TImplicitObjectInstanced<FHeightField>>();
+				DrawShapesImpl(Particle, ShapeTransform, Scaled->GetInstancedObject(), Shapes, 0.0f, Color, Settings);
 				break;
 			}
 			default:
@@ -481,13 +482,13 @@ namespace Chaos
 			}
 		}
 
-		void DrawShapesImpl(const FGeometryParticleHandle* Particle, const FRigidTransform3& ShapeTransform, const FImplicitObject* Shape, const FReal Margin, const FColor& Color, const FChaosDebugDrawSettings& Settings)
+		void DrawShapesImpl(const FGeometryParticleHandle* Particle, const FRigidTransform3& ShapeTransform, const FImplicitObject* Implicit, const FShapeOrShapesArray& Shapes, const FReal Margin, const FColor& Color, const FChaosDebugDrawSettings& Settings)
 		{
 
-			if (!Particle || !Shape) return;
+			if (!Particle || !Implicit) return;
 
-			const EImplicitObjectType PackedType = Shape->GetType(); // Type includes scaling and instancing data
-			const EImplicitObjectType InnerType = GetInnerType(Shape->GetType());
+			const EImplicitObjectType PackedType = Implicit->GetType(); // Type includes scaling and instancing data
+			const EImplicitObjectType InnerType = GetInnerType(Implicit->GetType());
 
 			// Are we within the region of interest?
 			const FReal ParticleSize = Particle->HasBounds() ? 0.5f * Particle->LocalBounds().Extents().Size() : TNumericLimits<FReal>::Max();
@@ -501,41 +502,55 @@ namespace Chaos
 			{
 				if (IsInstanced(PackedType))
 				{
-					DrawShapesScaledImpl<true>(Particle, ShapeTransform, Shape, Margin, Color, Settings);
+					DrawShapesScaledImpl<true>(Particle, ShapeTransform, Implicit, Shapes, Margin, Color, Settings);
 				}
 				else
 				{
-					DrawShapesScaledImpl<false>(Particle, ShapeTransform, Shape, Margin, Color, Settings);
+					DrawShapesScaledImpl<false>(Particle, ShapeTransform, Implicit, Shapes, Margin, Color, Settings);
 				}
 				return;
 			}
 			else if (IsInstanced(PackedType))
 			{
-				DrawShapesInstancedImpl(Particle, ShapeTransform, Shape, Margin, Color, Settings);
+				DrawShapesInstancedImpl(Particle, ShapeTransform, Implicit, Shapes, Margin, Color, Settings);
 				return;
 			}
 			else if (InnerType == ImplicitObjectType::Transformed)
 			{
-				const TImplicitObjectTransformed<FReal, 3>* Transformed = Shape->template GetObject<TImplicitObjectTransformed<FReal, 3>>();
+				const TImplicitObjectTransformed<FReal, 3>* Transformed = Implicit->template GetObject<TImplicitObjectTransformed<FReal, 3>>();
 				FRigidTransform3 TransformedTransform = FRigidTransform3(ShapeTransform.TransformPosition(Transformed->GetTransform().GetLocation()), ShapeTransform.GetRotation() * Transformed->GetTransform().GetRotation());
-				DrawShapesImpl(Particle, TransformedTransform, Transformed->GetTransformedObject(), Margin, Color, Settings);
+				DrawShapesImpl(Particle, TransformedTransform, Transformed->GetTransformedObject(), Shapes, Margin, Color, Settings);
 				return;
 			}
 			else if (InnerType == ImplicitObjectType::Union)
 			{
-				const FImplicitObjectUnion* Union = Shape->template GetObject<FImplicitObjectUnion>();
-				for (auto& UnionShape : Union->GetObjects())
+				const FImplicitObjectUnion* Union = Implicit->template GetObject<FImplicitObjectUnion>();
+				int32 UnionIdx = 0;
+				for (auto& UnionImplicit : Union->GetObjects())
 				{
-					DrawShapesImpl(Particle, ShapeTransform, UnionShape.Get(), Margin, Color, Settings);
+					// Retrieve shape from union's shapes array
+					const FPerShapeData* PerShapeData = nullptr;
+					if (ensure(!Shapes.IsSingleShape()))
+					{
+						const FShapesArray& ShapesArray = *Shapes.GetShapesArray();
+						PerShapeData = ShapesArray[UnionIdx].Get();
+					}
+					
+					DrawShapesImpl(Particle, ShapeTransform, UnionImplicit.Get(), FShapeOrShapesArray(PerShapeData), Margin, Color, Settings);
+					UnionIdx++;
 				}
 				return;
 			}
 			else if (InnerType == ImplicitObjectType::UnionClustered)
 			{
-				const FImplicitObjectUnionClustered* Union = Shape->template GetObject<FImplicitObjectUnionClustered>();
-				for (auto& UnionShape : Union->GetObjects())
+				const FImplicitObjectUnionClustered* Union = Implicit->template GetObject<FImplicitObjectUnionClustered>();
+				for (auto& UnionImplicit : Union->GetObjects())
 				{
-					DrawShapesImpl(Particle, ShapeTransform, UnionShape.Get(), Margin, Color, Settings);
+					const TPBDRigidParticleHandle<FReal, 3>* OriginalParticle = Union->FindParticleForImplicitObject(UnionImplicit.Get());
+					if (ensure(OriginalParticle))
+					{
+						DrawShapesImpl(Particle, ShapeTransform, UnionImplicit.Get(), FShapeOrShapesArray(OriginalParticle), Margin, Color, Settings);
+					}
 				}
 				return;
 			}
@@ -544,7 +559,13 @@ namespace Chaos
 			// Whether we should show meshes and non-mesh shapes
 			bool bShowMeshes = Settings.bShowComplexCollision;
 			bool bShowNonMeshes = Settings.bShowSimpleCollision;
-			const FPerShapeData* ShapeData = Particle->GetImplicitShape(Shape);
+
+			const FPerShapeData* ShapeData = nullptr;
+			if (Shapes.IsValid() && ensure(Shapes.IsSingleShape()))
+			{
+				ShapeData = Shapes.GetShape();
+			}
+			
 			if (ShapeData != nullptr)
 			{
 				bShowMeshes = (Settings.bShowComplexCollision && (ShapeData->GetCollisionTraceType() != EChaosCollisionTraceFlag::Chaos_CTF_UseSimpleAsComplex))
@@ -579,14 +600,14 @@ namespace Chaos
 			{
 			case ImplicitObjectType::Sphere:
 			{
-				const TSphere<FReal, 3>* Sphere = Shape->template GetObject<TSphere<FReal, 3>>();
+				const TSphere<FReal, 3>* Sphere = Implicit->template GetObject<TSphere<FReal, 3>>();
 				const FVec3 P = ShapeTransform.TransformPosition(Sphere->GetCenter());
 				FDebugDrawQueue::GetInstance().DrawDebugSphere(P, Sphere->GetRadius(), 8, ShapeColor, false, KINDA_SMALL_NUMBER, Settings.DrawPriority, Settings.ShapeThicknesScale * Settings.LineThickness);
 				break;
 			}
 			case ImplicitObjectType::Box:
 			{
-				const TBox<FReal, 3>* Box = Shape->template GetObject<TBox<FReal, 3>>();
+				const TBox<FReal, 3>* Box = Implicit->template GetObject<TBox<FReal, 3>>();
 				const FVec3 P = ShapeTransform.TransformPosition(Box->GetCenter());
 				FDebugDrawQueue::GetInstance().DrawDebugBox(P, (FReal)0.5 * Box->Extents(), ShapeTransform.GetRotation(), ShapeColor, false, KINDA_SMALL_NUMBER, Settings.DrawPriority, Settings.ShapeThicknesScale * Settings.LineThickness);
 				break;
@@ -595,7 +616,7 @@ namespace Chaos
 				break;
 			case ImplicitObjectType::Capsule:
 			{
-				const FCapsule* Capsule = Shape->template GetObject<FCapsule>();
+				const FCapsule* Capsule = Implicit->template GetObject<FCapsule>();
 				const FVec3 P = ShapeTransform.TransformPosition(Capsule->GetCenter());
 				const FRotation3 Q = ShapeTransform.GetRotation() * FRotationMatrix::MakeFromZ(Capsule->GetAxis());
 				FDebugDrawQueue::GetInstance().DrawDebugCapsule(P, (FReal)0.5 * Capsule->GetHeight() + Capsule->GetRadius(), Capsule->GetRadius(), Q, ShapeColor, false, KINDA_SMALL_NUMBER, Settings.DrawPriority, Settings.ShapeThicknesScale * Settings.LineThickness);
@@ -603,14 +624,14 @@ namespace Chaos
 			}
 			case ImplicitObjectType::LevelSet:
 			{
-				const FLevelSet* LevelSet = Shape->template GetObject<FLevelSet>();
+				const FLevelSet* LevelSet = Implicit->template GetObject<FLevelSet>();
 				DrawShapesLevelSetImpl(Particle, ShapeTransform, LevelSet, ShapeColor, Settings);
 				break;
 			}
 			break;
 			case ImplicitObjectType::Convex:
 			{
-				const FConvex* Convex = Shape->template GetObject<FConvex>();
+				const FConvex* Convex = Implicit->template GetObject<FConvex>();
 
 				const FReal NetMargin = bChaosDebugDebugDrawCoreShapes ? Margin + Convex->GetMargin() : 0.0f;
 				DrawShapesConvexImpl(Particle, ShapeTransform, Convex, NetMargin, ShapeColor, Settings);
@@ -639,13 +660,13 @@ namespace Chaos
 				break;
 			case ImplicitObjectType::TriangleMesh:
 			{
-				const FTriangleMeshImplicitObject* TriangleMesh = Shape->template GetObject<FTriangleMeshImplicitObject>();
+				const FTriangleMeshImplicitObject* TriangleMesh = Implicit->template GetObject<FTriangleMeshImplicitObject>();
 				DrawShapesTriangleMeshImpl(Particle, ShapeTransform, TriangleMesh, ShapeColor, Settings);
 				break;
 			}
 			case ImplicitObjectType::HeightField:
 			{
-				const FHeightField* HeightField = Shape->template GetObject<FHeightField>();
+				const FHeightField* HeightField = Implicit->template GetObject<FHeightField>();
 				DrawShapesHeightFieldImpl(Particle, ShapeTransform, HeightField, ShapeColor, Settings);
 				break;
 			}
@@ -672,7 +693,7 @@ namespace Chaos
 			if (bChaosDebugDebugDrawShapeBounds)
 			{
 				const FColor ShapeBoundsColor = FColor::Orange;
-				const FAABB3& ShapeBounds = Shape->BoundingBox();
+				const FAABB3& ShapeBounds = Implicit->BoundingBox();
 				const FVec3 ShapeBoundsPos = ShapeTransform.TransformPosition(ShapeBounds.Center());
 				FDebugDrawQueue::GetInstance().DrawDebugBox(ShapeBoundsPos, 0.5f * ShapeBounds.Extents(), ShapeTransform.GetRotation(), ShapeBoundsColor, false, KINDA_SMALL_NUMBER, Settings.DrawPriority, Settings.LineThickness);
 			}
@@ -683,7 +704,7 @@ namespace Chaos
 			FVec3 P = SpaceTransform.TransformPosition(Particle->ObjectState() == EObjectStateType::Dynamic ? Particle->CastToRigidParticle()->P() : Particle->X());
 			FRotation3 Q = SpaceTransform.GetRotation() * (Particle->ObjectState() == EObjectStateType::Dynamic ? Particle->CastToRigidParticle()->Q() : Particle->R());
 
-			DrawShapesImpl(Particle, FRigidTransform3(P, Q), Particle->Geometry().Get(), 0.0f, InColor, Settings);
+			DrawShapesImpl(Particle, FRigidTransform3(P, Q), Particle->Geometry().Get(), FShapeOrShapesArray(Particle), 0.0f, InColor, Settings);
 		}
 
 		void DrawParticleShapesImpl(const FRigidTransform3& SpaceTransform, const FGeometryParticle* Particle, const FColor& InColor, const FChaosDebugDrawSettings& Settings)
@@ -691,7 +712,7 @@ namespace Chaos
 			FVec3 P = SpaceTransform.TransformPosition(Particle->X());
 			FRotation3 Q = SpaceTransform.GetRotation() * (Particle->R());
 
-			DrawShapesImpl(Particle->Handle(), FRigidTransform3(P, Q), Particle->Geometry().Get(), 0.0f, InColor, Settings);
+			DrawShapesImpl(Particle->Handle(), FRigidTransform3(P, Q), Particle->Geometry().Get(), FShapeOrShapesArray(Particle->Handle()), 0.0f, InColor, Settings);
 		}
 
 		static EImplicitObjectType GetFirstConcreteShapeType(const FImplicitObject* Shape)
@@ -952,7 +973,7 @@ namespace Chaos
 
 		void DrawCollidingShapesImpl(const FRigidTransform3& SpaceTransform, const FPBDCollisionConstraints& Collisions, FRealSingle ColorScale, const FChaosDebugDrawSettings& Settings)
 		{
-			TArray<const FImplicitObject*> Shapes;
+			TArray<const FImplicitObject*> Implicits;
 			TArray<FConstGenericParticleHandle> ShapeParticles;
 			TArray<FRigidTransform3> ShapeTransforms;
 
@@ -972,16 +993,16 @@ namespace Chaos
 						const FRigidTransform3 ShapeWorldTransform0 = PointConstraint->GetShapeRelativeTransform0() * WorldActorTransform0;
 						const FRigidTransform3 ShapeWorldTransform1 = PointConstraint->GetShapeRelativeTransform1() * WorldActorTransform1;
 
-						if (!Shapes.Contains(Implicit0))
+						if (!Implicits.Contains(Implicit0))
 						{
-							Shapes.Add(Implicit0);
+							Implicits.Add(Implicit0);
 							ShapeParticles.Add(Particle0);
 							ShapeTransforms.Add(ShapeWorldTransform0);
 						}
 
-						if (!Shapes.Contains(Implicit1))
+						if (!Implicits.Contains(Implicit1))
 						{
-							Shapes.Add(Implicit1);
+							Implicits.Add(Implicit1);
 							ShapeParticles.Add(Particle1);
 							ShapeTransforms.Add(ShapeWorldTransform1);
 						}
@@ -989,12 +1010,13 @@ namespace Chaos
 				}
 			}
 
-			for (int32 ShapeIndex = 0; ShapeIndex < Shapes.Num(); ++ShapeIndex)
+			for (int32 ShapeIndex = 0; ShapeIndex < Implicits.Num(); ++ShapeIndex)
 			{
 				DrawShapesImpl(
 					ShapeParticles[ShapeIndex]->Handle(), 
 					ShapeTransforms[ShapeIndex], 
-					Shapes[ShapeIndex], 
+					Implicits[ShapeIndex],
+					FShapeOrShapesArray(ShapeParticles[ShapeIndex]->Handle()),
 					0.0f, 
 					ShapeParticles[ShapeIndex]->IsDynamic() ? FColor::Yellow : FColor::Red, 
 					Settings);
diff --git a/Engine/Source/Runtime/Experimental/Chaos/Private/Chaos/Collision/PBDCollisionConstraint.cpp b/Engine/Source/Runtime/Experimental/Chaos/Private/Chaos/Collision/PBDCollisionConstraint.cpp
--- a/Engine/Source/Runtime/Experimental/Chaos/Private/Chaos/Collision/PBDCollisionConstraint.cpp
+++ b/Engine/Source/Runtime/Experimental/Chaos/Private/Chaos/Collision/PBDCollisionConstraint.cpp
@@ -103,10 +103,12 @@ namespace Chaos
 	void FPBDCollisionConstraint::MakeInline(
 		FGeometryParticleHandle* Particle0,
 		const FImplicitObject* Implicit0,
+		const FPerShapeData* Shape0,
 		const FBVHParticles* Simplicial0,
 		const FRigidTransform3& ImplicitLocalTransform0,
 		FGeometryParticleHandle* Particle1,
 		const FImplicitObject* Implicit1,
+		const FPerShapeData* Shape1,
 		const FBVHParticles* Simplicial1,
 		const FRigidTransform3& ImplicitLocalTransform1,
 		const FReal InCullDistance,
@@ -118,6 +120,8 @@ namespace Chaos
 		OutConstraint.Particle[1] = Particle1;
 		OutConstraint.Implicit[0] = Implicit0;
 		OutConstraint.Implicit[1] = Implicit1;
+		OutConstraint.Shape[0] = Shape0;
+		OutConstraint.Shape[1] = Shape1;
 		OutConstraint.Simplicial[0] = Simplicial0;
 		OutConstraint.Simplicial[1] = Simplicial1;
 
@@ -127,17 +131,19 @@ namespace Chaos
 	TUniquePtr<FPBDCollisionConstraint> FPBDCollisionConstraint::Make(
 		FGeometryParticleHandle* Particle0,
 		const FImplicitObject* Implicit0,
+		const FPerShapeData* Shape0,
 		const FBVHParticles* Simplicial0,
 		const FRigidTransform3& ImplicitLocalTransform0,
 		FGeometryParticleHandle* Particle1,
 		const FImplicitObject* Implicit1,
+		const FPerShapeData* Shape1,
 		const FBVHParticles* Simplicial1,
 		const FRigidTransform3& ImplicitLocalTransform1,
 		const FReal InCullDistance,
 		const bool bInUseManifold,
 		const EContactShapesType ShapesType)
 	{
-		FPBDCollisionConstraint* Constraint = new FPBDCollisionConstraint(Particle0, Implicit0, Simplicial0, Particle1, Implicit1, Simplicial1);
+		FPBDCollisionConstraint* Constraint = new FPBDCollisionConstraint(Particle0, Implicit0, Shape0, Simplicial0, Particle1, Implicit1, Shape1, Simplicial1);
 		
 		Constraint->Setup(ECollisionCCDType::Disabled, ShapesType, ImplicitLocalTransform0, ImplicitLocalTransform1, InCullDistance, bInUseManifold);
 
@@ -168,6 +174,7 @@ namespace Chaos
 		: ImplicitTransform{ FRigidTransform3(), FRigidTransform3() }
 		, Particle{ nullptr, nullptr }
 		, Implicit{ nullptr, nullptr }
+		, Shape{ nullptr, nullptr }
 		, Simplicial{ nullptr, nullptr }
 		, Manifold()
 		, Stiffness(1)
@@ -196,13 +203,16 @@ namespace Chaos
 	FPBDCollisionConstraint::FPBDCollisionConstraint(
 		FGeometryParticleHandle* Particle0,
 		const FImplicitObject* Implicit0,
+		const FPerShapeData* Shape0,
 		const FBVHParticles* Simplicial0,
 		FGeometryParticleHandle* Particle1,
 		const FImplicitObject* Implicit1,
+		const FPerShapeData* Shape1,
 		const FBVHParticles* Simplicial1)
 		: ImplicitTransform{ FRigidTransform3(), FRigidTransform3() }
 		, Particle{ Particle0, Particle1 }
 		, Implicit{ Implicit0, Implicit1 }
+		, Shape{ Shape0, Shape1 }
 		, Simplicial{ Simplicial0, Simplicial1 }
 		, Manifold()
 		, Stiffness(1)
diff --git a/Engine/Source/Runtime/Experimental/Chaos/Private/Chaos/Collision/ParticlePairMidPhase.cpp b/Engine/Source/Runtime/Experimental/Chaos/Private/Chaos/Collision/ParticlePairMidPhase.cpp
--- a/Engine/Source/Runtime/Experimental/Chaos/Private/Chaos/Collision/ParticlePairMidPhase.cpp
+++ b/Engine/Source/Runtime/Experimental/Chaos/Private/Chaos/Collision/ParticlePairMidPhase.cpp
@@ -62,23 +62,25 @@ namespace Chaos
 		const FBVHParticles* BVHParticles1 = FConstGenericParticleHandle(Particle1)->CollisionParticles().Get();
 		const FRigidTransform3& ShapeRelativeTransform1 = (FRigidTransform3)InShape1->GetLeafRelativeTransform();
 
-		return FPBDCollisionConstraint::Make(Particle0, Implicit0, BVHParticles0, ShapeRelativeTransform0, Particle1, Implicit1, BVHParticles1, ShapeRelativeTransform1, CullDistance, bUseManifold, ShapePairType);
+		return FPBDCollisionConstraint::Make(Particle0, Implicit0, InShape0, BVHParticles0, ShapeRelativeTransform0, Particle1, Implicit1, InShape1, BVHParticles1, ShapeRelativeTransform1, CullDistance, bUseManifold, ShapePairType);
 	}
 
 	TUniquePtr<FPBDCollisionConstraint> CreateImplicitPairConstraint(
 		FGeometryParticleHandle* Particle0,
 		const FImplicitObject* Implicit0,
+		const FPerShapeData* Shape0,
 		const FBVHParticles* BVHParticles0,
 		const FRigidTransform3& ShapeRelativeTransform0,
 		FGeometryParticleHandle* Particle1,
 		const FImplicitObject* Implicit1,
+		const FPerShapeData* Shape1,
 		const FBVHParticles* BVHParticles1,
 		const FRigidTransform3& ShapeRelativeTransform1,
 		const FReal CullDistance,
 		const EContactShapesType ShapePairType,
 		const bool bUseManifold)
 	{
-		return FPBDCollisionConstraint::Make(Particle0, Implicit0, BVHParticles0, ShapeRelativeTransform0, Particle1, Implicit1, BVHParticles1, ShapeRelativeTransform1, CullDistance, bUseManifold, ShapePairType);
+		return FPBDCollisionConstraint::Make(Particle0, Implicit0, Shape0, BVHParticles0, ShapeRelativeTransform0, Particle1, Implicit1, Shape1, BVHParticles1, ShapeRelativeTransform1, CullDistance, bUseManifold, ShapePairType);
 	}
 
 
@@ -479,8 +481,10 @@ namespace Chaos
 			Particle0,
 			Particle1,
 			Implicit0,
+			Shape0,
 			BVHParticles0,
 			Implicit1,
+			Shape1,
 			BVHParticles1,
 			ParticleWorldTransform0,
 			ShapeRelativeTransform0,
@@ -503,10 +507,12 @@ namespace Chaos
 	FPBDCollisionConstraint* FMultiShapePairCollisionDetector::FindOrCreateConstraint(
 		FGeometryParticleHandle* InParticle0,
 		const FImplicitObject* Implicit0,
+		const FPerShapeData* InShape0,
 		const FBVHParticles* BVHParticles0,
 		const FRigidTransform3& ShapeRelativeTransform0,
 		FGeometryParticleHandle* InParticle1,
 		const FImplicitObject* Implicit1,
+		const FPerShapeData* InShape1,
 		const FBVHParticles* BVHParticles1,
 		const FRigidTransform3& ShapeRelativeTransform1,
 		const FReal CullDistance,
@@ -542,7 +548,7 @@ namespace Chaos
 		if (Constraint == nullptr)
 		{
 			// NOTE: Using InParticle0 and InParticle1 here because the order may be different to what we have stored
-			Constraint = CreateConstraint(InParticle0, Implicit0, BVHParticles0, ShapeRelativeTransform0, InParticle1, Implicit1, BVHParticles1, ShapeRelativeTransform1, CullDistance, ShapePairType, bUseManifold, Key);
+			Constraint = CreateConstraint(InParticle0, Implicit0, InShape0, BVHParticles0, ShapeRelativeTransform0, InParticle1, Implicit1, InShape1, BVHParticles1, ShapeRelativeTransform1, CullDistance, ShapePairType, bUseManifold, Key);
 		}
 
 		// @todo(chaos): we already have the shape world transforms at the calling site - pass them in
@@ -560,17 +566,19 @@ namespace Chaos
 	FPBDCollisionConstraint* FMultiShapePairCollisionDetector::FindOrCreateSweptConstraint(
 		FGeometryParticleHandle* InParticle0,
 		const FImplicitObject* Implicit0,
+		const FPerShapeData* InShape0,
 		const FBVHParticles* BVHParticles0,
 		const FRigidTransform3& ShapeRelativeTransform0,
 		FGeometryParticleHandle* InParticle1,
 		const FImplicitObject* Implicit1,
+		const FPerShapeData* InShape1,
 		const FBVHParticles* BVHParticles1,
 		const FRigidTransform3& ShapeRelativeTransform1,
 		const FReal CullDistance,
 		const EContactShapesType ShapePairType)
 	{
 		const bool bUseManifold = true;
-		FPBDCollisionConstraint* Constraint = FindOrCreateConstraint(InParticle0, Implicit0, BVHParticles0, ShapeRelativeTransform0, InParticle1, Implicit1, BVHParticles1, ShapeRelativeTransform1, CullDistance, ShapePairType, bUseManifold);
+		FPBDCollisionConstraint* Constraint = FindOrCreateConstraint(InParticle0, Implicit0, InShape0, BVHParticles0, ShapeRelativeTransform0, InParticle1, Implicit1, InShape1, BVHParticles1, ShapeRelativeTransform1, CullDistance, ShapePairType, bUseManifold);
 		if (Constraint != nullptr)
 		{
 			Constraint->SetCCDEnabled(true);
@@ -591,10 +599,12 @@ namespace Chaos
 	FPBDCollisionConstraint* FMultiShapePairCollisionDetector::CreateConstraint(
 		FGeometryParticleHandle* InParticle0,
 		const FImplicitObject* Implicit0,
+		const FPerShapeData* InShape0,
 		const FBVHParticles* BVHParticles0,
 		const FRigidTransform3& ShapeRelativeTransform0,
 		FGeometryParticleHandle* InParticle1,
 		const FImplicitObject* Implicit1,
+		const FPerShapeData* InShape1,
 		const FBVHParticles* BVHParticles1,
 		const FRigidTransform3& ShapeRelativeTransform1,
 		const FReal CullDistance,
@@ -603,7 +613,7 @@ namespace Chaos
 		const FCollisionParticlePairConstraintKey& Key)
 	{
 		PHYSICS_CSV_SCOPED_EXPENSIVE(PhysicsVerbose, NarrowPhase_CreateConstraint);
-		TUniquePtr<FPBDCollisionConstraint> Constraint = CreateImplicitPairConstraint(InParticle0, Implicit0, BVHParticles0, ShapeRelativeTransform0, InParticle1, Implicit1, BVHParticles1, ShapeRelativeTransform1, CullDistance, ShapePairType, bUseManifold);
+		TUniquePtr<FPBDCollisionConstraint> Constraint = CreateImplicitPairConstraint(InParticle0, Implicit0, InShape0, BVHParticles0, ShapeRelativeTransform0, InParticle1, Implicit1, InShape1, BVHParticles1, ShapeRelativeTransform1, CullDistance, ShapePairType, bUseManifold);
 
 		Constraint->GetContainerCookie().MidPhase = &MidPhase;
 		Constraint->GetContainerCookie().bIsMultiShapePair = true;
@@ -885,9 +895,8 @@ namespace Chaos
 	{
 		if (!Constraint.GetContainerCookie().bIsMultiShapePair)
 		{
-			// @todo(chaos): remove GetImplicitShape - we should store the shape in the constraint
-			const FPerShapeData* Shape0 = Constraint.GetParticle0()->GetImplicitShape(Constraint.GetImplicit0());
-			const FPerShapeData* Shape1 = Constraint.GetParticle1()->GetImplicitShape(Constraint.GetImplicit1());
+			const FPerShapeData* Shape0 = Constraint.GetShape0();
+			const FPerShapeData* Shape1 = Constraint.GetShape1();
 
 			// @todo(chaos): fix O(N) search for shape pair - store the index in the cookie (it will be the same
 			// as long as the ShapesArray on each particle has not changed)
diff --git a/Engine/Source/Runtime/Experimental/Chaos/Private/Chaos/CollisionResolution.cpp b/Engine/Source/Runtime/Experimental/Chaos/Private/Chaos/CollisionResolution.cpp
--- a/Engine/Source/Runtime/Experimental/Chaos/Private/Chaos/CollisionResolution.cpp
+++ b/Engine/Source/Runtime/Experimental/Chaos/Private/Chaos/CollisionResolution.cpp
@@ -522,13 +522,13 @@ namespace Chaos
 		}
 
 		template<typename T_TRAITS>
-		void ConstructBoxBoxConstraints(TGeometryParticleHandle<FReal, 3>* Particle0, TGeometryParticleHandle<FReal, 3>* Particle1, const FImplicitObject* Implicit0, const FImplicitObject* Implicit1, const FRigidTransform3& ParticleWorldTransform0, const FRigidTransform3& LocalTransform0, const FRigidTransform3& ParticleWorldTransform1, const FRigidTransform3& LocalTransform1, const FReal CullDistance, const FReal Dt, const FCollisionContext& Context)
+		void ConstructBoxBoxConstraints(TGeometryParticleHandle<FReal, 3>* Particle0, TGeometryParticleHandle<FReal, 3>* Particle1, const FImplicitObject* Implicit0, const FPerShapeData* Shape0, const FImplicitObject* Implicit1, const FPerShapeData* Shape1, const FRigidTransform3& ParticleWorldTransform0, const FRigidTransform3& LocalTransform0, const FRigidTransform3& ParticleWorldTransform1, const FRigidTransform3& LocalTransform1, const FReal CullDistance, const FReal Dt, const FCollisionContext& Context)
 		{
 			const TBox<FReal, 3>* Object0 = Implicit0->template GetObject<const TBox<FReal, 3> >();
 			const TBox<FReal, 3>* Object1 = Implicit1->template GetObject<const TBox<FReal, 3> >();
 			if (ensure(Object0 && Object1))
 			{
-				FPBDCollisionConstraint* Constraint = Context.CollisionAllocator->FindOrCreateConstraint(Particle0, Implicit0, nullptr, LocalTransform0, Particle1, Implicit1, nullptr, LocalTransform1, CullDistance, EContactShapesType::BoxBox, Context.bAllowManifolds);
+				FPBDCollisionConstraint* Constraint = Context.CollisionAllocator->FindOrCreateConstraint(Particle0, Implicit0, Shape0, nullptr, LocalTransform0, Particle1, Implicit1, Shape1, nullptr, LocalTransform1, CullDistance, EContactShapesType::BoxBox, Context.bAllowManifolds);
 				if (T_TRAITS::bImmediateUpdate && (Constraint != nullptr))
 				{
 					const FRigidTransform3 WorldTransform0 = LocalTransform0 * ParticleWorldTransform0;
@@ -558,7 +558,7 @@ namespace Chaos
 		}
 
 		template<typename T_TRAITS>
-		void ConstructBoxHeightFieldConstraints(TGeometryParticleHandle<FReal, 3>* Particle0, TGeometryParticleHandle<FReal, 3>* Particle1, const FImplicitObject* Implicit0, const FImplicitObject* Implicit1, const FRigidTransform3& ParticleWorldTransform0, const FRigidTransform3& LocalTransform0, const FRigidTransform3& ParticleWorldTransform1, const FRigidTransform3& LocalTransform1, const FReal CullDistance, const FReal Dt, const FCollisionContext& Context)
+		void ConstructBoxHeightFieldConstraints(TGeometryParticleHandle<FReal, 3>* Particle0, TGeometryParticleHandle<FReal, 3>* Particle1, const FImplicitObject* Implicit0, const FPerShapeData* Shape0, const FImplicitObject* Implicit1, const FPerShapeData* Shape1, const FRigidTransform3& ParticleWorldTransform0, const FRigidTransform3& LocalTransform0, const FRigidTransform3& ParticleWorldTransform1, const FRigidTransform3& LocalTransform1, const FReal CullDistance, const FReal Dt, const FCollisionContext& Context)
 		{
 
 			const TBox<FReal, 3>* Object0 = Implicit0->template GetObject<const TBox<FReal, 3> >();
@@ -566,7 +566,7 @@ namespace Chaos
 			if (ensure(Object0 && Object1))
 			{
 				// @todo(chaos): one-shot manifold
-				FPBDCollisionConstraint* Constraint = Context.CollisionAllocator->FindOrCreateConstraint(Particle0, Implicit0, nullptr, LocalTransform0, Particle1, Implicit1, nullptr, LocalTransform1, CullDistance, EContactShapesType::BoxHeightField, Context.bAllowManifolds);
+				FPBDCollisionConstraint* Constraint = Context.CollisionAllocator->FindOrCreateConstraint(Particle0, Implicit0, Shape0, nullptr, LocalTransform0, Particle1, Implicit1, Shape1, nullptr, LocalTransform1, CullDistance, EContactShapesType::BoxHeightField, Context.bAllowManifolds);
 				if (T_TRAITS::bImmediateUpdate && (Constraint != nullptr))
 				{
 					FRigidTransform3 WorldTransform0 = LocalTransform0 * ParticleWorldTransform0;
@@ -653,13 +653,13 @@ namespace Chaos
 
 
 		template<typename T_TRAITS>
-		void ConstructBoxPlaneConstraints(TGeometryParticleHandle<FReal, 3>* Particle0, TGeometryParticleHandle<FReal, 3>* Particle1, const FImplicitObject* Implicit0, const FImplicitObject* Implicit1, const FRigidTransform3& ParticleWorldTransform0, const FRigidTransform3& LocalTransform0, const FRigidTransform3& ParticleWorldTransform1, const FRigidTransform3& LocalTransform1, const FReal CullDistance, const FReal Dt, const FCollisionContext& Context)
+		void ConstructBoxPlaneConstraints(TGeometryParticleHandle<FReal, 3>* Particle0, TGeometryParticleHandle<FReal, 3>* Particle1, const FImplicitObject* Implicit0, const FPerShapeData* Shape0, const FImplicitObject* Implicit1, const FPerShapeData* Shape1, const FRigidTransform3& ParticleWorldTransform0, const FRigidTransform3& LocalTransform0, const FRigidTransform3& ParticleWorldTransform1, const FRigidTransform3& LocalTransform1, const FReal CullDistance, const FReal Dt, const FCollisionContext& Context)
 		{
 			const TBox<FReal, 3>* Object0 = Implicit0->template GetObject<const TBox<FReal, 3> >();
 			const TPlane<FReal, 3>* Object1 = Implicit1->template GetObject<const TPlane<FReal, 3> >();
 			if (ensure(Object0 && Object1))
 			{
-				FPBDCollisionConstraint* Constraint = Context.CollisionAllocator->FindOrCreateConstraint(Particle0, Implicit0, nullptr, LocalTransform0, Particle1, Implicit1, nullptr, LocalTransform1, CullDistance, EContactShapesType::BoxPlane, Context.bAllowManifolds);
+				FPBDCollisionConstraint* Constraint = Context.CollisionAllocator->FindOrCreateConstraint(Particle0, Implicit0, Shape0, nullptr, LocalTransform0, Particle1, Implicit1, Shape1, nullptr, LocalTransform1, CullDistance, EContactShapesType::BoxPlane, Context.bAllowManifolds);
 				if (T_TRAITS::bImmediateUpdate && (Constraint != nullptr))
 				{
 					FRigidTransform3 WorldTransform0 = LocalTransform0 * ParticleWorldTransform0;
@@ -688,7 +688,7 @@ namespace Chaos
 		}
 
 		template<typename T_TRAITS>
-		void ConstructBoxTriangleMeshConstraints(TGeometryParticleHandle<FReal, 3>* Particle0, TGeometryParticleHandle<FReal, 3>* Particle1, const FImplicitObject* Implicit0, const FImplicitObject* Implicit1, const FRigidTransform3& ParticleWorldTransform0, const FRigidTransform3& LocalTransform0, const FRigidTransform3& ParticleWorldTransform1, const FRigidTransform3& LocalTransform1, const FReal CullDistance, const FReal Dt, const FCollisionContext& Context)
+		void ConstructBoxTriangleMeshConstraints(TGeometryParticleHandle<FReal, 3>* Particle0, TGeometryParticleHandle<FReal, 3>* Particle1, const FImplicitObject* Implicit0, const FPerShapeData* Shape0, const FImplicitObject* Implicit1, const FPerShapeData* Shape1, const FRigidTransform3& ParticleWorldTransform0, const FRigidTransform3& LocalTransform0, const FRigidTransform3& ParticleWorldTransform1, const FRigidTransform3& LocalTransform1, const FReal CullDistance, const FReal Dt, const FCollisionContext& Context)
 		{
 			const TBox<FReal, 3>* Object0 = Implicit0->template GetObject<const TBox<FReal, 3> >();
 			if (ensure(Object0))
@@ -696,7 +696,7 @@ namespace Chaos
 				if (const TImplicitObjectScaled<FTriangleMeshImplicitObject>* ScaledTriangleMesh = Implicit1->template GetObject<const TImplicitObjectScaled<FTriangleMeshImplicitObject>>())
 				{
 					// @todo(chaos): one-shot manifold
-					FPBDCollisionConstraint* Constraint = Context.CollisionAllocator->FindOrCreateConstraint(Particle0, Implicit0, nullptr, LocalTransform0, Particle1, Implicit1, nullptr, LocalTransform1, CullDistance, EContactShapesType::BoxTriMesh, Context.bAllowManifolds);
+					FPBDCollisionConstraint* Constraint = Context.CollisionAllocator->FindOrCreateConstraint(Particle0, Implicit0, Shape0, nullptr, LocalTransform0, Particle1, Implicit1, Shape1, nullptr, LocalTransform1, CullDistance, EContactShapesType::BoxTriMesh, Context.bAllowManifolds);
 					if (T_TRAITS::bImmediateUpdate && (Constraint != nullptr))
 					{
 						FRigidTransform3 WorldTransform0 = LocalTransform0 * ParticleWorldTransform0;
@@ -707,7 +707,7 @@ namespace Chaos
 				else if (const FTriangleMeshImplicitObject* TriangleMesh = Implicit1->template GetObject<const FTriangleMeshImplicitObject>())
 				{
 					// @todo(chaos): one-shot manifold
-					FPBDCollisionConstraint* Constraint = Context.CollisionAllocator->FindOrCreateConstraint(Particle0, Implicit0, nullptr, LocalTransform0, Particle1, Implicit1, nullptr, LocalTransform1, CullDistance, EContactShapesType::BoxTriMesh, Context.bAllowManifolds);
+					FPBDCollisionConstraint* Constraint = Context.CollisionAllocator->FindOrCreateConstraint(Particle0, Implicit0, Shape0, nullptr, LocalTransform0, Particle1, Implicit1, Shape1, nullptr, LocalTransform1, CullDistance, EContactShapesType::BoxTriMesh, Context.bAllowManifolds);
 					if (T_TRAITS::bImmediateUpdate && (Constraint != nullptr))
 					{
 						FRigidTransform3 WorldTransform0 = LocalTransform0 * ParticleWorldTransform0;
@@ -739,13 +739,13 @@ namespace Chaos
 		}
 
 		template<typename T_TRAITS>
-		void ConstructSphereSphereConstraints(TGeometryParticleHandle<FReal, 3>* Particle0, TGeometryParticleHandle<FReal, 3>* Particle1, const FImplicitObject* Implicit0, const FImplicitObject* Implicit1, const FRigidTransform3& ParticleWorldTransform0, const FRigidTransform3& LocalTransform0, const FRigidTransform3& ParticleWorldTransform1, const FRigidTransform3& LocalTransform1, const FReal CullDistance, const FReal Dt, const FCollisionContext& Context)
+		void ConstructSphereSphereConstraints(TGeometryParticleHandle<FReal, 3>* Particle0, TGeometryParticleHandle<FReal, 3>* Particle1, const FImplicitObject* Implicit0, const FPerShapeData* Shape0, const FImplicitObject* Implicit1, const FPerShapeData* Shape1, const FRigidTransform3& ParticleWorldTransform0, const FRigidTransform3& LocalTransform0, const FRigidTransform3& ParticleWorldTransform1, const FRigidTransform3& LocalTransform1, const FReal CullDistance, const FReal Dt, const FCollisionContext& Context)
 		{
 			const TSphere<FReal, 3>* Object0 = Implicit0->template GetObject<const TSphere<FReal, 3> >();
 			const TSphere<FReal, 3>* Object1 = Implicit1->template GetObject<const TSphere<FReal, 3> >();
 			if (ensure(Object0 && Object1))
 			{
-				FPBDCollisionConstraint* Constraint = Context.CollisionAllocator->FindOrCreateConstraint(Particle0, Implicit0, nullptr, LocalTransform0, Particle1, Implicit1, nullptr, LocalTransform1, CullDistance, EContactShapesType::SphereSphere, Context.bAllowManifolds);
+				FPBDCollisionConstraint* Constraint = Context.CollisionAllocator->FindOrCreateConstraint(Particle0, Implicit0, Shape0, nullptr, LocalTransform0, Particle1, Implicit1, Shape1, nullptr, LocalTransform1, CullDistance, EContactShapesType::SphereSphere, Context.bAllowManifolds);
 				if (T_TRAITS::bImmediateUpdate && (Constraint != nullptr))
 				{
 					FRigidTransform3 WorldTransform0 = LocalTransform0 * ParticleWorldTransform0;
@@ -780,14 +780,14 @@ namespace Chaos
 		}
 
 		template<typename T_TRAITS>
-		void ConstructSphereHeightFieldConstraints(TGeometryParticleHandle<FReal, 3>* Particle0, TGeometryParticleHandle<FReal, 3>* Particle1, const FImplicitObject* Implicit0, const FImplicitObject* Implicit1, const FRigidTransform3& ParticleWorldTransform0, const FRigidTransform3& LocalTransform0, const FRigidTransform3& ParticleWorldTransform1, const FRigidTransform3& LocalTransform1, const FReal CullDistance, const FReal Dt, const FCollisionContext& Context)
+		void ConstructSphereHeightFieldConstraints(TGeometryParticleHandle<FReal, 3>* Particle0, TGeometryParticleHandle<FReal, 3>* Particle1, const FImplicitObject* Implicit0, const FPerShapeData* Shape0, const FImplicitObject* Implicit1, const FPerShapeData* Shape1, const FRigidTransform3& ParticleWorldTransform0, const FRigidTransform3& LocalTransform0, const FRigidTransform3& ParticleWorldTransform1, const FRigidTransform3& LocalTransform1, const FReal CullDistance, const FReal Dt, const FCollisionContext& Context)
 		{
 
 			const TSphere<FReal, 3>* Object0 = Implicit0->template GetObject<const TSphere<FReal, 3> >();
 			const FHeightField* Object1 = Implicit1->template GetObject<const FHeightField >();
 			if (ensure(Object0 && Object1))
 			{
-				FPBDCollisionConstraint* Constraint = Context.CollisionAllocator->FindOrCreateConstraint(Particle0, Implicit0, nullptr, LocalTransform0, Particle1, Implicit1, nullptr, LocalTransform1, CullDistance, EContactShapesType::SphereHeightField, Context.bAllowManifolds);
+				FPBDCollisionConstraint* Constraint = Context.CollisionAllocator->FindOrCreateConstraint(Particle0, Implicit0, Shape0, nullptr, LocalTransform0, Particle1, Implicit1, Shape1, nullptr, LocalTransform1, CullDistance, EContactShapesType::SphereHeightField, Context.bAllowManifolds);
 				if (T_TRAITS::bImmediateUpdate && (Constraint != nullptr))
 				{
 					FRigidTransform3 WorldTransform0 = LocalTransform0 * ParticleWorldTransform0;
@@ -797,13 +797,13 @@ namespace Chaos
 			}
 		}
 
-		void ConstructSphereHeightFieldConstraintsSwept(TGeometryParticleHandle<FReal, 3>* Particle0, TGeometryParticleHandle<FReal, 3>* Particle1, const FImplicitObject* Implicit0, const FImplicitObject* Implicit1, const FRigidTransform3& ParticleWorldTransform0, const FRigidTransform3& LocalTransform0, const FRigidTransform3& ParticleWorldTransform1, const FRigidTransform3& LocalTransform1, const FReal CullDistance, const FReal Dt, const FVec3& Dir, FReal Length, const FCollisionContext& Context)
+		void ConstructSphereHeightFieldConstraintsSwept(TGeometryParticleHandle<FReal, 3>* Particle0, TGeometryParticleHandle<FReal, 3>* Particle1, const FImplicitObject* Implicit0, const FPerShapeData* Shape0, const FImplicitObject* Implicit1, const FPerShapeData* Shape1, const FRigidTransform3& ParticleWorldTransform0, const FRigidTransform3& LocalTransform0, const FRigidTransform3& ParticleWorldTransform1, const FRigidTransform3& LocalTransform1, const FReal CullDistance, const FReal Dt, const FVec3& Dir, FReal Length, const FCollisionContext& Context)
 		{
 			const TSphere<FReal, 3>* Object0 = Implicit0->template GetObject<const TSphere<FReal, 3>>();
 			const FHeightField* Object1 = Implicit1->template GetObject<const FHeightField>();
 			if (ensure(Object0 && Object1))
 			{
-				FPBDCollisionConstraint* Constraint = Context.CollisionAllocator->FindOrCreateSweptConstraint(Particle0, Implicit0, nullptr, LocalTransform0, Particle1, Implicit1, nullptr, LocalTransform1, CullDistance, EContactShapesType::SphereHeightField);
+				FPBDCollisionConstraint* Constraint = Context.CollisionAllocator->FindOrCreateSweptConstraint(Particle0, Implicit0, Shape0, nullptr, LocalTransform0, Particle1, Implicit1, Shape1, nullptr, LocalTransform1, CullDistance, EContactShapesType::SphereHeightField);
 				FRigidTransform3 WorldTransformX0 = LocalTransform0 * FRigidTransform3(Particle0->X(), Particle0->R());
 				FRigidTransform3 WorldTransform1 = LocalTransform1 * ParticleWorldTransform1;
 				UpdateSphereHeightFieldConstraintSwept(Particle0, *Object0, WorldTransformX0, *Object1, WorldTransform1, Dir, Length, Dt, *Constraint);
@@ -827,14 +827,14 @@ namespace Chaos
 		}
 
 		template<typename T_TRAITS>
-		void ConstructSpherePlaneConstraints(TGeometryParticleHandle<FReal, 3>* Particle0, TGeometryParticleHandle<FReal, 3>* Particle1, const FImplicitObject* Implicit0, const FImplicitObject* Implicit1, const FRigidTransform3& ParticleWorldTransform0, const FRigidTransform3& LocalTransform0, const FRigidTransform3& ParticleWorldTransform1, const FRigidTransform3& LocalTransform1, const FReal CullDistance, const FReal Dt, const FCollisionContext& Context)
+		void ConstructSpherePlaneConstraints(TGeometryParticleHandle<FReal, 3>* Particle0, TGeometryParticleHandle<FReal, 3>* Particle1, const FImplicitObject* Implicit0, const FPerShapeData* Shape0, const FImplicitObject* Implicit1, const FPerShapeData* Shape1, const FRigidTransform3& ParticleWorldTransform0, const FRigidTransform3& LocalTransform0, const FRigidTransform3& ParticleWorldTransform1, const FRigidTransform3& LocalTransform1, const FReal CullDistance, const FReal Dt, const FCollisionContext& Context)
 		{
 
 			const TSphere<FReal, 3>* Object0 = Implicit0->template GetObject<const TSphere<FReal, 3> >();
 			const TPlane<FReal, 3>* Object1 = Implicit1->template GetObject<const TPlane<FReal, 3> >();
 			if (ensure(Object0 && Object1))
 			{
-				FPBDCollisionConstraint* Constraint = Context.CollisionAllocator->FindOrCreateConstraint(Particle0, Implicit0, nullptr, LocalTransform0, Particle1, Implicit1, nullptr, LocalTransform1, CullDistance, EContactShapesType::SpherePlane, Context.bAllowManifolds);
+				FPBDCollisionConstraint* Constraint = Context.CollisionAllocator->FindOrCreateConstraint(Particle0, Implicit0, Shape0, nullptr, LocalTransform0, Particle1, Implicit1, Shape1, nullptr, LocalTransform1, CullDistance, EContactShapesType::SpherePlane, Context.bAllowManifolds);
 				if (T_TRAITS::bImmediateUpdate && (Constraint != nullptr))
 				{
 					FRigidTransform3 WorldTransform0 = LocalTransform0 * ParticleWorldTransform0;
@@ -862,14 +862,14 @@ namespace Chaos
 		}
 
 		template<typename T_TRAITS>
-		void ConstructSphereBoxConstraints(TGeometryParticleHandle<FReal, 3>* Particle0, TGeometryParticleHandle<FReal, 3>* Particle1, const FImplicitObject* Implicit0, const FImplicitObject* Implicit1, const FRigidTransform3& ParticleWorldTransform0, const FRigidTransform3& LocalTransform0, const FRigidTransform3& ParticleWorldTransform1, const FRigidTransform3& LocalTransform1, const FReal CullDistance, const FReal Dt, const FCollisionContext& Context)
+		void ConstructSphereBoxConstraints(TGeometryParticleHandle<FReal, 3>* Particle0, TGeometryParticleHandle<FReal, 3>* Particle1, const FImplicitObject* Implicit0, const FPerShapeData* Shape0, const FImplicitObject* Implicit1, const FPerShapeData* Shape1, const FRigidTransform3& ParticleWorldTransform0, const FRigidTransform3& LocalTransform0, const FRigidTransform3& ParticleWorldTransform1, const FRigidTransform3& LocalTransform1, const FReal CullDistance, const FReal Dt, const FCollisionContext& Context)
 		{
 
 			const TSphere<FReal, 3>* Object0 = Implicit0->template GetObject<const TSphere<FReal, 3> >();
 			const TBox<FReal, 3>* Object1 = Implicit1->template GetObject<const TBox<FReal, 3> >();
 			if (ensure(Object0 && Object1))
 			{
-				FPBDCollisionConstraint* Constraint = Context.CollisionAllocator->FindOrCreateConstraint(Particle0, Implicit0, nullptr, LocalTransform0, Particle1, Implicit1, nullptr, LocalTransform1, CullDistance, EContactShapesType::SphereBox, Context.bAllowManifolds);
+				FPBDCollisionConstraint* Constraint = Context.CollisionAllocator->FindOrCreateConstraint(Particle0, Implicit0, Shape0, nullptr, LocalTransform0, Particle1, Implicit1, Shape1, nullptr, LocalTransform1, CullDistance, EContactShapesType::SphereBox, Context.bAllowManifolds);
 				if (T_TRAITS::bImmediateUpdate && (Constraint != nullptr))
 				{
 					FRigidTransform3 WorldTransform0 = LocalTransform0 * ParticleWorldTransform0;
@@ -898,13 +898,13 @@ namespace Chaos
 		}
 
 		template<typename T_TRAITS>
-		void ConstructSphereCapsuleConstraints(TGeometryParticleHandle<FReal, 3>* Particle0, TGeometryParticleHandle<FReal, 3>* Particle1, const FImplicitObject* Implicit0, const FImplicitObject* Implicit1, const FRigidTransform3& ParticleWorldTransform0, const FRigidTransform3& LocalTransform0, const FRigidTransform3& ParticleWorldTransform1, const FRigidTransform3& LocalTransform1, const FReal CullDistance, const FReal Dt, const FCollisionContext& Context)
+		void ConstructSphereCapsuleConstraints(TGeometryParticleHandle<FReal, 3>* Particle0, TGeometryParticleHandle<FReal, 3>* Particle1, const FImplicitObject* Implicit0, const FPerShapeData* Shape0, const FImplicitObject* Implicit1, const FPerShapeData* Shape1, const FRigidTransform3& ParticleWorldTransform0, const FRigidTransform3& LocalTransform0, const FRigidTransform3& ParticleWorldTransform1, const FRigidTransform3& LocalTransform1, const FReal CullDistance, const FReal Dt, const FCollisionContext& Context)
 		{
 			const TSphere<FReal, 3>* Object0 = Implicit0->template GetObject<const TSphere<FReal, 3> >();
 			const FCapsule* Object1 = Implicit1->template GetObject<const FCapsule >();
 			if (ensure(Object0 && Object1))
 			{
-				FPBDCollisionConstraint* Constraint = Context.CollisionAllocator->FindOrCreateConstraint(Particle0, Implicit0, nullptr, LocalTransform0, Particle1, Implicit1, nullptr, LocalTransform1, CullDistance, EContactShapesType::SphereCapsule, Context.bAllowManifolds);
+				FPBDCollisionConstraint* Constraint = Context.CollisionAllocator->FindOrCreateConstraint(Particle0, Implicit0, Shape0, nullptr, LocalTransform0, Particle1, Implicit1, Shape1, nullptr, LocalTransform1, CullDistance, EContactShapesType::SphereCapsule, Context.bAllowManifolds);
 				if (T_TRAITS::bImmediateUpdate && (Constraint != nullptr))
 				{
 					FRigidTransform3 WorldTransform0 = LocalTransform0 * ParticleWorldTransform0;
@@ -931,13 +931,13 @@ namespace Chaos
 		}
 
 		template<typename T_TRAITS>
-		void ConstructSphereConvexConstraints(TGeometryParticleHandle<FReal, 3>* Particle0, TGeometryParticleHandle<FReal, 3>* Particle1, const FImplicitObject* Implicit0, const FImplicitObject* Implicit1, const FRigidTransform3& ParticleWorldTransform0, const FRigidTransform3& LocalTransform0, const FRigidTransform3& ParticleWorldTransform1, const FRigidTransform3& LocalTransform1, const FReal CullDistance, const FReal Dt, const FCollisionContext& Context)
+		void ConstructSphereConvexConstraints(TGeometryParticleHandle<FReal, 3>* Particle0, TGeometryParticleHandle<FReal, 3>* Particle1, const FImplicitObject* Implicit0, const FPerShapeData* Shape0, const FImplicitObject* Implicit1, const FPerShapeData* Shape1, const FRigidTransform3& ParticleWorldTransform0, const FRigidTransform3& LocalTransform0, const FRigidTransform3& ParticleWorldTransform1, const FRigidTransform3& LocalTransform1, const FReal CullDistance, const FReal Dt, const FCollisionContext& Context)
 		{
 			const TSphere<FReal, 3>* Object0 = Implicit0->template GetObject<const TSphere<FReal, 3> >();
 			const FImplicitObject* Object1 = Implicit1;
 			if (ensure(Object0 && Object1))
 			{
-				FPBDCollisionConstraint* Constraint = Context.CollisionAllocator->FindOrCreateConstraint(Particle0, Implicit0, nullptr, LocalTransform0, Particle1, Implicit1, nullptr, LocalTransform1, CullDistance, EContactShapesType::SphereConvex, Context.bAllowManifolds);
+				FPBDCollisionConstraint* Constraint = Context.CollisionAllocator->FindOrCreateConstraint(Particle0, Implicit0, Shape0, nullptr, LocalTransform0, Particle1, Implicit1, Shape1, nullptr, LocalTransform1, CullDistance, EContactShapesType::SphereConvex, Context.bAllowManifolds);
 				if (T_TRAITS::bImmediateUpdate && (Constraint != nullptr))
 				{
 					FRigidTransform3 WorldTransform0 = LocalTransform0 * ParticleWorldTransform0;
@@ -974,14 +974,14 @@ namespace Chaos
 		}
 
 		template<typename T_TRAITS>
-		void ConstructSphereTriangleMeshConstraints(TGeometryParticleHandle<FReal, 3>* Particle0, TGeometryParticleHandle<FReal, 3>* Particle1, const FImplicitObject* Implicit0, const FImplicitObject* Implicit1, const FRigidTransform3& ParticleWorldTransform0, const FRigidTransform3& LocalTransform0, const FRigidTransform3& ParticleWorldTransform1, const FRigidTransform3& LocalTransform1, const FReal CullDistance, const FReal Dt, const FCollisionContext& Context)
+		void ConstructSphereTriangleMeshConstraints(TGeometryParticleHandle<FReal, 3>* Particle0, TGeometryParticleHandle<FReal, 3>* Particle1, const FImplicitObject* Implicit0, const FPerShapeData* Shape0, const FImplicitObject* Implicit1, const FPerShapeData* Shape1, const FRigidTransform3& ParticleWorldTransform0, const FRigidTransform3& LocalTransform0, const FRigidTransform3& ParticleWorldTransform1, const FRigidTransform3& LocalTransform1, const FReal CullDistance, const FReal Dt, const FCollisionContext& Context)
 		{
 			const TSphere<FReal, 3>* Object0 = Implicit0->template GetObject<const TSphere<FReal, 3> >();
 			if (ensure(Object0))
 			{
 				if (const TImplicitObjectScaled<FTriangleMeshImplicitObject>* ScaledTriangleMesh = Implicit1->template GetObject<const TImplicitObjectScaled<FTriangleMeshImplicitObject>>())
 				{
-					FPBDCollisionConstraint* Constraint = Context.CollisionAllocator->FindOrCreateConstraint(Particle0, Implicit0, nullptr, LocalTransform0, Particle1, Implicit1, nullptr, LocalTransform1, CullDistance, EContactShapesType::SphereTriMesh, Context.bAllowManifolds);
+					FPBDCollisionConstraint* Constraint = Context.CollisionAllocator->FindOrCreateConstraint(Particle0, Implicit0, Shape0, nullptr, LocalTransform0, Particle1, Implicit1, Shape1, nullptr, LocalTransform1, CullDistance, EContactShapesType::SphereTriMesh, Context.bAllowManifolds);
 					if (T_TRAITS::bImmediateUpdate && (Constraint != nullptr))
 					{
 						FRigidTransform3 WorldTransform0 = LocalTransform0 * ParticleWorldTransform0;
@@ -991,7 +991,7 @@ namespace Chaos
 				}
 				else if (const FTriangleMeshImplicitObject* TriangleMesh = Implicit1->template GetObject<const FTriangleMeshImplicitObject>())
 				{
-					FPBDCollisionConstraint* Constraint = Context.CollisionAllocator->FindOrCreateConstraint(Particle0, Implicit0, nullptr, LocalTransform0, Particle1, Implicit1, nullptr, LocalTransform1, CullDistance, EContactShapesType::SphereTriMesh, Context.bAllowManifolds);
+					FPBDCollisionConstraint* Constraint = Context.CollisionAllocator->FindOrCreateConstraint(Particle0, Implicit0, Shape0, nullptr, LocalTransform0, Particle1, Implicit1, Shape1, nullptr, LocalTransform1, CullDistance, EContactShapesType::SphereTriMesh, Context.bAllowManifolds);
 					if (T_TRAITS::bImmediateUpdate && (Constraint != nullptr))
 					{
 						FRigidTransform3 WorldTransform0 = LocalTransform0 * ParticleWorldTransform0;
@@ -1007,21 +1007,21 @@ namespace Chaos
 		}
 
 
-		void ConstructSphereTriangleMeshConstraintsSwept(TGeometryParticleHandle<FReal, 3>* Particle0, TGeometryParticleHandle<FReal, 3>* Particle1, const FImplicitObject* Implicit0, const FImplicitObject* Implicit1, const FRigidTransform3& ParticleWorldTransform0, const FRigidTransform3& LocalTransform0, const FRigidTransform3& ParticleWorldTransform1, const FRigidTransform3& LocalTransform1, const FReal CullDistance, const FReal Dt, const FVec3& Dir, FReal Length, const FCollisionContext& Context)
+		void ConstructSphereTriangleMeshConstraintsSwept(TGeometryParticleHandle<FReal, 3>* Particle0, TGeometryParticleHandle<FReal, 3>* Particle1, const FImplicitObject* Implicit0, const FPerShapeData* Shape0, const FImplicitObject* Implicit1, const FPerShapeData* Shape1, const FRigidTransform3& ParticleWorldTransform0, const FRigidTransform3& LocalTransform0, const FRigidTransform3& ParticleWorldTransform1, const FRigidTransform3& LocalTransform1, const FReal CullDistance, const FReal Dt, const FVec3& Dir, FReal Length, const FCollisionContext& Context)
 		{
 			const TSphere<FReal, 3>* Object0 = Implicit0->template GetObject<const TSphere<FReal, 3>>();
 			if (ensure(Object0))
 			{
 				if (const TImplicitObjectScaled<FTriangleMeshImplicitObject>* ScaledTriangleMesh = Implicit1->template GetObject<const TImplicitObjectScaled<FTriangleMeshImplicitObject>>())
 				{
-					FPBDCollisionConstraint* Constraint = Context.CollisionAllocator->FindOrCreateSweptConstraint(Particle0, Implicit0, nullptr, LocalTransform0, Particle1, Implicit1, nullptr, LocalTransform1, CullDistance, EContactShapesType::SphereTriMesh);
+					FPBDCollisionConstraint* Constraint = Context.CollisionAllocator->FindOrCreateSweptConstraint(Particle0, Implicit0, Shape0, nullptr, LocalTransform0, Particle1, Implicit1, Shape1, nullptr, LocalTransform1, CullDistance, EContactShapesType::SphereTriMesh);
 					FRigidTransform3 WorldTransformX0 = LocalTransform0 * FRigidTransform3(Particle0->X(), Particle0->R());
 					FRigidTransform3 WorldTransform1 = LocalTransform1 * ParticleWorldTransform1;
 					UpdateSphereTriangleMeshConstraintSwept(Particle0, *Object0, WorldTransformX0, *ScaledTriangleMesh, WorldTransform1, Dir, Length, Dt, *Constraint);
 				}
 				else if (const FTriangleMeshImplicitObject* TriangleMesh = Implicit1->template GetObject<const FTriangleMeshImplicitObject>())
 				{
-					FPBDCollisionConstraint* Constraint = Context.CollisionAllocator->FindOrCreateSweptConstraint(Particle0, Implicit0, nullptr, LocalTransform0, Particle1, Implicit1, nullptr, LocalTransform1, CullDistance, EContactShapesType::SphereTriMesh);
+					FPBDCollisionConstraint* Constraint = Context.CollisionAllocator->FindOrCreateSweptConstraint(Particle0, Implicit0, Shape0, nullptr, LocalTransform0, Particle1, Implicit1, Shape1, nullptr, LocalTransform1, CullDistance, EContactShapesType::SphereTriMesh);
 					FRigidTransform3 WorldTransformX0 = LocalTransform0 * FRigidTransform3(Particle0->X(), Particle0->R());
 					FRigidTransform3 WorldTransform1 = LocalTransform1 * ParticleWorldTransform1;
 					UpdateSphereTriangleMeshConstraintSwept(Particle0, *Object0, WorldTransformX0, *TriangleMesh, WorldTransform1, Dir, Length, Dt, *Constraint);
@@ -1053,13 +1053,13 @@ namespace Chaos
 		}
 
 		template<typename T_TRAITS>
-		void ConstructCapsuleCapsuleConstraints(TGeometryParticleHandle<FReal, 3>* Particle0, TGeometryParticleHandle<FReal, 3>* Particle1, const FImplicitObject* Implicit0, const FImplicitObject* Implicit1, const FRigidTransform3& ParticleWorldTransform0, const FRigidTransform3& LocalTransform0, const FRigidTransform3& ParticleWorldTransform1, const FRigidTransform3& LocalTransform1, const FReal CullDistance, const FReal Dt, const FCollisionContext& Context)
+		void ConstructCapsuleCapsuleConstraints(TGeometryParticleHandle<FReal, 3>* Particle0, TGeometryParticleHandle<FReal, 3>* Particle1, const FImplicitObject* Implicit0, const FPerShapeData* Shape0, const FImplicitObject* Implicit1, const FPerShapeData* Shape1, const FRigidTransform3& ParticleWorldTransform0, const FRigidTransform3& LocalTransform0, const FRigidTransform3& ParticleWorldTransform1, const FRigidTransform3& LocalTransform1, const FReal CullDistance, const FReal Dt, const FCollisionContext& Context)
 		{
 			const FCapsule* Object0 = Implicit0->template GetObject<const FCapsule >();
 			const FCapsule* Object1 = Implicit1->template GetObject<const FCapsule >();
 			if (ensure(Object0 && Object1))
 			{
-				FPBDCollisionConstraint* Constraint = Context.CollisionAllocator->FindOrCreateConstraint(Particle0, Implicit0, nullptr, LocalTransform0, Particle1, Implicit1, nullptr, LocalTransform1, CullDistance, EContactShapesType::CapsuleCapsule, Context.bAllowManifolds);
+				FPBDCollisionConstraint* Constraint = Context.CollisionAllocator->FindOrCreateConstraint(Particle0, Implicit0, Shape0, nullptr, LocalTransform0, Particle1, Implicit1, Shape1, nullptr, LocalTransform1, CullDistance, EContactShapesType::CapsuleCapsule, Context.bAllowManifolds);
 				if (T_TRAITS::bImmediateUpdate && (Constraint != nullptr))
 				{
 					FRigidTransform3 WorldTransform0 = LocalTransform0 * ParticleWorldTransform0;
@@ -1094,13 +1094,13 @@ namespace Chaos
 		}
 
 		template<typename T_TRAITS>
-		void ConstructCapsuleBoxConstraints(TGeometryParticleHandle<FReal, 3>* Particle0, TGeometryParticleHandle<FReal, 3>* Particle1, const FImplicitObject* Implicit0, const FImplicitObject* Implicit1, const FRigidTransform3& ParticleWorldTransform0, const FRigidTransform3& LocalTransform0, const FRigidTransform3& ParticleWorldTransform1, const FRigidTransform3& LocalTransform1, const FReal CullDistance, const FReal Dt, const FCollisionContext& Context)
+		void ConstructCapsuleBoxConstraints(TGeometryParticleHandle<FReal, 3>* Particle0, TGeometryParticleHandle<FReal, 3>* Particle1, const FImplicitObject* Implicit0, const FPerShapeData* Shape0, const FImplicitObject* Implicit1, const FPerShapeData* Shape1, const FRigidTransform3& ParticleWorldTransform0, const FRigidTransform3& LocalTransform0, const FRigidTransform3& ParticleWorldTransform1, const FRigidTransform3& LocalTransform1, const FReal CullDistance, const FReal Dt, const FCollisionContext& Context)
 		{
 			const FCapsule* Object0 = Implicit0->template GetObject<const FCapsule>();
 			const TBox<FReal, 3>* Object1 = Implicit1->template GetObject<const TBox<FReal, 3>>();
 			if (ensure(Object0 && Object1))
 			{
-				FPBDCollisionConstraint* Constraint = Context.CollisionAllocator->FindOrCreateConstraint(Particle0, Implicit0, nullptr, LocalTransform0, Particle1, Implicit1, nullptr, LocalTransform1, CullDistance, EContactShapesType::CapsuleBox, Context.bAllowManifolds);
+				FPBDCollisionConstraint* Constraint = Context.CollisionAllocator->FindOrCreateConstraint(Particle0, Implicit0, Shape0, nullptr, LocalTransform0, Particle1, Implicit1, Shape1, nullptr, LocalTransform1, CullDistance, EContactShapesType::CapsuleBox, Context.bAllowManifolds);
 				if (T_TRAITS::bImmediateUpdate && (Constraint != nullptr))
 				{
 					FRigidTransform3 WorldTransform0 = LocalTransform0 * ParticleWorldTransform0;
@@ -1176,13 +1176,13 @@ namespace Chaos
 		}
 
 		template<typename T_TRAITS>
-		void ConstructCapsuleConvexConstraints(TGeometryParticleHandle<FReal, 3>* Particle0, TGeometryParticleHandle<FReal, 3>* Particle1, const FImplicitObject* Implicit0, const FImplicitObject* Implicit1, const FRigidTransform3& ParticleWorldTransform0, const FRigidTransform3& LocalTransform0, const FRigidTransform3& ParticleWorldTransform1, const FRigidTransform3& LocalTransform1, const FReal CullDistance, const FReal Dt, const FCollisionContext& Context)
+		void ConstructCapsuleConvexConstraints(TGeometryParticleHandle<FReal, 3>* Particle0, TGeometryParticleHandle<FReal, 3>* Particle1, const FImplicitObject* Implicit0, const FPerShapeData* Shape0, const FImplicitObject* Implicit1, const FPerShapeData* Shape1, const FRigidTransform3& ParticleWorldTransform0, const FRigidTransform3& LocalTransform0, const FRigidTransform3& ParticleWorldTransform1, const FRigidTransform3& LocalTransform1, const FReal CullDistance, const FReal Dt, const FCollisionContext& Context)
 		{
 			const FCapsule* Object0 = Implicit0->template GetObject<const FCapsule>();
 			const FImplicitObject* Object1 = Implicit1;
 			if (ensure(Object0 && Object1))
 			{
-				FPBDCollisionConstraint* Constraint = Context.CollisionAllocator->FindOrCreateConstraint(Particle0, Implicit0, nullptr, LocalTransform0, Particle1, Implicit1, nullptr, LocalTransform1, CullDistance, EContactShapesType::CapsuleConvex, Context.bAllowManifolds);
+				FPBDCollisionConstraint* Constraint = Context.CollisionAllocator->FindOrCreateConstraint(Particle0, Implicit0, Shape0, nullptr, LocalTransform0, Particle1, Implicit1, Shape1, nullptr, LocalTransform1, CullDistance, EContactShapesType::CapsuleConvex, Context.bAllowManifolds);
 				if (T_TRAITS::bImmediateUpdate && (Constraint != nullptr))
 				{
 					FRigidTransform3 WorldTransform0 = LocalTransform0 * ParticleWorldTransform0;
@@ -1221,7 +1221,7 @@ namespace Chaos
 		}
 
 		template<typename T_TRAITS>
-		void ConstructCapsuleHeightFieldConstraints(TGeometryParticleHandle<FReal, 3>* Particle0, TGeometryParticleHandle<FReal, 3>* Particle1, const FImplicitObject* Implicit0, const FImplicitObject* Implicit1, const FRigidTransform3& ParticleWorldTransform0, const FRigidTransform3& LocalTransform0, const FRigidTransform3& ParticleWorldTransform1, const FRigidTransform3& LocalTransform1, const FReal CullDistance, const FReal Dt, const FCollisionContext& Context)
+		void ConstructCapsuleHeightFieldConstraints(TGeometryParticleHandle<FReal, 3>* Particle0, TGeometryParticleHandle<FReal, 3>* Particle1, const FImplicitObject* Implicit0, const FPerShapeData* Shape0, const FImplicitObject* Implicit1, const FPerShapeData* Shape1, const FRigidTransform3& ParticleWorldTransform0, const FRigidTransform3& LocalTransform0, const FRigidTransform3& ParticleWorldTransform1, const FRigidTransform3& LocalTransform1, const FReal CullDistance, const FReal Dt, const FCollisionContext& Context)
 		{
 			CONDITIONAL_SCOPE_CYCLE_COUNTER(STAT_Collisions_ConstructCapsuleHeightFieldConstraints, ConstraintsDetailedStats);
 
@@ -1229,7 +1229,7 @@ namespace Chaos
 			const FHeightField* Object1 = Implicit1->template GetObject<const FHeightField >();
 			if (ensure(Object0 && Object1))
 			{
-				FPBDCollisionConstraint* Constraint = Context.CollisionAllocator->FindOrCreateConstraint(Particle0, Implicit0, nullptr, LocalTransform0, Particle1, Implicit1, nullptr, LocalTransform1, CullDistance, EContactShapesType::CapsuleHeightField, Context.bAllowManifolds);
+				FPBDCollisionConstraint* Constraint = Context.CollisionAllocator->FindOrCreateConstraint(Particle0, Implicit0, Shape0, nullptr, LocalTransform0, Particle1, Implicit1, Shape1, nullptr, LocalTransform1, CullDistance, EContactShapesType::CapsuleHeightField, Context.bAllowManifolds);
 				if (T_TRAITS::bImmediateUpdate && (Constraint != nullptr))
 				{
 					FRigidTransform3 WorldTransform0 = LocalTransform0 * ParticleWorldTransform0;
@@ -1239,14 +1239,14 @@ namespace Chaos
 			}
 		}
 
-		void ConstructCapsuleHeightFieldConstraintsSwept(TGeometryParticleHandle<FReal, 3>* Particle0, TGeometryParticleHandle<FReal, 3>* Particle1, const FImplicitObject* Implicit0, const FImplicitObject* Implicit1, const FRigidTransform3& ParticleWorldTransform0, const FRigidTransform3& LocalTransform0, const FRigidTransform3& ParticleWorldTransform1, const FRigidTransform3& LocalTransform1, const FReal CullDistance, const FReal Dt, const FVec3& Dir, FReal Length, const FCollisionContext& Context)
+		void ConstructCapsuleHeightFieldConstraintsSwept(TGeometryParticleHandle<FReal, 3>* Particle0, TGeometryParticleHandle<FReal, 3>* Particle1, const FImplicitObject* Implicit0, const FPerShapeData* Shape0, const FImplicitObject* Implicit1, const FPerShapeData* Shape1, const FRigidTransform3& ParticleWorldTransform0, const FRigidTransform3& LocalTransform0, const FRigidTransform3& ParticleWorldTransform1, const FRigidTransform3& LocalTransform1, const FReal CullDistance, const FReal Dt, const FVec3& Dir, FReal Length, const FCollisionContext& Context)
 		{
 			CONDITIONAL_SCOPE_CYCLE_COUNTER(STAT_Collisions_ConstructCapsuleHeightFieldConstraintsSwept, ConstraintsDetailedStats);
 			const FCapsule* Object0 = Implicit0->template GetObject<const FCapsule >();
 			const FHeightField* Object1 = Implicit1->template GetObject<const FHeightField >();
 			if (ensure(Object0 && Object1))
 			{
-				FPBDCollisionConstraint* Constraint = Context.CollisionAllocator->FindOrCreateSweptConstraint(Particle0, Implicit0, nullptr, LocalTransform0, Particle1, Implicit1, nullptr, LocalTransform1, CullDistance, EContactShapesType::CapsuleHeightField);
+				FPBDCollisionConstraint* Constraint = Context.CollisionAllocator->FindOrCreateSweptConstraint(Particle0, Implicit0, Shape0, nullptr, LocalTransform0, Particle1, Implicit1, Shape1, nullptr, LocalTransform1, CullDistance, EContactShapesType::CapsuleHeightField);
 				FRigidTransform3 WorldTransformX0 = LocalTransform0 * FRigidTransform3(Particle0->X(), Particle0->R());
 				FRigidTransform3 WorldTransform1 = LocalTransform1 * ParticleWorldTransform1;
 				UpdateCapsuleHeightFieldConstraintSwept(Particle0, *Object0, WorldTransformX0, *Object1, WorldTransform1, Dir, Length, Dt, *Constraint);
@@ -1284,7 +1284,7 @@ namespace Chaos
 		}
 
 		template<typename T_TRAITS>
-		void ConstructCapsuleTriangleMeshConstraints(TGeometryParticleHandle<FReal, 3>* Particle0, TGeometryParticleHandle<FReal, 3>* Particle1, const FImplicitObject* Implicit0, const FImplicitObject* Implicit1, const FRigidTransform3& ParticleWorldTransform0, const FRigidTransform3& LocalTransform0, const FRigidTransform3& ParticleWorldTransform1, const FRigidTransform3& LocalTransform1, const FReal CullDistance, const FReal Dt, const FCollisionContext& Context)
+		void ConstructCapsuleTriangleMeshConstraints(TGeometryParticleHandle<FReal, 3>* Particle0, TGeometryParticleHandle<FReal, 3>* Particle1, const FImplicitObject* Implicit0, const FPerShapeData* Shape0, const FImplicitObject* Implicit1, const FPerShapeData* Shape1, const FRigidTransform3& ParticleWorldTransform0, const FRigidTransform3& LocalTransform0, const FRigidTransform3& ParticleWorldTransform1, const FRigidTransform3& LocalTransform1, const FReal CullDistance, const FReal Dt, const FCollisionContext& Context)
 		{
 			CONDITIONAL_SCOPE_CYCLE_COUNTER(STAT_Collisions_ConstructCapsuleTriangleMeshConstraints, ConstraintsDetailedStats);
 
@@ -1293,7 +1293,7 @@ namespace Chaos
 			{
 				if (const TImplicitObjectScaled<FTriangleMeshImplicitObject>* ScaledTriangleMesh = Implicit1->template GetObject<const TImplicitObjectScaled<FTriangleMeshImplicitObject>>())
 				{
-					FPBDCollisionConstraint* Constraint = Context.CollisionAllocator->FindOrCreateConstraint(Particle0, Implicit0, nullptr, LocalTransform0, Particle1, Implicit1, nullptr, LocalTransform1, CullDistance, EContactShapesType::CapsuleTriMesh, Context.bAllowManifolds);
+					FPBDCollisionConstraint* Constraint = Context.CollisionAllocator->FindOrCreateConstraint(Particle0, Implicit0, Shape0, nullptr, LocalTransform0, Particle1, Implicit1, Shape1,  nullptr, LocalTransform1, CullDistance, EContactShapesType::CapsuleTriMesh, Context.bAllowManifolds);
 					if (T_TRAITS::bImmediateUpdate && (Constraint != nullptr))
 					{
 						FRigidTransform3 WorldTransform0 = LocalTransform0 * ParticleWorldTransform0;
@@ -1303,7 +1303,7 @@ namespace Chaos
 				}
 				else if (const FTriangleMeshImplicitObject* TriangleMesh = Implicit1->template GetObject<const FTriangleMeshImplicitObject>())
 				{
-					FPBDCollisionConstraint* Constraint = Context.CollisionAllocator->FindOrCreateConstraint(Particle0, Implicit0, nullptr, LocalTransform0, Particle1, Implicit1, nullptr, LocalTransform1, CullDistance, EContactShapesType::CapsuleTriMesh, Context.bAllowManifolds);
+					FPBDCollisionConstraint* Constraint = Context.CollisionAllocator->FindOrCreateConstraint(Particle0, Implicit0, Shape0, nullptr, LocalTransform0, Particle1, Implicit1, Shape1, nullptr, LocalTransform1, CullDistance, EContactShapesType::CapsuleTriMesh, Context.bAllowManifolds);
 					if (T_TRAITS::bImmediateUpdate && (Constraint != nullptr))
 					{
 						FRigidTransform3 WorldTransform0 = LocalTransform0 * ParticleWorldTransform0;
@@ -1318,7 +1318,7 @@ namespace Chaos
 			}
 		}
 
-		void ConstructCapsuleTriangleMeshConstraintsSwept(TGeometryParticleHandle<FReal, 3>* Particle0, TGeometryParticleHandle<FReal, 3>* Particle1, const FImplicitObject* Implicit0, const FImplicitObject* Implicit1, const FRigidTransform3& ParticleWorldTransform0, const FRigidTransform3& LocalTransform0, const FRigidTransform3& ParticleWorldTransform1, const FRigidTransform3& LocalTransform1, const FReal CullDistance, const FReal Dt, const FVec3& Dir, FReal Length, const FCollisionContext& Context)
+		void ConstructCapsuleTriangleMeshConstraintsSwept(TGeometryParticleHandle<FReal, 3>* Particle0, TGeometryParticleHandle<FReal, 3>* Particle1, const FImplicitObject* Implicit0, const FPerShapeData* Shape0, const FImplicitObject* Implicit1, const FPerShapeData* Shape1, const FRigidTransform3& ParticleWorldTransform0, const FRigidTransform3& LocalTransform0, const FRigidTransform3& ParticleWorldTransform1, const FRigidTransform3& LocalTransform1, const FReal CullDistance, const FReal Dt, const FVec3& Dir, FReal Length, const FCollisionContext& Context)
 		{
 			CONDITIONAL_SCOPE_CYCLE_COUNTER(STAT_Collisions_ConstructCapsuleTriangleMeshConstraintsSwept, ConstraintsDetailedStats);
 
@@ -1327,14 +1327,14 @@ namespace Chaos
 			{
 				if (const TImplicitObjectScaled<FTriangleMeshImplicitObject>* ScaledTriangleMesh = Implicit1->template GetObject<const TImplicitObjectScaled<FTriangleMeshImplicitObject>>())
 				{
-					FPBDCollisionConstraint* Constraint = Context.CollisionAllocator->FindOrCreateSweptConstraint(Particle0, Implicit0, nullptr, LocalTransform0, Particle1, Implicit1, nullptr, LocalTransform1, CullDistance, EContactShapesType::CapsuleTriMesh);
+					FPBDCollisionConstraint* Constraint = Context.CollisionAllocator->FindOrCreateSweptConstraint(Particle0, Implicit0, Shape0, nullptr, LocalTransform0, Particle1, Implicit1, Shape1, nullptr, LocalTransform1, CullDistance, EContactShapesType::CapsuleTriMesh);
 					FRigidTransform3 WorldTransformX0 = LocalTransform0 * FRigidTransform3(Particle0->X(), Particle0->R());
 					FRigidTransform3 WorldTransform1 = LocalTransform1 * ParticleWorldTransform1;
 					UpdateCapsuleTriangleMeshConstraintSwept(Particle0, *Object0, WorldTransformX0, *ScaledTriangleMesh, WorldTransform1, Dir, Length, Dt, *Constraint);
 				}
 				else if (const FTriangleMeshImplicitObject* TriangleMesh = Implicit1->template GetObject<const FTriangleMeshImplicitObject>())
 				{
-					FPBDCollisionConstraint* Constraint = Context.CollisionAllocator->FindOrCreateSweptConstraint(Particle0, Implicit0, nullptr, LocalTransform0, Particle1, Implicit1, nullptr, LocalTransform1, CullDistance, EContactShapesType::CapsuleTriMesh);
+					FPBDCollisionConstraint* Constraint = Context.CollisionAllocator->FindOrCreateSweptConstraint(Particle0, Implicit0, Shape0, nullptr, LocalTransform0, Particle1, Implicit1, Shape1, nullptr, LocalTransform1, CullDistance, EContactShapesType::CapsuleTriMesh);
 					FRigidTransform3 WorldTransformX0 = LocalTransform0 * FRigidTransform3(Particle0->X(), Particle0->R());
 					FRigidTransform3 WorldTransform1 = LocalTransform1 * ParticleWorldTransform1;
 					UpdateCapsuleTriangleMeshConstraintSwept(Particle0, *Object0, WorldTransformX0, *TriangleMesh, WorldTransform1, Dir, Length, Dt, *Constraint);
@@ -1369,13 +1369,13 @@ namespace Chaos
 
 
 		template<typename T_TRAITS>
-		void ConstructGenericConvexConvexConstraints(TGeometryParticleHandle<FReal, 3>* Particle0, TGeometryParticleHandle<FReal, 3>* Particle1, const FImplicitObject* Implicit0, const FImplicitObject* Implicit1, const FRigidTransform3& ParticleWorldTransform0, const FRigidTransform3& LocalTransform0, const FRigidTransform3& ParticleWorldTransform1, const FRigidTransform3& LocalTransform1, const FReal CullDistance, const FReal Dt, const FCollisionContext& Context)
+		void ConstructGenericConvexConvexConstraints(TGeometryParticleHandle<FReal, 3>* Particle0, TGeometryParticleHandle<FReal, 3>* Particle1, const FImplicitObject* Implicit0, const FPerShapeData* Shape0, const FImplicitObject* Implicit1, const FPerShapeData* Shape1, const FRigidTransform3& ParticleWorldTransform0, const FRigidTransform3& LocalTransform0, const FRigidTransform3& ParticleWorldTransform1, const FRigidTransform3& LocalTransform1, const FReal CullDistance, const FReal Dt, const FCollisionContext& Context)
 		{
 			CONDITIONAL_SCOPE_CYCLE_COUNTER(STAT_Collisions_ConstructGenericConvexConvexConstraints, ConstraintsDetailedStats);
 			EImplicitObjectType Implicit0Type = Particle0->Geometry()->GetType();
 			EImplicitObjectType Implicit1Type = Particle1->Geometry()->GetType();
 
-			FPBDCollisionConstraint* Constraint = Context.CollisionAllocator->FindOrCreateConstraint(Particle0, Implicit0, nullptr, LocalTransform0, Particle1, Implicit1, nullptr, LocalTransform1, CullDistance, EContactShapesType::GenericConvexConvex, Context.bAllowManifolds);
+			FPBDCollisionConstraint* Constraint = Context.CollisionAllocator->FindOrCreateConstraint(Particle0, Implicit0, Shape0, nullptr, LocalTransform0, Particle1, Implicit1, Shape1, nullptr, LocalTransform1, CullDistance, EContactShapesType::GenericConvexConvex, Context.bAllowManifolds);
 			if (T_TRAITS::bImmediateUpdate && (Constraint != nullptr))
 			{
 				FRigidTransform3 WorldTransform0 = LocalTransform0 * ParticleWorldTransform0;
@@ -1384,12 +1384,12 @@ namespace Chaos
 			}
 		}
 
-		void ConstructGenericConvexConvexConstraintsSwept(TGeometryParticleHandle<FReal, 3>* Particle0, TGeometryParticleHandle<FReal, 3>* Particle1, const FImplicitObject* Implicit0, const FImplicitObject* Implicit1, const FRigidTransform3& ParticleWorldTransform0, const FRigidTransform3& LocalTransform0, const FVec3& StartX0, const FRigidTransform3& ParticleWorldTransform1, const FRigidTransform3& LocalTransform1, const FVec3& StartX1, const FReal CullDistance, const FReal Dt, const FVec3& Dir, FReal Length, const FCollisionContext& Context)
+		void ConstructGenericConvexConvexConstraintsSwept(TGeometryParticleHandle<FReal, 3>* Particle0, TGeometryParticleHandle<FReal, 3>* Particle1, const FImplicitObject* Implicit0, const FPerShapeData* Shape0, const FImplicitObject* Implicit1, const FPerShapeData* Shape1, const FRigidTransform3& ParticleWorldTransform0, const FRigidTransform3& LocalTransform0, const FVec3& StartX0, const FRigidTransform3& ParticleWorldTransform1, const FRigidTransform3& LocalTransform1, const FVec3& StartX1, const FReal CullDistance, const FReal Dt, const FVec3& Dir, FReal Length, const FCollisionContext& Context)
 		{
 			CONDITIONAL_SCOPE_CYCLE_COUNTER(STAT_Collisions_ConstructGenericConvexConvexConstraintsSwept, ConstraintsDetailedStats);
 			FGenericParticleHandle RigidParticle0 = FGenericParticleHandle(Particle0);
 			FGenericParticleHandle RigidParticle1 = FGenericParticleHandle(Particle1);
-			FPBDCollisionConstraint* Constraint = Context.CollisionAllocator->FindOrCreateSweptConstraint(Particle0, Implicit0, nullptr, LocalTransform0, Particle1, Implicit1, nullptr, LocalTransform1, CullDistance, EContactShapesType::GenericConvexConvex);
+			FPBDCollisionConstraint* Constraint = Context.CollisionAllocator->FindOrCreateSweptConstraint(Particle0, Implicit0, Shape0, nullptr, LocalTransform0, Particle1, Implicit1, Shape1, nullptr, LocalTransform1, CullDistance, EContactShapesType::GenericConvexConvex);
 			// Rotations of WorldTransformXRx are not actually used in UpdateGenericConvexConvexConstraintSwept.
 			FRigidTransform3 WorldTransformXR0 = LocalTransform0 * FRigidTransform3(StartX0, RigidParticle0->R());
 			FRigidTransform3 WorldTransformXR1 = LocalTransform1 * FRigidTransform3(StartX1, RigidParticle1->R());
@@ -1431,13 +1431,13 @@ namespace Chaos
 		}
 
 		template<typename T_TRAITS>
-		void ConstructConvexHeightFieldConstraints(TGeometryParticleHandle<FReal, 3>* Particle0, TGeometryParticleHandle<FReal, 3>* Particle1, const FImplicitObject* Implicit0, const FImplicitObject* Implicit1, const FRigidTransform3& ParticleWorldTransform0, const FRigidTransform3& LocalTransform0, const FRigidTransform3& ParticleWorldTransform1, const FRigidTransform3& LocalTransform1, const FReal CullDistance, const FReal Dt, const FCollisionContext& Context)
+		void ConstructConvexHeightFieldConstraints(TGeometryParticleHandle<FReal, 3>* Particle0, TGeometryParticleHandle<FReal, 3>* Particle1, const FImplicitObject* Implicit0, const FPerShapeData* Shape0, const FImplicitObject* Implicit1, const FPerShapeData* Shape1, const FRigidTransform3& ParticleWorldTransform0, const FRigidTransform3& LocalTransform0, const FRigidTransform3& ParticleWorldTransform1, const FRigidTransform3& LocalTransform1, const FReal CullDistance, const FReal Dt, const FCollisionContext& Context)
 		{
 			CONDITIONAL_SCOPE_CYCLE_COUNTER(STAT_Collisions_ConstructConvexHeightFieldConstraints, ConstraintsDetailedStats);
 			const FHeightField* Object1 = Implicit1->template GetObject<const FHeightField >();
 			if (ensure(Implicit0->IsConvex() && Object1))
 			{
-				FPBDCollisionConstraint* Constraint = Context.CollisionAllocator->FindOrCreateConstraint(Particle0, Implicit0, nullptr, LocalTransform0, Particle1, Implicit1, nullptr, LocalTransform1, CullDistance, EContactShapesType::ConvexHeightField, Context.bAllowManifolds);
+				FPBDCollisionConstraint* Constraint = Context.CollisionAllocator->FindOrCreateConstraint(Particle0, Implicit0, Shape0, nullptr, LocalTransform0, Particle1, Implicit1, Shape1, nullptr, LocalTransform1, CullDistance, EContactShapesType::ConvexHeightField, Context.bAllowManifolds);
 				if (T_TRAITS::bImmediateUpdate && (Constraint != nullptr))
 				{
 					FRigidTransform3 WorldTransform0 = LocalTransform0 * ParticleWorldTransform0;
@@ -1447,13 +1447,13 @@ namespace Chaos
 			}
 		}
 
-		void ConstructConvexHeightFieldConstraintsSwept(TGeometryParticleHandle<FReal, 3>* Particle0, TGeometryParticleHandle<FReal, 3>* Particle1, const FImplicitObject* Implicit0, const FImplicitObject* Implicit1, const FRigidTransform3& ParticleWorldTransform0, const FRigidTransform3& LocalTransform0, const FRigidTransform3& ParticleWorldTransform1, const FRigidTransform3& LocalTransform1, const FReal CullDistance, const FReal Dt, const FVec3& Dir, FReal Length, const FCollisionContext& Context)
+		void ConstructConvexHeightFieldConstraintsSwept(TGeometryParticleHandle<FReal, 3>* Particle0, TGeometryParticleHandle<FReal, 3>* Particle1, const FImplicitObject* Implicit0, const FPerShapeData* Shape0, const FImplicitObject* Implicit1, const FPerShapeData* Shape1, const FRigidTransform3& ParticleWorldTransform0, const FRigidTransform3& LocalTransform0, const FRigidTransform3& ParticleWorldTransform1, const FRigidTransform3& LocalTransform1, const FReal CullDistance, const FReal Dt, const FVec3& Dir, FReal Length, const FCollisionContext& Context)
 		{
 			CONDITIONAL_SCOPE_CYCLE_COUNTER(STAT_Collisions_ConstructConvexHeightFieldConstraintsSwept, ConstraintsDetailedStats);
 			const FHeightField* Object1 = Implicit1->template GetObject<const FHeightField >();
 			if (ensure(Implicit0->IsConvex() && Object1))
 			{
-				FPBDCollisionConstraint* Constraint = Context.CollisionAllocator->FindOrCreateSweptConstraint(Particle0, Implicit0, nullptr, LocalTransform0, Particle1, Implicit1, nullptr, LocalTransform1, CullDistance, EContactShapesType::ConvexHeightField);
+				FPBDCollisionConstraint* Constraint = Context.CollisionAllocator->FindOrCreateSweptConstraint(Particle0, Implicit0, Shape0, nullptr, LocalTransform0, Particle1, Implicit1, Shape1, nullptr, LocalTransform1, CullDistance, EContactShapesType::ConvexHeightField);
 				FRigidTransform3 WorldTransformX0 = LocalTransform0 * FRigidTransform3(Particle0->X(), Particle0->R());
 				FRigidTransform3 WorldTransform1 = LocalTransform1 * ParticleWorldTransform1;
 				UpdateConvexHeightFieldConstraintSwept(Particle0, *Implicit0, WorldTransformX0, *Object1, WorldTransform1, Dir, Length, Dt, *Constraint);
@@ -1488,12 +1488,12 @@ namespace Chaos
 		}
 
 		template<typename T_TRAITS>
-		void ConstructConvexTriangleMeshConstraints(TGeometryParticleHandle<FReal, 3>* Particle0, TGeometryParticleHandle<FReal, 3>* Particle1, const FImplicitObject* Implicit0, const FImplicitObject* Implicit1, const FRigidTransform3& ParticleWorldTransform0, const FRigidTransform3& LocalTransform0, const FRigidTransform3& ParticleWorldTransform1, const FRigidTransform3& LocalTransform1, const FReal CullDistance, const FReal Dt, const FCollisionContext& Context)
+		void ConstructConvexTriangleMeshConstraints(TGeometryParticleHandle<FReal, 3>* Particle0, TGeometryParticleHandle<FReal, 3>* Particle1, const FImplicitObject* Implicit0, const FPerShapeData* Shape0, const FImplicitObject* Implicit1, const FPerShapeData* Shape1, const FRigidTransform3& ParticleWorldTransform0, const FRigidTransform3& LocalTransform0, const FRigidTransform3& ParticleWorldTransform1, const FRigidTransform3& LocalTransform1, const FReal CullDistance, const FReal Dt, const FCollisionContext& Context)
 		{
 			CONDITIONAL_SCOPE_CYCLE_COUNTER(STAT_Collisions_ConstructConvexTriangleMeshConstraints, ConstraintsDetailedStats);
 			if (ensure(Implicit0->IsConvex()))
 			{
-				FPBDCollisionConstraint* Constraint = Context.CollisionAllocator->FindOrCreateConstraint(Particle0, Implicit0, nullptr, LocalTransform0, Particle1, Implicit1, nullptr, LocalTransform1, CullDistance, EContactShapesType::ConvexTriMesh, Context.bAllowManifolds);
+				FPBDCollisionConstraint* Constraint = Context.CollisionAllocator->FindOrCreateConstraint(Particle0, Implicit0, Shape0, nullptr, LocalTransform0, Particle1, Implicit1, Shape1, nullptr, LocalTransform1, CullDistance, EContactShapesType::ConvexTriMesh, Context.bAllowManifolds);
 				if (T_TRAITS::bImmediateUpdate && (Constraint != nullptr))
 				{
 					FRigidTransform3 WorldTransform0 = LocalTransform0 * ParticleWorldTransform0;
@@ -1503,21 +1503,21 @@ namespace Chaos
 			}
 		}
 
-		void ConstructConvexTriangleMeshConstraintsSwept(TGeometryParticleHandle<FReal, 3>* Particle0, TGeometryParticleHandle<FReal, 3>* Particle1, const FImplicitObject* Implicit0, const FImplicitObject* Implicit1, const FRigidTransform3& ParticleWorldTransform0, const FRigidTransform3& LocalTransform0, const FRigidTransform3& ParticleWorldTransform1, const FRigidTransform3& LocalTransform1, const FReal CullDistance, const FReal Dt, const FVec3& Dir, FReal Length, const FCollisionContext& Context)
+		void ConstructConvexTriangleMeshConstraintsSwept(TGeometryParticleHandle<FReal, 3>* Particle0, TGeometryParticleHandle<FReal, 3>* Particle1, const FImplicitObject* Implicit0, const FPerShapeData* Shape0, const FImplicitObject* Implicit1, const FPerShapeData* Shape1, const FRigidTransform3& ParticleWorldTransform0, const FRigidTransform3& LocalTransform0, const FRigidTransform3& ParticleWorldTransform1, const FRigidTransform3& LocalTransform1, const FReal CullDistance, const FReal Dt, const FVec3& Dir, FReal Length, const FCollisionContext& Context)
 		{
 			CONDITIONAL_SCOPE_CYCLE_COUNTER(STAT_Collisions_ConstructConvexTriangleMeshConstraintsSwept, ConstraintsDetailedStats);
 			if (ensure(Implicit0->IsConvex()))
 			{
 				if (const TImplicitObjectScaled<FTriangleMeshImplicitObject>* ScaledTriangleMesh = Implicit1->template GetObject<const TImplicitObjectScaled<FTriangleMeshImplicitObject>>())
 				{
-					FPBDCollisionConstraint* Constraint = Context.CollisionAllocator->FindOrCreateSweptConstraint(Particle0, Implicit0, nullptr, LocalTransform0, Particle1, Implicit1, nullptr, LocalTransform1, CullDistance, EContactShapesType::ConvexTriMesh);
+					FPBDCollisionConstraint* Constraint = Context.CollisionAllocator->FindOrCreateSweptConstraint(Particle0, Implicit0, Shape0, nullptr, LocalTransform0, Particle1, Implicit1, Shape1, nullptr, LocalTransform1, CullDistance, EContactShapesType::ConvexTriMesh);
 					FRigidTransform3 WorldTransformX0 = LocalTransform0 * FRigidTransform3(Particle0->X(), Particle0->R());
 					FRigidTransform3 WorldTransform1 = LocalTransform1 * ParticleWorldTransform1;
 					UpdateConvexTriangleMeshConstraintSwept(Particle0, *Implicit0, WorldTransformX0, *ScaledTriangleMesh, WorldTransform1, Dir, Length, Dt, *Constraint);
 				}
 				else if (const FTriangleMeshImplicitObject* TriangleMesh = Implicit1->template GetObject<const FTriangleMeshImplicitObject>())
 				{
-					FPBDCollisionConstraint* Constraint = Context.CollisionAllocator->FindOrCreateSweptConstraint(Particle0, Implicit0, nullptr, LocalTransform0, Particle1, Implicit1, nullptr, LocalTransform1, CullDistance, EContactShapesType::ConvexTriMesh);
+					FPBDCollisionConstraint* Constraint = Context.CollisionAllocator->FindOrCreateSweptConstraint(Particle0, Implicit0, Shape0, nullptr, LocalTransform0, Particle1, Implicit1, Shape1, nullptr, LocalTransform1, CullDistance, EContactShapesType::ConvexTriMesh);
 					FRigidTransform3 WorldTransformX0 = LocalTransform0 * FRigidTransform3(Particle0->X(), Particle0->R());
 					FRigidTransform3 WorldTransform1 = LocalTransform1 * ParticleWorldTransform1;
 					UpdateConvexTriangleMeshConstraintSwept(Particle0, *Implicit0, WorldTransformX0, *TriangleMesh, WorldTransform1, Dir, Length, Dt, *Constraint);
@@ -1590,7 +1590,7 @@ namespace Chaos
 
 
 		template<typename T_TRAITS>
-		void ConstructLevelsetLevelsetConstraints(TGeometryParticleHandle<FReal, 3>* Particle0, TGeometryParticleHandle<FReal, 3>* Particle1, const FImplicitObject* Implicit0, const FBVHParticles* Simplicial0, const FImplicitObject* Implicit1, const FBVHParticles* Simplicial1, const FRigidTransform3& ParticleWorldTransform0, const FRigidTransform3& LocalTransform0, const FRigidTransform3& ParticleWorldTransform1, const FRigidTransform3& LocalTransform1, const FReal CullDistance, const FReal Dt, const FCollisionContext& Context)
+		void ConstructLevelsetLevelsetConstraints(TGeometryParticleHandle<FReal, 3>* Particle0, TGeometryParticleHandle<FReal, 3>* Particle1, const FImplicitObject* Implicit0, const FPerShapeData* Shape0, const FBVHParticles* Simplicial0, const FImplicitObject* Implicit1, const FPerShapeData* Shape1, const FBVHParticles* Simplicial1, const FRigidTransform3& ParticleWorldTransform0, const FRigidTransform3& LocalTransform0, const FRigidTransform3& ParticleWorldTransform1, const FRigidTransform3& LocalTransform1, const FReal CullDistance, const FReal Dt, const FCollisionContext& Context)
 		{
 			const bool bAllowManifolds = Context.bAllowManifolds && bChaos_Collision_AllowLevelsetManifolds;
 			FPBDCollisionConstraint* Constraint = nullptr;
@@ -1601,11 +1601,11 @@ namespace Chaos
 			int32 P0NumCollisionParticles = Simplicial0 ? Simplicial0->Size() : Particle0->CastToRigidParticle()->CollisionParticlesSize();
 			if (!Particle1->Geometry() || (bIsParticleDynamic0 && !P0NumCollisionParticles && Particle0->Geometry() && !Particle0->Geometry()->IsUnderlyingUnion()))
 			{
-				Constraint = Context.CollisionAllocator->FindOrCreateConstraint(Particle1, Implicit1, Simplicial1, LocalTransform1, Particle0, Implicit0, Simplicial0, LocalTransform0, CullDistance, EContactShapesType::LevelSetLevelSet, bAllowManifolds);
+				Constraint = Context.CollisionAllocator->FindOrCreateConstraint(Particle1, Implicit1, Shape1, Simplicial1, LocalTransform1, Particle0, Implicit0, Shape0, Simplicial0, LocalTransform0, CullDistance, EContactShapesType::LevelSetLevelSet, bAllowManifolds);
 			}
 			else
 			{
-				Constraint = Context.CollisionAllocator->FindOrCreateConstraint(Particle0, Implicit0, Simplicial0, LocalTransform0, Particle1, Implicit1, Simplicial1, LocalTransform1, CullDistance, EContactShapesType::LevelSetLevelSet, bAllowManifolds);
+				Constraint = Context.CollisionAllocator->FindOrCreateConstraint(Particle0, Implicit0, Shape0, Simplicial0, LocalTransform0, Particle1, Implicit1, Shape1, Simplicial1, LocalTransform1, CullDistance, EContactShapesType::LevelSetLevelSet, bAllowManifolds);
 			}
 
 			if (Constraint != nullptr)
@@ -1867,7 +1867,7 @@ namespace Chaos
 		}
 
 		template<typename T_TRAITS>
-		void ConstructConstraintsImpl(TGeometryParticleHandle<FReal, 3>* Particle0, TGeometryParticleHandle<FReal, 3>* Particle1, const FImplicitObject* Implicit0, const FBVHParticles* Simplicial0, const FImplicitObject* Implicit1, const FBVHParticles* Simplicial1, const FRigidTransform3& ParticleWorldTransform0, const FRigidTransform3& LocalTransform0, const FRigidTransform3& ParticleWorldTransform1, const FRigidTransform3& LocalTransform1, const FReal CullDistance, const FReal Dt, const FCollisionContext& Context)
+		void ConstructConstraintsImpl(TGeometryParticleHandle<FReal, 3>* Particle0, TGeometryParticleHandle<FReal, 3>* Particle1, const FImplicitObject* Implicit0, const FPerShapeData* Shape0, const FBVHParticles* Simplicial0, const FImplicitObject* Implicit1, const FPerShapeData* Shape1, const FBVHParticles* Simplicial1, const FRigidTransform3& ParticleWorldTransform0, const FRigidTransform3& LocalTransform0, const FRigidTransform3& ParticleWorldTransform1, const FRigidTransform3& LocalTransform1, const FReal CullDistance, const FReal Dt, const FCollisionContext& Context)
 		{
 			CONDITIONAL_SCOPE_CYCLE_COUNTER(STAT_Collisions_ConstructConstraintsInternal, ConstraintsDetailedStats);
 			INC_DWORD_STAT(STAT_ChaosCollisionCounter_NumShapePairs);
@@ -1948,213 +1948,213 @@ namespace Chaos
 
 			if (Implicit0Type == TBox<FReal, 3>::StaticType() && Implicit1Type == TBox<FReal, 3>::StaticType() && !bUseGenericSweptConstraints)
 			{
-				ConstructBoxBoxConstraints<T_TRAITS>(Particle0, Particle1, Implicit0, Implicit1, ParticleWorldTransform0, LocalTransform0, ParticleWorldTransform1, LocalTransform1, CullDistance, Dt, Context);
+				ConstructBoxBoxConstraints<T_TRAITS>(Particle0, Particle1, Implicit0, Shape0, Implicit1, Shape1, ParticleWorldTransform0, LocalTransform0, ParticleWorldTransform1, LocalTransform1, CullDistance, Dt, Context);
 			}
 			else if (Implicit0Type == TBox<FReal, 3>::StaticType() && Implicit1Type == FHeightField::StaticType())
 			{
-				ConstructBoxHeightFieldConstraints<T_TRAITS>(Particle0, Particle1, Implicit0, Implicit1, ParticleWorldTransform0, LocalTransform0, ParticleWorldTransform1, LocalTransform1, CullDistance, Dt, Context);
+				ConstructBoxHeightFieldConstraints<T_TRAITS>(Particle0, Particle1, Implicit0, Shape0, Implicit1, Shape1, ParticleWorldTransform0, LocalTransform0, ParticleWorldTransform1, LocalTransform1, CullDistance, Dt, Context);
 			}
 			else if (Implicit0Type == FHeightField::StaticType() && Implicit1Type == TBox<FReal, 3>::StaticType())
 			{
-				ConstructBoxHeightFieldConstraints<T_TRAITS>(Particle1, Particle0, Implicit1, Implicit0, ParticleWorldTransform1, LocalTransform1, ParticleWorldTransform0, LocalTransform0, CullDistance, Dt, Context);
+				ConstructBoxHeightFieldConstraints<T_TRAITS>(Particle1, Particle0, Implicit1, Shape1, Implicit0, Shape0, ParticleWorldTransform1, LocalTransform1, ParticleWorldTransform0, LocalTransform0, CullDistance, Dt, Context);
 			}
 			else if (Implicit0Type == TBox<FReal, 3>::StaticType() && Implicit1Type == TPlane<FReal, 3>::StaticType() && !bUseGenericSweptConstraints)
 			{
-				ConstructBoxPlaneConstraints<T_TRAITS>(Particle0, Particle1, Implicit0, Implicit1, ParticleWorldTransform0, LocalTransform0, ParticleWorldTransform1, LocalTransform1, CullDistance, Dt, Context);
+				ConstructBoxPlaneConstraints<T_TRAITS>(Particle0, Particle1, Implicit0, Shape0, Implicit1, Shape1, ParticleWorldTransform0, LocalTransform0, ParticleWorldTransform1, LocalTransform1, CullDistance, Dt, Context);
 			}
 			else if (Implicit0Type == TPlane<FReal, 3>::StaticType() && Implicit1Type == TBox<FReal, 3>::StaticType() && !bUseGenericSweptConstraints)
 			{
-				ConstructBoxPlaneConstraints<T_TRAITS>(Particle1, Particle0, Implicit1, Implicit0, ParticleWorldTransform1, LocalTransform1, ParticleWorldTransform0, LocalTransform0, CullDistance, Dt, Context);
+				ConstructBoxPlaneConstraints<T_TRAITS>(Particle1, Particle0, Implicit1, Shape1, Implicit0, Shape0, ParticleWorldTransform1, LocalTransform1, ParticleWorldTransform0, LocalTransform0, CullDistance, Dt, Context);
 			}
 			else if (Implicit0Type == TSphere<FReal, 3>::StaticType() && Implicit1Type == TSphere<FReal, 3>::StaticType() && !bUseGenericSweptConstraints)
 			{
-				ConstructSphereSphereConstraints<T_TRAITS>(Particle0, Particle1, Implicit0, Implicit1, ParticleWorldTransform0, LocalTransform0, ParticleWorldTransform1, LocalTransform1, CullDistance, Dt, Context);
+				ConstructSphereSphereConstraints<T_TRAITS>(Particle0, Particle1, Implicit0, Shape0, Implicit1, Shape1, ParticleWorldTransform0, LocalTransform0, ParticleWorldTransform1, LocalTransform1, CullDistance, Dt, Context);
 			}
 			else if (Implicit0Type == TSphere<FReal, 3>::StaticType() && Implicit1Type == FHeightField::StaticType())
 			{
 				if (bUseCCD)
 				{
-					ConstructSphereHeightFieldConstraintsSwept(Particle0, Particle1, Implicit0, Implicit1, ParticleWorldTransform0, LocalTransform0, ParticleWorldTransform1, LocalTransform1, CullDistance, Dt, DirCCD, LengthCCD, Context);
+					ConstructSphereHeightFieldConstraintsSwept(Particle0, Particle1, Implicit0, Shape0, Implicit1, Shape1, ParticleWorldTransform0, LocalTransform0, ParticleWorldTransform1, LocalTransform1, CullDistance, Dt, DirCCD, LengthCCD, Context);
 					return;
 				}
-				ConstructSphereHeightFieldConstraints<T_TRAITS>(Particle0, Particle1, Implicit0, Implicit1, ParticleWorldTransform0, LocalTransform0, ParticleWorldTransform1, LocalTransform1, CullDistance, Dt, Context);
+				ConstructSphereHeightFieldConstraints<T_TRAITS>(Particle0, Particle1, Implicit0, Shape0, Implicit1, Shape1, ParticleWorldTransform0, LocalTransform0, ParticleWorldTransform1, LocalTransform1, CullDistance, Dt, Context);
 			}
 			else if (Implicit0Type == FHeightField::StaticType() && Implicit1Type == TSphere<FReal, 3>::StaticType())
 			{
 				if (bUseCCD)
 				{
-					ConstructSphereHeightFieldConstraintsSwept(Particle1, Particle0, Implicit1, Implicit0, ParticleWorldTransform1, LocalTransform1, ParticleWorldTransform0, LocalTransform0, CullDistance, Dt, DirCCD, LengthCCD, Context);
+					ConstructSphereHeightFieldConstraintsSwept(Particle1, Particle0, Implicit1, Shape1, Implicit0, Shape0, ParticleWorldTransform1, LocalTransform1, ParticleWorldTransform0, LocalTransform0, CullDistance, Dt, DirCCD, LengthCCD, Context);
 					return;
 				}
-				ConstructSphereHeightFieldConstraints<T_TRAITS>(Particle1, Particle0, Implicit1, Implicit0, ParticleWorldTransform1, LocalTransform1, ParticleWorldTransform0, LocalTransform0, CullDistance, Dt, Context);
+				ConstructSphereHeightFieldConstraints<T_TRAITS>(Particle1, Particle0, Implicit1, Shape1, Implicit0, Shape0, ParticleWorldTransform1, LocalTransform1, ParticleWorldTransform0, LocalTransform0, CullDistance, Dt, Context);
 			}
 			else if (Implicit0Type == TSphere<FReal, 3>::StaticType() && Implicit1Type == TPlane<FReal, 3>::StaticType() && !bUseGenericSweptConstraints)
 			{
-				ConstructSpherePlaneConstraints<T_TRAITS>(Particle0, Particle1, Implicit0, Implicit1, ParticleWorldTransform0, LocalTransform0, ParticleWorldTransform1, LocalTransform1, CullDistance, Dt, Context);
+				ConstructSpherePlaneConstraints<T_TRAITS>(Particle0, Particle1, Implicit0, Shape0, Implicit1, Shape1, ParticleWorldTransform0, LocalTransform0, ParticleWorldTransform1, LocalTransform1, CullDistance, Dt, Context);
 			}
 			else if (Implicit0Type == TPlane<FReal, 3>::StaticType() && Implicit1Type == TSphere<FReal, 3>::StaticType() && !bUseGenericSweptConstraints)
 			{
-				ConstructSpherePlaneConstraints<T_TRAITS>(Particle1, Particle0, Implicit1, Implicit0, ParticleWorldTransform1, LocalTransform1, ParticleWorldTransform0, LocalTransform0, CullDistance, Dt, Context);
+				ConstructSpherePlaneConstraints<T_TRAITS>(Particle1, Particle0, Implicit1, Shape1, Implicit0, Shape0, ParticleWorldTransform1, LocalTransform1, ParticleWorldTransform0, LocalTransform0, CullDistance, Dt, Context);
 			}
 			else if (Implicit0Type == TSphere<FReal, 3>::StaticType() && Implicit1Type == TBox<FReal, 3>::StaticType() && !bUseGenericSweptConstraints)
 			{
-				ConstructSphereBoxConstraints<T_TRAITS>(Particle0, Particle1, Implicit0, Implicit1, ParticleWorldTransform0, LocalTransform0, ParticleWorldTransform1, LocalTransform1, CullDistance, Dt, Context);
+				ConstructSphereBoxConstraints<T_TRAITS>(Particle0, Particle1, Implicit0, Shape0, Implicit1, Shape1, ParticleWorldTransform0, LocalTransform0, ParticleWorldTransform1, LocalTransform1, CullDistance, Dt, Context);
 			}
 			else if (Implicit0Type == TBox<FReal, 3>::StaticType() && Implicit1Type == TSphere<FReal, 3>::StaticType() && !bUseGenericSweptConstraints)
 			{
-				ConstructSphereBoxConstraints<T_TRAITS>(Particle1, Particle0, Implicit1, Implicit0, ParticleWorldTransform1, LocalTransform1, ParticleWorldTransform0, LocalTransform0, CullDistance, Dt, Context);
+				ConstructSphereBoxConstraints<T_TRAITS>(Particle1, Particle0, Implicit1, Shape1, Implicit0, Shape0, ParticleWorldTransform1, LocalTransform1, ParticleWorldTransform0, LocalTransform0, CullDistance, Dt, Context);
 			}
 			else if (Implicit0Type == TSphere<FReal, 3>::StaticType() && Implicit1Type == FCapsule::StaticType() && !bUseGenericSweptConstraints)
 			{
-				ConstructSphereCapsuleConstraints<T_TRAITS>(Particle0, Particle1, Implicit0, Implicit1, ParticleWorldTransform0, LocalTransform0, ParticleWorldTransform1, LocalTransform1, CullDistance, Dt, Context);
+				ConstructSphereCapsuleConstraints<T_TRAITS>(Particle0, Particle1, Implicit0, Shape0, Implicit1, Shape1, ParticleWorldTransform0, LocalTransform0, ParticleWorldTransform1, LocalTransform1, CullDistance, Dt, Context);
 			}
 			else if (Implicit0Type == FCapsule::StaticType() && Implicit1Type == TSphere<FReal, 3>::StaticType() && !bUseGenericSweptConstraints)
 			{
-				ConstructSphereCapsuleConstraints<T_TRAITS>(Particle1, Particle0, Implicit1, Implicit0, ParticleWorldTransform1, LocalTransform1, ParticleWorldTransform0, LocalTransform0, CullDistance, Dt, Context);
+				ConstructSphereCapsuleConstraints<T_TRAITS>(Particle1, Particle0, Implicit1, Shape1, Implicit0, Shape0, ParticleWorldTransform1, LocalTransform1, ParticleWorldTransform0, LocalTransform0, CullDistance, Dt, Context);
 			}
 			else if (Implicit0Type == TSphere<FReal, 3>::StaticType() && Implicit1Type == FImplicitConvex3::StaticType() && !bUseGenericSweptConstraints)
 			{
-				ConstructSphereConvexConstraints<T_TRAITS>(Particle0, Particle1, Implicit0, Implicit1, ParticleWorldTransform0, LocalTransform0, ParticleWorldTransform1, LocalTransform1, CullDistance, Dt, Context);
+				ConstructSphereConvexConstraints<T_TRAITS>(Particle0, Particle1, Implicit0, Shape0, Implicit1, Shape1, ParticleWorldTransform0, LocalTransform0, ParticleWorldTransform1, LocalTransform1, CullDistance, Dt, Context);
 			}
 			else if (Implicit0Type == FImplicitConvex3::StaticType() && Implicit1Type == TSphere<FReal, 3>::StaticType() && !bUseGenericSweptConstraints)
 			{
-				ConstructSphereConvexConstraints<T_TRAITS>(Particle1, Particle0, Implicit1, Implicit0, ParticleWorldTransform1, LocalTransform1, ParticleWorldTransform0, LocalTransform0, CullDistance, Dt, Context);
+				ConstructSphereConvexConstraints<T_TRAITS>(Particle1, Particle0, Implicit1, Shape1, Implicit0, Shape0, ParticleWorldTransform1, LocalTransform1, ParticleWorldTransform0, LocalTransform0, CullDistance, Dt, Context);
 			}
 			else if (Implicit0Type == FCapsule::StaticType() && Implicit1Type == FCapsule::StaticType() && !bUseGenericSweptConstraints)
 			{
-				ConstructCapsuleCapsuleConstraints<T_TRAITS>(Particle0, Particle1, Implicit0, Implicit1, ParticleWorldTransform0, LocalTransform0, ParticleWorldTransform1, LocalTransform1, CullDistance, Dt, Context);
+				ConstructCapsuleCapsuleConstraints<T_TRAITS>(Particle0, Particle1, Implicit0, Shape0, Implicit1, Shape1, ParticleWorldTransform0, LocalTransform0, ParticleWorldTransform1, LocalTransform1, CullDistance, Dt, Context);
 			}
 			else if (Implicit0Type == FCapsule::StaticType() && Implicit1Type == TBox<FReal, 3>::StaticType() && !bUseGenericSweptConstraints)
 			{
-				ConstructCapsuleBoxConstraints<T_TRAITS>(Particle0, Particle1, Implicit0, Implicit1, ParticleWorldTransform0, LocalTransform0, ParticleWorldTransform1, LocalTransform1, CullDistance, Dt, Context);
+				ConstructCapsuleBoxConstraints<T_TRAITS>(Particle0, Particle1, Implicit0, Shape0, Implicit1, Shape1, ParticleWorldTransform0, LocalTransform0, ParticleWorldTransform1, LocalTransform1, CullDistance, Dt, Context);
 			}
 			else if (Implicit0Type == TBox<FReal, 3>::StaticType() && Implicit1Type == FCapsule::StaticType() && !bUseGenericSweptConstraints)
 			{
-				ConstructCapsuleBoxConstraints<T_TRAITS>(Particle1, Particle0, Implicit1, Implicit0, ParticleWorldTransform1, LocalTransform1, ParticleWorldTransform0, LocalTransform0, CullDistance, Dt, Context);
+				ConstructCapsuleBoxConstraints<T_TRAITS>(Particle1, Particle0, Implicit1, Shape1, Implicit0, Shape0, ParticleWorldTransform1, LocalTransform1, ParticleWorldTransform0, LocalTransform0, CullDistance, Dt, Context);
 			}
 			else if (Implicit0Type == FCapsule::StaticType() && Implicit1Type == FImplicitConvex3::StaticType() && !bUseGenericSweptConstraints)
 			{
-				ConstructCapsuleConvexConstraints<T_TRAITS>(Particle0, Particle1, Implicit0, Implicit1, ParticleWorldTransform0, LocalTransform0, ParticleWorldTransform1, LocalTransform1, CullDistance, Dt, Context);
+				ConstructCapsuleConvexConstraints<T_TRAITS>(Particle0, Particle1, Implicit0, Shape0, Implicit1, Shape1, ParticleWorldTransform0, LocalTransform0, ParticleWorldTransform1, LocalTransform1, CullDistance, Dt, Context);
 			}
 			else if (Implicit0Type == FImplicitConvex3::StaticType() && Implicit1Type == FCapsule::StaticType() && !bUseGenericSweptConstraints)
 			{
-				ConstructCapsuleConvexConstraints<T_TRAITS>(Particle1, Particle0, Implicit1, Implicit0, ParticleWorldTransform1, LocalTransform1, ParticleWorldTransform0, LocalTransform0, CullDistance, Dt, Context);
+				ConstructCapsuleConvexConstraints<T_TRAITS>(Particle1, Particle0, Implicit1, Shape1, Implicit0, Shape0, ParticleWorldTransform1, LocalTransform1, ParticleWorldTransform0, LocalTransform0, CullDistance, Dt, Context);
 			}
 			else if (Implicit0Type == FCapsule::StaticType() && Implicit1Type == FHeightField::StaticType())
 			{
 				if (bUseCCD)
 				{
-					ConstructCapsuleHeightFieldConstraintsSwept(Particle0, Particle1, Implicit0, Implicit1, ParticleWorldTransform0, LocalTransform0, ParticleWorldTransform1, LocalTransform1, CullDistance, Dt, DirCCD, LengthCCD, Context);
+					ConstructCapsuleHeightFieldConstraintsSwept(Particle0, Particle1, Implicit0, Shape0, Implicit1, Shape1, ParticleWorldTransform0, LocalTransform0, ParticleWorldTransform1, LocalTransform1, CullDistance, Dt, DirCCD, LengthCCD, Context);
 					return;
 				}
 
-				ConstructCapsuleHeightFieldConstraints<T_TRAITS>(Particle0, Particle1, Implicit0, Implicit1, ParticleWorldTransform0, LocalTransform0, ParticleWorldTransform1, LocalTransform1, CullDistance, Dt, Context);
+				ConstructCapsuleHeightFieldConstraints<T_TRAITS>(Particle0, Particle1, Implicit0, Shape0, Implicit1, Shape1, ParticleWorldTransform0, LocalTransform0, ParticleWorldTransform1, LocalTransform1, CullDistance, Dt, Context);
 			}
 			else if (Implicit0Type == FHeightField::StaticType() && Implicit1Type == FCapsule::StaticType())
 			{
 				if (bUseCCD)
 				{
-					ConstructCapsuleHeightFieldConstraintsSwept(Particle1, Particle0, Implicit1, Implicit0, ParticleWorldTransform1, LocalTransform1, ParticleWorldTransform0, LocalTransform0, CullDistance, Dt, DirCCD, LengthCCD, Context);
+					ConstructCapsuleHeightFieldConstraintsSwept(Particle1, Particle0, Implicit1, Shape1, Implicit0, Shape0, ParticleWorldTransform1, LocalTransform1, ParticleWorldTransform0, LocalTransform0, CullDistance, Dt, DirCCD, LengthCCD, Context);
 					return;
 				}
 
-				ConstructCapsuleHeightFieldConstraints<T_TRAITS>(Particle1, Particle0, Implicit1, Implicit0, ParticleWorldTransform1, LocalTransform1, ParticleWorldTransform0, LocalTransform0, CullDistance, Dt, Context);
+				ConstructCapsuleHeightFieldConstraints<T_TRAITS>(Particle1, Particle0, Implicit1, Shape1, Implicit0, Shape0, ParticleWorldTransform1, LocalTransform1, ParticleWorldTransform0, LocalTransform0, CullDistance, Dt, Context);
 			}
 			else if (Implicit0Type == TBox<FReal, 3>::StaticType() && Implicit1Type == FTriangleMeshImplicitObject::StaticType())
 			{
-				ConstructBoxTriangleMeshConstraints<T_TRAITS>(Particle0, Particle1, Implicit0, Implicit1, ParticleWorldTransform0, LocalTransform0, ParticleWorldTransform1, LocalTransform1, CullDistance, Dt, Context);
+				ConstructBoxTriangleMeshConstraints<T_TRAITS>(Particle0, Particle1, Implicit0, Shape0, Implicit1, Shape1, ParticleWorldTransform0, LocalTransform0, ParticleWorldTransform1, LocalTransform1, CullDistance, Dt, Context);
 			}
 			else if (Implicit0Type == FTriangleMeshImplicitObject::StaticType() && Implicit1Type == TBox<FReal, 3>::StaticType())
 			{
-				ConstructBoxTriangleMeshConstraints<T_TRAITS>(Particle1, Particle0, Implicit1, Implicit0, ParticleWorldTransform1, LocalTransform1, ParticleWorldTransform0, LocalTransform0, CullDistance, Dt, Context);
+				ConstructBoxTriangleMeshConstraints<T_TRAITS>(Particle1, Particle0, Implicit1, Shape1, Implicit0, Shape0, ParticleWorldTransform1, LocalTransform1, ParticleWorldTransform0, LocalTransform0, CullDistance, Dt, Context);
 			}
 			else if (Implicit0Type == TSphere<FReal, 3>::StaticType() && Implicit1Type == FTriangleMeshImplicitObject::StaticType())
 			{
 				if (bUseCCD)
 				{
-					ConstructSphereTriangleMeshConstraintsSwept(Particle0, Particle1, Implicit0, Implicit1, ParticleWorldTransform0, LocalTransform0, ParticleWorldTransform1, LocalTransform1, CullDistance, Dt, DirCCD, LengthCCD, Context);
+					ConstructSphereTriangleMeshConstraintsSwept(Particle0, Particle1, Implicit0, Shape0, Implicit1, Shape1, ParticleWorldTransform0, LocalTransform0, ParticleWorldTransform1, LocalTransform1, CullDistance, Dt, DirCCD, LengthCCD, Context);
 					return;
 				}
-				ConstructSphereTriangleMeshConstraints<T_TRAITS>(Particle0, Particle1, Implicit0, Implicit1, ParticleWorldTransform0, LocalTransform0, ParticleWorldTransform1, LocalTransform1, CullDistance, Dt, Context);
+				ConstructSphereTriangleMeshConstraints<T_TRAITS>(Particle0, Particle1, Implicit0, Shape0, Implicit1, Shape1, ParticleWorldTransform0, LocalTransform0, ParticleWorldTransform1, LocalTransform1, CullDistance, Dt, Context);
 			}
 			else if (Implicit0Type == FTriangleMeshImplicitObject::StaticType() && Implicit1Type == TSphere<FReal, 3>::StaticType())
 			{
 				if (bUseCCD)
 				{
-					ConstructSphereTriangleMeshConstraintsSwept(Particle1, Particle0, Implicit1, Implicit0, ParticleWorldTransform1, LocalTransform1, ParticleWorldTransform0, LocalTransform0, CullDistance, Dt, DirCCD, LengthCCD, Context);
+					ConstructSphereTriangleMeshConstraintsSwept(Particle1, Particle0, Implicit1, Shape1, Implicit0, Shape0, ParticleWorldTransform1, LocalTransform1, ParticleWorldTransform0, LocalTransform0, CullDistance, Dt, DirCCD, LengthCCD, Context);
 					return;
 				}
-				ConstructSphereTriangleMeshConstraints<T_TRAITS>(Particle1, Particle0, Implicit1, Implicit0, ParticleWorldTransform1, LocalTransform1, ParticleWorldTransform0, LocalTransform0, CullDistance, Dt, Context);
+				ConstructSphereTriangleMeshConstraints<T_TRAITS>(Particle1, Particle0, Implicit1, Shape1, Implicit0, Shape0, ParticleWorldTransform1, LocalTransform1, ParticleWorldTransform0, LocalTransform0, CullDistance, Dt, Context);
 			}
 			else if (Implicit0Type == FCapsule::StaticType() && Implicit1Type == FTriangleMeshImplicitObject::StaticType())
 			{
 				if (bUseCCD)
 				{
-					ConstructCapsuleTriangleMeshConstraintsSwept(Particle0, Particle1, Implicit0, Implicit1, ParticleWorldTransform0, LocalTransform0, ParticleWorldTransform1, LocalTransform1, CullDistance, Dt, DirCCD, LengthCCD, Context);
+					ConstructCapsuleTriangleMeshConstraintsSwept(Particle0, Particle1, Implicit0, Shape0, Implicit1, Shape1, ParticleWorldTransform0, LocalTransform0, ParticleWorldTransform1, LocalTransform1, CullDistance, Dt, DirCCD, LengthCCD, Context);
 					return;
 				}
 
-				ConstructCapsuleTriangleMeshConstraints<T_TRAITS>(Particle0, Particle1, Implicit0, Implicit1, ParticleWorldTransform0, LocalTransform0, ParticleWorldTransform1, LocalTransform1, CullDistance, Dt, Context);
+				ConstructCapsuleTriangleMeshConstraints<T_TRAITS>(Particle0, Particle1, Implicit0, Shape0, Implicit1, Shape1, ParticleWorldTransform0, LocalTransform0, ParticleWorldTransform1, LocalTransform1, CullDistance, Dt, Context);
 			}
 			else if (Implicit0Type == FTriangleMeshImplicitObject::StaticType() && Implicit1Type == FCapsule::StaticType())
 			{
 				if (bUseCCD)
 				{
-					ConstructCapsuleTriangleMeshConstraintsSwept(Particle1, Particle0, Implicit1, Implicit0, ParticleWorldTransform1, LocalTransform1, ParticleWorldTransform0, LocalTransform0, CullDistance, Dt, DirCCD, LengthCCD, Context);
+					ConstructCapsuleTriangleMeshConstraintsSwept(Particle1, Particle0, Implicit1, Shape1, Implicit0, Shape0, ParticleWorldTransform1, LocalTransform1, ParticleWorldTransform0, LocalTransform0, CullDistance, Dt, DirCCD, LengthCCD, Context);
 					return;
 				}
 
-				ConstructCapsuleTriangleMeshConstraints<T_TRAITS>(Particle1, Particle0, Implicit1, Implicit0, ParticleWorldTransform1, LocalTransform1, ParticleWorldTransform0, LocalTransform0, CullDistance, Dt, Context);
+				ConstructCapsuleTriangleMeshConstraints<T_TRAITS>(Particle1, Particle0, Implicit1, Shape1, Implicit0, Shape0, ParticleWorldTransform1, LocalTransform1, ParticleWorldTransform0, LocalTransform0, CullDistance, Dt, Context);
 			}
 			else if (bIsConvex0 && Implicit1Type == FHeightField::StaticType())
 			{
 				if (bUseCCD)
 				{
-					ConstructConvexHeightFieldConstraintsSwept(Particle0, Particle1, Implicit0, Implicit1, ParticleWorldTransform0, LocalTransform0, ParticleWorldTransform1, LocalTransform1, CullDistance, Dt, DirCCD, LengthCCD, Context);
+					ConstructConvexHeightFieldConstraintsSwept(Particle0, Particle1, Implicit0, Shape0, Implicit1, Shape1, ParticleWorldTransform0, LocalTransform0, ParticleWorldTransform1, LocalTransform1, CullDistance, Dt, DirCCD, LengthCCD, Context);
 					return;
 				}
 
-				ConstructConvexHeightFieldConstraints<T_TRAITS>(Particle0, Particle1, Implicit0, Implicit1, ParticleWorldTransform0, LocalTransform0, ParticleWorldTransform1, LocalTransform1, CullDistance, Dt, Context);
+				ConstructConvexHeightFieldConstraints<T_TRAITS>(Particle0, Particle1, Implicit0, Shape0, Implicit1, Shape1, ParticleWorldTransform0, LocalTransform0, ParticleWorldTransform1, LocalTransform1, CullDistance, Dt, Context);
 			}
 			else if (Implicit0Type == FHeightField::StaticType() && bIsConvex1)
 			{
 				if (bUseCCD)
 				{
-					ConstructConvexHeightFieldConstraintsSwept(Particle1, Particle0, Implicit1, Implicit0, ParticleWorldTransform1, LocalTransform1, ParticleWorldTransform0, LocalTransform0, CullDistance, Dt, DirCCD, LengthCCD, Context);
+					ConstructConvexHeightFieldConstraintsSwept(Particle1, Particle0, Implicit1, Shape1, Implicit0, Shape0, ParticleWorldTransform1, LocalTransform1, ParticleWorldTransform0, LocalTransform0, CullDistance, Dt, DirCCD, LengthCCD, Context);
 					return;
 				}
 
-				ConstructConvexHeightFieldConstraints<T_TRAITS>(Particle1, Particle0, Implicit1, Implicit0, ParticleWorldTransform1, LocalTransform1, ParticleWorldTransform0, LocalTransform0, CullDistance, Dt, Context);
+				ConstructConvexHeightFieldConstraints<T_TRAITS>(Particle1, Particle0, Implicit1, Shape1, Implicit0, Shape0, ParticleWorldTransform1, LocalTransform1, ParticleWorldTransform0, LocalTransform0, CullDistance, Dt, Context);
 			}
 			else if (bIsConvex0 && Implicit1Type == FTriangleMeshImplicitObject::StaticType())
 			{
 				if (bUseCCD)
 				{
-					ConstructConvexTriangleMeshConstraintsSwept(Particle0, Particle1, Implicit0, Implicit1, ParticleWorldTransform0, LocalTransform0, ParticleWorldTransform1, LocalTransform1, CullDistance, Dt, DirCCD, LengthCCD, Context);
+					ConstructConvexTriangleMeshConstraintsSwept(Particle0, Particle1, Implicit0, Shape0, Implicit1, Shape1, ParticleWorldTransform0, LocalTransform0, ParticleWorldTransform1, LocalTransform1, CullDistance, Dt, DirCCD, LengthCCD, Context);
 					return;
 				}
 
-				ConstructConvexTriangleMeshConstraints<T_TRAITS>(Particle0, Particle1, Implicit0, Implicit1, ParticleWorldTransform0, LocalTransform0, ParticleWorldTransform1, LocalTransform1, CullDistance, Dt, Context);
+				ConstructConvexTriangleMeshConstraints<T_TRAITS>(Particle0, Particle1, Implicit0, Shape0, Implicit1, Shape1, ParticleWorldTransform0, LocalTransform0, ParticleWorldTransform1, LocalTransform1, CullDistance, Dt, Context);
 			}
 			else if (Implicit0Type == FTriangleMeshImplicitObject::StaticType() && bIsConvex1)
 			{
 				if (bUseCCD)
 				{
-					ConstructConvexTriangleMeshConstraintsSwept(Particle1, Particle0, Implicit1, Implicit0, ParticleWorldTransform1, LocalTransform1, ParticleWorldTransform0, LocalTransform0, CullDistance, Dt, DirCCD, LengthCCD, Context);
+					ConstructConvexTriangleMeshConstraintsSwept(Particle1, Particle0, Implicit1, Shape1, Implicit0, Shape0, ParticleWorldTransform1, LocalTransform1, ParticleWorldTransform0, LocalTransform0, CullDistance, Dt, DirCCD, LengthCCD, Context);
 					return;
 				}
 
-				ConstructConvexTriangleMeshConstraints<T_TRAITS>(Particle1, Particle0, Implicit1, Implicit0, ParticleWorldTransform1, LocalTransform1, ParticleWorldTransform0, LocalTransform0, CullDistance, Dt, Context);
+				ConstructConvexTriangleMeshConstraints<T_TRAITS>(Particle1, Particle0, Implicit1, Shape1, Implicit0, Shape0, ParticleWorldTransform1, LocalTransform1, ParticleWorldTransform0, LocalTransform0, CullDistance, Dt, Context);
 			}
 			else if (bIsConvex0 && bIsConvex1)
 			{
 				if (bUseCCD)
 				{
-					ConstructGenericConvexConvexConstraintsSwept(Particle0, Particle1, Implicit0, Implicit1, ParticleWorldTransform0, LocalTransform0, StartX0, ParticleWorldTransform1, LocalTransform1, StartX1, CullDistance, Dt, DirCCD, LengthCCD, Context);
+					ConstructGenericConvexConvexConstraintsSwept(Particle0, Particle1, Implicit0, Shape0, Implicit1, Shape1, ParticleWorldTransform0, LocalTransform0, StartX0, ParticleWorldTransform1, LocalTransform1, StartX1, CullDistance, Dt, DirCCD, LengthCCD, Context);
 					return;
 				}
 
-				ConstructGenericConvexConvexConstraints<T_TRAITS>(Particle0, Particle1, Implicit0, Implicit1, ParticleWorldTransform0, LocalTransform0, ParticleWorldTransform1, LocalTransform1, CullDistance, Dt, Context);
+				ConstructGenericConvexConvexConstraints<T_TRAITS>(Particle0, Particle1, Implicit0, Shape0, Implicit1, Shape1, ParticleWorldTransform0, LocalTransform0, ParticleWorldTransform1, LocalTransform1, CullDistance, Dt, Context);
 			}
 			else
 			{
@@ -2162,11 +2162,11 @@ namespace Chaos
 				bool bSwap = (!FConstGenericParticleHandle(Particle0)->IsDynamic() && FConstGenericParticleHandle(Particle1)->IsDynamic());
 				if (!bSwap)
 				{
-					ConstructLevelsetLevelsetConstraints<T_TRAITS>(Particle0, Particle1, Implicit0, Simplicial0, Implicit1, Simplicial1, ParticleWorldTransform0, LocalTransform0, ParticleWorldTransform1, LocalTransform1, CullDistance, Dt, Context);
+					ConstructLevelsetLevelsetConstraints<T_TRAITS>(Particle0, Particle1, Implicit0, Shape0, Simplicial0, Implicit1, Shape1, Simplicial1, ParticleWorldTransform0, LocalTransform0, ParticleWorldTransform1, LocalTransform1, CullDistance, Dt, Context);
 				}
 				else
 				{
-					ConstructLevelsetLevelsetConstraints<T_TRAITS>(Particle1, Particle0, Implicit1, Simplicial1, Implicit0, Simplicial0, ParticleWorldTransform1, LocalTransform1, ParticleWorldTransform0, LocalTransform0, CullDistance, Dt, Context);
+					ConstructLevelsetLevelsetConstraints<T_TRAITS>(Particle1, Particle0, Implicit1, Shape1, Simplicial1, Implicit0, Shape0, Simplicial0, ParticleWorldTransform1, LocalTransform1, ParticleWorldTransform0, LocalTransform0, CullDistance, Dt, Context);
 				}
 			}
 		}
@@ -2207,7 +2207,7 @@ namespace Chaos
 		}
 
 		template<typename T_TRAITS>
-		void LegacyConstructConstraints(TGeometryParticleHandle<FReal, 3>* Particle0, TGeometryParticleHandle<FReal, 3>* Particle1, const FImplicitObject* Implicit0, const FBVHParticles* Simplicial0, const FImplicitObject* Implicit1, const FBVHParticles* Simplicial1, const FRigidTransform3& ParticleWorldTransform0, const FRigidTransform3& LocalTransform0, const FRigidTransform3& ParticleWorldTransform1, const FRigidTransform3& LocalTransform1, const FReal CullDistance, FReal Dt, const FCollisionContext& Context)
+		void LegacyConstructConstraints(TGeometryParticleHandle<FReal, 3>* Particle0, TGeometryParticleHandle<FReal, 3>* Particle1, const FImplicitObject* Implicit0, const FShapeOrShapesArray& Shapes0, const FBVHParticles* Simplicial0, const FImplicitObject* Implicit1, const FShapeOrShapesArray& Shapes1, const FBVHParticles* Simplicial1, const FRigidTransform3& ParticleWorldTransform0, const FRigidTransform3& LocalTransform0, const FRigidTransform3& ParticleWorldTransform1, const FRigidTransform3& LocalTransform1, const FReal CullDistance, FReal Dt, const FCollisionContext& Context)
 		{
 			CONDITIONAL_SCOPE_CYCLE_COUNTER(STAT_Collisions_ConstructConstraints, ConstraintsDetailedStats);
 
@@ -2217,7 +2217,18 @@ namespace Chaos
 
 			if (!Implicit0 || !Implicit1)
 			{
-				ConstructLevelsetLevelsetConstraints<T_TRAITS>(Particle0, Particle1, Implicit0, Simplicial0, Implicit1, Simplicial1, ParticleWorldTransform0, LocalTransform0, ParticleWorldTransform1, LocalTransform1, CullDistance, Dt, Context);
+				const FPerShapeData* Shape0	= nullptr;
+				if (ensure(Shapes0.IsSingleShape()))
+				{
+					Shape0 = Shapes0.GetShape();
+				}
+				const FPerShapeData* Shape1 = nullptr;
+				if (ensure(Shapes1.IsSingleShape()))
+				{
+					Shape1 = Shapes1.GetShape();
+				}
+
+				ConstructLevelsetLevelsetConstraints<T_TRAITS>(Particle0, Particle1, Implicit0, Shape0, Simplicial0, Implicit1, Shape1, Simplicial1, ParticleWorldTransform0, LocalTransform0, ParticleWorldTransform1, LocalTransform1, CullDistance, Dt, Context);
 				return;
 			}
 
@@ -2231,21 +2242,21 @@ namespace Chaos
 				const TImplicitObjectTransformed<FReal, 3>* TransformedImplicit1 = Implicit1->template GetObject<const TImplicitObjectTransformed<FReal, 3>>();
 				FRigidTransform3 TransformedTransform0 = TransformedImplicit0->GetTransform() * LocalTransform0;
 				FRigidTransform3 TransformedTransform1 = TransformedImplicit1->GetTransform() * LocalTransform1;
-				LegacyConstructConstraints<T_TRAITS>(Particle0, Particle1, TransformedImplicit0->GetTransformedObject(), Simplicial0, TransformedImplicit1->GetTransformedObject(), Simplicial1, ParticleWorldTransform0, TransformedTransform0, ParticleWorldTransform1, TransformedTransform1, CullDistance, Dt, Context);
+				LegacyConstructConstraints<T_TRAITS>(Particle0, Particle1, TransformedImplicit0->GetTransformedObject(), Shapes0, Simplicial0, TransformedImplicit1->GetTransformedObject(), Shapes1, Simplicial1, ParticleWorldTransform0, TransformedTransform0, ParticleWorldTransform1, TransformedTransform1, CullDistance, Dt, Context);
 				return;
 			}
 			else if (Implicit0OuterType == TImplicitObjectTransformed<FReal, 3>::StaticType())
 			{
 				const TImplicitObjectTransformed<FReal, 3>* TransformedImplicit0 = Implicit0->template GetObject<const TImplicitObjectTransformed<FReal, 3>>();
 				FRigidTransform3 TransformedTransform0 = TransformedImplicit0->GetTransform() * LocalTransform0;
-				LegacyConstructConstraints<T_TRAITS>(Particle0, Particle1, TransformedImplicit0->GetTransformedObject(), Simplicial0, Implicit1, Simplicial1, ParticleWorldTransform0, TransformedTransform0, ParticleWorldTransform1, LocalTransform1, CullDistance, Dt, Context);
+				LegacyConstructConstraints<T_TRAITS>(Particle0, Particle1, TransformedImplicit0->GetTransformedObject(), Shapes0, Simplicial0, Implicit1, Shapes1, Simplicial1, ParticleWorldTransform0, TransformedTransform0, ParticleWorldTransform1, LocalTransform1, CullDistance, Dt, Context);
 				return;
 			}
 			else if (Implicit1OuterType == TImplicitObjectTransformed<FReal, 3>::StaticType())
 			{
 				const TImplicitObjectTransformed<FReal, 3>* TransformedImplicit1 = Implicit1->template GetObject<const TImplicitObjectTransformed<FReal, 3>>();
 				FRigidTransform3 TransformedTransform1 = TransformedImplicit1->GetTransform() * LocalTransform1;
-				LegacyConstructConstraints<T_TRAITS>(Particle0, Particle1, Implicit0, Simplicial0, TransformedImplicit1->GetTransformedObject(), Simplicial1, ParticleWorldTransform0, LocalTransform0, ParticleWorldTransform1, TransformedTransform1, CullDistance, Dt, Context);
+				LegacyConstructConstraints<T_TRAITS>(Particle0, Particle1, Implicit0, Shapes0, Simplicial0, TransformedImplicit1->GetTransformedObject(), Shapes1, Simplicial1, ParticleWorldTransform0, LocalTransform0, ParticleWorldTransform1, TransformedTransform1, CullDistance, Dt, Context);
 				return;
 			}
 
@@ -2266,17 +2277,17 @@ namespace Chaos
 				}
 				if (InnerImplicit0 && InnerImplicit1)
 				{
-					LegacyConstructConstraints<T_TRAITS>(Particle0, Particle1, InnerImplicit0, Simplicial0, InnerImplicit1, Simplicial1, ParticleWorldTransform0, LocalTransform0, ParticleWorldTransform1, LocalTransform1, CullDistance, Dt, Context);
+					LegacyConstructConstraints<T_TRAITS>(Particle0, Particle1, InnerImplicit0, Shapes0, Simplicial0, InnerImplicit1, Shapes1, Simplicial1, ParticleWorldTransform0, LocalTransform0, ParticleWorldTransform1, LocalTransform1, CullDistance, Dt, Context);
 					return;
 				}
 				else if (InnerImplicit0 && !InnerImplicit1)
 				{
-					LegacyConstructConstraints<T_TRAITS>(Particle0, Particle1, InnerImplicit0, Simplicial0, Implicit1, Simplicial1, ParticleWorldTransform0, LocalTransform0, ParticleWorldTransform1, LocalTransform1, CullDistance, Dt, Context);
+					LegacyConstructConstraints<T_TRAITS>(Particle0, Particle1, InnerImplicit0, Shapes0, Simplicial0, Implicit1, Shapes1, Simplicial1, ParticleWorldTransform0, LocalTransform0, ParticleWorldTransform1, LocalTransform1, CullDistance, Dt, Context);
 					return;
 				}
 				else if (!InnerImplicit0 && InnerImplicit1)
 				{
-					LegacyConstructConstraints<T_TRAITS>(Particle0, Particle1, Implicit0, Simplicial0, InnerImplicit1, Simplicial1, ParticleWorldTransform0, LocalTransform0, ParticleWorldTransform1, LocalTransform1, CullDistance, Dt, Context);
+					LegacyConstructConstraints<T_TRAITS>(Particle0, Particle1, Implicit0, Shapes0, Simplicial0, InnerImplicit1, Shapes1, Simplicial1, ParticleWorldTransform0, LocalTransform0, ParticleWorldTransform1, LocalTransform1, CullDistance, Dt, Context);
 					return;
 				}
 
@@ -2286,13 +2297,22 @@ namespace Chaos
 			if (Implicit0OuterType == FImplicitObjectUnion::StaticType())
 			{
 				const FImplicitObjectUnion* Union0 = Implicit0->template GetObject<FImplicitObjectUnion>();
+				int32 ShapeIdx = 0;
 				for (const auto& Child0 : Union0->GetObjects())
 				{
-					// If shape is not sim'd, we may end up iterating over a lot of shapes on particle1's union and wasting time filtering.
-					if (Context.bFilteringEnabled == false || FilterHasSimEnabled(Particle0->GetImplicitShape(Child0.Get())))
+					if (ensure(!Shapes0.IsSingleShape()))
 					{
-						LegacyConstructConstraints<T_TRAITS>(Particle0, Particle1, Child0.Get(), Simplicial0, Implicit1, Simplicial1, ParticleWorldTransform0, LocalTransform0, ParticleWorldTransform1, LocalTransform1, CullDistance, Dt, Context);
+						const FShapesArray& ShapesArray = *Shapes0.GetShapesArray();
+						const FPerShapeData* ChildShape0 = ShapesArray[ShapeIdx].Get();
+					
+						// If shape is not sim'd, we may end up iterating over a lot of shapes on particle1's union and wasting time filtering.
+						if (Context.bFilteringEnabled == false || FilterHasSimEnabled(ChildShape0))
+						{
+							LegacyConstructConstraints<T_TRAITS>(Particle0, Particle1, Child0.Get(), FShapeOrShapesArray(ChildShape0), Simplicial0, Implicit1, Shapes1, Simplicial1, ParticleWorldTransform0, LocalTransform0, ParticleWorldTransform1, LocalTransform1, CullDistance, Dt, Context);
+						}
 					}
+
+					ShapeIdx++;
 				}
 				return;
 			}
@@ -2302,7 +2322,7 @@ namespace Chaos
 				const FImplicitObjectUnionClustered* Union0 = Implicit0->template GetObject<FImplicitObjectUnionClustered>();
 				if (Implicit1->HasBoundingBox())
 				{
-					TArray<Pair<Pair<const FImplicitObject*, const FBVHParticles*>, FRigidTransform3>> Children;
+					TArray<FLargeUnionClusteredImplicitInfo> Children;
 
 					// Need to get transformed bounds of 1 in the space of 0
 					FRigidTransform3 TM0 = LocalTransform0 * ParticleWorldTransform0;
@@ -2312,9 +2332,9 @@ namespace Chaos
 
 					Union0->FindAllIntersectingClusteredObjects(Children, QueryBounds);
 
-					for (const Pair<Pair<const FImplicitObject*, const FBVHParticles*>, FRigidTransform3>& Child0 : Children)
+					for (const FLargeUnionClusteredImplicitInfo& Child0 : Children)
 					{
-						LegacyConstructConstraints<T_TRAITS>(Particle0, Particle1, Child0.First.First, Child0.First.Second, Implicit1, Simplicial1, ParticleWorldTransform0, Child0.Second * LocalTransform0, ParticleWorldTransform1, LocalTransform1, CullDistance, Dt, Context);
+						LegacyConstructConstraints<T_TRAITS>(Particle0, Particle1, Child0.Implicit, Child0.Shapes, Child0.BVHParticles, Implicit1, Shapes1, Simplicial1, ParticleWorldTransform0, Child0.Transform * LocalTransform0, ParticleWorldTransform1, LocalTransform1, CullDistance, Dt, Context);
 					}
 				}
 				else
@@ -2322,7 +2342,11 @@ namespace Chaos
 					for (const TUniquePtr<FImplicitObject>& Child0 : Union0->GetObjects())
 					{
 						const TPBDRigidParticleHandle<FReal, 3>* OriginalParticle = Union0->FindParticleForImplicitObject(Child0.Get());
-						LegacyConstructConstraints<T_TRAITS>(Particle0, Particle1, Child0.Get(), OriginalParticle ? OriginalParticle->CollisionParticles().Get() : Simplicial0, Implicit1, Simplicial1, ParticleWorldTransform0, LocalTransform0, ParticleWorldTransform1, LocalTransform1, CullDistance, Dt, Context);
+
+						if (ensure(OriginalParticle))
+						{
+							LegacyConstructConstraints<T_TRAITS>(Particle0, Particle1, Child0.Get(), FShapeOrShapesArray(OriginalParticle), OriginalParticle->CollisionParticles().Get(), Implicit1, Shapes1, Simplicial1, ParticleWorldTransform0, LocalTransform0, ParticleWorldTransform1, LocalTransform1, CullDistance, Dt, Context);
+						}
 					}
 				}
 				return;
@@ -2331,13 +2355,22 @@ namespace Chaos
 			if (Implicit1OuterType == FImplicitObjectUnion::StaticType())
 			{
 				const FImplicitObjectUnion* Union1 = Implicit1->template GetObject<FImplicitObjectUnion>();
+				int32 ShapeIdx = 0;
 				for (const auto& Child1 : Union1->GetObjects())
 				{
-					// If shape is not sim'd, we may end up iterating over a lot of shapes on particle1's union and wasting time filtering.
-					if (Context.bFilteringEnabled == false || FilterHasSimEnabled(Particle1->GetImplicitShape(Child1.Get())))
+					if (ensure(!Shapes1.IsSingleShape()))
 					{
-						LegacyConstructConstraints<T_TRAITS>(Particle0, Particle1, Implicit0, Simplicial0, Child1.Get(), Simplicial1, ParticleWorldTransform0, LocalTransform0, ParticleWorldTransform1, LocalTransform1, CullDistance, Dt, Context);
+						const FShapesArray& ShapesArray = *Shapes1.GetShapesArray();
+						const FPerShapeData* ChildShape1 = ShapesArray[ShapeIdx].Get();
+
+						// If shape is not sim'd, we may end up iterating over a lot of shapes on particle1's union and wasting time filtering.
+						if (Context.bFilteringEnabled == false || FilterHasSimEnabled(ChildShape1))
+						{
+							LegacyConstructConstraints<T_TRAITS>(Particle0, Particle1, Implicit0, Shapes0, Simplicial0, Child1.Get(), FShapeOrShapesArray(ChildShape1), Simplicial1, ParticleWorldTransform0, LocalTransform0, ParticleWorldTransform1, LocalTransform1, CullDistance, Dt, Context);
+						}
 					}
+
+					ShapeIdx++;
 				}
 				return;
 			}
@@ -2347,7 +2380,7 @@ namespace Chaos
 				const FImplicitObjectUnionClustered* Union1 = Implicit1->template GetObject<FImplicitObjectUnionClustered>();
 				if (Implicit0->HasBoundingBox())
 				{
-					TArray<Pair<Pair<const FImplicitObject*, const FBVHParticles*>, FRigidTransform3>> Children;
+					TArray<FLargeUnionClusteredImplicitInfo> Children;
 
 					// Need to get transformed bounds of 0 in the space of 1
 					FRigidTransform3 TM0 = LocalTransform0 * ParticleWorldTransform0;
@@ -2360,9 +2393,9 @@ namespace Chaos
 						Union1->FindAllIntersectingClusteredObjects(Children, QueryBounds);
 					}
 
-					for (const Pair<Pair<const FImplicitObject*, const FBVHParticles*>, FRigidTransform3>& Child1 : Children)
+					for (const FLargeUnionClusteredImplicitInfo& Child1 : Children)
 					{
-						LegacyConstructConstraints<T_TRAITS>(Particle0, Particle1, Implicit0, Simplicial0, Child1.First.First, Child1.First.Second, ParticleWorldTransform0, LocalTransform0, ParticleWorldTransform1, Child1.Second * LocalTransform1, CullDistance, Dt, Context);
+						LegacyConstructConstraints<T_TRAITS>(Particle0, Particle1, Implicit0, Shapes0, Simplicial0, Child1.Implicit, Child1.Shapes, Child1.BVHParticles, ParticleWorldTransform0, LocalTransform0, ParticleWorldTransform1, Child1.Transform * LocalTransform1, CullDistance, Dt, Context);
 					}
 				}
 				else
@@ -2370,21 +2403,36 @@ namespace Chaos
 					for (const TUniquePtr<FImplicitObject>& Child1 : Union1->GetObjects())
 					{
 						const TPBDRigidParticleHandle<FReal, 3>* OriginalParticle = Union1->FindParticleForImplicitObject(Child1.Get());
-						LegacyConstructConstraints<T_TRAITS>(Particle0, Particle1, Implicit0, Simplicial0, Child1.Get(), OriginalParticle ? OriginalParticle->CollisionParticles().Get() : Simplicial1, ParticleWorldTransform0, LocalTransform0, ParticleWorldTransform1, LocalTransform1, CullDistance, Dt, Context);
+						if (ensure(OriginalParticle))
+						{
+							LegacyConstructConstraints<T_TRAITS>(Particle0, Particle1, Implicit0, Shapes0, Simplicial0, Child1.Get(), FShapeOrShapesArray(OriginalParticle), OriginalParticle->CollisionParticles().Get(), ParticleWorldTransform0, LocalTransform0, ParticleWorldTransform1, LocalTransform1, CullDistance, Dt, Context);
+						}
 					}
 				}
 				return;
 			}
 
+			// Should not have unions at this point, shouldn't be shapes array
+			const FPerShapeData* Shape0 = nullptr;
+			if (ensure(Shapes0.IsSingleShape()))
+			{
+				Shape0 = Shapes0.GetShape();
+			}
+			const FPerShapeData* Shape1 = nullptr;
+			if (ensure(Shapes1.IsSingleShape()))
+			{
+				Shape1 = Shapes1.GetShape();
+			}
+
 			// Check shape pair filtering if enable
-			if (Context.bFilteringEnabled && !DoCollide(Implicit0Type, Particle0->GetImplicitShape(Implicit0), Implicit1Type, Particle1->GetImplicitShape(Implicit1)))
+			if (Context.bFilteringEnabled && !DoCollide(Implicit0Type, Shape0, Implicit1Type, Shape1))
 			{
 				return;
 			}
 
 			// If we get here, we have a pair of concrete shapes (i.e., no wrappers or containers)
 			// Create a constraint for the shape pair
-			ConstructConstraintsImpl<T_TRAITS>(Particle0, Particle1, Implicit0, Simplicial0, Implicit1, Simplicial1, ParticleWorldTransform0, LocalTransform0, ParticleWorldTransform1, LocalTransform1, CullDistance, Dt, Context);
+			ConstructConstraintsImpl<T_TRAITS>(Particle0, Particle1, Implicit0, Shape0, Simplicial0, Implicit1, Shape1, Simplicial1, ParticleWorldTransform0, LocalTransform0, ParticleWorldTransform1, LocalTransform1, CullDistance, Dt, Context);
 		}
 
 		EContactShapesType CalculateShapePairType(const EImplicitObjectType Implicit0Type, const EImplicitObjectType Implicit1Type, const bool bIsConvex0, const bool bIsConvex1, const bool bIsBVH0, const bool bIsBVH1,  bool &bOutSwap)
@@ -2563,7 +2611,7 @@ namespace Chaos
 			return CalculateShapePairType(Implicit0Type, Implicit1Type, bIsConvex0, bIsConvex1, bIsBVH0, bIsBVH1, bOutSwap);
 		}
 
-		void ConstructConstraints(TGeometryParticleHandle<FReal, 3>* Particle0, TGeometryParticleHandle<FReal, 3>* Particle1, const FImplicitObject* Implicit0, const FBVHParticles* Simplicial0, const FImplicitObject* Implicit1, const FBVHParticles* Simplicial1, const FRigidTransform3& ParticleWorldTransform0, const FRigidTransform3& LocalTransform0, const FRigidTransform3& ParticleWorldTransform1, const FRigidTransform3& LocalTransform1, const FReal CullDistance, const FReal dT, const FCollisionContext& Context)
+		void ConstructConstraints(TGeometryParticleHandle<FReal, 3>* Particle0, TGeometryParticleHandle<FReal, 3>* Particle1, const FImplicitObject* Implicit0, const FPerShapeData* Shape0, const FBVHParticles* Simplicial0, const FImplicitObject* Implicit1, const FPerShapeData* Shape1, const FBVHParticles* Simplicial1, const FRigidTransform3& ParticleWorldTransform0, const FRigidTransform3& LocalTransform0, const FRigidTransform3& ParticleWorldTransform1, const FRigidTransform3& LocalTransform1, const FReal CullDistance, const FReal dT, const FCollisionContext& Context)
 		{
 			INC_DWORD_STAT(STAT_ChaosCollisionCounter_NumParticlePairs);
 
@@ -2572,12 +2620,12 @@ namespace Chaos
 			if (bDeferUpdate)
 			{
 				using TTraits = TConstructCollisionTraits<false>;
-				LegacyConstructConstraints<TTraits>(Particle0, Particle1, Implicit0, Simplicial0, Implicit1, Simplicial1, ParticleWorldTransform0, LocalTransform0, ParticleWorldTransform1, LocalTransform1, CullDistance, dT, Context);
+				LegacyConstructConstraints<TTraits>(Particle0, Particle1, Implicit0, FShapeOrShapesArray(Shape0), Simplicial0, Implicit1, FShapeOrShapesArray(Shape1), Simplicial1, ParticleWorldTransform0, LocalTransform0, ParticleWorldTransform1, LocalTransform1, CullDistance, dT, Context);
 			}
 			else
 			{
 				using TTraits = TConstructCollisionTraits<true>;
-				LegacyConstructConstraints<TTraits>(Particle0, Particle1, Implicit0, Simplicial0, Implicit1, Simplicial1, ParticleWorldTransform0, LocalTransform0, ParticleWorldTransform1, LocalTransform1, CullDistance, dT, Context);
+				LegacyConstructConstraints<TTraits>(Particle0, Particle1, Implicit0, FShapeOrShapesArray(Shape0), Simplicial0, Implicit1, FShapeOrShapesArray(Shape1), Simplicial1, ParticleWorldTransform0, LocalTransform0, ParticleWorldTransform1, LocalTransform1, CullDistance, dT, Context);
 			}
 		}
 
diff --git a/Engine/Source/Runtime/Experimental/Chaos/Private/Chaos/ContactModification.cpp b/Engine/Source/Runtime/Experimental/Chaos/Private/Chaos/ContactModification.cpp
--- a/Engine/Source/Runtime/Experimental/Chaos/Private/Chaos/ContactModification.cpp
+++ b/Engine/Source/Runtime/Experimental/Chaos/Private/Chaos/ContactModification.cpp
@@ -426,6 +426,11 @@ namespace Chaos
 		return { Constraint->GetParticle0(), Constraint->GetParticle1() };
 	}
 
+	TVec2<const FPerShapeData*> FContactPairModifier::GetShapePair() const
+	{
+		return { Constraint->GetShape0(), Constraint->GetShape1() };
+	}
+
 	void FContactPairModifierIterator::SeekValidContact()
 	{
 		// Not valid to call from end.
diff --git a/Engine/Source/Runtime/Experimental/Chaos/Private/Chaos/GeometryParticles.cpp b/Engine/Source/Runtime/Experimental/Chaos/Private/Chaos/GeometryParticles.cpp
--- a/Engine/Source/Runtime/Experimental/Chaos/Private/Chaos/GeometryParticles.cpp
+++ b/Engine/Source/Runtime/Experimental/Chaos/Private/Chaos/GeometryParticles.cpp
@@ -306,6 +306,29 @@ namespace Chaos
 
 	}
 
+	FShapeOrShapesArray::FShapeOrShapesArray(const FGeometryParticleHandle* Particle)
+	{
+		const FImplicitObject* Geometry = Particle->Geometry().Get();
+		if (Geometry)
+		{
+			if (Geometry->IsUnderlyingUnion())
+			{
+				ShapeArray = &Particle->ShapesArray();
+				bIsSingleShape = false;
+			}
+			else
+			{
+				Shape = Particle->ShapesArray()[0].Get();
+				bIsSingleShape = true;
+			}
+
+			return;
+		}
+
+		Shape = nullptr;
+		bIsSingleShape = true;
+	}
+
 
 	template <typename T, int d, EGeometryParticlesSimType SimType>
 	void TGeometryParticlesImp<T, d, SimType>::SetHandle(int32 Index, FGeometryParticleHandle* Handle)
@@ -363,91 +386,6 @@ namespace Chaos
 		Ar << SerializableGeometryParticles;
 	}
 
-	template <typename T, int d, EGeometryParticlesSimType SimType>
-	void TGeometryParticlesImp<T, d, SimType>::MapImplicitShapes()
-	{
-		int32 NumShapeArrays = MShapesArray.Num();
-		ImplicitShapeMap.Resize(NumShapeArrays);
-		for (int32 Index = 0; Index < NumShapeArrays; ++Index)
-		{
-			MapImplicitShapes(Index);
-		}
-	}
-
-	template <typename T, int d, EGeometryParticlesSimType SimType>
-	void TGeometryParticlesImp<T, d, SimType>::MapImplicitShapes(int32 Index)
-	{
-		checkSlow(Index >= 0 && Index < ImplicitShapeMap.Num());
-
-		TMap<const FImplicitObject*, int32>& Mapping = ImplicitShapeMap[Index];
-		FShapesArray& ShapeArray = MShapesArray[Index];
-		Mapping.Reset();
-
-		bool bHasCollision = false;
-		for (int32 ShapeIndex = 0; ShapeIndex < ShapeArray.Num(); ++ ShapeIndex)
-		{
-			const FImplicitObject* ImplicitObject = ShapeArray[ShapeIndex]->GetGeometry().Get();
-			Mapping.Add(ImplicitObject, ShapeIndex);
-
-			const FImplicitObject* ImplicitChildObject = Utilities::ImplicitChildHelper(ImplicitObject);
-			if (ImplicitChildObject != ImplicitObject)
-			{
-				Mapping.Add(ImplicitChildObject, ShapeIndex);
-			}
-
-			const FCollisionData& CollisionData = ShapeArray[ShapeIndex]->GetCollisionData();
-			bHasCollision |= CollisionData.HasCollisionData();
-		}
-
-		HasCollision(Index) = bHasCollision;
-
-		if (MGeometry[Index])
-		{
-			int32 CurrentShapeIndex = INDEX_NONE;
-
-			if (const auto* Union = MGeometry[Index]->template GetObject<FImplicitObjectUnion>())
-			{
-				for (const TUniquePtr<FImplicitObject>& ImplicitObject : Union->GetObjects())
-				{
-					if (ImplicitObject.Get())
-					{
-						if (const FImplicitObject* ImplicitChildObject = Utilities::ImplicitChildHelper(ImplicitObject.Get()))
-						{
-							if (ImplicitShapeMap[Index].Contains(ImplicitObject.Get()))
-							{
-								ImplicitShapeMap[Index].Add(ImplicitChildObject, CopyTemp(ImplicitShapeMap[Index][ImplicitObject.Get()]));
-							}
-							else if (ImplicitShapeMap[Index].Contains(ImplicitChildObject))
-							{
-								ImplicitShapeMap[Index].Add(ImplicitObject.Get(), CopyTemp(ImplicitShapeMap[Index][ImplicitChildObject]));
-							}
-							
-						}
-					}
-				}
-			}
-			else
-			{
-				if (const FImplicitObject* ImplicitChildObject = Utilities::ImplicitChildHelper(MGeometry[Index].Get()))
-				{
-					if (ImplicitShapeMap[Index].Contains(MGeometry[Index].Get()))
-					{
-						ImplicitShapeMap[Index].Add(ImplicitChildObject, CopyTemp(ImplicitShapeMap[Index][MGeometry[Index].Get()]));
-					}
-					else if (ImplicitShapeMap[Index].Contains(ImplicitChildObject))
-					{
-						ImplicitShapeMap[Index].Add(MGeometry[Index].Get(), CopyTemp(ImplicitShapeMap[Index][ImplicitChildObject]));
-					}
-					
-				}
-			}
-
-		}
-	}
-
-
-
-	
 	template class TGeometryParticlesImp<FReal, 3, EGeometryParticlesSimType::RigidBodySim>;
 	template class TGeometryParticlesImp<FRealSingle, 3, EGeometryParticlesSimType::Other>;
 	template class TGeometryParticlesImp<FRealDouble, 3, EGeometryParticlesSimType::Other>;
diff --git a/Engine/Source/Runtime/Experimental/Chaos/Private/Chaos/ImplicitObjectUnion.cpp b/Engine/Source/Runtime/Experimental/Chaos/Private/Chaos/ImplicitObjectUnion.cpp
--- a/Engine/Source/Runtime/Experimental/Chaos/Private/Chaos/ImplicitObjectUnion.cpp
+++ b/Engine/Source/Runtime/Experimental/Chaos/Private/Chaos/ImplicitObjectUnion.cpp
@@ -225,7 +225,7 @@ FImplicitObjectUnionClustered::FImplicitObjectUnionClustered(FImplicitObjectUnio
 	Type = ImplicitObjectType::UnionClustered;
 }
 
-void FImplicitObjectUnionClustered::FindAllIntersectingClusteredObjects(TArray<Pair<Pair<const FImplicitObject*,const FBVHParticles*>,FRigidTransform3>>& Out, const FAABB3& LocalBounds) const
+void FImplicitObjectUnionClustered::FindAllIntersectingClusteredObjects(TArray<FLargeUnionClusteredImplicitInfo>& Out, const FAABB3& LocalBounds) const
 {
 	if (LargeUnionData)
 	{
@@ -234,8 +234,22 @@ void FImplicitObjectUnionClustered::FindAllIntersectingClusteredObjects(TArray<P
 		for (int32 Idx : Overlaps)
 		{
 			const FImplicitObject* Obj = LargeUnionData->GeomParticles.Geometry(Idx).Get();
+			const FShapesArray& ShapesArray = LargeUnionData->GeomParticles.ShapesArray(Idx);
+
+			// If implicit is union return shapes array of its particle, otherwise retrieve single shape.
+			FShapeOrShapesArray Shapes;
+			if (Obj->IsUnderlyingUnion())
+			{
+				Shapes = FShapeOrShapesArray(&ShapesArray);
+			}
+			else
+			{
+				Shapes = FShapeOrShapesArray(ShapesArray[0].Get());
+			}
+			
+
 			const FBVHParticles* Simplicial = MOriginalParticleLookupHack.IsValidIndex(Idx) ? MOriginalParticleLookupHack[Idx]->CollisionParticles().Get() : nullptr;
-			Out.Add(MakePair(MakePair(Obj,Simplicial),FRigidTransform3(LargeUnionData->GeomParticles.X(Idx), LargeUnionData->GeomParticles.R(Idx))));
+			Out.Add(FLargeUnionClusteredImplicitInfo(Obj,FRigidTransform3(LargeUnionData->GeomParticles.X(Idx), LargeUnionData->GeomParticles.R(Idx)), Simplicial, Shapes));
 		}
 	}
 	else
@@ -255,8 +269,18 @@ void FImplicitObjectUnionClustered::FindAllIntersectingClusteredObjects(TArray<P
 		for (int32 Idx = 0; Idx < LocalOut.Num(); ++Idx)
 		{
 			auto& OutElem = LocalOut[Idx];
-			const FBVHParticles* Simplicial = MOriginalParticleLookupHack.IsValidIndex(Idxs[Idx]) ? MOriginalParticleLookupHack[Idxs[Idx]]->CollisionParticles().Get() : nullptr;
-			Out.Add(MakePair(MakePair(OutElem.First, Simplicial), OutElem.Second));
+
+			FShapeOrShapesArray Shapes;
+			const FBVHParticles* Simplicial = nullptr;
+
+			if (ensure(MOriginalParticleLookupHack.IsValidIndex(Idxs[Idx])))
+			{
+				FPBDRigidParticleHandle* Handle = MOriginalParticleLookupHack[Idxs[Idx]];
+				Shapes = FShapeOrShapesArray(Handle);
+				Simplicial = Handle->CollisionParticles().Get();
+			}
+
+			Out.Add(FLargeUnionClusteredImplicitInfo(OutElem.First, OutElem.Second, Simplicial, Shapes));
 		}
 	}
 }
diff --git a/Engine/Source/Runtime/Experimental/Chaos/Private/Chaos/ParticleHandle.cpp b/Engine/Source/Runtime/Experimental/Chaos/Private/Chaos/ParticleHandle.cpp
--- a/Engine/Source/Runtime/Experimental/Chaos/Private/Chaos/ParticleHandle.cpp
+++ b/Engine/Source/Runtime/Experimental/Chaos/Private/Chaos/ParticleHandle.cpp
@@ -26,64 +26,6 @@ namespace Chaos
 		}
 	}
 
-	template <typename T, int d>
-	void Chaos::TGeometryParticle<T, d>::MapImplicitShapes()
-	{
-		ImplicitShapeMap.Reset();
-
-		for (int32 ShapeIndex = 0; ShapeIndex < MShapesArray.Num(); ++ ShapeIndex)
-		{
-			const FImplicitObject* ImplicitObject = MShapesArray[ShapeIndex]->GetGeometry().Get();
-			ImplicitShapeMap.Add(ImplicitObject, ShapeIndex);
-
-			const FImplicitObject* ImplicitChildObject = Utilities::ImplicitChildHelper(ImplicitObject);
-			if (ImplicitChildObject != ImplicitObject)
-			{
-				ImplicitShapeMap.Add(ImplicitChildObject, ShapeIndex);
-			}
-		}
-
-		auto Geometry = MNonFrequentData.Read().Geometry();
-		if (Geometry)
-		{
-			int32 CurrentShapeIndex = INDEX_NONE;
-			if (const auto* Union = Geometry->template GetObject<FImplicitObjectUnion>())
-			{
-				for (const TUniquePtr<FImplicitObject>& ImplicitObject : Union->GetObjects())
-				{
-					if (ImplicitObject.Get())
-					{
-						if (const FImplicitObject* ImplicitChildObject = Utilities::ImplicitChildHelper(ImplicitObject.Get()))
-						{
-							if (ImplicitShapeMap.Contains(ImplicitObject.Get()))
-							{
-								ImplicitShapeMap.Add(ImplicitChildObject, CopyTemp(ImplicitShapeMap[ImplicitObject.Get()]));
-							}
-							else if (ImplicitShapeMap.Contains(ImplicitChildObject))
-							{
-								ImplicitShapeMap.Add(ImplicitObject.Get(), CopyTemp(ImplicitShapeMap[ImplicitChildObject]));
-							}
-						}
-					}
-				}
-			}
-			else 
-			{
-				if (const FImplicitObject* ImplicitChildObject = Utilities::ImplicitChildHelper(Geometry.Get()))
-				{
-					if (ImplicitShapeMap.Contains(Geometry.Get()))
-					{
-						ImplicitShapeMap.Add(ImplicitChildObject, CopyTemp(ImplicitShapeMap[Geometry.Get()]));
-					}
-					else if (ImplicitShapeMap.Contains(ImplicitChildObject))
-					{
-						ImplicitShapeMap.Add(Geometry.Get(), CopyTemp(ImplicitShapeMap[ImplicitChildObject]));
-					}
-				}
-			}
-		}
-	}
-
 	inline FImplicitObject* GetInstancedImplicitHelper(FImplicitObject* Implicit0)
 	{
 		EImplicitObjectType Implicit0OuterType = Implicit0->GetType();
@@ -187,11 +129,17 @@ namespace Chaos
 		}
 		else
 		{
-			if (const auto* PerShapeData = GetImplicitShape(Implicit))
+
+			// Find our shape and see if sim is enabled.
+			for (const TUniquePtr<FPerShapeData>& Shape : ShapesArray())
 			{
-				if (!PerShapeData->GetSimEnabled())
+				if (Shape->GetGeometry().Get() == Implicit)
 				{
-					return;
+					if (!Shape->GetSimEnabled())
+					{
+						return;
+					}
+					break;
 				}
 			}
 			if (bIgnoreAnalyticCollisions)
diff --git a/Engine/Source/Runtime/Experimental/Chaos/Private/EventDefaults.cpp b/Engine/Source/Runtime/Experimental/Chaos/Private/EventDefaults.cpp
--- a/Engine/Source/Runtime/Experimental/Chaos/Private/EventDefaults.cpp
+++ b/Engine/Source/Runtime/Experimental/Chaos/Private/EventDefaults.cpp
@@ -131,9 +131,9 @@ namespace Chaos
 
 						const FImplicitObject* Implicit0 = Constraint.GetImplicit0();
 						const FImplicitObject* Implicit1 = Constraint.GetImplicit1();
-
-						const FPerShapeData* Shape0 = Particle0->GetImplicitShape(Implicit0);
-						const FPerShapeData* Shape1 = Particle1->GetImplicitShape(Implicit1);
+	
+						const FPerShapeData* Shape0 = Constraint.GetShape0();
+						const FPerShapeData* Shape1 = Constraint.GetShape1();
 
 						// If we don't have a filter - allow the notify, otherwise obey the filter flag
 						const bool bFilter0Notify = Shape0 ? Shape0->GetSimData().HasFlag(EFilterFlags::ContactNotify) : true;
@@ -184,8 +184,8 @@ namespace Chaos
 						Data.Proxy1 = Particle0 ? const_cast<IPhysicsProxyBase*>(Particle0->PhysicsProxy()) : nullptr;
 						Data.Proxy2 = Particle1 ? const_cast<IPhysicsProxyBase*>(Particle1->PhysicsProxy()) : nullptr;
 
-						const FPerShapeData* Shape0 = Particle0 ? Particle0->GetImplicitShape(Constraint.GetImplicit0()) : nullptr;
-						const FPerShapeData* Shape1 = Particle1 ? Particle1->GetImplicitShape(Constraint.GetImplicit1()) : nullptr;
+						const FPerShapeData* Shape0 = Constraint.GetShape0();
+						const FPerShapeData* Shape1 = Constraint.GetShape1();
 
 						Data.Mat1 = ResolveMaterial(Shape0, Constraint);
 						Data.Mat2 = ResolveMaterial(Shape1, Constraint);
diff --git a/Engine/Source/Runtime/Experimental/Chaos/Public/Chaos/ChaosDebugDraw.h b/Engine/Source/Runtime/Experimental/Chaos/Public/Chaos/ChaosDebugDraw.h
--- a/Engine/Source/Runtime/Experimental/Chaos/Public/Chaos/ChaosDebugDraw.h
+++ b/Engine/Source/Runtime/Experimental/Chaos/Public/Chaos/ChaosDebugDraw.h
@@ -14,6 +14,7 @@ namespace Chaos
 	template <typename PayloadType, typename T, int d> class ISpatialAcceleration;
 	class FPBDSuspensionConstraints;
 	class FRigidClustering;
+	class FShapeOrShapesArray;
 
 	namespace DebugDraw
 	{
@@ -211,7 +212,7 @@ namespace Chaos
 		CHAOS_API void DrawJointConstraints(const FRigidTransform3& SpaceTransform, const TArray<FPBDJointConstraintHandle*>& ConstraintHandles, FRealSingle ColorScale, const FChaosDebugDrawJointFeatures& FeatureMask = FChaosDebugDrawJointFeatures::MakeDefault(), const FChaosDebugDrawSettings* Settings = nullptr);
 		CHAOS_API void DrawJointConstraints(const FRigidTransform3& SpaceTransform, const FPBDJointConstraints& Constraints, FRealSingle ColorScale, const FChaosDebugDrawJointFeatures& FeatureMask = FChaosDebugDrawJointFeatures::MakeDefault(), const FChaosDebugDrawSettings* Settings = nullptr);
 		CHAOS_API void DrawSimulationSpace(const FSimulationSpace& SimSpace, const FChaosDebugDrawSettings* Settings = nullptr);
-		CHAOS_API void DrawShape(const FRigidTransform3& ShapeTransform, const FImplicitObject* Shape, const FColor& Color, const FChaosDebugDrawSettings* Settings = nullptr);
+		CHAOS_API void DrawShape(const FRigidTransform3& ShapeTransform, const FImplicitObject* Implicit, const FShapeOrShapesArray& Shapes, const FColor& Color, const FChaosDebugDrawSettings* Settings = nullptr);
 		CHAOS_API void DrawConstraintGraph(const FRigidTransform3& ShapeTransform, const FPBDConstraintGraph& Graph, const FChaosDebugDrawSettings* Settings = nullptr);
 		CHAOS_API void DrawConnectionGraph(const FRigidClustering& Clustering, const FChaosDebugDrawSettings* Settings = nullptr);
 		CHAOS_API void DrawCollidingShapes(const FRigidTransform3& SpaceTransform, const FPBDCollisionConstraints& Collisions, FRealSingle ColorScale, const FChaosDebugDrawSettings* Settings = nullptr);
diff --git a/Engine/Source/Runtime/Experimental/Chaos/Public/Chaos/Collision/PBDCollisionConstraint.h b/Engine/Source/Runtime/Experimental/Chaos/Public/Chaos/Collision/PBDCollisionConstraint.h
--- a/Engine/Source/Runtime/Experimental/Chaos/Public/Chaos/Collision/PBDCollisionConstraint.h
+++ b/Engine/Source/Runtime/Experimental/Chaos/Public/Chaos/Collision/PBDCollisionConstraint.h
@@ -28,6 +28,7 @@ namespace Chaos
 	class FSolverBody;
 	class FSolverBodyContainer;
 	class FPBDCollisionSolverContainer;
+	class FPerShapeData;
 
 	UE_DEPRECATED(4.27, "Use FPBDCollisionConstraint instead")
 	typedef FPBDCollisionConstraint FRigidBodyPointContactConstraint;
@@ -220,10 +221,12 @@ namespace Chaos
 		static void MakeInline(
 			FGeometryParticleHandle* Particle0,
 			const FImplicitObject* Implicit0,
+			const FPerShapeData* Shape0,
 			const FBVHParticles* Simplicial0,
 			const FRigidTransform3& ImplicitLocalTransform0,
 			FGeometryParticleHandle* Particle1,
 			const FImplicitObject* Implicit1,
+			const FPerShapeData* Shape1,
 			const FBVHParticles* Simplicial1,
 			const FRigidTransform3& ImplicitLocalTransform1,
 			const FReal InCullDistance,
@@ -239,10 +242,12 @@ namespace Chaos
 		static TUniquePtr<FPBDCollisionConstraint> Make(
 			FGeometryParticleHandle* Particle0,
 			const FImplicitObject* Implicit0,
+			const FPerShapeData* Shape0,
 			const FBVHParticles* Simplicial0,
 			const FRigidTransform3& ImplicitLocalTransform0,
 			FGeometryParticleHandle* Particle1,
 			const FImplicitObject* Implicit1,
+			const FPerShapeData* Shape1,
 			const FBVHParticles* Simplicial1,
 			const FRigidTransform3& ImplicitLocalTransform1,
 			const FReal InCullDistance,
@@ -302,6 +307,10 @@ namespace Chaos
 		const FImplicitObject* GetImplicit1() const { return Implicit[1]; }
 		const FImplicitObject* GetImplicit(const int32 ParticleIndex) const { check((ParticleIndex >= 0) && (ParticleIndex < 2)); return Implicit[ParticleIndex]; }
 
+		const FPerShapeData* GetShape0() const { return Shape[0]; }
+		const FPerShapeData* GetShape1() const { return Shape[1]; }
+		const FPerShapeData* GetShape(const int32 ParticleIndex) const { check((ParticleIndex >= 0) && (ParticleIndex < 2)); return Shape[ParticleIndex]; }
+
 		const FBVHParticles* GetCollisionParticles0() const { return Simplicial[0]; }
 		const FBVHParticles* GetCollisionParticles1() const { return Simplicial[1]; }
 		const FBVHParticles* GetCollisionParticles(const int32 ParticleIndex) const { check((ParticleIndex >= 0) && (ParticleIndex < 2)); return Simplicial[ParticleIndex]; }
@@ -550,9 +559,11 @@ namespace Chaos
 		FPBDCollisionConstraint(
 			FGeometryParticleHandle* Particle0,
 			const FImplicitObject* Implicit0,
+			const FPerShapeData* Shape0,
 			const FBVHParticles* Simplicial0,
 			FGeometryParticleHandle* Particle1,
 			const FImplicitObject* Implicit1,
+			const FPerShapeData* Shape1,
 			const FBVHParticles* Simplicial1);
 
 		// Set all the data not initialized in the constructor
@@ -604,6 +615,7 @@ namespace Chaos
 		
 		FGeometryParticleHandle* Particle[2];
 		const FImplicitObject* Implicit[2];
+		const FPerShapeData* Shape[2];
 		const FBVHParticles* Simplicial[2];
 		FCollisionContact Manifold;// @todo(chaos): rename to FCollisionMaterial or something
 		FReal Stiffness;
diff --git a/Engine/Source/Runtime/Experimental/Chaos/Public/Chaos/Collision/ParticlePairMidPhase.h b/Engine/Source/Runtime/Experimental/Chaos/Public/Chaos/Collision/ParticlePairMidPhase.h
--- a/Engine/Source/Runtime/Experimental/Chaos/Public/Chaos/Collision/ParticlePairMidPhase.h
+++ b/Engine/Source/Runtime/Experimental/Chaos/Public/Chaos/Collision/ParticlePairMidPhase.h
@@ -174,10 +174,12 @@ namespace Chaos
 		FPBDCollisionConstraint* FindOrCreateConstraint(
 			FGeometryParticleHandle* InParticle0,
 			const FImplicitObject* Implicit0,
+			const FPerShapeData* Shape0,
 			const FBVHParticles* BVHParticles0,
 			const FRigidTransform3& ShapeRelativeTransform0,
 			FGeometryParticleHandle* InParticle1,
 			const FImplicitObject* Implicit1,
+			const FPerShapeData* Shape1,
 			const FBVHParticles* BVHParticles1,
 			const FRigidTransform3& ShapeRelativeTransform1,
 			const FReal CullDistance,
@@ -190,10 +192,12 @@ namespace Chaos
 		FPBDCollisionConstraint* FindOrCreateSweptConstraint(
 			FGeometryParticleHandle* InParticle0,
 			const FImplicitObject* Implicit0,
+			const FPerShapeData* Shape0,
 			const FBVHParticles* BVHParticles0,
 			const FRigidTransform3& ShapeRelativeTransform0,
 			FGeometryParticleHandle* InParticle1,
 			const FImplicitObject* Implicit1,
+			const FPerShapeData* Shape1,
 			const FBVHParticles* BVHParticles1,
 			const FRigidTransform3& ShapeRelativeTransform1,
 			const FReal CullDistance,
@@ -224,10 +228,12 @@ namespace Chaos
 		FPBDCollisionConstraint* CreateConstraint(
 			FGeometryParticleHandle* Particle0,
 			const FImplicitObject* Implicit0,
+			const FPerShapeData* Shape0,
 			const FBVHParticles* BVHParticles0,
 			const FRigidTransform3& ShapeRelativeTransform0,
 			FGeometryParticleHandle* Particle1,
 			const FImplicitObject* Implicit1,
+			const FPerShapeData* Shape1,
 			const FBVHParticles* BVHParticles1,
 			const FRigidTransform3& ShapeRelativeTransform1,
 			const FReal CullDistance,
diff --git a/Engine/Source/Runtime/Experimental/Chaos/Public/Chaos/CollisionResolution.h b/Engine/Source/Runtime/Experimental/Chaos/Public/Chaos/CollisionResolution.h
--- a/Engine/Source/Runtime/Experimental/Chaos/Public/Chaos/CollisionResolution.h
+++ b/Engine/Source/Runtime/Experimental/Chaos/Public/Chaos/CollisionResolution.h
@@ -70,7 +70,7 @@ namespace Chaos
 		bool CHAOS_API UpdateConstraintFromGeometrySwept(FPBDCollisionConstraint& Constraint, const FRigidTransform3& ParticleTransform0, const FRigidTransform3& ParticleTransform1, const FReal Dt);
 
 		// Create constraints for the particle pair. This could create multiple constraints: one for each potentially colliding shape pair in multi-shape particles.
-		void CHAOS_API ConstructConstraints(TGeometryParticleHandle<FReal, 3>* Particle0, TGeometryParticleHandle<FReal, 3>* Particle1, const FImplicitObject* Implicit0, const FBVHParticles* Simplicial0, const FImplicitObject* Implicit1, const FBVHParticles* Simplicial1, const FRigidTransform3& ParticleWorldTransform0, const FRigidTransform3& Transform0, const FRigidTransform3& ParticleWorldTransform1, const FRigidTransform3& Transform1, const FReal CullDistance, const FReal Dt,const FCollisionContext& Context);
+		void CHAOS_API ConstructConstraints(TGeometryParticleHandle<FReal, 3>* Particle0, TGeometryParticleHandle<FReal, 3>* Particle1, const FImplicitObject* Implicit0, const FPerShapeData* Shape0, const FBVHParticles* Simplicial0, const FImplicitObject* Implicit1, const FPerShapeData* Shape1, const FBVHParticles* Simplicial1, const FRigidTransform3& ParticleWorldTransform0, const FRigidTransform3& Transform0, const FRigidTransform3& ParticleWorldTransform1, const FRigidTransform3& Transform1, const FReal CullDistance, const FReal Dt,const FCollisionContext& Context);
 
 		// @todo(chaos): this is only called in tests - should it really be exposed?
 		template<ECollisionUpdateType UpdateType>
diff --git a/Engine/Source/Runtime/Experimental/Chaos/Public/Chaos/ContactModification.h b/Engine/Source/Runtime/Experimental/Chaos/Public/Chaos/ContactModification.h
--- a/Engine/Source/Runtime/Experimental/Chaos/Public/Chaos/ContactModification.h
+++ b/Engine/Source/Runtime/Experimental/Chaos/Public/Chaos/ContactModification.h
@@ -11,6 +11,7 @@ namespace Chaos
 	class FPBDCollisionConstraints;
 	class FCollisionContactModifier;
 	class FContactPairModifier;
+	class FPerShapeData;
 
 	class CHAOS_API FContactPairModifier
 	{
@@ -202,6 +203,11 @@ namespace Chaos
 		*/
 		TVec2<FGeometryParticleHandle*> GetParticlePair() const;
 
+		/*
+		* Get shape pair from constraint.
+		*/
+		TVec2<const FPerShapeData*> GetShapePair() const;
+
 	private:
 		/**
 		 * @brief Update cached shape transforms in the constraint after modifying particle positions
diff --git a/Engine/Source/Runtime/Experimental/Chaos/Public/Chaos/GeometryParticles.h b/Engine/Source/Runtime/Experimental/Chaos/Public/Chaos/GeometryParticles.h
--- a/Engine/Source/Runtime/Experimental/Chaos/Public/Chaos/GeometryParticles.h
+++ b/Engine/Source/Runtime/Experimental/Chaos/Public/Chaos/GeometryParticles.h
@@ -27,6 +27,9 @@ namespace Chaos
 	class FConstraintHandle;
 	class FParticleCollisions;
 
+	using FShapesArray = TArray<TUniquePtr<FPerShapeData>, TInlineAllocator<1>>;
+	using FConstraintHandleArray = TArray<FConstraintHandle*>;
+
 	/** Data that is associated with geometry. If a union is used an entry is created per internal geometry */
 	class CHAOS_API FPerShapeData
 	{
@@ -298,16 +301,66 @@ namespace Chaos
 		const FImplicitObject* LeafGeometry;
 	};
 
+	/**
+	* Union between shape and shapes array pointers, used for passing around shapes with
+	* implicit that could be single implicit or union.
+	*/
+	class FShapeOrShapesArray
+	{
+	public:
+
+		// Store particle's shape array if particle has union geometry, otherwise individual shape.
+		FShapeOrShapesArray(const FGeometryParticleHandle* Particle);
+
+		FShapeOrShapesArray()
+			: Shape(nullptr)
+			, bIsSingleShape(true)
+		{}
+
+		FShapeOrShapesArray(const FPerShapeData* InShape)
+			: Shape(InShape)
+			, bIsSingleShape(true)
+		{}
+
+		FShapeOrShapesArray(const FShapesArray* InShapeArray)
+			: ShapeArray(InShapeArray)
+			, bIsSingleShape(false)
+		{}
+
+		bool IsSingleShape() const { return bIsSingleShape; }
+
+		bool IsValid() const { return Shape != nullptr; }
+
+		// Do not call without checking IsSingleShape().
+		const FPerShapeData* GetShape() const
+		{
+			check(bIsSingleShape);
+			return Shape;
+		}
+
+		// Do not call without checking IsSingleShape().
+		const FShapesArray* GetShapesArray() const
+		{
+			check(!bIsSingleShape);
+			return ShapeArray;
+		}
+
+	private:
+		union
+		{
+			const FPerShapeData* Shape;
+			const FShapesArray* ShapeArray;
+		};
+
+		bool bIsSingleShape;
+	};
+
 	inline FChaosArchive& operator<<(FChaosArchive& Ar, FPerShapeData& Shape)
 	{
 		Shape.Serialize(Ar);
 		return Ar;
 	}
 
-	using FShapesArray = TArray<TUniquePtr<FPerShapeData>, TInlineAllocator<1>>;
-
-	using FConstraintHandleArray = TArray<FConstraintHandle*>;
-
 	void CHAOS_API UpdateShapesArrayFromGeometry(FShapesArray& ShapesArray, TSerializablePtr<FImplicitObject> Geometry, const FRigidTransform3& ActorTM, IPhysicsProxyBase* Proxy);
 
 	FORCEINLINE uint32 GetTypeHash(const FParticleID& Unique)
@@ -383,7 +436,6 @@ namespace Chaos
 			TArrayCollection::AddArray(&MParticleIDs);
 			TArrayCollection::AddArray(&MHasCollision);
 			TArrayCollection::AddArray(&MShapesArray);
-			TArrayCollection::AddArray(&ImplicitShapeMap);
 			TArrayCollection::AddArray(&MLocalBounds);
 			TArrayCollection::AddArray(&MWorldSpaceInflatedBounds);
 			TArrayCollection::AddArray(&MHasBounds);
@@ -423,7 +475,6 @@ namespace Chaos
 			, MPhysicsProxy(MoveTemp(Other.MPhysicsProxy))
 			, MHasCollision(MoveTemp(Other.MHasCollision))
 			, MShapesArray(MoveTemp(Other.MShapesArray))
-			, ImplicitShapeMap(MoveTemp(Other.ImplicitShapeMap))
 			, MLocalBounds(MoveTemp(Other.MLocalBounds))
 			, MWorldSpaceInflatedBounds(MoveTemp(Other.MWorldSpaceInflatedBounds))
 			, MHasBounds(MoveTemp(Other.MHasBounds))
@@ -450,7 +501,6 @@ namespace Chaos
 			TArrayCollection::AddArray(&MDynamicGeometry);
 			TArrayCollection::AddArray(&MHasCollision);
 			TArrayCollection::AddArray(&MShapesArray);
-			TArrayCollection::AddArray(&ImplicitShapeMap);
 			TArrayCollection::AddArray(&MLocalBounds);
 			TArrayCollection::AddArray(&MWorldSpaceInflatedBounds);
 			TArrayCollection::AddArray(&MHasBounds);
@@ -493,7 +543,6 @@ namespace Chaos
 			TArrayCollection::AddArray(&MDynamicGeometry);
 			TArrayCollection::AddArray(&MHasCollision);
 			TArrayCollection::AddArray(&MShapesArray);
-			TArrayCollection::AddArray(&ImplicitShapeMap);
 			TArrayCollection::AddArray(&MLocalBounds);
 			TArrayCollection::AddArray(&MWorldSpaceInflatedBounds);
 			TArrayCollection::AddArray(&MHasBounds);
@@ -579,7 +628,6 @@ namespace Chaos
 			MGeometry[Index] = InGeometry;
 
 			UpdateShapesArray(Index);
-			MapImplicitShapes(Index);
 
 			MHasBounds[Index] = (InGeometry && InGeometry->HasBoundingBox());
 			if (MHasBounds[Index])
@@ -757,11 +805,6 @@ public:
 			if (Ar.CustomVer(FPhysicsObjectVersion::GUID) >= FPhysicsObjectVersion::PerShapeData)
 			{
 				Ar << MShapesArray;
-
-				if(Ar.IsLoading())
-				{
-					MapImplicitShapes();
-				}
 			}
 
 			if (Ar.CustomVer(FPhysicsObjectVersion::GUID) >= FPhysicsObjectVersion::SerializeGTGeometryParticles)
@@ -816,22 +859,6 @@ public:
 
 		FORCEINLINE EParticleType ParticleType() const { return MParticleType; }
 
-		CHAOS_API const FPerShapeData* GetImplicitShape(int32 Index, const FImplicitObject* InObject)
-		{
-			checkSlow(Index >= 0 && Index < ImplicitShapeMap.Num());
-			TMap<const FImplicitObject*, int32>& Mapping = ImplicitShapeMap[Index];
-			FShapesArray& ShapeArray = MShapesArray[Index];
-			int32* ShapeIndex = Mapping.Find(InObject);
-
-			if(ShapeIndex && ShapeArray.IsValidIndex(*ShapeIndex))
-			{
-				return ShapeArray[*ShapeIndex].Get();
-			}
-
-			return nullptr;
-		}
-
-
 		FORCEINLINE TArray<TRotation<T, d>>& AllR() { return MR; }
 		FORCEINLINE TArray<TAABB<T, d>>& AllLocalBounds() { return MLocalBounds; }
 		FORCEINLINE TArray<TAABB<T, d>>& AllWorldSpaceInflatedBounds() { return MWorldSpaceInflatedBounds; }
@@ -857,7 +884,6 @@ public:
 		TArrayCollectionArray<IPhysicsProxyBase*> MPhysicsProxy;
 		TArrayCollectionArray<bool> MHasCollision;
 		TArrayCollectionArray<FShapesArray> MShapesArray;
-		TArrayCollectionArray<TMap<const FImplicitObject*, int32>> ImplicitShapeMap;
 		TArrayCollectionArray<TAABB<T,d>> MLocalBounds;
 		TArrayCollectionArray<TAABB<T, d>> MWorldSpaceInflatedBounds;
 		TArrayCollectionArray<bool> MHasBounds;
@@ -874,13 +900,8 @@ public:
 		void UpdateShapesArray(const int32 Index)
 		{
 			UpdateShapesArrayFromGeometry(MShapesArray[Index], MGeometry[Index], FRigidTransform3(X(Index), R(Index)), nullptr);
-			MapImplicitShapes(Index);
 		}
 
-		CHAOS_API void MapImplicitShapes();
-
-		CHAOS_API void MapImplicitShapes(int32 Index);
-
 		template <typename T2, int d2, EGeometryParticlesSimType SimType2>
 		friend class TGeometryParticlesImp;
 
diff --git a/Engine/Source/Runtime/Experimental/Chaos/Public/Chaos/ImplicitObjectUnion.h b/Engine/Source/Runtime/Experimental/Chaos/Public/Chaos/ImplicitObjectUnion.h
--- a/Engine/Source/Runtime/Experimental/Chaos/Public/Chaos/ImplicitObjectUnion.h
+++ b/Engine/Source/Runtime/Experimental/Chaos/Public/Chaos/ImplicitObjectUnion.h
@@ -202,6 +202,22 @@ protected:
 	TUniquePtr<FLargeImplicitObjectUnionData> LargeUnionData;	//only needed when there are many objects
 };
 
+struct FLargeUnionClusteredImplicitInfo
+{
+	FLargeUnionClusteredImplicitInfo(const FImplicitObject* InImplicit, const FRigidTransform3& InTransform, const FBVHParticles* InBVHParticles, const FShapeOrShapesArray& InShapes)
+		: Implicit(InImplicit)
+		, Transform(InTransform)
+		, BVHParticles(InBVHParticles)
+		, Shapes(InShapes)
+	{
+	}
+
+	const FImplicitObject* Implicit;
+	FRigidTransform3 Transform;
+	const FBVHParticles* BVHParticles;
+	FShapeOrShapesArray Shapes;
+};
+
 class CHAOS_API FImplicitObjectUnionClustered: public FImplicitObjectUnion
 {
 public:
@@ -216,7 +232,7 @@ public:
 		return ImplicitObjectType::UnionClustered;
 	}
 
-	void FindAllIntersectingClusteredObjects(TArray<Pair<Pair<const FImplicitObject*, const FBVHParticles*>, FRigidTransform3>>& Out, const FAABB3& LocalBounds) const;
+	void FindAllIntersectingClusteredObjects(TArray<FLargeUnionClusteredImplicitInfo>& Out, const FAABB3& LocalBounds) const;
 	TArray<FPBDRigidParticleHandle*> FindAllIntersectingChildren(const FAABB3& LocalBounds) const;
 
 #if CHAOS_PARTICLEHANDLE_TODO
diff --git a/Engine/Source/Runtime/Experimental/Chaos/Public/Chaos/ParticleHandle.h b/Engine/Source/Runtime/Experimental/Chaos/Public/Chaos/ParticleHandle.h
--- a/Engine/Source/Runtime/Experimental/Chaos/Public/Chaos/ParticleHandle.h
+++ b/Engine/Source/Runtime/Experimental/Chaos/Public/Chaos/ParticleHandle.h
@@ -666,11 +666,6 @@ public:
 		GeometryParticles->SetHandle(ParticleIdx, this);
 	}
 
-	const FPerShapeData* GetImplicitShape(const FImplicitObject* InObject) const
-	{
-		return GeometryParticles->GetImplicitShape(ParticleIdx, InObject);
-	}
-
 	FWeakParticleHandle& WeakParticleHandle()
 	{
 		return GeometryParticles->WeakParticleHandle(ParticleIdx);
@@ -1891,11 +1886,6 @@ public:
 		{
 			UpdateShapeBounds();
 		}
-
-		if(Ar.IsLoading())
-		{
-			MapImplicitShapes();
-		}
 	}
 
 	virtual bool IsParticleValid() const
@@ -2038,7 +2028,6 @@ public:
 	{
 		ensure(InShapesArray.Num() == MShapesArray.Num());
 		MShapesArray = MoveTemp(InShapesArray);
-		MapImplicitShapes();
 	}
 
 	void MergeShapesArray(FShapesArray&& OtherShapesArray)
@@ -2049,7 +2038,6 @@ public:
 			ensure(Idx < MShapesArray.Num());
 			MShapesArray[Idx++] = MoveTemp(Shape);
 		}
-		MapImplicitShapes();
 	}
 
 	void SetIgnoreAnalyticCollisionsImp(FImplicitObject* Implicit, bool bIgnoreAnalyticCollisions);
@@ -2140,16 +2128,6 @@ public:
 		return nullptr;
 	}
 
-	const FPerShapeData* GetImplicitShape(const FImplicitObject* InImplicit) const
-	{
-		const int32* ShapeIndex = ImplicitShapeMap.Find(InImplicit);
-		if(ShapeIndex)
-		{
-			return MShapesArray[*ShapeIndex].Get();
-		}
-
-		return nullptr;
-	}
 
 	void SyncRemoteData(FDirtyPropertiesManager& Manager, int32 DataIdx, FDirtyChaosProperties& RemoteData, const TArray<int32>& ShapeDataIndices, FShapeDirtyData* ShapesRemoteData) const
 	{
@@ -2240,8 +2218,6 @@ private:
 	void* MUserData;
 
 	FShapesArray MShapesArray;
-	TMap<const FImplicitObject*, int32> ImplicitShapeMap;
-
 
 public:
 	// Ryan: FGeometryCollectionPhysicsProxy needs access to GeometrySharedLowLevel(), 
@@ -2267,7 +2243,6 @@ protected:
 	void UpdateShapesArray()
 	{
 		UpdateShapesArrayFromGeometry(MShapesArray, MakeSerializable(MNonFrequentData.Read().Geometry()), FRigidTransform3(X(), R()), Proxy);
-		MapImplicitShapes();
 	}
 
 	virtual void SyncRemoteDataImp(FDirtyPropertiesManager& Manager, int32 DataIdx, const FDirtyChaosProperties& RemoteData) const
@@ -2275,8 +2250,6 @@ protected:
 		MXR.SyncRemote(Manager, DataIdx, RemoteData);
 		MNonFrequentData.SyncRemote(Manager, DataIdx, RemoteData);
 	}
-
-	CHAOS_API void MapImplicitShapes();
 };
 
 template <typename T, int d>
diff --git a/Engine/Source/Runtime/PhysicsCore/Private/SQAccelerator.cpp b/Engine/Source/Runtime/PhysicsCore/Private/SQAccelerator.cpp
--- a/Engine/Source/Runtime/PhysicsCore/Private/SQAccelerator.cpp
+++ b/Engine/Source/Runtime/PhysicsCore/Private/SQAccelerator.cpp
@@ -17,6 +17,7 @@
 #include "Chaos/CastingUtilities.h"
 #include "Chaos/ISpatialAcceleration.h"
 #include "Chaos/PBDCollisionConstraints.h"
+#include "Chaos/GeometryParticles.h"
 #include "Chaos/GeometryQueries.h"
 #include "Chaos/DebugDrawQueue.h"
 
@@ -428,7 +429,7 @@ private:
 		}
 		else if (SQ == ESQType::Overlap)
 		{
-			Chaos::DebugDraw::DrawShape(StartTM, QueryGeom, bHit ? FColor::Red : FColor::Green);
+			Chaos::DebugDraw::DrawShape(StartTM, QueryGeom, Chaos::FShapeOrShapesArray(), bHit ? FColor::Red : FColor::Green);
 		}
 
 		if (Instance.bHasBounds)
