Chaos Cloth - Fix collisions against convexes now that they're floats.

Additionally add lots of compile time verification that the ispc structs and enums match the c structs and enums for collision objects. This required adding a new accessor (which I called ISPCDataVerifier) to many of the implicit objects which can report the offset and size of its members.

Highlighted code changes:

TaperedCylinder.h
Switched this to hold FPlaneConcrete instead of FPlane. The difference is that FPlane is an FImplicitObject and otherwise has a bunch of passthroughs to FPlaneConcrete. We never use the fact that these are FImplicitObjects, so we don't need this additional 32 bytes per plane.

PerParticlePBDCollisionConstraint.ispc
On top of updating the convex code to be the correct float types, I also added th e FTaperedCapsule and FTaperedCylinder ispc structs. This makes it easier to verify this data.

ImplicitObject.h, ChaosClothingSimulationCollider.cpp
Technically all TSpheres think they have ImplicitObjectType::Sphere, but all of our code assumes that ImplicitObjectType::Sphere corresponds with TSphere<FReal,3> when casting to the derived type. I added a new typedef for FSphere which is this one type.

diff --git a/Engine/Plugins/Experimental/ChaosCloth/Source/Chaos/Private/ChaosCloth/ChaosClothingSimulationCollider.cpp b/Engine/Plugins/Experimental/ChaosCloth/Source/Chaos/Private/ChaosCloth/ChaosClothingSimulationCollider.cpp
--- a/Engine/Plugins/Experimental/ChaosCloth/Source/Chaos/Private/ChaosCloth/ChaosClothingSimulationCollider.cpp
+++ b/Engine/Plugins/Experimental/ChaosCloth/Source/Chaos/Private/ChaosCloth/ChaosClothingSimulationCollider.cpp
@@ -124,7 +124,7 @@ void FClothingSimulationCollider::FLODData::Add(
 			{
 				// Sphere
 				Solver->SetCollisionGeometry(CapsuleOffset, Index,
-					MakeUnique<TSphere<FReal, 3>>(Center, MaxRadius));
+					MakeUnique<FSphere>(Center, MaxRadius));
 			}
 			else if (MaxRadius - MinRadius < KINDA_SMALL_NUMBER)
 			{
@@ -147,9 +147,9 @@ void FClothingSimulationCollider::FLODData::Add(
 					Objects.Add(TUniquePtr<FImplicitObject>(
 						new FTaperedCylinder(P0, P1, Radius0, Radius1)));
 					Objects.Add(TUniquePtr<FImplicitObject>(
-						new TSphere<FReal, 3>(P0, Radius0)));
+						new FSphere(P0, Radius0)));
 					Objects.Add(TUniquePtr<FImplicitObject>(
-						new TSphere<FReal, 3>(P1, Radius1)));
+						new FSphere(P1, Radius1)));
 					Solver->SetCollisionGeometry(CapsuleOffset, Index,
 						MakeUnique<FImplicitObjectUnion>(MoveTemp(Objects)));
 				}
@@ -181,7 +181,7 @@ void FClothingSimulationCollider::FLODData::Add(
 			BaseTransforms[Index] = Softs::FSolverRigidTransform3::Identity;
 
 			Solver->SetCollisionGeometry(SphereOffset, Index,
-				MakeUnique<TSphere<FReal, 3>>(
+				MakeUnique<FSphere>(
 					Sphere.LocalPosition * InScale,
 					Sphere.Radius * InScale));
 
@@ -264,7 +264,7 @@ void FClothingSimulationCollider::FLODData::Add(
 			else
 			{
 				UE_LOG(LogChaosCloth, Warning, TEXT("Replacing invalid convex collision by a default unit sphere."));
-				Solver->SetCollisionGeometry(ConvexOffset, Index, MakeUnique<TSphere<FReal, 3>>(FVec3(0.0f), 1.0f));  // Default to a unit sphere to replace the faulty convex
+				Solver->SetCollisionGeometry(ConvexOffset, Index, MakeUnique<FSphere>(FVec3(0.0f), 1.0f));  // Default to a unit sphere to replace the faulty convex
 			}
 		}
 	}
diff --git a/Engine/Source/Runtime/Core/Public/Math/Vector.isph b/Engine/Source/Runtime/Core/Public/Math/Vector.isph
--- a/Engine/Source/Runtime/Core/Public/Math/Vector.isph
+++ b/Engine/Source/Runtime/Core/Public/Math/Vector.isph
@@ -2637,7 +2637,12 @@ inline uniform FVector VectorAbs(const uniform FVector& A)
 	return *((uniform FVector *uniform)&Result);
 }
 
-inline bool IsVectorEqual(const FVector &A, const FVector &B)
+inline bool IsVectorEqual(const FVector3d &A, const FVector3d &B)
+{
+	return A.V[0] == B.V[0] && A.V[1] == B.V[1] && A.V[2] == B.V[2];
+}
+
+inline bool IsVectorEqual(const FVector3f &A, const FVector3f &B)
 {
 	return A.V[0] == B.V[0] && A.V[1] == B.V[1] && A.V[2] == B.V[2];
 }
@@ -3467,9 +3472,15 @@ inline varying float VectorDot( const varying FVector3f &Vec1, const varying FVe
 	return Temp.V[0] + Temp.V[1] + Temp.V[2];
 }
 
-inline varying FReal VectorDot( const varying FVector &Vec1, const uniform FVector &Vec2 )
+inline varying double VectorDot( const varying FVector3d &Vec1, const uniform FVector3d &Vec2 )
+{
+	const varying FVector3d Temp = Vec1 * Vec2;
+	return Temp.V[0] + Temp.V[1] + Temp.V[2];
+}
+
+inline varying float VectorDot( const varying FVector3f &Vec1, const uniform FVector3f &Vec2 )
 {
-	const varying FVector Temp = Vec1 * Vec2;
+	const varying FVector3f Temp = Vec1 * Vec2;
 	return Temp.V[0] + Temp.V[1] + Temp.V[2];
 }
 
diff --git a/Engine/Source/Runtime/Experimental/Chaos/Private/Chaos/PerParticlePBDCollisionConstraint.cpp b/Engine/Source/Runtime/Experimental/Chaos/Private/Chaos/PerParticlePBDCollisionConstraint.cpp
--- a/Engine/Source/Runtime/Experimental/Chaos/Private/Chaos/PerParticlePBDCollisionConstraint.cpp
+++ b/Engine/Source/Runtime/Experimental/Chaos/Private/Chaos/PerParticlePBDCollisionConstraint.cpp
@@ -5,11 +5,104 @@
 #include "HAL/IConsoleManager.h"
 #if INTEL_ISPC
 #include "PerParticlePBDCollisionConstraint.ispc.generated.h"
+
+#if !UE_BUILD_SHIPPING
+// These includes are used for the static_asserts below
+#include "Chaos/Capsule.h"
+#include "Chaos/Convex.h"
+#include "Chaos/ImplicitObjectUnion.h"
+#include "Chaos/Sphere.h"
+#include "Chaos/TaperedCapsule.h"
+#include "Chaos/TaperedCylinder.h"
+#endif 
 #endif
 
 #if INTEL_ISPC && !UE_BUILD_SHIPPING
 static_assert(sizeof(ispc::FVector3f) == sizeof(Chaos::Softs::FSolverVec3), "sizeof(ispc::FVector3f) != sizeof(Chaos::Softs::FSolverVec3)");
 static_assert(sizeof(ispc::FVector4f) == sizeof(Chaos::Softs::FSolverRotation3), "sizeof(ispc::FVector4f) != sizeof(Chaos::Softs::FSolverRotation3)");
+static_assert(ispc::ImplicitObjectType::Sphere == Chaos::ImplicitObjectType::Sphere);
+static_assert(ispc::ImplicitObjectType::Capsule == Chaos::ImplicitObjectType::Capsule);
+static_assert(ispc::ImplicitObjectType::Union == Chaos::ImplicitObjectType::Union);
+static_assert(ispc::ImplicitObjectType::TaperedCapsule == Chaos::ImplicitObjectType::TaperedCapsule);
+static_assert(ispc::ImplicitObjectType::Convex == Chaos::ImplicitObjectType::Convex);
+static_assert(sizeof(ispc::TArray) == sizeof(TArray<int>));
+// Sphere
+static_assert(sizeof(Chaos::FImplicitObject) == Chaos::FSphere::FISPCDataVerifier::OffsetOfCenter());
+static_assert(sizeof(ispc::FVector) == Chaos::FSphere::FISPCDataVerifier::SizeOfCenter());
+// Capsule
+static_assert(sizeof(Chaos::FImplicitObject) == Chaos::FCapsule::FISPCDataVerifier::OffsetOfMSegment());
+static_assert(sizeof(ispc::Segment) == Chaos::FCapsule::FISPCDataVerifier::SizeOfMSegment());
+// Union (only specific case of MObjects = [FTaperedCylinder, Sphere, Sphere] is used here.
+static_assert(sizeof(Chaos::FImplicitObject) == Chaos::FImplicitObjectUnion::FISPCDataVerifier::OffsetOfMObjects());
+static_assert(sizeof(ispc::TArray) == Chaos::FImplicitObjectUnion::FISPCDataVerifier::SizeOfMObjects());
+// TaperedCylinder
+static_assert(sizeof(Chaos::FImplicitObject) + offsetof(ispc::FTaperedCylinder, MPlane1) == Chaos::FTaperedCylinder::FISPCDataVerifier::OffsetOfMPlane1());
+static_assert(sizeof(ispc::FTaperedCylinder::MPlane1) == Chaos::FTaperedCylinder::FISPCDataVerifier::SizeOfMPlane1());
+static_assert(sizeof(Chaos::FImplicitObject) + offsetof(ispc::FTaperedCylinder, MPlane2) == Chaos::FTaperedCylinder::FISPCDataVerifier::OffsetOfMPlane2());
+static_assert(sizeof(ispc::FTaperedCylinder::MPlane2) == Chaos::FTaperedCylinder::FISPCDataVerifier::SizeOfMPlane2());
+static_assert(sizeof(Chaos::FImplicitObject) + offsetof(ispc::FTaperedCylinder, Height) == Chaos::FTaperedCylinder::FISPCDataVerifier::OffsetOfMHeight());
+static_assert(sizeof(ispc::FTaperedCylinder::Height) == Chaos::FTaperedCylinder::FISPCDataVerifier::SizeOfMHeight());
+static_assert(sizeof(Chaos::FImplicitObject) + offsetof(ispc::FTaperedCylinder, Radius1) == Chaos::FTaperedCylinder::FISPCDataVerifier::OffsetOfMRadius1());
+static_assert(sizeof(ispc::FTaperedCylinder::Radius1) == Chaos::FTaperedCylinder::FISPCDataVerifier::SizeOfMRadius1());
+static_assert(sizeof(Chaos::FImplicitObject) + offsetof(ispc::FTaperedCylinder, Radius2) == Chaos::FTaperedCylinder::FISPCDataVerifier::OffsetOfMRadius2());
+static_assert(sizeof(ispc::FTaperedCylinder::Radius2) == Chaos::FTaperedCylinder::FISPCDataVerifier::SizeOfMRadius2());
+// TaperedCapsule
+static_assert(sizeof(Chaos::FImplicitObject) + offsetof(ispc::FTaperedCapsule, Origin) == Chaos::FTaperedCapsule::FISPCDataVerifier::OffsetOfOrigin());
+static_assert(sizeof(ispc::FTaperedCapsule::Origin) == Chaos::FTaperedCapsule::FISPCDataVerifier::SizeOfOrigin());
+static_assert(sizeof(Chaos::FImplicitObject) + offsetof(ispc::FTaperedCapsule, Axis) == Chaos::FTaperedCapsule::FISPCDataVerifier::OffsetOfAxis());
+static_assert(sizeof(ispc::FTaperedCapsule::Axis) == Chaos::FTaperedCapsule::FISPCDataVerifier::SizeOfAxis());
+static_assert(sizeof(Chaos::FImplicitObject) + offsetof(ispc::FTaperedCapsule, Height) == Chaos::FTaperedCapsule::FISPCDataVerifier::OffsetOfHeight());
+static_assert(sizeof(ispc::FTaperedCapsule::Height) == Chaos::FTaperedCapsule::FISPCDataVerifier::SizeOfHeight());
+static_assert(sizeof(Chaos::FImplicitObject) + offsetof(ispc::FTaperedCapsule, Radius1) == Chaos::FTaperedCapsule::FISPCDataVerifier::OffsetOfRadius1());
+static_assert(sizeof(ispc::FTaperedCapsule::Radius1) == Chaos::FTaperedCapsule::FISPCDataVerifier::SizeOfRadius1());
+static_assert(sizeof(Chaos::FImplicitObject) + offsetof(ispc::FTaperedCapsule, Radius2) == Chaos::FTaperedCapsule::FISPCDataVerifier::OffsetOfRadius2());
+static_assert(sizeof(ispc::FTaperedCapsule::Radius2) == Chaos::FTaperedCapsule::FISPCDataVerifier::SizeOfRadius2());
+// Convex
+static_assert(sizeof(Chaos::FImplicitObject) + offsetof(ispc::FConvex, Planes) == Chaos::FConvex::FISPCDataVerifier::OffsetOfPlanes());
+static_assert(sizeof(ispc::FConvex::Planes) == Chaos::FConvex::FISPCDataVerifier::SizeOfPlanes());
+static_assert(sizeof(Chaos::FImplicitObject) + offsetof(ispc::FConvex, Vertices) == Chaos::FConvex::FISPCDataVerifier::OffsetOfVertices());
+static_assert(sizeof(ispc::FConvex::Vertices) == Chaos::FConvex::FISPCDataVerifier::SizeOfVertices());
+static_assert(sizeof(Chaos::FImplicitObject) + offsetof(ispc::FConvex, StructureData) == Chaos::FConvex::FISPCDataVerifier::OffsetOfStructureData());
+static_assert(sizeof(ispc::FConvex::StructureData) == Chaos::FConvex::FISPCDataVerifier::SizeOfStructureData());
+
+static_assert(sizeof(ispc::FPlaneConcrete3f) == sizeof(Chaos::FConvex::FPlaneType));
+static_assert(sizeof(ispc::FVector3f) == sizeof(Chaos::FConvex::FVec3Type)); // Vertices array
+
+static_assert(offsetof(ispc::FConvexStructureData, Data) == Chaos::FConvexStructureData::FISPCDataVerifier::OffsetOfData());
+static_assert(sizeof(ispc::FConvexStructureData::Data) == Chaos::FConvexStructureData::FISPCDataVerifier::SizeOfData());
+static_assert(offsetof(ispc::FConvexStructureData, IndexType) == Chaos::FConvexStructureData::FISPCDataVerifier::OffsetOfIndexType());
+static_assert(sizeof(ispc::FConvexStructureData::IndexType) == Chaos::FConvexStructureData::FISPCDataVerifier::SizeOfIndexType());
+
+static_assert(offsetof(ispc::FConvexStructureDataImp, Planes) == Chaos::FConvexStructureData::FConvexStructureDataSmall::FISPCDataVerifier::OffsetOfPlanes());
+static_assert(offsetof(ispc::FConvexStructureDataImp, Planes) == Chaos::FConvexStructureData::FConvexStructureDataMedium::FISPCDataVerifier::OffsetOfPlanes());
+static_assert(offsetof(ispc::FConvexStructureDataImp, Planes) == Chaos::FConvexStructureData::FConvexStructureDataLarge::FISPCDataVerifier::OffsetOfPlanes());
+static_assert(sizeof(ispc::FConvexStructureDataImp::Planes) == Chaos::FConvexStructureData::FConvexStructureDataSmall::FISPCDataVerifier::SizeOfPlanes());
+static_assert(sizeof(ispc::FConvexStructureDataImp::Planes) == Chaos::FConvexStructureData::FConvexStructureDataMedium::FISPCDataVerifier::SizeOfPlanes());
+static_assert(sizeof(ispc::FConvexStructureDataImp::Planes) == Chaos::FConvexStructureData::FConvexStructureDataLarge::FISPCDataVerifier::SizeOfPlanes());
+
+static_assert(offsetof(ispc::FConvexStructureDataImp, HalfEdges) == Chaos::FConvexStructureData::FConvexStructureDataSmall::FISPCDataVerifier::OffsetOfHalfEdges());
+static_assert(offsetof(ispc::FConvexStructureDataImp, HalfEdges) == Chaos::FConvexStructureData::FConvexStructureDataMedium::FISPCDataVerifier::OffsetOfHalfEdges());
+static_assert(offsetof(ispc::FConvexStructureDataImp, HalfEdges) == Chaos::FConvexStructureData::FConvexStructureDataLarge::FISPCDataVerifier::OffsetOfHalfEdges());
+static_assert(sizeof(ispc::FConvexStructureDataImp::HalfEdges) == Chaos::FConvexStructureData::FConvexStructureDataSmall::FISPCDataVerifier::SizeOfHalfEdges());
+static_assert(sizeof(ispc::FConvexStructureDataImp::HalfEdges) == Chaos::FConvexStructureData::FConvexStructureDataMedium::FISPCDataVerifier::SizeOfHalfEdges());
+static_assert(sizeof(ispc::FConvexStructureDataImp::HalfEdges) == Chaos::FConvexStructureData::FConvexStructureDataLarge::FISPCDataVerifier::SizeOfHalfEdges());
+
+static_assert(offsetof(ispc::FConvexStructureDataImp, Vertices) == Chaos::FConvexStructureData::FConvexStructureDataSmall::FISPCDataVerifier::OffsetOfVertices());
+static_assert(offsetof(ispc::FConvexStructureDataImp, Vertices) == Chaos::FConvexStructureData::FConvexStructureDataMedium::FISPCDataVerifier::OffsetOfVertices());
+static_assert(offsetof(ispc::FConvexStructureDataImp, Vertices) == Chaos::FConvexStructureData::FConvexStructureDataLarge::FISPCDataVerifier::OffsetOfVertices());
+static_assert(sizeof(ispc::FConvexStructureDataImp::Vertices) == Chaos::FConvexStructureData::FConvexStructureDataSmall::FISPCDataVerifier::SizeOfVertices());
+static_assert(sizeof(ispc::FConvexStructureDataImp::Vertices) == Chaos::FConvexStructureData::FConvexStructureDataMedium::FISPCDataVerifier::SizeOfVertices());
+static_assert(sizeof(ispc::FConvexStructureDataImp::Vertices) == Chaos::FConvexStructureData::FConvexStructureDataLarge::FISPCDataVerifier::SizeOfVertices());
+
+static_assert(sizeof(ispc::PlanesS) == sizeof(Chaos::FConvexStructureData::FConvexStructureDataSmall::FPlaneData));
+static_assert(sizeof(ispc::PlanesM) == sizeof(Chaos::FConvexStructureData::FConvexStructureDataMedium::FPlaneData));
+static_assert(sizeof(ispc::PlanesL) == sizeof(Chaos::FConvexStructureData::FConvexStructureDataLarge::FPlaneData));
+static_assert(sizeof(ispc::HalfEdgesS) == sizeof(Chaos::FConvexStructureData::FConvexStructureDataSmall::FHalfEdgeData));
+static_assert(sizeof(ispc::HalfEdgesM) == sizeof(Chaos::FConvexStructureData::FConvexStructureDataMedium::FHalfEdgeData));
+static_assert(sizeof(ispc::HalfEdgesL) == sizeof(Chaos::FConvexStructureData::FConvexStructureDataLarge::FHalfEdgeData));
+static_assert(sizeof(ispc::VerticesS) == sizeof(Chaos::FConvexStructureData::FConvexStructureDataSmall::FVertexData));
+static_assert(sizeof(ispc::VerticesM) == sizeof(Chaos::FConvexStructureData::FConvexStructureDataMedium::FVertexData));
+static_assert(sizeof(ispc::VerticesL) == sizeof(Chaos::FConvexStructureData::FConvexStructureDataLarge::FVertexData));
 
 bool bChaos_PerParticleCollision_ISPC_Enabled = true;
 FAutoConsoleVariableRef CVarChaosPerParticleCollisionISPCEnabled(TEXT("p.Chaos.PerParticleCollision.ISPC"), bChaos_PerParticleCollision_ISPC_Enabled, TEXT("Whether to use ISPC optimizations in per particle collisions"));
diff --git a/Engine/Source/Runtime/Experimental/Chaos/Private/Chaos/PerParticlePBDCollisionConstraint.ispc b/Engine/Source/Runtime/Experimental/Chaos/Private/Chaos/PerParticlePBDCollisionConstraint.ispc
--- a/Engine/Source/Runtime/Experimental/Chaos/Private/Chaos/PerParticlePBDCollisionConstraint.ispc
+++ b/Engine/Source/Runtime/Experimental/Chaos/Private/Chaos/PerParticlePBDCollisionConstraint.ispc
@@ -49,18 +49,6 @@ struct TArray
 	uint32 ArrayMax;
 };
 
-struct TPairU8
-{
-	uint16 Key;
-	uint8 Value;
-};
-
-struct TPairS32
-{
-	int32 Key;
-	int32 Value;
-};
-
 struct PlanesS
 {
 	unsigned int8 FirstHalfEdgeIndex;
@@ -122,15 +110,16 @@ struct Segment
 	FReal Length;
 };
 
-struct PlaneConcrete
+struct FPlaneConcrete3d
 {
-	FVector MX;
-	FVector MNormal;
+	FVector3d MX;
+	FVector3d MNormal;
 };
 
-struct Plane
+struct FPlaneConcrete3f
 {
-	PlaneConcrete MPlaneConcrete;
+	FVector3f MX;
+	FVector3f MNormal;
 };
 
 struct FConvexStructureDataImp
@@ -139,6 +128,15 @@ struct FConvexStructureDataImp
 	TArray HalfEdges;
 	TArray Vertices;
 	TArray Edges;
+	// TArray VertexPlanes; // this code never uses this, and we never assume to know the sizeof(FConvexStructureDataImp)
+};
+
+enum FConvexStructureDataIndexType
+{
+	EIndexType_None = 0,
+	EIndexType_Small,
+	EIndexType_Medium,
+	EIndexType_Large
 };
 
 struct FConvexStructureData
@@ -147,17 +145,62 @@ struct FConvexStructureData
 	int8 IndexType;
 };
 
-struct ConvexShape
+struct FConvex
 {
 	TArray Planes;
 	TArray Vertices;
-	FVector LocalBoundingBoxMin;
-	FVector LocalBoundingBoxMax;
+	FVector3f LocalBoundingBoxMin;
+	FVector3f LocalBoundingBoxMax;
 	FConvexStructureData StructureData;
-	FReal Volume;
-	FVector CenterOfMass;
+	// FReal Volume; // this code never uses this, and we never assume to know the sizeof(FConvex)
+	// FVector CenterOfMass; // this code never uses this, and we never assume to know the sizeof(FConvex)
 };
 
+struct FTaperedCapsule
+{
+	FVector Origin;
+	FVector Axis;
+	FReal Height;
+	FReal Radius1;
+	FReal Radius2;
+	// FAABB3 LocalBoundingBox; // this code never uses this, and we never assume to know the sizeof(FTaperedCapsule)
+};
+
+struct FTaperedCylinder
+{
+	FPlaneConcrete3d MPlane1;
+	FPlaneConcrete3d MPlane2;
+	FReal Height;
+	FReal Radius1;
+	FReal Radius2;
+	// FAABB3 MLocalBoundingBox; // this code never uses this, and we never assume to know the sizeof(FTaperedCylinder)
+};
+
+// We want to export all of the structs above, so defining exported method that uses them.
+export uniform int32 ExportHack_PerParticlePBDCollisionConstraintStructs(
+	uniform ImplicitObjectType Type,
+	uniform TArray* uniform Array,
+	uniform PlanesS* uniform _PlanesS,
+	uniform PlanesM* uniform _PlanesM,
+	uniform PlanesL* uniform _PlanesL,
+	uniform HalfEdgesS* uniform _HalfEdgesS,
+	uniform HalfEdgesM* uniform _HalfEdgesM,
+	uniform HalfEdgesL* uniform _HalfEdgesL,
+	uniform VerticesS* uniform _VerticesS,
+	uniform VerticesM* uniform _VerticesM,
+	uniform VerticesL* uniform _VerticesL,
+	uniform Segment* uniform _Segment,
+	uniform FPlaneConcrete3d* uniform _Plane3d,
+	uniform FPlaneConcrete3f* uniform _Plane3f,
+	uniform FConvexStructureDataIndexType IndexType,
+	uniform FConvex* uniform _FConvex,
+	uniform FTaperedCapsule* uniform _TaperedCapsule,
+	uniform FTaperedCylinder* uniform _TaperedCylinder
+)
+{
+	return 0;
+}
+
 static inline FReal SafeNormalize(FVector &Direction)
 {
 	const FReal SizeSqr = VectorSizeSquared(Direction);
@@ -185,26 +228,19 @@ static inline void PhiWithNormalSphere(const uniform uint8 *uniform TypedObjectP
 	Phi = SafeNormalize(Normal) - Radius;
 }
 
-static inline void PhiWithNormalPlane(const uniform uint8 *uniform TypedObjectPtr, const FVector &PhiWithNormalInput, FVector &Normal, FReal &Phi)
-{
-	const uniform Plane P = *((const uniform Plane *uniform)&TypedObjectPtr[0]);
-	Normal = P.MPlaneConcrete.MNormal;
-	Phi = VectorDot(PhiWithNormalInput - P.MPlaneConcrete.MX, P.MPlaneConcrete.MNormal);
-}
-
-static inline void PhiWithNormalPlane(const uniform Plane *varying InPlane, const FVector &PhiWithNormalInput, FVector &Normal, FReal &Phi)
+static inline void PhiWithNormalPlane(const uniform FPlaneConcrete3f *varying InPlane, const FVector3f &PhiWithNormalInput, FVector3f &Normal, float &Phi)
 {
-	const varying FVector MNormal = VectorGather(&InPlane->MPlaneConcrete.MNormal);
-	const varying FVector MX = VectorGather(&InPlane->MPlaneConcrete.MX);
+	const varying FVector3f MNormal = VectorGather(&InPlane->MNormal);
+	const varying FVector3f MX = VectorGather(&InPlane->MX);
 
 	Normal = MNormal;
 	Phi = VectorDot(PhiWithNormalInput - MX, MNormal);
 }
 
-static inline FReal SignedDistancePlane(const uniform Plane *uniform InPlane, const FVector &x)
+static inline float SignedDistancePlane(const uniform FPlaneConcrete3f *uniform InPlane, const FVector3f &x)
 {
-	const uniform Plane P = *InPlane;
-	return VectorDot(x - P.MPlaneConcrete.MX, P.MPlaneConcrete.MNormal);
+	const uniform FPlaneConcrete3f P = *InPlane;
+	return VectorDot(x - P.MX, P.MNormal);
 }
 
 static inline void PhiWithNormalCapsule(const uniform uint8 *uniform TypedObjectPtr, const FVector &PhiWithNormalInput, FVector &Normal, FReal &Phi, const uniform FReal Radius)
@@ -219,72 +255,64 @@ static inline void PhiWithNormalCapsule(const uniform uint8 *uniform TypedObject
 
 static inline void PhiWithNormalTaperedCapsule(const uniform uint8* uniform TypedObjectPtr, const FVector& PhiWithNormalInput, FVector& Normal, FReal& Phi)
 {
-	const uniform FVector Origin = *((const uniform FVector *uniform)& TypedObjectPtr[0]);
-	const uniform FVector Axis   = *((const uniform FVector *uniform)& TypedObjectPtr[sizeof(uniform FVector)]);
-	const uniform FReal Height   = *((const uniform FReal *uniform)& TypedObjectPtr[2*sizeof(uniform FVector)]);
-	const uniform FReal Radius1  = *((const uniform FReal *uniform)& TypedObjectPtr[2*sizeof(uniform FVector)+ sizeof(uniform FReal)]);
-	const uniform FReal Radius2  = *((const uniform FReal *uniform)& TypedObjectPtr[2*sizeof(uniform FVector)+ 2*sizeof(uniform FReal)]);
+	const uniform FTaperedCapsule CapsuleObj = *((const uniform FTaperedCapsule *uniform)&TypedObjectPtr[0]);
 
-	const FReal DistanceAlongAxis = clamp(VectorDot(PhiWithNormalInput - Origin, Axis), (uniform FReal)(0), Height);
-	const FVector ProjectedPoint = SetVector(DistanceAlongAxis) * Axis + Origin;
-	const FReal Alpha = DistanceAlongAxis / Height;
-	const FReal Radius = Radius1 * ((uniform FReal)(1.0) - Alpha) + Radius2 * Alpha;
+	const FReal DistanceAlongAxis = clamp(VectorDot(PhiWithNormalInput - CapsuleObj.Origin, CapsuleObj.Axis), (uniform FReal)(0), CapsuleObj.Height);
+	const FVector ProjectedPoint = SetVector(DistanceAlongAxis) * CapsuleObj.Axis + CapsuleObj.Origin;
+	const FReal Alpha = DistanceAlongAxis / CapsuleObj.Height;
+	const FReal Radius = CapsuleObj.Radius1 * ((uniform FReal)(1.0) - Alpha) + CapsuleObj.Radius2 * Alpha;
 	Normal = PhiWithNormalInput - ProjectedPoint;
 	Phi = SafeNormalize(Normal) - Radius;
 }
 
 static inline void PhiWithNormalTaperedCylinder(const uniform uint8 *uniform TypedObjectPtr, const FVector &PhiWithNormalInput, FVector &Normal, FReal &Phi, const uniform int SizeofFImplicitObject)
 {
-	const uniform Plane MPlane1 = *((const uniform Plane *uniform)&TypedObjectPtr[SizeofFImplicitObject]);
-	const uniform Plane MPlane2 = *((const uniform Plane *uniform)&TypedObjectPtr[2*SizeofFImplicitObject+sizeof(uniform Plane)]);
-	const uniform FReal MHeight = *((const uniform FReal *uniform)&TypedObjectPtr[2*SizeofFImplicitObject+2*sizeof(uniform Plane)]);
-	const uniform FReal MRadius1 = *((const uniform FReal *uniform)&TypedObjectPtr[2*SizeofFImplicitObject + 2*sizeof(uniform Plane) + sizeof(uniform FReal)]);
-	const uniform FReal MRadius2 = *((const uniform FReal *uniform)&TypedObjectPtr[2*SizeofFImplicitObject + 2*sizeof(uniform Plane) + 2*sizeof(uniform FReal)]);
-
-	const FReal Distance1 = VectorDot(PhiWithNormalInput - MPlane1.MPlaneConcrete.MX, MPlane1.MPlaneConcrete.MNormal);
-	// Used to be Distance2 = MPlane2.PhiWithNormal(x, Normal2); but that would trigger 
-	const FReal Distance2 = MHeight - Distance1;          // the ensure on Distance2 being slightly larger than MHeight in some border cases
-	const FVector SideVector = (PhiWithNormalInput - (MPlane1.MPlaneConcrete.MNormal * Distance1 + MPlane1.MPlaneConcrete.MX));
-	const FReal Alpha = Distance1 / MHeight;
-	const FReal Radius = MRadius1 * ((uniform FReal)(1.0) - Alpha) + MRadius2 * Alpha;
+	const uniform FTaperedCylinder CylinderObj = *((const uniform FTaperedCylinder *uniform)&TypedObjectPtr[0]);
+
+	const FReal Distance1 = VectorDot(PhiWithNormalInput - CylinderObj.MPlane1.MX, CylinderObj.MPlane1.MNormal);
+	// Used to be Distance2 = CylinderObj.MPlane2.PhiWithNormal(x, Normal2); but that would trigger 
+	const FReal Distance2 = CylinderObj.Height - Distance1;          // the ensure on Distance2 being slightly larger than MHeight in some border cases
+	const FVector SideVector = (PhiWithNormalInput - (CylinderObj.MPlane1.MNormal * Distance1 + CylinderObj.MPlane1.MX));
+	const FReal Alpha = Distance1 / CylinderObj.Height;
+	const FReal Radius = CylinderObj.Radius1 * ((uniform FReal)(1.0) - Alpha) + CylinderObj.Radius2 * Alpha;
 	const FReal SideDistance = VectorSize(SideVector) - Radius;
 	const FReal TopDistance = Distance1 < Distance2 ? Distance1 : Distance2;
 
 	if (Distance1 < SMALL_NUMBER)
 	{
-		const FVector v = PhiWithNormalInput - (MPlane1.MPlaneConcrete.MNormal * Distance1 + MPlane1.MPlaneConcrete.MX);
-		if (VectorSize(v) > MRadius1)
+		const FVector v = PhiWithNormalInput - (CylinderObj.MPlane1.MNormal * Distance1 + CylinderObj.MPlane1.MX);
+		if (VectorSize(v) > CylinderObj.Radius1)
 		{
-			const FVector Corner = VectorGetSafeNormal(v) * MRadius1 + MPlane1.MPlaneConcrete.MX;
+			const FVector Corner = VectorGetSafeNormal(v) * CylinderObj.Radius1 + CylinderObj.MPlane1.MX;
 			const FVector CornerVector = PhiWithNormalInput - Corner;
 			Normal = VectorGetSafeNormal(CornerVector);
 			Phi = VectorSize(CornerVector);
 		}
 		else
 		{
-			Normal = MPlane1.MPlaneConcrete.MNormal * (uniform FReal)(-1.0);
+			Normal = CylinderObj.MPlane1.MNormal * (uniform FReal)(-1.0);
 			Phi = -Distance1;
 		}
 	}
 	else if (Distance2 < SMALL_NUMBER)
 	{
-		const FVector v = PhiWithNormalInput - (MPlane2.MPlaneConcrete.MNormal * Distance2 + MPlane2.MPlaneConcrete.MX);
-		if (VectorSize(v) > MRadius2)
+		const FVector v = PhiWithNormalInput - (CylinderObj.MPlane2.MNormal * Distance2 + CylinderObj.MPlane2.MX);
+		if (VectorSize(v) > CylinderObj.Radius2)
 		{
-			const FVector Corner = VectorGetSafeNormal(v) * MRadius2 + MPlane2.MPlaneConcrete.MX;
+			const FVector Corner = VectorGetSafeNormal(v) * CylinderObj.Radius2 + CylinderObj.MPlane2.MX;
 			const FVector CornerVector = PhiWithNormalInput - Corner;
 			Normal = VectorGetSafeNormal(CornerVector);
 			Phi = VectorSize(CornerVector);
 		}
 		else
 		{
-			Normal = MPlane2.MPlaneConcrete.MNormal * (uniform FReal)(-1.0);
+			Normal = CylinderObj.MPlane2.MNormal * (uniform FReal)(-1.0);
 			Phi = -Distance2;
 		}
 	}
 	else if (SideDistance < (uniform FReal)(0.0) && TopDistance < -SideDistance)
 	{
-		Normal = Distance1 < Distance2 ? MPlane1.MPlaneConcrete.MNormal * (uniform FReal)(-1.0) : MPlane2.MPlaneConcrete.MNormal * (uniform FReal)(-1.0);
+		Normal = Distance1 < Distance2 ? CylinderObj.MPlane1.MNormal * (uniform FReal)(-1.0) : CylinderObj.MPlane2.MNormal * (uniform FReal)(-1.0);
 		Phi = -TopDistance;
 	}
 	else
@@ -294,8 +322,11 @@ static inline void PhiWithNormalTaperedCylinder(const uniform uint8 *uniform Typ
 	}
 }
 
-static inline void PhiWithNormalUnion(const uniform uint8 *uniform TypedPtr, const FVector &PhiWithNormalInput, FVector &Normal, FReal &Phi, const uniform int SizeofFImplicitObject, const uniform int OffsetOfMargin)
+static inline void PhiWithNormalUnion(const uniform uint8 *uniform TypedPtr, const FVector &PhiWithNormalInput, FVector &Normal, FReal &Phi, const uniform int SizeofFImplicitObject, const uniform int OffsetofGeometryType, const uniform int OffsetOfMargin)
 {
+	// This is expecting a MObjects = [TaperedCylinder, Sphere, Sphere]
+	assert( ((const uniform TArray *uniform)TypedPtr)->ArrayNum == 3 );
+
 	const uniform UniquePtr *uniform MObjects = (const uniform UniquePtr *uniform)TypedPtr;
 	Phi = FLT_MAX;
 	bool NeedsNormalize = false;
@@ -306,6 +337,7 @@ static inline void PhiWithNormalUnion(const uniform uint8 *uniform TypedPtr, con
 	FVector NextNormal;
 	FReal NextPhi;
 
+	assert( (Object[OffsetofGeometryType] & 0x3F) == TaperedCylinder );
 	PhiWithNormalTaperedCylinder(TypedObjectPtr, PhiWithNormalInput, NextNormal, NextPhi, SizeofFImplicitObject);
 
 	if (NextPhi < Phi)
@@ -324,6 +356,7 @@ static inline void PhiWithNormalUnion(const uniform uint8 *uniform TypedPtr, con
 	uniform FReal Margin = *((const uniform FReal *uniform)(Object + OffsetOfMargin));
 	TypedObjectPtr = Object + SizeofFImplicitObject;
 
+	assert( (Object[OffsetofGeometryType] & 0x3F) == Sphere );
 	PhiWithNormalSphere(TypedObjectPtr, PhiWithNormalInput, NextNormal, NextPhi, Margin);
 
 	if (NextPhi < Phi)
@@ -342,6 +375,7 @@ static inline void PhiWithNormalUnion(const uniform uint8 *uniform TypedPtr, con
 	Margin = *((const uniform FReal *uniform)(Object + OffsetOfMargin));
 	TypedObjectPtr = Object + SizeofFImplicitObject;
 
+	assert( (Object[OffsetofGeometryType] & 0x3F) == Sphere );
 	PhiWithNormalSphere(TypedObjectPtr, PhiWithNormalInput, NextNormal, NextPhi, Margin);
 
 	if (NextPhi < Phi)
@@ -362,41 +396,41 @@ static inline void PhiWithNormalUnion(const uniform uint8 *uniform TypedPtr, con
 	}
 }
 
-static inline FVector FindClosestPointOnLineSegment(const FVector& P0, const FVector& P1, const FVector& P)
+static inline FVector3f FindClosestPointOnLineSegment(const FVector3f& P0, const FVector3f& P1, const FVector3f& P)
 {
-	const FVector P10 = P1 - P0;
-	const FVector PP0 = P - P0;
-	const FReal Proj = VectorDot(P10, PP0);
-	const FReal Denom2 = VectorSizeSquared(P10);
-	const FReal NormalProj = Proj / Denom2;
-	const FVector P2 = P0 + NormalProj * P10;
-	return VectorSelect(Proj < (uniform FReal)(0.0) || Denom2 < KINDA_SMALL_NUMBER, P0, VectorSelect(NormalProj > (uniform FReal)(1.0), P1, P2));
+	const FVector3f P10 = P1 - P0;
+	const FVector3f PP0 = P - P0;
+	const float Proj = VectorDot(P10, PP0);
+	const float Denom2 = VectorSizeSquared(P10);
+	const float NormalProj = Proj / Denom2;
+	const FVector3f P2 = P0 + NormalProj * P10;
+	return VectorSelect(Proj < (uniform float)(0.0f) || Denom2 < KINDA_SMALL_NUMBER, P0, VectorSelect(NormalProj > (uniform float)(1.0f), P1, P2));
 }
 
-static inline FVector FindClosestPointOnTriangle(const FVector& ClosestPointOnPlane, const FVector& P0, const FVector& P1, const FVector& P2, const FVector& P)
+static inline FVector3f FindClosestPointOnTriangle(const FVector3f& ClosestPointOnPlane, const FVector3f& P0, const FVector3f& P1, const FVector3f& P2, const FVector3f& P)
 {
 	// ComputeBarycentricInPlane
-	const FVector P10 = P1 - P0;
-	const FVector P20 = P2 - P0;
-	const FVector PP0 = P - P0;
-	const FReal Size10 = VectorSizeSquared(P10);
-	const FReal Size20 = VectorSizeSquared(P20);
-	const FReal ProjSides = VectorDot(P10, P20);
-	const FReal ProjP1 = VectorDot(PP0, P10);
-	const FReal ProjP2 = VectorDot(PP0, P20);
-	const FReal Denom = Size10 * Size20 - ProjSides * ProjSides;
-	const FReal BaryX = (Size20 * ProjP1 - ProjSides * ProjP2) / Denom;
-	const FReal BaryY = (Size10 * ProjP2 - ProjSides * ProjP1) / Denom;
-
-	const bool bClosestPoint = (BaryX >= -KINDA_SMALL_NUMBER && BaryX <= 1 + KINDA_SMALL_NUMBER && BaryY >= -KINDA_SMALL_NUMBER && BaryY <= 1 + KINDA_SMALL_NUMBER && (BaryX + BaryY) <= (1 + KINDA_SMALL_NUMBER));
-
-	const FVector P10Closest = FindClosestPointOnLineSegment(P0, P1, P);
-	const FVector P20Closest = FindClosestPointOnLineSegment(P0, P2, P);
-	const FVector P21Closest = FindClosestPointOnLineSegment(P1, P2, P);
-
-	const FReal P10Dist2 = VectorSizeSquared(P - P10Closest);
-	const FReal P20Dist2 = VectorSizeSquared(P - P20Closest);
-	const FReal P21Dist2 = VectorSizeSquared(P - P21Closest);
+	const FVector3f P10 = P1 - P0;
+	const FVector3f P20 = P2 - P0;
+	const FVector3f PP0 = P - P0;
+	const float Size10 = VectorSizeSquared(P10);
+	const float Size20 = VectorSizeSquared(P20);
+	const float ProjSides = VectorDot(P10, P20);
+	const float ProjP1 = VectorDot(PP0, P10);
+	const float ProjP2 = VectorDot(PP0, P20);
+	const float Denom = Size10 * Size20 - ProjSides * ProjSides;
+	const float BaryX = (Size20 * ProjP1 - ProjSides * ProjP2) / Denom;
+	const float BaryY = (Size10 * ProjP2 - ProjSides * ProjP1) / Denom;
+
+	const bool bClosestPoint = (BaryX >= -FLOAT_KINDA_SMALL_NUMBER && BaryX <= 1 + FLOAT_KINDA_SMALL_NUMBER && BaryY >= -FLOAT_KINDA_SMALL_NUMBER && BaryY <= 1 + FLOAT_KINDA_SMALL_NUMBER && (BaryX + BaryY) <= (1 + FLOAT_KINDA_SMALL_NUMBER));
+
+	const FVector3f P10Closest = FindClosestPointOnLineSegment(P0, P1, P);
+	const FVector3f P20Closest = FindClosestPointOnLineSegment(P0, P2, P);
+	const FVector3f P21Closest = FindClosestPointOnLineSegment(P1, P2, P);
+
+	const float P10Dist2 = VectorSizeSquared(P - P10Closest);
+	const float P20Dist2 = VectorSizeSquared(P - P20Closest);
+	const float P21Dist2 = VectorSizeSquared(P - P21Closest);
 
 	return VectorSelect(bClosestPoint, ClosestPointOnPlane, VectorSelect(P10Dist2 < P20Dist2, VectorSelect(P10Dist2 < P21Dist2, P10Closest, P21Closest), VectorSelect(P20Dist2 < P21Dist2, P20Closest, P21Closest)));
 }
@@ -405,21 +439,21 @@ static inline int32 NumPlaneVertices(const uniform FConvexStructureData *uniform
 {
 	const uniform int8 IndexType = StructureData->IndexType;
 
-	if(IndexType == 1) // EIndexType::Small
+	if(IndexType == EIndexType_Small)
 	{
 		const uniform PlanesS *uniform Planes = (const uniform PlanesS *uniform)StructureData->Data->Planes.Data;
 
 		#pragma ignore warning(perf)
 		return (int32)Planes[PlaneIndex].NumHalfEdges;
 	}
-	else if(IndexType == 2) // EIndexType::Medium
+	else if(IndexType == EIndexType_Medium)
 	{
 		const uniform PlanesM *uniform Planes = (const uniform PlanesM *uniform)StructureData->Data->Planes.Data;
 
 		#pragma ignore warning(perf)
 		return (int32)Planes[PlaneIndex].NumHalfEdges;
 	}
-	else if(IndexType == 3) // EIndexType::Large
+	else if(IndexType == EIndexType_Large)
 	{
 		const uniform PlanesL *uniform Planes = (const uniform PlanesL *uniform)StructureData->Data->Planes.Data;
 
@@ -436,7 +470,7 @@ static inline int32 GetPlaneVertex(const uniform FConvexStructureData *uniform S
 {
 	const uniform int8 IndexType = StructureData->IndexType;
 
-	if(IndexType == 1) // EIndexType::Small
+	if(IndexType == EIndexType_Small)
 	{
 		const uniform PlanesS *uniform Planes = (const uniform PlanesS *uniform)StructureData->Data->Planes.Data;
 		const uniform HalfEdgesS *uniform HalfEdges = (const uniform HalfEdgesS *uniform)StructureData->Data->HalfEdges.Data;
@@ -447,7 +481,7 @@ static inline int32 GetPlaneVertex(const uniform FConvexStructureData *uniform S
 		#pragma ignore warning(perf)
 		return (int32)HalfEdges[HalfEdgeIndex].VertexIndex;
 	}
-	else if(IndexType == 2) // EIndexType::Medium
+	else if(IndexType == EIndexType_Medium)
 	{
 		const uniform PlanesM *uniform Planes = (const uniform PlanesM *uniform)StructureData->Data->Planes.Data;
 		const uniform HalfEdgesM *uniform HalfEdges = (const uniform HalfEdgesM *uniform)StructureData->Data->HalfEdges.Data;
@@ -458,7 +492,7 @@ static inline int32 GetPlaneVertex(const uniform FConvexStructureData *uniform S
 		#pragma ignore warning(perf)
 		return (int32)HalfEdges[HalfEdgeIndex].VertexIndex;
 	}
-	else if(IndexType == 3) // EIndexType::Large
+	else if(IndexType == EIndexType_Large)
 	{
 		const uniform PlanesL *uniform Planes = (const uniform PlanesL *uniform)StructureData->Data->Planes.Data;
 		const uniform HalfEdgesL *uniform HalfEdges = (const uniform HalfEdgesL *uniform)StructureData->Data->HalfEdges.Data;
@@ -475,23 +509,23 @@ static inline int32 GetPlaneVertex(const uniform FConvexStructureData *uniform S
 	}
 }
 
-static inline void PhiWithNormalConvex(const uniform uint8 *uniform TypedObjectPtr, const FVector &PhiWithNormalInput, FVector &Normal, FReal &Phi, const uniform int SizeofFImplicitObject)
+static inline void PhiWithNormalConvex(const uniform uint8 *uniform TypedObjectPtr, const FVector3f &PhiWithNormalInput, FVector3f &Normal, float &Phi)
 {
-	const uniform ConvexShape *uniform Object = (const uniform ConvexShape *uniform)TypedObjectPtr;
-	const uniform Plane *uniform Planes = (const uniform Plane *uniform)&Object->Planes.Data[0];
+	const uniform FConvex *uniform Object = (const uniform FConvex *uniform)TypedObjectPtr;
+	const uniform FPlaneConcrete3f *uniform Planes = (const uniform FPlaneConcrete3f *uniform)&Object->Planes.Data[0];
 	const uniform int32 NumPlanes = Object->Planes.ArrayNum;
 	if (NumPlanes == 0)
 	{
-		Phi = REAL_MAX;
+		Phi = FLT_MAX;
 		return;
 	}
 
-	FReal MaxPhi = -REAL_MAX;
+	float MaxPhi = -FLT_MAX;
 	int32 MaxPlane = 0;
 
 	for (uniform int32 Idx = 0; Idx < NumPlanes; ++Idx)
 	{
-		const FReal InnerPhi = SignedDistancePlane(&Planes[Idx], PhiWithNormalInput);
+		const float InnerPhi = SignedDistancePlane(&Planes[Idx], PhiWithNormalInput);
 		if (InnerPhi > MaxPhi)
 		{
 			MaxPhi = InnerPhi;
@@ -507,29 +541,29 @@ static inline void PhiWithNormalConvex(const uniform uint8 *uniform TypedObjectP
 
 	// If x is outside the convex mesh, we should find for the nearest point to triangles on the plane
 	const int32 PlaneVerticesNum = NumPlaneVertices(&Object->StructureData, MaxPlane);
-	const FVector XOnPlane = PhiWithNormalInput - Phi * Normal;
+	const FVector3f XOnPlane = PhiWithNormalInput - Phi * Normal;
 
 	if(PlaneVerticesNum > 2 && IsVectorEqual(XOnPlane, PhiWithNormalInput))
 	{
 		return;
 	}
 
-	FReal ClosestDistance = REAL_MAX;
-	FVector ClosestPoint = ZeroVector;
+	float ClosestDistance = FLT_MAX;
+	FVector3f ClosestPoint = FloatZeroVector;
 	for (uniform int32 Index = 0; Index < PlaneVerticesNum - 2; Index++)
 	{
-		const uniform FVector *uniform Vertices = (const uniform FVector *uniform)Object->Vertices.Data;
-		const uniform FVector *varying APtr = &Vertices[GetPlaneVertex(&Object->StructureData, MaxPlane, 0)];
-		const uniform FVector *varying BPtr = &Vertices[GetPlaneVertex(&Object->StructureData, MaxPlane, Index + 1)];
-		const uniform FVector *varying CPtr = &Vertices[GetPlaneVertex(&Object->StructureData, MaxPlane, Index + 2)];
+		const uniform FVector3f *uniform Vertices = (const uniform FVector3f *uniform)Object->Vertices.Data;
+		const uniform FVector3f *varying APtr = &Vertices[GetPlaneVertex(&Object->StructureData, MaxPlane, 0)];
+		const uniform FVector3f *varying BPtr = &Vertices[GetPlaneVertex(&Object->StructureData, MaxPlane, Index + 1)];
+		const uniform FVector3f *varying CPtr = &Vertices[GetPlaneVertex(&Object->StructureData, MaxPlane, Index + 2)];
 
-		const FVector A = VectorGather(APtr);
-		const FVector B = VectorGather(BPtr);
-		const FVector C = VectorGather(CPtr);
+		const FVector3f A = VectorGather(APtr);
+		const FVector3f B = VectorGather(BPtr);
+		const FVector3f C = VectorGather(CPtr);
 
-		const FVector Point = FindClosestPointOnTriangle(XOnPlane, A, B, C, PhiWithNormalInput);
+		const FVector3f Point = FindClosestPointOnTriangle(XOnPlane, A, B, C, PhiWithNormalInput);
 
-		const FReal Distance = VectorSize(Point - XOnPlane);
+		const float Distance = VectorSize(Point - XOnPlane);
 		if (Distance < ClosestDistance)
 		{
 			ClosestDistance = Distance;
@@ -537,9 +571,9 @@ static inline void PhiWithNormalConvex(const uniform uint8 *uniform TypedObjectP
 		}
 	}
 
-	const FVector Difference = PhiWithNormalInput - ClosestPoint;
+	const FVector3f Difference = PhiWithNormalInput - ClosestPoint;
 	Phi = VectorSize(Difference);
-	if (Phi > SMALL_NUMBER)
+	if (Phi > FLOAT_SMALL_NUMBER)
 	{
 		Normal = (Difference) / Phi;
 	}
@@ -645,7 +679,7 @@ export void ApplyPerParticleCollision(uniform FVector4f ParticlesPandInvM[],
 				}
 			case Union:
 				{
-					PhiWithNormalUnion(TypedPtr, PhiWithNormalInput, NormalReal, PhiReal, SizeofFImplicitObject, OffsetOfMargin);
+					PhiWithNormalUnion(TypedPtr, PhiWithNormalInput, NormalReal, PhiReal, SizeofFImplicitObject, OffsetofGeometryType, OffsetOfMargin);
 					Phi = (varying float)PhiReal;
 					Normal = ConvertVector3NativeTo3f(NormalReal);
 					break;
@@ -659,9 +693,7 @@ export void ApplyPerParticleCollision(uniform FVector4f ParticlesPandInvM[],
 				}
 			case Convex:
 				{
-					PhiWithNormalConvex(TypedPtr, PhiWithNormalInput, NormalReal, PhiReal, SizeofFImplicitObject);
-					Phi = (varying float)PhiReal;
-					Normal = ConvertVector3NativeTo3f(NormalReal);
+					PhiWithNormalConvex(TypedPtr, PhiWithNormalInputFloat, Normal, Phi);
 					break;
 				}
 			default:
@@ -813,7 +845,7 @@ export void ApplyPerParticleCollisionFastFriction(uniform FVector4f ParticlesPan
 				}
 			case Union:
 				{
-					PhiWithNormalUnion(TypedPtr, PhiWithNormalInput, NormalReal, PhiReal, SizeofFImplicitObject, OffsetOfMargin);
+					PhiWithNormalUnion(TypedPtr, PhiWithNormalInput, NormalReal, PhiReal, SizeofFImplicitObject, OffsetofGeometryType, OffsetOfMargin);
 					Phi = (varying float)PhiReal;
 					Normal = ConvertVector3NativeTo3f(NormalReal);
 					break;
@@ -827,9 +859,7 @@ export void ApplyPerParticleCollisionFastFriction(uniform FVector4f ParticlesPan
 				}
 			case Convex:
 				{
-					PhiWithNormalConvex(TypedPtr, PhiWithNormalInput, NormalReal, PhiReal, SizeofFImplicitObject);
-					Phi = (varying float)PhiReal;
-					Normal = ConvertVector3NativeTo3f(NormalReal);
+					PhiWithNormalConvex(TypedPtr, PhiWithNormalInputFloat, Normal, Phi);
 					break;
 				}
 			default:
@@ -964,7 +994,7 @@ export void ApplyPerParticleCollisionNoFriction(uniform FVector4f ParticlesPandI
 				}
 			case Union:
 				{
-					PhiWithNormalUnion(TypedPtr, PhiWithNormalInput, NormalReal, PhiReal, SizeofFImplicitObject, OffsetOfMargin);
+					PhiWithNormalUnion(TypedPtr, PhiWithNormalInput, NormalReal, PhiReal, SizeofFImplicitObject, OffsetofGeometryType, OffsetOfMargin);
 					Phi = (varying float)PhiReal;
 					Normal = ConvertVector3NativeTo3f(NormalReal);
 					break;
@@ -978,9 +1008,7 @@ export void ApplyPerParticleCollisionNoFriction(uniform FVector4f ParticlesPandI
 				}
 			case Convex:
 				{
-					PhiWithNormalConvex(TypedPtr, PhiWithNormalInput, NormalReal, PhiReal, SizeofFImplicitObject);
-					Phi = (varying float)PhiReal;
-					Normal = ConvertVector3NativeTo3f(NormalReal);
+					PhiWithNormalConvex(TypedPtr, PhiWithNormalInputFloat, Normal, Phi);
 					break;
 				}
 			default:
diff --git a/Engine/Source/Runtime/Experimental/Chaos/Public/Chaos/Capsule.h b/Engine/Source/Runtime/Experimental/Chaos/Public/Chaos/Capsule.h
--- a/Engine/Source/Runtime/Experimental/Chaos/Public/Chaos/Capsule.h
+++ b/Engine/Source/Runtime/Experimental/Chaos/Public/Chaos/Capsule.h
@@ -566,6 +566,18 @@ namespace Chaos
 		// Used for manifold generation
 		int32 NumPlanes() const { return 0; }
 
+#if INTEL_ISPC && !UE_BUILD_SHIPPING
+		// See PerParticlePBDCollisionConstraint.cpp
+		// ISPC code has matching structs for interpreting FImplicitObjects.
+		// This is used to verify that the structs stay the same.
+		struct FISPCDataVerifier
+		{
+			static constexpr int32 OffsetOfMSegment() { return offsetof(FCapsule, MSegment); }
+			static constexpr int32 SizeOfMSegment() { return sizeof(FCapsule::MSegment); }
+		};
+		friend FISPCDataVerifier;
+#endif // #if INTEL_ISPC && !UE_BUILD_SHIPPING
+
 	private:
 		void SetRadius(FReal InRadius) { SetMargin(InRadius); }
 
diff --git a/Engine/Source/Runtime/Experimental/Chaos/Public/Chaos/Convex.h b/Engine/Source/Runtime/Experimental/Chaos/Public/Chaos/Convex.h
--- a/Engine/Source/Runtime/Experimental/Chaos/Public/Chaos/Convex.h
+++ b/Engine/Source/Runtime/Experimental/Chaos/Public/Chaos/Convex.h
@@ -914,6 +914,22 @@ namespace Chaos
 			return FVec3(0);
 		}
 
+#if INTEL_ISPC && !UE_BUILD_SHIPPING
+		// See PerParticlePBDCollisionConstraint.cpp
+		// ISPC code has matching structs for interpreting FImplicitObjects.
+		// This is used to verify that the structs stay the same.
+		struct FISPCDataVerifier
+		{
+			static constexpr int32 OffsetOfPlanes() { return offsetof(FConvex, Planes); }
+			static constexpr int32 SizeOfPlanes() { return sizeof(FConvex::Planes); }
+			static constexpr int32 OffsetOfVertices() { return offsetof(FConvex, Vertices); }
+			static constexpr int32 SizeOfVertices() { return sizeof(FConvex::Vertices); }
+			static constexpr int32 OffsetOfStructureData() { return offsetof(FConvex, StructureData); }
+			static constexpr int32 SizeOfStructureData() { return sizeof(FConvex::StructureData); }
+		};
+		friend FISPCDataVerifier;
+#endif // #if INTEL_ISPC && !UE_BUILD_SHIPPING
+
 	private:
 		TArray<FPlaneType> Planes;
 		TArray<FVec3Type> Vertices; //copy of the vertices that are just on the convex hull boundary
diff --git a/Engine/Source/Runtime/Experimental/Chaos/Public/Chaos/ConvexHalfEdgeStructureData.h b/Engine/Source/Runtime/Experimental/Chaos/Public/Chaos/ConvexHalfEdgeStructureData.h
--- a/Engine/Source/Runtime/Experimental/Chaos/Public/Chaos/ConvexHalfEdgeStructureData.h
+++ b/Engine/Source/Runtime/Experimental/Chaos/Public/Chaos/ConvexHalfEdgeStructureData.h
@@ -467,6 +467,26 @@ namespace Chaos
 			return Ar;
 		}
 
+#if INTEL_ISPC && !UE_BUILD_SHIPPING
+		// See PerParticlePBDCollisionConstraint.cpp
+		// ISPC code has matching structs for interpreting FImplicitObjects.
+		// This is used to verify that the structs stay the same.
+		struct FISPCDataVerifier
+		{
+			static constexpr int32 OffsetOfPlanes() { return offsetof(TConvexHalfEdgeStructureData, Planes); }
+			static constexpr int32 SizeOfPlanes() { return sizeof(TConvexHalfEdgeStructureData::Planes); }
+			static constexpr int32 OffsetOfHalfEdges() { return offsetof(TConvexHalfEdgeStructureData, HalfEdges); }
+			static constexpr int32 SizeOfHalfEdges() { return sizeof(TConvexHalfEdgeStructureData::HalfEdges); }
+			static constexpr int32 OffsetOfVertices() { return offsetof(TConvexHalfEdgeStructureData, Vertices); }
+			static constexpr int32 SizeOfVertices() { return sizeof(TConvexHalfEdgeStructureData::Vertices); }
+			static constexpr int32 OffsetOfEdges() { return offsetof(TConvexHalfEdgeStructureData, Edges); }
+			static constexpr int32 SizeOfEdges() { return sizeof(TConvexHalfEdgeStructureData::Edges); }
+			static constexpr int32 OffsetOfVertexPlanes() { return offsetof(TConvexHalfEdgeStructureData, VertexPlanes); }
+			static constexpr int32 SizeOfVertexPlanes() { return sizeof(TConvexHalfEdgeStructureData::VertexPlanes); }
+		};
+		friend FISPCDataVerifier;
+#endif // #if INTEL_ISPC && !UE_BUILD_SHIPPING
+
 	private:
 
 		// The edge index of the previous edge on the plane (loops)
diff --git a/Engine/Source/Runtime/Experimental/Chaos/Public/Chaos/ConvexStructureData.h b/Engine/Source/Runtime/Experimental/Chaos/Public/Chaos/ConvexStructureData.h
--- a/Engine/Source/Runtime/Experimental/Chaos/Public/Chaos/ConvexStructureData.h
+++ b/Engine/Source/Runtime/Experimental/Chaos/Public/Chaos/ConvexStructureData.h
@@ -253,6 +253,20 @@ namespace Chaos
 			return Ar;
 		}
 
+#if INTEL_ISPC && !UE_BUILD_SHIPPING
+		// See PerParticlePBDCollisionConstraint.cpp
+		// ISPC code has matching structs for interpreting FImplicitObjects.
+		// This is used to verify that the structs stay the same.
+		struct FISPCDataVerifier
+		{
+			static constexpr int32 OffsetOfData() { return offsetof(FConvexStructureData, Data); }
+			static constexpr int32 SizeOfData() { return sizeof(FConvexStructureData::Data); }
+			static constexpr int32 OffsetOfIndexType() { return offsetof(FConvexStructureData, IndexType); }
+			static constexpr int32 SizeOfIndexType() { return sizeof(FConvexStructureData::IndexType); }
+		};
+		friend FISPCDataVerifier;
+#endif // #if INTEL_ISPC && !UE_BUILD_SHIPPING
+
 	private:
 
 		// Load data from an asset saved before we had a proper half-edge data structure
diff --git a/Engine/Source/Runtime/Experimental/Chaos/Public/Chaos/ImplicitObject.h b/Engine/Source/Runtime/Experimental/Chaos/Public/Chaos/ImplicitObject.h
--- a/Engine/Source/Runtime/Experimental/Chaos/Public/Chaos/ImplicitObject.h
+++ b/Engine/Source/Runtime/Experimental/Chaos/Public/Chaos/ImplicitObject.h
@@ -30,13 +30,14 @@ class FImplicitObject;
 
 using FAABB3 = TAABB<FReal, 3>;
 using FParticles = TParticles<FReal, 3>;
+using FSphere = TSphere<FReal, 3>; // warning: code assumes that FImplicitObjects with type ImplicitObjectType::Sphere are FSpheres, but all TSpheres will think they have ImplicitObjectType::Sphere.
 
 namespace ImplicitObjectType
 {
 	enum
 	{
 		//Note: add entries in order to avoid serialization issues (but before IsInstanced)
-		Sphere = 0,
+		Sphere = 0, // warning: code assumes that this is an FSphere, but all TSpheres will think this is their type.
 		Box,
 		Plane,
 		Capsule,
diff --git a/Engine/Source/Runtime/Experimental/Chaos/Public/Chaos/ImplicitObjectUnion.h b/Engine/Source/Runtime/Experimental/Chaos/Public/Chaos/ImplicitObjectUnion.h
--- a/Engine/Source/Runtime/Experimental/Chaos/Public/Chaos/ImplicitObjectUnion.h
+++ b/Engine/Source/Runtime/Experimental/Chaos/Public/Chaos/ImplicitObjectUnion.h
@@ -171,6 +171,18 @@ class CHAOS_API FImplicitObjectUnion : public FImplicitObject
 		return new FImplicitObjectUnion(MoveTemp(NewObjects));
 	}
 
+#if INTEL_ISPC && !UE_BUILD_SHIPPING
+	// See PerParticlePBDCollisionConstraint.cpp
+	// ISPC code has matching structs for interpreting FImplicitObjects.
+	// This is used to verify that the structs stay the same.
+	struct FISPCDataVerifier
+	{
+		static constexpr int32 OffsetOfMObjects() { return offsetof(FImplicitObjectUnion, MObjects); }
+		static constexpr int32 SizeOfMObjects() { return sizeof(FImplicitObjectUnion::MObjects); }
+	};
+	friend FISPCDataVerifier;
+#endif // #if INTEL_ISPC && !UE_BUILD_SHIPPING
+
 protected:
 	virtual Pair<FVec3, bool> FindClosestIntersectionImp(const FVec3& StartPoint, const FVec3& EndPoint, const FReal Thickness) const override
 	{
diff --git a/Engine/Source/Runtime/Experimental/Chaos/Public/Chaos/Sphere.h b/Engine/Source/Runtime/Experimental/Chaos/Public/Chaos/Sphere.h
--- a/Engine/Source/Runtime/Experimental/Chaos/Public/Chaos/Sphere.h
+++ b/Engine/Source/Runtime/Experimental/Chaos/Public/Chaos/Sphere.h
@@ -372,6 +372,18 @@ namespace Chaos
 			return  TUniquePtr<FImplicitObject>(new TSphere<T, d>(Center * Scale, GetRadius() * Scale.Min()));
 		}
 
+#if INTEL_ISPC && !UE_BUILD_SHIPPING
+		// See PerParticlePBDCollisionConstraint.cpp
+		// ISPC code has matching structs for interpreting FImplicitObjects.
+		// This is used to verify that the structs stay the same.
+		struct FISPCDataVerifier
+		{
+			static constexpr int32 OffsetOfCenter() { return offsetof(TSphere, Center); }
+			static constexpr int32 SizeOfCenter() { return sizeof(TSphere::Center); }
+		};
+		friend FISPCDataVerifier;
+#endif // #if INTEL_ISPC && !UE_BUILD_SHIPPING
+
 	private:
 		void SetRadius(FReal InRadius) { SetMargin(InRadius); }
 
diff --git a/Engine/Source/Runtime/Experimental/Chaos/Public/Chaos/TaperedCapsule.h b/Engine/Source/Runtime/Experimental/Chaos/Public/Chaos/TaperedCapsule.h
--- a/Engine/Source/Runtime/Experimental/Chaos/Public/Chaos/TaperedCapsule.h
+++ b/Engine/Source/Runtime/Experimental/Chaos/Public/Chaos/TaperedCapsule.h
@@ -3,6 +3,7 @@
 
 #include "Chaos/ImplicitObject.h"
 #include "Chaos/Core.h"
+#include "Chaos/TaperedCylinder.h"
 
 namespace Chaos
 {
@@ -217,6 +218,26 @@ namespace Chaos
 			return HashCombine(OriginAxisHash, PropertyHash);
 		}
 
+#if INTEL_ISPC && !UE_BUILD_SHIPPING
+		// See PerParticlePBDCollisionConstraint.cpp
+		// ISPC code has matching structs for interpreting FImplicitObjects.
+		// This is used to verify that the structs stay the same.
+		struct FISPCDataVerifier
+		{
+			static constexpr int32 OffsetOfOrigin() { return offsetof(FTaperedCapsule, Origin); }
+			static constexpr int32 SizeOfOrigin() { return sizeof(FTaperedCapsule::Origin); }
+			static constexpr int32 OffsetOfAxis() { return offsetof(FTaperedCapsule, Axis); }
+			static constexpr int32 SizeOfAxis() { return sizeof(FTaperedCapsule::Axis); }
+			static constexpr int32 OffsetOfHeight() { return offsetof(FTaperedCapsule, Height); }
+			static constexpr int32 SizeOfHeight() { return sizeof(FTaperedCapsule::Height); }
+			static constexpr int32 OffsetOfRadius1() { return offsetof(FTaperedCapsule, Radius1); }
+			static constexpr int32 SizeOfRadius1() { return sizeof(FTaperedCapsule::Radius1); }
+			static constexpr int32 OffsetOfRadius2() { return offsetof(FTaperedCapsule, Radius2); }
+			static constexpr int32 SizeOfRadius2() { return sizeof(FTaperedCapsule::Radius2); }
+		};
+		friend FISPCDataVerifier;
+#endif // #if INTEL_ISPC && !UE_BUILD_SHIPPING
+
 	private:
 		//Phi is distance from closest point on plane1
 		FReal GetRadius(const FReal& Phi) const
diff --git a/Engine/Source/Runtime/Experimental/Chaos/Public/Chaos/TaperedCylinder.h b/Engine/Source/Runtime/Experimental/Chaos/Public/Chaos/TaperedCylinder.h
--- a/Engine/Source/Runtime/Experimental/Chaos/Public/Chaos/TaperedCylinder.h
+++ b/Engine/Source/Runtime/Experimental/Chaos/Public/Chaos/TaperedCylinder.h
@@ -330,6 +330,26 @@ namespace Chaos
 			return HashCombine(PlaneHashes, PropertyHash);
 		}
 
+#if INTEL_ISPC && !UE_BUILD_SHIPPING
+		// See PerParticlePBDCollisionConstraint.cpp
+		// ISPC code has matching structs for interpreting FImplicitObjects.
+		// This is used to verify that the structs stay the same.
+		struct FISPCDataVerifier
+		{
+			static constexpr int32 OffsetOfMPlane1() { return offsetof(FTaperedCylinder, MPlane1); }
+			static constexpr int32 SizeOfMPlane1() { return sizeof(FTaperedCylinder::MPlane1); }
+			static constexpr int32 OffsetOfMPlane2() { return offsetof(FTaperedCylinder, MPlane2); }
+			static constexpr int32 SizeOfMPlane2() { return sizeof(FTaperedCylinder::MPlane2); }
+			static constexpr int32 OffsetOfMHeight() { return offsetof(FTaperedCylinder, MHeight); }
+			static constexpr int32 SizeOfMHeight() { return sizeof(FTaperedCylinder::MHeight); }
+			static constexpr int32 OffsetOfMRadius1() { return offsetof(FTaperedCylinder, MRadius1); }
+			static constexpr int32 SizeOfMRadius1() { return sizeof(FTaperedCylinder::MRadius1); }
+			static constexpr int32 OffsetOfMRadius2() { return offsetof(FTaperedCylinder, MRadius2); }
+			static constexpr int32 SizeOfMRadius2() { return sizeof(FTaperedCylinder::MRadius2); }
+		};
+		friend FISPCDataVerifier;
+#endif // #if INTEL_ISPC && !UE_BUILD_SHIPPING
+
 	private:
 		//Phi is distance from closest point on plane1
 		FReal GetRadius(const FReal& Phi) const
@@ -338,7 +358,7 @@ namespace Chaos
 			return MRadius1 * (static_cast<FReal>(1.) - Alpha) + MRadius2 * Alpha;
 		}
 
-		TPlane<FReal, 3> MPlane1, MPlane2;
+		TPlaneConcrete<FReal, 3> MPlane1, MPlane2;
 		FReal MHeight, MRadius1, MRadius2;
 		FAABB3 MLocalBoundingBox;
 	};
