Add access to dynamic collection from Geometry Collection Component

diff --git a/Engine/Source/Runtime/Experimental/GeometryCollectionEngine/Public/GeometryCollection/GeometryCollectionComponent.h b/Engine/Source/Runtime/Experimental/GeometryCollectionEngine/Public/GeometryCollection/GeometryCollectionComponent.h
--- a/Engine/Source/Runtime/Experimental/GeometryCollectionEngine/Public/GeometryCollection/GeometryCollectionComponent.h
+++ b/Engine/Source/Runtime/Experimental/GeometryCollectionEngine/Public/GeometryCollection/GeometryCollectionComponent.h
@@ -697,6 +697,8 @@ public:
 	/** Used by Niagara DI to query global matrices rather than recalculating them again */
 	const TArray<FMatrix>& GetGlobalMatrices() { return GlobalMatrices; }
 
+	const FGeometryDynamicCollection* GetDynamicCollection() const { return DynamicCollection.Get(); }
+	
 public:
 	UPROPERTY(BlueprintAssignable, Category = "Collision")
 	FOnChaosPhysicsCollision OnChaosPhysicsCollision;
