Invalidate island persistent flag when removing a node

diff --git a/Engine/Source/Runtime/Experimental/Chaos/Private/Chaos/Island/IslandGraph.cpp b/Engine/Source/Runtime/Experimental/Chaos/Private/Chaos/Island/IslandGraph.cpp
--- a/Engine/Source/Runtime/Experimental/Chaos/Private/Chaos/Island/IslandGraph.cpp
+++ b/Engine/Source/Runtime/Experimental/Chaos/Private/Chaos/Island/IslandGraph.cpp
@@ -140,10 +140,7 @@ void FIslandGraph<NodeType, EdgeType, IslandType>::RemoveNode(const NodeType& No
 			// If only one node and zero edges, we invalidate the node island
 			if(GraphNode.NodeEdges.Num() == 0  && GraphIslands.IsValidIndex(GraphNode.IslandIndex))
 			{
-				if(!GraphNode.bStationaryNode)
-				{
-					GraphIslands[GraphNode.IslandIndex].bIsPersistent = false;
-				}
+				GraphIslands[GraphNode.IslandIndex].bIsPersistent = false;
 			}
 			else
 			{
@@ -153,10 +150,7 @@ void FIslandGraph<NodeType, EdgeType, IslandType>::RemoveNode(const NodeType& No
 					if (GraphNode.NodeEdges.IsValidIndex(NodeEdgeIndex))
 					{
 						const int32 GraphEdgeIndex = GraphNode.NodeEdges[NodeEdgeIndex];
-						if(IsEdgeMoving(GraphEdgeIndex))
-						{
-							GraphIslands[GraphEdges[GraphEdgeIndex].IslandIndex].bIsPersistent = false;
-						}
+						GraphIslands[GraphEdges[GraphEdgeIndex].IslandIndex].bIsPersistent = false;
 						RemoveEdge(GraphEdgeIndex);
 					}
 				}
