Fix CCD between ragdoll and trimesh:
- handle box trimesh collision
- fix scaling issue in trimesh sweep

diff --git a/Engine/Source/Runtime/Experimental/Chaos/Private/Chaos/CollisionResolution.cpp b/Engine/Source/Runtime/Experimental/Chaos/Private/Chaos/CollisionResolution.cpp
--- a/Engine/Source/Runtime/Experimental/Chaos/Private/Chaos/CollisionResolution.cpp
+++ b/Engine/Source/Runtime/Experimental/Chaos/Private/Chaos/CollisionResolution.cpp
@@ -1785,6 +1785,7 @@ namespace Chaos
 				case EContactShapesType::CapsuleBox:
 				case EContactShapesType::CapsuleConvex:
 				case EContactShapesType::BoxBox:
+				case EContactShapesType::BoxConvex:
 				case EContactShapesType::BoxPlane:
 				case EContactShapesType::GenericConvexConvex:
 					UpdateGenericConvexConvexConstraintSwept(Particle0, Implicit0, WorldTransform0, Particle1, Implicit1, WorldTransform1, DirCCD, LengthCCD, Dt, Constraint);
@@ -1837,12 +1838,14 @@ namespace Chaos
 					}
 					return true;
 				}
+				case EContactShapesType::BoxHeightField:
 				case EContactShapesType::ConvexHeightField:
 				{
 					const FHeightField* Object1 = Implicit1.template GetObject<const FHeightField >();
 					UpdateConvexHeightFieldConstraintSwept(Particle0, Implicit0, WorldTransform0, *Object1, WorldTransform1, DirCCD, LengthCCD, Dt, Constraint);
 					return true;
 				}
+				case EContactShapesType::BoxTriMesh:
 				case EContactShapesType::ConvexTriMesh:
 					if (const TImplicitObjectScaled<FTriangleMeshImplicitObject>* ScaledTriangleMesh = Implicit1.template GetObject<const TImplicitObjectScaled<FTriangleMeshImplicitObject>>())
 					{
@@ -1857,6 +1860,9 @@ namespace Chaos
 						ensure(false);
 					}
 					return true;
+				default:
+					ensure(false);
+					break;
 				}
 			}
 
diff --git a/Engine/Source/Runtime/Experimental/Chaos/Private/Chaos/TriangleMeshImplicitObject.cpp b/Engine/Source/Runtime/Experimental/Chaos/Private/Chaos/TriangleMeshImplicitObject.cpp
--- a/Engine/Source/Runtime/Experimental/Chaos/Private/Chaos/TriangleMeshImplicitObject.cpp
+++ b/Engine/Source/Runtime/Experimental/Chaos/Private/Chaos/TriangleMeshImplicitObject.cpp
@@ -1052,8 +1052,9 @@ bool FTriangleMeshImplicitObject::SweepGeomImp(const QueryGeomType& QueryGeom, c
 		VisitorType SQVisitor(*this,Elements, QueryGeom,StartTM,Dir,ScaledDirNormalized,LengthScale,ScaledStartTM,Thickness,bComputeMTD, TriMeshScale, CullsBackFaceRaycastCode);
 
 		const FAABB3 QueryBounds = QueryGeom.BoundingBox().TransformedAABB(FRigidTransform3(FVec3::ZeroVector,StartTM.GetRotation()));
-		const FVec3 StartPoint = StartTM.TransformPosition(QueryBounds.Center());
-		const FVec3 Inflation = QueryBounds.Extents() * SafeInvScale(TriMeshScale) * 0.5 + FVec3(Thickness);
+		const FVec3 InvTriMeshScale = SafeInvScale(TriMeshScale);
+		const FVec3 StartPoint = QueryBounds.Center() * InvTriMeshScale + StartTM.GetLocation();
+		const FVec3 Inflation = QueryBounds.Extents() * InvTriMeshScale * 0.5 + FVec3(Thickness);
 		BVH.template Sweep<VisitorType>(StartPoint,Dir,Length,Inflation,SQVisitor);
 
 		if(SQVisitor.OutTime <= Length)
