Moving InterchangeImportCommon.h to InterchangeEngine module. Allowing other import plugins to use the AssetImportData utiliies.

diff --git a/Engine/Plugins/Experimental/Interchange/Runtime/Source/Import/Private/InterchangeImportCommon.cpp b/Engine/Plugins/Experimental/Interchange/Runtime/Source/Import/Private/InterchangeImportCommon.cpp
deleted file mode 100644
diff --git a/Engine/Plugins/Experimental/Interchange/Runtime/Source/Import/Private/InterchangeImportCommon.h b/Engine/Plugins/Experimental/Interchange/Runtime/Source/Import/Private/InterchangeImportCommon.h
deleted file mode 100644
diff --git a/Engine/Source/Runtime/Experimental/Interchange/Engine/Private/InterchangeImportCommon.cpp b/Engine/Source/Runtime/Experimental/Interchange/Engine/Private/InterchangeImportCommon.cpp
new file mode 100644
--- /dev/null
+++ b/Engine/Source/Runtime/Experimental/Interchange/Engine/Private/InterchangeImportCommon.cpp
@@ -0,0 +1,227 @@
+// Copyright Epic Games, Inc. All Rights Reserved.
+#include "InterchangeImportCommon.h"
+
+#include "CoreMinimal.h"
+#include "EditorFramework/AssetImportData.h"
+#include "InterchangeAssetImportData.h"
+#include "InterchangePipelineBase.h"
+#include "InterchangeSourceData.h"
+#include "Nodes/InterchangeBaseNodeContainer.h"
+#include "UObject/Object.h"
+#include "EditorFramework/AssetImportData.h"
+#include "HAL/FileManager.h"
+
+namespace UE
+{
+	namespace Interchange
+	{
+		namespace Private::ImportCommon
+		{
+			UInterchangeAssetImportData* BeginSetupAssetData(FFactoryCommon::FUpdateImportAssetDataParameters& Parameters)
+			{
+				if (!ensure(IsInGameThread()))
+				{
+					return nullptr;
+				}
+				if (!ensure(Parameters.SourceData && Parameters.AssetImportDataOuter))
+				{
+					return nullptr;
+				}
+
+				UInterchangeAssetImportData* AssetImportData = Cast<UInterchangeAssetImportData>(Parameters.AssetImportData);
+
+				if (!AssetImportData)
+				{
+					AssetImportData = NewObject<UInterchangeAssetImportData>(Parameters.AssetImportDataOuter, NAME_None);
+				}
+
+				ensure(AssetImportData);
+
+				return AssetImportData;
+			}
+
+
+			void EndSetupAssetData(FFactoryCommon::FUpdateImportAssetDataParameters& Parameters, UInterchangeAssetImportData* AssetImportData)
+			{
+				//Set the interchange node graph data
+				AssetImportData->NodeUniqueID = Parameters.NodeUniqueID;
+				FObjectDuplicationParameters DupParam(Parameters.NodeContainer, AssetImportData);
+				AssetImportData->NodeContainer = CastChecked<UInterchangeBaseNodeContainer>(StaticDuplicateObjectEx(DupParam));
+				AssetImportData->Pipelines.Reset();
+				for (const UInterchangePipelineBase* Pipeline : Parameters.Pipelines)
+				{
+					UInterchangePipelineBase* DupPipeline = Cast<UInterchangePipelineBase>(StaticDuplicateObject(Pipeline, AssetImportData));
+					if (DupPipeline)
+					{
+						AssetImportData->Pipelines.Add(DupPipeline);
+					}
+				}
+			}
+		}
+
+		FFactoryCommon::FUpdateImportAssetDataParameters::FUpdateImportAssetDataParameters(UObject* InAssetImportDataOuter
+																							, UAssetImportData* InAssetImportData
+																							, const UInterchangeSourceData* InSourceData
+																							, FString InNodeUniqueID
+																							, UInterchangeBaseNodeContainer* InNodeContainer
+																						   , const TArray<UInterchangePipelineBase*>& InPipelines)
+			: AssetImportDataOuter(InAssetImportDataOuter)
+			, AssetImportData(InAssetImportData)
+			, SourceData(InSourceData)
+			, NodeUniqueID(InNodeUniqueID)
+			, NodeContainer(InNodeContainer)
+			, Pipelines(InPipelines)
+		{
+			ensure(AssetImportDataOuter);
+			ensure(SourceData);
+			ensure(!NodeUniqueID.IsEmpty());
+			ensure(NodeContainer);
+		}
+
+		UAssetImportData* FFactoryCommon::UpdateImportAssetData(FUpdateImportAssetDataParameters& Parameters)
+		{
+			return UpdateImportAssetData(Parameters, [&Parameters](UInterchangeAssetImportData* AssetImportData)
+				{
+#if WITH_EDITORONLY_DATA
+					//Set the asset import data file source to allow reimport. TODO: manage MD5 Hash properly
+					TOptional<FMD5Hash> FileContentHash = Parameters.SourceData->GetFileContentHash();
+
+					//Update the first filename, TODO for asset using multiple source file we have to update the correct index
+					AssetImportData->Update(Parameters.SourceData->GetFilename(), FileContentHash.IsSet() ? &FileContentHash.GetValue() : nullptr);
+#endif //WITH_EDITORONLY_DATA
+				});
+
+		}
+
+		UAssetImportData* FFactoryCommon::UpdateImportAssetData(FUpdateImportAssetDataParameters& Parameters, TFunctionRef<void(UInterchangeAssetImportData*)> CustomFileSourceUpdate)
+		{
+#if WITH_EDITORONLY_DATA
+			if (!ensure(IsInGameThread()))
+			{
+				return nullptr;
+			}
+			if (!ensure(Parameters.SourceData && Parameters.AssetImportDataOuter))
+			{
+				return nullptr;
+			}
+
+			UInterchangeAssetImportData* AssetImportData = Private::ImportCommon::BeginSetupAssetData(Parameters);
+
+			if (Parameters.AssetImportData) //-V1051
+			{
+				if (!Parameters.AssetImportData->IsA<UInterchangeAssetImportData>())
+				{
+					// Migrate the old source data
+					TArray<FAssetImportInfo::FSourceFile> OldSourceFiles = Parameters.AssetImportData->SourceData.SourceFiles;
+					AssetImportData->SetSourceFiles(MoveTemp(OldSourceFiles));
+				}
+			}
+
+			CustomFileSourceUpdate(AssetImportData);
+
+
+			Private::ImportCommon::EndSetupAssetData(Parameters, AssetImportData);
+
+			// Return the asset import data so it can be set on the imported asset.
+			return AssetImportData;
+#endif //#if WITH_EDITORONLY_DATA
+			return nullptr;
+		}
+
+#if WITH_EDITORONLY_DATA
+		FFactoryCommon::FSetImportAssetDataParameters::FSetImportAssetDataParameters(UObject* InAssetImportDataOuter
+																					, UAssetImportData* InAssetImportData
+																					, const UInterchangeSourceData* InSourceData
+																					, FString InNodeUniqueID
+																					, UInterchangeBaseNodeContainer* InNodeContainer
+																					, const TArray<UInterchangePipelineBase*>& InPipelines)
+			: FUpdateImportAssetDataParameters(InAssetImportDataOuter
+				, InAssetImportData
+				, InSourceData
+				, InNodeUniqueID
+				, InNodeContainer
+				, InPipelines
+				)
+			, SourceFiles()
+		{
+		}
+
+		UAssetImportData* FFactoryCommon::SetImportAssetData(FSetImportAssetDataParameters& Parameters)
+		{
+			UInterchangeAssetImportData* AssetImportData = Private::ImportCommon::BeginSetupAssetData(Parameters);
+
+			// Update the source files
+			{
+				if (Parameters.SourceFiles.IsEmpty())
+				{
+					TOptional<FMD5Hash> FileContentHash = Parameters.SourceData->GetFileContentHash();
+
+					// Todo add display label?
+					Parameters.SourceFiles.Emplace(
+						AssetImportData->SanitizeImportFilename(Parameters.SourceData->GetFilename()),
+						IFileManager::Get().GetTimeStamp(*Parameters.SourceData->GetFilename()),
+						FileContentHash.IsSet() ? FileContentHash.GetValue() : FMD5Hash()
+						);
+				}
+				else
+				{
+					for (FAssetImportInfo::FSourceFile& Source : Parameters.SourceFiles)
+					{
+						// This is done here since this is not thread safe
+						Source.RelativeFilename = AssetImportData->SanitizeImportFilename(Source.RelativeFilename);
+					}
+				}
+
+				AssetImportData->SetSourceFiles(MoveTemp(Parameters.SourceFiles));
+			}
+
+			Private::ImportCommon::EndSetupAssetData(Parameters, AssetImportData);
+
+			// Return the asset import data so it can be set on the imported asset.
+			return AssetImportData;
+		}
+#endif //#if WITH_EDITORONLY_DATA
+
+
+		void FFactoryCommon::ApplyReimportStrategyToAsset(UObject* Asset
+										  , UInterchangeBaseNode* PreviousAssetNode
+										  , UInterchangeBaseNode* CurrentAssetNode
+										  , UInterchangeBaseNode* PipelineAssetNode)
+		{
+			if (!ensure(PreviousAssetNode) || !ensure(PipelineAssetNode) || !ensure(CurrentAssetNode))
+			{
+				return;
+			}
+			EReimportStrategyFlags ReimportStrategyFlags = PipelineAssetNode->GetReimportStrategyFlags();
+			switch (ReimportStrategyFlags)
+			{
+				case EReimportStrategyFlags::ApplyNoProperties:
+				{
+					//We want to have no effect
+					break;
+				}
+					
+				case EReimportStrategyFlags::ApplyPipelineProperties:
+				{
+					//Directly apply pipeline node attribute to the asset
+					PipelineAssetNode->ApplyAllCustomAttributeToObject(Asset);
+					break;
+				}
+				
+				case EReimportStrategyFlags::ApplyEditorChangedProperties:
+				{
+					TArray<FAttributeKey> RemovedAttributes;
+					TArray<FAttributeKey> AddedAttributes;
+					TArray<FAttributeKey> ModifiedAttributes;
+					UInterchangeBaseNode::CompareNodeStorage(PreviousAssetNode, CurrentAssetNode, RemovedAttributes, AddedAttributes, ModifiedAttributes);
+
+					//set all ModifedAttributes from the CurrentAssetNode to the pipeline node. This will put back all user changes
+					UInterchangeBaseNode::CopyStorageAttributes(CurrentAssetNode, PipelineAssetNode, ModifiedAttributes);
+					//Now apply the pipeline node attribute to the asset
+					PipelineAssetNode->ApplyAllCustomAttributeToObject(Asset);
+					break;
+				}
+			}
+		}
+	}
+}
diff --git a/Engine/Source/Runtime/Experimental/Interchange/Engine/Public/InterchangeImportCommon.h b/Engine/Source/Runtime/Experimental/Interchange/Engine/Public/InterchangeImportCommon.h
new file mode 100644
--- /dev/null
+++ b/Engine/Source/Runtime/Experimental/Interchange/Engine/Public/InterchangeImportCommon.h
@@ -0,0 +1,90 @@
+// Copyright Epic Games, Inc. All Rights Reserved.
+#pragma once
+
+#include "CoreMinimal.h"
+#include "EditorFramework/AssetImportData.h"
+
+class UAssetImportData;
+class UInterchangeAssetImportData;
+class UInterchangeBaseNode;
+class UInterchangeBaseNodeContainer;
+class UInterchangePipelineBase;
+class UInterchangeSourceData;
+class UObject;
+
+namespace UE
+{
+	namespace Interchange
+	{
+		/**
+		 * All the code we cannot put in the base factory class because of dependencies (like Engine dep)
+		 * Will be available here.
+		 */
+		class INTERCHANGEENGINE_API FFactoryCommon
+		{
+		public:
+
+			struct INTERCHANGEENGINE_API FUpdateImportAssetDataParameters
+			{
+				UObject* AssetImportDataOuter = nullptr;
+				UAssetImportData* AssetImportData = nullptr;
+				const UInterchangeSourceData* SourceData = nullptr;
+				FString NodeUniqueID;
+				UInterchangeBaseNodeContainer* NodeContainer;
+				const TArray<UInterchangePipelineBase*> Pipelines;
+
+				FUpdateImportAssetDataParameters(UObject* InAssetImportDataOuter
+																	, UAssetImportData* InAssetImportData
+																	, const UInterchangeSourceData* InSourceData
+																	, FString InNodeUniqueID
+																	, UInterchangeBaseNodeContainer* InNodeContainer
+																	, const TArray<UInterchangePipelineBase*>& InPipelines);
+			};
+
+			/**
+			 * Update the AssetImportData source file of the specified asset in the parameters. Also update the node container and the node unique id.
+			 * If the AssetImportData is null it will create one. If the AssetImportData is not an UInterchangeAssetImportData it will create a new one.
+			 * @return The source data that should be stored on the asset or nullptr if a parameter is invalid
+			 */
+			static UAssetImportData* UpdateImportAssetData(FUpdateImportAssetDataParameters& Parameters);
+
+			/**
+			 * Update the AssetImportData of the specified asset in the parameters. Also update the node container and the node unique id.
+			 * If the AssetImportData is null it will create one. If the AssetImportData is not an UInterchangeAssetImportData it will create a new one.
+			 * The file source update is done by calling the function parameter CustomFileSourceUpdate, so its the client responsability to properly update the file source.
+			 * @return The source data that should be stored on the asset or nullptr if a parameter is invalid
+			 */
+			static UAssetImportData* UpdateImportAssetData(FUpdateImportAssetDataParameters& Parameters, TFunctionRef<void(UInterchangeAssetImportData*)> CustomFileSourceUpdate);
+
+#if WITH_EDITORONLY_DATA
+			struct INTERCHANGEENGINE_API FSetImportAssetDataParameters : public FUpdateImportAssetDataParameters
+			{
+				// Allow the factory to provide is own list of source files.
+				TArray<FAssetImportInfo::FSourceFile> SourceFiles;
+
+				FSetImportAssetDataParameters(UObject* InAssetImportDataOuter
+					, UAssetImportData* InAssetImportData
+					, const UInterchangeSourceData* InSourceData
+					, FString InNodeUniqueID
+					, UInterchangeBaseNodeContainer* InNodeContainer
+					, const TArray<UInterchangePipelineBase*>& InPipelines);
+			};
+
+			/**
+			 * Set the AssetImportData source file of the specified asset in the parameters. Also update the node container and the node unique id.
+			 * If the AssetImportData is null it will create one. If the AssetImportData is not an UInterchangeAssetImportData it will create a new one.
+			 * @return The source data that should be stored on the asset or nullptr if a parameter is invalid
+			 */
+			static UAssetImportData* SetImportAssetData(FSetImportAssetDataParameters& Parameters);
+#endif // WITH_EDITORONLY_DATA
+
+			/**
+			 * Apply the current strategy to the PipelineAssetNode
+			 */
+			static void ApplyReimportStrategyToAsset(UObject* Asset
+											  , UInterchangeBaseNode* PreviousAssetNode
+											  , UInterchangeBaseNode* CurrentAssetNode
+											  , UInterchangeBaseNode* PipelineAssetNode);
+		};
+	}
+}
