Add profiling stats for performance tracking.

diff --git a/Engine/Source/Runtime/Experimental/Chaos/Private/Chaos/PBDRigidsEvolutionGBF.cpp b/Engine/Source/Runtime/Experimental/Chaos/Private/Chaos/PBDRigidsEvolutionGBF.cpp
--- a/Engine/Source/Runtime/Experimental/Chaos/Private/Chaos/PBDRigidsEvolutionGBF.cpp
+++ b/Engine/Source/Runtime/Experimental/Chaos/Private/Chaos/PBDRigidsEvolutionGBF.cpp
@@ -258,6 +258,66 @@ void FPBDRigidsEvolutionGBF::BuildDisabledParticles(const int32 Island, TArray<T
 int32 DrawAwake = 0;
 FAutoConsoleVariableRef CVarDrawAwake(TEXT("p.chaos.DebugDrawAwake"),DrawAwake,TEXT("Draw particles that are awake"));
 
+#define CSV_CUSTOM_STAT_PER_ISLAND_HELPER(Stat) CSV_CUSTOM_STAT(PhysicsVerbose, Stat, Stats[Stat] * 1000.0, ECsvCustomStatOp::Set);
+
+//Using aligned to avoid false sharing during perf capture
+struct alignas(64) FPerIslandStats
+{
+#if CSV_PROFILER
+
+	enum EPerIslandStat
+	{
+		/** Any stat added to this enum must be added to ReportStats. Can use inl if it's a problem, but seems fine to just copy paste here*/
+		PerIslandSolve_ReloadCacheTotalSerialized = 0,
+		PerIslandSolve_GatherTotalSerialized,
+		PerIslandSolve_ApplyTotalSerialized,
+		PerIslandSolve_UpdateVelocitiesTotalSerialized,
+		PerIslandSolve_ApplyPushOutTotalSerialized,
+		PerIslandSolve_ScatterTotalSerialized,
+		NumStats
+	};
+
+	void ReportStats() const
+	{
+		CSV_CUSTOM_STAT_PER_ISLAND_HELPER(PerIslandSolve_ReloadCacheTotalSerialized);
+		CSV_CUSTOM_STAT_PER_ISLAND_HELPER(PerIslandSolve_GatherTotalSerialized);
+		CSV_CUSTOM_STAT_PER_ISLAND_HELPER(PerIslandSolve_ApplyTotalSerialized);
+		CSV_CUSTOM_STAT_PER_ISLAND_HELPER(PerIslandSolve_UpdateVelocitiesTotalSerialized);
+		CSV_CUSTOM_STAT_PER_ISLAND_HELPER(PerIslandSolve_ApplyPushOutTotalSerialized);
+		CSV_CUSTOM_STAT_PER_ISLAND_HELPER(PerIslandSolve_ScatterTotalSerialized);
+	}
+
+	FPerIslandStats()
+	{
+		for (double& Stat : Stats) { Stat = 0.0; }
+	}
+
+	static FPerIslandStats Flatten(const TArray<FPerIslandStats>& PerIslandStats)
+	{
+		FPerIslandStats Flat;
+		for(const FPerIslandStats& IslandStats : PerIslandStats)
+		{
+			for(int32 Stat = 0; Stat < NumStats; ++Stat)
+			{
+				Flat.Stats[Stat] += IslandStats.Stats[Stat];
+			}
+		}
+
+		return Flat;
+	}
+
+	double Stats[EPerIslandStat::NumStats];
+#else
+	void ReportStats() const{}
+#endif
+
+};
+
+#if CSV_PROFILER
+#define CSV_SCOPED_PER_ISLAND_TIMING_STAT(Stat) FScopedDurationTimer Timer_##Stat(PerIslandStats[GroupIndex].Stats[FPerIslandStats::Stat])
+#else
+#define CSV_SCOPED_PER_ISLAND_TIMING_STAT(Stat)
+#endif
 void FPBDRigidsEvolutionGBF::AdvanceOneTimeStepImpl(const FReal Dt, const FSubStepInfo& SubStepInfo)
 {
 	SCOPE_CYCLE_COUNTER(STAT_Evolution_AdvanceOneTimeStep);
@@ -377,7 +437,15 @@ void FPBDRigidsEvolutionGBF::AdvanceOneTimeStepImpl(const FReal Dt, const FSubSt
 	if(Dt > 0)
 	{
 		SCOPE_CYCLE_COUNTER(STAT_Evolution_ParallelSolve);
-		PhysicsParallelFor(GetConstraintGraph().NumGroups(), [&](int32 GroupIndex) {
+		CSV_SCOPED_TIMING_STAT(PhysicsVerbose, StepSolver_PerIslandSolve);
+		const int32 NumGroups = GetConstraintGraph().NumGroups();
+		CSV_CUSTOM_STAT(PhysicsCounters, NumIslandGroups, NumGroups, ECsvCustomStatOp::Set);
+		TArray<FPerIslandStats> PerIslandStats;
+#if CSV_PROFILER
+		PerIslandStats.AddDefaulted(NumGroups);
+#endif
+
+		PhysicsParallelFor(NumGroups, [&](int32 GroupIndex) {
 
 			if(GetConstraintGraph().GetIslandGroups().IsValidIndex(GroupIndex) &&
 			  !GetConstraintGraph().GetIslandGroup(GroupIndex)->IsValid())
@@ -386,7 +454,7 @@ void FPBDRigidsEvolutionGBF::AdvanceOneTimeStepImpl(const FReal Dt, const FSubSt
 			
 			// Reload cache if necessary
 			{
-				CSV_SCOPED_TIMING_STAT(PhysicsVerbose, StepSolver_ReloadCache);
+				CSV_SCOPED_PER_ISLAND_TIMING_STAT(PerIslandSolve_ReloadCacheTotalSerialized);
 				for(auto& IslandSolver : GetConstraintGraph().GetGroupIslands(GroupIndex))
 				{
 					ReloadParticlesCache(IslandSolver->GetIslandIndex());
@@ -396,7 +464,7 @@ void FPBDRigidsEvolutionGBF::AdvanceOneTimeStepImpl(const FReal Dt, const FSubSt
 			// Collect all the data that the constraint solvers operate on
 			{
 				SCOPE_CYCLE_COUNTER(STAT_Evolution_Gather);
-				CSV_SCOPED_TIMING_STAT(PhysicsVerbose, StepSolver_Gather);
+				CSV_SCOPED_PER_ISLAND_TIMING_STAT(PerIslandSolve_GatherTotalSerialized);
 				GatherSolverInput(Dt, GroupIndex);
 			}
 
@@ -406,7 +474,7 @@ void FPBDRigidsEvolutionGBF::AdvanceOneTimeStepImpl(const FReal Dt, const FSubSt
 			{
 				SCOPE_CYCLE_COUNTER(STAT_Evolution_ApplyConstraintsPhase1);
 				CSV_SCOPED_TIMING_STAT(Chaos, ApplyConstraints);
-				CSV_SCOPED_TIMING_STAT(PhysicsVerbose, StepSolver_Apply);
+				CSV_SCOPED_PER_ISLAND_TIMING_STAT(PerIslandSolve_ApplyTotalSerialized);
 				ApplyConstraintsPhase1(Dt, GroupIndex);
 			}
 
@@ -424,7 +492,7 @@ void FPBDRigidsEvolutionGBF::AdvanceOneTimeStepImpl(const FReal Dt, const FSubSt
 			// Update implicit velocities from results of constraint solver phase 1
 			{
 				SCOPE_CYCLE_COUNTER(STAT_Evolution_UpdateVelocites);
-				CSV_SCOPED_TIMING_STAT(PhysicsVerbose, StepSolver_UpdateVelocities);
+				CSV_SCOPED_PER_ISLAND_TIMING_STAT(PerIslandSolve_UpdateVelocitiesTotalSerialized);
 				SetImplicitVelocities(Dt, GroupIndex);
 			}
 
@@ -434,7 +502,7 @@ void FPBDRigidsEvolutionGBF::AdvanceOneTimeStepImpl(const FReal Dt, const FSubSt
 			// For QPBD this is the velocity solve step
 			{
 				SCOPE_CYCLE_COUNTER(STAT_Evolution_ApplyConstraintsPhase2);
-				CSV_SCOPED_TIMING_STAT(PhysicsVerbose, StepSolver_ApplyPushOut);
+				CSV_SCOPED_PER_ISLAND_TIMING_STAT(PerIslandSolve_ApplyPushOutTotalSerialized);
 				ApplyConstraintsPhase2(Dt, GroupIndex);
 			}
 
@@ -442,7 +510,7 @@ void FPBDRigidsEvolutionGBF::AdvanceOneTimeStepImpl(const FReal Dt, const FSubSt
 			// that is accessed externally (net impusles, break info, etc)
 			{
 				SCOPE_CYCLE_COUNTER(STAT_Evolution_Scatter);
-				CSV_SCOPED_TIMING_STAT(PhysicsVerbose, StepSolver_Scatter);
+				CSV_SCOPED_PER_ISLAND_TIMING_STAT(PerIslandSolve_ScatterTotalSerialized);
 				ScatterSolverOutput(Dt, GroupIndex);
 			}
 
@@ -459,6 +527,9 @@ void FPBDRigidsEvolutionGBF::AdvanceOneTimeStepImpl(const FReal Dt, const FSubSt
 				}
 			}
 		});
+
+		FPerIslandStats FlattenedStats = FPerIslandStats::Flatten(PerIslandStats);
+		FlattenedStats.ReportStats();
 	}
 
 	{
