Chaos : performance : fix bad performance for creating convex based geometry collection with a large number of pieces

diff --git a/Engine/Source/Runtime/Experimental/Chaos/Private/Chaos/Convex.cpp b/Engine/Source/Runtime/Experimental/Chaos/Private/Chaos/Convex.cpp
--- a/Engine/Source/Runtime/Experimental/Chaos/Private/Chaos/Convex.cpp
+++ b/Engine/Source/Runtime/Experimental/Chaos/Private/Chaos/Convex.cpp
@@ -10,7 +10,7 @@ namespace Chaos
 {
 	TUniquePtr<FImplicitObject> FConvex::CopyWithScale(const FVec3& Scale) const
 	{
-		TSharedPtr<FConvex, ESPMode::ThreadSafe> ConvexCopy = MakeShared<FConvex, ESPMode::ThreadSafe>(GetVertices(), GetMargin());
+		TSharedPtr<FConvex, ESPMode::ThreadSafe> ConvexCopy = MakeShared<FConvex, ESPMode::ThreadSafe>(FConvex(*this));
 		return  TUniquePtr<FImplicitObject>(new TImplicitObjectScaled<FConvex>(ConvexCopy, Scale));
 	}
 
diff --git a/Engine/Source/Runtime/Experimental/Chaos/Public/Chaos/Convex.h b/Engine/Source/Runtime/Experimental/Chaos/Public/Chaos/Convex.h
--- a/Engine/Source/Runtime/Experimental/Chaos/Public/Chaos/Convex.h
+++ b/Engine/Source/Runtime/Experimental/Chaos/Public/Chaos/Convex.h
@@ -42,7 +42,6 @@ namespace Chaos
 			, Volume(0.f)
 			, CenterOfMass(FVec3Type(0.f))
 		{}
-		FConvex(const FConvex&) = delete;
 		FConvex(FConvex&& Other)
 		    : FImplicitObject(EImplicitObject::IsConvex | EImplicitObject::HasBoundingBox, ImplicitObjectType::Convex)
 			, Planes(MoveTemp(Other.Planes))
@@ -153,7 +152,7 @@ namespace Chaos
 
 		virtual TUniquePtr<FImplicitObject> Copy() const
 		{
-			return TUniquePtr<FImplicitObject>(new FConvex(GetVertices(), GetMargin()));
+			return TUniquePtr<FImplicitObject>(new FConvex(*this));
 		}
 
 		virtual TUniquePtr<FImplicitObject> CopyWithScale(const FVec3& Scale) const override;
@@ -930,6 +929,20 @@ namespace Chaos
 		friend FISPCDataVerifier;
 #endif // #if INTEL_ISPC && !UE_BUILD_SHIPPING
 
+	private:
+		// IMPORTANT to keep this copy constructor private to avoid unintented expensive copies 
+		FConvex(const FConvex& Other)
+			: FImplicitObject(EImplicitObject::IsConvex | EImplicitObject::HasBoundingBox, ImplicitObjectType::Convex)
+			, Planes(Other.Planes)
+			, Vertices(Other.Vertices)
+			, LocalBoundingBox(Other.LocalBoundingBox)
+			, Volume(Other.Volume)
+			, CenterOfMass(Other.CenterOfMass)
+		{
+			StructureData.CopyFrom(Other.StructureData);
+			SetMargin(Other.Margin);
+		}
+
 	private:
 		TArray<FPlaneType> Planes;
 		TArray<FVec3Type> Vertices; //copy of the vertices that are just on the convex hull boundary
diff --git a/Engine/Source/Runtime/Experimental/Chaos/Public/Chaos/ConvexStructureData.h b/Engine/Source/Runtime/Experimental/Chaos/Public/Chaos/ConvexStructureData.h
--- a/Engine/Source/Runtime/Experimental/Chaos/Public/Chaos/ConvexStructureData.h
+++ b/Engine/Source/Runtime/Experimental/Chaos/Public/Chaos/ConvexStructureData.h
@@ -127,8 +127,33 @@ namespace Chaos
 
 			return *this;
 		}
+		
+		void CopyFrom(const FConvexStructureData& Other)
+		{
+			DestroyDataContainer();
 
+			check(Data.Ptr == nullptr);
+			check(IndexType == EIndexType::None);
 
+			if (Other.IndexType == EIndexType::Large)
+			{
+				check(Other.Data.DataL);
+				Data.DataL = new FConvexStructureDataLarge(*Other.Data.DataL);
+			}
+			else if (Other.IndexType == EIndexType::Medium)
+			{
+				check(Other.Data.DataM);
+				Data.DataM = new FConvexStructureDataMedium(*Other.Data.DataM);
+			}
+			else if (Other.IndexType == EIndexType::Small)
+			{
+				check(Other.Data.DataS);
+				Data.DataS = new FConvexStructureDataSmall(*Other.Data.DataS);
+			}
+
+			IndexType = Other.IndexType;
+		}
+		
 		inline bool IsValid() const
 		{
 			return (Data.Ptr != nullptr);
