Chaos : Clustering
-Only set the initial active clusters to sleep when the Object Type is initialized as sleeping

diff --git a/Engine/Source/Runtime/Experimental/Chaos/Private/Chaos/PBDRigidClustering.cpp b/Engine/Source/Runtime/Experimental/Chaos/Private/Chaos/PBDRigidClustering.cpp
--- a/Engine/Source/Runtime/Experimental/Chaos/Private/Chaos/PBDRigidClustering.cpp
+++ b/Engine/Source/Runtime/Experimental/Chaos/Private/Chaos/PBDRigidClustering.cpp
@@ -256,8 +256,14 @@ namespace Chaos
 
 		if(ClusterUnionMap.Num())
 		{
+			struct FClusterGroup {
+				TArray<FPBDRigidParticleHandle*> Bodies = TArray < FPBDRigidParticleHandle*>();
+				bool bIsSleeping = true;
+			};
+
+
 			TMap<FPBDRigidParticleHandle*, FPBDRigidParticleHandle*> ChildToParentMap;
-			TMap<int32, TArray<FPBDRigidParticleHandle*>> NewClusterGroups;
+			TMap<int32, FClusterGroup> NewClusterGroups;
 
 			// Walk the list of registered cluster groups
 			for(TTuple<int32, TArray<FPBDRigidClusteredParticleHandle* >>& Group : ClusterUnionMap)
@@ -270,9 +276,10 @@ namespace Chaos
 					// First see if this is a new group
 					if(!NewClusterGroups.Contains(ClusterGroupID))
 					{
-						NewClusterGroups.Add(ClusterGroupID, TArray < FPBDRigidParticleHandle*>());
+						NewClusterGroups.Add(ClusterGroupID, FClusterGroup());
 					}
 
+					bool bIsSleeping = true;
 					TArray<FPBDRigidParticleHandle*> ClusterBodies;
 					for(FPBDRigidClusteredParticleHandle* ActiveCluster : Handles)
 					{
@@ -281,13 +288,16 @@ namespace Chaos
 							// If this is an external cluster (from the rest collection) we release its children and append them to the current group
 							TSet<FPBDRigidParticleHandle*> Children;
 							
+							// let sleeping clusters stay asleep
+							bIsSleeping &= ActiveCluster->ObjectState() == EObjectStateType::Sleeping;
+
 							{
 								// First disable breaking data generation - this is not a break we're just reclustering under a dynamic parent.
 								TGuardValue<bool> BreakFlagGuard(DoGenerateBreakingData, false);
 								Children = ReleaseClusterParticles(ActiveCluster, nullptr, true);
 							}
 
-							NewClusterGroups[ClusterGroupID].Append(Children.Array());
+							NewClusterGroups[ClusterGroupID].Bodies.Append(Children.Array());
 							
 							for(FPBDRigidParticleHandle* Child : Children)
 							{
@@ -295,21 +305,29 @@ namespace Chaos
 							}
 						}
 					}
+					NewClusterGroups[ClusterGroupID].bIsSleeping = bIsSleeping;
 				}
 			}
 
 			// For new cluster groups, create an internal cluster parent.
-			for(TTuple<int32, TArray<FPBDRigidParticleHandle* >>& Group : NewClusterGroups)
+			for(TTuple<int32, FClusterGroup>& Group : NewClusterGroups)
 			{
 				int32 ClusterGroupID = FMath::Abs(Group.Key);
 
-				TArray<FPBDRigidParticleHandle*> ActiveCluster = Group.Value;
+				TArray<FPBDRigidParticleHandle*> ActiveCluster = Group.Value.Bodies;
 
 				FClusterCreationParameters Parameters(0.3f, 100, false, !!UnionsHaveCollisionParticles);
 				Parameters.ConnectionMethod = MClusterUnionConnectionType;
-				TPBDRigidClusteredParticleHandleImp<FReal, 3, true>* Handle = CreateClusterParticle(-ClusterGroupID, MoveTemp(Group.Value), Parameters, TSharedPtr<FImplicitObject, ESPMode::ThreadSafe>());
+				TPBDRigidClusteredParticleHandleImp<FReal, 3, true>* Handle = 
+					CreateClusterParticle(-ClusterGroupID, MoveTemp(Group.Value.Bodies), Parameters, 
+						TSharedPtr<FImplicitObject, ESPMode::ThreadSafe>());
 				Handle->SetInternalCluster(true);
 
+				if (Group.Value.bIsSleeping)
+				{
+					MEvolution.SetParticleObjectState(Handle, Chaos::EObjectStateType::Sleeping);
+				}
+
 				MEvolution.SetPhysicsMaterial(Handle, MEvolution.GetPhysicsMaterial(ActiveCluster[0]));
 
 				for(FPBDRigidParticleHandle* Constituent : ActiveCluster)
@@ -1015,6 +1033,47 @@ namespace Chaos
 		return AllActivatedChildren;
 	}
 
+
+	DECLARE_CYCLE_STAT(TEXT("FRigidClustering::Visitor"), STAT_ClusterVisitor, STATGROUP_Chaos);
+	void FRigidClustering::Visitor(FClusterHandle Cluster, FVisitorFunction Function)
+	{
+		if (Cluster)
+		{
+			if (MChildren.Contains(Cluster) && MChildren[Cluster].Num())
+			{
+				SCOPE_CYCLE_COUNTER(STAT_ClusterVisitor);
+
+				// TQueue is a linked list, which has no preallocator.
+				TQueue<FRigidHandle> Queue;
+				for (Chaos::FPBDRigidParticleHandle* Child : MChildren[Cluster])
+				{
+					Queue.Enqueue(Child);
+				}
+
+				FRigidHandle CurrentHandle = nullptr;
+				while (Queue.Dequeue(CurrentHandle))
+				{
+					if (FClusterHandle CurrentClusterHandle = CurrentHandle->CastToClustered())
+					{
+						// @question : Maybe we should just store the leaf node bodies in a
+						// map, that will require Memory(n*log(n))
+						if (MChildren.Contains(CurrentClusterHandle))
+						{
+							for (Chaos::FPBDRigidParticleHandle* Child : MChildren[CurrentClusterHandle])
+							{
+								Queue.Enqueue(Child);
+							}
+						}
+					}
+					if (CurrentHandle)
+					{
+						Function(*this, CurrentHandle);
+					}
+				}
+			}
+		}
+	}
+
 	DECLARE_CYCLE_STAT(TEXT("TPBDRigidClustering<>::GetActiveClusterIndex"), STAT_GetActiveClusterIndex, STATGROUP_Chaos);
 	FPBDRigidParticleHandle* 
 	FRigidClustering::GetActiveClusterIndex(
diff --git a/Engine/Source/Runtime/Experimental/Chaos/Private/PhysicsProxy/GeometryCollectionPhysicsProxy.cpp b/Engine/Source/Runtime/Experimental/Chaos/Private/PhysicsProxy/GeometryCollectionPhysicsProxy.cpp
--- a/Engine/Source/Runtime/Experimental/Chaos/Private/PhysicsProxy/GeometryCollectionPhysicsProxy.cpp
+++ b/Engine/Source/Runtime/Experimental/Chaos/Private/PhysicsProxy/GeometryCollectionPhysicsProxy.cpp
@@ -908,6 +908,15 @@ void FGeometryCollectionPhysicsProxy::InitializeBodiesPT(Chaos::FPBDRigidsSolver
 		{
 			if (FClusterHandle* Handle = SolverParticleHandles[TransformGroupIndex])
 			{
+				// Sleeping Geometry Collections:
+				//   A sleeping geometry collection is dynamic internally, and then the top level
+				//   active clusters are set to sleeping. Sleeping is not propagated up from the 
+				//   leaf nodes like kinematic or dynamic clusters. 
+				if (!Handle->Disabled() && Parameters.ObjectType == EObjectStateTypeEnum::Chaos_Object_Sleeping)
+				{
+					RigidsSolver->GetEvolution()->SetParticleObjectState(Handle, Chaos::EObjectStateType::Sleeping);
+				}
+
 				RigidsSolver->GetEvolution()->DirtyParticle(*Handle);
 			}
 		}
diff --git a/Engine/Source/Runtime/Experimental/Chaos/Public/Chaos/PBDRigidClustering.h b/Engine/Source/Runtime/Experimental/Chaos/Public/Chaos/PBDRigidClustering.h
--- a/Engine/Source/Runtime/Experimental/Chaos/Public/Chaos/PBDRigidClustering.h
+++ b/Engine/Source/Runtime/Experimental/Chaos/Public/Chaos/PBDRigidClustering.h
@@ -53,11 +53,12 @@ class CHAOS_API FRigidClustering
 {
 public:
 
-	typedef FPBDRigidsEvolutionGBF                      FRigidEvolution;
-	typedef FPBDRigidParticleHandle*					FRigidHandle;
-	typedef TArray<FRigidHandle>						FRigidHandleArray;
-	typedef FPBDRigidClusteredParticleHandle*			FClusterHandle;
-	typedef TMap<FClusterHandle, FRigidHandleArray>	    FClusterMap;
+	typedef FPBDRigidsEvolutionGBF								FRigidEvolution;
+	typedef FPBDRigidParticleHandle*							FRigidHandle;
+	typedef TArray<FRigidHandle>								FRigidHandleArray;
+	typedef FPBDRigidClusteredParticleHandle*					FClusterHandle;
+	typedef TMap<FClusterHandle, FRigidHandleArray>				FClusterMap;
+	typedef TFunction<void(FRigidClustering&, FRigidHandle)>	FVisitorFunction;
 
 	FRigidClustering(FRigidEvolution& InEvolution, FPBDRigidClusteredParticles& InParticles);
 	~FRigidClustering();
@@ -174,6 +175,13 @@ public:
 	// Access
 	//
 
+	/**
+	*
+	*  Visitor
+	*   Walk all the decendents of the current cluster and execute FVisitorFunction.
+	*   FVisitorFunction = [](FRigidClustering& Clustering, FRigidHandle RigidHandle){}
+	*/
+	void Visitor(FClusterHandle Cluster, FVisitorFunction Function);
 
 	/*
 	*  GetActiveClusterIndex
@@ -261,6 +269,10 @@ public:
 	const TSet<Chaos::FPBDRigidClusteredParticleHandle*>& GetTopLevelClusterParents() const { return TopLevelClusterParents; }
 	TSet<Chaos::FPBDRigidClusteredParticleHandle*>& GetTopLevelClusterParents() { return TopLevelClusterParents; }
 
+	FRigidEvolution& GetEvolution() { return MEvolution; }
+	const FRigidEvolution& GetEvolution() const { return MEvolution; }
+
+
  protected:
 
 	void ComputeStrainFromCollision(const FPBDCollisionConstraints& CollisionRule);
diff --git a/Engine/Source/Runtime/Experimental/Chaos/Public/GeometryCollection/GeometryCollectionSimulationCoreTypes.h b/Engine/Source/Runtime/Experimental/Chaos/Public/GeometryCollection/GeometryCollectionSimulationCoreTypes.h
--- a/Engine/Source/Runtime/Experimental/Chaos/Public/GeometryCollection/GeometryCollectionSimulationCoreTypes.h
+++ b/Engine/Source/Runtime/Experimental/Chaos/Public/GeometryCollection/GeometryCollectionSimulationCoreTypes.h
@@ -219,6 +219,7 @@ struct FSimulationParameters
 		, CacheBeginTime(Other.CacheBeginTime)
 		, ReverseCacheBeginTime(Other.ReverseCacheBeginTime)
 		, bClearCache(Other.bClearCache)
+		, ObjectType(Other.ObjectType)
 		, PhysicalMaterialHandle(Other.PhysicalMaterialHandle)
 		, bGenerateBreakingData(Other.bGenerateBreakingData)
 		, bGenerateCollisionData(Other.bGenerateCollisionData)
diff --git a/Engine/Source/Runtime/Experimental/GeometryCollectionEngine/Private/GeometryCollection/GeometryCollectionComponent.cpp b/Engine/Source/Runtime/Experimental/GeometryCollectionEngine/Private/GeometryCollection/GeometryCollectionComponent.cpp
--- a/Engine/Source/Runtime/Experimental/GeometryCollectionEngine/Private/GeometryCollection/GeometryCollectionComponent.cpp
+++ b/Engine/Source/Runtime/Experimental/GeometryCollectionEngine/Private/GeometryCollection/GeometryCollectionComponent.cpp
@@ -1971,21 +1971,22 @@ void UGeometryCollectionComponent::OnCreatePhysicsState()
 			// If the Component is set to Dynamic, we look to the RestCollection for initial dynamic state override per transform.
 			TManagedArray<int32> & DynamicState = DynamicCollection->DynamicState;
 
-			if (ObjectType != EObjectStateTypeEnum::Chaos_Object_UserDefined)
+			EObjectStateTypeEnum LocalObjectType = (ObjectType != EObjectStateTypeEnum::Chaos_Object_Sleeping) ? ObjectType : EObjectStateTypeEnum::Chaos_Object_Dynamic;
+			if (LocalObjectType != EObjectStateTypeEnum::Chaos_Object_UserDefined)
 			{
-				if (RestCollection && (ObjectType == EObjectStateTypeEnum::Chaos_Object_Dynamic))
+				if (RestCollection && (LocalObjectType == EObjectStateTypeEnum::Chaos_Object_Dynamic))
 				{
 					TManagedArray<int32>& InitialDynamicState = RestCollection->GetGeometryCollection()->InitialDynamicState;
 					for (int i = 0; i < DynamicState.Num(); i++)
 					{
-						DynamicState[i] = (InitialDynamicState[i] == static_cast<int32>(Chaos::EObjectStateType::Uninitialized)) ? static_cast<int32>(ObjectType) : InitialDynamicState[i];
+						DynamicState[i] = (InitialDynamicState[i] == static_cast<int32>(Chaos::EObjectStateType::Uninitialized)) ? static_cast<int32>(LocalObjectType) : InitialDynamicState[i];
 					}
 				}
 				else
 				{
 					for (int i = 0; i < DynamicState.Num(); i++)
 					{
-						DynamicState[i] = static_cast<int32>(ObjectType);
+						DynamicState[i] = static_cast<int32>(LocalObjectType);
 					}
 				}
 			}
