Fix cluster union children being filtered out during collision.
Remove dependency on clustered union particle mapping always finding a particle.

diff --git a/Engine/Source/Runtime/Experimental/Chaos/Private/Chaos/CollisionResolution.cpp b/Engine/Source/Runtime/Experimental/Chaos/Private/Chaos/CollisionResolution.cpp
--- a/Engine/Source/Runtime/Experimental/Chaos/Private/Chaos/CollisionResolution.cpp
+++ b/Engine/Source/Runtime/Experimental/Chaos/Private/Chaos/CollisionResolution.cpp
@@ -2340,7 +2340,8 @@ namespace Chaos
 
 					for (const FLargeUnionClusteredImplicitInfo& Child0 : Children)
 					{
-						LegacyConstructConstraints<T_TRAITS>(Particle0, Particle1, Child0.Implicit, Child0.Shapes, Child0.BVHParticles, Implicit1, Shapes1, Simplicial1, ParticleWorldTransform0, Child0.Transform * LocalTransform0, ParticleWorldTransform1, LocalTransform1, CullDistance, Dt, Context);
+						// Provide no shape for clustered child implicits, these use their parent's filter data, which was already tested when generating shape pair in midphase.
+						LegacyConstructConstraints<T_TRAITS>(Particle0, Particle1, Child0.Implicit, FShapeOrShapesArray(), Child0.BVHParticles, Implicit1, Shapes1, Simplicial1, ParticleWorldTransform0, Child0.Transform * LocalTransform0, ParticleWorldTransform1, LocalTransform1, CullDistance, Dt, Context);
 					}
 				}
 				else
@@ -2349,10 +2350,8 @@ namespace Chaos
 					{
 						const TPBDRigidParticleHandle<FReal, 3>* OriginalParticle = Union0->FindParticleForImplicitObject(Child0.Get());
 
-						if (ensure(OriginalParticle))
-						{
-							LegacyConstructConstraints<T_TRAITS>(Particle0, Particle1, Child0.Get(), FShapeOrShapesArray(OriginalParticle), OriginalParticle->CollisionParticles().Get(), Implicit1, Shapes1, Simplicial1, ParticleWorldTransform0, LocalTransform0, ParticleWorldTransform1, LocalTransform1, CullDistance, Dt, Context);
-						}
+						// Provide no shape for clustered child implicits, these use their parent's filter data, which was already tested when generating shape pair in midphase.
+						LegacyConstructConstraints<T_TRAITS>(Particle0, Particle1, Child0.Get(), FShapeOrShapesArray(), OriginalParticle ? OriginalParticle->CollisionParticles().Get() : Simplicial0, Implicit1, Shapes1, Simplicial1, ParticleWorldTransform0, LocalTransform0, ParticleWorldTransform1, LocalTransform1, CullDistance, Dt, Context);
 					}
 				}
 				return;
@@ -2401,7 +2400,8 @@ namespace Chaos
 
 					for (const FLargeUnionClusteredImplicitInfo& Child1 : Children)
 					{
-						LegacyConstructConstraints<T_TRAITS>(Particle0, Particle1, Implicit0, Shapes0, Simplicial0, Child1.Implicit, Child1.Shapes, Child1.BVHParticles, ParticleWorldTransform0, LocalTransform0, ParticleWorldTransform1, Child1.Transform * LocalTransform1, CullDistance, Dt, Context);
+						// Provide no shape for clustered child implicits, these use their parent's filter data, which was already tested when generating shape pair in midphase.
+						LegacyConstructConstraints<T_TRAITS>(Particle0, Particle1, Implicit0, Shapes0, Simplicial0, Child1.Implicit, FShapeOrShapesArray(), Child1.BVHParticles, ParticleWorldTransform0, LocalTransform0, ParticleWorldTransform1, Child1.Transform * LocalTransform1, CullDistance, Dt, Context);
 					}
 				}
 				else
@@ -2409,10 +2409,9 @@ namespace Chaos
 					for (const TUniquePtr<FImplicitObject>& Child1 : Union1->GetObjects())
 					{
 						const TPBDRigidParticleHandle<FReal, 3>* OriginalParticle = Union1->FindParticleForImplicitObject(Child1.Get());
-						if (ensure(OriginalParticle))
-						{
-							LegacyConstructConstraints<T_TRAITS>(Particle0, Particle1, Implicit0, Shapes0, Simplicial0, Child1.Get(), FShapeOrShapesArray(OriginalParticle), OriginalParticle->CollisionParticles().Get(), ParticleWorldTransform0, LocalTransform0, ParticleWorldTransform1, LocalTransform1, CullDistance, Dt, Context);
-						}
+
+						// Provide no shape for clustered child implicits, these use their parent's filter data, which was already tested when generating shape pair in midphase.
+						LegacyConstructConstraints<T_TRAITS>(Particle0, Particle1, Implicit0, Shapes0, Simplicial0, Child1.Get(), FShapeOrShapesArray(), OriginalParticle ? OriginalParticle->CollisionParticles().Get() : Simplicial1, ParticleWorldTransform0, LocalTransform0, ParticleWorldTransform1, LocalTransform1, CullDistance, Dt, Context);
 					}
 				}
 				return;
diff --git a/Engine/Source/Runtime/Experimental/Chaos/Private/Chaos/ImplicitObjectUnion.cpp b/Engine/Source/Runtime/Experimental/Chaos/Private/Chaos/ImplicitObjectUnion.cpp
--- a/Engine/Source/Runtime/Experimental/Chaos/Private/Chaos/ImplicitObjectUnion.cpp
+++ b/Engine/Source/Runtime/Experimental/Chaos/Private/Chaos/ImplicitObjectUnion.cpp
@@ -234,22 +234,8 @@ void FImplicitObjectUnionClustered::FindAllIntersectingClusteredObjects(TArray<F
 		for (int32 Idx : Overlaps)
 		{
 			const FImplicitObject* Obj = LargeUnionData->GeomParticles.Geometry(Idx).Get();
-			const FShapesArray& ShapesArray = LargeUnionData->GeomParticles.ShapesArray(Idx);
-
-			// If implicit is union return shapes array of its particle, otherwise retrieve single shape.
-			FShapeOrShapesArray Shapes;
-			if (Obj->IsUnderlyingUnion())
-			{
-				Shapes = FShapeOrShapesArray(&ShapesArray);
-			}
-			else
-			{
-				Shapes = FShapeOrShapesArray(ShapesArray[0].Get());
-			}
-			
-
 			const FBVHParticles* Simplicial = MOriginalParticleLookupHack.IsValidIndex(Idx) ? MOriginalParticleLookupHack[Idx]->CollisionParticles().Get() : nullptr;
-			Out.Add(FLargeUnionClusteredImplicitInfo(Obj,FRigidTransform3(LargeUnionData->GeomParticles.X(Idx), LargeUnionData->GeomParticles.R(Idx)), Simplicial, Shapes));
+			Out.Add(FLargeUnionClusteredImplicitInfo(Obj,FRigidTransform3(LargeUnionData->GeomParticles.X(Idx), LargeUnionData->GeomParticles.R(Idx)), Simplicial));
 		}
 	}
 	else
@@ -269,18 +255,9 @@ void FImplicitObjectUnionClustered::FindAllIntersectingClusteredObjects(TArray<F
 		for (int32 Idx = 0; Idx < LocalOut.Num(); ++Idx)
 		{
 			auto& OutElem = LocalOut[Idx];
+			const FBVHParticles* Simplicial = MOriginalParticleLookupHack.IsValidIndex(Idxs[Idx]) ? MOriginalParticleLookupHack[Idxs[Idx]]->CollisionParticles().Get() : nullptr;
 
-			FShapeOrShapesArray Shapes;
-			const FBVHParticles* Simplicial = nullptr;
-
-			if (ensure(MOriginalParticleLookupHack.IsValidIndex(Idxs[Idx])))
-			{
-				FPBDRigidParticleHandle* Handle = MOriginalParticleLookupHack[Idxs[Idx]];
-				Shapes = FShapeOrShapesArray(Handle);
-				Simplicial = Handle->CollisionParticles().Get();
-			}
-
-			Out.Add(FLargeUnionClusteredImplicitInfo(OutElem.First, OutElem.Second, Simplicial, Shapes));
+			Out.Add(FLargeUnionClusteredImplicitInfo(OutElem.First, OutElem.Second, Simplicial));
 		}
 	}
 }
diff --git a/Engine/Source/Runtime/Experimental/Chaos/Public/Chaos/ImplicitObjectUnion.h b/Engine/Source/Runtime/Experimental/Chaos/Public/Chaos/ImplicitObjectUnion.h
--- a/Engine/Source/Runtime/Experimental/Chaos/Public/Chaos/ImplicitObjectUnion.h
+++ b/Engine/Source/Runtime/Experimental/Chaos/Public/Chaos/ImplicitObjectUnion.h
@@ -216,18 +216,16 @@ protected:
 
 struct FLargeUnionClusteredImplicitInfo
 {
-	FLargeUnionClusteredImplicitInfo(const FImplicitObject* InImplicit, const FRigidTransform3& InTransform, const FBVHParticles* InBVHParticles, const FShapeOrShapesArray& InShapes)
+	FLargeUnionClusteredImplicitInfo(const FImplicitObject* InImplicit, const FRigidTransform3& InTransform, const FBVHParticles* InBVHParticles)
 		: Implicit(InImplicit)
 		, Transform(InTransform)
 		, BVHParticles(InBVHParticles)
-		, Shapes(InShapes)
 	{
 	}
 
 	const FImplicitObject* Implicit;
 	FRigidTransform3 Transform;
 	const FBVHParticles* BVHParticles;
-	FShapeOrShapesArray Shapes;
 };
 
 class CHAOS_API FImplicitObjectUnionClustered: public FImplicitObjectUnion
