Chaos : Geometry Collection
- Update the SOA when global transforms for animation collections are updated.

diff --git a/Engine/Source/Runtime/Experimental/Chaos/Private/PhysicsProxy/GeometryCollectionPhysicsProxy.cpp b/Engine/Source/Runtime/Experimental/Chaos/Private/PhysicsProxy/GeometryCollectionPhysicsProxy.cpp
--- a/Engine/Source/Runtime/Experimental/Chaos/Private/PhysicsProxy/GeometryCollectionPhysicsProxy.cpp
+++ b/Engine/Source/Runtime/Experimental/Chaos/Private/PhysicsProxy/GeometryCollectionPhysicsProxy.cpp
@@ -1518,7 +1518,11 @@ void FGeometryCollectionPhysicsProxy::SetClusteredParticleKinematicTarget(Chaos:
 	{
 		Chaos::TKinematicTarget<Chaos::FReal, 3> NewKinematicTarget;
 		NewKinematicTarget.SetTargetMode(NewWorldTransform);
-		Handle->SetKinematicTarget(NewKinematicTarget);
+
+		if (Chaos::FPhysicsSolver* RBDSolver = GetSolver<Chaos::FPhysicsSolver>())
+		{
+			RBDSolver->GetEvolution()->SetParticleKinematicTarget(Handle, NewKinematicTarget);
+		}
 	}
 }
 void FGeometryCollectionPhysicsProxy::BufferPhysicsResults(Chaos::FPBDRigidsSolver* CurrentSolver, Chaos::FDirtyGeometryCollectionData& BufferData)
