Chaos : Clustering
- Fix static analysis build.

diff --git a/Engine/Source/Runtime/Experimental/Chaos/Private/Chaos/PBDRigidClustering.cpp b/Engine/Source/Runtime/Experimental/Chaos/Private/Chaos/PBDRigidClustering.cpp
--- a/Engine/Source/Runtime/Experimental/Chaos/Private/Chaos/PBDRigidClustering.cpp
+++ b/Engine/Source/Runtime/Experimental/Chaos/Private/Chaos/PBDRigidClustering.cpp
@@ -1053,21 +1053,24 @@ namespace Chaos
 				FRigidHandle CurrentHandle = nullptr;
 				while (Queue.Dequeue(CurrentHandle))
 				{
-					if (FClusterHandle CurrentClusterHandle = CurrentHandle->CastToClustered())
+					if (CurrentHandle)
 					{
-						// @question : Maybe we should just store the leaf node bodies in a
-						// map, that will require Memory(n*log(n))
-						if (MChildren.Contains(CurrentClusterHandle))
+						if (FClusterHandle CurrentClusterHandle = CurrentHandle->CastToClustered())
 						{
-							for (Chaos::FPBDRigidParticleHandle* Child : MChildren[CurrentClusterHandle])
+							// @question : Maybe we should just store the leaf node bodies in a
+							// map, that will require Memory(n*log(n))
+							if (MChildren.Contains(CurrentClusterHandle))
 							{
-								Queue.Enqueue(Child);
+								for (Chaos::FPBDRigidParticleHandle* Child : MChildren[CurrentClusterHandle])
+								{
+									Queue.Enqueue(Child);
+								}
 							}
 						}
-					}
-					if (CurrentHandle)
-					{
-						Function(*this, CurrentHandle);
+						if (CurrentHandle)
+						{
+							Function(*this, CurrentHandle);
+						}
 					}
 				}
 			}
