Headless chaos fix test failure - increase clamp distance

diff --git a/Engine/Source/Runtime/Experimental/Chaos/Private/Chaos/PBDSuspensionConstraints.cpp b/Engine/Source/Runtime/Experimental/Chaos/Private/Chaos/PBDSuspensionConstraints.cpp
--- a/Engine/Source/Runtime/Experimental/Chaos/Private/Chaos/PBDSuspensionConstraints.cpp
+++ b/Engine/Source/Runtime/Experimental/Chaos/Private/Chaos/PBDSuspensionConstraints.cpp
@@ -16,7 +16,7 @@ FAutoConsoleVariableRef CVarChaosSuspensionHardstopEnabled(TEXT("p.Chaos.Suspens
 bool bChaos_Suspension_VelocitySolve = true;
 FAutoConsoleVariableRef CVarChaosSuspensionVelocitySolve(TEXT("p.Chaos.Suspension.VelocitySolve"), bChaos_Suspension_VelocitySolve, TEXT("Enable/Disable VelocitySolve"));
 
-float Chaos_Suspension_MaxPushoutVelocity = 50.f;
+float Chaos_Suspension_MaxPushoutVelocity = 100.f;
 FAutoConsoleVariableRef CVarChaosSuspensionMaxPushoutVelocity(TEXT("p.Chaos.Suspension.MaxPushoutVelocity"), Chaos_Suspension_MaxPushoutVelocity, TEXT("Chaos Suspension Max Pushout Velocity Value"));
 
 float Chaos_Suspension_MaxPushout = 5.f;
