Chaos : Optimization : reduce unecessary update of geometry collection particles causing poor performance even before destruction
- Make sure we remove non dynamic clustered particles from the active particle array to triggering calls to the geoimetry collection physics proxy
- Remove the static geometry collecion array from the dirty particle view

diff --git a/Engine/Source/Runtime/Experimental/Chaos/Public/Chaos/PBDRigidsSOAs.h b/Engine/Source/Runtime/Experimental/Chaos/Public/Chaos/PBDRigidsSOAs.h
--- a/Engine/Source/Runtime/Experimental/Chaos/Public/Chaos/PBDRigidsSOAs.h
+++ b/Engine/Source/Runtime/Experimental/Chaos/Public/Chaos/PBDRigidsSOAs.h
@@ -636,6 +636,16 @@ public:
 	{
 		RemoveClusteredParticle(ClusteredParticle);
 		InsertClusteredParticle(ClusteredParticle);
+
+		if (ClusteredParticle->ObjectState() != EObjectStateType::Dynamic)
+		{
+			RemoveFromActiveArray(ClusteredParticle->CastToRigidParticle(), /*bStillDirty=*/true);
+		}
+		else
+		{
+			AddToActiveArray(ClusteredParticle->CastToRigidParticle());
+		}
+		
 		UpdateViews();
 	}
 
@@ -1064,7 +1074,6 @@ private:
 			{
 				{&ActiveParticlesMapArray.GetArray()},
 				//	{&DynamicClusteredMapArray.GetArray()},  Cluster particles appear in the ActiveParticlesArray
-				{&StaticGeometryCollectionArray},
 				{&KinematicGeometryCollectionArray},
 				{&DynamicGeometryCollectionArray},
 				{&FellAsleepGeometryCollectionArray.GetArray()},
