Fix CCD with fast-moving kinematic objects

diff --git a/Engine/Source/Runtime/Experimental/Chaos/Private/Chaos/CCDUtilities.cpp b/Engine/Source/Runtime/Experimental/Chaos/Private/Chaos/CCDUtilities.cpp
--- a/Engine/Source/Runtime/Experimental/Chaos/Private/Chaos/CCDUtilities.cpp
+++ b/Engine/Source/Runtime/Experimental/Chaos/Private/Chaos/CCDUtilities.cpp
@@ -334,9 +334,9 @@ namespace Chaos
 					{
 						if (CCDConstraint->FastMovingKinematicIndex != INDEX_NONE)
 						{
-							const FConstGenericParticleHandle Particle1 = FGenericParticleHandle(CCDConstraint->SweptConstraint->GetParticle(CCDConstraint->FastMovingKinematicIndex));
+							const FConstGenericParticleHandle KinematicParticle = FGenericParticleHandle(CCDConstraint->SweptConstraint->GetParticle(CCDConstraint->FastMovingKinematicIndex));
 							const FVec3 Normal = CCDConstraint->SweptConstraint->CalculateWorldContactNormal();
-							const FVec3 Offset = FVec3::DotProduct(Particle1->V() * ((1.f - IslandTOI) * Dt), Normal) * Normal;
+							const FVec3 Offset = FVec3::DotProduct(KinematicParticle->V() * ((1.f - IslandTOI) * Dt), Normal) * Normal;
 							ClipParticleP(CCDConstraintParticles[0], Offset);
 						}
 						else
@@ -347,7 +347,17 @@ namespace Chaos
 					}
 					if (CCDConstraintParticles[1])
 					{
-						ClipParticleP(CCDConstraintParticles[1]);
+						if (CCDConstraint->FastMovingKinematicIndex != INDEX_NONE)
+						{
+							const FConstGenericParticleHandle KinematicParticle = FGenericParticleHandle(CCDConstraint->SweptConstraint->GetParticle(CCDConstraint->FastMovingKinematicIndex));
+							const FVec3 Normal = CCDConstraint->SweptConstraint->CalculateWorldContactNormal();
+							const FVec3 Offset = FVec3::DotProduct(KinematicParticle->V() * ((1.f - IslandTOI) * Dt), Normal) * Normal;
+							ClipParticleP(CCDConstraintParticles[1], Offset);
+						}
+						else
+						{
+							ClipParticleP(CCDConstraintParticles[1]);
+						}
 						CCDConstraintParticles[1]->Done = true;
 					}
 				}
diff --git a/Engine/Source/Runtime/Experimental/Chaos/Private/Chaos/Collision/ParticlePairMidPhase.cpp b/Engine/Source/Runtime/Experimental/Chaos/Private/Chaos/Collision/ParticlePairMidPhase.cpp
--- a/Engine/Source/Runtime/Experimental/Chaos/Private/Chaos/Collision/ParticlePairMidPhase.cpp
+++ b/Engine/Source/Runtime/Experimental/Chaos/Private/Chaos/Collision/ParticlePairMidPhase.cpp
@@ -361,13 +361,16 @@ namespace Chaos
 			Constraint->ResetManifold();
 			Constraint->ResetActiveManifoldContacts();
 
+			FConstGenericParticleHandle P0 = Particle0;
+			FConstGenericParticleHandle P1 = Particle1;
+			// For kinematic particles, X = P (at TOI=1), we need to compute P-V*dt to get position at TOI=0. 
+			const FVec3 StartX0 = P0->ObjectState() == EObjectStateType::Kinematic ? P0->P() - P0->V() * Dt : P0->X();
+			const FVec3 StartX1 = P1->ObjectState() == EObjectStateType::Kinematic ? P1->P() - P1->V() * Dt : P1->X();
 			// Note: It is unusual that we are mixing X and Q. 
 			// This is due to how CCD rewinds the position (not rotation) and then sweeps to find the first contact at the most recent orientation Q
 			// NOTE: These are actor transforms, not CoM transforms
-			FConstGenericParticleHandle P0 = Particle0;
-			FConstGenericParticleHandle P1 = Particle1;
-			const FRigidTransform3 CCDParticleWorldTransform0 = FRigidTransform3(P0->CCDEnabled() ? P0->X() : P0->P(), P0->Q());
-			const FRigidTransform3 CCDParticleWorldTransform1 = FRigidTransform3(P1->CCDEnabled() ? P1->X() : P1->P(), P1->Q());
+			const FRigidTransform3 CCDParticleWorldTransform0 = FRigidTransform3(StartX0, P0->Q());
+			const FRigidTransform3 CCDParticleWorldTransform1 = FRigidTransform3(StartX1, P1->Q());
 			const FRigidTransform3 CCDShapeWorldTransform0 = Constraint->ImplicitTransform[0] * CCDParticleWorldTransform0;
 			const FRigidTransform3 CCDShapeWorldTransform1 = Constraint->ImplicitTransform[1] * CCDParticleWorldTransform1;
 			Collisions::UpdateConstraintSwept(*Constraint.Get(), CCDShapeWorldTransform0, CCDShapeWorldTransform1, Dt);
