Chaos : fix crash in clustering when particle does not have geometry

diff --git a/Engine/Source/Runtime/Experimental/Chaos/Private/Chaos/PBDRigidClustering.cpp b/Engine/Source/Runtime/Experimental/Chaos/Private/Chaos/PBDRigidClustering.cpp
--- a/Engine/Source/Runtime/Experimental/Chaos/Private/Chaos/PBDRigidClustering.cpp
+++ b/Engine/Source/Runtime/Experimental/Chaos/Private/Chaos/PBDRigidClustering.cpp
@@ -463,7 +463,7 @@ namespace Chaos
 						ClusterBreak.Velocity = Child->V();
 						ClusterBreak.AngularVelocity = Child->W();
 						ClusterBreak.Mass = Child->M();
-						if (Child->Geometry()->HasBoundingBox())
+						if (Child->Geometry() && Child->Geometry()->HasBoundingBox())
 						{
 							ClusterBreak.BoundingBox = Child->Geometry()->BoundingBox();
 						}
