Linear solver projection

diff --git a/Engine/Source/Runtime/Engine/Private/PhysicsEngine/ImmediatePhysics/ImmediatePhysicsChaos/ImmediatePhysicsSimulation_Chaos.cpp b/Engine/Source/Runtime/Engine/Private/PhysicsEngine/ImmediatePhysics/ImmediatePhysicsChaos/ImmediatePhysicsSimulation_Chaos.cpp
--- a/Engine/Source/Runtime/Engine/Private/PhysicsEngine/ImmediatePhysics/ImmediatePhysicsChaos/ImmediatePhysicsSimulation_Chaos.cpp
+++ b/Engine/Source/Runtime/Engine/Private/PhysicsEngine/ImmediatePhysics/ImmediatePhysicsChaos/ImmediatePhysicsSimulation_Chaos.cpp
@@ -152,7 +152,7 @@ int32 ChaosImmediate_SolverType = (int32)Chaos::EConstraintSolverType::QuasiPbd;
 FAutoConsoleVariableRef CVarChaosImmPhysSolverType(TEXT("p.Chaos.ImmPhys.SolverType"), ChaosImmediate_SolverType, TEXT("0 = None; 1 = GbfPbd; 2 = Pbd; 3 = QuasiPbd"));
 
 // Whether to use the linear joint solver which is significantly faster than the non-linear one but less accurate. Only applies to the QuasiPBD Solver
-bool bChaosImmediate_Joint_UseLinearSolver = false;
+bool bChaosImmediate_Joint_UseLinearSolver = true;
 FAutoConsoleVariableRef CVarChaosImmPhysJointUseCachedSolver(TEXT("p.Chaos.ImmPhys.Joint.UseLinearSolver"), bChaosImmediate_Joint_UseLinearSolver, TEXT("Use linear version of joint solver. (default is true"));
 
 //
diff --git a/Engine/Source/Runtime/Experimental/Chaos/Private/Chaos/Joint/PBDJointCachedSolverGaussSeidel.cpp b/Engine/Source/Runtime/Experimental/Chaos/Private/Chaos/Joint/PBDJointCachedSolverGaussSeidel.cpp
--- a/Engine/Source/Runtime/Experimental/Chaos/Private/Chaos/Joint/PBDJointCachedSolverGaussSeidel.cpp
+++ b/Engine/Source/Runtime/Experimental/Chaos/Private/Chaos/Joint/PBDJointCachedSolverGaussSeidel.cpp
@@ -227,6 +227,46 @@ void FPBDJointCachedSolver::Init(
 	LastDQs[0] = FVec3(0.f);
 	LastDQs[1] = FVec3(0.f);
 }
+	
+void FPBDJointCachedSolver::InitProjection(
+	const FReal Dt,
+	const FPBDJointSolverSettings& SolverSettings,
+	const FPBDJointSettings& JointSettings)
+{
+	bool bHasLinearProjection = false;
+	{
+		const FReal LinearProjection = FPBDJointUtilities::GetLinearProjection(SolverSettings, JointSettings);
+		bHasLinearProjection = JointSettings.bProjectionEnabled && LinearProjection > 0.0;
+	}
+	bool bHasAngularProjection = false;
+	{
+		const FReal AngularProjection = FPBDJointUtilities::GetAngularProjection(SolverSettings, JointSettings);
+		bHasAngularProjection = JointSettings.bProjectionEnabled && AngularProjection > 0.0;
+	}
+	if(bHasLinearProjection || bHasAngularProjection)
+	{
+		ComputeBodyState(0);
+		ComputeBodyState(1);
+
+		ConnectorRs[1].EnforceShortestArcWith(ConnectorRs[0]);
+
+		InvMScales[0] = 0.0;
+		InvMScales[1] = 1.0;
+
+		UpdateMass0();
+		UpdateMass1();
+
+		if(bHasLinearProjection)
+		{
+			InitPositionConstraints(Dt, SolverSettings, JointSettings);
+		}
+
+		if(bHasAngularProjection)
+		{
+			InitRotationConstraints(Dt, SolverSettings, JointSettings);
+		}
+	}
+}
 
 void FPBDJointCachedSolver::Deinit()
 {
@@ -352,8 +392,8 @@ void FPBDJointCachedSolver::InitPositionConstraints(
 	const TVec3<EJointMotionType>& LinearMotion = JointSettings.LinearMotionTypes;
 
 
-	const FVec3 ConstraintArm0Limited = ConnectorXs[1] - P(0);
-	const FVec3 ConstraintArm1Limited = ConnectorXs[1] - P(1);
+	const FVec3 ConstraintArm0Limited = ConnectorXs[1] - CurrentPs[0];
+	const FVec3 ConstraintArm1Limited = ConnectorXs[1] - CurrentPs[1];
 	
 	FVec3 ConstraintArm0Locked = ConstraintArm0Limited;
 	FVec3 ConstraintArm1Locked = ConstraintArm1Limited;
@@ -804,13 +844,12 @@ void FPBDJointCachedSolver::InitTwistConstraint(
 		const FPBDJointSettings& JointSettings,
 		const FReal Dt)
 {
-	FVec3 TwistAxis;
-	FReal TwistAngle;
-	FPBDJointUtilities::GetTwistAxisAngle(ConnectorRs[0], ConnectorRs[1], TwistAxis, TwistAngle);
+	 FVec3 TwistAxis;
+	 FReal TwistAngle;
+	 FPBDJointUtilities::GetTwistAxisAngle(ConnectorRs[0], ConnectorRs[1], TwistAxis, TwistAngle);
 	
-	// Project the angle directly to avoid checking the limits during the solve.
-	CorrectAxisAngleConstraint(JointSettings,(int32)EJointAngularConstraintIndex::Twist, TwistAxis, TwistAngle);
-	InitRotationConstraintDatas( JointSettings, (int32)EJointAngularConstraintIndex::Twist, TwistAxis, TwistAngle, JointSettings.TwistRestitution, Dt, false);
+	 // Project the angle directly to avoid checking the limits during the solve.
+	 InitRotationConstraintDatas( JointSettings, (int32)EJointAngularConstraintIndex::Twist, TwistAxis, TwistAngle, JointSettings.TwistRestitution, Dt, true);
 }
 
 void FPBDJointCachedSolver::InitConeConstraint(
@@ -1117,8 +1156,8 @@ void FPBDJointCachedSolver::InitAxisPositionDrive(
 		const FVec3& DeltaVelocity,
 		const FReal Dt)
 {	
-	const FVec3 ConstraintArm0 = ConnectorXs[0] - P(0);
-	const FVec3 ConstraintArm1 = ConnectorXs[1] - P(1);
+	const FVec3 ConstraintArm0 = ConnectorXs[0] - CurrentPs[0];
+	const FVec3 ConstraintArm1 = ConnectorXs[1] - CurrentPs[1];
 
 	PositionDrives.UpdateDatas(ConstraintIndex, ConstraintAxis, FVec3::DotProduct(DeltaPosition, ConstraintAxis),
 		0.0f,  true, ConstraintArm0, ConstraintArm1,
@@ -1126,6 +1165,166 @@ void FPBDJointCachedSolver::InitAxisPositionDrive(
 
 	InitPositionDatasMass(PositionDrives, ConstraintIndex, Dt);
 }
+/** APPLY POSITION PROJECTIONS *********************************************************************************/
+	
+void FPBDJointCachedSolver::ApplyProjections(
+		const FReal Dt,
+		const FReal InSolverStiffness,
+		const FPBDJointSolverSettings& SolverSettings,
+		const FPBDJointSettings& JointSettings,
+		const bool bLastIteration)
+{
+	if (!IsDynamic(1))
+	{
+		// If child is kinematic, return. 
+		return;
+	}
+
+	SolverStiffness = InSolverStiffness;
+	
+
+	if (SolverSettings.bSolvePositionLast)
+	{
+		ApplyRotationProjection(Dt, SolverSettings, JointSettings);
+		ApplyPositionProjection(Dt, SolverSettings, JointSettings);
+	}
+	else
+	{
+		ApplyPositionProjection(Dt, SolverSettings, JointSettings);
+		ApplyRotationProjection(Dt, SolverSettings, JointSettings);
+	}
+	if(bLastIteration)
+	{
+		//Final position fixup
+		const TVec3<EJointMotionType>& LinearMotion = JointSettings.LinearMotionTypes;
+		const bool bLinearLocked = (LinearMotion[0] == EJointMotionType::Locked) && (LinearMotion[1] == EJointMotionType::Locked) && (LinearMotion[2] == EJointMotionType::Locked);
+		if (bLinearLocked)
+		{
+			const FReal LinearProjection = FPBDJointUtilities::GetLinearProjection(SolverSettings, JointSettings);
+			if (JointSettings.bProjectionEnabled && (LinearProjection > 0))
+			{
+				const FVec3 DP1 = -LinearProjection * (ConnectorXs[1] - ConnectorXs[0]);
+				ApplyPositionDelta(1, DP1);
+			}
+		
+			// Add velocity correction from the net projection motion
+			if (Chaos_Joint_VelProjectionAlpha > 0.0f)
+			{
+				const FSolverReal VelocityScale = Chaos_Joint_VelProjectionAlpha  / static_cast<FSolverReal>(Dt);
+				const FSolverVec3 DV1 = Body1().DP() * VelocityScale;
+				const FSolverVec3 DW1 = Body1().DQ()* VelocityScale;
+		
+				Body(1).ApplyVelocityDelta(DV1, DW1);
+			}
+		}
+	}
+}
+
+void FPBDJointCachedSolver::ApplyRotationProjection(
+			const FReal Dt,
+			const FPBDJointSolverSettings& SolverSettings,
+			const FPBDJointSettings& JointSettings)
+{
+	const FReal AngularProjection = FPBDJointUtilities::GetAngularProjection(SolverSettings, JointSettings);
+	
+	const TVec3<EJointMotionType>& LinearMotion = JointSettings.LinearMotionTypes;
+	const bool bLinearLocked = (LinearMotion[0] == EJointMotionType::Locked) && (LinearMotion[1] == EJointMotionType::Locked) && (LinearMotion[2] == EJointMotionType::Locked);
+
+	for(int32 ConstraintIndex = 0; ConstraintIndex < 3; ++ConstraintIndex)
+	{
+		if(RotationConstraints.bValidDatas[ConstraintIndex])
+		{
+			if (JointSettings.bProjectionEnabled && (AngularProjection > 0.0f))
+			{
+				FReal DeltaAngle = RotationConstraints.ConstraintCX[ConstraintIndex] +
+					FVec3::DotProduct(Body(1).DQ() - Body(0).DQ(), RotationConstraints.ConstraintAxis[ConstraintIndex]);
+	
+				bool NeedsSolve = false;
+				if(RotationConstraints.bLimitsCheck[ConstraintIndex])
+				{
+					if(DeltaAngle > RotationConstraints.ConstraintLimits[ConstraintIndex] )
+					{
+						DeltaAngle -= RotationConstraints.ConstraintLimits[ConstraintIndex];
+						NeedsSolve = true;
+					}
+					else if(DeltaAngle < -RotationConstraints.ConstraintLimits[ConstraintIndex])
+					{
+						DeltaAngle += RotationConstraints.ConstraintLimits[ConstraintIndex];
+						NeedsSolve = true;
+			
+					}
+				}
+
+				if (!RotationConstraints.bLimitsCheck[ConstraintIndex] || (RotationConstraints.bLimitsCheck[ConstraintIndex] && NeedsSolve && FMath::Abs(DeltaAngle) > AngleTolerance))
+				{
+					const FReal IM = -FVec3::DotProduct(RotationConstraints.ConstraintAxis[ConstraintIndex], RotationConstraints.ConstraintDRAxis[ConstraintIndex][1]);
+					const FReal DeltaLambda = SolverStiffness * RotationConstraints.ConstraintHardStiffness[ConstraintIndex] * DeltaAngle / IM;
+					
+					const FVec3 DR1 = AngularProjection * RotationConstraints.ConstraintDRAxis[ConstraintIndex][1] * DeltaLambda;
+					ApplyRotationDelta(1, DR1);
+
+					if(bLinearLocked)
+					{
+						const FVec3 DP1 = -AngularProjection * FVec3::CrossProduct(DR1, PositionConstraints.ConstraintArms[ConstraintIndex][1]);
+						ApplyPositionDelta(1,DP1);
+					}
+				}
+			}
+		}
+	}
+}
+
+void FPBDJointCachedSolver::ApplyPositionProjection(
+		const FReal Dt,
+		const FPBDJointSolverSettings& SolverSettings,
+		const FPBDJointSettings& JointSettings)
+{
+	const FReal LinearProjection = FPBDJointUtilities::GetLinearProjection(SolverSettings, JointSettings);
+
+	for(int32 ConstraintIndex = 0; ConstraintIndex < 3; ++ConstraintIndex)
+	{
+		if(PositionConstraints.bValidDatas[ConstraintIndex])
+		{
+			if (JointSettings.bProjectionEnabled && (LinearProjection > 0.0f))
+			{
+				const FVec3 CX = Body(1).DP()  - Body(0).DP() +
+					FVec3::CrossProduct(Body(1).DQ(), PositionConstraints.ConstraintArms[ConstraintIndex][1]) -
+					FVec3::CrossProduct(Body(0).DQ(), PositionConstraints.ConstraintArms[ConstraintIndex][0]);
+				
+				FReal DeltaPosition = PositionConstraints.ConstraintCX[ConstraintIndex] + FVec3::DotProduct(CX, PositionConstraints.ConstraintAxis[ConstraintIndex]);
+				
+				bool NeedsSolve = false;
+				if(PositionConstraints.bLimitsCheck[ConstraintIndex])
+				{
+					if(DeltaPosition > PositionConstraints.ConstraintLimits[ConstraintIndex] )
+					{
+						DeltaPosition -= PositionConstraints.ConstraintLimits[ConstraintIndex];
+						NeedsSolve = true;
+					}
+					else if(DeltaPosition < -PositionConstraints.ConstraintLimits[ConstraintIndex])
+					{
+						DeltaPosition += PositionConstraints.ConstraintLimits[ConstraintIndex];
+						NeedsSolve = true;
+					}
+				}
+				if (!PositionConstraints.bLimitsCheck[ConstraintIndex] || (PositionConstraints.bLimitsCheck[ConstraintIndex] && NeedsSolve && FMath::Abs(DeltaPosition ) > PositionTolerance))
+				{
+					const FVec3 AngularAxis1 = FVec3::CrossProduct(PositionConstraints.ConstraintArms[ConstraintIndex][1], PositionConstraints.ConstraintAxis[ConstraintIndex]);
+					const FReal IM = InvM(1) - FVec3::DotProduct(AngularAxis1, PositionConstraints.ConstraintDRAxis[ConstraintIndex][1]);
+					const FReal DeltaLambda = SolverStiffness * PositionConstraints.ConstraintHardStiffness[ConstraintIndex] * DeltaPosition / IM;
+	
+					const FVec3 DX = PositionConstraints.ConstraintAxis[ConstraintIndex] * DeltaLambda;
+					
+					const FVec3 DP1 = -LinearProjection * InvM(1) * DX;
+					const FVec3 DR1 = LinearProjection * PositionConstraints.ConstraintDRAxis[ConstraintIndex][1] * DeltaLambda;
+					
+					ApplyPositionDelta(1,DP1);
+					ApplyRotationDelta(1,DR1);
+				}
+			}
+		}
+	}
+}
 
 /** APPLY POSITION  DRIVES *********************************************************************************/
 
diff --git a/Engine/Source/Runtime/Experimental/Chaos/Private/Chaos/PBDJointConstraints.cpp b/Engine/Source/Runtime/Experimental/Chaos/Private/Chaos/PBDJointConstraints.cpp
--- a/Engine/Source/Runtime/Experimental/Chaos/Private/Chaos/PBDJointConstraints.cpp
+++ b/Engine/Source/Runtime/Experimental/Chaos/Private/Chaos/PBDJointConstraints.cpp
@@ -1024,12 +1024,39 @@ namespace Chaos
 		}
 		return (NumActive > 0);
 	}
+	
+	void FPBDJointConstraints::PreparePhase3Serial(const FReal Dt, const int32 It, FPBDIslandSolverData& SolverData)
+	{
+		CSV_SCOPED_TIMING_STAT(Chaos, ProjectJointConstraints);
+		SCOPE_CYCLE_COUNTER(STAT_Joints_ApplyProjection);
+
+		if(It == 0 && Settings.bUseLinearSolver)
+		{
+			for (int32 ConstraintIndex : SolverData.GetConstraintIndices(ContainerId))
+			{
+				const FPBDJointSettings& JointSettings = ConstraintSettings[ConstraintIndex];
+				FPBDJointCachedSolver& Solver = CachedConstraintSolvers[ConstraintIndex];
+				
+				if ((FMath::IsNearlyZero(Solver.InvM(0)) && FMath::IsNearlyZero(Solver.InvM(1))))
+				{
+					return;
+				}
+				Solver.InitProjection(
+					Dt,
+					Settings,
+					JointSettings);
+			}
+		}
+	}
 
 	bool FPBDJointConstraints::ApplyPhase3Serial(const FReal Dt, const int32 It, const int32 NumIts, FPBDIslandSolverData& SolverData)
 	{
 		CSV_SCOPED_TIMING_STAT(Chaos, ProjectJointConstraints);
 		SCOPE_CYCLE_COUNTER(STAT_Joints_ApplyProjection);
 
+		// Prepare phase 3 for the linear solver in order to partially re-init the solver
+		PreparePhase3Serial(Dt, It, SolverData);
+
 		// UpdateConstraintProjection(It, NumIts, SolverData);
 		for (int32 ConstraintIndex : SolverData.GetConstraintIndices(ContainerId))
 		{
@@ -1039,7 +1066,6 @@ namespace Chaos
 		return true;
 	}
 
-
 	void FPBDJointConstraints::UpdateConstraintProjection(const int32 It, const int32 NumIts, const FPBDIslandSolverData& SolverData)
 	{
 		const bool bUseProjection = It == NumIts - 1;
@@ -1363,12 +1389,8 @@ namespace Chaos
 				return false;
 			}
 
-			// @todo(chaos): rebuild the joint state at the latest (post-solve) transforms
-			//Solver.Update(Dt, Settings, JointSettings);
-			//Solver.UpdateMasses(FReal(0), FReal(1));
-
 			const FReal IterationStiffness = CalculateIterationStiffness(It, NumIts);
-			Solver.ApplyProjections(Dt, IterationStiffness, Settings, JointSettings);
+			Solver.ApplyProjections(Dt, IterationStiffness, Settings, JointSettings, (It==(NumIts-1)));
 		}
 		else
 		{
diff --git a/Engine/Source/Runtime/Experimental/Chaos/Public/Chaos/Joint/PBDJointCachedSolverGaussSeidel.h b/Engine/Source/Runtime/Experimental/Chaos/Public/Chaos/Joint/PBDJointCachedSolverGaussSeidel.h
--- a/Engine/Source/Runtime/Experimental/Chaos/Public/Chaos/Joint/PBDJointCachedSolverGaussSeidel.h
+++ b/Engine/Source/Runtime/Experimental/Chaos/Public/Chaos/Joint/PBDJointCachedSolverGaussSeidel.h
@@ -201,6 +201,11 @@ struct FAxisConstraintDatas
 			const FRigidTransform3& XL0,
 			const FRigidTransform3& XL1);
 
+		void InitProjection(
+			const FReal Dt,
+			const FPBDJointSolverSettings& SolverSettings,
+			const FPBDJointSettings& JointSettings);
+
 		void Deinit();
 
 		// @todo(chaos): this doesn't do much now we have SolverBodies - remove it
@@ -229,7 +234,18 @@ struct FAxisConstraintDatas
 			const FReal Dt,
 			const FReal InSolverStiffness,
 			const FPBDJointSolverSettings& SolverSettings,
-			const FPBDJointSettings& JointSettings){}
+			const FPBDJointSettings& JointSettings,
+			const bool bLastIteration);
+
+		void ApplyPositionProjection(
+			const FReal Dt,
+			const FPBDJointSolverSettings& SolverSettings,
+			const FPBDJointSettings& JointSettings);
+
+		void ApplyRotationProjection(
+			const FReal Dt,
+			const FPBDJointSolverSettings& SolverSettings,
+			const FPBDJointSettings& JointSettings);
 
 		void EnableProjection();
 
diff --git a/Engine/Source/Runtime/Experimental/Chaos/Public/Chaos/PBDJointConstraints.h b/Engine/Source/Runtime/Experimental/Chaos/Public/Chaos/PBDJointConstraints.h
--- a/Engine/Source/Runtime/Experimental/Chaos/Public/Chaos/PBDJointConstraints.h
+++ b/Engine/Source/Runtime/Experimental/Chaos/Public/Chaos/PBDJointConstraints.h
@@ -279,6 +279,7 @@ namespace Chaos
 		bool ApplyPhase1Serial(const FReal Dt, const int32 It, const int32 NumIts, FPBDIslandSolverData& SolverData);
 		bool ApplyPhase2Serial(const FReal Dt, const int32 It, const int32 NumIts, FPBDIslandSolverData& SolverData);
 		bool ApplyPhase3Serial(const FReal Dt, const int32 It, const int32 NumIts, FPBDIslandSolverData& SolverData);
+		void PreparePhase3Serial(const FReal Dt, const int32 It, FPBDIslandSolverData& SolverData);
 
 		/**
 		 * Set the solver method to use
