Parallelize scatter and color solve. Saves about a factor of 2-3 in scatter. Color solve is disabled by default.

diff --git a/Engine/Source/Runtime/Experimental/Chaos/Private/Chaos/Collision/SolverCollisionContainer.cpp b/Engine/Source/Runtime/Experimental/Chaos/Private/Chaos/Collision/SolverCollisionContainer.cpp
--- a/Engine/Source/Runtime/Experimental/Chaos/Private/Chaos/Collision/SolverCollisionContainer.cpp
+++ b/Engine/Source/Runtime/Experimental/Chaos/Private/Chaos/Collision/SolverCollisionContainer.cpp
@@ -520,11 +520,11 @@ namespace Chaos
 		}
 		else if (bApplyStaticFriction)
 		{
-			return SolvePositionWithFrictionImpl(Dt, BeginIndex, EndIndex, MaxPushOut);
+			return SolvePositionWithFrictionImpl(Dt, BeginIndex, EndIndex, MaxPushOut, bParallel);
 		}
 		else
 		{
-			return SolvePositionNoFrictionImpl(Dt, BeginIndex, EndIndex, MaxPushOut);
+			return SolvePositionNoFrictionImpl(Dt, BeginIndex, EndIndex, MaxPushOut, bParallel);
 		}
 	}
 
@@ -557,31 +557,47 @@ namespace Chaos
 	}
 
 	// Solve position with friction (last few iterations each tick)
-	bool FPBDCollisionSolverContainer::SolvePositionWithFrictionImpl(const FReal InDt, const int32 BeginIndex, const int32 EndIndex, const FReal InMaxPushOut)
+	bool FPBDCollisionSolverContainer::SolvePositionWithFrictionImpl(const FReal InDt, const int32 BeginIndex, const int32 EndIndex, const FReal InMaxPushOut, const bool bParallel)
 	{
+		if (EndIndex == BeginIndex)
+		{
+			return false;
+		}
 		const FSolverReal Dt = FSolverReal(InDt);
 		const FSolverReal MaxPushOut = FSolverReal(InMaxPushOut);
 
-		bool bNeedsAnotherIteration = false;
-		for (int32 SolverIndex = BeginIndex; SolverIndex < EndIndex; ++SolverIndex)
+		InnerPhysicsParallelForRange(EndIndex - BeginIndex, [&](int32 StartRangeIndex, int32 EndRangeIndex)
 		{
-			bNeedsAnotherIteration |= CollisionSolvers[SolverIndex].GetSolver().SolvePositionWithFriction(Dt, MaxPushOut);
-		}
-		return bNeedsAnotherIteration;
+			int32 LocalStartIndex = StartRangeIndex + BeginIndex;
+			int32 LocalEndIndex = EndRangeIndex + BeginIndex;
+			for (int32 SolverIndex = LocalStartIndex; SolverIndex < LocalEndIndex; ++SolverIndex)
+			{
+				CollisionSolvers[SolverIndex].GetSolver().SolvePositionWithFriction(Dt, MaxPushOut);
+			}
+		}, 100, !bParallel);
+		return true;
 	}
 
 	// Solve position without friction (first few iterations each tick)
-	bool FPBDCollisionSolverContainer::SolvePositionNoFrictionImpl(const FReal InDt, const int32 BeginIndex, const int32 EndIndex, const FReal InMaxPushOut)
+	bool FPBDCollisionSolverContainer::SolvePositionNoFrictionImpl(const FReal InDt, const int32 BeginIndex, const int32 EndIndex, const FReal InMaxPushOut, const bool bParallel)
 	{
+		if (EndIndex == BeginIndex)
+		{
+			return false;
+		}
 		const FSolverReal Dt = FSolverReal(InDt);
 		const FSolverReal MaxPushOut = FSolverReal(InMaxPushOut);
 
-		bool bNeedsAnotherIteration = false;
-		for (int32 SolverIndex = BeginIndex; SolverIndex < EndIndex; ++SolverIndex)
+		InnerPhysicsParallelForRange(EndIndex - BeginIndex, [&](int32 StartRangeIndex, int32 EndRangeIndex)
 		{
-			bNeedsAnotherIteration |= CollisionSolvers[SolverIndex].GetSolver().SolvePositionNoFriction(Dt, MaxPushOut);
-		}
-		return bNeedsAnotherIteration;
+			int32 LocalStartIndex = StartRangeIndex + BeginIndex;
+			int32 LocalEndIndex = EndRangeIndex + BeginIndex;
+			for (int32 SolverIndex = LocalStartIndex; SolverIndex < LocalEndIndex; ++SolverIndex)
+			{
+				CollisionSolvers[SolverIndex].GetSolver().SolvePositionNoFriction(Dt, MaxPushOut);
+			}
+		}, 100, !bParallel);
+		return true;
 	}
 
 	bool FPBDCollisionSolverContainer::SolveVelocityImpl(const FReal InDt, const int32 It, const int32 NumIts, const int32 BeginIndex, const int32 EndIndex, const FPBDCollisionSolverSettings& SolverSettings, const bool bParallel)
@@ -620,11 +636,20 @@ namespace Chaos
 		SCOPE_CYCLE_COUNTER(STAT_Collisions_Scatter);
 		check(BeginIndex >= 0);
 		check(EndIndex <= CollisionSolvers.Num());
-
-		for (int32 SolverIndex = BeginIndex; SolverIndex < EndIndex; ++SolverIndex)
+		if (EndIndex == BeginIndex)
 		{
-			CollisionSolvers[SolverIndex].ScatterOutput(Dt);
+			return;
 		}
+
+		InnerPhysicsParallelForRange(EndIndex - BeginIndex, [&](int32 LocalStartIndex, int32 LocalEndIndex)
+		{
+			int32 LoopStartIndex = LocalStartIndex + BeginIndex;
+			int32 LoopEndIndex = LocalEndIndex + BeginIndex;
+			for (int32 SolverIndex = LoopStartIndex; SolverIndex < LoopEndIndex; ++SolverIndex)
+			{
+				CollisionSolvers[SolverIndex].ScatterOutput(Dt);
+			}
+		}, 100);
 	}
 
 }
diff --git a/Engine/Source/Runtime/Experimental/Chaos/Private/Chaos/Framework/Parallel.cpp b/Engine/Source/Runtime/Experimental/Chaos/Private/Chaos/Framework/Parallel.cpp
--- a/Engine/Source/Runtime/Experimental/Chaos/Private/Chaos/Framework/Parallel.cpp
+++ b/Engine/Source/Runtime/Experimental/Chaos/Private/Chaos/Framework/Parallel.cpp
@@ -10,8 +10,8 @@ namespace Chaos
 {
 	CHAOS_API int32 GSingleThreadedPhysics = 0;
 	CHAOS_API int32 InnerParallelForBatchSize = 0;
-	CHAOS_API int32 MinRangeBatchSize = 10;
-	CHAOS_API int32 MaxRangeBatchWorkers = 100;
+	CHAOS_API int32 MinRangeBatchSize = 0;
+	CHAOS_API int32 MaxNumWorkers = 100;
 	
 #if !UE_BUILD_SHIPPING
 	CHAOS_API bool bDisablePhysicsParallelFor = false;
@@ -23,7 +23,7 @@ namespace Chaos
 	FAutoConsoleVariableRef CVarDisableCollisionParallelFor(TEXT("p.Chaos.DisableCollisionParallelFor"), bDisableCollisionParallelFor, TEXT("Disable parallel execution for Chaos Collisions (also disabled by DisableParticleParallelFor)"));
 	FAutoConsoleVariableRef CVarInnerPhysicsBatchSize(TEXT("p.Chaos.InnerParallelForBatchSize"), InnerParallelForBatchSize, TEXT("Set the batch size threshold for inner parallel fors"));
 	FAutoConsoleVariableRef CVarMinRangeBatchSize(TEXT("p.Chaos.MinRangeBatchSize"), MinRangeBatchSize, TEXT("Set the min range batch size for parallel for"));
-	FAutoConsoleVariableRef CVarMaxRangeBatchWorkers(TEXT("p.Chaos.MaxRangeBatchWorkers"), MaxRangeBatchWorkers, TEXT("Set the max range batch num workers for parallel for"));
+	FAutoConsoleVariableRef CVarMaxRangeBatchWorkers(TEXT("p.Chaos.MaxNumWorkers"), MaxNumWorkers, TEXT("Set the max number of workers for physics"));
 #endif
 }
 
@@ -34,11 +34,11 @@ void Chaos::InnerPhysicsParallelFor(int32 Num, TFunctionRef<void(int32)> InCalla
 	PhysicsParallelFor(Num, InCallable, (BatchSize > InnerParallelForBatchSize) ? bForceSingleThreaded : true);
 }
 
-void Chaos::InnerPhysicsParallelForRange(int32 InNum, TFunctionRef<void(int32, int32)> InCallable, bool bForceSingleThreaded)
+void Chaos::InnerPhysicsParallelForRange(int32 InNum, TFunctionRef<void(int32, int32)> InCallable, const int32 InMinBatchSize, bool bForceSingleThreaded)
 {
 	int32 NumWorkers = int32(LowLevelTasks::FScheduler::Get().GetNumWorkers());
 	int32 BatchSize = FMath::DivideAndRoundUp<int32>(InNum, NumWorkers);
-	PhysicsParallelForRange(InNum, InCallable, (BatchSize > InnerParallelForBatchSize) ? bForceSingleThreaded : true);
+	PhysicsParallelForRange(InNum, InCallable, InMinBatchSize, (BatchSize > InnerParallelForBatchSize) ? bForceSingleThreaded : true);
 }
 
 void Chaos::PhysicsParallelFor(int32 InNum, TFunctionRef<void(int32)> InCallable, bool bForceSingleThreaded)
@@ -66,7 +66,7 @@ void Chaos::PhysicsParallelFor(int32 InNum, TFunctionRef<void(int32)> InCallable
 	::ParallelFor(InNum, PassThrough, !!GSingleThreadedPhysics || bDisablePhysicsParallelFor || bForceSingleThreaded);
 }
 
-void Chaos::PhysicsParallelForRange(int32 InNum, TFunctionRef<void(int32, int32)> InCallable, bool bForceSingleThreaded)
+void Chaos::PhysicsParallelForRange(int32 InNum, TFunctionRef<void(int32, int32)> InCallable, const int32 InMinBatchSize, bool bForceSingleThreaded)
 {
 	TRACE_CPUPROFILER_EVENT_SCOPE(Chaos_PhysicsParallelFor);
 	using namespace Chaos;
@@ -87,18 +87,19 @@ void Chaos::PhysicsParallelForRange(int32 InNum, TFunctionRef<void(int32, int32)
 		NumWorkers++; //named threads help with the work
 	}
 	NumWorkers = FMath::Min(NumWorkers, InNum);
-	NumWorkers = FMath::Min(NumWorkers, MaxRangeBatchWorkers);
+	NumWorkers = FMath::Min(NumWorkers, MaxNumWorkers);
 	check(NumWorkers > 0);
 	int32 BatchSize = FMath::DivideAndRoundUp<int32>(InNum, NumWorkers);
+	int32 MinBatchSize = FMath::Max(InMinBatchSize, MinRangeBatchSize);
 	// @todo(mlentine): Find a better batch size in this case
-	if (InNum < MinRangeBatchSize)
+	if (InNum < MinBatchSize)
 	{
 		NumWorkers = 1;
 		BatchSize = InNum;
 	}
 	else
 	{
-		while (BatchSize < MinRangeBatchSize && NumWorkers > 1)
+		while (BatchSize < MinBatchSize && NumWorkers > 1)
 		{
 			NumWorkers /= 2;
 			BatchSize = FMath::DivideAndRoundUp<int32>(InNum, NumWorkers);
diff --git a/Engine/Source/Runtime/Experimental/Chaos/Private/Chaos/PBDConstraintRule.cpp b/Engine/Source/Runtime/Experimental/Chaos/Private/Chaos/PBDConstraintRule.cpp
--- a/Engine/Source/Runtime/Experimental/Chaos/Private/Chaos/PBDConstraintRule.cpp
+++ b/Engine/Source/Runtime/Experimental/Chaos/Private/Chaos/PBDConstraintRule.cpp
@@ -335,7 +335,7 @@ namespace Chaos
 													Constraint->GatherInput(Dt, Particle0Level, Particle1Level, *IslandGroup);
 												}
 											}
-										});
+										}, 10);
 										for (int32 ConstraintIndex = OffsetBegin; ConstraintIndex < OffsetEnd; ++ConstraintIndex)
 										{
 											FConstraintContainerHandle* Constraint = SortedConstraints[ConstraintIndex];
diff --git a/Engine/Source/Runtime/Experimental/Chaos/Public/Chaos/Collision/SolverCollisionContainer.h b/Engine/Source/Runtime/Experimental/Chaos/Public/Chaos/Collision/SolverCollisionContainer.h
--- a/Engine/Source/Runtime/Experimental/Chaos/Public/Chaos/Collision/SolverCollisionContainer.h
+++ b/Engine/Source/Runtime/Experimental/Chaos/Public/Chaos/Collision/SolverCollisionContainer.h
@@ -82,8 +82,8 @@ namespace Chaos
 		void UpdateVelocityShockPropagation(const FReal Dt, const int32 It, const int32 NumIts, const int32 BeginIndex, const int32 EndIndex, const FPBDCollisionSolverSettings& SolverSettings);
 		bool SolvePositionImpl(const FReal Dt, const int32 It, const int32 NumIts, const int32 BeginIndex, const int32 EndIndex, const FPBDCollisionSolverSettings& SolverSettings, const bool bParallel);
 		bool SolvePositionIncrementalImpl(const FReal Dt, const int32 BeginIndex, const int32 EndIndex, const FReal MaxPushOut, const bool bApplyStaticFriction);
-		bool SolvePositionWithFrictionImpl(const FReal Dt, const int32 BeginIndex, const int32 EndIndex, const FReal MaxPushOut);
-		bool SolvePositionNoFrictionImpl(const FReal Dt, const int32 BeginIndex, const int32 EndIndex, const FReal MaxPushOut);
+		bool SolvePositionWithFrictionImpl(const FReal Dt, const int32 BeginIndex, const int32 EndIndex, const FReal MaxPushOut, const bool bParallel);
+		bool SolvePositionNoFrictionImpl(const FReal Dt, const int32 BeginIndex, const int32 EndIndex, const FReal MaxPushOut, const bool bParallel);
 		bool SolveVelocityImpl(const FReal Dt, const int32 It, const int32 NumIts, const int32 BeginIndex, const int32 EndIndex, const FPBDCollisionSolverSettings& SolverSettings, const bool bParallel);
 		void ScatterOutputImpl(const FReal Dt, const int32 BeginIndex, const int32 EndIndex, const bool bParallel);
 
diff --git a/Engine/Source/Runtime/Experimental/Chaos/Public/Chaos/Framework/Parallel.h b/Engine/Source/Runtime/Experimental/Chaos/Public/Chaos/Framework/Parallel.h
--- a/Engine/Source/Runtime/Experimental/Chaos/Public/Chaos/Framework/Parallel.h
+++ b/Engine/Source/Runtime/Experimental/Chaos/Public/Chaos/Framework/Parallel.h
@@ -6,9 +6,9 @@
 
 namespace Chaos
 {
-	void CHAOS_API PhysicsParallelForRange(int32 InNum, TFunctionRef<void(int32, int32)> InCallable, bool bForceSingleThreaded = false);
+	void CHAOS_API PhysicsParallelForRange(int32 InNum, TFunctionRef<void(int32, int32)> InCallable, const int32 MinBatchSize, bool bForceSingleThreaded = false);
 	void CHAOS_API PhysicsParallelFor(int32 InNum, TFunctionRef<void(int32)> InCallable, bool bForceSingleThreaded = false);
-	void CHAOS_API InnerPhysicsParallelForRange(int32 InNum, TFunctionRef<void(int32, int32)> InCallable, bool bForceSingleThreaded = false);
+	void CHAOS_API InnerPhysicsParallelForRange(int32 InNum, TFunctionRef<void(int32, int32)> InCallable, const int32 MinBatchSize, bool bForceSingleThreaded = false);
 	void CHAOS_API InnerPhysicsParallelFor(int32 InNum, TFunctionRef<void(int32)> InCallable, bool bForceSingleThreaded = false);
 	//void CHAOS_API PhysicsParallelFor_RecursiveDivide(int32 InNum, TFunctionRef<void(int32)> InCallable, bool bForceSingleThreaded = false);
 
@@ -21,5 +21,6 @@ namespace Chaos
 	CHAOS_API extern bool bDisablePhysicsParallelFor;
 	CHAOS_API extern bool bDisableParticleParallelFor;
 	CHAOS_API extern bool bDisableCollisionParallelFor;
+	CHAOS_API extern int32 MaxNumWorkers;
 #endif
 }
\ No newline at end of file
