Add max clamp on size of island groups.

diff --git a/Engine/Source/Runtime/Experimental/Chaos/Private/Chaos/Island/IslandManager.cpp b/Engine/Source/Runtime/Experimental/Chaos/Private/Chaos/Island/IslandManager.cpp
--- a/Engine/Source/Runtime/Experimental/Chaos/Private/Chaos/Island/IslandManager.cpp
+++ b/Engine/Source/Runtime/Experimental/Chaos/Private/Chaos/Island/IslandManager.cpp
@@ -329,7 +329,7 @@ void FPBDIslandManager::InitializeGroups()
 {
 	// @todo(chaos): is the number of worker threads a good indicator of how many threads we get in the solver loop? (Currently uses ParallelFor)
 	// Check for use of the "-onethread" command line arg, and physics threading disabled (GetNumWorkerThreads() is not affected by these)
-	const int32 NumWorkerThreads = (FApp::ShouldUseThreadingForPerformance() && !GSingleThreadedPhysics) ? FTaskGraphInterface::Get().GetNumWorkerThreads() : 0;
+	const int32 NumWorkerThreads = (FApp::ShouldUseThreadingForPerformance() && !GSingleThreadedPhysics) ? FMath::Min(FTaskGraphInterface::Get().GetNumWorkerThreads(), Chaos::MaxNumWorkers) : 0;
 	const int32 MaxIslandGroups = FMath::Max(1, FMath::CeilToInt32(FReal(NumWorkerThreads) * GChaosSolverIslandGroupsMultiplier));
 
 	IslandGroups.SetNum(MaxIslandGroups, false);
