Parallelize Island Gather. Makes about a 3x improvment when using joints.

diff --git a/Engine/Source/Runtime/Experimental/Chaos/Private/Chaos/PBDConstraintRule.cpp b/Engine/Source/Runtime/Experimental/Chaos/Private/Chaos/PBDConstraintRule.cpp
--- a/Engine/Source/Runtime/Experimental/Chaos/Private/Chaos/PBDConstraintRule.cpp
+++ b/Engine/Source/Runtime/Experimental/Chaos/Private/Chaos/PBDConstraintRule.cpp
@@ -28,8 +28,9 @@ namespace Chaos
 	int32 ChaosDisableIslandColors = true;
 	FAutoConsoleVariableRef CVarChaosDisableIslandColors(TEXT("p.Chaos.Islands.DisableColors"), ChaosDisableIslandColors, TEXT(""));
 	
-	DECLARE_CYCLE_STAT(TEXT("FPBDConstraintRule::Gather"), STAT_Constraints_Gather, STATGROUP_Chaos);
+	DECLARE_CYCLE_STAT(TEXT("FPBDConstraintRule::Island_Gather"), STAT_Constraints_Island_Gather, STATGROUP_Chaos);
 	DECLARE_CYCLE_STAT(TEXT("FPBDConstraintRule::Color_Gather"), STAT_Constraints_Color_Gather, STATGROUP_Chaos);
+	DECLARE_CYCLE_STAT(TEXT("FPBDConstraintRule::AddToGraph"), STAT_Constraints_AddToGraph, STATGROUP_Chaos);
 
 	template<class ConstraintType>
 	TSimpleConstraintRule<ConstraintType>::TSimpleConstraintRule(int32 InPriority, FConstraints& InConstraints)
@@ -132,6 +133,7 @@ namespace Chaos
 	template<class ConstraintType>
 	void TPBDConstraintGraphRuleImpl<ConstraintType>::AddToGraph()
 	{
+		SCOPE_CYCLE_COUNTER(STAT_Constraints_AddToGraph);
 		ConstraintGraph->ReserveConstraints(Constraints.NumConstraints());
 
 		for (typename FConstraints::FConstraintContainerHandle * ConstraintHandle : Constraints.GetConstraintHandles())
@@ -157,7 +159,7 @@ namespace Chaos
 	template<class ConstraintType>
 	void TPBDConstraintIslandRule<ConstraintType>::GatherSolverInput(const FReal Dt, int32 GroupIndex)
 	{
-		SCOPE_CYCLE_COUNTER(STAT_Constraints_Gather);
+		SCOPE_CYCLE_COUNTER(STAT_Constraints_Island_Gather);
 		if(FPBDIslandGroup* IslandGroup = ConstraintGraph->GetIslandGroup(GroupIndex))
 		{
 			// This will reset the number of constraints inside the solver datas. For now we keep this function since according
@@ -173,20 +175,27 @@ namespace Chaos
 				{
 					const TArray<FConstraintHandleHolder>& IslandConstraints = ConstraintGraph->GetIslandConstraints(IslandSolver->GetIslandIndex());
 					
-					for (FConstraintHandle* ConstraintHandle : IslandConstraints)
+					if (IslandConstraints.Num())
 					{
-						if (ConstraintHandle->GetContainerId() == GetContainerId())
+						InnerPhysicsParallelForRange(IslandConstraints.Num(), [&](const int32 StartRangeIndex, const int32 EndRangeIndex)
 						{
-							FConstraintContainerHandle* Constraint = ConstraintHandle->As<FConstraintContainerHandle>();
-
-							// Note we are building the SolverBodies as we go, in the order that we visit them. Each constraint
-							// references two bodies, so we won't strictly be accessing only in cache order, but it's about as good as it can be.
-							if (Constraint->IsEnabled())
+							for (int32 Index = StartRangeIndex; Index < EndRangeIndex; ++Index)
 							{
-								// @todo(chaos): we should provide Particle Levels in the island rule as well (see TPBDConstraintColorRule)
-								Constraint->GatherInput(Dt, INDEX_NONE, INDEX_NONE, *IslandGroup);
+								FConstraintHandle* ConstraintHandle = IslandConstraints[Index];
+								if (ConstraintHandle->GetContainerId() == GetContainerId())
+								{
+									FConstraintContainerHandle* Constraint = ConstraintHandle->As<FConstraintContainerHandle>();
+
+									// Note we are building the SolverBodies as we go, in the order that we visit them. Each constraint
+									// references two bodies, so we won't strictly be accessing only in cache order, but it's about as good as it can be.
+									if (Constraint->IsEnabled())
+									{
+										// @todo(chaos): we should provide Particle Levels in the island rule as well (see TPBDConstraintColorRule)
+										Constraint->GatherInput(Dt, INDEX_NONE, INDEX_NONE, *IslandGroup);
+									}
+								}
 							}
-						}
+						}, 1);
 					}
 				}
 			}
