Disable expensive stat in test builds. Saves about 200us in comparisons.

diff --git a/Engine/Source/Runtime/Experimental/Chaos/Private/PBDRigidsSolver.cpp b/Engine/Source/Runtime/Experimental/Chaos/Private/PBDRigidsSolver.cpp
--- a/Engine/Source/Runtime/Experimental/Chaos/Private/PBDRigidsSolver.cpp
+++ b/Engine/Source/Runtime/Experimental/Chaos/Private/PBDRigidsSolver.cpp
@@ -1460,7 +1460,7 @@ CSV_CUSTOM_STAT(PhysicsCounters, Name, Value, ECsvCustomStatOp::Set);
 		CHAOS_COUNTER_STAT(NumContacts, NumCollisionConstraints());
 		CHAOS_COUNTER_STAT(NumJoints, NumJointConstraints());
 
-#if !UE_BUILD_SHIPPING
+#if !UE_BUILD_SHIPPING && !UE_BUILD_TEST
 		UpdateExpensiveStatCounters();
 #endif
 
