Parallelize Event Processing.

diff --git a/Engine/Source/Runtime/Experimental/Chaos/Private/EventDefaults.cpp b/Engine/Source/Runtime/Experimental/Chaos/Private/EventDefaults.cpp
--- a/Engine/Source/Runtime/Experimental/Chaos/Private/EventDefaults.cpp
+++ b/Engine/Source/Runtime/Experimental/Chaos/Private/EventDefaults.cpp
@@ -17,6 +17,9 @@
 
 namespace Chaos
 {
+	DECLARE_CYCLE_STAT(TEXT("EventDefaults::RegisterCollisionEvent_Filter"), STAT_Events_RegisterCollisionEvent_Filter, STATGROUP_Chaos);
+	DECLARE_CYCLE_STAT(TEXT("EventDefaults::RegisterCollisionEvent_Notify"), STAT_Events_RegisterCollisionEvent_Notify, STATGROUP_Chaos);
+
 	/**
 	 * Resolves a material from a shape and a collision.
 	 * Either both shapes are simple primitives, or one is simple and the other is a heightfield or trimesh.
@@ -110,156 +113,193 @@ namespace Chaos
 				TArray<const Chaos::FPBDCollisionConstraintHandle*> ValidCollisionHandles;
 				ValidCollisionHandles.SetNumUninitialized(CollisionRule.NumConstraints());
 				int32 NumValidCollisions = 0;
-				const FReal MinDeltaVelocityForHitEvents = FChaosSolversModule::GetModule()->GetSettingsProvider().GetMinDeltaVelocityForHitEvents();
-				for (const Chaos::FPBDCollisionConstraintHandle * ContactHandle : CollisionRule.GetConstConstraintHandles())
 				{
-					if (NumValidCollisions >= CollisionRule.NumConstraints())
+					SCOPE_CYCLE_COUNTER(STAT_Events_RegisterCollisionEvent_Filter);
+					const FReal MinDeltaVelocityForHitEvents = FChaosSolversModule::GetModule()->GetSettingsProvider().GetMinDeltaVelocityForHitEvents();
+					const FPBDCollisionConstraints::FConstHandles& CollisionHandles = CollisionRule.GetConstConstraintHandles();
+					TArray<bool> ValidArray;
+					ValidArray.SetNum(CollisionHandles.Num());
+					InnerPhysicsParallelForRange(CollisionHandles.Num(), [&](int32 StartRangeIndex, int32 EndRangeIndex)
 					{
-						break;
-					}
+						for (int32 Index = StartRangeIndex; Index < EndRangeIndex; ++Index)
+						{
+							const Chaos::FPBDCollisionConstraintHandle* ContactHandle = CollisionHandles[Index];
+							ValidArray[Index] = false;
+							if (NumValidCollisions >= CollisionRule.NumConstraints())
+							{
+								break;
+							}
 
-					const FPBDCollisionConstraint& Constraint = ContactHandle->GetContact();
+							const FPBDCollisionConstraint& Constraint = ContactHandle->GetContact();
 
-					if (ensure(!Constraint.AccumulatedImpulse.ContainsNaN() && FMath::IsFinite(Constraint.GetPhi())))
-					{
-						const FGeometryParticleHandle* Particle0 = Constraint.GetParticle0();
-						const FGeometryParticleHandle* Particle1 = Constraint.GetParticle1();
-						const FKinematicGeometryParticleHandle* Body0 = Particle0->CastToKinematicParticle();
+							if (ensure(!Constraint.AccumulatedImpulse.ContainsNaN() && FMath::IsFinite(Constraint.GetPhi())))
+							{
+								const FPerShapeData* Shape0 = Constraint.GetShape0();
+								const FPerShapeData* Shape1 = Constraint.GetShape1();
 
-						// presently when a rigidbody or kinematic hits static geometry then Body1 is null
-						const FKinematicGeometryParticleHandle* Body1 = Particle1->CastToKinematicParticle();
+								// If we don't have a filter - allow the notify, otherwise obey the filter flag
+								const bool bFilter0Notify = Shape0 ? Shape0->GetSimData().HasFlag(EFilterFlags::ContactNotify) : true;
+								const bool bFilter1Notify = Shape1 ? Shape1->GetSimData().HasFlag(EFilterFlags::ContactNotify) : true;
 
-						const FImplicitObject* Implicit0 = Constraint.GetImplicit0();
-						const FImplicitObject* Implicit1 = Constraint.GetImplicit1();
-	
-						const FPerShapeData* Shape0 = Constraint.GetShape0();
-						const FPerShapeData* Shape1 = Constraint.GetShape1();
+								if (!bFilter0Notify && !bFilter1Notify)
+								{
+									// No need to notify - engine didn't request notifications for either shape.
+									continue;
+								}
 
-						// If we don't have a filter - allow the notify, otherwise obey the filter flag
-						const bool bFilter0Notify = Shape0 ? Shape0->GetSimData().HasFlag(EFilterFlags::ContactNotify) : true;
-						const bool bFilter1Notify = Shape1 ? Shape1->GetSimData().HasFlag(EFilterFlags::ContactNotify) : true;
+								const FGeometryParticleHandle* Particle0 = Constraint.GetParticle0();
+								const FGeometryParticleHandle* Particle1 = Constraint.GetParticle1();
+								const FKinematicGeometryParticleHandle* Body0 = Particle0->CastToKinematicParticle();
 
-						if(!bFilter0Notify && !bFilter1Notify)
-						{
-							// No need to notify - engine didn't request notifications for either shape.
-							continue;
-						}
+								// presently when a rigidbody or kinematic hits static geometry then Body1 is null
+								const FKinematicGeometryParticleHandle* Body1 = Particle1->CastToKinematicParticle();
 
-						const FKinematicGeometryParticleHandle* Primary = Body0 ? Body0 : Body1;
-						const FKinematicGeometryParticleHandle* Secondary = Body0 ? Body1 : Body0;
+								const FKinematicGeometryParticleHandle* Primary = Body0 ? Body0 : Body1;
+								const FKinematicGeometryParticleHandle* Secondary = Body0 ? Body1 : Body0;
 
-						if (!Constraint.AccumulatedImpulse.IsZero() && Primary)
-						{
-							if (ensure(!Constraint.CalculateWorldContactLocation().ContainsNaN() &&
-								!Constraint.CalculateWorldContactNormal().ContainsNaN()) &&
-								!Primary->V().ContainsNaN() &&
-								!Primary->W().ContainsNaN() &&
-								(Secondary == nullptr || ((!Secondary->V().ContainsNaN()) && !Secondary->W().ContainsNaN())))
-							{
-								ValidCollisionHandles[NumValidCollisions] = ContactHandle;
-								NumValidCollisions++;
+								if (!Constraint.AccumulatedImpulse.IsZero() && Primary)
+								{
+									if (ensure(!Constraint.CalculateWorldContactLocation().ContainsNaN() &&
+										!Constraint.CalculateWorldContactNormal().ContainsNaN()) &&
+										!Primary->V().ContainsNaN() &&
+										!Primary->W().ContainsNaN() &&
+										(Secondary == nullptr || ((!Secondary->V().ContainsNaN()) && !Secondary->W().ContainsNaN())))
+									{
+										ValidArray[Index] = true;
+									}
+								}
 							}
 						}
+					}, 100);
+					for (int32 Index = 0; Index < CollisionHandles.Num(); ++Index)
+					{
+						if (ValidArray[Index])
+						{
+							const Chaos::FPBDCollisionConstraintHandle* ContactHandle = CollisionHandles[Index];
+							ValidCollisionHandles[NumValidCollisions] = ContactHandle;
+							NumValidCollisions++;
+						}
 					}
 				}
 
-				ValidCollisionHandles.SetNum(NumValidCollisions);
-
-				if(ValidCollisionHandles.Num() > 0)
 				{
-					for (int32 IdxCollision = 0; IdxCollision < ValidCollisionHandles.Num(); ++IdxCollision)
+					SCOPE_CYCLE_COUNTER(STAT_Events_RegisterCollisionEvent_Notify);
+					ValidCollisionHandles.SetNum(NumValidCollisions);
+					if (ValidCollisionHandles.Num() > 0)
 					{
-						const FPBDCollisionConstraint& Constraint = ValidCollisionHandles[IdxCollision]->GetContact();
+						AllCollisionsDataArray.SetNum(NumValidCollisions);
+						InnerPhysicsParallelForRange(ValidCollisionHandles.Num(), [&](int32 StartRangeIndex, int32 EndRangeIndex)
+						{
+							for (int32 IdxCollision = StartRangeIndex; IdxCollision < EndRangeIndex; ++IdxCollision)
+							{
+								const FPBDCollisionConstraint& Constraint = ValidCollisionHandles[IdxCollision]->GetContact();
 
-						const FGeometryParticleHandle* Particle0 = Constraint.GetParticle0();
-						const FGeometryParticleHandle* Particle1 = Constraint.GetParticle1();
+								const FGeometryParticleHandle* Particle0 = Constraint.GetParticle0();
+								const FGeometryParticleHandle* Particle1 = Constraint.GetParticle1();
 
-						FCollidingData Data;
-						Data.Location = Constraint.CalculateWorldContactLocation();
-						Data.AccumulatedImpulse = Constraint.AccumulatedImpulse;
-						Data.Normal = Constraint.CalculateWorldContactNormal();
-						Data.PenetrationDepth = Constraint.GetPhi();
-						
-						// @todo(chaos): fix this casting
-						Data.Proxy1 = Particle0 ? const_cast<IPhysicsProxyBase*>(Particle0->PhysicsProxy()) : nullptr;
-						Data.Proxy2 = Particle1 ? const_cast<IPhysicsProxyBase*>(Particle1->PhysicsProxy()) : nullptr;
+								FCollidingData Data;
+								Data.Location = Constraint.CalculateWorldContactLocation();
+								Data.AccumulatedImpulse = Constraint.AccumulatedImpulse;
+								Data.Normal = Constraint.CalculateWorldContactNormal();
+								Data.PenetrationDepth = Constraint.GetPhi();
 
-						const FPerShapeData* Shape0 = Constraint.GetShape0();
-						const FPerShapeData* Shape1 = Constraint.GetShape1();
+								// @todo(chaos): fix this casting
+								Data.Proxy1 = Particle0 ? const_cast<IPhysicsProxyBase*>(Particle0->PhysicsProxy()) : nullptr;
+								Data.Proxy2 = Particle1 ? const_cast<IPhysicsProxyBase*>(Particle1->PhysicsProxy()) : nullptr;
 
-						Data.Mat1 = ResolveMaterial(Shape0, Constraint);
-						Data.Mat2 = ResolveMaterial(Shape1, Constraint);
+								const FPerShapeData* Shape0 = Constraint.GetShape0();
+								const FPerShapeData* Shape1 = Constraint.GetShape1();
 
-						// Collision constraints require both proxies are valid. If either is not, we needn't record the collision event.
-						if (Data.Proxy1 == nullptr || Data.Proxy2 == nullptr)
-						{
-							continue;
-						}
+								Data.Mat1 = ResolveMaterial(Shape0, Constraint);
+								Data.Mat2 = ResolveMaterial(Shape1, Constraint);
 
-						if (const FPBDRigidParticleHandle * Rigid0 = Particle0->CastToRigidParticle())
-						{
-							Data.DeltaVelocity1 = Rigid0->V() - Rigid0->PreV();
-						}
-						if (const FPBDRigidParticleHandle * Rigid1 = Particle1->CastToRigidParticle())
-						{
-							Data.DeltaVelocity2 = Rigid1->V() - Rigid1->PreV();
-						}
+								// Collision constraints require both proxies are valid. If either is not, we needn't record the collision event.
+								if (Data.Proxy1 == nullptr || Data.Proxy2 == nullptr)
+								{
+									continue;
+								}
 
-						// todo: do we need these anymore now we are storing the particles you can access all of this stuff from there
-						// do we still need these now we have pointers to particles returned?
-						const FPBDRigidParticleHandle* PBDRigid0 = Particle0->CastToRigidParticle();
-						if (PBDRigid0 && PBDRigid0->ObjectState() == EObjectStateType::Dynamic)
-						{
-							Data.Velocity1 = PBDRigid0->V();
-							Data.AngularVelocity1 = PBDRigid0->W();
-							Data.Mass1 = PBDRigid0->M();
-						}
+								if (const FPBDRigidParticleHandle* Rigid0 = Particle0->CastToRigidParticle())
+								{
+									Data.DeltaVelocity1 = Rigid0->V() - Rigid0->PreV();
+								}
+								if (const FPBDRigidParticleHandle* Rigid1 = Particle1->CastToRigidParticle())
+								{
+									Data.DeltaVelocity2 = Rigid1->V() - Rigid1->PreV();
+								}
 
-						const FPBDRigidParticleHandle* PBDRigid1 = Particle1->CastToRigidParticle();
-						if (PBDRigid1 && PBDRigid1->ObjectState() == EObjectStateType::Dynamic)
-						{
-							Data.Velocity2 = PBDRigid1->V();
-							Data.AngularVelocity2 = PBDRigid1->W();
-							Data.Mass2 = PBDRigid1->M();
-						}
+								// todo: do we need these anymore now we are storing the particles you can access all of this stuff from there
+								// do we still need these now we have pointers to particles returned?
+								const FPBDRigidParticleHandle* PBDRigid0 = Particle0->CastToRigidParticle();
+								if (PBDRigid0 && PBDRigid0->ObjectState() == EObjectStateType::Dynamic)
+								{
+									Data.Velocity1 = PBDRigid0->V();
+									Data.AngularVelocity1 = PBDRigid0->W();
+									Data.Mass1 = PBDRigid0->M();
+								}
 
-						IPhysicsProxyBase* const PhysicsProxy = const_cast<IPhysicsProxyBase*>(Particle0->PhysicsProxy());
-						IPhysicsProxyBase* const OtherPhysicsProxy = const_cast<IPhysicsProxyBase*>(Particle1->PhysicsProxy());
-							
-						const FSolverCollisionEventFilter* SolverCollisionEventFilter = Solver->GetEventFilters()->GetCollisionFilter();
-						if (!SolverCollisionEventFilter->Enabled() || SolverCollisionEventFilter->Pass(Data))
+								const FPBDRigidParticleHandle* PBDRigid1 = Particle1->CastToRigidParticle();
+								if (PBDRigid1 && PBDRigid1->ObjectState() == EObjectStateType::Dynamic)
+								{
+									Data.Velocity2 = PBDRigid1->V();
+									Data.AngularVelocity2 = PBDRigid1->W();
+									Data.Mass2 = PBDRigid1->M();
+								}
 
-						{
-							const int32 NewIdx = AllCollisionsDataArray.Add(FCollidingData());
-							FCollidingData& CollisionDataArrayItem = AllCollisionsDataArray[NewIdx];
+								IPhysicsProxyBase* const PhysicsProxy = const_cast<IPhysicsProxyBase*>(Particle0->PhysicsProxy());
+								IPhysicsProxyBase* const OtherPhysicsProxy = const_cast<IPhysicsProxyBase*>(Particle1->PhysicsProxy());
 
-							CollisionDataArrayItem = Data;
+								const FSolverCollisionEventFilter* SolverCollisionEventFilter = Solver->GetEventFilters()->GetCollisionFilter();
+								if (!SolverCollisionEventFilter->Enabled() || SolverCollisionEventFilter->Pass(Data))
 
+								{
+									AllCollisionsDataArray[IdxCollision] = Data;
 #if TODO_REIMPLEMENT_RIGID_CLUSTERING
-							// If Constraint.ParticleIndex is a cluster store an index for a mesh in this cluster
-							if (ClusterIdsArray[Constraint.ParticleIndex].NumChildren > 0)
+									// If Constraint.ParticleIndex is a cluster store an index for a mesh in this cluster
+									if (ClusterIdsArray[Constraint.ParticleIndex].NumChildren > 0)
+									{
+										int32 ParticleIndexMesh = GetParticleIndexMesh(ParentToChildrenMap, Constraint.ParticleIndex);
+										ensure(ParticleIndexMesh != INDEX_NONE);
+										CollisionDataArrayItem.ParticleIndexMesh = ParticleIndexMesh;
+									}
+									// If Constraint.LevelsetIndex is a cluster store an index for a mesh in this cluster
+									if (ClusterIdsArray[Constraint.LevelsetIndex].NumChildren > 0)
+									{
+										int32 LevelsetIndexMesh = GetParticleIndexMesh(ParentToChildrenMap, Constraint.LevelsetIndex);
+										ensure(LevelsetIndexMesh != INDEX_NONE);
+										CollisionDataArrayItem.LevelsetIndexMesh = LevelsetIndexMesh;
+									}
+#endif
+								}
+							}
+						}, 10);
+						int32 IdxCollision = 0, IdxCollision2 = NumValidCollisions - 1, NumDeleted = 0;
+						while (IdxCollision <= IdxCollision2)
+						{
+							if (AllCollisionsDataArray[IdxCollision].Proxy1 == nullptr && AllCollisionsDataArray[IdxCollision2].Proxy1 == nullptr)
 							{
-								int32 ParticleIndexMesh = GetParticleIndexMesh(ParentToChildrenMap, Constraint.ParticleIndex);
-								ensure(ParticleIndexMesh != INDEX_NONE);
-								CollisionDataArrayItem.ParticleIndexMesh = ParticleIndexMesh;
+								NumDeleted++;
+								IdxCollision2--;
 							}
-							// If Constraint.LevelsetIndex is a cluster store an index for a mesh in this cluster
-							if (ClusterIdsArray[Constraint.LevelsetIndex].NumChildren > 0)
+							else if (AllCollisionsDataArray[IdxCollision].Proxy1 == nullptr)
 							{
-								int32 LevelsetIndexMesh = GetParticleIndexMesh(ParentToChildrenMap, Constraint.LevelsetIndex);
-								ensure(LevelsetIndexMesh != INDEX_NONE);
-								CollisionDataArrayItem.LevelsetIndexMesh = LevelsetIndexMesh;
+								AllCollisionsDataArray[IdxCollision] = AllCollisionsDataArray[IdxCollision2];
+								NumDeleted++;
+								IdxCollision2--;
 							}
-#endif
-
-							// Add to AllCollisionsIndicesByPhysicsProxy
-							AllCollisionsIndicesByPhysicsProxy.FindOrAdd(PhysicsProxy).Add(FEventManager::EncodeCollisionIndex(NewIdx, false));
-
-							if (OtherPhysicsProxy && OtherPhysicsProxy != PhysicsProxy)
+							else
 							{
-								AllCollisionsIndicesByPhysicsProxy.FindOrAdd(OtherPhysicsProxy).Add(FEventManager::EncodeCollisionIndex(NewIdx, true));
+								// Add to AllCollisionsIndicesByPhysicsProxy
+								AllCollisionsIndicesByPhysicsProxy.FindOrAdd(AllCollisionsDataArray[IdxCollision].Proxy1).Add(FEventManager::EncodeCollisionIndex(IdxCollision, false));
+
+								if (AllCollisionsDataArray[IdxCollision].Proxy2 && AllCollisionsDataArray[IdxCollision].Proxy2 != AllCollisionsDataArray[IdxCollision].Proxy1)
+								{
+									AllCollisionsIndicesByPhysicsProxy.FindOrAdd(AllCollisionsDataArray[IdxCollision].Proxy2).Add(FEventManager::EncodeCollisionIndex(IdxCollision, true));
+								}
+								IdxCollision++;
 							}
 						}
+						AllCollisionsDataArray.SetNum(NumValidCollisions - NumDeleted);
 					}
 				}
 			}
