Use cvars for parallelfor sizes. Don't resize arrays when not necessary to save perf. Add PreGather for joints and other constraints to be thread safe.

diff --git a/Engine/Source/Programs/HeadlessChaos/Public/HeadlessChaosCollisionConstraints.h b/Engine/Source/Programs/HeadlessChaos/Public/HeadlessChaosCollisionConstraints.h
--- a/Engine/Source/Programs/HeadlessChaos/Public/HeadlessChaosCollisionConstraints.h
+++ b/Engine/Source/Programs/HeadlessChaos/Public/HeadlessChaosCollisionConstraints.h
@@ -144,7 +144,7 @@ public:
 
 		for (auto Handle : CollisionConstraints.GetConstraintHandles())
 		{
-			Handle->PreGatherInput(SolverData);
+			Handle->PreGatherInput(Dt, SolverData);
 			Handle->GatherInput(Dt, INDEX_NONE, INDEX_NONE, SolverData);
 		}
 	}
diff --git a/Engine/Source/Runtime/Experimental/Chaos/Private/Chaos/Collision/PBDCollisionConstraintHandle.cpp b/Engine/Source/Runtime/Experimental/Chaos/Private/Chaos/Collision/PBDCollisionConstraintHandle.cpp
--- a/Engine/Source/Runtime/Experimental/Chaos/Private/Chaos/Collision/PBDCollisionConstraintHandle.cpp
+++ b/Engine/Source/Runtime/Experimental/Chaos/Private/Chaos/Collision/PBDCollisionConstraintHandle.cpp
@@ -57,7 +57,7 @@ namespace Chaos
 		return { GetContact().GetParticle0(), GetContact().GetParticle1() };
 	}
 
-	void FPBDCollisionConstraintHandle::PreGatherInput(FPBDIslandSolverData& SolverData)
+	void FPBDCollisionConstraintHandle::PreGatherInput(const FReal Dt, FPBDIslandSolverData& SolverData)
 	{
 		ConcreteContainer()->PreGatherInput(GetContact(), SolverData);
 	}
diff --git a/Engine/Source/Runtime/Experimental/Chaos/Private/Chaos/Collision/SolverCollisionContainer.cpp b/Engine/Source/Runtime/Experimental/Chaos/Private/Chaos/Collision/SolverCollisionContainer.cpp
--- a/Engine/Source/Runtime/Experimental/Chaos/Private/Chaos/Collision/SolverCollisionContainer.cpp
+++ b/Engine/Source/Runtime/Experimental/Chaos/Private/Chaos/Collision/SolverCollisionContainer.cpp
@@ -574,7 +574,7 @@ namespace Chaos
 			{
 				CollisionSolvers[SolverIndex].GetSolver().SolvePositionWithFriction(Dt, MaxPushOut);
 			}
-		}, 100, !bParallel);
+		}, Chaos::LargeBatchSize, !bParallel);
 		return true;
 	}
 
@@ -596,7 +596,7 @@ namespace Chaos
 			{
 				CollisionSolvers[SolverIndex].GetSolver().SolvePositionNoFriction(Dt, MaxPushOut);
 			}
-		}, 100, !bParallel);
+		}, Chaos::LargeBatchSize, !bParallel);
 		return true;
 	}
 
@@ -649,7 +649,7 @@ namespace Chaos
 			{
 				CollisionSolvers[SolverIndex].ScatterOutput(Dt);
 			}
-		}, 100);
+		}, Chaos::LargeBatchSize);
 	}
 
 }
diff --git a/Engine/Source/Runtime/Experimental/Chaos/Private/Chaos/Framework/Parallel.cpp b/Engine/Source/Runtime/Experimental/Chaos/Private/Chaos/Framework/Parallel.cpp
--- a/Engine/Source/Runtime/Experimental/Chaos/Private/Chaos/Framework/Parallel.cpp
+++ b/Engine/Source/Runtime/Experimental/Chaos/Private/Chaos/Framework/Parallel.cpp
@@ -12,6 +12,8 @@ namespace Chaos
 	CHAOS_API int32 InnerParallelForBatchSize = 0;
 	CHAOS_API int32 MinRangeBatchSize = 0;
 	CHAOS_API int32 MaxNumWorkers = 100;
+	CHAOS_API int32 SmallBatchSize = 10;
+	CHAOS_API int32 LargeBatchSize = 100;
 	
 #if !UE_BUILD_SHIPPING
 	CHAOS_API bool bDisablePhysicsParallelFor = false;
@@ -24,6 +26,8 @@ namespace Chaos
 	FAutoConsoleVariableRef CVarInnerPhysicsBatchSize(TEXT("p.Chaos.InnerParallelForBatchSize"), InnerParallelForBatchSize, TEXT("Set the batch size threshold for inner parallel fors"));
 	FAutoConsoleVariableRef CVarMinRangeBatchSize(TEXT("p.Chaos.MinRangeBatchSize"), MinRangeBatchSize, TEXT("Set the min range batch size for parallel for"));
 	FAutoConsoleVariableRef CVarMaxRangeBatchWorkers(TEXT("p.Chaos.MaxNumWorkers"), MaxNumWorkers, TEXT("Set the max number of workers for physics"));
+	FAutoConsoleVariableRef CVarSmallBatchSize(TEXT("p.Chaos.SmallBatchSize"), SmallBatchSize, TEXT("Small batch size for chaos parallel loops"));
+	FAutoConsoleVariableRef CVarLargeBatchSize(TEXT("p.Chaos.LargeBatchSize"), LargeBatchSize, TEXT("Large batch size for chaos parallel loops"));
 #endif
 }
 
diff --git a/Engine/Source/Runtime/Experimental/Chaos/Private/Chaos/PBDConstraintRule.cpp b/Engine/Source/Runtime/Experimental/Chaos/Private/Chaos/PBDConstraintRule.cpp
--- a/Engine/Source/Runtime/Experimental/Chaos/Private/Chaos/PBDConstraintRule.cpp
+++ b/Engine/Source/Runtime/Experimental/Chaos/Private/Chaos/PBDConstraintRule.cpp
@@ -177,25 +177,38 @@ namespace Chaos
 					
 					if (IslandConstraints.Num())
 					{
-						InnerPhysicsParallelForRange(IslandConstraints.Num(), [&](const int32 StartRangeIndex, const int32 EndRangeIndex)
+						for (int32 Index = 0; Index < IslandConstraints.Num(); ++Index)
 						{
-							for (int32 Index = StartRangeIndex; Index < EndRangeIndex; ++Index)
+							FConstraintHandle* ConstraintHandle = IslandConstraints[Index];
+							if (ConstraintHandle->GetContainerId() == GetContainerId())
 							{
-								FConstraintHandle* ConstraintHandle = IslandConstraints[Index];
-								if (ConstraintHandle->GetContainerId() == GetContainerId())
+								FConstraintContainerHandle* Constraint = ConstraintHandle->As<FConstraintContainerHandle>();
+
+								// Note we are building the SolverBodies as we go, in the order that we visit them. Each constraint
+								// references two bodies, so we won't strictly be accessing only in cache order, but it's about as good as it can be.
+								if (Constraint->IsEnabled())
 								{
-									FConstraintContainerHandle* Constraint = ConstraintHandle->As<FConstraintContainerHandle>();
+									// @todo(chaos): we should provide Particle Levels in the island rule as well (see TPBDConstraintColorRule)
+									Constraint->PreGatherInput(Dt, *IslandGroup);
+								}
+							}
+						}
+						InnerPhysicsParallelFor(IslandConstraints.Num(), [&](const int32 Index)
+						{
+							FConstraintHandle* ConstraintHandle = IslandConstraints[Index];
+							if (ConstraintHandle->GetContainerId() == GetContainerId())
+							{
+								FConstraintContainerHandle* Constraint = ConstraintHandle->As<FConstraintContainerHandle>();
 
-									// Note we are building the SolverBodies as we go, in the order that we visit them. Each constraint
-									// references two bodies, so we won't strictly be accessing only in cache order, but it's about as good as it can be.
-									if (Constraint->IsEnabled())
-									{
-										// @todo(chaos): we should provide Particle Levels in the island rule as well (see TPBDConstraintColorRule)
-										Constraint->GatherInput(Dt, INDEX_NONE, INDEX_NONE, *IslandGroup);
-									}
+								// Note we are building the SolverBodies as we go, in the order that we visit them. Each constraint
+								// references two bodies, so we won't strictly be accessing only in cache order, but it's about as good as it can be.
+								if (Constraint->IsEnabled())
+								{
+									// @todo(chaos): we should provide Particle Levels in the island rule as well (see TPBDConstraintColorRule)
+									Constraint->GatherInput(Dt, INDEX_NONE, INDEX_NONE, *IslandGroup);
 								}
 							}
-						}, 1);
+						});
 					}
 				}
 			}
@@ -320,7 +333,7 @@ namespace Chaos
 											FConstraintContainerHandle* Constraint = SortedConstraints[ConstraintIndex];
 											if (Constraint->IsEnabled())
 											{
-												Constraint->PreGatherInput(*IslandGroup);
+												Constraint->PreGatherInput(Dt, *IslandGroup);
 											}
 										}
 										// Calculate the range of indices for this color as a set of independent contacts
@@ -344,7 +357,7 @@ namespace Chaos
 													Constraint->GatherInput(Dt, Particle0Level, Particle1Level, *IslandGroup);
 												}
 											}
-										}, 10);
+										}, Chaos::SmallBatchSize);
 										for (int32 ConstraintIndex = OffsetBegin; ConstraintIndex < OffsetEnd; ++ConstraintIndex)
 										{
 											FConstraintContainerHandle* Constraint = SortedConstraints[ConstraintIndex];
diff --git a/Engine/Source/Runtime/Experimental/Chaos/Private/Chaos/PBDJointConstraints.cpp b/Engine/Source/Runtime/Experimental/Chaos/Private/Chaos/PBDJointConstraints.cpp
--- a/Engine/Source/Runtime/Experimental/Chaos/Private/Chaos/PBDJointConstraints.cpp
+++ b/Engine/Source/Runtime/Experimental/Chaos/Private/Chaos/PBDJointConstraints.cpp
@@ -174,8 +174,14 @@ namespace Chaos
 		return ConcreteContainer()->SetConstraintEnabled(ConstraintIndex, bInEnabled);
 	}
 
+	void FPBDJointConstraintHandle::PreGatherInput(const FReal Dt, FPBDIslandSolverData& SolverData)
+	{
+		ConcreteContainer()->PreGatherInput(Dt, ConstraintIndex, SolverData);
+	}
+
 	void FPBDJointConstraintHandle::GatherInput(const FReal Dt, const int32 Particle0Level, const int32 Particle1Level, FPBDIslandSolverData& SolverData)
 	{
+
 		ConcreteContainer()->GatherInput(Dt, ConstraintIndex, Particle0Level, Particle1Level, SolverData);
 	}
 
@@ -856,13 +862,28 @@ namespace Chaos
 		SolverData.GetConstraintIndices(ContainerId).Reset(NumIslandConstraints);
 	}
 
-	void FPBDJointConstraints::GatherInput(const FReal Dt, const int32 ConstraintIndex, const int32 Particle0Level, const int32 Particle1Level, FPBDIslandSolverData& SolverData)
+	void FPBDJointConstraints::PreGatherInput(const FReal Dt, const int32 ConstraintIndex, FPBDIslandSolverData& SolverData)
 	{
-		SCOPE_CYCLE_COUNTER(STAT_Joints_Gather);
 		check(!ConstraintStates[ConstraintIndex].bDisabled);
 
 		SolverData.GetConstraintIndices(ContainerId).Add(ConstraintIndex);
 
+		int32 Index0, Index1;
+		GetConstrainedParticleIndices(ConstraintIndex, Index0, Index1);
+		FGenericParticleHandle Particle0 = FGenericParticleHandle(ConstraintParticles[ConstraintIndex][Index0]);
+		FGenericParticleHandle Particle1 = FGenericParticleHandle(ConstraintParticles[ConstraintIndex][Index1]);
+
+		// Find the solver bodies for the particles we constrain. This will add them to the container
+		// if they aren't there already, and ensure that they are populated with the latest data.
+		SolverData.GetBodyContainer().FindOrAdd(Particle0);
+		SolverData.GetBodyContainer().FindOrAdd(Particle1);
+	}
+
+	void FPBDJointConstraints::GatherInput(const FReal Dt, const int32 ConstraintIndex, const int32 Particle0Level, const int32 Particle1Level, FPBDIslandSolverData& SolverData)
+	{
+		SCOPE_CYCLE_COUNTER(STAT_Joints_Gather);
+		check(!ConstraintStates[ConstraintIndex].bDisabled);
+
 		const FPBDJointSettings& JointSettings = ConstraintSettings[ConstraintIndex];
 		const FTransformPair& JointFrames = JointSettings.ConnectorTransforms;
 
diff --git a/Engine/Source/Runtime/Experimental/Chaos/Private/Chaos/PBDPositionConstraints.cpp b/Engine/Source/Runtime/Experimental/Chaos/Private/Chaos/PBDPositionConstraints.cpp
--- a/Engine/Source/Runtime/Experimental/Chaos/Private/Chaos/PBDPositionConstraints.cpp
+++ b/Engine/Source/Runtime/Experimental/Chaos/Private/Chaos/PBDPositionConstraints.cpp
@@ -10,6 +10,11 @@ namespace Chaos
 		return ConcreteContainer()->GetConstrainedParticles(ConstraintIndex);
 	}
 
+	void FPBDPositionConstraintHandle::PreGatherInput(const FReal Dt, FPBDIslandSolverData& SolverData)
+	{
+		ConcreteContainer()->PreGatherInput(Dt, ConstraintIndex, SolverData);
+	}
+
 	void FPBDPositionConstraintHandle::GatherInput(const FReal Dt, const int32 Particle0Level, const int32 Particle1Level, FPBDIslandSolverData& SolverData)
 	{
 		ConcreteContainer()->GatherInput(Dt, ConstraintIndex, Particle0Level, Particle1Level, SolverData);
@@ -20,12 +25,16 @@ namespace Chaos
 		SolverData.GetConstraintIndices(ContainerId).Reset(NumIslandConstraints);
 	}
 
-	void FPBDPositionConstraints::GatherInput(const FReal Dt, const int32 ConstraintIndex, const int32 Particle0Level, const int32 Particle1Level, FPBDIslandSolverData& SolverData)
+	void FPBDPositionConstraints::PreGatherInput(const FReal Dt, const int32 ConstraintIndex, FPBDIslandSolverData& SolverData)
 	{
 		SolverData.GetConstraintIndices(ContainerId).Add(ConstraintIndex);
 
 		ConstraintSolverBodies[ConstraintIndex] = SolverData.GetBodyContainer().FindOrAdd(ConstrainedParticles[ConstraintIndex]);
 	}
+
+	void FPBDPositionConstraints::GatherInput(const FReal Dt, const int32 ConstraintIndex, const int32 Particle0Level, const int32 Particle1Level, FPBDIslandSolverData& SolverData)
+	{
+	}
 	
 	void FPBDPositionConstraints::ScatterOutput(FReal Dt, FPBDIslandSolverData& SolverData)
 	{
diff --git a/Engine/Source/Runtime/Experimental/Chaos/Private/Chaos/PBDRigidDynamicSpringConstraints.cpp b/Engine/Source/Runtime/Experimental/Chaos/Private/Chaos/PBDRigidDynamicSpringConstraints.cpp
--- a/Engine/Source/Runtime/Experimental/Chaos/Private/Chaos/PBDRigidDynamicSpringConstraints.cpp
+++ b/Engine/Source/Runtime/Experimental/Chaos/Private/Chaos/PBDRigidDynamicSpringConstraints.cpp
@@ -12,6 +12,11 @@ TVec2<FGeometryParticleHandle*> FPBDRigidDynamicSpringConstraintHandle::GetConst
 	return ConcreteContainer()->GetConstrainedParticles(ConstraintIndex);
 }
 
+void FPBDRigidDynamicSpringConstraintHandle::PreGatherInput(const FReal Dt, FPBDIslandSolverData& SolverData)
+{
+	ConcreteContainer()->PreGatherInput(Dt, ConstraintIndex, SolverData);
+}
+
 void FPBDRigidDynamicSpringConstraintHandle::GatherInput(const FReal Dt, const int32 Particle0Level, const int32 Particle1Level, FPBDIslandSolverData& SolverData)
 {
 	ConcreteContainer()->GatherInput(Dt, ConstraintIndex, Particle0Level, Particle1Level, SolverData);
@@ -107,7 +112,7 @@ void FPBDRigidDynamicSpringConstraints::SetNumIslandConstraints(const int32 NumI
 	SolverData.GetConstraintIndices(ContainerId).Reset(NumIslandConstraints);
 }
 
-void FPBDRigidDynamicSpringConstraints::GatherInput(const FReal Dt, const int32 ConstraintIndex, const int32 Particle0Level, const int32 Particle1Level, FPBDIslandSolverData& SolverData)
+void FPBDRigidDynamicSpringConstraints::PreGatherInput(const FReal Dt, const int32 ConstraintIndex, FPBDIslandSolverData& SolverData)
 {
 	SolverData.GetConstraintIndices(ContainerId).Add(ConstraintIndex);
 
@@ -118,6 +123,10 @@ void FPBDRigidDynamicSpringConstraints::GatherInput(const FReal Dt, const int32
 	};
 }
 
+void FPBDRigidDynamicSpringConstraints::GatherInput(const FReal Dt, const int32 ConstraintIndex, const int32 Particle0Level, const int32 Particle1Level, FPBDIslandSolverData& SolverData)
+{
+}
+
 void FPBDRigidDynamicSpringConstraints::ScatterOutput(FReal Dt, FPBDIslandSolverData& SolverData)
 {
 	for (int32 ConstraintIndex : SolverData.GetConstraintIndices(ContainerId))
diff --git a/Engine/Source/Runtime/Experimental/Chaos/Private/Chaos/PBDRigidSpringConstraints.cpp b/Engine/Source/Runtime/Experimental/Chaos/Private/Chaos/PBDRigidSpringConstraints.cpp
--- a/Engine/Source/Runtime/Experimental/Chaos/Private/Chaos/PBDRigidSpringConstraints.cpp
+++ b/Engine/Source/Runtime/Experimental/Chaos/Private/Chaos/PBDRigidSpringConstraints.cpp
@@ -37,6 +37,11 @@ namespace Chaos
 		ConcreteContainer()->SetRestLength(ConstraintIndex, SpringLength);
 	}
 
+	void FPBDRigidSpringConstraintHandle::PreGatherInput(const FReal Dt, FPBDIslandSolverData& SolverData)
+	{
+		ConcreteContainer()->PreGatherInput(Dt, ConstraintIndex, SolverData);
+	}
+
 	void FPBDRigidSpringConstraintHandle::GatherInput(const FReal Dt, const int32 Particle0Level, const int32 Particle1Level, FPBDIslandSolverData& SolverData)
 	{
 		ConcreteContainer()->GatherInput(Dt, ConstraintIndex, Particle0Level, Particle1Level, SolverData);
@@ -130,6 +135,14 @@ namespace Chaos
 		return SpringSettings[ConstraintIndex].Stiffness * Delta / CombinedInvMass;
 	}
 
+	void FPBDRigidSpringConstraints::PreGatherInput(const FReal Dt, FPBDIslandSolverData& SolverData)
+	{
+		for (int32 ConstraintIndex = 0; ConstraintIndex < NumConstraints(); ++ConstraintIndex)
+		{
+			PreGatherInput(Dt, ConstraintIndex, SolverData);
+		}
+	}
+
 	void FPBDRigidSpringConstraints::GatherInput(const FReal Dt, FPBDIslandSolverData& SolverData)
 	{
 		for (int32 ConstraintIndex = 0; ConstraintIndex < NumConstraints(); ++ConstraintIndex)
@@ -148,7 +161,7 @@ namespace Chaos
 		SolverData.GetConstraintIndices(ContainerId).Reset(NumIslandConstraints);
 	}
 
-	void FPBDRigidSpringConstraints::GatherInput(const FReal Dt, const int32 ConstraintIndex, const int32 Particle0Level, const int32 Particle1Level, FPBDIslandSolverData& SolverData)
+	void FPBDRigidSpringConstraints::PreGatherInput(const FReal Dt, const int32 ConstraintIndex, FPBDIslandSolverData& SolverData)
 	{
 		SolverData.GetConstraintIndices(ContainerId).Add(ConstraintIndex);
 	
@@ -159,6 +172,10 @@ namespace Chaos
 		};
 	}
 
+	void FPBDRigidSpringConstraints::GatherInput(const FReal Dt, const int32 ConstraintIndex, const int32 Particle0Level, const int32 Particle1Level, FPBDIslandSolverData& SolverData)
+	{
+	}
+
 	void FPBDRigidSpringConstraints::ScatterOutput(FReal Dt, FPBDIslandSolverData& SolverData)
 	{
 		for (int32 ConstraintIndex : SolverData.GetConstraintIndices(ContainerId))
diff --git a/Engine/Source/Runtime/Experimental/Chaos/Private/Chaos/PBDSuspensionConstraints.cpp b/Engine/Source/Runtime/Experimental/Chaos/Private/Chaos/PBDSuspensionConstraints.cpp
--- a/Engine/Source/Runtime/Experimental/Chaos/Private/Chaos/PBDSuspensionConstraints.cpp
+++ b/Engine/Source/Runtime/Experimental/Chaos/Private/Chaos/PBDSuspensionConstraints.cpp
@@ -54,12 +54,16 @@ namespace Chaos
 		return ConcreteContainer()->GetConstrainedParticles(ConstraintIndex);
 	}
 
+	void FPBDSuspensionConstraintHandle::PreGatherInput(const FReal Dt, FPBDIslandSolverData& SolverData)
+	{
+		ConcreteContainer()->PreGatherInput(Dt, ConstraintIndex, SolverData);
+	}
+
 	void FPBDSuspensionConstraintHandle::GatherInput(const FReal Dt, const int32 Particle0Level, const int32 Particle1Level, FPBDIslandSolverData& SolverData)
 	{
 		ConcreteContainer()->GatherInput(Dt, ConstraintIndex, Particle0Level, Particle1Level, SolverData);
 	}
 
-
 	Chaos::FPBDSuspensionConstraints::FConstraintContainerHandle* FPBDSuspensionConstraints::AddConstraint(TGeometryParticleHandle<FReal, 3>* Particle, const FVec3& InSuspensionLocalOffset, const FPBDSuspensionSettings& InConstraintSettings)
 	{
 		int32 NewIndex = ConstrainedParticles.Num();
@@ -118,7 +122,8 @@ namespace Chaos
 		SolverData.GetConstraintIndices(ContainerId).Reset(NumIslandConstraints);
 	}
 
-	void FPBDSuspensionConstraints::GatherInput(const FReal Dt, const int32 ConstraintIndex, const int32 Particle0Level, const int32 Particle1Level, FPBDIslandSolverData& SolverData)
+	// @todo: This does per body work so can't be safely parallelized
+	void FPBDSuspensionConstraints::PreGatherInput(const FReal Dt, const int32 ConstraintIndex, FPBDIslandSolverData& SolverData)
 	{
 		SolverData.GetConstraintIndices(ContainerId).Add(ConstraintIndex);
 
@@ -209,6 +214,10 @@ namespace Chaos
 
 	}
 
+	void FPBDSuspensionConstraints::GatherInput(const FReal Dt, const int32 ConstraintIndex, const int32 Particle0Level, const int32 Particle1Level, FPBDIslandSolverData& SolverData)
+	{
+	}
+
 	void FPBDSuspensionConstraints::ScatterOutput(FReal Dt, FPBDIslandSolverData& SolverData)
 	{
 		for (int32 ConstraintIndex : SolverData.GetConstraintIndices(ContainerId))
diff --git a/Engine/Source/Runtime/Experimental/Chaos/Private/EventDefaults.cpp b/Engine/Source/Runtime/Experimental/Chaos/Private/EventDefaults.cpp
--- a/Engine/Source/Runtime/Experimental/Chaos/Private/EventDefaults.cpp
+++ b/Engine/Source/Runtime/Experimental/Chaos/Private/EventDefaults.cpp
@@ -170,7 +170,7 @@ namespace Chaos
 								}
 							}
 						}
-					}, 100);
+					}, Chaos::LargeBatchSize);
 					for (int32 Index = 0; Index < CollisionHandles.Num(); ++Index)
 					{
 						if (ValidArray[Index])
@@ -272,7 +272,7 @@ namespace Chaos
 #endif
 								}
 							}
-						}, 10);
+						}, Chaos::SmallBatchSize);
 						int32 IdxCollision = 0, IdxCollision2 = NumValidCollisions - 1, NumDeleted = 0;
 						while (IdxCollision <= IdxCollision2)
 						{
@@ -299,7 +299,7 @@ namespace Chaos
 								IdxCollision++;
 							}
 						}
-						AllCollisionsDataArray.SetNum(NumValidCollisions - NumDeleted);
+						AllCollisionsDataArray.SetNum(NumValidCollisions - NumDeleted, false);
 					}
 				}
 			}
diff --git a/Engine/Source/Runtime/Experimental/Chaos/Public/Chaos/Collision/PBDCollisionConstraintHandle.h b/Engine/Source/Runtime/Experimental/Chaos/Public/Chaos/Collision/PBDCollisionConstraintHandle.h
--- a/Engine/Source/Runtime/Experimental/Chaos/Public/Chaos/Collision/PBDCollisionConstraintHandle.h
+++ b/Engine/Source/Runtime/Experimental/Chaos/Public/Chaos/Collision/PBDCollisionConstraintHandle.h
@@ -87,7 +87,7 @@ namespace Chaos
 
 		TVector<TGeometryParticleHandle<FReal, 3>*, 2> GetConstrainedParticles();
 
-		void PreGatherInput(FPBDIslandSolverData& SolverData);
+		void PreGatherInput(const FReal Dt, FPBDIslandSolverData& SolverData);
 		void GatherInput(FReal Dt, const int32 Particle0Level, const int32 Particle1Level, FPBDIslandSolverData& SolverData);
 
 		FSolverBody* GetSolverBody0();
@@ -103,4 +103,4 @@ namespace Chaos
 		}
 	};
 
-}
\ No newline at end of file
+}
diff --git a/Engine/Source/Runtime/Experimental/Chaos/Public/Chaos/Framework/Parallel.h b/Engine/Source/Runtime/Experimental/Chaos/Public/Chaos/Framework/Parallel.h
--- a/Engine/Source/Runtime/Experimental/Chaos/Public/Chaos/Framework/Parallel.h
+++ b/Engine/Source/Runtime/Experimental/Chaos/Public/Chaos/Framework/Parallel.h
@@ -14,6 +14,8 @@ namespace Chaos
 
 
 	CHAOS_API extern int32 MaxNumWorkers;
+	CHAOS_API extern int32 SmallBatchSize;
+	CHAOS_API extern int32 LargeBatchSize;
 #if UE_BUILD_SHIPPING
 	const bool bDisablePhysicsParallelFor = false;
 	const bool bDisableParticleParallelFor = false;
diff --git a/Engine/Source/Runtime/Experimental/Chaos/Public/Chaos/PBDJointConstraints.h b/Engine/Source/Runtime/Experimental/Chaos/Public/Chaos/PBDJointConstraints.h
--- a/Engine/Source/Runtime/Experimental/Chaos/Public/Chaos/PBDJointConstraints.h
+++ b/Engine/Source/Runtime/Experimental/Chaos/Public/Chaos/PBDJointConstraints.h
@@ -56,6 +56,7 @@ namespace Chaos
 		void SetSettings(const FPBDJointSettings& Settings);
 		TVec2<FGeometryParticleHandle*> GetConstrainedParticles() const;
 
+		void PreGatherInput(const FReal Dt, FPBDIslandSolverData& SolverData);
 		void GatherInput(const FReal Dt, const int32 Particle0Level, const int32 Particle1Level, FPBDIslandSolverData& SolverData);
 
 		static const FConstraintHandleTypeID& StaticType()
@@ -275,6 +276,7 @@ namespace Chaos
 		//
 
 		void SetNumIslandConstraints(const int32 NumIslandConstraints, FPBDIslandSolverData& SolverData);
+		void PreGatherInput(const FReal Dt, const int32 ConstraintIndex, FPBDIslandSolverData& SolverData);
 		void GatherInput(const FReal Dt, const int32 ConstraintIndex, const int32 Particle0Level, const int32 Particle1Level, FPBDIslandSolverData& SolverData);
 		bool ApplyPhase1Serial(const FReal Dt, const int32 It, const int32 NumIts, FPBDIslandSolverData& SolverData);
 		bool ApplyPhase2Serial(const FReal Dt, const int32 It, const int32 NumIts, FPBDIslandSolverData& SolverData);
@@ -332,6 +334,8 @@ namespace Chaos
 		TArray<FPBDJointCachedSolver> CachedConstraintSolvers;
 
 		EConstraintSolverType SolverType;
+
+
 	};
 
 }
diff --git a/Engine/Source/Runtime/Experimental/Chaos/Public/Chaos/PBDPositionConstraints.h b/Engine/Source/Runtime/Experimental/Chaos/Public/Chaos/PBDPositionConstraints.h
--- a/Engine/Source/Runtime/Experimental/Chaos/Public/Chaos/PBDPositionConstraints.h
+++ b/Engine/Source/Runtime/Experimental/Chaos/Public/Chaos/PBDPositionConstraints.h
@@ -24,6 +24,7 @@ namespace Chaos
 			: TIndexedContainerConstraintHandle<FPBDPositionConstraints>(InConstraintContainer, InConstraintIndex) {}
 		TVector<FGeometryParticleHandle*, 2> GetConstrainedParticles() const;
 
+		void PreGatherInput(const FReal Dt, FPBDIslandSolverData& SolverData);
 		void GatherInput(const FReal Dt, const int32 Particle0Level, const int32 Particle1Level, FPBDIslandSolverData& SolverData);
 
 		static const FConstraintHandleTypeID& StaticType()
@@ -186,6 +187,7 @@ namespace Chaos
 		void UpdatePositionBasedState(const FReal Dt) {}
 
 		void SetNumIslandConstraints(const int32 NumIslandConstraints, FPBDIslandSolverData& SolverData);
+		void PreGatherInput(const FReal Dt, const int32 ConstraintIndex, FPBDIslandSolverData& SolverData);
 		void GatherInput(const FReal Dt, const int32 ConstraintIndex, const int32 Particle0Level, const int32 Particle1Level, FPBDIslandSolverData& SolverData);
 		void ScatterOutput(FReal Dt, FPBDIslandSolverData& SolverData);
 
diff --git a/Engine/Source/Runtime/Experimental/Chaos/Public/Chaos/PBDRigidDynamicSpringConstraints.h b/Engine/Source/Runtime/Experimental/Chaos/Public/Chaos/PBDRigidDynamicSpringConstraints.h
--- a/Engine/Source/Runtime/Experimental/Chaos/Public/Chaos/PBDRigidDynamicSpringConstraints.h
+++ b/Engine/Source/Runtime/Experimental/Chaos/Public/Chaos/PBDRigidDynamicSpringConstraints.h
@@ -29,6 +29,7 @@ namespace Chaos
 
 		TVec2<FGeometryParticleHandle*> GetConstrainedParticles() const;
 
+		void PreGatherInput(const FReal Dt, FPBDIslandSolverData& SolverData);
 		void GatherInput(const FReal Dt, const int32 Particle0Level, const int32 Particle1Level, FPBDIslandSolverData& SolverData);
 
 		static const FConstraintHandleTypeID& StaticType()
@@ -207,6 +208,7 @@ namespace Chaos
 		void UpdatePositionBasedState(const FReal Dt);
 
 		void SetNumIslandConstraints(const int32 NumIslandConstraints, FPBDIslandSolverData& SolverData);
+		void PreGatherInput(const FReal Dt, const int32 ConstraintIndex, FPBDIslandSolverData& SolverData);
 		void GatherInput(const FReal Dt, const int32 ConstraintIndex, const int32 Particle0Level, const int32 Particle1Level, FPBDIslandSolverData& SolverData);
 		void ScatterOutput(FReal Dt, FPBDIslandSolverData& SolverData);
 
diff --git a/Engine/Source/Runtime/Experimental/Chaos/Public/Chaos/PBDRigidSpringConstraints.h b/Engine/Source/Runtime/Experimental/Chaos/Public/Chaos/PBDRigidSpringConstraints.h
--- a/Engine/Source/Runtime/Experimental/Chaos/Public/Chaos/PBDRigidSpringConstraints.h
+++ b/Engine/Source/Runtime/Experimental/Chaos/Public/Chaos/PBDRigidSpringConstraints.h
@@ -37,6 +37,7 @@ namespace Chaos
 		FReal GetRestLength() const;
 		void SetRestLength(const FReal SpringLength);
 
+		void PreGatherInput(const FReal Dt, FPBDIslandSolverData& SolverData);
 		void GatherInput(const FReal Dt, const int32 Particle0Level, const int32 Particle1Level, FPBDIslandSolverData& SolverData);
 
 		static const FConstraintHandleTypeID& StaticType()
@@ -162,6 +163,7 @@ namespace Chaos
 		// Simple Rule API
 		//
 
+		void PreGatherInput(const FReal Dt, FPBDIslandSolverData& SolverData);
 		void GatherInput(const FReal Dt, FPBDIslandSolverData& SolverData);
 		void ScatterOutput(const FReal Dt, FPBDIslandSolverData& SolverData);
 
@@ -178,6 +180,7 @@ namespace Chaos
 		void UpdatePositionBasedState(const FReal Dt) {}
 
 		void SetNumIslandConstraints(const int32 NumIslandConstraints, FPBDIslandSolverData& SolverData);
+		void PreGatherInput(const FReal Dt, const int32 ConstraintIndex, FPBDIslandSolverData& SolverData);
 		void GatherInput(const FReal Dt, const int32 ConstraintIndex, const int32 Particle0Level, const int32 Particle1Level, FPBDIslandSolverData& SolverData);
 
 		bool ApplyPhase1Serial(const FReal Dt, const int32 It, const int32 NumIts, FPBDIslandSolverData& SolverData);
diff --git a/Engine/Source/Runtime/Experimental/Chaos/Public/Chaos/PBDSuspensionConstraints.h b/Engine/Source/Runtime/Experimental/Chaos/Public/Chaos/PBDSuspensionConstraints.h
--- a/Engine/Source/Runtime/Experimental/Chaos/Public/Chaos/PBDSuspensionConstraints.h
+++ b/Engine/Source/Runtime/Experimental/Chaos/Public/Chaos/PBDSuspensionConstraints.h
@@ -35,6 +35,7 @@ namespace Chaos
 
 		TVec2<FGeometryParticleHandle*> GetConstrainedParticles() const;
 
+		void PreGatherInput(const FReal Dt, FPBDIslandSolverData& SolverData);
 		void GatherInput(const FReal Dt, const int32 Particle0Level, const int32 Particle1Level, FPBDIslandSolverData& SolverData);
 
 		static const FConstraintHandleTypeID& StaticType()
@@ -233,6 +234,7 @@ namespace Chaos
 		void UpdatePositionBasedState(const FReal Dt) {}
 
 		void SetNumIslandConstraints(const int32 NumIslandConstraints, FPBDIslandSolverData& SolverData);
+		void PreGatherInput(const FReal Dt, const int32 ConstraintIndex, FPBDIslandSolverData& SolverData);
 		void GatherInput(const FReal Dt, const int32 ConstraintIndex, const int32 Particle0Level, const int32 Particle1Level, FPBDIslandSolverData& SolverData);
 		void ScatterOutput(FReal Dt, FPBDIslandSolverData& SolverData);
 
