Move some spammy physics lock stats to the PhysicsVerbose category, which is off by default.

diff --git a/Engine/Source/Runtime/Experimental/Chaos/Public/Framework/Threading.h b/Engine/Source/Runtime/Experimental/Chaos/Public/Framework/Threading.h
--- a/Engine/Source/Runtime/Experimental/Chaos/Public/Framework/Threading.h
+++ b/Engine/Source/Runtime/Experimental/Chaos/Public/Framework/Threading.h
@@ -697,7 +697,7 @@ FORCEINLINE void EnsureIsInGameThreadContext()
 		TPhysicsSceneGuardScopedWrite(MutexType& InMutex)
 			: Mutex(InMutex)
 		{
-			CSV_SCOPED_TIMING_STAT_EXCLUSIVE(AcquireSceneWriteLock);
+			CSV_SCOPED_TIMING_STAT(PhysicsVerbose,AcquireSceneWriteLock);
 			Mutex.WriteLock();
 		}
 
@@ -726,7 +726,7 @@ FORCEINLINE void EnsureIsInGameThreadContext()
 		TPhysicsSceneGuardScopedRead(MutexType& InMutex)
 			: Mutex(InMutex)
 		{
-			CSV_SCOPED_TIMING_STAT_EXCLUSIVE(AcquireSceneReadLock);
+			CSV_SCOPED_TIMING_STAT(PhysicsVerbose, AcquireSceneReadLock);
 			Mutex.ReadLock();
 		}
 
