Check for div by zero when threading is disabled.

diff --git a/Engine/Source/Runtime/Experimental/Chaos/Private/Chaos/Framework/Parallel.cpp b/Engine/Source/Runtime/Experimental/Chaos/Private/Chaos/Framework/Parallel.cpp
--- a/Engine/Source/Runtime/Experimental/Chaos/Private/Chaos/Framework/Parallel.cpp
+++ b/Engine/Source/Runtime/Experimental/Chaos/Private/Chaos/Framework/Parallel.cpp
@@ -34,6 +34,10 @@ namespace Chaos
 void Chaos::InnerPhysicsParallelFor(int32 Num, TFunctionRef<void(int32)> InCallable, bool bForceSingleThreaded)
 {
 	int32 NumWorkers = int32(LowLevelTasks::FScheduler::Get().GetNumWorkers());
+	if (NumWorkers == 0)
+	{
+		NumWorkers = 1;
+	}
 	int32 BatchSize = FMath::DivideAndRoundUp<int32>(Num, NumWorkers);
 	PhysicsParallelFor(Num, InCallable, (BatchSize > InnerParallelForBatchSize) ? bForceSingleThreaded : true);
 }
@@ -41,6 +45,10 @@ void Chaos::InnerPhysicsParallelFor(int32 Num, TFunctionRef<void(int32)> InCalla
 void Chaos::InnerPhysicsParallelForRange(int32 InNum, TFunctionRef<void(int32, int32)> InCallable, const int32 InMinBatchSize, bool bForceSingleThreaded)
 {
 	int32 NumWorkers = int32(LowLevelTasks::FScheduler::Get().GetNumWorkers());
+	if (NumWorkers == 0)
+	{
+		NumWorkers = 1;
+	}
 	int32 BatchSize = FMath::DivideAndRoundUp<int32>(InNum, NumWorkers);
 	PhysicsParallelForRange(InNum, InCallable, InMinBatchSize, (BatchSize > InnerParallelForBatchSize) ? bForceSingleThreaded : true);
 }
