Chaos - fix RBAN joints pre-gather

diff --git a/Engine/Source/Runtime/Experimental/Chaos/Private/Chaos/PBDCollisionConstraints.cpp b/Engine/Source/Runtime/Experimental/Chaos/Private/Chaos/PBDCollisionConstraints.cpp
--- a/Engine/Source/Runtime/Experimental/Chaos/Private/Chaos/PBDCollisionConstraints.cpp
+++ b/Engine/Source/Runtime/Experimental/Chaos/Private/Chaos/PBDCollisionConstraints.cpp
@@ -397,7 +397,7 @@ namespace Chaos
 		}
 	}
 
-	void FPBDCollisionConstraints::GatherInput(const FReal Dt, FPBDIslandSolverData& SolverData)
+	void FPBDCollisionConstraints::PreGatherInput(const FReal Dt, FPBDIslandSolverData& SolverData)
 	{
 		if (SolverType == EConstraintSolverType::QuasiPbd)
 		{
@@ -406,6 +406,19 @@ namespace Chaos
 				if (Constraint->IsEnabled())
 				{
 					PreGatherInput(*Constraint, SolverData);
+				}
+			}
+		}
+	}
+
+	void FPBDCollisionConstraints::GatherInput(const FReal Dt, FPBDIslandSolverData& SolverData)
+	{
+		if (SolverType == EConstraintSolverType::QuasiPbd)
+		{
+			for (FPBDCollisionConstraint* Constraint : GetConstraints())
+			{
+				if (Constraint->IsEnabled())
+				{
 					GatherInput(Dt, *Constraint, INDEX_NONE, INDEX_NONE, SolverData);
 				}
 			}
diff --git a/Engine/Source/Runtime/Experimental/Chaos/Private/Chaos/PBDConstraintRule.cpp b/Engine/Source/Runtime/Experimental/Chaos/Private/Chaos/PBDConstraintRule.cpp
--- a/Engine/Source/Runtime/Experimental/Chaos/Private/Chaos/PBDConstraintRule.cpp
+++ b/Engine/Source/Runtime/Experimental/Chaos/Private/Chaos/PBDConstraintRule.cpp
@@ -81,6 +81,7 @@ namespace Chaos
 		{
 			SolverData->GetConstraintIndex(Constraints.GetContainerId()) = 0;
 			Constraints.SetNumIslandConstraints(Constraints.NumConstraints(), *SolverData);
+			Constraints.PreGatherInput(Dt, *SolverData);
 			Constraints.GatherInput(Dt, *SolverData);
 		}
 	}
diff --git a/Engine/Source/Runtime/Experimental/Chaos/Private/Chaos/PBDJointConstraints.cpp b/Engine/Source/Runtime/Experimental/Chaos/Private/Chaos/PBDJointConstraints.cpp
--- a/Engine/Source/Runtime/Experimental/Chaos/Private/Chaos/PBDJointConstraints.cpp
+++ b/Engine/Source/Runtime/Experimental/Chaos/Private/Chaos/PBDJointConstraints.cpp
@@ -181,7 +181,6 @@ namespace Chaos
 
 	void FPBDJointConstraintHandle::GatherInput(const FReal Dt, const int32 Particle0Level, const int32 Particle1Level, FPBDIslandSolverData& SolverData)
 	{
-
 		ConcreteContainer()->GatherInput(Dt, ConstraintIndex, Particle0Level, Particle1Level, SolverData);
 	}
 
@@ -819,6 +818,17 @@ namespace Chaos
 	//
 	//////////////////////////////////////////////////////////////////////////
 
+	void FPBDJointConstraints::PreGatherInput(const FReal Dt, FPBDIslandSolverData& SolverData)
+	{
+		for (int32 ConstraintIndex = 0; ConstraintIndex < NumConstraints(); ++ConstraintIndex)
+		{
+			if (!ConstraintStates[ConstraintIndex].bDisabled)
+			{
+				PreGatherInput(Dt, ConstraintIndex, SolverData);
+			}
+		}
+	}
+
 	void FPBDJointConstraints::GatherInput(const FReal Dt, FPBDIslandSolverData& SolverData)
 	{
 		for (int32 ConstraintIndex = 0; ConstraintIndex < NumConstraints(); ++ConstraintIndex)
diff --git a/Engine/Source/Runtime/Experimental/Chaos/Public/Chaos/PBDCollisionConstraints.h b/Engine/Source/Runtime/Experimental/Chaos/Public/Chaos/PBDCollisionConstraints.h
--- a/Engine/Source/Runtime/Experimental/Chaos/Public/Chaos/PBDCollisionConstraints.h
+++ b/Engine/Source/Runtime/Experimental/Chaos/Public/Chaos/PBDCollisionConstraints.h
@@ -154,6 +154,7 @@ public:
 	// Simple Rule API
 	//
 
+	void PreGatherInput(const FReal Dt, FPBDIslandSolverData& SolverData);
 	void GatherInput(const FReal Dt, FPBDIslandSolverData& SolverData);
 	void ScatterOutput(const FReal Dt, FPBDIslandSolverData& SolverData);
 
diff --git a/Engine/Source/Runtime/Experimental/Chaos/Public/Chaos/PBDJointConstraints.h b/Engine/Source/Runtime/Experimental/Chaos/Public/Chaos/PBDJointConstraints.h
--- a/Engine/Source/Runtime/Experimental/Chaos/Public/Chaos/PBDJointConstraints.h
+++ b/Engine/Source/Runtime/Experimental/Chaos/Public/Chaos/PBDJointConstraints.h
@@ -264,6 +264,7 @@ namespace Chaos
 		// Simple Rule API
 		//
 
+		void PreGatherInput(const FReal Dt, FPBDIslandSolverData& SolverData);
 		void GatherInput(const FReal Dt, FPBDIslandSolverData& SolverData);
 		void ScatterOutput(const FReal Dt, FPBDIslandSolverData& SolverData);
 
