Fix collision update after CCD solve

diff --git a/Engine/Source/Runtime/Experimental/Chaos/Private/Chaos/CCDUtilities.cpp b/Engine/Source/Runtime/Experimental/Chaos/Private/Chaos/CCDUtilities.cpp
--- a/Engine/Source/Runtime/Experimental/Chaos/Private/Chaos/CCDUtilities.cpp
+++ b/Engine/Source/Runtime/Experimental/Chaos/Private/Chaos/CCDUtilities.cpp
@@ -566,15 +566,13 @@ namespace Chaos
 	{
 		// We need to update the world-space contact points at the final locations
 		// @todo(chaos): parallelize this code
+		// @todo(chaos): These SweptConstraints might contain non-CCD particles and those non-CCD particles might collide with other non-CCD particles, which are modeled in normal collision constraints. Those normal collision constraints might need to be updated as well.
 		for (FPBDCollisionConstraint* SweptConstraint : SweptConstraints)
 		{
-			FRigidTransform3 ShapeWorldTransform0 = SweptConstraint->GetShapeWorldTransform0();
-			FRigidTransform3 ShapeWorldTransform1 = SweptConstraint->GetShapeWorldTransform1();
-			ShapeWorldTransform0.SetTranslation(FConstGenericParticleHandle(SweptConstraint->GetParticle0())->P());
-			ShapeWorldTransform1.SetTranslation(FConstGenericParticleHandle(SweptConstraint->GetParticle1())->P());
-
-			SweptConstraint->SetShapeWorldTransforms(ShapeWorldTransform0, ShapeWorldTransform1);
-			SweptConstraint->UpdateManifoldContacts();
+			const FConstGenericParticleHandle P0 = FConstGenericParticleHandle(SweptConstraint->GetParticle0());
+			const FConstGenericParticleHandle P1 = FConstGenericParticleHandle(SweptConstraint->GetParticle1());
+			SweptConstraint->ResetManifold();
+			Collisions::UpdateConstraintFromGeometry<ECollisionUpdateType::Deepest>(*SweptConstraint, FRigidTransform3(P0->P(), P0->Q()), FRigidTransform3(P1->P(), P1->Q()), Dt);
 
 			// @todo(zhenglin): Removing constraints that has Phi larger than CullDistance could reduce the island sizes in the normal solve. But I could not get this to work...
 			// if (SweptConstraint->GetPhi() > SweptConstraint->GetCullDistance())
