Chaos : Fix slow removal of geometry collection particle
- remove the secial remove behavior for GC and clustered particle as this is not a requirement and causing massive slowdown when removing thousands of those particles from solver

diff --git a/Engine/Source/Runtime/Experimental/Chaos/Public/Chaos/PBDRigidsSOAs.h b/Engine/Source/Runtime/Experimental/Chaos/Public/Chaos/PBDRigidsSOAs.h
--- a/Engine/Source/Runtime/Experimental/Chaos/Public/Chaos/PBDRigidsSOAs.h
+++ b/Engine/Source/Runtime/Experimental/Chaos/Public/Chaos/PBDRigidsSOAs.h
@@ -208,10 +208,8 @@ public:
 		DynamicKinematicParticles = MakeUnique<FPBDRigidParticles>();
 
 		ClusteredParticles = MakeUnique< FPBDRigidClusteredParticles>();
-		ClusteredParticles->RemoveParticleBehavior() = ERemoveParticleBehavior::Remove;	//clustered particles maintain relative ordering
 
 		GeometryCollectionParticles = MakeUnique<TPBDGeometryCollectionParticles<FReal, 3>>();
-		GeometryCollectionParticles->RemoveParticleBehavior() = ERemoveParticleBehavior::Remove;	//clustered particles maintain relative ordering
 		bGeometryCollectionDirty = false;
 
 		UpdateViews();
