Partially revert only the event part. This fixed EdTest but breaks events. Still working on a real fix.

diff --git a/Engine/Source/Runtime/Experimental/Chaos/Private/EventDefaults.cpp b/Engine/Source/Runtime/Experimental/Chaos/Private/EventDefaults.cpp
--- a/Engine/Source/Runtime/Experimental/Chaos/Private/EventDefaults.cpp
+++ b/Engine/Source/Runtime/Experimental/Chaos/Private/EventDefaults.cpp
@@ -182,124 +182,103 @@ namespace Chaos
 					}
 				}
 
+				ValidCollisionHandles.SetNum(NumValidCollisions);
+
+				if(ValidCollisionHandles.Num() > 0)
 				{
-					SCOPE_CYCLE_COUNTER(STAT_Events_RegisterCollisionEvent_Notify);
-					ValidCollisionHandles.SetNum(NumValidCollisions);
-					if (ValidCollisionHandles.Num() > 0)
+					for (int32 IdxCollision = 0; IdxCollision < ValidCollisionHandles.Num(); ++IdxCollision)
 					{
-						AllCollisionsDataArray.SetNum(NumValidCollisions);
-						InnerPhysicsParallelForRange(ValidCollisionHandles.Num(), [&](int32 StartRangeIndex, int32 EndRangeIndex)
-						{
-							for (int32 IdxCollision = StartRangeIndex; IdxCollision < EndRangeIndex; ++IdxCollision)
-							{
-								const FPBDCollisionConstraint& Constraint = ValidCollisionHandles[IdxCollision]->GetContact();
+						const FPBDCollisionConstraint& Constraint = ValidCollisionHandles[IdxCollision]->GetContact();
 
-								const FGeometryParticleHandle* Particle0 = Constraint.GetParticle0();
-								const FGeometryParticleHandle* Particle1 = Constraint.GetParticle1();
+						const FGeometryParticleHandle* Particle0 = Constraint.GetParticle0();
+						const FGeometryParticleHandle* Particle1 = Constraint.GetParticle1();
 
-								FCollidingData Data;
-								Data.Location = Constraint.CalculateWorldContactLocation();
-								Data.AccumulatedImpulse = Constraint.AccumulatedImpulse;
-								Data.Normal = Constraint.CalculateWorldContactNormal();
-								Data.PenetrationDepth = Constraint.GetPhi();
+						FCollidingData Data;
+						Data.Location = Constraint.CalculateWorldContactLocation();
+						Data.AccumulatedImpulse = Constraint.AccumulatedImpulse;
+						Data.Normal = Constraint.CalculateWorldContactNormal();
+						Data.PenetrationDepth = Constraint.GetPhi();
+						
+						// @todo(chaos): fix this casting
+						Data.Proxy1 = Particle0 ? const_cast<IPhysicsProxyBase*>(Particle0->PhysicsProxy()) : nullptr;
+						Data.Proxy2 = Particle1 ? const_cast<IPhysicsProxyBase*>(Particle1->PhysicsProxy()) : nullptr;
 
-								// @todo(chaos): fix this casting
-								Data.Proxy1 = Particle0 ? const_cast<IPhysicsProxyBase*>(Particle0->PhysicsProxy()) : nullptr;
-								Data.Proxy2 = Particle1 ? const_cast<IPhysicsProxyBase*>(Particle1->PhysicsProxy()) : nullptr;
+						const FPerShapeData* Shape0 = Constraint.GetShape0();
+						const FPerShapeData* Shape1 = Constraint.GetShape1();
 
-								const FPerShapeData* Shape0 = Constraint.GetShape0();
-								const FPerShapeData* Shape1 = Constraint.GetShape1();
+						Data.Mat1 = ResolveMaterial(Shape0, Constraint);
+						Data.Mat2 = ResolveMaterial(Shape1, Constraint);
 
-								Data.Mat1 = ResolveMaterial(Shape0, Constraint);
-								Data.Mat2 = ResolveMaterial(Shape1, Constraint);
+						// Collision constraints require both proxies are valid. If either is not, we needn't record the collision event.
+						if (Data.Proxy1 == nullptr || Data.Proxy2 == nullptr)
+						{
+							continue;
+						}
 
-								// Collision constraints require both proxies are valid. If either is not, we needn't record the collision event.
-								if (Data.Proxy1 == nullptr || Data.Proxy2 == nullptr)
-								{
-									continue;
-								}
+						if (const FPBDRigidParticleHandle * Rigid0 = Particle0->CastToRigidParticle())
+						{
+							Data.DeltaVelocity1 = Rigid0->V() - Rigid0->PreV();
+						}
+						if (const FPBDRigidParticleHandle * Rigid1 = Particle1->CastToRigidParticle())
+						{
+							Data.DeltaVelocity2 = Rigid1->V() - Rigid1->PreV();
+						}
 
-								if (const FPBDRigidParticleHandle* Rigid0 = Particle0->CastToRigidParticle())
-								{
-									Data.DeltaVelocity1 = Rigid0->V() - Rigid0->PreV();
-								}
-								if (const FPBDRigidParticleHandle* Rigid1 = Particle1->CastToRigidParticle())
-								{
-									Data.DeltaVelocity2 = Rigid1->V() - Rigid1->PreV();
-								}
+						// todo: do we need these anymore now we are storing the particles you can access all of this stuff from there
+						// do we still need these now we have pointers to particles returned?
+						const FPBDRigidParticleHandle* PBDRigid0 = Particle0->CastToRigidParticle();
+						if (PBDRigid0 && PBDRigid0->ObjectState() == EObjectStateType::Dynamic)
+						{
+							Data.Velocity1 = PBDRigid0->V();
+							Data.AngularVelocity1 = PBDRigid0->W();
+							Data.Mass1 = PBDRigid0->M();
+						}
 
-								// todo: do we need these anymore now we are storing the particles you can access all of this stuff from there
-								// do we still need these now we have pointers to particles returned?
-								const FPBDRigidParticleHandle* PBDRigid0 = Particle0->CastToRigidParticle();
-								if (PBDRigid0 && PBDRigid0->ObjectState() == EObjectStateType::Dynamic)
-								{
-									Data.Velocity1 = PBDRigid0->V();
-									Data.AngularVelocity1 = PBDRigid0->W();
-									Data.Mass1 = PBDRigid0->M();
-								}
+						const FPBDRigidParticleHandle* PBDRigid1 = Particle1->CastToRigidParticle();
+						if (PBDRigid1 && PBDRigid1->ObjectState() == EObjectStateType::Dynamic)
+						{
+							Data.Velocity2 = PBDRigid1->V();
+							Data.AngularVelocity2 = PBDRigid1->W();
+							Data.Mass2 = PBDRigid1->M();
+						}
 
-								const FPBDRigidParticleHandle* PBDRigid1 = Particle1->CastToRigidParticle();
-								if (PBDRigid1 && PBDRigid1->ObjectState() == EObjectStateType::Dynamic)
-								{
-									Data.Velocity2 = PBDRigid1->V();
-									Data.AngularVelocity2 = PBDRigid1->W();
-									Data.Mass2 = PBDRigid1->M();
-								}
+						IPhysicsProxyBase* const PhysicsProxy = const_cast<IPhysicsProxyBase*>(Particle0->PhysicsProxy());
+						IPhysicsProxyBase* const OtherPhysicsProxy = const_cast<IPhysicsProxyBase*>(Particle1->PhysicsProxy());
+							
+						const FSolverCollisionEventFilter* SolverCollisionEventFilter = Solver->GetEventFilters()->GetCollisionFilter();
+						if (!SolverCollisionEventFilter->Enabled() || SolverCollisionEventFilter->Pass(Data))
 
-								IPhysicsProxyBase* const PhysicsProxy = const_cast<IPhysicsProxyBase*>(Particle0->PhysicsProxy());
-								IPhysicsProxyBase* const OtherPhysicsProxy = const_cast<IPhysicsProxyBase*>(Particle1->PhysicsProxy());
+						{
+							const int32 NewIdx = AllCollisionsDataArray.Add(FCollidingData());
+							FCollidingData& CollisionDataArrayItem = AllCollisionsDataArray[NewIdx];
 
-								const FSolverCollisionEventFilter* SolverCollisionEventFilter = Solver->GetEventFilters()->GetCollisionFilter();
-								if (!SolverCollisionEventFilter->Enabled() || SolverCollisionEventFilter->Pass(Data))
+							CollisionDataArrayItem = Data;
 
-								{
-									AllCollisionsDataArray[IdxCollision] = Data;
 #if TODO_REIMPLEMENT_RIGID_CLUSTERING
-									// If Constraint.ParticleIndex is a cluster store an index for a mesh in this cluster
-									if (ClusterIdsArray[Constraint.ParticleIndex].NumChildren > 0)
-									{
-										int32 ParticleIndexMesh = GetParticleIndexMesh(ParentToChildrenMap, Constraint.ParticleIndex);
-										ensure(ParticleIndexMesh != INDEX_NONE);
-										CollisionDataArrayItem.ParticleIndexMesh = ParticleIndexMesh;
-									}
-									// If Constraint.LevelsetIndex is a cluster store an index for a mesh in this cluster
-									if (ClusterIdsArray[Constraint.LevelsetIndex].NumChildren > 0)
-									{
-										int32 LevelsetIndexMesh = GetParticleIndexMesh(ParentToChildrenMap, Constraint.LevelsetIndex);
-										ensure(LevelsetIndexMesh != INDEX_NONE);
-										CollisionDataArrayItem.LevelsetIndexMesh = LevelsetIndexMesh;
-									}
-#endif
-								}
-							}
-						}, Chaos::SmallBatchSize);
-						int32 IdxCollision = 0, IdxCollision2 = NumValidCollisions - 1, NumDeleted = 0;
-						while (IdxCollision <= IdxCollision2)
-						{
-							if (AllCollisionsDataArray[IdxCollision].Proxy1 == nullptr && AllCollisionsDataArray[IdxCollision2].Proxy1 == nullptr)
+							// If Constraint.ParticleIndex is a cluster store an index for a mesh in this cluster
+							if (ClusterIdsArray[Constraint.ParticleIndex].NumChildren > 0)
 							{
-								NumDeleted++;
-								IdxCollision2--;
+								int32 ParticleIndexMesh = GetParticleIndexMesh(ParentToChildrenMap, Constraint.ParticleIndex);
+								ensure(ParticleIndexMesh != INDEX_NONE);
+								CollisionDataArrayItem.ParticleIndexMesh = ParticleIndexMesh;
 							}
-							else if (AllCollisionsDataArray[IdxCollision].Proxy1 == nullptr)
+							// If Constraint.LevelsetIndex is a cluster store an index for a mesh in this cluster
+							if (ClusterIdsArray[Constraint.LevelsetIndex].NumChildren > 0)
 							{
-								AllCollisionsDataArray[IdxCollision] = AllCollisionsDataArray[IdxCollision2];
-								NumDeleted++;
-								IdxCollision2--;
+								int32 LevelsetIndexMesh = GetParticleIndexMesh(ParentToChildrenMap, Constraint.LevelsetIndex);
+								ensure(LevelsetIndexMesh != INDEX_NONE);
+								CollisionDataArrayItem.LevelsetIndexMesh = LevelsetIndexMesh;
 							}
-							else
-							{
-								// Add to AllCollisionsIndicesByPhysicsProxy
-								AllCollisionsIndicesByPhysicsProxy.FindOrAdd(AllCollisionsDataArray[IdxCollision].Proxy1).Add(FEventManager::EncodeCollisionIndex(IdxCollision, false));
+#endif
 
-								if (AllCollisionsDataArray[IdxCollision].Proxy2 && AllCollisionsDataArray[IdxCollision].Proxy2 != AllCollisionsDataArray[IdxCollision].Proxy1)
-								{
-									AllCollisionsIndicesByPhysicsProxy.FindOrAdd(AllCollisionsDataArray[IdxCollision].Proxy2).Add(FEventManager::EncodeCollisionIndex(IdxCollision, true));
-								}
-								IdxCollision++;
+							// Add to AllCollisionsIndicesByPhysicsProxy
+							AllCollisionsIndicesByPhysicsProxy.FindOrAdd(PhysicsProxy).Add(FEventManager::EncodeCollisionIndex(NewIdx, false));
+
+							if (OtherPhysicsProxy && OtherPhysicsProxy != PhysicsProxy)
+							{
+								AllCollisionsIndicesByPhysicsProxy.FindOrAdd(OtherPhysicsProxy).Add(FEventManager::EncodeCollisionIndex(NewIdx, true));
 							}
 						}
-						AllCollisionsDataArray.SetNum(NumValidCollisions - NumDeleted, false);
 					}
 				}
 			}
