Chaos : fix potential crash where geometry collection particle end up in the clustered particle array and remain there after being destroyed

diff --git a/Engine/Source/Runtime/Experimental/Chaos/Public/Chaos/PBDRigidsSOAs.h b/Engine/Source/Runtime/Experimental/Chaos/Public/Chaos/PBDRigidsSOAs.h
--- a/Engine/Source/Runtime/Experimental/Chaos/Public/Chaos/PBDRigidsSOAs.h
+++ b/Engine/Source/Runtime/Experimental/Chaos/Public/Chaos/PBDRigidsSOAs.h
@@ -632,8 +632,11 @@ public:
 
 	void SetClusteredParticleSOA(FPBDRigidClusteredParticleHandle* ClusteredParticle)
 	{
-		RemoveClusteredParticle(ClusteredParticle);
-		InsertClusteredParticle(ClusteredParticle);
+		if (ClusteredParticle->GetParticleType() != Chaos::EParticleType::GeometryCollection)
+		{ 
+			RemoveClusteredParticle(ClusteredParticle);
+			InsertClusteredParticle(ClusteredParticle);
+		}
 
 		if (ClusteredParticle->ObjectState() != EObjectStateType::Dynamic)
 		{
