Fixed Nanite geometry collection scene proxy crash due to missing raster bin material (GCs still don't support bins, but they now remap raster bin to the default material bin correctly)

diff --git a/Engine/Source/Runtime/Experimental/GeometryCollectionEngine/Private/GeometryCollection/GeometryCollectionSceneProxy.cpp b/Engine/Source/Runtime/Experimental/GeometryCollectionEngine/Private/GeometryCollection/GeometryCollectionSceneProxy.cpp
--- a/Engine/Source/Runtime/Experimental/GeometryCollectionEngine/Private/GeometryCollection/GeometryCollectionSceneProxy.cpp
+++ b/Engine/Source/Runtime/Experimental/GeometryCollectionEngine/Private/GeometryCollection/GeometryCollectionSceneProxy.cpp
@@ -1443,6 +1443,7 @@ FNaniteGeometryCollectionSceneProxy::FNaniteGeometryCollectionSceneProxy(UGeomet
 
 		UMaterialInterface* MaterialInterface = bValidMeshSection ? Component->GetMaterial(MeshSection.MaterialID) : nullptr;
 
+		// TODO: PROG_RASTER (Implement programmable raster support)
 		const bool bInvalidMaterial = !MaterialInterface || MaterialInterface->GetBlendMode() != BLEND_Opaque;
 		if (bInvalidMaterial)
 		{
@@ -1473,6 +1474,7 @@ FNaniteGeometryCollectionSceneProxy::FNaniteGeometryCollectionSceneProxy(UGeomet
 		check(MaterialInterface->GetBlendMode() == BLEND_Opaque);
 
 		MaterialSections[SectionIndex].ShadingMaterialProxy = MaterialInterface->GetRenderProxy();
+		MaterialSections[SectionIndex].RasterMaterialProxy  = MaterialInterface->GetRenderProxy(); // TODO: PROG_RASTER (Implement programmable raster support)
 		MaterialSections[SectionIndex].MaterialIndex = MeshSection.MaterialID;
 	}
 
