Fix geometry collection GT particle being part of the wrong acceleration structure

diff --git a/Engine/Source/Runtime/Experimental/Chaos/Private/PhysicsProxy/GeometryCollectionPhysicsProxy.cpp b/Engine/Source/Runtime/Experimental/Chaos/Private/PhysicsProxy/GeometryCollectionPhysicsProxy.cpp
--- a/Engine/Source/Runtime/Experimental/Chaos/Private/PhysicsProxy/GeometryCollectionPhysicsProxy.cpp
+++ b/Engine/Source/Runtime/Experimental/Chaos/Private/PhysicsProxy/GeometryCollectionPhysicsProxy.cpp
@@ -438,7 +438,7 @@ void FGeometryCollectionPhysicsProxy::Initialize(Chaos::FPBDRigidsEvolutionBase
 			// IMPORTANT: we need to set eth right spatial index because GT particle is static and PT particle is rigid
 			// this is causing a mismatch when using the separate acceleration structures optimization which can cause crashes when destroying the particle while async tracing 
 			// todo(chaos) we should eventually refactor this code to use rigid particles on the GT side for geometry collection  
-			P->SetSpatialIdx(Chaos::FSpatialAccelerationIdx{ 0,0 });
+			P->SetSpatialIdx(Chaos::FSpatialAccelerationIdx{ 0,1 });
 
 			const Chaos::FShapesArray& Shapes = P->ShapesArray();
 			const int32 NumShapes = Shapes.Num();
