Fix vehicle forced sleeping, disabled particles were being added to island manager by constraints, ultimately waking the island when they were removed again.

diff --git a/Engine/Source/Runtime/Experimental/Chaos/Private/Chaos/PBDConstraintRule.cpp b/Engine/Source/Runtime/Experimental/Chaos/Private/Chaos/PBDConstraintRule.cpp
--- a/Engine/Source/Runtime/Experimental/Chaos/Private/Chaos/PBDConstraintRule.cpp
+++ b/Engine/Source/Runtime/Experimental/Chaos/Private/Chaos/PBDConstraintRule.cpp
@@ -141,7 +141,13 @@ namespace Chaos
 		{
 			if (ConstraintHandle->IsEnabled())
 			{
-				ConstraintGraph->AddConstraint(GetContainerId(), ConstraintHandle, ConstraintHandle->GetConstrainedParticles());
+				bool DisabledParticle0 = (ConstraintHandle->GetConstrainedParticles()[0] && ConstraintHandle->GetConstrainedParticles()[0]->CastToRigidParticle()) ? ConstraintHandle->GetConstrainedParticles()[0]->CastToRigidParticle()->Disabled() : false;
+				bool DisabledParticle1 = (ConstraintHandle->GetConstrainedParticles()[1] && ConstraintHandle->GetConstrainedParticles()[1]->CastToRigidParticle()) ? ConstraintHandle->GetConstrainedParticles()[1]->CastToRigidParticle()->Disabled() : false;
+
+				if (!DisabledParticle0 && !DisabledParticle1)
+				{
+					ConstraintGraph->AddConstraint(GetContainerId(), ConstraintHandle, ConstraintHandle->GetConstrainedParticles());
+				}
 			}
 		}
 	}
