Chaos : Fix potential crash when object is moving from GlobalPayload to regular leaves when using dynamic tree
- creation of new leaf was rewriting the global payload index and causing duplicate in the global payload array

diff --git a/Engine/Source/Runtime/Experimental/Chaos/Public/Chaos/AABBTree.h b/Engine/Source/Runtime/Experimental/Chaos/Public/Chaos/AABBTree.h
--- a/Engine/Source/Runtime/Experimental/Chaos/Public/Chaos/AABBTree.h
+++ b/Engine/Source/Runtime/Experimental/Chaos/Public/Chaos/AABBTree.h
@@ -1090,7 +1090,8 @@ public:
 		Nodes[AllocatedNodeIdx].ParentNode = INDEX_NONE;
 		FAABBTreePayloadInfo* PayloadInfo = PayloadToInfo.Find(Payload);
 		check(PayloadInfo);
-		*PayloadInfo = FAABBTreePayloadInfo{ INDEX_NONE, INDEX_NONE, LeafIndex, INDEX_NONE, AllocatedNodeIdx };
+		PayloadInfo->LeafIdx = LeafIndex;
+		PayloadInfo->NodeIdx = AllocatedNodeIdx;
 
 		return AllocatedNodeIdx;
 	}
