Fix for Auto and EdTest.

diff --git a/Engine/Source/Runtime/Experimental/Chaos/Private/EventDefaults.cpp b/Engine/Source/Runtime/Experimental/Chaos/Private/EventDefaults.cpp
--- a/Engine/Source/Runtime/Experimental/Chaos/Private/EventDefaults.cpp
+++ b/Engine/Source/Runtime/Experimental/Chaos/Private/EventDefaults.cpp
@@ -184,12 +184,12 @@ namespace Chaos
 
 				{
 					SCOPE_CYCLE_COUNTER(STAT_Events_RegisterCollisionEvent_Notify);
-					FCollisionDataArray DupAllCollisionsDataArray;
-					DupAllCollisionsDataArray.SetNum(NumValidCollisions);
+
 					ValidCollisionHandles.SetNum(NumValidCollisions);
 					if (ValidCollisionHandles.Num() > 0)
 					{
-						//AllCollisionsDataArray.SetNum(NumValidCollisions);
+						FCollisionDataArray DupAllCollisionsDataArray;
+						DupAllCollisionsDataArray.SetNum(NumValidCollisions);
 						InnerPhysicsParallelForRange(ValidCollisionHandles.Num(), [&](int32 StartRangeIndex, int32 EndRangeIndex)
 						{
 							for (int32 IdxCollision = StartRangeIndex; IdxCollision < EndRangeIndex; ++IdxCollision)
@@ -277,7 +277,7 @@ namespace Chaos
 						}, Chaos::SmallBatchSize);
 						for (int32 IdxCollision = 0; IdxCollision < NumValidCollisions; ++IdxCollision)
 						{
-							if (DupAllCollisionsDataArray[IdxCollision].Proxy1)
+							if (DupAllCollisionsDataArray[IdxCollision].Proxy1 != nullptr)
 							{
 								int32 NewIdx = AllCollisionsDataArray.Add(DupAllCollisionsDataArray[IdxCollision]);
 								AllCollisionsIndicesByPhysicsProxy.FindOrAdd(AllCollisionsDataArray[NewIdx].Proxy1).Add(FEventManager::EncodeCollisionIndex(NewIdx, false));
