Chaos overlap queries performance improvement: GJK overlap is vectorized+
#[fyi] Michael.Lentine, Benn.Gallagher, Chris.Caulfield

diff --git a/Engine/Source/Runtime/Experimental/Chaos/Private/Chaos/TriangleMeshImplicitObject.cpp b/Engine/Source/Runtime/Experimental/Chaos/Private/Chaos/TriangleMeshImplicitObject.cpp
--- a/Engine/Source/Runtime/Experimental/Chaos/Private/Chaos/TriangleMeshImplicitObject.cpp
+++ b/Engine/Source/Runtime/Experimental/Chaos/Private/Chaos/TriangleMeshImplicitObject.cpp
@@ -811,7 +811,13 @@ bool FTriangleMeshImplicitObject::OverlapGeomImp(const QueryGeomType& QueryGeom,
 				//However, maybe we should check if it's behind the triangle plane. Also, we should enforce this winding in some way
 				const FVec3 Offset = FVec3::CrossProduct(AB, AC);
 
-				if (GJKIntersection(FTriangle(A, B, C), WorldScaleQueryGeom, WorldScaleQueryTM, Thickness, Offset))
+				VectorRegister4Float ASimd = MakeVectorRegisterFloat((float)A.X, (float)A.Y, (float)A.Z, 0.0f);
+				VectorRegister4Float BSimd = MakeVectorRegisterFloat((float)B.X, (float)B.Y, (float)B.Z, 0.0f);
+				VectorRegister4Float CSimd = MakeVectorRegisterFloat((float)C.X, (float)C.Y, (float)C.Z, 0.0f);
+
+				FTriangleRegister Tri(ASimd, BSimd, CSimd);
+
+				if (GJKIntersection(Tri, WorldScaleQueryGeom, WorldScaleQueryTM, Thickness, Offset))
 				{
 					return true;
 				}
diff --git a/Engine/Source/Runtime/Experimental/Chaos/Public/Chaos/GJK.h b/Engine/Source/Runtime/Experimental/Chaos/Public/Chaos/GJK.h
--- a/Engine/Source/Runtime/Experimental/Chaos/Public/Chaos/GJK.h
+++ b/Engine/Source/Runtime/Experimental/Chaos/Public/Chaos/GJK.h
@@ -38,6 +38,92 @@ namespace Chaos
 		@param ThicknessB The amount of geometry inflation for Geometry B(for example if the surface distance of two geometries with thickness 0 would be 2, a thickness of 0.5 would give a distance of 1.5)
 		@return True if the geometries overlap, False otherwise 
 	 */
+#if GJK_VECTORIZED
+	template <typename T, typename TGeometryA, typename TGeometryB>
+	bool GJKIntersection(const TGeometryA& RESTRICT A, const TGeometryB& RESTRICT B, const TRigidTransform<T, 3>& BToATM, const T InThicknessA = 0, const TVector<T, 3>& InitialDir = TVector<T, 3>(1, 0, 0), const T InThicknessB = 0)
+	{
+
+		const UE::Math::TQuat<T>& RotationDouble = BToATM.GetRotation();
+		VectorRegister4Float RotationSimd = MakeVectorRegisterFloatFromDouble(MakeVectorRegister(RotationDouble.X, RotationDouble.Y, RotationDouble.Z, RotationDouble.W));
+
+		const UE::Math::TVector<T>& TranslationDouble = BToATM.GetTranslation();
+		const VectorRegister4Float TranslationSimd = MakeVectorRegisterFloatFromDouble(MakeVectorRegister(TranslationDouble.X, TranslationDouble.Y, TranslationDouble.Z, 0.0));
+		// Normalize rotation
+		RotationSimd = VectorNormalizeSafe(RotationSimd, GlobalVectorConstants::Float0001);
+
+		const VectorRegister4Float InitialDirSimd = MakeVectorRegisterFloatFromDouble(MakeVectorRegister(InitialDir[0], InitialDir[1], InitialDir[2], 0.0));
+		
+		VectorRegister4Float VSimd = VectorNegate(InitialDirSimd);
+		VSimd = VectorNormalizeSafe(VSimd, MakeVectorRegisterFloatConstant(-1.f, 0.f, 0.f, 0.f));
+
+		const VectorRegister4Float AToBRotationSimd = VectorQuaternionInverse(RotationSimd);
+		bool bTerminate;
+		bool bNearZero = false;
+		int NumIterations = 0;
+		VectorRegister4Float PrevDist2Simd = MakeVectorRegisterFloatConstant(FLT_MAX, FLT_MAX, FLT_MAX, FLT_MAX);
+
+		VectorRegister4Float SimplexSimd[4] = { VectorZeroFloat(), VectorZeroFloat(), VectorZeroFloat(), VectorZeroFloat() };
+		//VectorRegister4Float As[4] = { VectorZeroFloat(), VectorZeroFloat(), VectorZeroFloat(), VectorZeroFloat() };
+		//VectorRegister4Float Bs[4] = { VectorZeroFloat(), VectorZeroFloat(), VectorZeroFloat(), VectorZeroFloat() };
+		VectorRegister4Float BarycentricSimd;
+		VectorRegister4Int NumVerts = GlobalVectorConstants::IntZero;
+
+		const T ThicknessA = A.GetMargin() + InThicknessA;
+		const T ThicknessB = B.GetMargin() + InThicknessB;
+		const T Inflation = ThicknessA + ThicknessB + static_cast<T>(1e-3);
+		const VectorRegister4Float InflationSimd = MakeVectorRegisterFloat((float)Inflation, (float)Inflation, (float)Inflation, (float)Inflation);
+
+		int32 VertexIndexA = INDEX_NONE, VertexIndexB = INDEX_NONE;
+		do
+		{
+			if (!ensure(NumIterations++ < 32))	//todo: take this out
+			{
+				break;	//if taking too long just stop. This should never happen
+			}
+			const VectorRegister4Float NegVSimd = VectorNegate(VSimd);
+			const VectorRegister4Float SupportASimd = A.SupportCoreSimd(NegVSimd, A.GetMargin());
+			const VectorRegister4Float VInBSimd = VectorQuaternionRotateVector(AToBRotationSimd, VSimd);
+			const VectorRegister4Float SupportBLocalSimd = B.SupportCoreSimd(VInBSimd, B.GetMargin());
+			//const TVector<T, 3> SupportB = BToATM.TransformPositionNoScale(SupportBLocal);  // Original code
+			const VectorRegister4Float SupportBSimd = VectorAdd(VectorQuaternionRotateVector(RotationSimd, SupportBLocalSimd), TranslationSimd);
+			const VectorRegister4Float WSimd = VectorSubtract(SupportASimd, SupportBSimd);
+
+			if (VectorMaskBits(VectorCompareGT(VectorDot3(VSimd, WSimd), InflationSimd)))
+			{
+				return false;
+			}
+
+			{
+				// Convert simdInt to int
+				alignas(16) int32 NumVertsInts[4];
+				VectorIntStoreAligned(NumVerts, NumVertsInts);
+				const int32 NumVertsInt = NumVertsInts[0];
+				SimplexSimd[NumVertsInt] = WSimd;
+			}
+
+			NumVerts = VectorIntAdd(NumVerts, GlobalVectorConstants::IntOne);
+
+			VSimd = VectorSimplexFindClosestToOrigin<false>(SimplexSimd, NumVerts, BarycentricSimd, nullptr, nullptr);
+
+			const VectorRegister4Float NewDist2Simd = VectorDot3(VSimd, VSimd);///
+			bNearZero = VectorMaskBits(VectorCompareLT(NewDist2Simd, VectorMultiply(InflationSimd, InflationSimd))) != 0;
+
+			//as simplices become degenerate we will stop making progress. This is a side-effect of precision, in that case take V as the current best approximation
+			//question: should we take previous v in case it's better?
+			bool bMadeProgress = VectorMaskBits(VectorCompareLT(NewDist2Simd, PrevDist2Simd)) != 0;
+			bTerminate = bNearZero || !bMadeProgress;
+			PrevDist2Simd = NewDist2Simd;
+
+			if (!bTerminate)
+			{
+				VSimd = VectorDivide(VSimd,VectorSqrt(NewDist2Simd));
+			}
+
+		} while (!bTerminate);
+
+		return bNearZero;
+	}
+#else
 	template <typename T, typename TGeometryA, typename TGeometryB>
 	bool GJKIntersection(const TGeometryA& A, const TGeometryB& B, const TRigidTransform<T, 3>& BToATM, const T InThicknessA = 0, const TVector<T, 3>& InitialDir = TVector<T, 3>(1, 0, 0), const T InThicknessB = 0)
 	{
@@ -103,6 +189,8 @@ namespace Chaos
 		return bNearZero;
 	}
 
+#endif
+
 	/** 
 		Determines if two convex geometries in the same space overlap
 		IMPORTANT: the two convex geometries must be in the same space!
diff --git a/Engine/Source/Runtime/Experimental/Chaos/Public/Chaos/SimplexVectorized.h b/Engine/Source/Runtime/Experimental/Chaos/Public/Chaos/SimplexVectorized.h
--- a/Engine/Source/Runtime/Experimental/Chaos/Public/Chaos/SimplexVectorized.h
+++ b/Engine/Source/Runtime/Experimental/Chaos/Public/Chaos/SimplexVectorized.h
@@ -7,8 +7,8 @@
 
 namespace Chaos
 {
-
-	FORCEINLINE VectorRegister4Float VectorLineSimplexFindOrigin(VectorRegister4Float* Simplex, VectorRegister4Int& NumVerts, VectorRegister4Float& OutBarycentric, VectorRegister4Float* A, VectorRegister4Float* B)
+	template <bool CalculatExtraInformation>
+	FORCEINLINE VectorRegister4Float VectorLineSimplexFindOrigin(VectorRegister4Float* RESTRICT Simplex, VectorRegister4Int& RESTRICT NumVerts, VectorRegister4Float& RESTRICT OutBarycentric, VectorRegister4Float* RESTRICT A, VectorRegister4Float* RESTRICT B)
 	{
 		const VectorRegister4Float& X0 = Simplex[0];
 		const VectorRegister4Float& X1 = Simplex[1];
@@ -30,8 +30,12 @@ namespace Chaos
 		const VectorRegister4Float IsX1 = VectorBitwiseOr(DotBigger, X0ToX1SquaredSmall);
 
 		Simplex[0] = VectorSelect(IsX1, Simplex[1], Simplex[0]);
-		A[0] = VectorSelect(IsX1, A[1], A[0]);
-		B[0] = VectorSelect(IsX1, B[1], B[0]);
+
+		if (CalculatExtraInformation)
+		{
+			A[0] = VectorSelect(IsX1, A[1], A[0]);
+			B[0] = VectorSelect(IsX1, B[1], B[0]);
+		}
 
 		VectorRegister4Float Ratio = VectorDivide(Dot, X0ToX1Squared);
 
@@ -49,14 +53,18 @@ namespace Chaos
 		const VectorRegister4Int IsX0OrX1Int = VectorCast4FloatTo4Int(IsX0OrX1);
 		NumVerts = VectorIntSelect(IsX0OrX1Int, GlobalVectorConstants::IntOne, NumVerts);
 
-		OutBarycentric = VectorSelect(IsX0OrX1, OutBarycentricIfX0OrX1, OutBarycentricOtherwise);
+		if (CalculatExtraInformation)
+		{
+			OutBarycentric = VectorSelect(IsX0OrX1, OutBarycentricIfX0OrX1, OutBarycentricOtherwise);
+		}
 
 		return Closest;
 	}
 
 	// Based on an algorithm in Real Time Collision Detection - Ericson (very close to that)
 	// Using the same variable name conventions for easy reference
-	FORCEINLINE VectorRegister4Float  TriangleSimplexFindOriginFast(VectorRegister4Float* Simplex, VectorRegister4Int& NumVerts, VectorRegister4Float& OutBarycentric, VectorRegister4Float* As, VectorRegister4Float* Bs)
+	template <bool CalculatExtraInformation>
+	FORCEINLINE VectorRegister4Float  TriangleSimplexFindOriginFast(VectorRegister4Float* RESTRICT Simplex, VectorRegister4Int& RESTRICT NumVerts, VectorRegister4Float& RESTRICT OutBarycentric, VectorRegister4Float* RESTRICT As, VectorRegister4Float* RESTRICT Bs)
 	{
 		const VectorRegister4Float& A = Simplex[0];
 		const VectorRegister4Float& B = Simplex[1];
@@ -77,7 +85,10 @@ namespace Chaos
 		if (VectorMaskBits(IsA))
 		{
 			NumVerts = GlobalVectorConstants::IntOne;
-			OutBarycentric = GlobalVectorConstants::Float1000;
+			if (CalculatExtraInformation)
+			{
+				OutBarycentric = GlobalVectorConstants::Float1000;
+			}
 			return A;
 		}
 
@@ -93,10 +104,16 @@ namespace Chaos
 		if (VectorMaskBits(IsB))
 		{
 			NumVerts = GlobalVectorConstants::IntOne;
-			OutBarycentric = GlobalVectorConstants::Float1000;
+			if (CalculatExtraInformation)
+			{
+				OutBarycentric = GlobalVectorConstants::Float1000;
+			}
 			Simplex[0] = B;
-			As[0] = As[1];
-			Bs[0] = Bs[1];
+			if (CalculatExtraInformation)
+			{
+				As[0] = As[1];
+				Bs[0] = Bs[1];
+			}
 			return B;
 		}
 
@@ -119,7 +136,10 @@ namespace Chaos
 			const VectorRegister4Float v = VectorDivide(d1, NormalizationDenominatorAB);
 			const VectorRegister4Float OneMinusV = VectorSubtract(GlobalVectorConstants::FloatOne, v);
 			// b0	a1	a2	a3
-			OutBarycentric = VectorUnpackLo(OneMinusV, v);
+			if (CalculatExtraInformation)
+			{
+				OutBarycentric = VectorUnpackLo(OneMinusV, v);
+			}
 			return VectorMultiplyAdd(v, AB, A);
 		}
 
@@ -134,10 +154,17 @@ namespace Chaos
 		if (VectorMaskBits(IsC))
 		{
 			NumVerts = GlobalVectorConstants::IntOne;
-			OutBarycentric = GlobalVectorConstants::Float1000;
+			if (CalculatExtraInformation)
+			{
+				OutBarycentric = GlobalVectorConstants::Float1000;
+			}
+
 			Simplex[0] = C;
-			As[0] = As[2];
-			Bs[0] = Bs[2];
+			if (CalculatExtraInformation)
+			{
+				As[0] = As[2];
+				Bs[0] = Bs[2];
+			}
 			return C;
 		}
 
@@ -159,10 +186,16 @@ namespace Chaos
 			NumVerts = two;
 			const VectorRegister4Float OneMinusW = VectorSubtract(GlobalVectorConstants::FloatOne, w);
 			// b0	a1	a2	a3
-			OutBarycentric = VectorUnpackLo(OneMinusW, w);
+			if (CalculatExtraInformation)
+			{
+				OutBarycentric = VectorUnpackLo(OneMinusW, w);
+			}
 			Simplex[1] = C;
-			As[1] = As[2];
-			Bs[1] = Bs[2];
+			if (CalculatExtraInformation)
+			{
+				As[1] = As[2];
+				Bs[1] = Bs[2];
+			}
 			return VectorMultiplyAdd(w, AC, A);
 		}
 
@@ -186,15 +219,21 @@ namespace Chaos
 			NumVerts = two;
 			const VectorRegister4Float OneMinusW = VectorSubtract(GlobalVectorConstants::FloatOne, w);
 			// b0	a1	a2	a3
-			OutBarycentric = VectorUnpackLo(OneMinusW, w);
+			if (CalculatExtraInformation)
+			{
+				OutBarycentric = VectorUnpackLo(OneMinusW, w);
+			}
 			const VectorRegister4Float CMinusB = VectorSubtract(C, B);
 			const VectorRegister4Float Result = VectorMultiplyAdd(w, CMinusB, B);
 			Simplex[0] = B;
 			Simplex[1] = C;
-			As[0] = As[1];
-			Bs[0] = Bs[1];
-			As[1] = As[2];
-			Bs[1] = Bs[2];
+			if (CalculatExtraInformation)
+			{
+				As[0] = As[1];
+				Bs[0] = Bs[1];
+				As[1] = As[2];
+				Bs[1] = Bs[2];
+			}
 			return Result;
 		}
 
@@ -209,7 +248,10 @@ namespace Chaos
 		// b0	a1	a2	a3
 		const VectorRegister4Float OneMinusVMinusW_W = VectorUnpackLo(OneMinusVMinusW, w);
 		// a0	b0	a1	b1
-		OutBarycentric = VectorUnpackLo(OneMinusVMinusW_W, v);
+		if (CalculatExtraInformation)
+		{
+			OutBarycentric = VectorUnpackLo(OneMinusVMinusW_W, v);
+		}
 
 		// We know that we are inside the triangle so we can use the projected point we calculated above. 
 		// The closest point can also be derived from the barycentric coordinates, but it will contain 
@@ -235,8 +277,8 @@ namespace Chaos
 		return (MaskA == MaskB) && !IsZero;
 	}
 
-
-	FORCEINLINE VectorRegister4Float VectorTetrahedronSimplexFindOrigin(VectorRegister4Float* Simplex, VectorRegister4Int& NumVerts, VectorRegister4Float& OutBarycentric, VectorRegister4Float* A, VectorRegister4Float* B)
+	template <bool CalculatExtraInformation>
+	FORCEINLINE VectorRegister4Float VectorTetrahedronSimplexFindOrigin(VectorRegister4Float* RESTRICT Simplex, VectorRegister4Int& RESTRICT NumVerts, VectorRegister4Float& RESTRICT OutBarycentric, VectorRegister4Float* RESTRICT A, VectorRegister4Float* RESTRICT B)
 	{
 		const VectorRegister4Float& X0 = Simplex[0];
 		const VectorRegister4Float& X1 = Simplex[1];
@@ -263,8 +305,22 @@ namespace Chaos
 		constexpr VectorRegister4Int ThreeInt = MakeVectorRegisterIntConstant(3, 3, 3, 3);
 		VectorRegister4Int SubNumVerts[4] = { ThreeInt, ThreeInt, ThreeInt, ThreeInt };
 		VectorRegister4Float SubSimplices[4][3] = { {Simplex[1], Simplex[2], Simplex[3]}, {Simplex[0], Simplex[2], Simplex[3]}, {Simplex[0], Simplex[1], Simplex[3]}, {Simplex[0], Simplex[1], Simplex[2]} };
-		VectorRegister4Float SubAs[4][3] = { {A[1], A[2], A[3]}, {A[0], A[2], A[3]}, {A[0], A[1], A[3]}, {A[0], A[1], A[2]} };
-		VectorRegister4Float SubBs[4][3] = { {B[1], B[2], B[3]}, {B[0], B[2], B[3]}, {B[0], B[1], B[3]}, {B[0], B[1], B[2]} };
+		VectorRegister4Float SubAs[4][3];
+		VectorRegister4Float SubBs[4][3];
+		if (CalculatExtraInformation)
+		{
+			//SubAs = { {A[1], A[2], A[3]}, {A[0], A[2], A[3]}, {A[0], A[1], A[3]}, {A[0], A[1], A[2]} };
+			//SubBs = { {B[1], B[2], B[3]}, {B[0], B[2], B[3]}, {B[0], B[1], B[3]}, {B[0], B[1], B[2]} };
+			SubAs[0][0] = A[1]; SubAs[0][1] = A[2]; SubAs[0][2] = A[3];
+			SubAs[1][0] = A[0]; SubAs[1][1] = A[2]; SubAs[1][2] = A[3];
+			SubAs[2][0] = A[0]; SubAs[2][1] = A[1]; SubAs[2][2] = A[3];
+			SubAs[3][0] = A[0]; SubAs[3][1] = A[1]; SubAs[3][2] = A[2];
+
+			SubBs[0][0] = B[1]; SubBs[0][1] = B[2]; SubBs[0][2] = B[3];
+			SubBs[1][0] = B[0]; SubBs[1][1] = B[2]; SubBs[1][2] = B[3];
+			SubBs[2][0] = B[0]; SubBs[2][1] = B[1]; SubBs[2][2] = B[3];
+			SubBs[3][0] = B[0]; SubBs[3][1] = B[1]; SubBs[3][2] = B[2];
+		}
 		VectorRegister4Float ClosestPointSub[4];
 		VectorRegister4Float SubBarycentric[4];
 		constexpr VectorRegister4Int IndexNone = MakeVectorRegisterIntConstant(INDEX_NONE, INDEX_NONE, INDEX_NONE, INDEX_NONE);
@@ -285,7 +341,7 @@ namespace Chaos
 			if (!bSignMatch[Idx])
 			{
 				bInside = false;
-				ClosestPointSub[Idx] = TriangleSimplexFindOriginFast(SubSimplices[Idx], SubNumVerts[Idx], SubBarycentric[Idx], SubAs[Idx], SubBs[Idx]);
+				ClosestPointSub[Idx] = TriangleSimplexFindOriginFast<CalculatExtraInformation>(SubSimplices[Idx], SubNumVerts[Idx], SubBarycentric[Idx], SubAs[Idx], SubBs[Idx]);
 
 				const VectorRegister4Float Dist2 = VectorDot3(ClosestPointSub[Idx], ClosestPointSub[Idx]);
 
@@ -313,7 +369,10 @@ namespace Chaos
 			const VectorRegister4Float OutBarycentric0101 = VectorUnpackLo(OutBarycentricVectors[0], OutBarycentricVectors[1]);
 			const VectorRegister4Float OutBarycentric2323 = VectorUnpackLo(OutBarycentricVectors[2], OutBarycentricVectors[3]);
 			// a0	a1	b0	b1
-			OutBarycentric = VectorMoveLh(OutBarycentric0101, OutBarycentric2323);
+			if (CalculatExtraInformation)
+			{
+				OutBarycentric = VectorMoveLh(OutBarycentric0101, OutBarycentric2323);
+			}
 
 			return VectorZeroFloat();
 		}
@@ -322,20 +381,28 @@ namespace Chaos
 		VectorIntStoreAligned(ClosestTriangleIdx, ClosestTriangleIdxInts);
 		int32 ClosestTriangleIdxInt = ClosestTriangleIdxInts[0];
 		NumVerts = SubNumVerts[ClosestTriangleIdxInt];
-		OutBarycentric = SubBarycentric[ClosestTriangleIdxInt];
+		if (CalculatExtraInformation)
+		{
+			OutBarycentric = SubBarycentric[ClosestTriangleIdxInt];
+		}
 
 		for (int i = 0; i < 3; i++)
 		{
 			Simplex[i] = SubSimplices[ClosestTriangleIdxInt][i];
-			A[i] = SubAs[ClosestTriangleIdxInt][i];
-			B[i] = SubBs[ClosestTriangleIdxInt][i];
+			if (CalculatExtraInformation)
+			{
+				A[i] = SubAs[ClosestTriangleIdxInt][i];
+				B[i] = SubBs[ClosestTriangleIdxInt][i];
+			}
 		}
 
 		return ClosestPointSub[ClosestTriangleIdxInt];
 	}
 
 
-	FORCEINLINE VectorRegister4Float VectorSimplexFindClosestToOrigin(VectorRegister4Float* Simplex, VectorRegister4Int& NumVerts, VectorRegister4Float& OutBarycentric, VectorRegister4Float* A, VectorRegister4Float* B)
+	// CalculatExtraHitInformation : Should we calculate the BaryCentric coordinates, As and Bs?
+	template <bool CalculatExtraInformation = true>
+	FORCEINLINE VectorRegister4Float VectorSimplexFindClosestToOrigin(VectorRegister4Float* RESTRICT Simplex, VectorRegister4Int& RESTRICT NumVerts, VectorRegister4Float& RESTRICT OutBarycentric, VectorRegister4Float* RESTRICT A, VectorRegister4Float* RESTRICT B)
 	{
 		VectorRegister4Float ClosestPoint;
 		alignas(16) int32 NumVertsInt[4];
@@ -343,22 +410,25 @@ namespace Chaos
 		switch (NumVertsInt[0])
 		{
 		case 1:
-			OutBarycentric = GlobalVectorConstants::Float1000;
+			if (CalculatExtraInformation)
+			{
+				OutBarycentric = GlobalVectorConstants::Float1000;
+			}
 			ClosestPoint = Simplex[0]; 
 			break;
 		case 2:
 		{
-			ClosestPoint = VectorLineSimplexFindOrigin(Simplex, NumVerts, OutBarycentric, A, B);
+			ClosestPoint = VectorLineSimplexFindOrigin<CalculatExtraInformation>(Simplex, NumVerts, OutBarycentric, A, B);
 			break;
 		}
 		case 3:
 		{
-			ClosestPoint = TriangleSimplexFindOriginFast(Simplex, NumVerts, OutBarycentric, A, B);
+			ClosestPoint = TriangleSimplexFindOriginFast<CalculatExtraInformation>(Simplex, NumVerts, OutBarycentric, A, B);
 			break;
 		}
 		case 4:
 		{
-			ClosestPoint = VectorTetrahedronSimplexFindOrigin(Simplex, NumVerts, OutBarycentric, A, B);
+			ClosestPoint = VectorTetrahedronSimplexFindOrigin<CalculatExtraInformation>(Simplex, NumVerts, OutBarycentric, A, B);
 			break;
 		}
 		default:
