Chaos : fix mass for scaled GC

diff --git a/Engine/Source/Runtime/Experimental/Chaos/Private/PhysicsProxy/GeometryCollectionPhysicsProxy.cpp b/Engine/Source/Runtime/Experimental/Chaos/Private/PhysicsProxy/GeometryCollectionPhysicsProxy.cpp
--- a/Engine/Source/Runtime/Experimental/Chaos/Private/PhysicsProxy/GeometryCollectionPhysicsProxy.cpp
+++ b/Engine/Source/Runtime/Experimental/Chaos/Private/PhysicsProxy/GeometryCollectionPhysicsProxy.cpp
@@ -647,6 +647,9 @@ void FGeometryCollectionPhysicsProxy::InitializeBodiesPT(Chaos::FPBDRigidsSolver
 		const float StrainDefault = Parameters.DamageThreshold.Num() ? Parameters.DamageThreshold[0] : 0;
 		// Add the rigid bodies
 
+		const FVector WorldScale = Parameters.WorldTransform.GetScale3D();
+		const FVector::FReal MassScale = WorldScale.X * WorldScale.Y * WorldScale.Z;
+		
 		// Iterating over the geometry group is a fast way of skipping everything that's
 		// not a leaf node, as each geometry has a transform index, which is a shortcut
 		// for the case when there's a 1-to-1 mapping between transforms and geometries.
@@ -659,7 +662,7 @@ void FGeometryCollectionPhysicsProxy::InitializeBodiesPT(Chaos::FPBDRigidsSolver
 				// Mass space -> Composed parent space -> world
 				const FTransform WorldTransform = 
 					MassToLocal[TransformGroupIndex] * Transform[TransformGroupIndex] * Parameters.WorldTransform;
-
+				
 				PopulateSimulatedParticle(
 					Handle,
 					Parameters.Shared,
@@ -667,8 +670,8 @@ void FGeometryCollectionPhysicsProxy::InitializeBodiesPT(Chaos::FPBDRigidsSolver
 					Implicits[TransformGroupIndex],
 					SimFilter,
 					QueryFilter,
-					Mass[TransformGroupIndex],
-					InertiaTensor[TransformGroupIndex],
+					Mass[TransformGroupIndex] * MassScale,
+					InertiaTensor[TransformGroupIndex] * MassScale,
 					WorldTransform,
 					static_cast<uint8>(DynamicState[TransformGroupIndex]),
 					static_cast<int16>(CollisionGroup[TransformGroupIndex]),
@@ -1014,6 +1017,9 @@ FGeometryCollectionPhysicsProxy::BuildClusters(
 	const TManagedArray<FVector3f>& InertiaTensor = 
 		Parameters.RestCollection->GetAttribute<FVector3f>("InertiaTensor", FTransformCollection::TransformGroup);
 
+	const FVector WorldScale = Parameters.WorldTransform.GetScale3D();
+	const FVector::FReal MassScale = WorldScale.X * WorldScale.Y * WorldScale.Z;
+	
 	PopulateSimulatedParticle(
 		Parent,
 		Parameters.Shared, 
@@ -1021,8 +1027,8 @@ FGeometryCollectionPhysicsProxy::BuildClusters(
 		nullptr, // Parent->Geometry() ? Parent->Geometry() : Implicits[CollectionClusterIndex], 
 		SimFilter,
 		QueryFilter,
-		Parent->M() > 0.0 ? Parent->M() : Mass[CollectionClusterIndex], 
-		FVector3f(Parent->I().GetDiagonal() != Chaos::FVec3(0.0) ? Parent->I().GetDiagonal() : Chaos::FVec3(InertiaTensor[CollectionClusterIndex])),
+		Parent->M() > 0.0 ? Parent->M() : Mass[CollectionClusterIndex] * MassScale, 
+		FVector3f(Parent->I().GetDiagonal() != Chaos::FVec3(0.0) ? Parent->I().GetDiagonal() : Chaos::FVec3(InertiaTensor[CollectionClusterIndex] * MassScale)),
 		ParticleTM, 
 		(uint8)DynamicState[CollectionClusterIndex], 
 		0,
