Chaos - GJKWarmStart optimization

diff --git a/Engine/Source/Runtime/Experimental/Chaos/Private/Chaos/CollisionOneShotManifolds.cpp b/Engine/Source/Runtime/Experimental/Chaos/Private/Chaos/CollisionOneShotManifolds.cpp
--- a/Engine/Source/Runtime/Experimental/Chaos/Private/Chaos/CollisionOneShotManifolds.cpp
+++ b/Engine/Source/Runtime/Experimental/Chaos/Private/Chaos/CollisionOneShotManifolds.cpp
@@ -78,6 +78,10 @@ namespace Chaos
 	FAutoConsoleVariableRef CVarChaos_Collision_EnableEdgePrune(TEXT("p.Chaos.Collision.EnableEdgePrune"), bChaos_Collision_EnableEdgePrune, TEXT(""));
 	FAutoConsoleVariableRef CVarChaos_Collision_EdgePrunePlaneDistance(TEXT("p.Chaos.Collision.EdgePrunePlaneDistance"), Chaos_Collision_EdgePrunePlaneDistance, TEXT(""));
 
+	// Whether to use the new index-less GJK. This should be enabled full time once tested
+	bool bChaos_Collision_UseGJK2 = false;
+	FAutoConsoleVariableRef CVarChaos_Collision_UseGJK2(TEXT("p.Chaos.Collision.UseGJK2"), bChaos_Collision_UseGJK2, TEXT(""));
+
 	namespace Collisions
 	{
 		// Forward delarations we need from CollisionRestitution.cpp
@@ -619,7 +623,17 @@ namespace Chaos
 			const TGJKCoreShape<GeometryA> AWithMargin(A, MarginA);
 			const TGJKCoreShape<GeometryB> BWithMargin(B, MarginB);
 
-			if (GJKPenetrationWarmStartable(AWithMargin, BWithMargin, BToATM, Penetration, ClosestA, ClosestB, NormalA, NormalB, VertexIndexA, VertexIndexB, InOutGjkWarmStartData, OutMaxMarginDelta, GJKEpsilon, EPAEpsilon))
+			bool bHaveContact = false;
+			if (bChaos_Collision_UseGJK2)
+			{
+				bHaveContact = GJKPenetrationWarmStartable2(AWithMargin, BWithMargin, BToATM, Penetration, ClosestA, ClosestB, NormalA, NormalB, VertexIndexA, VertexIndexB, InOutGjkWarmStartData, OutMaxMarginDelta, GJKEpsilon, EPAEpsilon);
+			}
+			else
+			{
+				bHaveContact = GJKPenetrationWarmStartable(AWithMargin, BWithMargin, BToATM, Penetration, ClosestA, ClosestB, NormalA, NormalB, VertexIndexA, VertexIndexB, InOutGjkWarmStartData, OutMaxMarginDelta, GJKEpsilon, EPAEpsilon);
+			}
+
+			if (bHaveContact)
 			{
 				Contact.ShapeContactPoints[0] = ClosestA;
 				Contact.ShapeContactPoints[1] = ClosestB;
@@ -646,7 +660,17 @@ namespace Chaos
 			const FReal GJKEpsilon = Chaos_Collision_GJKEpsilon;
 			const FReal EPAEpsilon = Chaos_Collision_EPAEpsilon;
 
-			if (GJKPenetrationSameSpace(A, B, Penetration, ClosestA, ClosestB, Normal, VertexIndexA, VertexIndexB, OutMaxMarginDelta, GJKEpsilon, EPAEpsilon))
+			bool bHaveContact = false;
+			if (bChaos_Collision_UseGJK2)
+			{
+				bHaveContact = GJKPenetrationSameSpace2(A, B, Penetration, ClosestA, ClosestB, Normal, VertexIndexA, VertexIndexB, OutMaxMarginDelta, GJKEpsilon, EPAEpsilon);
+			}
+			else
+			{
+				bHaveContact = GJKPenetrationSameSpace(A, B, Penetration, ClosestA, ClosestB, Normal, VertexIndexA, VertexIndexB, OutMaxMarginDelta, GJKEpsilon, EPAEpsilon);
+			}
+
+			if (bHaveContact)
 			{
 				Contact.ShapeContactPoints[0] = ClosestA;
 				Contact.ShapeContactPoints[1] = ClosestB;
diff --git a/Engine/Source/Runtime/Experimental/Chaos/Private/Chaos/CollisionOneShotManifoldsMiscShapes.cpp b/Engine/Source/Runtime/Experimental/Chaos/Private/Chaos/CollisionOneShotManifoldsMiscShapes.cpp
--- a/Engine/Source/Runtime/Experimental/Chaos/Private/Chaos/CollisionOneShotManifoldsMiscShapes.cpp
+++ b/Engine/Source/Runtime/Experimental/Chaos/Private/Chaos/CollisionOneShotManifoldsMiscShapes.cpp
@@ -30,6 +30,7 @@ namespace Chaos
 	extern FRealSingle Chaos_Collision_Manifold_CapsuleDeepPenetrationFraction;
 	extern FRealSingle Chaos_Collision_Manifold_CapsuleRadialContactFraction;
 
+	extern bool bChaos_Collision_UseGJK2;
 	extern FRealSingle Chaos_Collision_GJKEpsilon;
 	extern FRealSingle Chaos_Collision_EPAEpsilon;
 
@@ -255,18 +256,36 @@ namespace Chaos
 			// Primary contact
 			// NOTE: swapped contact point order to match desired output order (Sphere, Convex)
 			FContactPoint ClosestContactPoint;
-			GJKPenetrationSameSpace(
-				GJKConvex, 
-				GJKSphere, 
-				ClosestPenetration,
-				ClosestContactPoint.ShapeContactPoints[1], 
-				ClosestContactPoint.ShapeContactPoints[0], 
-				ClosestContactPoint.ShapeContactNormal, 
-				ClosestVertexIndexConvex,
-				ClosestVertexIndexSphere,
-				ClosestSupportMaxDelta,
-				GJKEpsilon, 
-				EPAEpsilon);
+			if (bChaos_Collision_UseGJK2)
+			{
+				GJKPenetrationSameSpace2(
+					GJKConvex, 
+					GJKSphere, 
+					ClosestPenetration,
+					ClosestContactPoint.ShapeContactPoints[1], 
+					ClosestContactPoint.ShapeContactPoints[0], 
+					ClosestContactPoint.ShapeContactNormal, 
+					ClosestVertexIndexConvex,
+					ClosestVertexIndexSphere,
+					ClosestSupportMaxDelta,
+					GJKEpsilon, 
+					EPAEpsilon);
+			}
+			else
+			{
+				GJKPenetrationSameSpace(
+					GJKConvex,
+					GJKSphere,
+					ClosestPenetration,
+					ClosestContactPoint.ShapeContactPoints[1],
+					ClosestContactPoint.ShapeContactPoints[0],
+					ClosestContactPoint.ShapeContactNormal,
+					ClosestVertexIndexConvex,
+					ClosestVertexIndexSphere,
+					ClosestSupportMaxDelta,
+					GJKEpsilon,
+					EPAEpsilon);
+			}
 
 			// Stop now if beyond cull distance
 			const FReal ClosestPhi = -ClosestPenetration;
diff --git a/Engine/Source/Runtime/Experimental/Chaos/Public/Chaos/GJK.h b/Engine/Source/Runtime/Experimental/Chaos/Public/Chaos/GJK.h
--- a/Engine/Source/Runtime/Experimental/Chaos/Public/Chaos/GJK.h
+++ b/Engine/Source/Runtime/Experimental/Chaos/Public/Chaos/GJK.h
@@ -327,6 +327,35 @@ namespace Chaos
 			}
 		}
 
+		void Restore2(const TRigidTransform<T, 3>& BToATM, int32& OutNumVerts, TVec3<T> OutSimplex[], TVec3<T>& OutV, T& OutDistance, const T Epsilon)
+		{
+			OutNumVerts = 0;
+
+			if (NumVerts > 0)
+			{
+				for (int32 VertIndex = 0; VertIndex < NumVerts; ++VertIndex)
+				{
+					OutSimplex[VertIndex] = As[VertIndex] - BToATM.TransformPositionNoScale(Bs[VertIndex]);
+				}
+
+				const TVec3<T> V = SimplexFindClosestToOrigin2(OutSimplex, NumVerts, Barycentric, As, Bs);
+				const T DistanceSq = V.SizeSquared();
+
+				// If the origin is inside the simplex at the new transform, we need to abort the restore
+				// This is necessary to cover the very-small separation case where we use the normal
+				// calculated in the previous iteration in GJKL, but we have no way to restore that.
+				// Note: we have already written to the simplex but that's ok because we reset the vert count.
+				if (DistanceSq > FMath::Square(Epsilon))
+				{
+					const T Distance = FMath::Sqrt(DistanceSq);
+					OutNumVerts = NumVerts;
+					OutV = V / Distance;
+					OutDistance = Distance;
+				}
+			}
+		}
+
+
 		// Maximum number of vertices that a GJK simplex can have
 		static const int32 MaxSimplexVerts = 4;
 
@@ -544,6 +573,174 @@ namespace Chaos
 		}
 	}
 
+	// Same as GJKPenetrationWarmStartable but with an index-less algorithm
+	template <typename T, typename TGeometryA, typename TGeometryB>
+	bool GJKPenetrationWarmStartable2(const TGeometryA& A, const TGeometryB& B, const TRigidTransform<T, 3>& BToATM, T& OutPenetration, TVec3<T>& OutClosestA, TVec3<T>& OutClosestB, TVec3<T>& OutNormalA, TVec3<T>& OutNormalB, int32& OutVertexA, int32& OutVertexB, TGJKSimplexData<T>& InOutSimplexData, T& OutMaxSupportDelta, const T Epsilon = T(1.e-3), const T EPAEpsilon = T(1.e-2))
+	{
+		T SupportDeltaA = 0;
+		T SupportDeltaB = 0;
+		T MaxSupportDelta = 0;
+
+		int32& VertexIndexA = OutVertexA;
+		int32& VertexIndexB = OutVertexB;
+
+		// Return the support vertex in A-local space for a vector V in A-local space
+		auto SupportAFunc = [&A, &SupportDeltaA, &VertexIndexA](const TVec3<T>& V)
+		{
+			return A.SupportCore(V, A.GetMargin(), &SupportDeltaA, VertexIndexA);
+		};
+
+		const TRotation<T, 3> AToBRotation = BToATM.GetRotation().Inverse();
+
+		// Return the support point on B, in B-local space, for a vector V in A-local space
+		auto SupportBFunc = [&B, &AToBRotation, &SupportDeltaB, &VertexIndexB](const TVec3<T>& V)
+		{
+			const TVec3<T> VInB = AToBRotation * V;
+			return B.SupportCore(VInB, B.GetMargin(), &SupportDeltaB, VertexIndexB);
+		};
+
+		// V and Simplex are in A-local space
+		TVec3<T> V = TVec3<T>(-1, 0, 0);
+		TVec3<T> Simplex[4];
+		int32 NumVerts = 0;
+		T Distance = FLT_MAX;
+
+		// If we have warm-start data, rebuild the simplex from the stored data
+		InOutSimplexData.Restore2(BToATM, NumVerts, Simplex, V, Distance, Epsilon);
+
+		TVec3<T> Normal = -V;					// Remember the last good normal (i.e. don't update it if separation goes less than Epsilon and we can no longer normalize)
+		bool bIsResult = false;					// True if GJK cannot make any more progress
+		bool bIsContact = false;				// True if shapes are within Epsilon or overlapping - GJK cannot provide a solution
+		int32 NumIterations = 0;
+		const T ThicknessA = A.GetMargin();
+		const T ThicknessB = B.GetMargin();
+		while (!bIsContact && !bIsResult)
+		{
+			if (!ensure(NumIterations++ < 32))	//todo: take this out
+			{
+				break;	//if taking too long just stop. This should never happen
+			}
+
+			const TVec3<T> NegV = -V;
+			const TVec3<T> SupportA = SupportAFunc(NegV);
+			const TVec3<T> SupportB = SupportBFunc(V);
+			const TVec3<T> SupportBInA = BToATM.TransformPositionNoScale(SupportB);
+			const TVec3<T> W = SupportA - SupportBInA;
+			MaxSupportDelta = FMath::Max(SupportDeltaA, SupportDeltaB);
+
+			// If we didn't move to at least ConvergedDistance or closer, assume we have reached a minimum
+			const T ConvergenceTolerance = 1.e-4f;
+			const T ConvergedDistance = (1.0f - ConvergenceTolerance) * Distance;
+			const T VW = TVec3<T>::DotProduct(V, W);
+
+			InOutSimplexData.As[NumVerts] = SupportA;
+			InOutSimplexData.Bs[NumVerts] = SupportB;
+			Simplex[NumVerts++] = W;
+
+			V = SimplexFindClosestToOrigin2(Simplex, NumVerts, InOutSimplexData.Barycentric, InOutSimplexData.As, InOutSimplexData.Bs);
+			T NewDistance = V.Size();
+
+			// Are we overlapping or too close for GJK to get a good result?
+			bIsContact = (NewDistance < Epsilon);
+
+			// If we did not get closer in this iteration, we are in a degenerate situation or we have the result
+			bIsResult = (NewDistance >= (Distance - Epsilon));
+
+			// If we are not too close, update the separating vector and keep iterating
+			// If we are too close we drop through to EPA, but if EPA determines the shapes are actually separated 
+			// after all, we will end up using the normal from the previous loop
+			if (!bIsContact)
+			{
+				V /= NewDistance;
+				Normal = -V;
+			}
+			Distance = NewDistance;
+		}
+		
+		// Save the warm-start data (not much to store since we updated most of it in the loop above)
+		InOutSimplexData.NumVerts = NumVerts;
+
+		if (bIsContact)
+		{
+			// We did not converge or we detected an overlap situation, so run EPA to get contact data
+			// Rebuild the simplex into a variable sized array - EPA can generate any number of vertices
+			TArray<TVec3<T>> VertsA;
+			TArray<TVec3<T>> VertsB;
+			VertsA.Reserve(8);
+			VertsB.Reserve(8);
+			for (int i = 0; i < NumVerts; ++i)
+			{
+				VertsA.Add(InOutSimplexData.As[i]);
+				VertsB.Add(BToATM.TransformPositionNoScale(InOutSimplexData.Bs[i]));
+			}
+
+			auto SupportBInAFunc = [&B, &BToATM, &AToBRotation, &SupportDeltaB, &VertexIndexB](const TVec3<T>& V)
+			{
+				const TVec3<T> VInB = AToBRotation * V;
+				const TVec3<T> SupportBLocal = B.SupportCore(VInB, B.GetMargin(), &SupportDeltaB, VertexIndexB);
+				return BToATM.TransformPositionNoScale(SupportBLocal);
+			};
+
+			T Penetration;
+			TVec3<T> MTD, ClosestA, ClosestBInA;
+			const EEPAResult EPAResult = EPA(VertsA, VertsB, SupportAFunc, SupportBInAFunc, Penetration, MTD, ClosestA, ClosestBInA, EPAEpsilon);
+
+			switch (EPAResult)
+			{
+			case EEPAResult::MaxIterations:
+				// Possibly a solution but with unknown error. Just return the last EPA state. 
+				// Fall through...
+			case EEPAResult::Ok:
+				// EPA has a solution - return it now
+				OutNormalA = MTD;
+				OutNormalB = BToATM.InverseTransformVectorNoScale(MTD);
+				OutPenetration = Penetration + ThicknessA + ThicknessB;
+				OutClosestA = ClosestA + MTD * ThicknessA;
+				OutClosestB = BToATM.InverseTransformPositionNoScale(ClosestBInA - MTD * ThicknessB);
+				OutMaxSupportDelta = MaxSupportDelta;
+				return true;
+			case EEPAResult::BadInitialSimplex:
+				// The origin is outside the simplex. Must be a touching contact and EPA setup will have calculated the normal
+				// and penetration but we keep the position generated by GJK.
+				Normal = MTD;
+				Distance = -Penetration;
+				break;
+			case EEPAResult::Degenerate:
+				// We hit a degenerate simplex condition and could not reach a solution so use whatever near touching point GJK came up with.
+				// The result from EPA under these circumstances is not usable.
+				//UE_LOG(LogChaos, Warning, TEXT("EPA Degenerate Case"));
+				break;
+			}
+		}
+
+		// @todo(chaos): handle the case where EPA hits a degenerate triangle in the simplex.
+		// Currently we return a touching contact with the last position and normal from GJK. We could run SAT instead.
+
+		// GJK converged or we have a touching contact
+		{
+			TVec3<T> ClosestA(0);
+			TVec3<T> ClosestB(0);
+			for (int i = 0; i < NumVerts; ++i)
+			{
+				ClosestA += InOutSimplexData.As[i] * InOutSimplexData.Barycentric[i];
+				ClosestB += InOutSimplexData.Bs[i] * InOutSimplexData.Barycentric[i];
+			}
+
+			OutNormalA = Normal;
+			OutNormalB = BToATM.InverseTransformVectorNoScale(Normal);
+
+			T Penetration = ThicknessA + ThicknessB - Distance;
+			OutPenetration = Penetration;
+			OutClosestA = ClosestA + OutNormalA * ThicknessA;
+			OutClosestB = ClosestB - OutNormalB * ThicknessB;
+
+			OutMaxSupportDelta = MaxSupportDelta;
+
+			// @todo(chaos): we should pass back failure for the degenerate case so we can decide how to handle it externally.
+			return true;
+		}
+	}
+
 	/**
 	 * @brief Calculate the penetration data for two shapes using GJK, assuming both shapes are already in the same space.
 	 * This is intended for use with triangles which have been transformed into the space of the convex shape.
@@ -716,6 +913,157 @@ namespace Chaos
 		}
 	}
 
+	template <typename T, typename TGeometryA, typename TGeometryB>
+	bool GJKPenetrationSameSpace2(const TGeometryA& A, const TGeometryB& B, T& OutPenetration, TVec3<T>& OutClosestA, TVec3<T>& OutClosestB, TVec3<T>& OutNormal, int32& OutVertexA, int32& OutVertexB, T& OutMaxSupportDelta, const T Epsilon = T(1.e-3), const T EPAEpsilon = T(1.e-2))
+	{
+		TGJKSimplexData<T> SimplexData;
+		T SupportDeltaA = 0;
+		T SupportDeltaB = 0;
+		T MaxSupportDelta = 0;
+
+		int32& VertexIndexA = OutVertexA;
+		int32& VertexIndexB = OutVertexB;
+
+		// Return the support vertex in A-local space for a vector V in A-local space
+		auto SupportAFunc = [&A, &SupportDeltaA, &VertexIndexA](const TVec3<T>& V)
+		{
+			return A.SupportCore(V, A.GetMargin(), &SupportDeltaA, VertexIndexA);
+		};
+
+		// Return the support point on B, in B-local space, for a vector V in A-local space
+		auto SupportBFunc = [&B, &SupportDeltaB, &VertexIndexB](const TVec3<T>& V)
+		{
+			return B.SupportCore(V, B.GetMargin(), &SupportDeltaB, VertexIndexB);
+		};
+
+		// V and Simplex are in A-local space
+		TVec3<T> V = TVec3<T>(-1, 0, 0);
+		TVec3<T> Simplex[4];
+		int32 NumVerts = 0;
+		T Distance = FLT_MAX;
+
+		TVec3<T> Normal = -V;					// Remember the last good normal (i.e. don't update it if separation goes less than Epsilon and we can no longer normalize)
+		bool bIsResult = false;				// True if GJK cannot make any more progress
+		bool bIsContact = false;				// True if shapes are within Epsilon or overlapping - GJK cannot provide a solution
+		int32 NumIterations = 0;
+		const T ThicknessA = A.GetMargin();
+		const T ThicknessB = B.GetMargin();
+		const T SeparatedDistance = ThicknessA + ThicknessB + Epsilon;
+		while (!bIsContact && !bIsResult)
+		{
+			if (!ensure(NumIterations++ < 32))	//todo: take this out
+			{
+				break;	//if taking too long just stop. This should never happen
+			}
+
+			const TVec3<T> NegV = -V;
+			const TVec3<T> SupportA = SupportAFunc(NegV);
+			const TVec3<T> SupportB = SupportBFunc(V);
+			const TVec3<T> W = SupportA - SupportB;
+			MaxSupportDelta = FMath::Max(SupportDeltaA, SupportDeltaB);
+
+			// If we didn't move to at least ConvergedDistance or closer, assume we have reached a minimum
+			const T ConvergenceTolerance = 1.e-4f;
+			const T ConvergedDistance = (1.0f - ConvergenceTolerance) * Distance;
+			const T VW = TVec3<T>::DotProduct(V, W);
+
+			SimplexData.As[NumVerts] = SupportA;
+			SimplexData.Bs[NumVerts] = SupportB;
+			Simplex[NumVerts++] = W;
+
+			V = SimplexFindClosestToOrigin2(Simplex, NumVerts, SimplexData.Barycentric, SimplexData.As, SimplexData.Bs);
+			T NewDistance = V.Size();
+
+			// Are we overlapping or too close for GJK to get a good result?
+			bIsContact = (NewDistance < Epsilon);
+
+			// If we did not get closer in this iteration, we are in a degenerate situation
+			bIsResult = (NewDistance >= (Distance - Epsilon));
+
+			// If we are not too close, update the separating vector and keep iterating
+			// If we are too close we drop through to EPA, but if EPA determines the shapes are actually separated 
+			// after all, we will wend up using the normal from the previous loop
+			if (!bIsContact)
+			{
+				V /= NewDistance;
+				Normal = -V;
+			}
+			Distance = NewDistance;
+		}
+
+		// Save the warm-start data (not much to store since we updated most of it in the loop above)
+		SimplexData.NumVerts = NumVerts;
+
+		if (bIsContact)
+		{
+			// We did not converge or we detected an overlap situation, so run EPA to get contact data
+			// Rebuild the simplex into a variable sized array - EPA can generate any mumber of vertices
+			TArray<TVec3<T>> VertsA;
+			TArray<TVec3<T>> VertsB;
+			VertsA.Reserve(8);
+			VertsB.Reserve(8);
+			for (int i = 0; i < NumVerts; ++i)
+			{
+				VertsA.Add(SimplexData.As[i]);
+				VertsB.Add(SimplexData.Bs[i]);
+			}
+
+			T Penetration;
+			TVec3<T> MTD, ClosestA, ClosestB;
+			const EEPAResult EPAResult = EPA(VertsA, VertsB, SupportAFunc, SupportBFunc, Penetration, MTD, ClosestA, ClosestB, EPAEpsilon);
+
+			switch (EPAResult)
+			{
+			case EEPAResult::MaxIterations:
+				// Possibly a solution but with unknown error. Just return the last EPA state. 
+				// Fall through...
+			case EEPAResult::Ok:
+				// EPA has a solution - return it now
+				OutNormal = MTD;
+				OutPenetration = Penetration + ThicknessA + ThicknessB;
+				OutClosestA = ClosestA + MTD * ThicknessA;
+				OutClosestB = ClosestB - MTD * ThicknessB;
+				OutMaxSupportDelta = MaxSupportDelta;
+				return true;
+			case EEPAResult::BadInitialSimplex:
+				// The origin is outside the simplex. Must be a touching contact and EPA setup will have calculated the normal
+				// and penetration but we keep the position generated by GJK.
+				Normal = MTD;
+				Distance = -Penetration;
+				//UE_LOG(LogChaos, Warning, TEXT("EPA Touching Case"));
+				break;
+			case EEPAResult::Degenerate:
+				// We hit a degenerate simplex condition and could not reach a solution so use whatever near touching point GJK came up with.
+				// The result from EPA under these circumstances is not usable.
+				//UE_LOG(LogChaos, Warning, TEXT("EPA Degenerate Case"));
+				break;
+			}
+		}
+
+		// @todo(chaos): handle the case where EPA hits a degenerate triangle in the simplex.
+		// Currently we return a touching contact with the last position and normal from GJK. We could run SAT instead.
+
+		// GJK converged or we have a touching contact
+		{
+			TVec3<T> ClosestA(0);
+			TVec3<T> ClosestB(0);
+			for (int i = 0; i < NumVerts; ++i)
+			{
+				ClosestA += SimplexData.As[i] * SimplexData.Barycentric[i];
+				ClosestB += SimplexData.Bs[i] * SimplexData.Barycentric[i];
+			}
+
+			OutPenetration = ThicknessA + ThicknessB - Distance;
+			OutClosestA = ClosestA + Normal * ThicknessA;
+			OutClosestB = ClosestB - Normal * ThicknessB;
+			OutNormal = Normal;
+			OutMaxSupportDelta = MaxSupportDelta;
+
+			// @todo(chaos): we should pass back failure for the degenerate case so we can decide how to handle it externally.
+			return true;
+		}
+	}
+
 
 	template <typename TGeometryA, typename TGeometryB, bool bNegativePenetrationAllowed = false, typename T>
 	bool GJKPenetrationImpl(const TGeometryA& A, const TGeometryB& B, const TRigidTransform<T, 3>& BToATM, T& OutPenetration, TVec3<T>& OutClosestA, TVec3<T>& OutClosestB, TVec3<T>& OutNormal, int32& OutClosestVertexIndexA, int32& OutClosestVertexIndexB, const T InThicknessA = 0.0f, const T InThicknessB = 0.0f, const TVector<T, 3>& InitialDir = TVector<T, 3>(1, 0, 0), const T Epsilon = 1.e-3f)
diff --git a/Engine/Source/Runtime/Experimental/Chaos/Public/Chaos/Simplex.h b/Engine/Source/Runtime/Experimental/Chaos/Public/Chaos/Simplex.h
--- a/Engine/Source/Runtime/Experimental/Chaos/Public/Chaos/Simplex.h
+++ b/Engine/Source/Runtime/Experimental/Chaos/Public/Chaos/Simplex.h
@@ -40,6 +40,45 @@ namespace Chaos
 		return Closest;
 	}
 
+	template <typename T>
+	TVec3<T> LineSimplexFindOrigin2(TVec3<T>* Simplex, int32& NumVerts, T* OutBarycentric, TVec3<T>* A, TVec3<T>* B)
+	{
+		const TVec3<T>& X0 = Simplex[0];
+		const TVec3<T>& X1 = Simplex[1];
+		const TVec3<T> X0ToX1 = X1 - X0;
+
+		//Closest Point = (-X0 dot X1-X0) / ||(X1-X0)||^2 * (X1-X0)
+
+		const TVec3<T> X0ToOrigin = -X0;
+		const T Dot = TVec3<T>::DotProduct(X0ToOrigin, X0ToX1);
+
+		if (Dot <= 0)
+		{
+			NumVerts = 1;
+			OutBarycentric[0] = 1;
+			return X0;
+		}
+
+		const T X0ToX1Squared = X0ToX1.SizeSquared();
+
+		if (X0ToX1Squared <= Dot || X0ToX1Squared <= std::numeric_limits<T>::min())	//if dividing gives 1+ or the line is degenerate
+		{
+			NumVerts = 1;
+			OutBarycentric[0] = 1;
+			Simplex[0] = Simplex[1];
+			A[0] = A[1];
+			B[0] = B[1];
+			return X1;
+		}
+
+		const T Ratio = FMath::Clamp(Dot / X0ToX1Squared, T(0), T(1));
+		const TVec3<T> Closest = Ratio * (X0ToX1)+X0;	//note: this could pass X1 by machine epsilon, but doesn't seem worth it for now
+		OutBarycentric[0] = 1 - Ratio;
+		OutBarycentric[1] = Ratio;
+		return Closest;
+	}
+
+
 	struct FSimplex
 	{
 		int32 NumVerts;
@@ -230,6 +269,163 @@ namespace Chaos
 		}
 	}
 
+	template <typename T>
+	TVec3<T> TriangleSimplexFindOrigin2(TVec3<T>* Simplex, int32& NumVerts, T* OutBarycentric, TVec3<T>* As, TVec3<T>* Bs)
+	{
+		const TVec3<T>& A = Simplex[0];
+		const TVec3<T>& B = Simplex[1];
+		const TVec3<T>& C = Simplex[2];
+
+		// Vertex region A
+		const TVec3<T> AB = B - A;
+		const TVec3<T> AC = C - A;
+		const TVec3<T> AO = -A;
+
+		const T d1 = TVec3<T>::DotProduct(AB, AO);
+		const T d2 = TVec3<T>::DotProduct(AC, AO);
+
+		const bool bIsD1LEZero = (d1 <= T(0));
+		const bool bIsD2LEZero = (d2 <= T(0));
+		const bool bIsA = bIsD1LEZero && bIsD2LEZero;
+		if (bIsA)
+		{
+			NumVerts = 1;
+			OutBarycentric[0] = T(1);
+			return A;
+		}
+
+		//Vertex region B
+		const TVec3<T> BO = -B;
+		const T d3 = TVec3<T>::DotProduct(AB, BO);
+		const T d4 = TVec3<T>::DotProduct(AC, BO);
+
+		const bool bIsD3GEZero = (d3 >= T(0));
+		const bool bIsD3GED4 = (d4 >= d4);
+		const bool bIsB = bIsD3GEZero && bIsD3GED4;
+		if (bIsB)
+		{
+			NumVerts = 1;
+			OutBarycentric[0] = T(1);
+			Simplex[0] = B;
+			As[0] = As[1];
+			Bs[0] = Bs[1];
+			return B;
+		}
+
+		// Edge AB
+		const T d1d4 = d1 * d4;
+		const T vc = d1d4 - d2 * d3;
+		const T NormalizationDenominatorAB = d1 - d3;
+
+		const bool bIsZeroGEvc = (vc <= T(0));
+		const bool bIsD1GEZero = (d1 >= T(0));
+		const bool bIsZeroGED3 = (d3 <= T(0));
+		const bool bIsNDABGTZero = (NormalizationDenominatorAB > T(0));
+		const bool bIsAB = bIsZeroGEvc && bIsD1GEZero && bIsZeroGED3 && bIsNDABGTZero;
+
+		if (bIsAB)
+		{
+			NumVerts = 2;
+			const T V = d1 / NormalizationDenominatorAB;
+			const T OneMinusV = T(1) - V;
+			OutBarycentric[0] = OneMinusV;
+			OutBarycentric[1] = V;
+			return A + V * AB;
+		}
+
+		// Vertex C
+		const TVec3<T> CO = -C;
+		const T d5 = TVec3<T>::DotProduct(AB, CO);
+		const T d6 = TVec3<T>::DotProduct(AC, CO);
+		const bool bIsD6GEZero = (d6 >= T(0));
+		const bool bIsD6GED5 = (d6 >= d5);
+		const bool bIsC = bIsD6GEZero && bIsD6GED5;
+		if (bIsC)
+		{
+			NumVerts = 1;
+			OutBarycentric[0] = T(1);
+			Simplex[0] = C;
+			As[0] = As[2];
+			Bs[0] = Bs[2];
+			return C;
+		}
+
+		// Edge AC
+		const T d5d2 = d5 * d2;
+		const T vb = d5d2 - d1 * d6;
+		const T NormalizationDenominatorAC = d2 - d6;
+
+		const bool bIsZeroGEvb = (vb <= T(0));
+		const bool bIsD2GEZero = (d2 >= T(0));
+		const bool bIsZeroGED6 = (d6 <= T(0));
+		const bool bIsNDACGTZero = (NormalizationDenominatorAC > T(0));
+		const bool bIsAC = bIsZeroGEvb && bIsD2GEZero && bIsZeroGED6 && bIsNDACGTZero;
+		if (bIsAC)
+		{
+			const T W = d2 / NormalizationDenominatorAC;
+			const T OneMinusW = T(1) - W;
+			NumVerts = 2;
+			OutBarycentric[0] = OneMinusW;
+			OutBarycentric[1] = W;
+			Simplex[1] = C;
+			As[1] = As[2];
+			Bs[1] = Bs[2];
+			return A + W * AC;
+		}
+
+		// Edge BC
+		const T d3d6 = d3 * d6;
+		const T va = d3d6 - d5 * d4;
+		const T d4MinusD3 = d4 - d3;
+		const T d5MinusD6 = d5 - d6;
+		const T NormalizationDenominatorBC = d4MinusD3 + d5MinusD6;
+
+		const bool bIsZeroGEva = (va <= T(0));
+		const bool bIsD4MinusD3GEZero = (d4MinusD3 >= T(0));
+		const bool bIsD5MinusD6GEZero = (d5MinusD6 >= T(0));
+		const bool bIsNDBCGTZero = (NormalizationDenominatorBC > T(0));
+		const bool bIsBC = bIsZeroGEva && bIsD4MinusD3GEZero && bIsD5MinusD6GEZero && bIsNDBCGTZero;
+		if (bIsBC)
+		{
+			const T W = d4MinusD3 / NormalizationDenominatorBC;
+			const T OneMinusW = T(1) - W;
+			NumVerts = 2;
+			OutBarycentric[0] = OneMinusW;
+			OutBarycentric[1] = W;
+			const TVec3<T> CMinusB = C - B;
+			const TVec3<T> Result = B + W * CMinusB;
+			Simplex[0] = B;
+			Simplex[1] = C;
+			As[0] = As[1];
+			Bs[0] = Bs[1];
+			As[1] = As[2];
+			Bs[1] = Bs[2];
+			return Result;
+		}
+
+		// Inside triangle
+		const T denom = T(1) / (va + vb + vc);
+		const T U = va * denom;
+		const T V = vb * denom;
+		const T W = vc * denom;
+		NumVerts = 3;
+		OutBarycentric[0] = U;
+		OutBarycentric[1] = V;
+		OutBarycentric[2] = W;
+
+		// We know that we are inside the triangle so we project the origin onto the plane
+		// The closest point can also be derived from the barycentric coordinates, but it will contain 
+		// numerical error from the determinant calculation  and can cause GJK to terminate with a poor solution.
+		// (E.g., this caused jittering when walking on box with dimensions of 100000cm or more).
+		// This fix the unit test TestSmallCapsuleLargeBoxGJKRaycast_Vertical
+		// Previously was return VectorMultiplyAdd(AC, w, VectorMultiplyAdd(AB, v, A));
+		const TVec3<T> TriNormal = TVec3<T>::CrossProduct(AB, AC);
+		const T TriNormal2 = TVec3<T>::DotProduct(TriNormal, TriNormal);
+		const TVec3<T> TriNormalOverSize2 = TriNormal / TriNormal2;
+		const T SignedDistance = TVec3<T>::DotProduct(A, TriNormalOverSize2);
+		return TriNormal * SignedDistance;
+	}
+
 	template <typename T>
 	TVec3<T> TetrahedronSimplexFindOrigin(const TVec3<T>* Simplex, FSimplex& Idxs, T* OutBarycentric)
 	{
@@ -303,6 +499,79 @@ namespace Chaos
 		return ClosestPointSub[ClosestTriangleIdx];
 	}
 
+	template <typename T>
+	TVec3<T> TetrahedronSimplexFindOrigin2(TVec3<T>* Simplex, int32& NumVerts, T* OutBarycentric, TVec3<T>* A, TVec3<T>* B)
+	{
+		const TVec3<T>& X0 = Simplex[0];
+		const TVec3<T>& X1 = Simplex[1];
+		const TVec3<T>& X2 = Simplex[2];
+		const TVec3<T>& X3 = Simplex[3];
+
+		//Use signed volumes to determine if origin is inside or outside
+		/*
+			M = [X0x X1x X2x X3x;
+				 X0y X1y X2y X3y;
+				 X0z X1z X2z X3z;
+				 1   1   1   1]
+		*/
+
+		T Cofactors[4];
+		Cofactors[0] = -TVec3<T>::DotProduct(X1, TVec3<T>::CrossProduct(X2, X3));
+		Cofactors[1] = TVec3<T>::DotProduct(X0, TVec3<T>::CrossProduct(X2, X3));
+		Cofactors[2] = -TVec3<T>::DotProduct(X0, TVec3<T>::CrossProduct(X1, X3));
+		Cofactors[3] = TVec3<T>::DotProduct(X0, TVec3<T>::CrossProduct(X1, X2));
+		T DetM = (Cofactors[0] + Cofactors[1]) + (Cofactors[2] + Cofactors[3]);
+
+		bool bSignMatch[4];
+		int32 SubNumVerts[4] = { 3, 3, 3, 3 };
+		TVec3<T> SubSimplices[4][3] = { {Simplex[1], Simplex[2], Simplex[3]}, {Simplex[0], Simplex[2], Simplex[3]}, {Simplex[0], Simplex[1], Simplex[3]}, {Simplex[0], Simplex[1], Simplex[2]} };
+		TVec3<T> SubAs[4][3] = { {A[1], A[2], A[3]}, {A[0], A[2], A[3]}, {A[0], A[1], A[3]}, {A[0], A[1], A[2]} };
+		TVec3<T> SubBs[4][3] = { {B[1], B[2], B[3]}, {B[0], B[2], B[3]}, {B[0], B[1], B[3]}, {B[0], B[1], B[2]} };
+		TVec3<T> ClosestPointSub[4];
+		T SubBarycentric[4][4];
+		int32 ClosestTriangleIdx = INDEX_NONE;
+		T MinTriangleDist2 = 0;
+
+		bool bInside = true;
+		for (int Idx = 0; Idx < 4; ++Idx)
+		{
+			bSignMatch[Idx] = SignMatch(DetM, Cofactors[Idx]);
+			if (!bSignMatch[Idx])
+			{
+				bInside = false;
+				ClosestPointSub[Idx] = TriangleSimplexFindOrigin2(SubSimplices[Idx], SubNumVerts[Idx], SubBarycentric[Idx], SubAs[Idx], SubBs[Idx]);
+
+				const T Dist2 = ClosestPointSub[Idx].SizeSquared();
+				if (ClosestTriangleIdx == INDEX_NONE || Dist2 < MinTriangleDist2)
+				{
+					MinTriangleDist2 = Dist2;
+					ClosestTriangleIdx = Idx;
+				}
+			}
+		}
+
+		if (bInside)
+		{
+			OutBarycentric[0] = Cofactors[0] / DetM;
+			OutBarycentric[1] = Cofactors[1] / DetM;
+			OutBarycentric[2] = Cofactors[2] / DetM;
+			OutBarycentric[3] = Cofactors[3] / DetM;
+
+			return TVec3<T>(0);
+		}
+
+		NumVerts = SubNumVerts[ClosestTriangleIdx];
+		for (int i = 0; i < 3; i++)
+		{
+			OutBarycentric[i] = SubBarycentric[ClosestTriangleIdx][i];
+			Simplex[i] = SubSimplices[ClosestTriangleIdx][i];
+			A[i] = SubAs[ClosestTriangleIdx][i];
+			B[i] = SubBs[ClosestTriangleIdx][i];
+		}
+
+		return ClosestPointSub[ClosestTriangleIdx];
+	}
+
 	template <typename T>
 	void ReorderGJKArray(T* Data, FSimplex& Idxs)
 	{
@@ -364,4 +633,32 @@ namespace Chaos
 
 		return ClosestPoint;
 	}
+
+	template <typename T>
+	TVec3<T> SimplexFindClosestToOrigin2(TVec3<T>* Simplex, int32& NumVerts, T* OutBarycentric, TVec3<T>* A, TVec3<T>* B)
+	{
+		TVec3<T> ClosestPoint;
+		switch (NumVerts)
+		{
+		case 1:
+			OutBarycentric[0] = T(1);
+			return Simplex[0];
+		case 2:
+		{
+			return LineSimplexFindOrigin2(Simplex, NumVerts, OutBarycentric, A, B);
+		}
+		case 3:
+		{
+			return TriangleSimplexFindOrigin2(Simplex, NumVerts, OutBarycentric, A, B);
+		}
+		case 4:
+		{
+			return TetrahedronSimplexFindOrigin2(Simplex, NumVerts, OutBarycentric, A, B);
+		}
+		default:
+			ensure(false);
+			return TVec3<T>(0);
+		}
+	}
+
 }
