Chaos : Store and compute proper InertiaTensor and RotationOfMass in FConvex

diff --git a/Engine/Source/Runtime/Core/Public/UObject/UE5ReleaseStreamObjectVersion.h b/Engine/Source/Runtime/Core/Public/UObject/UE5ReleaseStreamObjectVersion.h
--- a/Engine/Source/Runtime/Core/Public/UObject/UE5ReleaseStreamObjectVersion.h
+++ b/Engine/Source/Runtime/Core/Public/UObject/UE5ReleaseStreamObjectVersion.h
@@ -111,6 +111,9 @@ struct CORE_API FUE5ReleaseStreamObjectVersion
 
 		// Serialize LandscapeActorGuid in FLandscapeActorDesc sub class.
 		WorldPartitionLandscapeActorDescSerializeLandscapeActorGuid,
+		
+		// add inertia tensor and rotation of mass to convex 
+		AddedInertiaTensorAndRotationOfMassAddedToConvex,
 
 		// -----<new versions can be added above this line>-------------------------------------------------
 		VersionPlusOne,
diff --git a/Engine/Source/Runtime/Experimental/Chaos/Private/Chaos/Convex.cpp b/Engine/Source/Runtime/Experimental/Chaos/Private/Chaos/Convex.cpp
--- a/Engine/Source/Runtime/Experimental/Chaos/Private/Chaos/Convex.cpp
+++ b/Engine/Source/Runtime/Experimental/Chaos/Private/Chaos/Convex.cpp
@@ -419,4 +419,70 @@ namespace Chaos
 		// Generate a new convex from the points
 		*this = FConvex(NewPoints, 0.0f);
 	}
+
+	void FConvex::ComputeUnitMassInertiaTensorAndRotationOfMass(const FReal InVolume)
+	{
+		// we only compute the inertia tensor using unit mass and scale it upon request
+		if (InVolume < SMALL_NUMBER)
+		{
+			UnitMassInertiaTensor = FVec3{1., 1., 1.};
+			RotationOfMass = FRotation3::Identity;
+			return;
+		}
+
+		// verify that the StructureData has been built before we call process further
+		ensure(StructureData.IsValid());
+		
+		constexpr FReal Mass = 1.0;
+		const FReal Density = Mass / InVolume;
+
+		static const FMatrix33 Standard(2, 1, 1, 2, 1, 2);
+		FMatrix33 Covariance(0);
+
+		for (int32 PlaneIndex = 0; PlaneIndex < Planes.Num(); ++PlaneIndex)
+		{
+			const int32 PlaneVerticesNum = NumPlaneVertices(PlaneIndex);
+			if (PlaneVerticesNum > 2)
+			{
+				// compute plane center
+				FVec3 PlaneCenter(0);
+				for (int32 VertexIndex = 0; VertexIndex < PlaneVerticesNum; VertexIndex++)
+				{
+					PlaneCenter += GetVertex(GetPlaneVertex(PlaneIndex, VertexIndex));
+				}
+				PlaneCenter /= PlaneVerticesNum;
+
+				// Now break down the plane in triangle ( fan around PlaneCenter ) 
+				FMatrix33 DeltaMatrix(0);
+				for (int32 VertexIndex = 0; VertexIndex < PlaneVerticesNum; VertexIndex++)
+				{
+					FVec3 DeltaVector = PlaneCenter - CenterOfMass;
+					DeltaMatrix.M[0][0] = DeltaVector[0];
+					DeltaMatrix.M[1][0] = DeltaVector[1];
+					DeltaMatrix.M[2][0] = DeltaVector[2];
+
+					DeltaVector = GetVertex(GetPlaneVertex(PlaneIndex, VertexIndex)) - CenterOfMass;
+					DeltaMatrix.M[0][1] = DeltaVector[0];
+					DeltaMatrix.M[1][1] = DeltaVector[1];
+					DeltaMatrix.M[2][1] = DeltaVector[2];
+
+					DeltaVector = GetVertex(GetPlaneVertex(PlaneIndex, (VertexIndex + 1) % PlaneVerticesNum)) - CenterOfMass;
+					DeltaMatrix.M[0][2] = DeltaVector[0];
+					DeltaMatrix.M[1][2] = DeltaVector[1];
+					DeltaMatrix.M[2][2] = DeltaVector[2];
+
+					FReal Det = DeltaMatrix.M[0][0] * (DeltaMatrix.M[1][1] * DeltaMatrix.M[2][2] - DeltaMatrix.M[1][2] * DeltaMatrix.M[2][1]) -
+						DeltaMatrix.M[0][1] * (DeltaMatrix.M[1][0] * DeltaMatrix.M[2][2] - DeltaMatrix.M[1][2] * DeltaMatrix.M[2][0]) +
+						DeltaMatrix.M[0][2] * (DeltaMatrix.M[1][0] * DeltaMatrix.M[2][1] - DeltaMatrix.M[1][1] * DeltaMatrix.M[2][0]);
+					const FMatrix33 ScaledStandard = Standard * Det;
+					Covariance += DeltaMatrix * ScaledStandard * DeltaMatrix.GetTransposed();
+				}
+			}
+		}
+		FReal Trace = Covariance.M[0][0] + Covariance.M[1][1] + Covariance.M[2][2];
+		FMatrix33 TraceMat(Trace, Trace, Trace);
+		FMatrix33 InertiaTensor = (TraceMat - Covariance) * (1 / (FReal)120) * Density;
+		RotationOfMass = TransformToLocalSpace(InertiaTensor);
+		UnitMassInertiaTensor = InertiaTensor.GetDiagonal();
+	}
 }
\ No newline at end of file
diff --git a/Engine/Source/Runtime/Experimental/Chaos/Public/Chaos/Convex.h b/Engine/Source/Runtime/Experimental/Chaos/Public/Chaos/Convex.h
--- a/Engine/Source/Runtime/Experimental/Chaos/Public/Chaos/Convex.h
+++ b/Engine/Source/Runtime/Experimental/Chaos/Public/Chaos/Convex.h
@@ -15,6 +15,7 @@
 #include "UObject/ReleaseObjectVersion.h"
 #include "UObject/PhysicsObjectVersion.h"
 #include "UObject/UE5MainStreamObjectVersion.h"
+#include "UObject/UE5ReleaseStreamObjectVersion.h"
 #include "UObject/FortniteMainBranchObjectVersion.h"
 
 namespace Chaos
@@ -41,6 +42,8 @@ namespace Chaos
 		    : FImplicitObject(EImplicitObject::IsConvex | EImplicitObject::HasBoundingBox, ImplicitObjectType::Convex)
 			, Volume(0.f)
 			, CenterOfMass(FVec3Type(0.f))
+			, UnitMassInertiaTensor(1., 1., 1.)
+			, RotationOfMass(FRotation3::Identity)
 		{}
 		FConvex(FConvex&& Other)
 		    : FImplicitObject(EImplicitObject::IsConvex | EImplicitObject::HasBoundingBox, ImplicitObjectType::Convex)
@@ -50,6 +53,8 @@ namespace Chaos
 			, StructureData(MoveTemp(Other.StructureData))
 			, Volume(MoveTemp(Other.Volume))
 			, CenterOfMass(MoveTemp(Other.CenterOfMass))
+			, UnitMassInertiaTensor(MoveTemp(Other.UnitMassInertiaTensor))
+			, RotationOfMass(MoveTemp(Other.RotationOfMass))
 		{}
 
 		// NOTE: This constructor will result in approximate COM and volume calculations, since it does
@@ -72,6 +77,8 @@ namespace Chaos
 			// For now we approximate COM and volume with the bounding box
 			CenterOfMass = LocalBoundingBox.GetCenterOfMass();
 			Volume = LocalBoundingBox.GetVolume();
+
+			ComputeUnitMassInertiaTensorAndRotationOfMass(Volume);
 		}
 
 		FConvex(TArray<FPlaneType>&& InPlanes, TArray<TArray<int32>>&& InFaceIndices, TArray<FVec3Type>&& InVertices)
@@ -87,8 +94,10 @@ namespace Chaos
 			// For now we approximate COM and volume with the bounding box
 			CenterOfMass = LocalBoundingBox.GetCenterOfMass();
 			Volume = LocalBoundingBox.GetVolume();
-
+			
 			CreateStructureData(MoveTemp(InFaceIndices));
+
+			ComputeUnitMassInertiaTensorAndRotationOfMass(Volume);
 		}
 
 		FConvex(const TArray<FVec3Type>& InVertices, const FReal InMargin)
@@ -107,7 +116,7 @@ namespace Chaos
 			// @todo(chaos): this only works with triangles. Fix that an we can run MergeFaces before calling this
 			CalculateVolumeAndCenterOfMass(Vertices, FaceIndices, Volume, CenterOfMass);
 
-			// @todo(chaos):																				 should be based on size, or passed in
+			// @todo(chaos): should be based on size, or passed in
 			if (!FConvexBuilder::bUseGeometryTConvexHull3)
 			{
 				// @todo(convex) : TConvexHull3 does not need to merge faces, and 
@@ -121,6 +130,8 @@ namespace Chaos
 
 			CreateStructureData(MoveTemp(FaceIndices));
 
+			ComputeUnitMassInertiaTensorAndRotationOfMass(Volume);
+
 			SetMargin(InMargin);
 		}
 
@@ -146,6 +157,8 @@ namespace Chaos
 			StructureData = MoveTemp(Other.StructureData);
 			Volume = MoveTemp(Other.Volume);
 			CenterOfMass = MoveTemp(Other.CenterOfMass);
+			UnitMassInertiaTensor = MoveTemp(Other.UnitMassInertiaTensor);
+			RotationOfMass = MoveTemp(Other.RotationOfMass);
 
 			return *this;
 		}
@@ -732,16 +745,16 @@ namespace Chaos
 		{
 			return Volume;
 		}
-
+		
 		const FMatrix33 GetInertiaTensor(const FReal Mass) const
 		{
-			// TODO: More precise inertia!
-			return LocalBoundingBox.GetInertiaTensor(Mass);
+			const FVec3 ScaledInertiaTensorDiagonal{ UnitMassInertiaTensor * Mass };
+			return FMatrix33(ScaledInertiaTensorDiagonal.X, ScaledInertiaTensorDiagonal.Y, ScaledInertiaTensorDiagonal.Z);
 		}
 
 		FRotation3 GetRotationOfMass() const
 		{
-			return FRotation3::FromIdentity();
+			return RotationOfMass;
 		}
 
 		const FVec3 GetCenterOfMass() const
@@ -774,6 +787,7 @@ namespace Chaos
 			Ar.UsingCustomVersion(FUE5MainStreamObjectVersion::GUID);
 			Ar.UsingCustomVersion(FPhysicsObjectVersion::GUID);
 			Ar.UsingCustomVersion(FFortniteMainBranchObjectVersion::GUID);
+			Ar.UsingCustomVersion(FUE5ReleaseStreamObjectVersion::GUID);
 			FImplicitObject::SerializeImp(Ar);
 
 			if (Ar.CustomVer(FExternalPhysicsCustomObjectVersion::GUID) < FExternalPhysicsCustomObjectVersion::ConvexUsesTPlaneConcrete)
@@ -827,6 +841,13 @@ namespace Chaos
 				Ar << VolumeFloat;
 				Volume = (FRealType)VolumeFloat;
 
+				// Some assets have stored NaN ( empty convex objects from Geometry Collection )
+				// Fixing this so that next time they get saved the Volume is a valid number
+				if (FMath::IsNaN(Volume) && Ar.IsLoading())
+				{
+					Volume = 0;
+				}
+
 				Ar << CenterOfMass;
 			}
 			else if (Ar.IsLoading())
@@ -860,6 +881,16 @@ namespace Chaos
 				FConvexBuilder::BuildPlaneVertexIndices(Planes, Vertices, FaceIndices);
 				CreateStructureData(MoveTemp(FaceIndices));
 			}
+
+			if (Ar.CustomVer(FUE5ReleaseStreamObjectVersion::GUID) >= FUE5ReleaseStreamObjectVersion::AddedInertiaTensorAndRotationOfMassAddedToConvex)
+			{
+				Ar << UnitMassInertiaTensor;
+				Ar << RotationOfMass;
+			}
+			else if (Ar.IsLoading())
+			{
+				ComputeUnitMassInertiaTensorAndRotationOfMass(Volume);
+			}
 		}
 
 		virtual void Serialize(FChaosArchive& Ar) override
@@ -930,6 +961,8 @@ namespace Chaos
 #endif // #if INTEL_ISPC && !UE_BUILD_SHIPPING
 
 	private:
+		void ComputeUnitMassInertiaTensorAndRotationOfMass(const FReal InVolume);
+		
 		// IMPORTANT to keep this copy constructor private to avoid unintented expensive copies 
 		FConvex(const FConvex& Other)
 			: FImplicitObject(EImplicitObject::IsConvex | EImplicitObject::HasBoundingBox, ImplicitObjectType::Convex)
@@ -938,6 +971,8 @@ namespace Chaos
 			, LocalBoundingBox(Other.LocalBoundingBox)
 			, Volume(Other.Volume)
 			, CenterOfMass(Other.CenterOfMass)
+			, UnitMassInertiaTensor(Other.UnitMassInertiaTensor)
+			, RotationOfMass(Other.RotationOfMass)
 		{
 			StructureData.CopyFrom(Other.StructureData);
 			SetMargin(Other.Margin);
@@ -950,5 +985,7 @@ namespace Chaos
 		FConvexStructureData StructureData;
 		FRealType Volume;
 		FVec3Type CenterOfMass;
+		FVec3 UnitMassInertiaTensor;
+		FRotation3 RotationOfMass;
 	};
 }
diff --git a/Engine/Source/Runtime/Experimental/ChaosCore/Public/Chaos/Core.h b/Engine/Source/Runtime/Experimental/ChaosCore/Public/Chaos/Core.h
--- a/Engine/Source/Runtime/Experimental/ChaosCore/Public/Chaos/Core.h
+++ b/Engine/Source/Runtime/Experimental/ChaosCore/Public/Chaos/Core.h
@@ -42,5 +42,6 @@ namespace Chaos
 	// NOTE: if you get a merge conflict on the GUID, you must replace it with a new GUID - do not accept the source or target
 	// or you will likely get DDC version conflicts resulting in crashes during load.
 	// Core version string for Chaos data. Any DDC builder dependent on Chaos for serialization should depend on this version
-	static const FString ChaosVersionString = TEXT("997CB556-41E1-4D22-B53C-986B3562586C");
+	static const FString ChaosVersionString = TEXT("A022A2AA-C0D0-4977-BCE3-3A6CF7D3B5C6");
+
 }
