Fix FVector conversion errors after FVector constructor was made explicit.

diff --git a/Engine/Source/Runtime/Engine/Private/PhysicsEngine/ImmediatePhysics/ImmediatePhysicsChaos/ImmediatePhysicsActorHandle_Chaos.cpp b/Engine/Source/Runtime/Engine/Private/PhysicsEngine/ImmediatePhysics/ImmediatePhysicsChaos/ImmediatePhysicsActorHandle_Chaos.cpp
--- a/Engine/Source/Runtime/Engine/Private/PhysicsEngine/ImmediatePhysics/ImmediatePhysicsChaos/ImmediatePhysicsActorHandle_Chaos.cpp
+++ b/Engine/Source/Runtime/Engine/Private/PhysicsEngine/ImmediatePhysics/ImmediatePhysicsChaos/ImmediatePhysicsActorHandle_Chaos.cpp
@@ -642,12 +642,12 @@ namespace ImmediatePhysics_Chaos
 
 	FVector FActorHandle::GetInverseInertia() const
 	{
-		return Handle()->InvI();
+		return FVector(Handle()->InvI());
 	}
 
 	FVector FActorHandle::GetInertia() const
 	{
-		return Handle()->I();
+		return FVector(Handle()->I());
 	}
 
 	void FActorHandle::SetMaxDepenetrationVelocity(FReal NewMaxDepenetrationVelocity)
diff --git a/Engine/Source/Runtime/Experimental/Chaos/Private/Chaos/CollisionResolutionUtil.cpp b/Engine/Source/Runtime/Experimental/Chaos/Private/Chaos/CollisionResolutionUtil.cpp
--- a/Engine/Source/Runtime/Experimental/Chaos/Private/Chaos/CollisionResolutionUtil.cpp
+++ b/Engine/Source/Runtime/Experimental/Chaos/Private/Chaos/CollisionResolutionUtil.cpp
@@ -50,14 +50,14 @@ namespace Chaos
 			if (bIsRigidDynamic0)
 			{
 				Jr0 = FVec3::CrossProduct(VectorToPoint1, Impulse);
-				IInvJr0 = PBDRigid0->Q().RotateVector(PBDRigid0->InvI() * PBDRigid0->Q().UnrotateVector(Jr0));
+				IInvJr0 = PBDRigid0->Q().RotateVector(FVec3(PBDRigid0->InvI()) * PBDRigid0->Q().UnrotateVector(Jr0));
 				ImpulseRatioNumerator0 = FVec3::DotProduct(Impulse, PBDRigid0->V() - KinematicVelocity) + FVec3::DotProduct(Jr0, PBDRigid0->W());
 				ImpulseRatioDenom0 = ImpulseSizeSQ / PBDRigid0->M() + FVec3::DotProduct(Jr0, IInvJr0);
 			}
 			if (bIsRigidDynamic1)
 			{
 				Jr1 = FVec3::CrossProduct(VectorToPoint2, Impulse);
-				IInvJr1 = PBDRigid1->Q().RotateVector(PBDRigid1->InvI() * PBDRigid1->Q().UnrotateVector(Jr1));
+				IInvJr1 = PBDRigid1->Q().RotateVector(FVec3(PBDRigid1->InvI()) * PBDRigid1->Q().UnrotateVector(Jr1));
 				ImpulseRatioNumerator1 = FVec3::DotProduct(Impulse, PBDRigid1->V() - KinematicVelocity) + FVec3::DotProduct(Jr1, PBDRigid1->W());
 				ImpulseRatioDenom1 = ImpulseSizeSQ / PBDRigid1->M() + FVec3::DotProduct(Jr1, IInvJr1);
 			}
diff --git a/Engine/Source/Runtime/Experimental/Chaos/Private/Chaos/Convex.cpp b/Engine/Source/Runtime/Experimental/Chaos/Private/Chaos/Convex.cpp
--- a/Engine/Source/Runtime/Experimental/Chaos/Private/Chaos/Convex.cpp
+++ b/Engine/Source/Runtime/Experimental/Chaos/Private/Chaos/Convex.cpp
@@ -448,7 +448,7 @@ namespace Chaos
 				FVec3 PlaneCenter(0);
 				for (int32 VertexIndex = 0; VertexIndex < PlaneVerticesNum; VertexIndex++)
 				{
-					PlaneCenter += GetVertex(GetPlaneVertex(PlaneIndex, VertexIndex));
+					PlaneCenter += FVec3(GetVertex(GetPlaneVertex(PlaneIndex, VertexIndex)));
 				}
 				PlaneCenter /= PlaneVerticesNum;
 
diff --git a/Engine/Source/Runtime/Experimental/Chaos/Private/PhysicsProxy/GeometryCollectionPhysicsProxy.cpp b/Engine/Source/Runtime/Experimental/Chaos/Private/PhysicsProxy/GeometryCollectionPhysicsProxy.cpp
--- a/Engine/Source/Runtime/Experimental/Chaos/Private/PhysicsProxy/GeometryCollectionPhysicsProxy.cpp
+++ b/Engine/Source/Runtime/Experimental/Chaos/Private/PhysicsProxy/GeometryCollectionPhysicsProxy.cpp
@@ -2427,7 +2427,7 @@ void FGeometryCollectionPhysicsProxy::InitializeSharedCollisionStructures(
 				TUniquePtr<FTriangleMesh> UnionMesh(new FTriangleMesh(MoveTemp(UnionMeshIndices)));
 				// TODO: Seems this should rotate full matrix and not discard off diagonals.
 				const FVec3& InertiaDiagonal = CollectionSpaceParticles->I(ClusterTransformIdx);
-				CollectionInertiaTensor[ClusterTransformIdx] = InertiaDiagonal;	// LWC_TODO: Precision loss
+				CollectionInertiaTensor[ClusterTransformIdx] = FVector3f(InertiaDiagonal);	// LWC_TODO: Precision loss
 				CollectionMass[ClusterTransformIdx] = (FRealSingle)CollectionSpaceParticles->M(ClusterTransformIdx);
 
 				const int32 SizeSpecificIdx = GeometryCollection::SizeSpecific::FindIndexForVolume(SharedParams.SizeSpecificData, InstanceBoundingBox);
diff --git a/Engine/Source/Runtime/PhysicsCore/Private/ChaosEngineInterface.cpp b/Engine/Source/Runtime/PhysicsCore/Private/ChaosEngineInterface.cpp
--- a/Engine/Source/Runtime/PhysicsCore/Private/ChaosEngineInterface.cpp
+++ b/Engine/Source/Runtime/PhysicsCore/Private/ChaosEngineInterface.cpp
@@ -686,7 +686,7 @@ FTransform FChaosEngineInterface::GetComTransformLocal_AssumesLocked(const FPhys
 
 FVector FChaosEngineInterface::GetLocalInertiaTensor_AssumesLocked(const FPhysicsActorHandle& InActorReference)
 {
-	return InActorReference->GetGameThreadAPI().I();
+	return FVector(InActorReference->GetGameThreadAPI().I());
 }
 
 FBox FChaosEngineInterface::GetBounds_AssumesLocked(const FPhysicsActorHandle& InActorReference)
