Chaos: Fix sweep face selection for triangle meshes that are separate but co-located, ensuring we pick the face most opposing the sweep direction.

diff --git a/Engine/Source/Runtime/Engine/Private/PhysicsEngine/Experimental/PhysInterface_Chaos.cpp b/Engine/Source/Runtime/Engine/Private/PhysicsEngine/Experimental/PhysInterface_Chaos.cpp
--- a/Engine/Source/Runtime/Engine/Private/PhysicsEngine/Experimental/PhysInterface_Chaos.cpp
+++ b/Engine/Source/Runtime/Engine/Private/PhysicsEngine/Experimental/PhysInterface_Chaos.cpp
@@ -1143,7 +1143,8 @@ bool FPhysInterface_Chaos::Sweep_Geom(FHitResult& OutHit, const FBodyInstance* I
 							Chaos::FVec3 WorldNormal;
 							Chaos::FReal Distance;
 							int32 FaceIdx;
-							if (Chaos::Utilities::CastHelper(ShapeAdapter.GetGeometry(), ActorTM, [&](const auto& Downcast, const auto& FullActorTM) { return Chaos::SweepQuery(*Shape->GetGeometry(), FullActorTM, Downcast, StartTM, Dir, DeltaMag, Distance, WorldPosition, WorldNormal, FaceIdx, 0.f, false); }))
+							Chaos::FVec3 FaceNormal;
+							if (Chaos::Utilities::CastHelper(ShapeAdapter.GetGeometry(), ActorTM, [&](const auto& Downcast, const auto& FullActorTM) { return Chaos::SweepQuery(*Shape->GetGeometry(), FullActorTM, Downcast, StartTM, Dir, DeltaMag, Distance, WorldPosition, WorldNormal, FaceIdx, FaceNormal, 0.f, false); }))
 							{
 								// we just like to make sure if the hit is made
 								FCollisionFilterData QueryFilter;
diff --git a/Engine/Source/Runtime/Experimental/Chaos/Private/Chaos/Collision/ContactPointsMiscShapes.cpp b/Engine/Source/Runtime/Experimental/Chaos/Private/Chaos/Collision/ContactPointsMiscShapes.cpp
--- a/Engine/Source/Runtime/Experimental/Chaos/Private/Chaos/Collision/ContactPointsMiscShapes.cpp
+++ b/Engine/Source/Runtime/Experimental/Chaos/Private/Chaos/Collision/ContactPointsMiscShapes.cpp
@@ -34,12 +34,13 @@ namespace Chaos
 
 		FReal OutTime = FLT_MAX;
 		int32 FaceIndex = -1;
+		FVec3 FaceNormal;
 		FVec3 Location, Normal;
 
 		Utilities::CastHelper(A, AStartTransform, [&](const auto& ADowncast, const FRigidTransform3& AFullTM)
 			{
 				// @todo(chaos): handle instances with margin
-				if (B.SweepGeom(ADowncast, AToBTM, LocalDir, Length, OutTime, Location, Normal, FaceIndex, 0.0f, true))
+				if (B.SweepGeom(ADowncast, AToBTM, LocalDir, Length, OutTime, Location, Normal, FaceIndex, FaceNormal, 0.0f, true))
 				{
 					// @todo(chaos): margin
 					Contact.ShapeContactPoints[0] = AToBTM.InverseTransformPosition(Location);
@@ -158,11 +159,12 @@ namespace Chaos
 		FReal OutTime = FLT_MAX;
 		FVec3 Location, Normal;
 		int32 FaceIndex = -1;
+		Chaos::FVec3 FaceNormal;
 
 		Utilities::CastHelper(A, AStartTransform, [&](const auto& ADowncast, const FRigidTransform3& AFullTM)
 			{
 				// @todo(chaos): handle Instanced with margin
-				if (B.LowLevelSweepGeom(ADowncast, AToBTM, LocalDir, Length, OutTime, Location, Normal, FaceIndex, 0.0f, true))
+				if (B.LowLevelSweepGeom(ADowncast, AToBTM, LocalDir, Length, OutTime, Location, Normal, FaceIndex, FaceNormal, 0.0f, true))
 				{
 					Contact.ShapeContactPoints[1] = Location;
 					Contact.ShapeContactNormal = Normal;
diff --git a/Engine/Source/Runtime/Experimental/Chaos/Private/Chaos/HeightField.cpp b/Engine/Source/Runtime/Experimental/Chaos/Private/Chaos/HeightField.cpp
--- a/Engine/Source/Runtime/Experimental/Chaos/Private/Chaos/HeightField.cpp
+++ b/Engine/Source/Runtime/Experimental/Chaos/Private/Chaos/HeightField.cpp
@@ -2043,42 +2043,42 @@ namespace Chaos
 		return bHit;
 	}
 
-	bool FHeightField::SweepGeom(const TSphere<FReal, 3>& QueryGeom, const FRigidTransform3& StartTM, const FVec3& Dir, const FReal Length, FReal& OutTime, FVec3& OutPosition, FVec3& OutNormal, int32& OutFaceIndex, const FReal Thickness, bool bComputeMTD) const
+	bool FHeightField::SweepGeom(const TSphere<FReal, 3>& QueryGeom, const FRigidTransform3& StartTM, const FVec3& Dir, const FReal Length, FReal& OutTime, FVec3& OutPosition, FVec3& OutNormal, int32& OutFaceIndex, FVec3& OutFaceNormal, const FReal Thickness, bool bComputeMTD) const
 	{
 		return SweepGeomImp(QueryGeom, StartTM, Dir, Length, OutTime, OutPosition, OutNormal, OutFaceIndex, Thickness, bComputeMTD);
 	}
 
-	bool FHeightField::SweepGeom(const TBox<FReal, 3>& QueryGeom, const FRigidTransform3& StartTM, const FVec3& Dir, const FReal Length, FReal& OutTime, FVec3& OutPosition, FVec3& OutNormal, int32& OutFaceIndex, const FReal Thickness, bool bComputeMTD) const
+	bool FHeightField::SweepGeom(const TBox<FReal, 3>& QueryGeom, const FRigidTransform3& StartTM, const FVec3& Dir, const FReal Length, FReal& OutTime, FVec3& OutPosition, FVec3& OutNormal, int32& OutFaceIndex, FVec3& OutFaceNormal, const FReal Thickness, bool bComputeMTD) const
 	{
 		return SweepGeomImp(QueryGeom, StartTM, Dir, Length, OutTime, OutPosition, OutNormal, OutFaceIndex, Thickness, bComputeMTD);
 	}
 
-	bool FHeightField::SweepGeom(const FCapsule& QueryGeom, const FRigidTransform3& StartTM, const FVec3& Dir, const FReal Length, FReal& OutTime, FVec3& OutPosition, FVec3& OutNormal, int32& OutFaceIndex, const FReal Thickness, bool bComputeMTD) const
+	bool FHeightField::SweepGeom(const FCapsule& QueryGeom, const FRigidTransform3& StartTM, const FVec3& Dir, const FReal Length, FReal& OutTime, FVec3& OutPosition, FVec3& OutNormal, int32& OutFaceIndex, FVec3& OutFaceNormal, const FReal Thickness, bool bComputeMTD) const
 	{
 		return SweepGeomImp(QueryGeom, StartTM, Dir, Length, OutTime, OutPosition, OutNormal, OutFaceIndex, Thickness, bComputeMTD);
 	}
 
-	bool FHeightField::SweepGeom(const FConvex& QueryGeom, const FRigidTransform3& StartTM, const FVec3& Dir, const FReal Length, FReal& OutTime, FVec3& OutPosition, FVec3& OutNormal, int32& OutFaceIndex, const FReal Thickness, bool bComputeMTD) const
+	bool FHeightField::SweepGeom(const FConvex& QueryGeom, const FRigidTransform3& StartTM, const FVec3& Dir, const FReal Length, FReal& OutTime, FVec3& OutPosition, FVec3& OutNormal, int32& OutFaceIndex, FVec3& OutFaceNormal, const FReal Thickness, bool bComputeMTD) const
 	{
 		return SweepGeomImp(QueryGeom, StartTM, Dir, Length, OutTime, OutPosition, OutNormal, OutFaceIndex, Thickness, bComputeMTD);
 	}
 
-	bool FHeightField::SweepGeom(const TImplicitObjectScaled<TSphere<FReal, 3>>& QueryGeom, const FRigidTransform3& StartTM, const FVec3& Dir, const FReal Length, FReal& OutTime, FVec3& OutPosition, FVec3& OutNormal, int32& OutFaceIndex, const FReal Thickness, bool bComputeMTD) const
+	bool FHeightField::SweepGeom(const TImplicitObjectScaled<TSphere<FReal, 3>>& QueryGeom, const FRigidTransform3& StartTM, const FVec3& Dir, const FReal Length, FReal& OutTime, FVec3& OutPosition, FVec3& OutNormal, int32& OutFaceIndex, FVec3& OutFaceNormal, const FReal Thickness, bool bComputeMTD) const
 	{
 		return SweepGeomImp(QueryGeom, StartTM, Dir, Length, OutTime, OutPosition, OutNormal, OutFaceIndex, Thickness, bComputeMTD);
 	}
 
-	bool FHeightField::SweepGeom(const TImplicitObjectScaled<TBox<FReal, 3>>& QueryGeom, const FRigidTransform3& StartTM, const FVec3& Dir, const FReal Length, FReal& OutTime, FVec3& OutPosition, FVec3& OutNormal, int32& OutFaceIndex, const FReal Thickness, bool bComputeMTD) const
+	bool FHeightField::SweepGeom(const TImplicitObjectScaled<TBox<FReal, 3>>& QueryGeom, const FRigidTransform3& StartTM, const FVec3& Dir, const FReal Length, FReal& OutTime, FVec3& OutPosition, FVec3& OutNormal, int32& OutFaceIndex, FVec3& OutFaceNormal, const FReal Thickness, bool bComputeMTD) const
 	{
 		return SweepGeomImp(QueryGeom, StartTM, Dir, Length, OutTime, OutPosition, OutNormal, OutFaceIndex, Thickness, bComputeMTD);
 	}
 
-	bool FHeightField::SweepGeom(const TImplicitObjectScaled<FCapsule>& QueryGeom, const FRigidTransform3& StartTM, const FVec3& Dir, const FReal Length, FReal& OutTime, FVec3& OutPosition, FVec3& OutNormal, int32& OutFaceIndex, const FReal Thickness, bool bComputeMTD) const
+	bool FHeightField::SweepGeom(const TImplicitObjectScaled<FCapsule>& QueryGeom, const FRigidTransform3& StartTM, const FVec3& Dir, const FReal Length, FReal& OutTime, FVec3& OutPosition, FVec3& OutNormal, int32& OutFaceIndex, FVec3& OutFaceNormal, const FReal Thickness, bool bComputeMTD) const
 	{
 		return SweepGeomImp(QueryGeom, StartTM, Dir, Length, OutTime, OutPosition, OutNormal, OutFaceIndex, Thickness, bComputeMTD);
 	}
 
-	bool FHeightField::SweepGeom(const TImplicitObjectScaled<FConvex>& QueryGeom, const FRigidTransform3& StartTM, const FVec3& Dir, const FReal Length, FReal& OutTime, FVec3& OutPosition, FVec3& OutNormal, int32& OutFaceIndex, const FReal Thickness, bool bComputeMTD) const
+	bool FHeightField::SweepGeom(const TImplicitObjectScaled<FConvex>& QueryGeom, const FRigidTransform3& StartTM, const FVec3& Dir, const FReal Length, FReal& OutTime, FVec3& OutPosition, FVec3& OutNormal, int32& OutFaceIndex, FVec3& OutFaceNormal, const FReal Thickness, bool bComputeMTD) const
 	{
 		return SweepGeomImp(QueryGeom, StartTM, Dir, Length, OutTime, OutPosition, OutNormal, OutFaceIndex, Thickness, bComputeMTD);
 	}
diff --git a/Engine/Source/Runtime/Experimental/Chaos/Private/Chaos/TriangleMeshImplicitObject.cpp b/Engine/Source/Runtime/Experimental/Chaos/Private/Chaos/TriangleMeshImplicitObject.cpp
--- a/Engine/Source/Runtime/Experimental/Chaos/Private/Chaos/TriangleMeshImplicitObject.cpp
+++ b/Engine/Source/Runtime/Experimental/Chaos/Private/Chaos/TriangleMeshImplicitObject.cpp
@@ -894,6 +894,7 @@ struct FTriangleMeshSweepVisitor
 	, LengthScale(InLengthScale)
 	, ScaledStartTM(InScaledStartTM)
 	, OutTime(TNumericLimits<FReal>::Max())
+	, OutFaceIndex(INDEX_NONE)
 	, TriMeshScale(InTriMeshScale)
 	{
 		VectorScaledDirNormalized = MakeVectorRegisterFloatFromDouble(MakeVectorRegister(InScaledDirNormalized.X, InScaledDirNormalized.Y, InScaledDirNormalized.Z, 0.0));
@@ -946,10 +947,10 @@ struct FTriangleMeshSweepVisitor
 		C = VectorMultiply(C, TriMeshScaleVector);
 
 		FTriangleRegister Tri(A, B, C);
+		const VectorRegister4Float TriNormal = VectorCross(VectorSubtract(B, A), VectorSubtract(C, A));
 
 		if(CullsBackFaceSweepsCode != 0)
 		{
-			const VectorRegister4Float TriNormal = VectorCross(VectorSubtract(B, A), VectorSubtract(C, A));
 			const VectorRegister4Float ReturnTrue = VectorCompareGT(VectorMultiply(VectorDot3(TriNormal, VectorScaledDirNormalized), VectorCullsBackFaceSweepsCode), VectorZero());
 			if (VectorMaskBits(ReturnTrue))
 			{
@@ -965,6 +966,7 @@ struct FTriangleMeshSweepVisitor
 				TransformSweepOutputsHelper(TriMeshScale, HitNormal, HitPosition, LengthScale, Time, OutNormal, OutPosition, OutTime);
 
 				OutFaceIndex = TriIdx;
+				VectorStoreFloat3(TriNormal, &OutFaceNormal);
 
 				if(Time <= 0)	//MTD or initial overlap
 				{
@@ -1001,6 +1003,7 @@ struct FTriangleMeshSweepVisitor
 	FVec3 OutPosition;
 	FVec3 OutNormal;
 	int32 OutFaceIndex;
+	FVec3 OutFaceNormal;
 
 	FVec3 TriMeshScale;
 };
@@ -1040,8 +1043,9 @@ FVec3 SafeInvScale(const FVec3& Scale)
 }
 
 template <typename QueryGeomType>
-bool FTriangleMeshImplicitObject::SweepGeomImp(const QueryGeomType& QueryGeom, const FRigidTransform3& StartTM, const FVec3& Dir, const FReal Length,
-	FReal& OutTime, FVec3& OutPosition, FVec3& OutNormal, int32& OutFaceIndex, const FReal Thickness, const bool bComputeMTD, FVec3 TriMeshScale) const
+bool FTriangleMeshImplicitObject::SweepGeomImp(const QueryGeomType& QueryGeom, const FRigidTransform3& StartTM, const FVec3& Dir, 
+	const FReal Length, FReal& OutTime, FVec3& OutPosition, FVec3& OutNormal, int32& OutFaceIndex, FVec3& OutFaceNormal, const FReal Thickness, 
+	const bool bComputeMTD, FVec3 TriMeshScale) const
 {
 	//QUICK_SCOPE_CYCLE_COUNTER(TrimeshSweep);
 	// Compute scaled sweep inputs to cache in visitor.
@@ -1069,6 +1073,7 @@ bool FTriangleMeshImplicitObject::SweepGeomImp(const QueryGeomType& QueryGeom, c
 			OutPosition = SQVisitor.OutPosition;
 			OutNormal = SQVisitor.OutNormal;
 			OutFaceIndex = SQVisitor.OutFaceIndex;
+			OutFaceNormal = GetFaceNormal(OutFaceIndex);
 			bHit = true;
 		}
 	};
@@ -1084,44 +1089,44 @@ bool FTriangleMeshImplicitObject::SweepGeomImp(const QueryGeomType& QueryGeom, c
 	return bHit;
 }
 
-bool FTriangleMeshImplicitObject::SweepGeom(const TSphere<FReal,3>& QueryGeom, const FRigidTransform3& StartTM, const FVec3& Dir, const FReal Length, FReal& OutTime, FVec3& OutPosition, FVec3& OutNormal, int32& OutFaceIndex, const FReal Thickness, const bool bComputeMTD, FVec3 TriMeshScale) const
+bool FTriangleMeshImplicitObject::SweepGeom(const TSphere<FReal,3>& QueryGeom, const FRigidTransform3& StartTM, const FVec3& Dir, const FReal Length, FReal& OutTime, FVec3& OutPosition, FVec3& OutNormal, int32& OutFaceIndex, FVec3& OutFaceNormal, const FReal Thickness, const bool bComputeMTD, FVec3 TriMeshScale) const
 {
-	return SweepGeomImp(QueryGeom, StartTM, Dir, Length, OutTime, OutPosition, OutNormal, OutFaceIndex, Thickness, bComputeMTD, TriMeshScale);
+	return SweepGeomImp(QueryGeom, StartTM, Dir, Length, OutTime, OutPosition, OutNormal, OutFaceIndex, OutFaceNormal, Thickness, bComputeMTD, TriMeshScale);
 }
 
-bool FTriangleMeshImplicitObject::SweepGeom(const TBox<FReal, 3>& QueryGeom, const FRigidTransform3& StartTM, const FVec3& Dir, const FReal Length, FReal& OutTime, FVec3& OutPosition, FVec3& OutNormal, int32& OutFaceIndex, const FReal Thickness, const bool bComputeMTD, FVec3 TriMeshScale) const
+bool FTriangleMeshImplicitObject::SweepGeom(const TBox<FReal, 3>& QueryGeom, const FRigidTransform3& StartTM, const FVec3& Dir, const FReal Length, FReal& OutTime, FVec3& OutPosition, FVec3& OutNormal, int32& OutFaceIndex, FVec3& OutFaceNormal, const FReal Thickness, const bool bComputeMTD, FVec3 TriMeshScale) const
 {
-	return SweepGeomImp(QueryGeom, StartTM, Dir, Length, OutTime, OutPosition, OutNormal, OutFaceIndex, Thickness, bComputeMTD, TriMeshScale);
+	return SweepGeomImp(QueryGeom, StartTM, Dir, Length, OutTime, OutPosition, OutNormal, OutFaceIndex, OutFaceNormal, Thickness, bComputeMTD, TriMeshScale);
 }
 
-bool FTriangleMeshImplicitObject::SweepGeom(const FCapsule& QueryGeom, const FRigidTransform3& StartTM, const FVec3& Dir, const FReal Length, FReal& OutTime, FVec3& OutPosition, FVec3& OutNormal, int32& OutFaceIndex, const FReal Thickness, const bool bComputeMTD, FVec3 TriMeshScale) const
+bool FTriangleMeshImplicitObject::SweepGeom(const FCapsule& QueryGeom, const FRigidTransform3& StartTM, const FVec3& Dir, const FReal Length, FReal& OutTime, FVec3& OutPosition, FVec3& OutNormal, int32& OutFaceIndex, FVec3& OutFaceNormal, const FReal Thickness, const bool bComputeMTD, FVec3 TriMeshScale) const
 {
-	return SweepGeomImp(QueryGeom, StartTM, Dir, Length, OutTime, OutPosition, OutNormal, OutFaceIndex, Thickness, bComputeMTD, TriMeshScale);
+	return SweepGeomImp(QueryGeom, StartTM, Dir, Length, OutTime, OutPosition, OutNormal, OutFaceIndex, OutFaceNormal, Thickness, bComputeMTD, TriMeshScale);
 }
 
-bool FTriangleMeshImplicitObject::SweepGeom(const FConvex& QueryGeom, const FRigidTransform3& StartTM, const FVec3& Dir, const FReal Length, FReal& OutTime, FVec3& OutPosition, FVec3& OutNormal, int32& OutFaceIndex, const FReal Thickness, const bool bComputeMTD, FVec3 TriMeshScale) const
+bool FTriangleMeshImplicitObject::SweepGeom(const FConvex& QueryGeom, const FRigidTransform3& StartTM, const FVec3& Dir, const FReal Length, FReal& OutTime, FVec3& OutPosition, FVec3& OutNormal, int32& OutFaceIndex, FVec3& OutFaceNormal, const FReal Thickness, const bool bComputeMTD, FVec3 TriMeshScale) const
 {
-	return SweepGeomImp(QueryGeom, StartTM, Dir, Length, OutTime, OutPosition, OutNormal, OutFaceIndex, Thickness, bComputeMTD, TriMeshScale);
+	return SweepGeomImp(QueryGeom, StartTM, Dir, Length, OutTime, OutPosition, OutNormal, OutFaceIndex, OutFaceNormal, Thickness, bComputeMTD, TriMeshScale);
 }
 
-bool FTriangleMeshImplicitObject::SweepGeom(const TImplicitObjectScaled<TSphere<FReal, 3>>& QueryGeom, const FRigidTransform3& StartTM, const FVec3& Dir, const FReal Length, FReal& OutTime, FVec3& OutPosition, FVec3& OutNormal, int32& OutFaceIndex, const FReal Thickness, const bool bComputeMTD, FVec3 TriMeshScale) const
+bool FTriangleMeshImplicitObject::SweepGeom(const TImplicitObjectScaled<TSphere<FReal, 3>>& QueryGeom, const FRigidTransform3& StartTM, const FVec3& Dir, const FReal Length, FReal& OutTime, FVec3& OutPosition, FVec3& OutNormal, int32& OutFaceIndex, FVec3& OutFaceNormal, const FReal Thickness, const bool bComputeMTD, FVec3 TriMeshScale) const
 {
-	return SweepGeomImp(QueryGeom, StartTM, Dir, Length, OutTime, OutPosition, OutNormal, OutFaceIndex, Thickness, bComputeMTD, TriMeshScale);
+	return SweepGeomImp(QueryGeom, StartTM, Dir, Length, OutTime, OutPosition, OutNormal, OutFaceIndex, OutFaceNormal, Thickness, bComputeMTD, TriMeshScale);
 }
 
-bool FTriangleMeshImplicitObject::SweepGeom(const TImplicitObjectScaled<TBox<FReal, 3>>& QueryGeom, const FRigidTransform3& StartTM, const FVec3& Dir, const FReal Length, FReal& OutTime, FVec3& OutPosition, FVec3& OutNormal, int32& OutFaceIndex, const FReal Thickness, const bool bComputeMTD, FVec3 TriMeshScale) const
+bool FTriangleMeshImplicitObject::SweepGeom(const TImplicitObjectScaled<TBox<FReal, 3>>& QueryGeom, const FRigidTransform3& StartTM, const FVec3& Dir, const FReal Length, FReal& OutTime, FVec3& OutPosition, FVec3& OutNormal, int32& OutFaceIndex, FVec3& OutFaceNormal, const FReal Thickness, const bool bComputeMTD, FVec3 TriMeshScale) const
 {
-	return SweepGeomImp(QueryGeom, StartTM, Dir, Length, OutTime, OutPosition, OutNormal, OutFaceIndex, Thickness, bComputeMTD, TriMeshScale);
+	return SweepGeomImp(QueryGeom, StartTM, Dir, Length, OutTime, OutPosition, OutNormal, OutFaceIndex, OutFaceNormal, Thickness, bComputeMTD, TriMeshScale);
 }
 
-bool FTriangleMeshImplicitObject::SweepGeom(const TImplicitObjectScaled<FCapsule>& QueryGeom, const FRigidTransform3& StartTM, const FVec3& Dir, const FReal Length, FReal& OutTime, FVec3& OutPosition, FVec3& OutNormal, int32& OutFaceIndex, const FReal Thickness, const bool bComputeMTD, FVec3 TriMeshScale) const
+bool FTriangleMeshImplicitObject::SweepGeom(const TImplicitObjectScaled<FCapsule>& QueryGeom, const FRigidTransform3& StartTM, const FVec3& Dir, const FReal Length, FReal& OutTime, FVec3& OutPosition, FVec3& OutNormal, int32& OutFaceIndex, FVec3& OutFaceNormal, const FReal Thickness, const bool bComputeMTD, FVec3 TriMeshScale) const
 {
-	return SweepGeomImp(QueryGeom, StartTM, Dir, Length, OutTime, OutPosition, OutNormal, OutFaceIndex, Thickness, bComputeMTD, TriMeshScale);
+	return SweepGeomImp(QueryGeom, StartTM, Dir, Length, OutTime, OutPosition, OutNormal, OutFaceIndex, OutFaceNormal, Thickness, bComputeMTD, TriMeshScale);
 }
 
-bool FTriangleMeshImplicitObject::SweepGeom(const TImplicitObjectScaled<FConvex>& QueryGeom, const FRigidTransform3& StartTM, const FVec3& Dir, const FReal Length, FReal& OutTime, FVec3& OutPosition, FVec3& OutNormal, int32& OutFaceIndex, const FReal Thickness, const bool bComputeMTD, FVec3 TriMeshScale) const
+bool FTriangleMeshImplicitObject::SweepGeom(const TImplicitObjectScaled<FConvex>& QueryGeom, const FRigidTransform3& StartTM, const FVec3& Dir, const FReal Length, FReal& OutTime, FVec3& OutPosition, FVec3& OutNormal, int32& OutFaceIndex, FVec3& OutFaceNormal, const FReal Thickness, const bool bComputeMTD, FVec3 TriMeshScale) const
 {
-	return SweepGeomImp(QueryGeom, StartTM, Dir, Length, OutTime, OutPosition, OutNormal, OutFaceIndex, Thickness, bComputeMTD, TriMeshScale);
+	return SweepGeomImp(QueryGeom, StartTM, Dir, Length, OutTime, OutPosition, OutNormal, OutFaceIndex, OutFaceNormal, Thickness, bComputeMTD, TriMeshScale);
 }
 
 template <typename IdxType>
diff --git a/Engine/Source/Runtime/Experimental/Chaos/Public/Chaos/GeometryQueries.h b/Engine/Source/Runtime/Experimental/Chaos/Public/Chaos/GeometryQueries.h
--- a/Engine/Source/Runtime/Experimental/Chaos/Public/Chaos/GeometryQueries.h
+++ b/Engine/Source/Runtime/Experimental/Chaos/Public/Chaos/GeometryQueries.h
@@ -120,7 +120,7 @@ namespace Chaos
 	}
 
 	template <typename SweptGeometry>
-	bool SweepQuery(const FImplicitObject& A, const FRigidTransform3& ATM, const SweptGeometry& B, const FRigidTransform3& BTM, const FVec3& Dir, const FReal Length, FReal& OutTime, FVec3& OutPosition, FVec3& OutNormal, int32& OutFaceIndex, const FReal Thickness, const bool bComputeMTD)
+	bool SweepQuery(const FImplicitObject& A, const FRigidTransform3& ATM, const SweptGeometry& B, const FRigidTransform3& BTM, const FVec3& Dir, const FReal Length, FReal& OutTime, FVec3& OutPosition, FVec3& OutNormal, int32& OutFaceIndex, FVec3& OutFaceNormal, const FReal Thickness, const bool bComputeMTD)
 	{
 		const EImplicitObjectType AType = A.GetType();
 		constexpr EImplicitObjectType BType = SweptGeometry::StaticType();
@@ -131,7 +131,7 @@ namespace Chaos
 		{
 			const TImplicitObjectTransformed<FReal, 3>& TransformedA = static_cast<const TImplicitObjectTransformed<FReal, 3>&>(A);
 			const FRigidTransform3 NewATM = TransformedA.GetTransform() * ATM;
-			return SweepQuery(*TransformedA.GetTransformedObject(), NewATM, B, BTM, Dir, Length, OutTime, OutPosition, OutNormal, OutFaceIndex, Thickness, bComputeMTD);
+			return SweepQuery(*TransformedA.GetTransformedObject(), NewATM, B, BTM, Dir, Length, OutTime, OutPosition, OutNormal, OutFaceIndex, OutFaceNormal, Thickness, bComputeMTD);
 		}
 
 		OutFaceIndex = INDEX_NONE;
@@ -204,13 +204,13 @@ namespace Chaos
 				case ImplicitObjectType::HeightField:
 				{
 					const FHeightField& AHeightField = static_cast<const FHeightField&>(A);
-					bResult = AHeightField.SweepGeom(B, BToATM, LocalDir, Length, OutTime, LocalPosition, LocalNormal, OutFaceIndex, Thickness, bComputeMTD);
+					bResult = AHeightField.SweepGeom(B, BToATM, LocalDir, Length, OutTime, LocalPosition, LocalNormal, OutFaceIndex, OutFaceNormal, Thickness, bComputeMTD);
 					break;
 				}
 				case ImplicitObjectType::TriangleMesh:
 				{
 					const FTriangleMeshImplicitObject& ATriangleMesh = static_cast<const FTriangleMeshImplicitObject&>(A);
-					bResult = ATriangleMesh.SweepGeom(B, BToATM, LocalDir, Length, OutTime, LocalPosition, LocalNormal, OutFaceIndex, Thickness, bComputeMTD);
+					bResult = ATriangleMesh.SweepGeom(B, BToATM, LocalDir, Length, OutTime, LocalPosition, LocalNormal, OutFaceIndex, OutFaceNormal, Thickness, bComputeMTD);
 					break;
 				}
 				case ImplicitObjectType::LevelSet:
@@ -223,13 +223,13 @@ namespace Chaos
 				if (IsScaled(AType))
 				{
 					const auto& AScaled = TImplicitObjectScaled<FTriangleMeshImplicitObject>::AsScaledChecked(A);
-					bResult = AScaled.LowLevelSweepGeom(B, BToATM, LocalDir, Length, OutTime, LocalPosition, LocalNormal, OutFaceIndex, Thickness, bComputeMTD);
+					bResult = AScaled.LowLevelSweepGeom(B, BToATM, LocalDir, Length, OutTime, LocalPosition, LocalNormal, OutFaceIndex, OutFaceNormal, Thickness, bComputeMTD);
 					break;
 				}
 				else if(IsInstanced(AType))
 				{
 					const auto& Instanced = TImplicitObjectInstanced<FTriangleMeshImplicitObject>::AsInstancedChecked(A);
-					bResult = Instanced.LowLevelSweepGeom(B,BToATM,LocalDir,Length,OutTime,LocalPosition,LocalNormal,OutFaceIndex,Thickness,bComputeMTD);
+					bResult = Instanced.LowLevelSweepGeom(B, BToATM, LocalDir, Length, OutTime, LocalPosition, LocalNormal, OutFaceIndex, OutFaceNormal, Thickness, bComputeMTD);
 					break;
 				}
 				else
diff --git a/Engine/Source/Runtime/Experimental/Chaos/Public/Chaos/HeightField.h b/Engine/Source/Runtime/Experimental/Chaos/Public/Chaos/HeightField.h
--- a/Engine/Source/Runtime/Experimental/Chaos/Public/Chaos/HeightField.h
+++ b/Engine/Source/Runtime/Experimental/Chaos/Public/Chaos/HeightField.h
@@ -67,14 +67,14 @@ namespace Chaos
 		bool OverlapGeom(const TImplicitObjectScaled<FCapsule>& QueryGeom, const FRigidTransform3& QueryTM, const FReal Thickness, FMTDInfo* OutMTD = nullptr) const;
 		bool OverlapGeom(const TImplicitObjectScaled<FConvex>& QueryGeom, const FRigidTransform3& QueryTM, const FReal Thickness, FMTDInfo* OutMTD = nullptr) const;
 
-		bool SweepGeom(const TSphere<FReal, 3>& QueryGeom, const FRigidTransform3& StartTM, const FVec3& Dir, const FReal Length, FReal& OutTime, FVec3& OutPosition, FVec3& OutNormal, int32& OutFaceIndex, const FReal Thickness = 0, bool bComputeMTD = false) const;
-		bool SweepGeom(const TBox<FReal, 3>& QueryGeom, const FRigidTransform3& StartTM, const FVec3& Dir, const FReal Length, FReal& OutTime, FVec3& OutPosition, FVec3& OutNormal, int32& OutFaceIndex, const FReal Thickness = 0, bool bComputeMTD = false) const;
-		bool SweepGeom(const FCapsule& QueryGeom, const FRigidTransform3& StartTM, const FVec3& Dir, const FReal Length, FReal& OutTime, FVec3& OutPosition, FVec3& OutNormal, int32& OutFaceIndex, const FReal Thickness = 0, bool bComputeMTD = false) const;
-		bool SweepGeom(const FConvex& QueryGeom, const FRigidTransform3& StartTM, const FVec3& Dir, const FReal Length, FReal& OutTime, FVec3& OutPosition, FVec3& OutNormal, int32& OutFaceIndex, const FReal Thickness = 0, bool bComputeMTD = false) const;
-		bool SweepGeom(const TImplicitObjectScaled<TSphere<FReal, 3>>& QueryGeom, const FRigidTransform3& StartTM, const FVec3& Dir, const FReal Length, FReal& OutTime, FVec3& OutPosition, FVec3& OutNormal, int32& OutFaceIndex, const FReal Thickness = 0, bool bComputeMTD = false) const;
-		bool SweepGeom(const TImplicitObjectScaled<TBox<FReal, 3>>& QueryGeom, const FRigidTransform3& StartTM, const FVec3& Dir, const FReal Length, FReal& OutTime, FVec3& OutPosition, FVec3& OutNormal, int32& OutFaceIndex, const FReal Thickness = 0, bool bComputeMTD = false) const;
-		bool SweepGeom(const TImplicitObjectScaled<FCapsule>& QueryGeom, const FRigidTransform3& StartTM, const FVec3& Dir, const FReal Length, FReal& OutTime, FVec3& OutPosition, FVec3& OutNormal, int32& OutFaceIndex, const FReal Thickness = 0, bool bComputeMTD = false) const;
-		bool SweepGeom(const TImplicitObjectScaled<FConvex>& QueryGeom, const FRigidTransform3& StartTM, const FVec3& Dir, const FReal Length, FReal& OutTime, FVec3& OutPosition, FVec3& OutNormal, int32& OutFaceIndex, const FReal Thickness = 0, bool bComputeMTD = false) const;
+		bool SweepGeom(const TSphere<FReal, 3>& QueryGeom, const FRigidTransform3& StartTM, const FVec3& Dir, const FReal Length, FReal& OutTime, FVec3& OutPosition, FVec3& OutNormal, int32& OutFaceIndex, FVec3& OutFaceNormal, const FReal Thickness = 0, bool bComputeMTD = false) const;
+		bool SweepGeom(const TBox<FReal, 3>& QueryGeom, const FRigidTransform3& StartTM, const FVec3& Dir, const FReal Length, FReal& OutTime, FVec3& OutPosition, FVec3& OutNormal, int32& OutFaceIndex, FVec3& OutFaceNormal, const FReal Thickness = 0, bool bComputeMTD = false) const;
+		bool SweepGeom(const FCapsule& QueryGeom, const FRigidTransform3& StartTM, const FVec3& Dir, const FReal Length, FReal& OutTime, FVec3& OutPosition, FVec3& OutNormal, int32& OutFaceIndex, FVec3& OutFaceNormal, const FReal Thickness = 0, bool bComputeMTD = false) const;
+		bool SweepGeom(const FConvex& QueryGeom, const FRigidTransform3& StartTM, const FVec3& Dir, const FReal Length, FReal& OutTime, FVec3& OutPosition, FVec3& OutNormal, int32& OutFaceIndex, FVec3& OutFaceNormal, const FReal Thickness = 0, bool bComputeMTD = false) const;
+		bool SweepGeom(const TImplicitObjectScaled<TSphere<FReal, 3>>& QueryGeom, const FRigidTransform3& StartTM, const FVec3& Dir, const FReal Length, FReal& OutTime, FVec3& OutPosition, FVec3& OutNormal, int32& OutFaceIndex, FVec3& OutFaceNormal, const FReal Thickness = 0, bool bComputeMTD = false) const;
+		bool SweepGeom(const TImplicitObjectScaled<TBox<FReal, 3>>& QueryGeom, const FRigidTransform3& StartTM, const FVec3& Dir, const FReal Length, FReal& OutTime, FVec3& OutPosition, FVec3& OutNormal, int32& OutFaceIndex, FVec3& OutFaceNormal, const FReal Thickness = 0, bool bComputeMTD = false) const;
+		bool SweepGeom(const TImplicitObjectScaled<FCapsule>& QueryGeom, const FRigidTransform3& StartTM, const FVec3& Dir, const FReal Length, FReal& OutTime, FVec3& OutPosition, FVec3& OutNormal, int32& OutFaceIndex, FVec3& OutFaceNormal, const FReal Thickness = 0, bool bComputeMTD = false) const;
+		bool SweepGeom(const TImplicitObjectScaled<FConvex>& QueryGeom, const FRigidTransform3& StartTM, const FVec3& Dir, const FReal Length, FReal& OutTime, FVec3& OutPosition, FVec3& OutNormal, int32& OutFaceIndex, FVec3& OutFaceNormal, const FReal Thickness = 0, bool bComputeMTD = false) const;
 
 		bool GJKContactPoint(const TBox<FReal, 3>& QueryGeom, const FRigidTransform3& QueryTM, const FReal Thickness, FVec3& ContactLocation, FVec3& ContactNormal, FReal& ContactPhi) const;
 		bool GJKContactPoint(const TSphere<FReal, 3>& QueryGeom, const FRigidTransform3& QueryTM, const FReal Thickness, FVec3& ContactLocation, FVec3& ContactNormal, FReal& ContactPhi) const;
diff --git a/Engine/Source/Runtime/Experimental/Chaos/Public/Chaos/ImplicitObjectScaled.h b/Engine/Source/Runtime/Experimental/Chaos/Public/Chaos/ImplicitObjectScaled.h
--- a/Engine/Source/Runtime/Experimental/Chaos/Public/Chaos/ImplicitObjectScaled.h
+++ b/Engine/Source/Runtime/Experimental/Chaos/Public/Chaos/ImplicitObjectScaled.h
@@ -218,9 +218,9 @@ public:
 
 	/** This is a low level function and assumes the internal object has a SweepGeom function. Should not be called directly. See GeometryQueries.h : SweepQuery */
 	template <typename QueryGeomType>
-	bool LowLevelSweepGeom(const QueryGeomType& B,const TRigidTransform<T,d>& BToATM,const TVector<T,d>& LocalDir,const T Length,T& OutTime,TVector<T,d>& LocalPosition,TVector<T,d>& LocalNormal,int32& OutFaceIndex,T Thickness = 0,bool bComputeMTD = false) const
+	bool LowLevelSweepGeom(const QueryGeomType& B, const TRigidTransform<T, d>& BToATM, const TVector<T, d>& LocalDir, const T Length, T& OutTime, TVector<T, d>& LocalPosition, TVector<T, d>& LocalNormal, int32& OutFaceIndex, TVector<T, d>& OutFaceNormal, T Thickness = 0, bool bComputeMTD = false) const
 	{
-		return MObject->SweepGeom(B,BToATM,LocalDir,Length,OutTime,LocalPosition,LocalNormal,OutFaceIndex,Thickness,bComputeMTD);
+		return MObject->SweepGeom(B, BToATM, LocalDir, Length, OutTime, LocalPosition, LocalNormal, OutFaceIndex, OutFaceNormal, Thickness, bComputeMTD);
 	}
 
 	/** This is a low level function and assumes the internal object has a OverlapGeom function. Should not be called directly. See GeometryQueries.h : OverlapQuery */
@@ -632,7 +632,7 @@ public:
 
 	/** This is a low level function and assumes the internal object has a SweepGeom function. Should not be called directly. See GeometryQueries.h : SweepQuery */
 	template <typename QueryGeomType>
-	bool LowLevelSweepGeom(const QueryGeomType& B, const TRigidTransform<T, d>& BToATM, const TVector<T, d>& LocalDir, const T Length, T& OutTime, TVector<T, d>& LocalPosition, TVector<T, d>& LocalNormal, int32& OutFaceIndex, T Thickness = 0, bool bComputeMTD = false) const
+	bool LowLevelSweepGeom(const QueryGeomType& B, const TRigidTransform<T, d>& BToATM, const TVector<T, d>& LocalDir, const T Length, T& OutTime, TVector<T, d>& LocalPosition, TVector<T, d>& LocalNormal, int32& OutFaceIndex, TVector<T, d>& OutFaceNormal, T Thickness = 0, bool bComputeMTD = false) const
 	{
 		ensure(Length > 0);
 		ensure(FMath::IsNearlyEqual(LocalDir.SizeSquared(), (T)1, (T)KINDA_SMALL_NUMBER));
@@ -652,7 +652,7 @@ public:
 
 			TRigidTransform<T, d> BToATMNoScale(BToATM.GetLocation() * MInvScale, BToATM.GetRotation());
 			
-			if (MObject->SweepGeom(B, BToATMNoScale, UnscaledDir, UnscaledLength, UnscaledTime, UnscaledPosition, UnscaledNormal, OutFaceIndex, Thickness, bComputeMTD, MScale))
+			if (MObject->SweepGeom(B, BToATMNoScale, UnscaledDir, UnscaledLength, UnscaledTime, UnscaledPosition, UnscaledNormal, OutFaceIndex, OutFaceNormal, Thickness, bComputeMTD, MScale))
 			{
 				const T NewTime = LengthScaleInv * UnscaledTime;
 				//We double check that NewTime < Length because of potential precision issues. When that happens we always keep the shortest hit first
diff --git a/Engine/Source/Runtime/Experimental/Chaos/Public/Chaos/TriangleMeshImplicitObject.h b/Engine/Source/Runtime/Experimental/Chaos/Public/Chaos/TriangleMeshImplicitObject.h
--- a/Engine/Source/Runtime/Experimental/Chaos/Public/Chaos/TriangleMeshImplicitObject.h
+++ b/Engine/Source/Runtime/Experimental/Chaos/Public/Chaos/TriangleMeshImplicitObject.h
@@ -218,15 +218,15 @@ namespace Chaos
 		bool OverlapGeom(const TImplicitObjectScaled<FCapsule>& QueryGeom, const FRigidTransform3& QueryTM, const FReal Thickness, FMTDInfo* OutMTD = nullptr, FVec3 TriMeshScale = FVec3(1.0f)) const;
 		bool OverlapGeom(const TImplicitObjectScaled<FConvex>& QueryGeom, const FRigidTransform3& QueryTM, const FReal Thickness, FMTDInfo* OutMTD = nullptr, FVec3 TriMeshScale = FVec3(1.0f)) const;
 
-		bool SweepGeom(const TSphere<FReal, 3>& QueryGeom, const FRigidTransform3& StartTM, const FVec3& Dir, const FReal Length, FReal& OutTime, FVec3& OutPosition, FVec3& OutNormal, int32& OutFaceIndex, const FReal Thickness = 0, const bool bComputeMTD = false, FVec3 TriMeshScale = FVec3(1.0f)) const;
-		bool SweepGeom(const TBox<FReal, 3>& QueryGeom, const FRigidTransform3& StartTM, const FVec3& Dir, const FReal Length, FReal& OutTime, FVec3& OutPosition, FVec3& OutNormal, int32& OutFaceIndex, const FReal Thickness = 0, const bool bComputeMTD = false, FVec3 TriMeshScale = FVec3(1.0f)) const;
-		bool SweepGeom(const FCapsule& QueryGeom, const FRigidTransform3& StartTM, const FVec3& Dir, const FReal Length, FReal& OutTime, FVec3& OutPosition, FVec3& OutNormal, int32& OutFaceIndex, const FReal Thickness = 0, const bool bComputeMTD = false, FVec3 TriMeshScale = FVec3(1.0f)) const;
-		bool SweepGeom(const FConvex& QueryGeom, const FRigidTransform3& StartTM, const FVec3& Dir, const FReal Length, FReal& OutTime, FVec3& OutPosition, FVec3& OutNormal, int32& OutFaceIndex, const FReal Thickness = 0, const bool bComputeMTD = false, FVec3 TriMeshScale = FVec3(1.0f)) const;
+		bool SweepGeom(const TSphere<FReal, 3>& QueryGeom, const FRigidTransform3& StartTM, const FVec3& Dir, const FReal Length, FReal& OutTime, FVec3& OutPosition, FVec3& OutNormal, int32& OutFaceIndex, FVec3& OutFaceNormal, const FReal Thickness = 0, const bool bComputeMTD = false, FVec3 TriMeshScale = FVec3(1.0f)) const;
+		bool SweepGeom(const TBox<FReal, 3>& QueryGeom, const FRigidTransform3& StartTM, const FVec3& Dir, const FReal Length, FReal& OutTime, FVec3& OutPosition, FVec3& OutNormal, int32& OutFaceIndex, FVec3& OutFaceNormal, const FReal Thickness = 0, const bool bComputeMTD = false, FVec3 TriMeshScale = FVec3(1.0f)) const;
+		bool SweepGeom(const FCapsule& QueryGeom, const FRigidTransform3& StartTM, const FVec3& Dir, const FReal Length, FReal& OutTime, FVec3& OutPosition, FVec3& OutNormal, int32& OutFaceIndex, FVec3& OutFaceNormal, const FReal Thickness = 0, const bool bComputeMTD = false, FVec3 TriMeshScale = FVec3(1.0f)) const;
+		bool SweepGeom(const FConvex& QueryGeom, const FRigidTransform3& StartTM, const FVec3& Dir, const FReal Length, FReal& OutTime, FVec3& OutPosition, FVec3& OutNormal, int32& OutFaceIndex, FVec3& OutFaceNormal, const FReal Thickness = 0, const bool bComputeMTD = false, FVec3 TriMeshScale = FVec3(1.0f)) const;
 
-		bool SweepGeom(const TImplicitObjectScaled<TSphere<FReal, 3>>& QueryGeom, const FRigidTransform3& StartTM, const FVec3& Dir, const FReal Length, FReal& OutTime, FVec3& OutPosition, FVec3& OutNormal, int32& OutFaceIndex, const FReal Thickness = 0, const bool bComputeMTD = false, FVec3 TriMeshScale = FVec3(1.0f)) const;
-		bool SweepGeom(const TImplicitObjectScaled<TBox<FReal, 3>>& QueryGeom, const FRigidTransform3& StartTM, const FVec3& Dir, const FReal Length, FReal& OutTime, FVec3& OutPosition, FVec3& OutNormal, int32& OutFaceIndex, const FReal Thickness = 0, const bool bComputeMTD = false, FVec3 TriMeshScale = FVec3(1.0f)) const;
-		bool SweepGeom(const TImplicitObjectScaled<FCapsule>& QueryGeom, const FRigidTransform3& StartTM, const FVec3& Dir, const FReal Length, FReal& OutTime, FVec3& OutPosition, FVec3& OutNormal, int32& OutFaceIndex, const FReal Thickness = 0, const bool bComputeMTD = false, FVec3 TriMeshScale = FVec3(1.0f)) const;
-		bool SweepGeom(const TImplicitObjectScaled<FConvex>& QueryGeom, const FRigidTransform3& StartTM, const FVec3& Dir, const FReal Length, FReal& OutTime, FVec3& OutPosition, FVec3& OutNormal, int32& OutFaceIndex, const FReal Thickness = 0, const bool bComputeMTD = false, FVec3 TriMeshScale = FVec3(1.0f)) const;
+		bool SweepGeom(const TImplicitObjectScaled<TSphere<FReal, 3>>& QueryGeom, const FRigidTransform3& StartTM, const FVec3& Dir, const FReal Length, FReal& OutTime, FVec3& OutPosition, FVec3& OutNormal, int32& OutFaceIndex, FVec3& OutFaceNormal, const FReal Thickness = 0, const bool bComputeMTD = false, FVec3 TriMeshScale = FVec3(1.0f)) const;
+		bool SweepGeom(const TImplicitObjectScaled<TBox<FReal, 3>>& QueryGeom, const FRigidTransform3& StartTM, const FVec3& Dir, const FReal Length, FReal& OutTime, FVec3& OutPosition, FVec3& OutNormal, int32& OutFaceIndex, FVec3& OutFaceNormal, const FReal Thickness = 0, const bool bComputeMTD = false, FVec3 TriMeshScale = FVec3(1.0f)) const;
+		bool SweepGeom(const TImplicitObjectScaled<FCapsule>& QueryGeom, const FRigidTransform3& StartTM, const FVec3& Dir, const FReal Length, FReal& OutTime, FVec3& OutPosition, FVec3& OutNormal, int32& OutFaceIndex, FVec3& OutFaceNormal, const FReal Thickness = 0, const bool bComputeMTD = false, FVec3 TriMeshScale = FVec3(1.0f)) const;
+		bool SweepGeom(const TImplicitObjectScaled<FConvex>& QueryGeom, const FRigidTransform3& StartTM, const FVec3& Dir, const FReal Length, FReal& OutTime, FVec3& OutPosition, FVec3& OutNormal, int32& OutFaceIndex, FVec3& OutFaceNormal, const FReal Thickness = 0, const bool bComputeMTD = false, FVec3 TriMeshScale = FVec3(1.0f)) const;
 
 		bool GJKContactPoint(const TSphere<FReal, 3>& QueryGeom, const FRigidTransform3& QueryTM, const FReal Thickness, FVec3& Location, FVec3& Normal, FReal& Penetration) const;
 		bool GJKContactPoint(const TBox<FReal, 3>& QueryGeom, const FRigidTransform3& QueryTM, const FReal Thickness, FVec3& Location, FVec3& Normal, FReal& Penetration) const;
@@ -512,7 +512,7 @@ namespace Chaos
 		bool OverlapGeomImp(const QueryGeomType& QueryGeom, const FRigidTransform3& QueryTM, const FReal Thickness, FMTDInfo* OutMTD = nullptr, FVec3 TriMeshScale = FVec3(1.0f)) const;
 
 		template<typename QueryGeomType>
-		bool SweepGeomImp(const QueryGeomType& QueryGeom, const FRigidTransform3& StartTM, const FVec3& Dir, const FReal Length, FReal& OutTime, FVec3& OutPosition, FVec3& OutNormal, int32& OutFaceIndex, const FReal Thickness, const bool bComputeMTD, FVec3 TriMeshScale = FVec3(1.0f)) const;
+		bool SweepGeomImp(const QueryGeomType& QueryGeom, const FRigidTransform3& StartTM, const FVec3& Dir, const FReal Length, FReal& OutTime, FVec3& OutPosition, FVec3& OutNormal, int32& OutFaceIndex, FVec3& OutFaceNormal, const FReal Thickness, const bool bComputeMTD, FVec3 TriMeshScale = FVec3(1.0f)) const;
 
 		template <typename IdxType>
 		bool RaycastImp(const TArray<TVec3<IdxType>>& Elements, const FVec3& StartPoint, const FVec3& Dir, const FReal Length, const FReal Thickness, FReal& OutTime, FVec3& OutPosition, FVec3& OutNormal, int32& OutFaceIndex) const;
diff --git a/Engine/Source/Runtime/PhysicsCore/Private/SQAccelerator.cpp b/Engine/Source/Runtime/PhysicsCore/Private/SQAccelerator.cpp
--- a/Engine/Source/Runtime/PhysicsCore/Private/SQAccelerator.cpp
+++ b/Engine/Source/Runtime/PhysicsCore/Private/SQAccelerator.cpp
@@ -133,6 +133,7 @@ struct TSQVisitor : public Chaos::ISpatialVisitor<TPayload, Chaos::FReal>
 		, OutputFlags(InOutputFlags)
 		, bAnyHit(false)
 		, DebugParams(InDebugParams)
+		, HitFaceNormal(Chaos::FVec3::ZeroVector)
 		, HitBuffer(InHitBuffer)
 		, QueryFilterData(InQueryFilterData)
 #if PHYSICS_INTERFACE_PHYSX
@@ -158,6 +159,7 @@ struct TSQVisitor : public Chaos::ISpatialVisitor<TPayload, Chaos::FReal>
 		: HalfExtents(InQueryGeom.BoundingBox().Extents() * 0.5)
 		, bAnyHit(false)
 		, DebugParams(InDebugParams)
+		, HitFaceNormal(Chaos::FVec3::ZeroVector)
 		, HitBuffer(InHitBuffer)
 		, QueryFilterData(InQueryFilterData)
 #if PHYSICS_INTERFACE_PHYSX
@@ -324,6 +326,7 @@ private:
 				FVec3 WorldPosition, WorldNormal;
 				Chaos::FReal Distance = 0;	//not needed but fixes compiler warning for overlap
 				int32 FaceIdx = INDEX_NONE;	//not needed but fixes compiler warning for overlap
+				FVec3 FaceNormal = FVec3::ZeroVector;
 				const bool bComputeMTD = !!((uint16)(OutputFlags.HitFlags & EHitFlags::MTD));
 
 				if (SQ == ESQType::Raycast)
@@ -342,7 +345,7 @@ private:
 				}
 				else if(SQ == ESQType::Sweep && CurData->CurrentLength > 0 && ensure(QueryGeom))
 				{
-					bHit = SweepQuery(*Geom, ActorTM, *QueryGeom, StartTM, CurData->Dir, CurData->CurrentLength, Distance, WorldPosition, WorldNormal, FaceIdx, 0.f, bComputeMTD);
+					bHit = SweepQuery(*Geom, ActorTM, *QueryGeom, StartTM, CurData->Dir, CurData->CurrentLength, Distance, WorldPosition, WorldNormal, FaceIdx, FaceNormal, 0.f, bComputeMTD);
 				}
 				else if ((SQ == ESQType::Overlap || (SQ == ESQType::Sweep && CurData->CurrentLength == 0)) && ensure(QueryGeom))
 				{
@@ -364,39 +367,81 @@ private:
 				if(bHit)
 				{
 					//QUICK_SCOPE_CYCLE_COUNTER(SQNarrowHit);
-					FillHitHelper(Hit, Distance, WorldPosition, WorldNormal, FaceIdx, bComputeMTD);
+					
+					bool bAcceptHit = true;
+					if constexpr(SQ == ESQType::Sweep && std::is_same_v<THitType, FSweepHit>)
+					{
+						const THitType* CurrentHit = HitBuffer.GetCurrentHit();
+						if(FaceIdx != INDEX_NONE)
+						{
+							if(CurrentHit)
+							{
+								constexpr static FReal CoLocationEpsilon = 1e-6;
+								const FReal DistDelta = FMath::Abs(Distance - CurrentHit->Distance);
+								const int32 OldFace = CurrentHit->FaceIndex;
+
+								if(DistDelta < CoLocationEpsilon && OldFace != INDEX_NONE)
+								{
+									// We already have a face hit from another triangle mesh - see if this one is better (more opposing the sweep)
+									const FReal OldDot = FVec3::DotProduct(CurData->Dir, HitFaceNormal);
+									const FReal NewDot = FVec3::DotProduct(CurData->Dir, FaceNormal);
+
+									if(NewDot < OldDot)
+									{
+										// More opposing
+										HitFaceNormal = FaceNormal;
+									}
+									else
+									{
+										// This hit is co-located but has a worse normal
+										bAcceptHit = false;
+									}
+								}
+							}
+							
+							if(bAcceptHit)
+							{
+								// Record the new face normal
+								HitFaceNormal = FaceNormal;
+							}
+						}
+					}
 
+					if(bAcceptHit)
+					{
+						FillHitHelper(Hit, Distance, WorldPosition, WorldNormal, FaceIdx, bComputeMTD);
 
 #if PHYSICS_INTERFACE_PHYSX
-					HitType = QueryFilterData.flags & PxQueryFlag::ePOSTFILTER ? QueryCallback.PostFilter(QueryFilterDataConcrete, Hit) : HitType;
+						HitType = QueryFilterData.flags & PxQueryFlag::ePOSTFILTER ? QueryCallback.PostFilter(QueryFilterDataConcrete, Hit) : HitType;
 #else
-					HitType = QueryFilterData.flags & FChaosQueryFlag::ePOSTFILTER ? QueryCallback.PostFilter(QueryFilterDataConcrete, Hit) : HitType;
+						HitType = QueryFilterData.flags & FChaosQueryFlag::ePOSTFILTER ? QueryCallback.PostFilter(QueryFilterDataConcrete, Hit) : HitType;
 #endif
 
-					if (HitType != ECollisionQueryHitType::None)
-					{
+						if(HitType != ECollisionQueryHitType::None)
+						{
 
-						//overlap never blocks
-						const bool bBlocker = (HitType == ECollisionQueryHitType::Block || bAnyHit || HitBuffer.WantsSingleResult());
-						HitBuffer.InsertHit(Hit, bBlocker);
+							//overlap never blocks
+							const bool bBlocker = (HitType == ECollisionQueryHitType::Block || bAnyHit || HitBuffer.WantsSingleResult());
+							HitBuffer.InsertHit(Hit, bBlocker);
 #if CHAOS_DEBUG_DRAW
-						bHitBufferIncreased = true;
+							bHitBufferIncreased = true;
 #endif
 
-						if (bBlocker && SQ != ESQType::Overlap)
-						{
-							CurData->SetLength(FMath::Max((FReal)0., Distance));	//Max is needed for MTD which returns negative distance
-							if (CurData->CurrentLength == 0 && (SQ == ESQType::Raycast || HitBuffer.WantsSingleResult()))	//raycasts always fail with distance 0, sweeps only matter if we want multi overlaps
+							if(bBlocker && SQ != ESQType::Overlap)
 							{
-								bContinue = false; //initial overlap so nothing will be better than this
-								break;
+								CurData->SetLength(FMath::Max((FReal)0., Distance));	//Max is needed for MTD which returns negative distance
+								if(CurData->CurrentLength == 0 && (SQ == ESQType::Raycast || HitBuffer.WantsSingleResult()))	//raycasts always fail with distance 0, sweeps only matter if we want multi overlaps
+								{
+									bContinue = false; //initial overlap so nothing will be better than this
+									break;
+								}
 							}
-						}
 
-						if (bAnyHit)
-						{
-							bContinue = false;
-							break;
+							if(bAnyHit)
+							{
+								bContinue = false;
+								break;
+							}
 						}
 					}
 				}
@@ -452,6 +497,7 @@ private:
 	FHitFlags OutputFlags;
 	bool bAnyHit;
 	const FQueryDebugParams DebugParams;
+	Chaos::FVec3 HitFaceNormal;
 	ChaosInterface::FSQHitBuffer<THitType>& HitBuffer;
 	const FQueryFilterData& QueryFilterData;
 	const FCollisionFilterData QueryFilterDataConcrete;
diff --git a/Engine/Source/Runtime/PhysicsCore/Public/ChaosSQTypes.h b/Engine/Source/Runtime/PhysicsCore/Public/ChaosSQTypes.h
--- a/Engine/Source/Runtime/PhysicsCore/Public/ChaosSQTypes.h
+++ b/Engine/Source/Runtime/PhysicsCore/Public/ChaosSQTypes.h
@@ -241,6 +241,8 @@ namespace ChaosInterface
 			Hits.Add(InTouchingHit);
 		}
 
+		const HitType* GetCurrentHit() const { return bHasBlockingHit ? &CurrentBlockingHit : nullptr; }
+
 		/** Does not do any distance verification. This is up to the SQ code to manage */
 		void InsertHit(const HitType& Hit, bool bBlocking)
 		{
