Chaos :  potential crash fix : make sure we keep Global Payload index when inserting a payload in an existing leaf in the AABBTree

diff --git a/Engine/Source/Runtime/Experimental/Chaos/Public/Chaos/AABBTree.h b/Engine/Source/Runtime/Experimental/Chaos/Public/Chaos/AABBTree.h
--- a/Engine/Source/Runtime/Experimental/Chaos/Public/Chaos/AABBTree.h
+++ b/Engine/Source/Runtime/Experimental/Chaos/Public/Chaos/AABBTree.h
@@ -1320,7 +1320,8 @@ public:
 			Nodes[BestSibling].ChildrenBounds[0] = ExpandedBounds;
 			UpdateAncestorBounds(BestSibling, true);
 			FAABBTreePayloadInfo* PayloadInfo = PayloadToInfo.Find(Payload);
-			*PayloadInfo = FAABBTreePayloadInfo{ INDEX_NONE, INDEX_NONE, LeafIdx, INDEX_NONE, BestSibling };
+			PayloadInfo->LeafIdx = LeafIdx;
+			PayloadInfo->NodeIdx = BestSibling;
 			return;
 		}
 
