Chaos : fix capsule raycast intersection issue
- A ray parallel to the capsule axis, starting within its local bounding of the cylinder part but not inside it woudl report a hit when it should not
- The fix verify that the computed hit position is actually on or inside the cylinder

diff --git a/Engine/Source/Runtime/Experimental/Chaos/Public/Chaos/Capsule.h b/Engine/Source/Runtime/Experimental/Chaos/Public/Chaos/Capsule.h
--- a/Engine/Source/Runtime/Experimental/Chaos/Public/Chaos/Capsule.h
+++ b/Engine/Source/Runtime/Experimental/Chaos/Public/Chaos/Capsule.h
@@ -14,6 +14,8 @@
 
 #include "UObject/ReleaseObjectVersion.h"
 
+PRAGMA_DISABLE_OPTIMIZATION
+
 namespace Chaos
 {
 	struct FCapsuleSpecializeSamplingHelper;
@@ -285,8 +287,15 @@ namespace Chaos
 					const FReal PositionLengthOnCoreCylinder = FVec3::DotProduct(CylinderToSpherePosition, MVector);
 					if (PositionLengthOnCoreCylinder >= 0 && PositionLengthOnCoreCylinder < MHeight)
 					{
+						const FVec3 SegmentToPosition = (CylinderToSpherePosition - MVector * PositionLengthOnCoreCylinder);
+						if (SegmentToPosition.SquaredLength() > R2)
+						{
+							// the contact point is actually outside of the cylinder
+							// this can happen if the ray starts within the cylinder bounds but do not intersect with the cylinder
+							return false;
+						}
 						OutTime = Time + RemovedLength; // Account for ray clipped against bounds
-						OutNormal = (CylinderToSpherePosition - MVector * PositionLengthOnCoreCylinder) / R;
+						OutNormal = SegmentToPosition / R;
 						OutPosition = SpherePosition - OutNormal * Thickness;
 						return true;
 					}
@@ -738,3 +747,4 @@ namespace Chaos
 	using TCapsuleSpecializeSamplingHelper UE_DEPRECATED(4.27, "Deprecated. this class is to be deleted, use FCapsuleSpecializeSamplingHelper instead") = FCapsuleSpecializeSamplingHelper;
 
 } // namespace Chaos
+PRAGMA_ENABLE_OPTIMIZATION
\ No newline at end of file
