Chaos
- better joint force-point selection to improved behaviour with large constraint errors and mismatched masses
- joint velocity solve for unilateral constraints does not add negative velocity correction

diff --git a/Engine/Source/Runtime/Experimental/Chaos/Private/Chaos/Joint/PBDJointCachedSolverGaussSeidel.cpp b/Engine/Source/Runtime/Experimental/Chaos/Private/Chaos/Joint/PBDJointCachedSolverGaussSeidel.cpp
--- a/Engine/Source/Runtime/Experimental/Chaos/Private/Chaos/Joint/PBDJointCachedSolverGaussSeidel.cpp
+++ b/Engine/Source/Runtime/Experimental/Chaos/Private/Chaos/Joint/PBDJointCachedSolverGaussSeidel.cpp
@@ -391,9 +391,18 @@ void FPBDJointCachedSolver::InitPositionConstraints(
 
 	const TVec3<EJointMotionType>& LinearMotion = JointSettings.LinearMotionTypes;
 
-
-	const FVec3 ConstraintArm0Limited = ConnectorXs[1] - CurrentPs[0];
-	const FVec3 ConstraintArm1Limited = ConnectorXs[1] - CurrentPs[1];
+	// We need to apply the corrections at the same world-space point on each body for conservation
+	// of momentum. When the two constraint anchors are not coincident and we have an error to correct
+	// we bias the point towards the connector of the lighter body. This reduces the amount of extra 
+	// rotation that gets applied as a result of using a connector point that is not the actual anchor 
+	// point. E.g., consider the case where one body is kinematic and the error to correct is very large.
+	// We would not want to select a point far from the dynamic body, becasue that can lead to a lot 
+	// of unnecessary rotation.
+	const FReal CXAlpha = InvM(1) / (InvM(0) + InvM(1));
+	const FVec3 SharedCX = FMath::Lerp(ConnectorXs[0], ConnectorXs[1], CXAlpha);
+
+	const FVec3 ConstraintArm0Limited = SharedCX - CurrentPs[0];
+	const FVec3 ConstraintArm1Limited = SharedCX - CurrentPs[1];
 	
 	FVec3 ConstraintArm0Locked = ConstraintArm0Limited;
 	FVec3 ConstraintArm1Locked = ConstraintArm1Limited;
@@ -595,24 +604,42 @@ void FPBDJointCachedSolver::SolveLinearVelocityConstraint(
 	const FVec3 CV1 = V(1) + FVec3::CrossProduct(W(1), PositionConstraints.ConstraintArms[ConstraintIndex][1]);
 	const FVec3 CV = CV1 - CV0;
 
-	const FReal DeltaLambda = SolverStiffness * PositionConstraints.ConstraintHardStiffness[ConstraintIndex] *
-	 (FVec3::DotProduct(CV, PositionConstraints.ConstraintAxis[ConstraintIndex]) - TargetVel) / PositionConstraints.ConstraintHardIM[ConstraintIndex];
-	
-	const FVec3 MDV = DeltaLambda * PositionConstraints.ConstraintAxis[ConstraintIndex];
+	const FReal DeltaVelocity = FVec3::DotProduct(CV, PositionConstraints.ConstraintAxis[ConstraintIndex]) - TargetVel;
+	FReal DeltaLambda = SolverStiffness * PositionConstraints.ConstraintHardStiffness[ConstraintIndex] * DeltaVelocity / PositionConstraints.ConstraintHardIM[ConstraintIndex];
 
-	if(Body(0).IsDynamic())
+	// If we have a unilateral contraint, do not allow negative impulses
+	// @todo(chaos): we should allow negative impulses up to the amount that was added by the position solve to prevent initial-overlap problems
+	bool bNeedsSolve = true;
+	if (PositionConstraints.bLimitsCheck[ConstraintIndex])
 	{
-		const FVec3 DV0 = InvM(0) * MDV;
-		const FVec3 DW0 = PositionConstraints.ConstraintDRAxis[ConstraintIndex][0] * DeltaLambda;
-
-		Body(0).ApplyVelocityDelta(DV0, DW0);
+		if ((DeltaVelocity > 0) && (PositionConstraints.ConstraintLambda[ConstraintIndex] < 0))
+		{
+			bNeedsSolve = false;
+		}
+		else if ((DeltaVelocity < 0) && (PositionConstraints.ConstraintLambda[ConstraintIndex] > 0))
+		{
+			bNeedsSolve = false;
+		}
 	}
-	if(Body(1).IsDynamic())
+
+	if (bNeedsSolve)
 	{
-		const FVec3 DV1 = -InvM(1) * MDV;
-		const FVec3 DW1 = PositionConstraints.ConstraintDRAxis[ConstraintIndex][1] * DeltaLambda;
+		const FVec3 MDV = DeltaLambda * PositionConstraints.ConstraintAxis[ConstraintIndex];
+
+		if(Body(0).IsDynamic())
+		{
+			const FVec3 DV0 = InvM(0) * MDV;
+			const FVec3 DW0 = PositionConstraints.ConstraintDRAxis[ConstraintIndex][0] * DeltaLambda;
+
+			Body(0).ApplyVelocityDelta(DV0, DW0);
+		}
+		if(Body(1).IsDynamic())
+		{
+			const FVec3 DV1 = -InvM(1) * MDV;
+			const FVec3 DW1 = PositionConstraints.ConstraintDRAxis[ConstraintIndex][1] * DeltaLambda;
 
-		Body(1).ApplyVelocityDelta(DV1, DW1);
+			Body(1).ApplyVelocityDelta(DV1, DW1);
+		}
 	}
 }
 
@@ -1069,20 +1096,39 @@ void FPBDJointCachedSolver::SolveAngularVelocityConstraint(
 {
 	const FVec3 CW = W(1) - W(0);
 
-	const FReal DeltaLambda = SolverStiffness * RotationConstraints.ConstraintHardStiffness[ConstraintIndex] *
-	 (FVec3::DotProduct(CW, RotationConstraints.ConstraintAxis[ConstraintIndex]) - TargetVel) / RotationConstraints.ConstraintHardIM[ConstraintIndex];
+	const FReal DeltaVelocity = FVec3::DotProduct(CW, RotationConstraints.ConstraintAxis[ConstraintIndex]) - TargetVel;
 
-	if(Body(0).IsDynamic())
+	// If we have a unilateral contraint, do not allow negative impulses
+	// @todo(chaos): we should allow negative impulses up to the amount that was added by the position solve to prevent initial-overlap problems
+	bool bNeedsSolve = true;
+	if (RotationConstraints.bLimitsCheck[ConstraintIndex])
 	{
-		const FVec3 DW0 = RotationConstraints.ConstraintDRAxis[ConstraintIndex][0] * DeltaLambda;
-	
-		Body(0).ApplyAngularVelocityDelta(DW0);
+		if ((DeltaVelocity > 0) && (RotationConstraints.ConstraintLambda[ConstraintIndex] < 0))
+		{
+			bNeedsSolve = false;
+		}
+		else if ((DeltaVelocity < 0) && (RotationConstraints.ConstraintLambda[ConstraintIndex] > 0))
+		{
+			bNeedsSolve = false;
+		}
 	}
-	if(Body(1).IsDynamic())
+
+	if (bNeedsSolve)
 	{
-		const FVec3 DW1 = RotationConstraints.ConstraintDRAxis[ConstraintIndex][1] * DeltaLambda;
+		const FReal DeltaLambda = SolverStiffness * RotationConstraints.ConstraintHardStiffness[ConstraintIndex] * DeltaVelocity / RotationConstraints.ConstraintHardIM[ConstraintIndex];
+
+		if(Body(0).IsDynamic())
+		{
+			const FVec3 DW0 = RotationConstraints.ConstraintDRAxis[ConstraintIndex][0] * DeltaLambda;
 	
-		Body(1).ApplyAngularVelocityDelta(DW1);
+			Body(0).ApplyAngularVelocityDelta(DW0);
+		}
+		if(Body(1).IsDynamic())
+		{
+			const FVec3 DW1 = RotationConstraints.ConstraintDRAxis[ConstraintIndex][1] * DeltaLambda;
+	
+			Body(1).ApplyAngularVelocityDelta(DW1);
+		}
 	}
 }
 
