Handle degenerate triangle in vectorized TriangleSimplexFindOriginFast. Fixes NaN in degen case.

diff --git a/Engine/Source/Runtime/Experimental/Chaos/Public/Chaos/SimplexVectorized.h b/Engine/Source/Runtime/Experimental/Chaos/Public/Chaos/SimplexVectorized.h
--- a/Engine/Source/Runtime/Experimental/Chaos/Public/Chaos/SimplexVectorized.h
+++ b/Engine/Source/Runtime/Experimental/Chaos/Public/Chaos/SimplexVectorized.h
@@ -70,9 +70,25 @@ namespace Chaos
 		const VectorRegister4Float& B = Simplex[1];
 		const VectorRegister4Float& C = Simplex[2];
 
-		// Vertex region A
 		const VectorRegister4Float AB = VectorSubtract(B, A);
 		const VectorRegister4Float AC = VectorSubtract(C, A);
+
+
+		// Handle degenerate triangle
+		const VectorRegister4Float TriNormal = VectorCross(AB, AC);
+		const VectorRegister4Float TriNormal2 = VectorDot3(TriNormal, TriNormal);
+		const VectorRegister4Float MinFloat = MakeVectorRegisterFloatConstant(std::numeric_limits<float>::min(), std::numeric_limits<float>::min(), std::numeric_limits<float>::min(), std::numeric_limits<float>::min());
+		const VectorRegister4Float AMin = VectorMultiply(A, MinFloat);
+		const VectorRegister4Float Eps2 = VectorDot3(AMin, AMin);
+		const VectorRegister4Float Eps2GENormal2 = VectorCompareGE(Eps2, TriNormal2);
+		if (VectorMaskBits(Eps2GENormal2))
+		{
+			constexpr VectorRegister4Int Two = MakeVectorRegisterIntConstant(2, 2, 2, 2);
+			NumVerts = Two;
+			return VectorLineSimplexFindOrigin<CalculatExtraInformation>(Simplex, NumVerts, OutBarycentric, As, Bs);
+		}
+
+		// Vertex region A
 		const VectorRegister4Float AO = VectorNegate(A);
 
 		const VectorRegister4Float d1 = VectorDot3(AB, AO);
@@ -259,8 +275,6 @@ namespace Chaos
 		// (E.g., this caused jittering when walking on box with dimensions of 100000cm or more).
 		// This fix the unit test TestSmallCapsuleLargeBoxGJKRaycast_Vertical
 		// Previously was return VectorMultiplyAdd(AC, w, VectorMultiplyAdd(AB, v, A));
-		const VectorRegister4Float TriNormal = VectorCross(AB, AC);
-		const VectorRegister4Float TriNormal2 = VectorDot3(TriNormal, TriNormal);
 		const VectorRegister4Float TriNormalOverSize2 = VectorDivide(TriNormal, TriNormal2);
 		const VectorRegister4Float SignedDistance = VectorDot3(A, TriNormalOverSize2);
 		return VectorMultiply(TriNormal, SignedDistance);
