Chaos : fix headless chaos test failing

diff --git a/Engine/Source/Runtime/Experimental/Chaos/Public/Chaos/Capsule.h b/Engine/Source/Runtime/Experimental/Chaos/Public/Chaos/Capsule.h
--- a/Engine/Source/Runtime/Experimental/Chaos/Public/Chaos/Capsule.h
+++ b/Engine/Source/Runtime/Experimental/Chaos/Public/Chaos/Capsule.h
@@ -14,8 +14,6 @@
 
 #include "UObject/ReleaseObjectVersion.h"
 
-PRAGMA_DISABLE_OPTIMIZATION
-
 namespace Chaos
 {
 	struct FCapsuleSpecializeSamplingHelper;
@@ -288,7 +286,7 @@ namespace Chaos
 					if (PositionLengthOnCoreCylinder >= 0 && PositionLengthOnCoreCylinder < MHeight)
 					{
 						const FVec3 SegmentToPosition = (CylinderToSpherePosition - MVector * PositionLengthOnCoreCylinder);
-						if (SegmentToPosition.SquaredLength() > R2)
+						if (SegmentToPosition.SquaredLength() > ( R2 + KINDA_SMALL_NUMBER))
 						{
 							// the contact point is actually outside of the cylinder
 							// this can happen if the ray starts within the cylinder bounds but do not intersect with the cylinder
@@ -747,4 +745,3 @@ namespace Chaos
 	using TCapsuleSpecializeSamplingHelper UE_DEPRECATED(4.27, "Deprecated. this class is to be deleted, use FCapsuleSpecializeSamplingHelper instead") = FCapsuleSpecializeSamplingHelper;
 
 } // namespace Chaos
-PRAGMA_ENABLE_OPTIMIZATION
\ No newline at end of file
