Chaos: Fixed non-virtual base class destructor leading to bad delete through a base class pointer (UB, caught by ASan)

diff --git a/Engine/Source/Runtime/Experimental/Chaos/Public/Chaos/GeometryParticles.h b/Engine/Source/Runtime/Experimental/Chaos/Public/Chaos/GeometryParticles.h
--- a/Engine/Source/Runtime/Experimental/Chaos/Public/Chaos/GeometryParticles.h
+++ b/Engine/Source/Runtime/Experimental/Chaos/Public/Chaos/GeometryParticles.h
@@ -40,7 +40,7 @@ namespace Chaos
 		static void UpdateGeometry(TUniquePtr<FPerShapeData>& ShapePtr, TSerializablePtr<FImplicitObject> InGeometry);
 		static bool RequiresCachedLeafInfo(const FRigidTransformRealSingle3& RelativeTransform, const FImplicitObject* LeafGeometry, const FImplicitObject* Geometry);
 
-		~FPerShapeData();
+		virtual ~FPerShapeData();
 		
 		void UpdateShapeBounds(const FRigidTransform3& WorldTM, const FVec3& BoundsExpansion = FVec3(0));
 
@@ -292,10 +292,6 @@ namespace Chaos
 			bHasCachedLeafInfo = true;
 		}
 
-		// WARNING: //////////////////////////////////////////////////////
-		// We don't have a virtual destructor on FPerShapeData (trying to save memory)
-		// FPerShapeDataCachedLeafInfo destructor will NOT be called, do not add non-PoD members.
-
 		FRigidTransformRealSingle3 LeafRelativeTransform;
 		FRigidTransform3 LeafWorldTransform;
 		const FImplicitObject* LeafGeometry;
