Chaos : Convex inertia tensor computation : handle case where we have 0 vertex convex

diff --git a/Engine/Source/Runtime/Experimental/Chaos/Private/Chaos/Convex.cpp b/Engine/Source/Runtime/Experimental/Chaos/Private/Chaos/Convex.cpp
--- a/Engine/Source/Runtime/Experimental/Chaos/Private/Chaos/Convex.cpp
+++ b/Engine/Source/Runtime/Experimental/Chaos/Private/Chaos/Convex.cpp
@@ -422,17 +422,14 @@ namespace Chaos
 
 	void FConvex::ComputeUnitMassInertiaTensorAndRotationOfMass(const FReal InVolume)
 	{
-		// we only compute the inertia tensor using unit mass and scale it upon request
-		if (InVolume < SMALL_NUMBER)
+		if (InVolume < SMALL_NUMBER || !StructureData.IsValid())
 		{
 			UnitMassInertiaTensor = FVec3{1., 1., 1.};
 			RotationOfMass = FRotation3::Identity;
 			return;
 		}
 
-		// verify that the StructureData has been built before we call process further
-		ensure(StructureData.IsValid());
-		
+		// we only compute the inertia tensor using unit mass and scale it upon request
 		constexpr FReal Mass = 1.0;
 		const FReal Density = Mass / InVolume;
 
