Awake island when kinematics are moving

diff --git a/Engine/Source/Runtime/Experimental/Chaos/Private/Chaos/Island/IslandGraph.cpp b/Engine/Source/Runtime/Experimental/Chaos/Private/Chaos/Island/IslandGraph.cpp
--- a/Engine/Source/Runtime/Experimental/Chaos/Private/Chaos/Island/IslandGraph.cpp
+++ b/Engine/Source/Runtime/Experimental/Chaos/Private/Chaos/Island/IslandGraph.cpp
@@ -734,9 +734,9 @@ void FIslandGraph<NodeType, EdgeType, IslandType>::UpdateGraph()
 				GraphNode.IslandIndex = GraphIslands.Emplace(GraphIsland);
 			}
 			// Update of the sleeping flag on the island
-			if(GraphIslands.IsValidIndex(GraphNode.IslandIndex) && !GraphNode.bStationaryNode)
+			if(!GraphNode.bStationaryNode)
 			{
-				if(GraphNode.NodeEdges.Num() == 0)
+				if(GraphNode.NodeEdges.Num() == 0 && GraphIslands.IsValidIndex(GraphNode.IslandIndex))
 				{
 					GraphIslands[GraphNode.IslandIndex].bIsSleeping = false;
 				}
diff --git a/Engine/Source/Runtime/Experimental/Chaos/Private/Chaos/Island/IslandManager.cpp b/Engine/Source/Runtime/Experimental/Chaos/Private/Chaos/Island/IslandManager.cpp
--- a/Engine/Source/Runtime/Experimental/Chaos/Private/Chaos/Island/IslandManager.cpp
+++ b/Engine/Source/Runtime/Experimental/Chaos/Private/Chaos/Island/IslandManager.cpp
@@ -49,7 +49,7 @@ FORCEINLINE bool IsStationaryParticle(const FGeometryParticleHandle* ParticleHan
 	if (ParticleHandle->ObjectState() == EObjectStateType::Kinematic)
 	{
 		const FKinematicGeometryParticleHandle* KinematicParticle = ParticleHandle->CastToKinematicParticle();
-		return KinematicParticle->V().IsZero();
+		return KinematicParticle->V().IsZero() && KinematicParticle->KinematicTarget().GetMode() == EKinematicTargetMode::None;
 	}
 	else
 	{
