Chaos - fix GJK typo which could lead to incorrect collision resuilts

diff --git a/Engine/Source/Runtime/Experimental/Chaos/Public/Chaos/Simplex.h b/Engine/Source/Runtime/Experimental/Chaos/Public/Chaos/Simplex.h
--- a/Engine/Source/Runtime/Experimental/Chaos/Public/Chaos/Simplex.h
+++ b/Engine/Source/Runtime/Experimental/Chaos/Public/Chaos/Simplex.h
@@ -300,7 +300,7 @@ namespace Chaos
 		const T d4 = TVec3<T>::DotProduct(AC, BO);
 
 		const bool bIsD3GEZero = (d3 >= T(0));
-		const bool bIsD3GED4 = (d4 >= d4);
+		const bool bIsD3GED4 = (d3 >= d4);
 		const bool bIsB = bIsD3GEZero && bIsD3GED4;
 		if (bIsB)
 		{
