Make ScaledImplicit constructor require shared ptr for ref count, to reduce risk of copied implicit dropping shared pointer.
Fix copy functions not copying shared ptr, fixing incorrect lifetime of shared geometry.

diff --git a/Engine/Source/Runtime/Engine/Private/PhysicsEngine/BodyInstance.cpp b/Engine/Source/Runtime/Engine/Private/PhysicsEngine/BodyInstance.cpp
--- a/Engine/Source/Runtime/Engine/Private/PhysicsEngine/BodyInstance.cpp
+++ b/Engine/Source/Runtime/Engine/Private/PhysicsEngine/BodyInstance.cpp
@@ -497,7 +497,8 @@ void FBodyInstance::UpdateTriMeshVertices(const TArray<FVector> & NewPositions)
 					TriMeshCopy->GetObjectChecked<FTriangleMeshImplicitObject>().UpdateVertices(NewPositions);
 					if (Scale != FVec3(1, 1, 1))
 					{
-						TUniquePtr<FImplicitObject> Scaled = MakeUnique<TImplicitObjectScaled<FTriangleMeshImplicitObject, /*bInstanced=*/false>>(MoveTemp(TriMeshCopy), Scale);
+						TSharedPtr<FTriangleMeshImplicitObject, ESPMode::ThreadSafe> SharedPtrForRefCount(nullptr); // Not instanced, no shared trimesh to ref count.
+						TUniquePtr<FImplicitObject> Scaled = MakeUnique<TImplicitObjectScaled<FTriangleMeshImplicitObject, /*bInstanced=*/false>>(MoveTemp(TriMeshCopy), SharedPtrForRefCount, Scale);
 						FPhysicsInterface::SetGeometry(Shape, MoveTemp(Scaled));
 					}
 					else
@@ -2161,7 +2162,7 @@ bool FBodyInstance::UpdateBodyScale(const FVector& InScale3D, bool bForceUpdate)
 						{
 							// While a body setup will instantiate the triangle mesh as a shared geometry, other methods might not (e.g. retopologized landscape)
 							TImplicitObjectScaled<FTriangleMeshImplicitObject>::ObjectType InnerObject = ScaledTriangleMesh->Object();
-							NewGeometry.Emplace(MakeUnique<Chaos::TImplicitObjectScaled<Chaos::FTriangleMeshImplicitObject>>(ScaledTriangleMesh->Object(), AdjustedScale3D));
+							NewGeometry.Emplace(MakeUnique<Chaos::TImplicitObjectScaled<Chaos::FTriangleMeshImplicitObject>>(ScaledTriangleMesh->Object(), InnerTriangleMesh, AdjustedScale3D));
 						}
 					}
 					else if (bIsInstanced)
diff --git a/Engine/Source/Runtime/Experimental/Chaos/Private/Chaos/TriangleMeshImplicitObject.cpp b/Engine/Source/Runtime/Experimental/Chaos/Private/Chaos/TriangleMeshImplicitObject.cpp
--- a/Engine/Source/Runtime/Experimental/Chaos/Private/Chaos/TriangleMeshImplicitObject.cpp
+++ b/Engine/Source/Runtime/Experimental/Chaos/Private/Chaos/TriangleMeshImplicitObject.cpp
@@ -31,8 +31,11 @@ namespace Chaos
 template <typename QueryGeomType>
 static auto MakeScaledHelper(const QueryGeomType& B, const FVec3& InvScale)
 {
+	// TODO: Fixup code using this and remove it.
+
 	TUniquePtr<QueryGeomType> HackBPtr(const_cast<QueryGeomType*>(&B));	//todo: hack, need scaled object to accept raw ptr similar to transformed implicit
-	TImplicitObjectScaled<QueryGeomType> ScaledB(MakeSerializable(HackBPtr), InvScale);
+	TSharedPtr<QueryGeomType, ESPMode::ThreadSafe> SharedPtrForRefCount(nullptr); // This scaled is temporary, use null shared ptr.
+	TImplicitObjectScaled<QueryGeomType> ScaledB(MakeSerializable(HackBPtr), SharedPtrForRefCount, InvScale);
 	HackBPtr.Release();
 	return ScaledB;
 }
@@ -41,7 +44,7 @@ template <typename QueryGeomType>
 static auto MakeScaledHelper(const TImplicitObjectScaled<QueryGeomType>& B, const FVec3& InvScale)
 {
 	//if scaled of scaled just collapse into one scaled
-	TImplicitObjectScaled<QueryGeomType> ScaledB(B.Object(), InvScale * B.GetScale());
+	TImplicitObjectScaled<QueryGeomType> ScaledB(B.Object(), B.GetSharedObject(), InvScale * B.GetScale());
 	return ScaledB;
 }
 
diff --git a/Engine/Source/Runtime/Experimental/Chaos/Public/Chaos/ImplicitObjectScaled.h b/Engine/Source/Runtime/Experimental/Chaos/Public/Chaos/ImplicitObjectScaled.h
--- a/Engine/Source/Runtime/Experimental/Chaos/Public/Chaos/ImplicitObjectScaled.h
+++ b/Engine/Source/Runtime/Experimental/Chaos/Public/Chaos/ImplicitObjectScaled.h
@@ -420,10 +420,11 @@ public:
 	using ObjectType = typename TChooseClass<bInstanced, TSerializablePtr<TConcrete>, TUniquePtr<TConcrete>>::Result;
 	using FImplicitObject::GetTypeName;
 
-	TImplicitObjectScaled(ObjectType Object, const FVec3& Scale, FReal InMargin = 0)
+	// If Object is instanced and requires ref counting must set SharedPtrForRefCount to ensure geometry remains valid, otherwise use null.
+	TImplicitObjectScaled(ObjectType Object, const TSharedPtr<TConcrete, ESPMode::ThreadSafe>& SharedPtrForRefCount, const FVec3& Scale, FReal InMargin = 0)
 	    : FImplicitObjectScaled(EImplicitObject::HasBoundingBox, Object->GetType())
 	    , MObject(MoveTemp(Object))
-		, MSharedPtrForRefCount(nullptr)
+		, MSharedPtrForRefCount(SharedPtrForRefCount)
 	{
 		ensureMsgf((IsScaled(MObject->GetType()) == false), TEXT("Scaled objects should not contain each other."));
 		ensureMsgf((IsInstanced(MObject->GetType()) == false), TEXT("Scaled objects should not contain instances."));
@@ -462,10 +463,11 @@ public:
 		SetScale(Scale);
 	}
 
-	TImplicitObjectScaled(ObjectType Object, TUniquePtr<Chaos::FImplicitObject> &&ObjectOwner, const FVec3& Scale, FReal InMargin = 0)
+	// If Object is instanced and requires ref counting must set SharedPtrForRefCount to ensure geometry remains valid, otherwise use null.
+	TImplicitObjectScaled(ObjectType Object, TUniquePtr<Chaos::FImplicitObject> &&ObjectOwner, const TSharedPtr<TConcrete, ESPMode::ThreadSafe>& SharedPtrForRefCount, const FVec3& Scale, FReal InMargin = 0)
 	    : FImplicitObjectScaled(EImplicitObject::HasBoundingBox, Object->GetType())
 	    , MObject(Object)
-		, MSharedPtrForRefCount(nullptr)
+		, MSharedPtrForRefCount(SharedPtrForRefCount)
 	{
 		ensureMsgf((IsScaled(MObject->GetType(true)) == false), TEXT("Scaled objects should not contain each other."));
 		ensureMsgf((IsInstanced(MObject->GetType(true)) == false), TEXT("Scaled objects should not contain instances."));
@@ -994,7 +996,7 @@ private:
 
 	static TImplicitObjectScaled<TConcrete, true>* CopyHelper(const TImplicitObjectScaled<TConcrete, true>* Obj)
 	{
-		return new TImplicitObjectScaled<TConcrete, true>(Obj->MObject, Obj->MScale, Obj->OuterMargin);
+		return new TImplicitObjectScaled<TConcrete, true>(Obj->MObject, Obj->MSharedPtrForRefCount, Obj->MScale, Obj->OuterMargin);
 	}
 
 	static TImplicitObjectScaled<TConcrete, false>* CopyHelper(const TImplicitObjectScaled<TConcrete, false>* Obj)
@@ -1002,7 +1004,7 @@ private:
 		TUniquePtr<FImplicitObject> DuplicatedShape = Obj->MObject->Copy();
 		
 		// We know the actual type of the underlying object pointer so we can cast it to the required type to make a copy of this implicit
-		return new TImplicitObjectScaled<TConcrete, false>(reinterpret_cast<TUniquePtr<TConcrete>&&>(DuplicatedShape), Obj->MScale, Obj->OuterMargin);
+		return new TImplicitObjectScaled<TConcrete, false>(reinterpret_cast<TUniquePtr<TConcrete>&&>(DuplicatedShape), Obj->MSharedPtrForRefCount, Obj->MScale, Obj->OuterMargin);
 	}
 
 	void UpdateBounds()
@@ -1017,8 +1019,12 @@ private:
 	template <typename QueryGeomType>
 	static auto MakeScaledHelper(const QueryGeomType& B, const TVector<T,d>& InvScale )
 	{
+		// TODO: Fixup code using this and remove it.
+
 		TUniquePtr<QueryGeomType> HackBPtr(const_cast<QueryGeomType*>(&B));	//todo: hack, need scaled object to accept raw ptr similar to transformed implicit
-		TImplicitObjectScaled<QueryGeomType> ScaledB(MakeSerializable(HackBPtr), InvScale);
+
+		TSharedPtr<QueryGeomType, ESPMode::ThreadSafe> SharedPtrForRefCount(nullptr); // This scaled is temporary, use fake shared ptr.
+		TImplicitObjectScaled<QueryGeomType> ScaledB(MakeSerializable(HackBPtr), SharedPtrForRefCount, InvScale);
 		HackBPtr.Release();
 		return ScaledB;
 	}
@@ -1027,7 +1033,7 @@ private:
 	static auto MakeScaledHelper(const TImplicitObjectScaled<QueryGeomType>& B, const TVector<T,d>& InvScale)
 	{
 		//if scaled of scaled just collapse into one scaled
-		TImplicitObjectScaled<QueryGeomType> ScaledB(B.Object(), InvScale * B.GetScale());
+		TImplicitObjectScaled<QueryGeomType> ScaledB(B.Object(), B.GetSharedObject(), InvScale * B.GetScale());
 		return ScaledB;
 	}
 
diff --git a/Engine/Source/Runtime/Landscape/Private/LandscapeCollision.cpp b/Engine/Source/Runtime/Landscape/Private/LandscapeCollision.cpp
--- a/Engine/Source/Runtime/Landscape/Private/LandscapeCollision.cpp
+++ b/Engine/Source/Runtime/Landscape/Private/LandscapeCollision.cpp
@@ -1633,8 +1633,9 @@ struct FMeshCollisionInitHelper
 
 		FVector Scale = FVector(ComponentScale.X * CollisionScale, ComponentScale.Y * CollisionScale, ComponentScale.Z);
 
+		TSharedPtr<Chaos::FTriangleMeshImplicitObject, ESPMode::ThreadSafe> SharedPtrForRefCount(nullptr); // Not shared trimesh, no need for ref counting trimesh.
 		{
-			TUniquePtr<Chaos::TImplicitObjectScaled<Chaos::FTriangleMeshImplicitObject>> ScaledTrimesh = MakeUnique<Chaos::TImplicitObjectScaled<Chaos::FTriangleMeshImplicitObject>>(MakeSerializable(MeshRef->Trimesh), Scale);
+			TUniquePtr<Chaos::TImplicitObjectScaled<Chaos::FTriangleMeshImplicitObject>> ScaledTrimesh = MakeUnique<Chaos::TImplicitObjectScaled<Chaos::FTriangleMeshImplicitObject>>(MakeSerializable(MeshRef->Trimesh), SharedPtrForRefCount, Scale);
 			TUniquePtr<Chaos::FPerShapeData> NewShape = Chaos::FPerShapeData::CreatePerShapeData(ShapeArray.Num(), MakeSerializable(ScaledTrimesh));
 
 			NewShape->SetQueryData(QueryFilter);
@@ -1649,7 +1650,7 @@ struct FMeshCollisionInitHelper
 #if WITH_EDITOR
 		if(!World->IsGameWorld())
 		{
-			TUniquePtr<Chaos::TImplicitObjectScaled<Chaos::FTriangleMeshImplicitObject>> ScaledTrimeshEd = MakeUnique<Chaos::TImplicitObjectScaled<Chaos::FTriangleMeshImplicitObject>>(MakeSerializable(MeshRef->EditorTrimesh), Scale);
+			TUniquePtr<Chaos::TImplicitObjectScaled<Chaos::FTriangleMeshImplicitObject>> ScaledTrimeshEd = MakeUnique<Chaos::TImplicitObjectScaled<Chaos::FTriangleMeshImplicitObject>>(MakeSerializable(MeshRef->EditorTrimesh), SharedPtrForRefCount, Scale);
 			TUniquePtr<Chaos::FPerShapeData> NewEdShape = Chaos::FPerShapeData::CreatePerShapeData(ShapeArray.Num(), MakeSerializable(ScaledTrimeshEd));
 
 			NewEdShape->SetQueryData(QueryFilterEd);
