Fix trimesh with negative scale computing negative inflation during sweep.

diff --git a/Engine/Source/Runtime/Experimental/Chaos/Private/Chaos/TriangleMeshImplicitObject.cpp b/Engine/Source/Runtime/Experimental/Chaos/Private/Chaos/TriangleMeshImplicitObject.cpp
--- a/Engine/Source/Runtime/Experimental/Chaos/Private/Chaos/TriangleMeshImplicitObject.cpp
+++ b/Engine/Source/Runtime/Experimental/Chaos/Private/Chaos/TriangleMeshImplicitObject.cpp
@@ -1067,7 +1067,7 @@ bool FTriangleMeshImplicitObject::SweepGeomImp(const QueryGeomType& QueryGeom, c
 		const FAABB3 QueryBounds = QueryGeom.BoundingBox().TransformedAABB(FRigidTransform3(FVec3::ZeroVector,StartTM.GetRotation()));
 		const FVec3 InvTriMeshScale = SafeInvScale(TriMeshScale);
 		const FVec3 StartPoint = QueryBounds.Center() * InvTriMeshScale + StartTM.GetLocation();
-		const FVec3 Inflation = QueryBounds.Extents() * InvTriMeshScale * 0.5 + FVec3(Thickness);
+		const FVec3 Inflation = QueryBounds.Extents() * InvTriMeshScale.GetAbs() * 0.5 + FVec3(Thickness);
 		BVH.template Sweep<VisitorType>(StartPoint,Dir,Length,Inflation,SQVisitor);
 
 		if(SQVisitor.OutTime <= Length)
