Fix Angular limits for linear joints

diff --git a/Engine/Source/Runtime/Experimental/Chaos/Private/Chaos/Joint/PBDJointCachedSolverGaussSeidel.cpp b/Engine/Source/Runtime/Experimental/Chaos/Private/Chaos/Joint/PBDJointCachedSolverGaussSeidel.cpp
--- a/Engine/Source/Runtime/Experimental/Chaos/Private/Chaos/Joint/PBDJointCachedSolverGaussSeidel.cpp
+++ b/Engine/Source/Runtime/Experimental/Chaos/Private/Chaos/Joint/PBDJointCachedSolverGaussSeidel.cpp
@@ -556,7 +556,7 @@ void FPBDJointCachedSolver::ApplyAxisPositionConstraint(
 			DeltaPosition += PositionConstraints.ConstraintLimits[ConstraintIndex];
 			NeedsSolve = true;
 		}
-	}
+	}  
 	if (!PositionConstraints.bLimitsCheck[ConstraintIndex] || (PositionConstraints.bLimitsCheck[ConstraintIndex] && NeedsSolve && FMath::Abs(DeltaPosition ) > PositionTolerance))
 	{
 		if ((PositionConstraints.MotionType[ConstraintIndex] == EJointMotionType::Limited) && PositionConstraints.bSoftLimit[ConstraintIndex])
@@ -726,17 +726,13 @@ FPBDJointUtilities::GetSoftTwistStiffness(SolverSettings, JointSettings),
 			// When using non linear solver, the cone swing direction could change at each iteration
 			// stabilizing the solver. In the linear case we need to constraint along the 2 directions
 			// for better stability
-			InitSwingConstraint(JointSettings, SolverSettings, Dt, EJointAngularConstraintIndex::Swing1);
-			InitSwingConstraint(JointSettings, SolverSettings, Dt, EJointAngularConstraintIndex::Swing2);
-
-			//InitConeConstraint(JointSettings, Dt);
+			InitPyramidSwingConstraint(JointSettings, Dt, true, true);
 		}
 		else if (bAngularLimited[S1] && bAngularLocked[S2])
 		{
-			InitSingleLockedSwingConstraint(JointSettings, Dt, EJointAngularConstraintIndex::Swing2);
 			if (!bDegenerate)
 			{
-				InitSwingConstraint(JointSettings, SolverSettings, Dt, EJointAngularConstraintIndex::Swing1 );
+				InitPyramidSwingConstraint(JointSettings, Dt, true, false);
 			}
 		}
 		else if (bAngularLimited[S1] && bAngularFree[S2])
@@ -748,16 +744,11 @@ FPBDJointUtilities::GetSoftTwistStiffness(SolverSettings, JointSettings),
 		}
 		else if (bAngularLocked[S1] && bAngularLimited[S2])
 		{
-			InitSingleLockedSwingConstraint(JointSettings, Dt, EJointAngularConstraintIndex::Swing1);
 			if (!bDegenerate)
 			{
-				InitSwingConstraint(JointSettings, SolverSettings, Dt, EJointAngularConstraintIndex::Swing2);
+				InitPyramidSwingConstraint(JointSettings, Dt, false, true);
 			}
 		}
-		else if (bAngularLocked[S1] && bAngularFree[S2])
-		{
-			InitSingleLockedSwingConstraint(JointSettings, Dt, EJointAngularConstraintIndex::Swing1);
-		}
 		else if (bAngularFree[S1] && bAngularLimited[S2])
 		{
 			if (!bDegenerate)
@@ -765,18 +756,15 @@ FPBDJointUtilities::GetSoftTwistStiffness(SolverSettings, JointSettings),
 				InitDualConeSwingConstraint(JointSettings, Dt, EJointAngularConstraintIndex::Swing2);
 			}
 		}
-		else if (bAngularFree[S1]  && bAngularLocked[S2])
-		{
-			InitSingleLockedSwingConstraint(JointSettings, Dt, EJointAngularConstraintIndex::Swing2);
-		}
 	}
 
 	// Note: single-swing locks are already handled above so we only need to do something here if both are locked
 	const bool bLockedTwist = SolverSettings.bEnableTwistLimits && bAngularLocked[TW];
-	const bool bLockedSwing = SolverSettings.bEnableSwingLimits && bAngularLocked[S1] && bAngularLocked[S2];
-	if (bLockedTwist || bLockedSwing)
+	const bool bLockedSwing1 = SolverSettings.bEnableSwingLimits && bAngularLocked[S1];
+	const bool bLockedSwing2 = SolverSettings.bEnableSwingLimits && bAngularLocked[S2];
+	if (bLockedTwist || bLockedSwing1 || bLockedSwing2)
 	{
-		InitLockedRotationConstraints(JointSettings, Dt, bLockedTwist, bLockedSwing);
+		InitLockedRotationConstraints(JointSettings, Dt, bLockedTwist, bLockedSwing1, bLockedSwing2);
 	}
 }
 
@@ -852,6 +840,32 @@ void FPBDJointCachedSolver::InitTwistConstraint(
 	 InitRotationConstraintDatas( JointSettings, (int32)EJointAngularConstraintIndex::Twist, TwistAxis, TwistAngle, JointSettings.TwistRestitution, Dt, true);
 }
 
+void FPBDJointCachedSolver::InitPyramidSwingConstraint(
+   const FPBDJointSettings& JointSettings,
+   const FReal Dt,
+   const bool bApplySwing1,
+   const bool bApplySwing2)
+{
+	// Decompose rotation of body 1 relative to body 0 into swing and twist rotations, assuming twist is X axis
+	FRotation3 R01Twist, R01Swing;
+	FPBDJointUtilities::DecomposeSwingTwistLocal(ConnectorRs[0], ConnectorRs[1], R01Swing, R01Twist);
+
+	const FRotation3 R0Swing = ConnectorRs[0] * R01Swing;
+
+	if(bApplySwing1)
+	{
+		const FVec3 SwingAxis = R0Swing * FJointConstants::Swing1Axis();
+		const FReal SwingAngle = 4.0 * FMath::Atan2(R01Swing.Z, (FReal)(1. + R01Swing.W));
+		InitRotationConstraintDatas( JointSettings, (int32)EJointAngularConstraintIndex::Swing1, SwingAxis, SwingAngle, JointSettings.SwingRestitution, Dt, true);
+	}
+	if(bApplySwing2)
+	{
+		const FVec3 SwingAxis = R0Swing * FJointConstants::Swing2Axis();
+		const FReal SwingAngle = 4.0 * FMath::Atan2(R01Swing.Y, (FReal)(1. + R01Swing.W));
+		InitRotationConstraintDatas( JointSettings, (int32)EJointAngularConstraintIndex::Swing2, SwingAxis, SwingAngle, JointSettings.SwingRestitution, Dt, true);
+	}
+}
+
 void FPBDJointCachedSolver::InitConeConstraint(
    const FPBDJointSettings& JointSettings,
    const FReal Dt)
@@ -861,7 +875,7 @@ void FPBDJointCachedSolver::InitConeConstraint(
 
 	FPBDJointUtilities::GetEllipticalConeAxisErrorLocal(ConnectorRs[0], ConnectorRs[1], 0.0, 0.0, SwingAxisLocal, SwingAngle);
 	SwingAxisLocal.SafeNormalize();
-
+	
 	const FVec3 SwingAxis = ConnectorRs[0] * SwingAxisLocal;
 	InitRotationConstraintDatas( JointSettings, (int32)EJointAngularConstraintIndex::Swing2, SwingAxis, SwingAngle, JointSettings.SwingRestitution, Dt, true);
 }
@@ -871,14 +885,14 @@ void FPBDJointCachedSolver::InitSingleLockedSwingConstraint(
 	const FReal Dt,
 	const EJointAngularConstraintIndex SwingConstraintIndex)
 {
-	 //NOTE: SwingAxis is not normalized in this mode. It has length Sin(SwingAngle).
-	 //Likewise, the SwingAngle is actually Sin(SwingAngle)
-	// FVec3 SwingAxis;
-	// FReal SwingAngle;
-	// FPBDJointUtilities::GetLockedSwingAxisAngle(ConnectorRs[0], ConnectorRs[1], SwingConstraintIndex, SwingAxis, SwingAngle);
-	//SwingAxis.SafeNormalize();
+	//NOTE: SwingAxis is not normalized in this mode. It has length Sin(SwingAngle).
+	//Likewise, the SwingAngle is actually Sin(SwingAngle)
+    // FVec3 SwingAxis;
+    // FReal SwingAngle;
+    // FPBDJointUtilities::GetLockedSwingAxisAngle(ConnectorRs[0], ConnectorRs[1], SwingConstraintIndex, SwingAxis, SwingAngle);
+    //SwingAxis.SafeNormalize();
 
-	// Using the locked swing axis angle results in potential axis switching since this axis is the result of OtherSwing x TwistAxis
+    // Using the locked swing axis angle results in potential axis switching since this axis is the result of OtherSwing x TwistAxis
 	FVec3 SwingAxis;
 	FReal SwingAngle;
 	FPBDJointUtilities::GetSwingAxisAngle(ConnectorRs[0], ConnectorRs[1], 0.0, SwingConstraintIndex, SwingAxis, SwingAngle);
@@ -886,6 +900,7 @@ void FPBDJointCachedSolver::InitSingleLockedSwingConstraint(
 	InitRotationConstraintDatas(JointSettings, (int32)SwingConstraintIndex, SwingAxis, SwingAngle, 0.0, Dt, false);
 }
 
+
 void FPBDJointCachedSolver::InitDualConeSwingConstraint(
 	const FPBDJointSettings& JointSettings,
 	const FReal Dt,
@@ -917,7 +932,8 @@ void FPBDJointCachedSolver::InitLockedRotationConstraints(
 	  const FPBDJointSettings& JointSettings,
 	  const FReal Dt,
 	  const bool bApplyTwist,
-	  const bool bApplySwing)
+	  const bool bApplySwing1,
+	  const bool bApplySwing2)
 {
 	FVec3 Axis0, Axis1, Axis2;
 	FPBDJointUtilities::GetLockedRotationAxes(ConnectorRs[0], ConnectorRs[1], Axis0, Axis1, Axis2);
@@ -929,11 +945,15 @@ void FPBDJointCachedSolver::InitLockedRotationConstraints(
 		InitRotationConstraintDatas(JointSettings, (int32)EJointAngularConstraintIndex::Twist, Axis0, R01.X, 0.0, Dt, false);
 	}
 
-	if (bApplySwing)
+	if (bApplySwing1)
 	{
-		InitRotationConstraintDatas(JointSettings, (int32)EJointAngularConstraintIndex::Swing2, Axis1, R01.Y, 0.0, Dt, false);
 		InitRotationConstraintDatas(JointSettings, (int32)EJointAngularConstraintIndex::Swing1, Axis2, R01.Z, 0.0, Dt, false);
 	}
+
+	if (bApplySwing2)
+	{
+		InitRotationConstraintDatas(JointSettings, (int32)EJointAngularConstraintIndex::Swing2, Axis1, R01.Y, 0.0, Dt, false);
+	}
 }
 
 /** APPLY ROTATION CONSTRAINT ******************************************************************************************/
@@ -980,10 +1000,7 @@ void FPBDJointCachedSolver::SolveRotationConstraintHard(
 	const FReal DeltaLambda = SolverStiffness * RotationConstraints.ConstraintHardStiffness[ConstraintIndex] * DeltaConstraint /
 		RotationConstraints.ConstraintHardIM[ConstraintIndex];
 
-	if (DeltaLambda >= 0)
-	{
-		RotationConstraints.ConstraintLambda[ConstraintIndex] += DeltaLambda;
-	}
+	RotationConstraints.ConstraintLambda[ConstraintIndex] += DeltaLambda;
 	SolveRotationConstraintDelta(ConstraintIndex, DeltaLambda, false, RotationConstraints);
 }
 
diff --git a/Engine/Source/Runtime/Experimental/Chaos/Public/Chaos/Joint/PBDJointCachedSolverGaussSeidel.h b/Engine/Source/Runtime/Experimental/Chaos/Public/Chaos/Joint/PBDJointCachedSolverGaussSeidel.h
--- a/Engine/Source/Runtime/Experimental/Chaos/Public/Chaos/Joint/PBDJointCachedSolverGaussSeidel.h
+++ b/Engine/Source/Runtime/Experimental/Chaos/Public/Chaos/Joint/PBDJointCachedSolverGaussSeidel.h
@@ -66,6 +66,7 @@ struct FAxisConstraintDatas
 	bool bLimitsCheck[3];
 	bool bAccelerationMode;
 	bool bSoftLimit[3];
+	
 	EJointMotionType MotionType[3];
 };
 	
@@ -363,6 +364,12 @@ struct FAxisConstraintDatas
 			const FPBDJointSettings& JointSettings,
 			const FReal Dt);
 
+		void InitPyramidSwingConstraint(
+		   const FPBDJointSettings& JointSettings,
+		   const FReal Dt,
+		   const bool bApplySwing1,
+		   const bool bApplySwing2);
+
 		// One Swing axis is free, and the other locked. This applies the lock: Body1 Twist axis is confined to a plane.
 		void InitSingleLockedSwingConstraint(
 			const FPBDJointSettings& JointSettings,
@@ -386,7 +393,8 @@ struct FAxisConstraintDatas
 			const FPBDJointSettings& JointSettings,
 			const FReal Dt,
 			const bool bApplyTwist,
-			const bool bApplySwing);
+			const bool bApplySwing1,
+			const bool bApplySwing2);
 
 		/** Apply Rotation constraints */
 
