Fixed a logic error in DoForSweepIntersectCellsImp

#localization none
#tests PIE, had a consistent repro set up and with the fix wasn't able to repro the bug anymore.

diff --git a/Engine/Source/Runtime/Experimental/Chaos/Public/Chaos/AABBTreeDirtyGridUtils.h b/Engine/Source/Runtime/Experimental/Chaos/Public/Chaos/AABBTreeDirtyGridUtils.h
--- a/Engine/Source/Runtime/Experimental/Chaos/Public/Chaos/AABBTreeDirtyGridUtils.h
+++ b/Engine/Source/Runtime/Experimental/Chaos/Public/Chaos/AABBTreeDirtyGridUtils.h
@@ -391,7 +391,7 @@ namespace Chaos
 						CrossingHorizontalCellBorderT = DtDy * ((FReal)(CurrentCellIndexY0 - ReferenceCellIndexY + (DeltaCelIndexY > 0 ? 1 : 0)) * DirtyElementGridCellSize - Y1);
 					}
 
-					if (CrossingHorizontalCellBorderT < CrossingVerticleCellBorderT)
+					if (CrossingHorizontalCellBorderT >= CrossingVerticleCellBorderT)
 					{
 						CurrentCellIndexY1 += DeltaCelIndexY;
 					}
