Chaos : GeometryCollection
- Transfer query filters to internal clusters.
- Enable transform updates on driven bodies so that line traces can resolve against unfractured collections.

diff --git a/Engine/Source/Runtime/Experimental/Chaos/Private/Chaos/PBDRigidClusteringAlgo.cpp b/Engine/Source/Runtime/Experimental/Chaos/Private/Chaos/PBDRigidClusteringAlgo.cpp
--- a/Engine/Source/Runtime/Experimental/Chaos/Private/Chaos/PBDRigidClusteringAlgo.cpp
+++ b/Engine/Source/Runtime/Experimental/Chaos/Private/Chaos/PBDRigidClusteringAlgo.cpp
@@ -299,7 +299,7 @@ namespace Chaos
 		bool bUseParticleImplicit = false;
 
 		// Need to extract a filter off one of the cluster children 
-		FCollisionFilterData Filter;
+		FCollisionFilterData SimFilter, QueryFilter;
 		for (FPBDRigidParticleHandle* Child : Children)
 		{
 			bool bFilterValid = false;
@@ -307,8 +307,10 @@ namespace Chaos
 			{
 				if (Shape)
 				{
-					Filter = Shape->GetSimData();
-					bFilterValid = Filter.Word0 != 0 || Filter.Word1 != 0 || Filter.Word2 != 0 || Filter.Word3 != 0;
+					SimFilter = Shape->GetSimData();
+					bFilterValid = SimFilter.Word0 != 0 || SimFilter.Word1 != 0 || SimFilter.Word2 != 0 || SimFilter.Word3 != 0;
+
+					QueryFilter = Shape->GetQueryData();
 				}
 
 				if (bFilterValid)
@@ -516,7 +518,8 @@ namespace Chaos
 		// Set the captured filter to our new shapes
 		for (const TUniquePtr<FPerShapeData>& Shape : Parent->ShapesArray())
 		{
-			Shape->SetSimData(Filter);
+			Shape->SetSimData(SimFilter);
+			Shape->SetQueryData(QueryFilter);
 		}
 	}
 
diff --git a/Engine/Source/Runtime/Experimental/Chaos/Private/PhysicsProxy/GeometryCollectionPhysicsProxy.cpp b/Engine/Source/Runtime/Experimental/Chaos/Private/PhysicsProxy/GeometryCollectionPhysicsProxy.cpp
--- a/Engine/Source/Runtime/Experimental/Chaos/Private/PhysicsProxy/GeometryCollectionPhysicsProxy.cpp
+++ b/Engine/Source/Runtime/Experimental/Chaos/Private/PhysicsProxy/GeometryCollectionPhysicsProxy.cpp
@@ -61,6 +61,14 @@ FAutoConsoleVariableRef CVarGeometryCollectionCollideAll(
 	GeometryCollectionCollideAll,
 	TEXT("Bypass the collision matrix and make geometry collections collide against everything"));
 
+
+bool bGeometryCollectionEnabledNestedChildTransformUpdates = true;
+FAutoConsoleVariableRef CVarEnabledNestedChildTransformUpdates(
+	TEXT("p.GeometryCollection.EnabledNestedChildTransformUpdates"),
+	bGeometryCollectionEnabledNestedChildTransformUpdates,
+	TEXT("Enable updates for driven, disabled, child bodies. Used for line trace results against geometry collections.[def: true]"));
+
+
 DEFINE_LOG_CATEGORY_STATIC(UGCC_LOG, Error, All);
 
 //==============================================================================
@@ -1710,6 +1718,18 @@ void FGeometryCollectionPhysicsProxy::BufferPhysicsResults(Chaos::FPBDRigidsSolv
 							ProxyElementHandle->R() = ParticleToWorld.GetRotation();
 							CurrentSolver->GetEvolution()->DirtyParticle(*ProxyElementHandle);
 						}
+
+						if (bGeometryCollectionEnabledNestedChildTransformUpdates)
+						{
+							if (!ClusterParent->Disabled())
+							{
+								FTransform ChildToWorld = Handle->ChildToParent() * FRigidTransform3(ClusterParent->X(), ClusterParent->R());
+								Handle->X() = Handle->P() = ChildToWorld.GetTranslation();
+								Handle->R() = Handle->Q() = ChildToWorld.GetRotation();
+								Handle->UpdateWorldSpaceState(ChildToWorld, FVec3(0));
+								CurrentSolver->GetEvolution()->DirtyParticle(*Handle);
+							}
+						}
 					}
 				}
 			}    // end if
