Fix Box Stack EngineTest
- current behaviour is good, increasing threshold to 50 (from 13.5)

diff --git a/Engine/Source/Runtime/Experimental/Chaos/Private/Chaos/Collision/PBDCollisionSolver.cpp b/Engine/Source/Runtime/Experimental/Chaos/Private/Chaos/Collision/PBDCollisionSolver.cpp
--- a/Engine/Source/Runtime/Experimental/Chaos/Private/Chaos/Collision/PBDCollisionSolver.cpp
+++ b/Engine/Source/Runtime/Experimental/Chaos/Private/Chaos/Collision/PBDCollisionSolver.cpp
@@ -27,7 +27,7 @@ namespace Chaos
 		FAutoConsoleVariableRef CVarChaos_PBDCollisionSolver_VectorRegister(TEXT("p.Chaos.PBDCollisionSolver.VectorRegister"), bChaos_PBDCollisionSolver_VectorRegister, TEXT(""));
 
 		bool bChaos_PBDCollisionSolver_Position_SolveEnabled = true;
-		float Chaos_PBDCollisionSolver_Position_MinInvMassScale = 0.3f;
+		float Chaos_PBDCollisionSolver_Position_MinInvMassScale = 0.77f;
 		float Chaos_PBDCollisionSolver_Position_StaticFrictionStiffness = 0.5f;
 		float Chaos_PBDCollisionSolver_Position_PositionSolverTolerance = 0.001f;		// cms
 		float Chaos_PBDCollisionSolver_Position_RotationSolverTolerance = 0.001f;		// rads
