The correct fix to the logic error in DoForSweepIntersectCellsImp

#localization none
#tests PIE, had a consistent repro set up and with the fix wasn't able to repro the bug anymore.

diff --git a/Engine/Source/Runtime/Experimental/Chaos/Public/Chaos/AABBTreeDirtyGridUtils.h b/Engine/Source/Runtime/Experimental/Chaos/Public/Chaos/AABBTreeDirtyGridUtils.h
--- a/Engine/Source/Runtime/Experimental/Chaos/Public/Chaos/AABBTreeDirtyGridUtils.h
+++ b/Engine/Source/Runtime/Experimental/Chaos/Public/Chaos/AABBTreeDirtyGridUtils.h
@@ -383,7 +383,8 @@ namespace Chaos
 					FReal CrossingHorizontalCellBorderT = std::numeric_limits<FReal>::max();
 					if (!DxTooSmall)
 					{
-						CrossingVerticleCellBorderT = DtDx * ((FReal)(CurrentCellIndexX0 - ReferenceCellIndexX + (DeltaCelIndexX > 0 ? 1 : 0)) * DirtyElementGridCellSize - X1);
+						// It should be X0 because CurrentCellIndexX0 has already been updated, but X1 will be updated later on.
+						CrossingVerticleCellBorderT = DtDx * ((FReal)(CurrentCellIndexX0 - ReferenceCellIndexX + (DeltaCelIndexX > 0 ? 1 : 0)) * DirtyElementGridCellSize - X0);
 					}
 
 					if (!DyTooSmall)
@@ -391,7 +392,7 @@ namespace Chaos
 						CrossingHorizontalCellBorderT = DtDy * ((FReal)(CurrentCellIndexY0 - ReferenceCellIndexY + (DeltaCelIndexY > 0 ? 1 : 0)) * DirtyElementGridCellSize - Y1);
 					}
 
-					if (CrossingHorizontalCellBorderT >= CrossingVerticleCellBorderT)
+					if (CrossingHorizontalCellBorderT < CrossingVerticleCellBorderT)
 					{
 						CurrentCellIndexY1 += DeltaCelIndexY;
 					}
