Chaos  : Fixed strain field not affecting clustered group

diff --git a/Engine/Source/Runtime/Experimental/Chaos/Public/PhysicsProxy/FieldSystemProxyHelper.h b/Engine/Source/Runtime/Experimental/Chaos/Public/PhysicsProxy/FieldSystemProxyHelper.h
--- a/Engine/Source/Runtime/Experimental/Chaos/Public/PhysicsProxy/FieldSystemProxyHelper.h
+++ b/Engine/Source/Runtime/Experimental/Chaos/Public/PhysicsProxy/FieldSystemProxyHelper.h
@@ -288,8 +288,11 @@ namespace Chaos
 
 			for (Chaos::FPBDRigidParticleHandle* Handle : Activated)
 			{
-				IPhysicsProxyBase* ParentProxy = Iter.Key->PhysicsProxy();
-				Handle->SetPhysicsProxy(ParentProxy);
+				if (Handle->PhysicsProxy() == nullptr)
+				{
+					IPhysicsProxyBase* ParentProxy = Iter.Key->PhysicsProxy();
+					Handle->SetPhysicsProxy(ParentProxy);
+				}
 			}
 		}
 	}
