Fix to the logic error in DoForSweepIntersectCellsImp

#localization none
#tests PIE, had a consistent repro set up and with the fix wasn't able to repro the bug anymore.

diff --git a/Engine/Source/Runtime/Experimental/Chaos/Public/Chaos/AABBTreeDirtyGridUtils.h b/Engine/Source/Runtime/Experimental/Chaos/Public/Chaos/AABBTreeDirtyGridUtils.h
--- a/Engine/Source/Runtime/Experimental/Chaos/Public/Chaos/AABBTreeDirtyGridUtils.h
+++ b/Engine/Source/Runtime/Experimental/Chaos/Public/Chaos/AABBTreeDirtyGridUtils.h
@@ -361,7 +361,6 @@ namespace Chaos
 				{
 					X0 += DirtyElementGridCellSize * (FReal)DeltaCelIndexX;
 				}
-				CurrentCellIndexX0 += DeltaCelIndexX;
 			}
 
 			// Advance line 1
@@ -383,8 +382,8 @@ namespace Chaos
 					FReal CrossingHorizontalCellBorderT = std::numeric_limits<FReal>::max();
 					if (!DxTooSmall)
 					{
-						// It should be X0 because CurrentCellIndexX0 has already been updated, but X1 will be updated later on.
-						CrossingVerticleCellBorderT = DtDx * ((FReal)(CurrentCellIndexX0 - ReferenceCellIndexX + (DeltaCelIndexX > 0 ? 1 : 0)) * DirtyElementGridCellSize - X0);
+
+						CrossingVerticleCellBorderT = DtDx * ((FReal)(CurrentCellIndexX0 - ReferenceCellIndexX + (DeltaCelIndexX > 0 ? 1 : 0)) * DirtyElementGridCellSize - X1);
 					}
 
 					if (!DyTooSmall)
@@ -409,6 +408,8 @@ namespace Chaos
 				}
 				CurrentCellIndexX1 += DeltaCelIndexX;
 			}
+
+			CurrentCellIndexX0 += DeltaCelIndexX;
 			Done = (CurrentCellIndexY0 == LastCellIndexY) && (DeltaCelIndexX * CurrentCellIndexX0 > LastCellIndexX * DeltaCelIndexX);
 		}
 	}
