GJK will now inlcude no margins for overlap calculations
Fixes:
On Component Begin and End Overlap are spammed when pawn is on edge of collision box

diff --git a/Engine/Source/Runtime/Experimental/Chaos/Public/Chaos/GJK.h b/Engine/Source/Runtime/Experimental/Chaos/Public/Chaos/GJK.h
--- a/Engine/Source/Runtime/Experimental/Chaos/Public/Chaos/GJK.h
+++ b/Engine/Source/Runtime/Experimental/Chaos/Public/Chaos/GJK.h
@@ -26,6 +26,31 @@
 
 namespace Chaos
 {
+
+	// Calculate the margins used for queries based on shape radius, shape margins and shape types
+	template <typename TGeometryA, typename TGeometryB, typename T>
+	void CalculateQueryMargins(const TGeometryA& A, const TGeometryB& B, T& outMarginA, T& outMarginB)
+	{
+		// Margin selection logic: we only need a small margin for sweeps since we only move the sweeping object
+		// to the point where it just touches.
+		// Spheres and Capsules: always use the core shape and full "margin" because it represents the radius
+		// Sphere/Capsule versus OtherShape: no margin on other
+		// OtherShape versus OtherShape: use margin of the smaller shape, zero margin on the other
+		const T RadiusA = A.GetRadius();
+		const T RadiusB = B.GetRadius();
+		const bool bHasRadiusA = RadiusA > 0;
+		const bool bHasRadiusB = RadiusB > 0;
+
+		// The sweep margins if required. Only one can be non-zero (we keep the smaller one)
+		const T SweepMarginScale = 0.05f;
+		const bool bAIsSmallest = A.GetMargin() < B.GetMargin();
+		const T SweepMarginA = (bHasRadiusA || bHasRadiusB) ? 0.0f : (bAIsSmallest ? SweepMarginScale * A.GetMargin() : 0.0f);
+		const T SweepMarginB = (bHasRadiusA || bHasRadiusB) ? 0.0f : (bAIsSmallest ? 0.0f : SweepMarginScale * B.GetMargin());
+
+		// Net margin (note: both SweepMargins are zero if either Radius is non-zero, and only one SweepMargin can be non-zero)
+		outMarginA = RadiusA + SweepMarginA;
+		outMarginB = RadiusB + SweepMarginB;
+	}
 	
 	/** 
 		Determines if two convex geometries overlap.
@@ -68,8 +93,13 @@ namespace Chaos
 		VectorRegister4Float BarycentricSimd;
 		VectorRegister4Int NumVerts = GlobalVectorConstants::IntZero;
 
-		const T ThicknessA = A.GetMargin() + InThicknessA;
-		const T ThicknessB = B.GetMargin() + InThicknessB;
+		T ThicknessA;
+		T ThicknessB;
+		CalculateQueryMargins(A, B, ThicknessA, ThicknessB);
+		ThicknessA += InThicknessA;
+		ThicknessB += InThicknessB;
+
+
 		const T Inflation = ThicknessA + ThicknessB + static_cast<T>(1e-3);
 		const VectorRegister4Float InflationSimd = MakeVectorRegisterFloat((float)Inflation, (float)Inflation, (float)Inflation, (float)Inflation);
 
@@ -81,9 +111,9 @@ namespace Chaos
 				break;	//if taking too long just stop. This should never happen
 			}
 			const VectorRegister4Float NegVSimd = VectorNegate(VSimd);
-			const VectorRegister4Float SupportASimd = A.SupportCoreSimd(NegVSimd, A.GetMargin());
+			const VectorRegister4Float SupportASimd = A.SupportCoreSimd(NegVSimd, ThicknessA);
 			const VectorRegister4Float VInBSimd = VectorQuaternionRotateVector(AToBRotationSimd, VSimd);
-			const VectorRegister4Float SupportBLocalSimd = B.SupportCoreSimd(VInBSimd, B.GetMargin());
+			const VectorRegister4Float SupportBLocalSimd = B.SupportCoreSimd(VInBSimd, ThicknessB);
 			//const TVector<T, 3> SupportB = BToATM.TransformPositionNoScale(SupportBLocal);  // Original code
 			const VectorRegister4Float SupportBSimd = VectorAdd(VectorQuaternionRotateVector(RotationSimd, SupportBLocalSimd), TranslationSimd);
 			const VectorRegister4Float WSimd = VectorSubtract(SupportASimd, SupportBSimd);
@@ -141,8 +171,11 @@ namespace Chaos
 		bool bNearZero = false;
 		int NumIterations = 0;
 		T PrevDist2 = FLT_MAX;
-		const T ThicknessA = A.GetMargin() + InThicknessA;
-		const T ThicknessB = B.GetMargin() + InThicknessB;
+		T ThicknessA;
+		T ThicknessB;
+		CalculateQueryMargins(A, B, ThicknessA, ThicknessB);
+		ThicknessA += InThicknessA;
+		ThicknessB += InThicknessB;
 		const T Inflation = ThicknessA + ThicknessB + static_cast<T>(1e-3);
 		const T Inflation2 = Inflation * Inflation;
 		int32 VertexIndexA = INDEX_NONE, VertexIndexB = INDEX_NONE;
@@ -153,9 +186,9 @@ namespace Chaos
 				break;	//if taking too long just stop. This should never happen
 			}
 			const TVector<T, 3> NegV = -V;
-			const TVector<T, 3> SupportA = A.SupportCore(NegV, A.GetMargin(), nullptr, VertexIndexA);
+			const TVector<T, 3> SupportA = A.SupportCore(NegV, ThicknessA, nullptr, VertexIndexA);
 			const TVector<T, 3> VInB = AToBRotation * V;
-			const TVector<T, 3> SupportBLocal = B.SupportCore(VInB, B.GetMargin(), nullptr, VertexIndexB);
+			const TVector<T, 3> SupportBLocal = B.SupportCore(VInB, ThicknessB, nullptr, VertexIndexB);
 			const TVector<T, 3> SupportB = BToATM.TransformPositionNoScale(SupportBLocal);
 			const TVector<T, 3> W = SupportA - SupportB;
 
@@ -218,8 +251,11 @@ namespace Chaos
 		bool bNearZero = false;
 		int NumIterations = 0;
 		T PrevDist2 = FLT_MAX;
-		const T ThicknessA = A.GetMargin() + InThicknessA;
-		const T ThicknessB = B.GetMargin() + InThicknessB;
+		T ThicknessA;
+		T ThicknessB;
+		CalculateQueryMargins(A, B, ThicknessA, ThicknessB);
+		ThicknessA += InThicknessA;
+		ThicknessB += InThicknessB;
 		const T Inflation = ThicknessA + ThicknessB + static_cast<T>(1e-3);
 		const T Inflation2 = Inflation * Inflation;
 		int32 VertexIndexA = INDEX_NONE, VertexIndexB = INDEX_NONE;
@@ -230,9 +266,9 @@ namespace Chaos
 				break;	//if taking too long just stop. This should never happen
 			}
 			const TVector<T, 3> NegV = -V;
-			const TVector<T, 3> SupportA = A.SupportCore(NegV, A.GetMargin(), nullptr, VertexIndexA);
+			const TVector<T, 3> SupportA = A.SupportCore(NegV, ThicknessA, nullptr, VertexIndexA);
 			const TVector<T, 3> VInB = V; // same space
-			const TVector<T, 3> SupportB = B.SupportCore(VInB, B.GetMargin(), nullptr, VertexIndexB);
+			const TVector<T, 3> SupportB = B.SupportCore(VInB, ThicknessB, nullptr, VertexIndexB);
 			const TVector<T, 3> W = SupportA - SupportB;
 
 			if (TVector<T, 3>::DotProduct(V, W) > Inflation)
@@ -1403,25 +1439,9 @@ namespace Chaos
 		ensure(FMath::IsNearlyEqual(RayDir.SizeSquared(), (T)1, (T)KINDA_SMALL_NUMBER));
 		ensure(RayLength > 0);
 
-		// Margin selection logic: we only need a small margin for sweeps since we only move the sweeping object
-		// to the point where it just touches.
-		// Spheres and Capsules: always use the core shape and full "margin" because it represents the radius
-		// Sphere/Capsule versus OtherShape: no margin on other
-		// OtherShape versus OtherShape: use margin of the smaller shape, zero margin on the other
-		const T RadiusA = A.GetRadius();
-		const T RadiusB = B.GetRadius();
-		const bool bHasRadiusA = RadiusA > 0;
-		const bool bHasRadiusB = RadiusB > 0;
-
-		// The sweep margins if required. Only one can be non-zero (we keep the smaller one)
-		const T SweepMarginScale = 0.05f;
-		const bool bAIsSmallest = A.GetMargin() < B.GetMargin();
-		const T SweepMarginA = (bHasRadiusA || bHasRadiusB) ? 0.0f : (bAIsSmallest ? SweepMarginScale * A.GetMargin() : 0.0f);
-		const T SweepMarginB = (bHasRadiusA || bHasRadiusB) ? 0.0f : (bAIsSmallest ? 0.0f : SweepMarginScale * B.GetMargin());
-
-		// Net margin (note: both SweepMargins are zero if either Radius is non-zero, and only one SweepMargin can be non-zero)
-		const T MarginA = RadiusA + SweepMarginA;
-		const T MarginB = RadiusB + SweepMarginB;
+		T MarginA;
+		T MarginB;
+		CalculateQueryMargins(A,B, MarginA, MarginB);
 
 		const TVector<T, 3> StartPoint = StartTM.GetLocation();
 
