Ignore broadphase collision flag for particle2

diff --git a/Engine/Source/Runtime/Experimental/Chaos/Public/Chaos/Collision/SpatialAccelerationBroadPhase.h b/Engine/Source/Runtime/Experimental/Chaos/Public/Chaos/Collision/SpatialAccelerationBroadPhase.h
--- a/Engine/Source/Runtime/Experimental/Chaos/Public/Chaos/Collision/SpatialAccelerationBroadPhase.h
+++ b/Engine/Source/Runtime/Experimental/Chaos/Public/Chaos/Collision/SpatialAccelerationBroadPhase.h
@@ -255,7 +255,8 @@ namespace Chaos
 						continue;
 					}
 					
-					if (Particle1.HasCollisionConstraintFlag(ECollisionConstraintFlags::CCF_BroadPhaseIgnoreCollisions) )
+					if (Particle1.HasCollisionConstraintFlag(ECollisionConstraintFlags::CCF_BroadPhaseIgnoreCollisions) || 
+						Particle2Generic->HasCollisionConstraintFlag(ECollisionConstraintFlags::CCF_BroadPhaseIgnoreCollisions))
 					{
 						if (IgnoreCollisionManager.IgnoresCollision(Particle1.UniqueIdx(), Particle2.UniqueIdx()))
 						{
