Chaos : fix crash when using removal event if particle is null or has no geometry

diff --git a/Engine/Source/Runtime/Experimental/Chaos/Private/Chaos/PBDRigidsEvolution.cpp b/Engine/Source/Runtime/Experimental/Chaos/Private/Chaos/PBDRigidsEvolution.cpp
--- a/Engine/Source/Runtime/Experimental/Chaos/Private/Chaos/PBDRigidsEvolution.cpp
+++ b/Engine/Source/Runtime/Experimental/Chaos/Private/Chaos/PBDRigidsEvolution.cpp
@@ -1089,6 +1089,11 @@ namespace Chaos
 
 	void FPBDRigidsEvolutionBase::DisableParticleWithRemovalEvent(FGeometryParticleHandle* Particle)
 	{
+		if (Particle == nullptr)
+		{
+			return;
+		}
+		
 		// Record removal for event generation
 		const int32 NewIdx = MAllRemovals.Add(FRemovalData());
 		FRemovalData& Removal = MAllRemovals[NewIdx];
@@ -1104,7 +1109,7 @@ namespace Chaos
 			Removal.Mass = 0.0;
 		}
 		
-		if (Particle->Geometry()->HasBoundingBox())
+		if (Particle->Geometry() && Particle->Geometry()->HasBoundingBox())
 		{
 			Removal.BoundingBox = Particle->Geometry()->BoundingBox();
 		}
