Chaos : fix computation of VSmooth and WSmooth to make sure it keeps the right magnitude
- this was causing issue with falling ragdolls goiung to sleep in some cases

diff --git a/Engine/Source/Runtime/Experimental/Chaos/Private/PhysicsProxy/SingleParticlePhysicsProxy.cpp b/Engine/Source/Runtime/Experimental/Chaos/Private/PhysicsProxy/SingleParticlePhysicsProxy.cpp
--- a/Engine/Source/Runtime/Experimental/Chaos/Private/PhysicsProxy/SingleParticlePhysicsProxy.cpp
+++ b/Engine/Source/Runtime/Experimental/Chaos/Private/PhysicsProxy/SingleParticlePhysicsProxy.cpp
@@ -108,7 +108,7 @@ void PushToPhysicsStateImp(const Chaos::FDirtyPropertiesManager& Manager, Chaos:
 			if(auto NewData = ParticleData.FindDynamics(Manager, DataIdx))
 			{
 				RigidHandle->SetDynamics(*NewData);
-				RigidHandle->ResetVSmoothFromForces();
+				Evolution.ResetVSmoothFromForces(*RigidHandle);
 			}
 
 			if(auto NewData = ParticleData.FindDynamicMisc(Manager,DataIdx))
diff --git a/Engine/Source/Runtime/Experimental/Chaos/Public/Chaos/PBDRigidParticles.h b/Engine/Source/Runtime/Experimental/Chaos/Public/Chaos/PBDRigidParticles.h
--- a/Engine/Source/Runtime/Experimental/Chaos/Public/Chaos/PBDRigidParticles.h
+++ b/Engine/Source/Runtime/Experimental/Chaos/Public/Chaos/PBDRigidParticles.h
@@ -178,22 +178,6 @@ class TPBDRigidParticles : public TRigidParticles<T, d>
 		this->SleepType(Index) = InSleepType;
 	}
 
-
-	void ResetVSmoothFromForces(int32 Index)
-	{
-		// Reset VSmooth to something roughly in the same direction as what V will be after integration.
-		// This is temp fix, if this is only re-computed after solve, island will get incorrectly put back to sleep even if it was just impulsed.
-		FReal FakeDT = (FReal)1. / (FReal)30.;
-		if (this->LinearImpulseVelocity(Index).IsNearlyZero() == false || this->Acceleration(Index).IsNearlyZero() == false)
-		{
-			this->VSmooth(Index) = this->Acceleration(Index) * FakeDT + this->LinearImpulseVelocity(Index);
-		}
-		if (this->AngularImpulseVelocity(Index).IsNearlyZero() == false || this->AngularAcceleration(Index).IsNearlyZero() == false)
-		{
-			this->WSmooth(Index) = this->AngularAcceleration(Index) * FakeDT + this->AngularImpulseVelocity(Index);
-		}
-	}
-
 	FString ToString(int32 index) const
 	{
 		FString BaseString = TRigidParticles<T, d>::ToString(index);
diff --git a/Engine/Source/Runtime/Experimental/Chaos/Public/Chaos/PBDRigidsEvolutionGBF.h b/Engine/Source/Runtime/Experimental/Chaos/Public/Chaos/PBDRigidsEvolutionGBF.h
--- a/Engine/Source/Runtime/Experimental/Chaos/Public/Chaos/PBDRigidsEvolutionGBF.h
+++ b/Engine/Source/Runtime/Experimental/Chaos/Public/Chaos/PBDRigidsEvolutionGBF.h
@@ -31,6 +31,7 @@ namespace Chaos
 		CHAOS_API extern FRealSingle HackMaxAngularVelocity;
 		CHAOS_API extern FRealSingle HackMaxVelocity;
 		CHAOS_API extern FRealSingle CCDEnableThresholdBoundsScale;
+		CHAOS_API extern FRealSingle SmoothedPositionLerpRate;
 		CHAOS_API extern bool bChaosCollisionCCDUseTightBoundingBox;
 	}
 
@@ -349,6 +350,27 @@ namespace Chaos
 
 		CHAOS_API void TransferJointConstraintCollisions();
 
+		// Resets VSmooth value to something plausible based on external forces to prevent object from going back to sleep if it was just impulsed.
+		template <bool bPersistent>
+		void ResetVSmoothFromForces(TPBDRigidParticleHandleImp<FReal, 3, bPersistent>& Particle)
+		{
+			const FReal SmoothRate = FMath::Clamp(CVars::SmoothedPositionLerpRate, 0.0f, 1.0f);
+	
+			// Reset VSmooth to something roughly in the same direction as what V will be after integration.
+			// This is temp fix, if this is only re-computed after solve, island will get incorrectly put back to sleep even if it was just impulsed.
+			FReal FakeDT = (FReal)1. / (FReal)30.;
+			if (Particle.LinearImpulseVelocity().IsNearlyZero() == false || Particle.Acceleration().IsNearlyZero() == false)
+			{
+				const FVec3 PredictedLinearVelocity = Particle.V() + Particle.Acceleration() * FakeDT + Particle.LinearImpulseVelocity();
+				Particle.VSmooth() =FMath::Lerp(Particle.VSmooth(), PredictedLinearVelocity, SmoothRate);
+			}
+			if (Particle.AngularImpulseVelocity().IsNearlyZero() == false || Particle.AngularAcceleration().IsNearlyZero() == false)
+			{
+				const FVec3 PredictedAngularVelocity = Particle.W() + Particle.AngularAcceleration() * FakeDT + Particle.AngularImpulseVelocity();
+				Particle.WSmooth() = FMath::Lerp(Particle.WSmooth(), PredictedAngularVelocity, SmoothRate);
+			}
+		}
+		
 	protected:
 
 		CHAOS_API void AdvanceOneTimeStepImpl(const FReal dt, const FSubStepInfo& SubStepInfo);
diff --git a/Engine/Source/Runtime/Experimental/Chaos/Public/Chaos/ParticleHandle.h b/Engine/Source/Runtime/Experimental/Chaos/Public/Chaos/ParticleHandle.h
--- a/Engine/Source/Runtime/Experimental/Chaos/Public/Chaos/ParticleHandle.h
+++ b/Engine/Source/Runtime/Experimental/Chaos/Public/Chaos/ParticleHandle.h
@@ -935,12 +935,6 @@ public:
 	TVector<T, d>& AngularImpulseVelocity() { return PBDRigidParticles->AngularImpulseVelocity(ParticleIdx); }
 	void SetAngularImpulseVelocity(const TVector<T, d>& InAngularImpulseVelocity, bool bInvalidate = false) { PBDRigidParticles->AngularImpulseVelocity(ParticleIdx) = InAngularImpulseVelocity; }
 
-	// Resets VSmooth value to something plausible based on external forces to prevent object from going back to sleep if it was just impulsed.
-	void ResetVSmoothFromForces()
-	{
-		PBDRigidParticles->ResetVSmoothFromForces(ParticleIdx);
-	}
-
 	void SetDynamics(const FParticleDynamics& Dynamics)
 	{
 		SetAcceleration(Dynamics.Acceleration());
