Joint sphericalcylindrical/planar position constraint

diff --git a/Engine/Source/Runtime/Experimental/Chaos/Private/Chaos/Joint/PBDJointCachedSolverGaussSeidel.cpp b/Engine/Source/Runtime/Experimental/Chaos/Private/Chaos/Joint/PBDJointCachedSolverGaussSeidel.cpp
--- a/Engine/Source/Runtime/Experimental/Chaos/Private/Chaos/Joint/PBDJointCachedSolverGaussSeidel.cpp
+++ b/Engine/Source/Runtime/Experimental/Chaos/Private/Chaos/Joint/PBDJointCachedSolverGaussSeidel.cpp
@@ -391,75 +391,211 @@ void FPBDJointCachedSolver::InitPositionConstraints(
 
 	const TVec3<EJointMotionType>& LinearMotion = JointSettings.LinearMotionTypes;
 
+	if (bLinearLocked[0] || bLinearLocked[1] || bLinearLocked[2]) 
+	{
+		// Process locked constraints
+		InitLockedPositionConstraint(JointSettings, Dt, LinearMotion);
+	}
+	if (bLinearLimited[0] || bLinearLimited[1] || bLinearLimited[2])
+	{
+		// Process limited constraints
+		if (bLinearLimited[0] && bLinearLimited[1] && bLinearLimited[2])
+		{
+			// Spherical constraint
+			InitSphericalPositionConstraint(JointSettings, Dt);
+		}
+		else if (bLinearLimited[1] && bLinearLimited[2])
+		{
+			// Cylindrical constraint along X axis
+			InitCylindricalPositionConstraint(JointSettings, Dt, 0);
+		}
+		else if (bLinearLimited[0] && bLinearLimited[2])
+		{
+			// Cylindrical constraint along Y axis
+			InitCylindricalPositionConstraint(JointSettings, Dt, 1);
+		}
+		else if (bLinearLimited[0] && bLinearLimited[1])
+		{
+			// Cylindrical constraint along Z axis
+			InitCylindricalPositionConstraint(JointSettings, Dt, 2);
+		}
+		else if (bLinearLimited[0])
+		{
+			// Planar constraint along X axis
+			InitPlanarPositionConstraint(JointSettings, Dt, 0);
+		}
+		else if (bLinearLimited[1])
+		{
+			// Planar constraint along Y axis
+			InitPlanarPositionConstraint(JointSettings, Dt, 1);
+		}
+		else if (bLinearLimited[2])
+		{
+			// Planar constraint along Z axis
+			InitPlanarPositionConstraint(JointSettings, Dt, 2);
+		}
+	}
+}
 
-	const FVec3 ConstraintArm0Limited = ConnectorXs[1] - CurrentPs[0];
-	const FVec3 ConstraintArm1Limited = ConnectorXs[1] - CurrentPs[1];
+void FPBDJointCachedSolver::InitPositionDatasMass(
+	FAxisConstraintDatas& PositionDatas, 
+	const int32 ConstraintIndex,
+	const FReal Dt)
+{
+	const FVec3 AngularAxis0 = FVec3::CrossProduct(PositionDatas.ConstraintArms[ConstraintIndex][0], PositionDatas.ConstraintAxis[ConstraintIndex]);
+	const FVec3 AngularAxis1 = FVec3::CrossProduct(PositionDatas.ConstraintArms[ConstraintIndex][1], PositionDatas.ConstraintAxis[ConstraintIndex]);
+	const FVec3 IA0 = Utilities::Multiply(InvI(0), AngularAxis0);
+	const FVec3 IA1 = Utilities::Multiply(InvI(1), AngularAxis1);
+	const FReal II0 = FVec3::DotProduct(AngularAxis0, IA0);
+	const FReal II1 = FVec3::DotProduct(AngularAxis1, IA1);
+
+	PositionDatas.UpdateMass(ConstraintIndex, IA0, IA1, InvM(0) + II0 + InvM(1) + II1, Dt);
+}
 	
-	FVec3 ConstraintArm0Locked = ConstraintArm0Limited;
-	FVec3 ConstraintArm1Locked = ConstraintArm1Limited;
+void FPBDJointCachedSolver::InitPositionConstraintDatas(
+	const int32 ConstraintIndex,
+	const FVec3& ConstraintAxis,
+	const FReal& ConstraintDelta,
+	const FReal ConstraintRestitution,
+	const FReal Dt,
+	const FReal ConstraintLimit,
+	const EJointMotionType JointType,
+	const FVec3& ConstraintArm0,
+	const FVec3& ConstraintArm1)
+{
+	const FVec3 LocalAxis = (ConstraintDelta < 0.0f) ? -ConstraintAxis : ConstraintAxis;
+	const FReal LocalDelta = (ConstraintDelta < 0.0f) ? -ConstraintDelta : ConstraintDelta;
 
+	PositionConstraints.MotionType[ConstraintIndex] = JointType;
+
+	if(JointType == EJointMotionType::Locked)
+	{
+		PositionConstraints.ConstraintLimits[ConstraintIndex] = 0.0f;
+		PositionConstraints.UpdateDatas(ConstraintIndex, LocalAxis, LocalDelta,
+			0.0, false, ConstraintArm0, ConstraintArm1);
+	}
+	else if(JointType == EJointMotionType::Limited)
+	{
+		PositionConstraints.ConstraintLimits[ConstraintIndex] =
+			FMath::Max(ConstraintLimit - GetLinearConstraintPadding(ConstraintIndex), (FReal)0.);
+		PositionConstraints.UpdateDatas(ConstraintIndex, LocalAxis, LocalDelta,
+			ConstraintRestitution, true, ConstraintArm0, ConstraintArm1);
+	}
+
+	const FVec3 CV0 = V(0) + FVec3::CrossProduct(W(0), ConstraintArm0);
+	const FVec3 CV1 = V(1) + FVec3::CrossProduct(W(1), ConstraintArm1);
+	const FVec3 CV = CV1 - CV0;
+
+	InitConstraintAxisLinearVelocities[ConstraintIndex] = FVec3::DotProduct(CV, LocalAxis);
+
+	InitPositionDatasMass(PositionConstraints, ConstraintIndex, Dt);
+}
+
+	
+void FPBDJointCachedSolver::InitLockedPositionConstraint(
+		const FPBDJointSettings& JointSettings,
+		const FReal Dt,
+		const TVec3<EJointMotionType>& LinearMotion)
+{
+	FVec3 ConstraintArm0 = ConnectorXs[1] - CurrentPs[0];
+	const FVec3 ConstraintArm1 = ConnectorXs[1] - CurrentPs[1];
+	
 	const FVec3 DX = ConnectorXs[1] - ConnectorXs[0];
 	const FMatrix33 R0M = ConnectorRs[0].ToMatrix(); 
 	FVec3 CX = FVec3::ZeroVector;
-
+	
+	// For a locked constraint we try to match an exact constraint, 
+	// it is why we are adding back the constraint projection along each axis.
+	// If the 3 axis are locked the constraintarm0 is then ConnectorXs[0] - CurrentPs[0];
 	for(int32 ConstraintIndex = 0; ConstraintIndex < 3; ++ConstraintIndex)
 	{
-		if(bLinearLocked[ConstraintIndex] || bLinearLimited[ConstraintIndex])
+		if(LinearMotion[ConstraintIndex] == EJointMotionType::Locked)
 		{
 			const FVec3 ConstraintAxis = R0M.GetAxis(ConstraintIndex);
 			CX[ConstraintIndex] = FVec3::DotProduct(DX, ConstraintAxis);
-
-			if(bLinearLocked[ConstraintIndex])
-			{
-				ConstraintArm0Locked -= ConstraintAxis * CX[ConstraintIndex];
-			}
+			ConstraintArm0 -= ConstraintAxis * CX[ConstraintIndex];
+			
 		}
 	}
-	
 	for(int32 ConstraintIndex = 0; ConstraintIndex < 3; ++ConstraintIndex)
 	{
-		if(bLinearLocked[ConstraintIndex] || bLinearLimited[ConstraintIndex])
+		if(LinearMotion[ConstraintIndex] == EJointMotionType::Locked)
 		{
 			const FVec3 ConstraintAxis = R0M.GetAxis(ConstraintIndex);
-			PositionConstraints.MotionType[ConstraintIndex] = LinearMotion[ConstraintIndex];
-			
-			if(bLinearLocked[ConstraintIndex])
-			{
-				PositionConstraints.ConstraintLimits[ConstraintIndex] = 0.0f;
-				PositionConstraints.UpdateDatas(ConstraintIndex, ConstraintAxis, CX[ConstraintIndex],
-					0.0, false, ConstraintArm0Locked, ConstraintArm1Locked);
-			}
-			else if(bLinearLimited[ConstraintIndex])
-			{
-				PositionConstraints.ConstraintLimits[ConstraintIndex] =
-					FMath::Max(JointSettings.LinearLimit - GetLinearConstraintPadding(ConstraintIndex), (FReal)0.);
-				PositionConstraints.UpdateDatas(ConstraintIndex, ConstraintAxis, CX[ConstraintIndex],
-					JointSettings.LinearRestitution, true, ConstraintArm0Limited, ConstraintArm1Limited);
-			}
-			const FVec3 CV0 = V(0) + FVec3::CrossProduct(W(0), PositionConstraints.ConstraintArms[ConstraintIndex][0]);
-			const FVec3 CV1 = V(1) + FVec3::CrossProduct(W(1), PositionConstraints.ConstraintArms[ConstraintIndex][1]);
-			const FVec3 CV = CV1 - CV0;
-			
-			InitConstraintAxisLinearVelocities[ConstraintIndex] = FVec3::DotProduct(CV, ConstraintAxis);
-
-			InitPositionDatasMass(PositionConstraints, ConstraintIndex, Dt);
+			InitPositionConstraintDatas(ConstraintIndex, ConstraintAxis,  CX[ConstraintIndex], 0.0, Dt,
+				0.0, EJointMotionType::Locked, ConstraintArm0, ConstraintArm1);
 		}
 	}
 }
 
-void FPBDJointCachedSolver::InitPositionDatasMass(
-	FAxisConstraintDatas& PositionDatas, 
-	const int32 ConstraintIndex,
+void FPBDJointCachedSolver::InitSphericalPositionConstraint(
+	const FPBDJointSettings& JointSettings,
 	const FReal Dt)
 {
-	const FVec3 AngularAxis0 = FVec3::CrossProduct(PositionDatas.ConstraintArms[ConstraintIndex][0], PositionDatas.ConstraintAxis[ConstraintIndex]);
-	const FVec3 AngularAxis1 = FVec3::CrossProduct(PositionDatas.ConstraintArms[ConstraintIndex][1], PositionDatas.ConstraintAxis[ConstraintIndex]);
-	const FVec3 IA0 = Utilities::Multiply(InvI(0), AngularAxis0);
-	const FVec3 IA1 = Utilities::Multiply(InvI(1), AngularAxis1);
-	const FReal II0 = FVec3::DotProduct(AngularAxis0, IA0);
-	const FReal II1 = FVec3::DotProduct(AngularAxis1, IA1);
+	FVec3 SphereAxis0;
+	FReal SphereDelta0;
+	FPBDJointUtilities::GetSphericalAxisDelta(ConnectorXs[0], ConnectorXs[1], SphereAxis0, SphereDelta0);
+	
+	const FVec3 SphereAxis1 = (FMath::Abs(FMath::Abs(SphereAxis0.Dot(FVec3(1,0,0))-1.0)) > SMALL_NUMBER) ? SphereAxis0.Cross(FVec3(1,0,0)) :
+	                          (FMath::Abs(FMath::Abs(SphereAxis0.Dot(FVec3(0,1,0))-1.0)) > SMALL_NUMBER) ? SphereAxis0.Cross(FVec3(0,1,0)) : SphereAxis0.Cross(FVec3(0,0,1)) ;
+	const FVec3 SphereAxis2 = SphereAxis0.Cross(SphereAxis1);
 
-	PositionDatas.UpdateMass(ConstraintIndex, IA0, IA1, InvM(0) + II0 + InvM(1) + II1, Dt);
+	// Using Connector1 for both conserves angular momentum and avoid having 
+	// too much torque applied onto the COM. But can only be used for limited constraints
+	const FVec3 ConstraintArm0 = ConnectorXs[1] - CurrentPs[0];
+	const FVec3 ConstraintArm1 = ConnectorXs[1] - CurrentPs[1];
+	
+	InitPositionConstraintDatas(0, SphereAxis0,  SphereDelta0, JointSettings.LinearRestitution, Dt,
+		JointSettings.LinearLimit, EJointMotionType::Limited, ConstraintArm0, ConstraintArm1);
+
+	// SphereAxis0 being the direction axis, the geometric error for the 2 other axis are 0
+	// We need to add these 2 constraints for a linear solver to avoid drifting away in
+	//  the other directions while solving. For a non linear solver since we are recomputing 
+	// the main direction at each time we don't need that
+	InitPositionConstraintDatas(1, SphereAxis1,  0.0, JointSettings.LinearRestitution, Dt,
+		JointSettings.LinearLimit, EJointMotionType::Limited, ConstraintArm0, ConstraintArm1);
+
+	InitPositionConstraintDatas(2, SphereAxis2,  0.0, JointSettings.LinearRestitution, Dt,
+		JointSettings.LinearLimit, EJointMotionType::Limited, ConstraintArm0, ConstraintArm1);
+}
+
+void FPBDJointCachedSolver::InitCylindricalPositionConstraint(
+	const FPBDJointSettings& JointSettings,
+	const FReal Dt,
+	const int32 AxisIndex)
+{
+	FVec3 PlaneAxis, RadialAxis0;
+	FReal PlaneDelta, RadialDelta0;
+	FPBDJointUtilities::GetCylindricalAxesDeltas(ConnectorRs[0], ConnectorXs[0], ConnectorXs[1],
+		AxisIndex, PlaneAxis, PlaneDelta, RadialAxis0, RadialDelta0);
+
+	const FVec3 RadialAxis1 = PlaneAxis.Cross(RadialAxis0);
+	const FReal RadialDelta1 =  (ConnectorXs[1] - ConnectorXs[0]).Dot(RadialAxis1);
+
+	const FVec3 ConstraintArm0 = ConnectorXs[1] - CurrentPs[0];
+	const FVec3 ConstraintArm1 = ConnectorXs[1] - CurrentPs[1];
+	
+	InitPositionConstraintDatas((AxisIndex+1)%3, RadialAxis0, RadialDelta0, JointSettings.LinearRestitution, Dt,
+		JointSettings.LinearLimit, EJointMotionType::Limited, ConstraintArm0, ConstraintArm1);
+
+	InitPositionConstraintDatas((AxisIndex+2)%3, RadialAxis1, RadialDelta1, JointSettings.LinearRestitution, Dt,
+	    JointSettings.LinearLimit, EJointMotionType::Limited, ConstraintArm0, ConstraintArm1);   
+}
+
+void FPBDJointCachedSolver::InitPlanarPositionConstraint(
+	const FPBDJointSettings& JointSettings,
+	const FReal Dt,
+	const int32 AxisIndex)
+{
+	FVec3 PlaneAxis;
+	FReal PlaneDelta;
+	FPBDJointUtilities::GetPlanarAxisDelta(ConnectorRs[0], ConnectorXs[0], ConnectorXs[1], AxisIndex, PlaneAxis, PlaneDelta);
+	
+	const FVec3 ConstraintArm0 = ConnectorXs[1] - CurrentPs[0];
+	const FVec3 ConstraintArm1 = ConnectorXs[1] - CurrentPs[1];
+	
+	InitPositionConstraintDatas(AxisIndex, PlaneAxis, PlaneDelta, JointSettings.LinearRestitution, Dt,
+		JointSettings.LinearLimit, EJointMotionType::Limited, ConstraintArm0, ConstraintArm1);
 }
 
 /** APPLY POSITION CONSTRAINT *****************************************************************************************/
@@ -707,6 +843,7 @@ FPBDJointUtilities::GetSoftTwistStiffness(SolverSettings, JointSettings),
 	const FVec3 Twist1 = ConnectorRs[1] * FJointConstants::TwistAxis();
 	const bool bDegenerate = (FVec3::DotProduct(Twist0, Twist1) < Chaos_Joint_DegenerateRotationLimit);
 
+	
 	// Apply twist constraint
 	// NOTE: Cannot calculate twist angle at 180degree swing
 	if (SolverSettings.bEnableTwistLimits)
diff --git a/Engine/Source/Runtime/Experimental/Chaos/Public/Chaos/Joint/PBDJointCachedSolverGaussSeidel.h b/Engine/Source/Runtime/Experimental/Chaos/Public/Chaos/Joint/PBDJointCachedSolverGaussSeidel.h
--- a/Engine/Source/Runtime/Experimental/Chaos/Public/Chaos/Joint/PBDJointCachedSolverGaussSeidel.h
+++ b/Engine/Source/Runtime/Experimental/Chaos/Public/Chaos/Joint/PBDJointCachedSolverGaussSeidel.h
@@ -294,6 +294,36 @@ struct FAxisConstraintDatas
 			const int32 ConstraintIndex,
 			const FReal Dt);
 		
+		void InitPositionConstraintDatas(
+			const int32 ConstraintIndex,
+			const FVec3& ConstraintAxis,
+			const FReal& ConstraintDelta,
+			const FReal ConstraintRestitution,
+			const FReal Dt,
+			const FReal ConstraintLimit,
+			const EJointMotionType JointType,
+			const FVec3& ConstraintArm0,
+			const FVec3& ConstraintArm1);
+
+		void InitLockedPositionConstraint(
+			const FPBDJointSettings& JointSettings,
+			const FReal Dt,
+			const TVec3<EJointMotionType>& LinearMotion);
+
+		void InitSphericalPositionConstraint(
+			const FPBDJointSettings& JointSettings,
+			const FReal Dt);
+		
+		void InitCylindricalPositionConstraint(
+			const FPBDJointSettings& JointSettings,
+			const FReal Dt,
+			const int32 AxisIndex);
+
+		void InitPlanarPositionConstraint(
+			const FPBDJointSettings& JointSettings,
+			const FReal Dt,
+			const int32 AxisIndex);
+		
 		/** Apply Position constraints */
 
 		void ApplyPositionConstraints(
