Chaos
- suspension fixes
- fix inactive hard stop constraints affecting unrelated bodies
- use implicit velocity in damping calculation (oversight from change to PBD)
- forces applied along surface normal, not suspension axis
- prevent springs applying negative forces

diff --git a/Engine/Source/Runtime/Experimental/Chaos/Private/Chaos/ChaosDebugDraw.cpp b/Engine/Source/Runtime/Experimental/Chaos/Private/Chaos/ChaosDebugDraw.cpp
--- a/Engine/Source/Runtime/Experimental/Chaos/Private/Chaos/ChaosDebugDraw.cpp
+++ b/Engine/Source/Runtime/Experimental/Chaos/Private/Chaos/ChaosDebugDraw.cpp
@@ -1348,6 +1348,12 @@ namespace Chaos
 			FConstGenericParticleHandle Particle = Constraints.GetConstrainedParticles(ConstraintIndex)[0];
 			const FPBDSuspensionSettings& ConstraintSettings = Constraints.GetSettings(ConstraintIndex);
 			const FPBDSuspensionResults& ConstraintResults = Constraints.GetResults(ConstraintIndex);
+			
+			if (!ConstraintSettings.Enabled)
+			{
+				return;
+			}
+			
 			const FVec3& PLocal = Constraints.GetConstraintPosition(ConstraintIndex);
 			const FRigidTransform3 ParticleTransform = FParticleUtilitiesPQ::GetActorWorldTransform(Particle);
 
@@ -1355,6 +1361,8 @@ namespace Chaos
 			const FVec3 AxisWorld = ParticleTransform.TransformVector(ConstraintSettings.Axis);
 			const FReal AxisLen = ConstraintResults.Length;
 			const FVec3 PushOutWorld = ConstraintResults.NetPushOut;
+			const FVec3 HardStopPushOutWorld = ConstraintResults.HardStopNetPushOut;
+			const FVec3 HardStopImpulseWorld = ConstraintResults.HardStopNetImpulse;
 
 			FDebugDrawQueue::GetInstance().DrawDebugLine(
 				SpaceTransform.TransformPosition(PWorld), 
@@ -1363,10 +1371,34 @@ namespace Chaos
 
 			if (Settings.PushOutScale > 0)
 			{
+				// Make the shorter line wider
+				const FReal SpringPushOutLen = PushOutWorld.Size();
+				const FReal HardStopPushOutLen = HardStopPushOutWorld.Size();
+				const FRealSingle SpringWidthScale = (SpringPushOutLen > HardStopPushOutLen) ? FRealSingle(1) : FRealSingle(2);
+				const FRealSingle HardStopWidthScale = (SpringPushOutLen > HardStopPushOutLen) ? FRealSingle(2) : FRealSingle(1);
+
+				const FColor HardStopPushOutColor = FColor(0, 200, 100);
+				const FColor SpringPushOutColor = FColor(0, 200, 0);
+
 				FDebugDrawQueue::GetInstance().DrawDebugLine(
 					SpaceTransform.TransformPosition(PWorld),
 					SpaceTransform.TransformPosition(PWorld + Settings.PushOutScale * PushOutWorld),
-					FColor::Blue, false, KINDA_SMALL_NUMBER, Settings.DrawPriority, Settings.LineThickness);
+					SpringPushOutColor, false, KINDA_SMALL_NUMBER, Settings.DrawPriority, SpringWidthScale * Settings.LineThickness);
+
+				FDebugDrawQueue::GetInstance().DrawDebugLine(
+					SpaceTransform.TransformPosition(PWorld),
+					SpaceTransform.TransformPosition(PWorld + Settings.PushOutScale * HardStopPushOutWorld),
+					HardStopPushOutColor, false, KINDA_SMALL_NUMBER, Settings.DrawPriority, HardStopWidthScale * Settings.LineThickness);
+			}
+
+			if (Settings.ImpulseScale > 0)
+			{
+				const FColor HardStopImpulseColor = FColor(200, 200, 0);
+
+				FDebugDrawQueue::GetInstance().DrawDebugLine(
+					SpaceTransform.TransformPosition(PWorld),
+					SpaceTransform.TransformPosition(PWorld + Settings.PushOutScale * HardStopImpulseWorld),
+					HardStopImpulseColor, false, KINDA_SMALL_NUMBER, Settings.DrawPriority, Settings.LineThickness * FRealSingle(0.5));
 			}
 		}
 
diff --git a/Engine/Source/Runtime/Experimental/Chaos/Private/Chaos/PBDSuspensionConstraints.cpp b/Engine/Source/Runtime/Experimental/Chaos/Private/Chaos/PBDSuspensionConstraints.cpp
--- a/Engine/Source/Runtime/Experimental/Chaos/Private/Chaos/PBDSuspensionConstraints.cpp
+++ b/Engine/Source/Runtime/Experimental/Chaos/Private/Chaos/PBDSuspensionConstraints.cpp
@@ -161,15 +161,14 @@ namespace Chaos
 				FVec3 AxisWorld = BodyQ.RotateVector(SuspensionCoMAxis);
 				FReal Distance = FVec3::DotProduct(WorldSpaceX - T, AxisWorld);
 
-				// if the suspension is NOT compressed right down to the hard-stop then we won't need a manifold, should be the case 99% time
-				if (Distance > Setting.MinLength)
-				{
-					return;
-				}
-
-				FReal WorldContactDeltaNormal = Setting.MinLength - Distance; // the distance we need to correct by
 				const FVec3 WorldArm = BodyQ.RotateVector(SuspensionCoMOffset);
 
+				// The hard-stop can only apply correction perpendicular to the surface
+				const FVec3& WorldContactNormal = Setting.Normal;
+
+				const FReal HardStopDistance = Setting.MinLength - Distance;
+				const FReal WorldContactDeltaNormal = FVec3::DotProduct(HardStopDistance * AxisWorld, WorldContactNormal);
+
 				// position the spoofed terrain at the same position at body0
 				FVec3 PosBody1 = T + Distance*AxisWorld;
 				Body1->SetP(PosBody1);
@@ -201,7 +200,7 @@ namespace Chaos
 					FSolverReal(0.1f),						// RestitutionVelocityThreshold,
 					FSolverVec3(WorldArm),					// RelativeContactPosition0,
 					FSolverVec3(0, 0, 0),					// RelativeContactPosition1,
-					FSolverVec3(AxisWorld),					// WorldContactNormal
+					FSolverVec3(WorldContactNormal),		// WorldContactNormal
 					FSolverVec3(0, 0, 0),					// WorldContactTangentU,
 					FSolverVec3(0, 0, 0),					// WorldContactTangentV,
 					FSolverReal(-WorldContactDeltaNormal),	// WorldContactDeltaNormal
@@ -222,7 +221,16 @@ namespace Chaos
 	{
 		for (int32 ConstraintIndex : SolverData.GetConstraintIndices(ContainerId))
 		{
+			if ((CollisionSolvers.Num() > 0) && (CollisionSolvers[ConstraintIndex] != nullptr) && (CollisionSolvers[ConstraintIndex]->NumManifoldPoints() > 0))
+			{
+				const FPBDCollisionSolver* CollisionSolver = CollisionSolvers[ConstraintIndex];
+				const FPBDCollisionSolverManifoldPoint& ManifoldPoint = CollisionSolver->GetManifoldPoint(0);
+				ConstraintResults[ConstraintIndex].HardStopNetPushOut = ManifoldPoint.NetPushOutNormal * ManifoldPoint.WorldContactNormal;
+				ConstraintResults[ConstraintIndex].HardStopNetImpulse = ManifoldPoint.NetImpulseNormal * ManifoldPoint.WorldContactNormal;
+			}
+
 			ConstraintSolverBodies[ConstraintIndex] = nullptr;
+			CollisionSolvers[ConstraintIndex]->SetSolverBodies(nullptr, nullptr);
 		}
 	}
 
@@ -231,12 +239,17 @@ namespace Chaos
 		if (bChaos_Suspension_Hardstop_Enabled)
 		{
 			// Suspension Hardstop
-			for (FPBDCollisionSolver* CollisionSolver : CollisionSolvers)
+			for (int32 ConstraintIndex : SolverData.GetConstraintIndices(ContainerId))
 			{
-				if (CollisionSolver->NumManifoldPoints() > 0)
+				const FPBDSuspensionSettings& Setting = ConstraintSettings[ConstraintIndex];
+				if (Setting.Enabled)
 				{
-					FReal MaxPushoutValue = FMath::Min(Chaos_Suspension_MaxPushout, Chaos_Suspension_MaxPushoutVelocity * Dt);
-					CollisionSolver->SolvePositionNoFriction(FSolverReal(Dt), FSolverReal(MaxPushoutValue));
+					FPBDCollisionSolver* CollisionSolver = CollisionSolvers[ConstraintIndex];
+					if ((CollisionSolver != nullptr) && (CollisionSolver->NumManifoldPoints() > 0))
+					{
+						FReal MaxPushoutValue = FMath::Min(Chaos_Suspension_MaxPushout, Chaos_Suspension_MaxPushoutVelocity * Dt);
+						CollisionSolver->SolvePositionNoFriction(FSolverReal(Dt), FSolverReal(MaxPushoutValue));
+					}
 				}
 			}
 		}
@@ -256,18 +269,24 @@ namespace Chaos
 
 	bool FPBDSuspensionConstraints::ApplyPhase2Serial(const FReal Dt, const int32 It, const int32 NumIts, FPBDIslandSolverData& SolverData)
 	{
-		bool bNeedsAnotherIteration = false;
-
 		if (bChaos_Suspension_Hardstop_Enabled && bChaos_Suspension_VelocitySolve)
 		{
 			// Suspension Hardstop
-			for (FPBDCollisionSolver* CollisionSolver : CollisionSolvers)
+			for (int32 ConstraintIndex : SolverData.GetConstraintIndices(ContainerId))
 			{
-				bNeedsAnotherIteration |= CollisionSolver->SolveVelocity(FSolverReal(Dt), false);
+				const FPBDSuspensionSettings& Setting = ConstraintSettings[ConstraintIndex];
+				if (Setting.Enabled)
+				{
+					FPBDCollisionSolver* CollisionSolver = CollisionSolvers[ConstraintIndex];
+					if ((CollisionSolver != nullptr) && (CollisionSolver->NumManifoldPoints() > 0))
+					{
+						CollisionSolver->SolveVelocity(FSolverReal(Dt), false);
+					}
+				}
 			}
 		}
 
-		return bNeedsAnotherIteration;
+		return true;
 	}
 
 	void FPBDSuspensionConstraints::ApplySingle(const FReal Dt, int32 ConstraintIndex)
@@ -288,31 +307,35 @@ namespace Chaos
 			const FVec3 SuspensionCoMOffset = Particle->RotationOfMass().UnrotateVector(SuspensionActorOffset - Particle->CenterOfMass());
 			const FVec3 SuspensionCoMAxis = Particle->RotationOfMass().UnrotateVector(Setting.Axis);
 
+			// @todo(chaos): use linearized error calculation
 			const FRotation3 BodyQ = Body.CorrectedQ();
 			const FVec3 BodyP = Body.CorrectedP();
 			const FVec3 WorldSpaceX = BodyQ.RotateVector(SuspensionCoMOffset) + BodyP;
 
 			FVec3 AxisWorld = BodyQ.RotateVector(SuspensionCoMAxis);
 
+			const FVec3& SurfaceNormal = Setting.Normal;
+
 			const float MPHToCmS = 100000.f / 2236.94185f;
 			const float SpeedThreshold = 10.0f * MPHToCmS;
 			const float FortyFiveDegreesThreshold = 0.707f;
 
-			if (AxisWorld.Z > FortyFiveDegreesThreshold)
-			{
-				if (Body.V().SquaredLength() < 1.0f)
-				{
-					AxisWorld = FVec3(0.f, 0.f, 1.f);
-				}
-				else
-				{
-					const FReal Speed = FMath::Abs(Body.V().Length());
-					if (Speed < SpeedThreshold)
-					{
-						AxisWorld = FMath::Lerp(FVec3(0.f, 0.f, 1.f), AxisWorld, Speed / SpeedThreshold);
-					}
-				}
-			}
+			// @todo(chaos): do we need this?
+			//if (AxisWorld.Z > FortyFiveDegreesThreshold)
+			//{
+			//	if (Body.V().SquaredLength() < 1.0f)
+			//	{
+			//		AxisWorld = FVec3(0.f, 0.f, 1.f);
+			//	}
+			//	else
+			//	{
+			//		const FReal Speed = FMath::Abs(Body.V().Length());
+			//		if (Speed < SpeedThreshold)
+			//		{
+			//			AxisWorld = FMath::Lerp(FVec3(0.f, 0.f, 1.f), AxisWorld, Speed / SpeedThreshold);
+			//		}
+			//	}
+			//}
 
 			FReal Distance = FVec3::DotProduct(WorldSpaceX - T, AxisWorld);
 			if (Distance >= Setting.MaxLength)
@@ -325,12 +348,14 @@ namespace Chaos
 			FVec3 DX = FVec3::ZeroVector;
 
 			// Require the velocity at the WorldSpaceX position - not the velocity of the particle origin
+			// NOTE: We are in the position solve phase and velocty is not updated. We must use the implicit 
+			// velocity in the damping calculation
+			// @todo(chaos): consider moving the damping term to the velocity solve phase
 			const FVec3 Diff = WorldSpaceX - BodyP;
-			FVec3 ArmVelocity = Body.V() - FVec3::CrossProduct(Diff, Body.W());
+			const FVec3 V = FVec3::CalculateVelocity(Body.X(), BodyP, Dt);
+			const FVec3 W = FRotation3::CalculateAngularVelocity(Body.R(), BodyQ, Dt);
+			FVec3 ArmVelocity = V - FVec3::CrossProduct(Diff, W);
 
-			// This constraint is causing considerable harm to the steering effect from the tires, using only the z component for damping
-			// makes this issue go away, rather than using DotProduct against the expected AxisWorld vector
-			FReal PointVelocityAlongAxis = FVec3::DotProduct(ArmVelocity, AxisWorld);
 
 			if (Distance < Setting.MinLength)
 			{
@@ -355,16 +380,26 @@ namespace Chaos
 				FReal DLambda = 0.f;
 				{
 					// #todo: Preload, better scaled spring damping like other suspension 0 -> 1 range
-					FReal SpringCompression = (Setting.MaxLength - Distance) /*+ Setting.SpringPreload*/;
-
-					FReal VelDt = PointVelocityAlongAxis;
+					const FReal AxisDotNormal = FVec3::DotProduct(AxisWorld, SurfaceNormal);
+					FReal SpringCompression = AxisDotNormal * (Setting.MaxLength - Distance) /*+ Setting.SpringPreload*/;
+					FReal SpringVelocity = FVec3::DotProduct(ArmVelocity, SurfaceNormal);
 
+					// @todo(chaos): this is not using the correct effective mass
 					const bool bAccelerationMode = false;
 					const FReal SpringMassScale = (bAccelerationMode) ? FReal(1) / Body.InvM() : FReal(1);
 					const FReal S = SpringMassScale * Setting.SpringStiffness * Dt * Dt;
 					const FReal D = SpringMassScale * Setting.SpringDamping * Dt;
-					DLambda = (S * SpringCompression - D * VelDt);
-					DX += DLambda * AxisWorld;
+					
+					// @todo(chaos): add missing XPBD term to get iteration-count independent behaviour
+					DLambda = (S * SpringCompression - D * SpringVelocity);
+
+					// Suspension springs cannot apply downward forces on the body
+					if (DLambda < 0)
+					{
+						DLambda = 0;
+					}
+
+					DX += DLambda * SurfaceNormal;
 				}
 			}
 
diff --git a/Engine/Source/Runtime/Experimental/Chaos/Public/Chaos/PBDSuspensionConstraintTypes.h b/Engine/Source/Runtime/Experimental/Chaos/Public/Chaos/PBDSuspensionConstraintTypes.h
--- a/Engine/Source/Runtime/Experimental/Chaos/Public/Chaos/PBDSuspensionConstraintTypes.h
+++ b/Engine/Source/Runtime/Experimental/Chaos/Public/Chaos/PBDSuspensionConstraintTypes.h
@@ -86,6 +86,8 @@ namespace Chaos
 		}
 
 		FVec3 NetPushOut;
+		FVec3 HardStopNetPushOut;
+		FVec3 HardStopNetImpulse;
 		FReal Length;
 	};
 
diff --git a/Engine/Source/Runtime/Experimental/ChaosSolverEngine/Private/Chaos/ChaosDebugDrawComponent.cpp b/Engine/Source/Runtime/Experimental/ChaosSolverEngine/Private/Chaos/ChaosDebugDrawComponent.cpp
--- a/Engine/Source/Runtime/Experimental/ChaosSolverEngine/Private/Chaos/ChaosDebugDrawComponent.cpp
+++ b/Engine/Source/Runtime/Experimental/ChaosSolverEngine/Private/Chaos/ChaosDebugDrawComponent.cpp
@@ -154,8 +154,11 @@ void DebugDrawChaos(AActor* DebugDrawActor, const TArray<Chaos::FLatentDrawComma
 				UE_VLOG(Command.TestBaseActor, LogChaos, Log, TEXT("%s"), *Command.Text);
 				break;
 			case FLatentDrawCommand::EDrawType::Circle:
-				//	UE_VLOG_CONE(Actor, LogChaos, Log, Command.Center, Command.ZAxis, 0.1f, PI, Command.Color, TEXT_EMPTY);
+			{
+				const FMatrix M = FRotationMatrix::MakeFromYZ(Command.YAxis, Command.ZAxis);
+				UE_VLOG_CIRCLE(Actor, LogChaos, Log, Command.Center, M.GetUnitAxis(EAxis::X), Command.Radius, Command.Color, TEXT_EMPTY);
 				break;
+			}
 			case FLatentDrawCommand::EDrawType::Capsule:
 			{
 				// VLOG Capsule uses the bottom end as the origin (though the variable is named Center)
