Fix crash on convex simplification filtering when input dimension is zero / the extreme points filter finds nothing to work with

diff --git a/Engine/Source/Runtime/Experimental/Chaos/Private/GeometryCollection/GeometryCollectionConvexUtility.cpp b/Engine/Source/Runtime/Experimental/Chaos/Private/GeometryCollection/GeometryCollectionConvexUtility.cpp
--- a/Engine/Source/Runtime/Experimental/Chaos/Private/GeometryCollection/GeometryCollectionConvexUtility.cpp
+++ b/Engine/Source/Runtime/Experimental/Chaos/Private/GeometryCollection/GeometryCollectionConvexUtility.cpp
@@ -94,7 +94,7 @@ typedef Chaos::TPlaneConcrete<Chaos::FReal, 3> FChaosPlane;
 // filter points s.t. they are spaced at least more than SimplificationDistanceThreshold apart (after starting with the 4 'extreme' points to ensure we cover a volume)
 void FilterHullPoints(const TArray<Chaos::FConvex::FVec3Type>& InPts, TArray<Chaos::FConvex::FVec3Type>& OutPts, double SimplificationDistanceThreshold)
 {
-	if (SimplificationDistanceThreshold > 0)
+	if (SimplificationDistanceThreshold > 0 && InPts.Num() > 0)
 	{
 		OutPts.Reset();
 
@@ -127,6 +127,10 @@ void FilterHullPoints(const TArray<Chaos::FConvex::FVec3Type>& InPts, TArray<Cha
 		for (int32 ExtremeIdx = 0; ExtremeIdx < ExtremePoints.Dimension + 1; ExtremeIdx++)
 		{
 			int32 ExtremePtIdx = ExtremePoints.Extreme[ExtremeIdx];
+			if (!InPts.IsValidIndex(ExtremePtIdx))
+			{
+				break;
+			}
 			Chaos::FVec3 ExtremePt = InPts[ExtremePtIdx];
 			OutPts.Add(ExtremePt);
 			Spatial.InsertPointUnsafe(ExtremePtIdx, ExtremePt);
