fix crash in geometry collection simulation when small convex turns into zero-vertex convex on regeneration with a new transform

diff --git a/Engine/Source/Runtime/Experimental/Chaos/Private/GeometryCollection/GeometryCollectionCollisionStructureManager.cpp b/Engine/Source/Runtime/Experimental/Chaos/Private/GeometryCollection/GeometryCollectionCollisionStructureManager.cpp
--- a/Engine/Source/Runtime/Experimental/Chaos/Private/GeometryCollection/GeometryCollectionCollisionStructureManager.cpp
+++ b/Engine/Source/Runtime/Experimental/Chaos/Private/GeometryCollection/GeometryCollectionCollisionStructureManager.cpp
@@ -269,12 +269,24 @@ FCollisionStructureManager::NewImplicitConvex(
 			}
 
 			Chaos::FReal Margin = (Chaos::FReal)(*ConvexGeometry)[Index]->BoundingBox().Extents().Min() * CollisionMarginFraction;
-			Chaos::FImplicitObject* Implicit = new Chaos::FConvex(ConvexVertices, Margin);
-			UpdateImplicitFlags(Implicit, CollisionType);
-			Implicits.Add(TUniquePtr<Chaos::FImplicitObject>(Implicit));
+			Chaos::FConvex* MarginConvex = new Chaos::FConvex(ConvexVertices, Margin);
+			if (MarginConvex->NumVertices() > 0)
+			{
+				Chaos::FImplicitObject* Implicit = MarginConvex;
+				UpdateImplicitFlags(Implicit, CollisionType);
+				Implicits.Add(TUniquePtr<Chaos::FImplicitObject>(Implicit));
+			}
+			else
+			{
+				delete MarginConvex;
+			}
 		}
 
-		if (Implicits.Num() == 1)
+		if (Implicits.Num() == 0)
+		{
+			return nullptr;
+		}
+		else if (Implicits.Num() == 1)
 		{
 			return Implicits[0].Release();
 		}
