Chaos : remove internal Epic wiki links in comments that end up in tooltips

diff --git a/Engine/Source/Runtime/Experimental/Chaos/Public/GeometryCollection/GeometryCollection.h b/Engine/Source/Runtime/Experimental/Chaos/Public/GeometryCollection/GeometryCollection.h
--- a/Engine/Source/Runtime/Experimental/Chaos/Public/GeometryCollection/GeometryCollection.h
+++ b/Engine/Source/Runtime/Experimental/Chaos/Public/GeometryCollection/GeometryCollection.h
@@ -24,9 +24,6 @@ namespace GeometryCollectionUV
 
 /**
 * FGeometryCollection (FTransformCollection)
-*  
-*    see : https://wiki.it.epicgames.net/display/~Brice.Criswell/Geometry+Collections
-*
 */
 class CHAOS_API FGeometryCollection : public FTransformCollection, 
 	public FGeometryCollectionConvexPropertiesInterface
diff --git a/Engine/Source/Runtime/Experimental/FieldSystem/Source/FieldSystemEngine/Public/Field/FieldSystemComponent.h b/Engine/Source/Runtime/Experimental/FieldSystem/Source/FieldSystemEngine/Public/Field/FieldSystemComponent.h
--- a/Engine/Source/Runtime/Experimental/FieldSystem/Source/FieldSystemEngine/Public/Field/FieldSystemComponent.h
+++ b/Engine/Source/Runtime/Experimental/FieldSystem/Source/FieldSystemEngine/Public/Field/FieldSystemComponent.h
@@ -192,8 +192,6 @@ public:
 	*    for examples of how to recreate variations of the above generic
 	*    fields using field networks
 	*
-	*    (https://wiki.it.epicgames.net/display/~Brice.Criswell/Fields)
-	*
 	*    @param Enabled Is this force enabled for evaluation.
 	*    @param Target Type of field supported by the solver.
 	*    @param MetaData Meta data used to assist in evaluation
@@ -218,8 +216,6 @@ public:
 	*    for examples of how to recreate variations of the above generic
 	*    fields using field networks
 	*
-	*    (https://wiki.it.epicgames.net/display/~Brice.Criswell/Fields)
-	*
 	*    @param Enabled Is this force enabled for evaluation.
 	*    @param Target Type of field supported by the solver.
 	*    @param MetaData Meta data used to assist in evaluation
@@ -252,8 +248,6 @@ public:
 	*    for examples of how to recreate variations of the above generic
 	*    fields using field networks
 	*
-	*    (https://wiki.it.epicgames.net/display/~Brice.Criswell/Fields)
-	*
 	*    @param Enabled Is this force enabled for evaluation.
 	*    @param Target Type of field supported by the solver.
 	*    @param MetaData Meta data used to assist in evaluation
