Chaos: destruction : make sure clusters have a GT shape representation on the gamethread for queries

diff --git a/Engine/Source/Runtime/Experimental/Chaos/Private/Chaos/PBDRigidClusteringAlgo.cpp b/Engine/Source/Runtime/Experimental/Chaos/Private/Chaos/PBDRigidClusteringAlgo.cpp
--- a/Engine/Source/Runtime/Experimental/Chaos/Private/Chaos/PBDRigidClusteringAlgo.cpp
+++ b/Engine/Source/Runtime/Experimental/Chaos/Private/Chaos/PBDRigidClusteringAlgo.cpp
@@ -392,7 +392,10 @@ namespace Chaos
 				OriginalPoints;
 		}
 
-		if (ProxyGeometry)
+		// ignore unions for now as we don't yet support deep copy of it
+		// on the GT they are only used by clusters that aggregate their children shapes ( see GeometryCollectionPhysicsProxy.cpp )
+		// by failing artificially this condition thmake sure we create a FImplicitObjectUnionClustered for this particle 
+		if (ProxyGeometry && ProxyGeometry->GetType() != ImplicitObjectType::Union)
 		{
 			const FVector Scale = Parameters.Scale;
 			auto DeepCopyImplicit = [&Scale](const TSharedPtr<Chaos::FImplicitObject, ESPMode::ThreadSafe>& ImplicitToCopy) -> TUniquePtr<Chaos::FImplicitObject>
diff --git a/Engine/Source/Runtime/Experimental/Chaos/Private/PhysicsProxy/GeometryCollectionPhysicsProxy.cpp b/Engine/Source/Runtime/Experimental/Chaos/Private/PhysicsProxy/GeometryCollectionPhysicsProxy.cpp
--- a/Engine/Source/Runtime/Experimental/Chaos/Private/PhysicsProxy/GeometryCollectionPhysicsProxy.cpp
+++ b/Engine/Source/Runtime/Experimental/Chaos/Private/PhysicsProxy/GeometryCollectionPhysicsProxy.cpp
@@ -68,6 +68,11 @@ FAutoConsoleVariableRef CVarEnabledNestedChildTransformUpdates(
 	bGeometryCollectionEnabledNestedChildTransformUpdates,
 	TEXT("Enable updates for driven, disabled, child bodies. Used for line trace results against geometry collections.[def: true]"));
 
+bool bGeometryCollectionAlwaysGenerateGTCollisionForClusters = true;
+FAutoConsoleVariableRef CVarGeometryCollectionAlwaysGenerateGTCollisionForClusters(
+	TEXT("p.GeometryCollection.AlwaysGenerateGTCollisionForClusters"),
+	bGeometryCollectionAlwaysGenerateGTCollisionForClusters,
+	TEXT("When enabled, always generate a game thread side collision for clusters.[def: true]"));
 
 DEFINE_LOG_CATEGORY_STATIC(UGCC_LOG, Error, All);
 
@@ -427,6 +432,7 @@ void FGeometryCollectionPhysicsProxy::Initialize(Chaos::FPBDRigidsEvolutionBase
 	GameThreadCollection.AddExternalAttribute<TUniquePtr<Chaos::FGeometryParticle>>(FGeometryCollection::ParticlesAttribute, FTransformCollection::TransformGroup, GTParticles);
 	
 
+	TArray<int32> ChildrenToCheckForParentFix;
 	if(ensure(NumTransforms == GameThreadCollection.Implicits.Num() && NumTransforms == GTParticles.Num())) // Implicits are in the transform group so this invariant should always hold
 	{
 		for(int32 Index = 0; Index < NumTransforms; ++Index)
@@ -443,11 +449,95 @@ void FGeometryCollectionPhysicsProxy::Initialize(Chaos::FPBDRigidsEvolutionBase
 			P->SetProxy(this);
 			P->SetGeometry(GameThreadCollection.Implicits[Index]);
 
-			// IMPORTANT: we need to set eth right spatial index because GT particle is static and PT particle is rigid
+			// this step is necessary for Phase 2 where we need to walk back the hierarchy from children to parent 
+			if (bGeometryCollectionAlwaysGenerateGTCollisionForClusters && GameThreadCollection.Children[Index].Num() == 0)
+			{
+				ChildrenToCheckForParentFix.Add(Index);
+			}
+			
+			// IMPORTANT: we need to set the right spatial index because GT particle is static and PT particle is rigid
 			// this is causing a mismatch when using the separate acceleration structures optimization which can cause crashes when destroying the particle while async tracing 
 			// todo(chaos) we should eventually refactor this code to use rigid particles on the GT side for geometry collection  
 			P->SetSpatialIdx(Chaos::FSpatialAccelerationIdx{ 0,1 });
+		}
+
+		if (bGeometryCollectionAlwaysGenerateGTCollisionForClusters)
+		{
+			// second phase: fixing parent geometries
+			// @todo(chaos) this could certainly be done ahead at generation time rather than runtime
+			TSet<int32> ParentToPotentiallyFix;
+			while (ChildrenToCheckForParentFix.Num())
+			{
+				// step 1 : find parents
+				for(const int32 ChildIndex: ChildrenToCheckForParentFix)
+				{
+					const int32 ParentIndex = GameThreadCollection.Parent[ChildIndex];
+					if (ParentIndex != INDEX_NONE)
+					{
+						ParentToPotentiallyFix.Add(ParentIndex);
+					}
+				}
+
+				// step 2: fix the parent if necessary
+				for (const int32 ParentToFixIndex: ParentToPotentiallyFix)
+				{
+					if (GameThreadCollection.Implicits[ParentToFixIndex] == nullptr)
+					{
+						const Chaos::FRigidTransform3 ParentShapeTransform =  GameThreadCollection.MassToLocal[ParentToFixIndex] * GameThreadCollection.Transform[ParentToFixIndex];
+				
+						// Make a union of the children geometry
+						TArray<TUniquePtr<Chaos::FImplicitObject>> ChildImplicits;
+						for (const int32& ChildIndex: GameThreadCollection.Children[ParentToFixIndex])
+						{
+							using FImplicitObjectTransformed = Chaos::TImplicitObjectTransformed<Chaos::FReal, 3>;
 
+							Chaos::FGeometryParticle* ChildParticle = GTParticles[ChildIndex].Get();
+							const FGeometryDynamicCollection::FSharedImplicit& ChildImplicit = GameThreadCollection.Implicits[ChildIndex];
+							if (ChildImplicit)
+							{
+								const Chaos::FRigidTransform3 ChildShapeTransform =  GameThreadCollection.MassToLocal[ChildIndex] * GameThreadCollection.Transform[ChildIndex];
+								const Chaos::FRigidTransform3 RelativeShapeTransform = ChildShapeTransform.GetRelativeTransform(ParentShapeTransform);
+						
+								// assumption that we only have can only have one level of union for any child
+								if (ChildImplicit->GetType() == Chaos::ImplicitObjectType::Union)
+								{
+									if (Chaos::FImplicitObjectUnion* Union = ChildImplicit->GetObject<Chaos::FImplicitObjectUnion>())
+									{
+										for (const TUniquePtr<Chaos::FImplicitObject>& ImplicitObject : Union->GetObjects())
+										{
+											TUniquePtr<Chaos::FImplicitObject> CopiedChildImplicit = ImplicitObject->DeepCopy();
+											FImplicitObjectTransformed* TransformedChildImplicit = new FImplicitObjectTransformed(MoveTemp(CopiedChildImplicit), RelativeShapeTransform);  
+											ChildImplicits.Add(TUniquePtr<FImplicitObjectTransformed>(TransformedChildImplicit));
+										}
+									}
+								}
+								else
+								{
+									TUniquePtr<Chaos::FImplicitObject> CopiedChildImplicit = GameThreadCollection.Implicits[ChildIndex]->DeepCopy();
+									FImplicitObjectTransformed* TransformedChildImplicit = new FImplicitObjectTransformed(MoveTemp(CopiedChildImplicit), RelativeShapeTransform);  
+									ChildImplicits.Add(TUniquePtr<FImplicitObjectTransformed>(TransformedChildImplicit));
+								}
+							}
+						}
+						if (ChildImplicits.Num() > 0)
+						{
+							Chaos::FImplicitObject* UnionImplicit = new Chaos::FImplicitObjectUnion(MoveTemp(ChildImplicits));
+							GameThreadCollection.Implicits[ParentToFixIndex] = FGeometryDynamicCollection::FSharedImplicit(UnionImplicit);
+						}
+						GTParticles[ParentToFixIndex]->SetGeometry(GameThreadCollection.Implicits[ParentToFixIndex]);
+					}
+				}
+
+				// step 3 : make the parent the new child to go up the hierarchy and continue the fixing
+				ChildrenToCheckForParentFix = ParentToPotentiallyFix.Array(); 
+				ParentToPotentiallyFix.Reset();
+			}
+		}
+		
+		// Phase 3 : finalization of shapes
+		for(int32 Index = 0; Index < NumTransforms; ++Index)
+		{
+			Chaos::FGeometryParticle* P = GTParticles[Index].Get();
 			const Chaos::FShapesArray& Shapes = P->ShapesArray();
 			const int32 NumShapes = Shapes.Num();
 			for(int32 ShapeIndex = 0; ShapeIndex < NumShapes; ++ShapeIndex)
