Add a failsafe for not enough groups compared to constraints. This should function but may result in non-optimal parallelization.

diff --git a/Engine/Source/Runtime/Experimental/Chaos/Private/Chaos/Island/IslandManager.cpp b/Engine/Source/Runtime/Experimental/Chaos/Private/Chaos/Island/IslandManager.cpp
--- a/Engine/Source/Runtime/Experimental/Chaos/Private/Chaos/Island/IslandManager.cpp
+++ b/Engine/Source/Runtime/Experimental/Chaos/Private/Chaos/Island/IslandManager.cpp
@@ -726,7 +726,8 @@ void FPBDIslandManager::BuildGroups(const int32 NumContainers)
 			IslandSolver->SetGroupIndex(GroupIndex);
 			GroupOffset += IslandSolver->NumConstraints();
 
-			if(GroupOffset > GroupSize)
+			// @todo: In theory we should never need the second check	
+			if(GroupOffset > GroupSize && (GroupIndex + 1) < NumGroups)
 			{
 				GroupIndex++;
 				GroupOffset = 0;
