Use a different BVH for Triangle Mesh

diff --git a/Engine/Source/Programs/HeadlessChaos/Private/HeadlessChaosApp.cpp b/Engine/Source/Programs/HeadlessChaos/Private/HeadlessChaosApp.cpp
--- a/Engine/Source/Programs/HeadlessChaos/Private/HeadlessChaosApp.cpp
+++ b/Engine/Source/Programs/HeadlessChaos/Private/HeadlessChaosApp.cpp
@@ -175,6 +175,7 @@ TEST(RaycastTests, Raycast) {
 	ChaosTest::CapsuleRaycastFastLargeDistance();
 	ChaosTest::CapsuleRaycastMissWithEndPointOnBounds();
 	ChaosTest::TriangleRaycast();
+	ChaosTest::TriangleRaycastDenegerated();
 	ChaosTest::BoxRaycast();
 	ChaosTest::VectorizedAABBRaycast();
 	ChaosTest::ScaledRaycast();
diff --git a/Engine/Source/Programs/HeadlessChaos/Private/HeadlessChaosTestRaycast.cpp b/Engine/Source/Programs/HeadlessChaos/Private/HeadlessChaosTestRaycast.cpp
--- a/Engine/Source/Programs/HeadlessChaos/Private/HeadlessChaosTestRaycast.cpp
+++ b/Engine/Source/Programs/HeadlessChaos/Private/HeadlessChaosTestRaycast.cpp
@@ -431,6 +431,28 @@ namespace ChaosTest
 		EXPECT_FLOAT_EQ(Position.Z, 1);
 	}
 
+	void TriangleRaycastDenegerated()
+	{
+		FReal Time;
+		FVec3 Position;
+		FVec3 Normal;
+		TArray<uint16> DummyMaterials;
+		int32 FaceIndex;
+
+		FTriangleMeshImplicitObject::ParticlesType Particles;
+		Particles.AddParticles(3);
+		Particles.X(0) = FVec3(1, 1, 1);
+		Particles.X(1) = FVec3(1, 1, 2);
+		Particles.X(2) = FVec3(1, 1, 3);
+		TArray<TVec3<int32>> Indices;
+		Indices.Emplace(0, 1, 2);
+		FTriangleMeshImplicitObject Tri(MoveTemp(Particles), MoveTemp(Indices), MoveTemp(DummyMaterials));
+
+		//simple into the triangle
+		bool bHit = Tri.Raycast(FVec3(1, -1, 1.5), FVec3(0, 1, 0), 2, 0, Time, Position, Normal, FaceIndex);
+		EXPECT_FALSE(bHit);
+	}
+
 	void BoxRaycast()
 	{
 		FReal Time;
diff --git a/Engine/Source/Programs/HeadlessChaos/Public/HeadlessChaosTestRaycast.h b/Engine/Source/Programs/HeadlessChaos/Public/HeadlessChaosTestRaycast.h
--- a/Engine/Source/Programs/HeadlessChaos/Public/HeadlessChaosTestRaycast.h
+++ b/Engine/Source/Programs/HeadlessChaos/Public/HeadlessChaosTestRaycast.h
@@ -23,6 +23,8 @@ namespace ChaosTest {
 
 	void TriangleRaycast();
 
+	void TriangleRaycastDenegerated();
+
 	void BoxRaycast();
 
 	void VectorizedAABBRaycast();
diff --git a/Engine/Source/Runtime/Core/Public/UObject/UE5MainStreamObjectVersion.h b/Engine/Source/Runtime/Core/Public/UObject/UE5MainStreamObjectVersion.h
--- a/Engine/Source/Runtime/Core/Public/UObject/UE5MainStreamObjectVersion.h
+++ b/Engine/Source/Runtime/Core/Public/UObject/UE5MainStreamObjectVersion.h
@@ -194,6 +194,9 @@ struct CORE_API FUE5MainStreamObjectVersion
 		// Fixed default value of volumetric cloud to be exact match with main view, more expenssive but we let user choosing how to lower the quality.
 		VolumetricCloudReflectionSampleCountDefaultUpdate,
 
+		// Use special BVH for TriangleMesh, instead of the AABBTree
+		UseTriangleMeshBVH,
+
 		// -----<new versions can be added above this line>-------------------------------------------------
 		VersionPlusOne,
 		LatestVersion = VersionPlusOne - 1
diff --git a/Engine/Source/Runtime/Experimental/Chaos/Private/Chaos/HeightField.cpp b/Engine/Source/Runtime/Experimental/Chaos/Private/Chaos/HeightField.cpp
--- a/Engine/Source/Runtime/Experimental/Chaos/Private/Chaos/HeightField.cpp
+++ b/Engine/Source/Runtime/Experimental/Chaos/Private/Chaos/HeightField.cpp
@@ -51,61 +51,6 @@ namespace Chaos
 			Sweep
 		};
 
-		/**
-		* Ray / triangle  intersection
-		* this provides a double sided test
-		* note : this method assumes that the triangle formed by A,B and C is well formed
-		*/
-		FORCEINLINE_DEBUGGABLE bool RayTriangleIntersection(
-			const FVec3& RayStart, const FVec3& RayDir, FReal RayLength,
-			const FVec3& A, const FVec3& B, const FVec3& C,
-			FReal& OutT, FVec3& OutN
-		)
-		{
-			const FVec3 AB = B - A; // edge 1
-			const FVec3 AC = C - A; // edge 2
-			const FVec3 Normal = FVec3::CrossProduct(AB, AC);
-			const FVec3 NegRayDir = -RayDir;
-
-			FReal Den = FVec3::DotProduct(NegRayDir, Normal);
-			if (FMath::Abs(Den) < SMALL_NUMBER)
-			{
-				// ray is parallel or away to the triangle plane it is a miss
-				return false;
-			}
-
-			const FReal InvDen = (FReal)1 / Den;
-
-			// let's compute the time to intersection
-			const FVec3 RayToA = RayStart - A;
-			const FReal Time = FVec3::DotProduct(RayToA, Normal) * InvDen;
-			if (Time < (FReal)0 || Time > RayLength)
-			{
-				return false;
-			}
-
-			// now compute baricentric coordinates
-			const FVec3 RayToACrossNegDir = FVec3::CrossProduct(NegRayDir, RayToA);
-			constexpr FReal Epsilon = SMALL_NUMBER;
-			const FReal UU = FVec3::DotProduct(AC, RayToACrossNegDir) * InvDen;
-			if (UU < -Epsilon || UU >(1 + Epsilon))
-			{
-				return false; // outside of the triangle
-			}
-			const FReal VV = -FVec3::DotProduct(AB, RayToACrossNegDir) * InvDen;
-			if (VV < -Epsilon || (VV + UU) >(1 + Epsilon))
-			{
-				return false; // outside of the triangle
-			}
-
-			// point is within the triangle, let's compute 
-			OutT = Time;
-			OutN = Normal.GetSafeNormal();
-			OutN *= FMath::Sign(Den);
-			return true;
-		}
-
-
 		FORCEINLINE_DEBUGGABLE bool VisitRaycast(int32 Payload, FReal& CurrentLength)
 		{
 			const int32 SubX = Payload % (GeomData->NumCols - 1);
diff --git a/Engine/Source/Runtime/Experimental/Chaos/Private/Chaos/TriangleMeshImplicitObject.cpp b/Engine/Source/Runtime/Experimental/Chaos/Private/Chaos/TriangleMeshImplicitObject.cpp
--- a/Engine/Source/Runtime/Experimental/Chaos/Private/Chaos/TriangleMeshImplicitObject.cpp
+++ b/Engine/Source/Runtime/Experimental/Chaos/Private/Chaos/TriangleMeshImplicitObject.cpp
@@ -12,7 +12,6 @@
 #include "Chaos/GeometryQueries.h"
 #include "Chaos/Utilities.h"
 
-//PRAGMA_DISABLE_OPTIMIZATION
 
 namespace Chaos
 {
@@ -135,20 +134,64 @@ struct FTriangleMeshRaycastVisitor
 		return nullptr;
 	}
 
-	template <ERaycastType SQType>
-	bool Visit(int32 TriIdx, FQueryFastData& CurData)
+	/**
+	 * find the intersection of a ray and the triangle at index TriIdx
+	 * @return true if the search should continue otherwise false to stop the search through the mesh 
+	 */
+	bool VisitRaycast(TSpatialVisitorData<int32> TriIdx, FQueryFastData& CurData)
+	{
+		constexpr FReal Epsilon = 1e-4f;
+
+		const int32 FaceIndex = TriIdx.Payload;
+		const FVec3& A = Particles.X(Elements[FaceIndex][0]);
+		const FVec3& B = Particles.X(Elements[FaceIndex][1]);
+		const FVec3& C = Particles.X(Elements[FaceIndex][2]);
+
+		// Note: the math here needs to match FTriangleMeshImplicitObject::GetFaceNormal
+		// @todo(chaos) we should really preprocess the face and remove the degenerated ones to avoid paying this runtime cost
+		const FVec3 AB = B - A;
+		const FVec3 AC = C - A;
+		FVec3 TriNormal = FVec3::CrossProduct(AB, AC);
+
+		if (bCullsBackFaceRaycast)
+		{
+			const bool bBackFace = (FVec3::DotProduct(Dir, TriNormal) > 0.0f);
+			if (bBackFace)
+			{
+				// skip this traingle and continue the visit
+				return true;
+			}
+	}
+
+		FVec3 HitNormal;
+		FReal HitTime;
+		if (RayTriangleIntersection(StartPoint, Dir, CurData.CurrentLength, A, B, C, HitTime, HitNormal))
+		{
+			OutPosition = StartPoint + (Dir * HitTime);
+			OutNormal = HitNormal;
+			OutTime = HitTime;
+			OutFaceIndex = FaceIndex;
+			CurData.SetLength(HitTime);	//prevent future rays from going any farther
+		}
+		// continue the visit
+		return true;
+	}
+	
+	bool VisitSweep(TSpatialVisitorData<int32> TriIdx, FQueryFastData& CurData)
 	{
 		constexpr FReal Epsilon = 1e-4f;
 		constexpr FReal Epsilon2 = Epsilon * Epsilon;
-		const FReal Thickness2 = SQType == ERaycastType::Sweep ? Thickness * Thickness : 0;
+		
+		const FReal Thickness2 = Thickness * Thickness;
 		FReal MinTime = 0;	//no need to initialize, but fixes warning
 
 		const FReal R = Thickness + Epsilon;
 		const FReal R2 = R * R;
 
-		const FVec3& A = Particles.X(Elements[TriIdx][0]);
-		const FVec3& B = Particles.X(Elements[TriIdx][1]);
-		const FVec3& C = Particles.X(Elements[TriIdx][2]);
+		const int32 FaceIndex = TriIdx.Payload;		
+		const FVec3& A = Particles.X(Elements[FaceIndex][0]);
+		const FVec3& B = Particles.X(Elements[FaceIndex][1]);
+		const FVec3& C = Particles.X(Elements[FaceIndex][2]);
 
 		// Note: the math here needs to match FTriangleMeshImplicitObject::GetFaceNormal
 		const FVec3 AB = B - A;
@@ -187,9 +230,7 @@ struct FTriangleMeshRaycastVisitor
 				if (DistToTriangle2 <= R2)
 				{
 					OutTime = 0;
-					OutFaceIndex = TriIdx;
-					//OutPosition = IntersectionPosition;
-					//OutNormal = RaycastNormal;	//We use the plane normal even when hitting triangle edges. This is to deal with triangles that approximate a single flat surface.
+					OutFaceIndex = FaceIndex;
 					return false; //no one will beat Time == 0
 				}
 			}
@@ -200,7 +241,7 @@ struct FTriangleMeshRaycastVisitor
 				bTriangleIntersects = DistToTriangle2 <= Epsilon2;	//raycast gave us the intersection point so sphere radius is already accounted for
 			}
 
-			if (SQType == ERaycastType::Sweep && !bTriangleIntersects)
+			if (!bTriangleIntersects)
 			{
 				//sphere is not immediately touching the triangle, but it could start intersecting the perimeter as it sweeps by
 				FVec3 BorderPositions[3];
@@ -267,7 +308,7 @@ struct FTriangleMeshRaycastVisitor
 					OutNormal = RaycastNormal;	//We use the plane normal even when hitting triangle edges. This is to deal with triangles that approximate a single flat surface.
 					OutTime = Time;
 					CurData.SetLength(Time);	//prevent future rays from going any farther
-					OutFaceIndex = TriIdx;
+					OutFaceIndex = FaceIndex;
 				}
 			}
 		}
@@ -275,16 +316,6 @@ struct FTriangleMeshRaycastVisitor
 		return true;
 	}
 
-	bool VisitRaycast(TSpatialVisitorData<int32> TriIdx, FQueryFastData& CurData)
-	{
-		return Visit<ERaycastType::Raycast>(TriIdx.Payload, CurData);
-	}
-
-	bool VisitSweep(TSpatialVisitorData<int32> TriIdx, FQueryFastData& CurData)
-	{
-		return Visit<ERaycastType::Sweep>(TriIdx.Payload, CurData);
-	}
-
 	bool VisitOverlap(TSpatialVisitorData<int32> TriIdx)
 	{
 		check(false);
@@ -301,8 +332,41 @@ struct FTriangleMeshRaycastVisitor
 	FVec3 OutNormal;
 	int32 OutFaceIndex;
 	bool bCullsBackFaceRaycast;
+	TArray<int32> RaycastTriangles;
+};
+
+struct FTriangleMeshOverlapVisitor
+{
+	FTriangleMeshOverlapVisitor(TArray<int32>& InResults) : CollectedResults(InResults) {}
+	bool VisitOverlap(int32 Instance)
+	{
+		CollectedResults.Add(Instance);
+		return true;
+	}
+	bool VisitSweep(int32 Instance, FQueryFastData& CurData)
+	{
+		check(false);
+		return true;
+	}
+	bool VisitRaycast(int32 Instance, FQueryFastData& CurData)
+	{
+		check(false);
+		return true;
+	}
+
+	TArray<int32>& CollectedResults;
 };
 
+TArray<int32> FTrimeshBVH::FindAllIntersections (const FAABB3& Intersection) const
+{
+	TArray<int32> Results;
+	FTriangleMeshOverlapVisitor Visitor(Results);
+	Overlap(Intersection, Visitor);
+
+	return Results;
+}
+
+
 FReal FTriangleMeshImplicitObject::PhiWithNormal(const FVec3& x, FVec3& Normal) const
 {
 	TSphere<FReal, 3> TestSphere(x, 0.0f);
@@ -320,11 +384,12 @@ bool FTriangleMeshImplicitObject::RaycastImp(const TArray<TVector<IdxType, 3>>&
 
 	if (Thickness > 0)
 	{
-		BVH.Sweep(StartPoint, Dir, Length, FVec3(Thickness), SQVisitor);
+		FVec3 QueryHalfExtents = FVec3(Thickness) * 0.5f;
+		FastBVH.Sweep(StartPoint, Dir, Length, QueryHalfExtents, SQVisitor);
 	}
 	else
 	{
-		BVH.Raycast(StartPoint, Dir, Length, SQVisitor);
+		FastBVH.Raycast(StartPoint, Dir, Length, SQVisitor);
 	}
 
 	if (SQVisitor.OutTime <= Length)
@@ -421,7 +486,7 @@ bool FTriangleMeshImplicitObject::ContactManifoldImp(const GeomType& QueryGeom,
 		FReal LocalContactPhi = FLT_MAX;
 		FVec3 LocalContactLocation, LocalContactNormal;
 
-		const TArray<int32> PotentialIntersections = BVH.FindAllIntersections(QueryBounds);
+		const TArray<int32> PotentialIntersections = FastBVH.FindAllIntersections(QueryBounds);
 
 		// MakeTriangleHelper get rid of the scale wrapper if necessary as MakeTriangle will try to infer properties from it 
 		FPBDCollisionConstraint Constraint = MakeTriangleHelper(WorldScaleGeom);
@@ -520,7 +585,7 @@ bool FTriangleMeshImplicitObject::GJKContactPointImp(const QueryGeomType& QueryG
 		FReal LocalContactPhi = FLT_MAX;
 		FVec3 LocalContactLocation, LocalContactNormal;
 
-		const TArray<int32> PotentialIntersections = BVH.FindAllIntersections(QueryBounds);
+		const TArray<int32> PotentialIntersections = FastBVH.FindAllIntersections(QueryBounds);
 
 		for (int32 TriIdx : PotentialIntersections)
 		{
@@ -651,7 +716,7 @@ bool FTriangleMeshImplicitObject::OverlapImp(const TArray<TVec3<IdxType>>& Eleme
 {
 	FAABB3 QueryBounds(Point, Point);
 	QueryBounds.Thicken(Thickness);
-	const TArray<int32> PotentialIntersections = BVH.FindAllIntersections(QueryBounds);
+	const TArray<int32> PotentialIntersections = FastBVH.FindAllIntersections(QueryBounds);
 
 	const FReal Epsilon = 1e-4f;
 	//ensure(Thickness > Epsilon);	//There's no hope for this to work unless thickness is large (really a sphere overlap test)
@@ -698,7 +763,7 @@ bool FTriangleMeshImplicitObject::Overlap(const FVec3& Point, const FReal Thickn
 
 void FTriangleMeshImplicitObject::VisitTriangles(const FAABB3& QueryBounds, const TFunction<void(const FTriangle& Triangle)>& Visitor) const
 {
-	const TArray<int32> PotentialIntersections = BVH.FindAllIntersections(QueryBounds);
+	const TArray<int32> PotentialIntersections = FastBVH.FindAllIntersections(QueryBounds);
 
 	auto TriangleProducer = [&](const auto& Elements)
 	{
@@ -759,7 +824,7 @@ bool FTriangleMeshImplicitObject::OverlapGeomImp(const QueryGeomType& QueryGeom,
 	QueryBounds.ThickenSymmetrically(FVec3(Thickness));
 	QueryBounds.ScaleWithNegative(InvTriMeshScale);
 
-	const TArray<int32> PotentialIntersections = BVH.FindAllIntersections(QueryBounds);
+	const TArray<int32> PotentialIntersections = FastBVH.FindAllIntersections(QueryBounds);
 
 	if (OutMTD)
 	{
@@ -1068,7 +1133,7 @@ bool FTriangleMeshImplicitObject::SweepGeomImp(const QueryGeomType& QueryGeom, c
 		const FVec3 InvTriMeshScale = SafeInvScale(TriMeshScale);
 		const FVec3 StartPoint = QueryBounds.Center() * InvTriMeshScale + StartTM.GetLocation();
 		const FVec3 Inflation = QueryBounds.Extents() * InvTriMeshScale.GetAbs() * 0.5 + FVec3(Thickness);
-		BVH.template Sweep<VisitorType>(StartPoint,Dir,Length,Inflation,SQVisitor);
+		FastBVH.template Sweep<VisitorType>(StartPoint, Dir, Length, Inflation, SQVisitor);
 
 		if(SQVisitor.OutTime <= Length)
 		{
@@ -1140,7 +1205,7 @@ int32 FTriangleMeshImplicitObject::FindMostOpposingFace(const TArray<TVec3<IdxTy
 
 	FAABB3 QueryBounds(Position - FVec3(SearchDist), Position + FVec3(SearchDist));
 
-	const TArray<int32> PotentialIntersections = BVH.FindAllIntersections(QueryBounds);
+	const TArray<int32> PotentialIntersections = FastBVH.FindAllIntersections(QueryBounds);
 	const FReal Epsilon = 1e-4f;
 
 	FReal MostOpposingDot = TNumericLimits<FReal>::Max();
@@ -1217,7 +1282,7 @@ TUniquePtr<FTriangleMeshImplicitObject> FTriangleMeshImplicitObject::CopySlowImp
 		ExternalVertexIndexMapCopy = MakeUnique<TArray<int32>>(*ExternalVertexIndexMap.Get());
 	}
 
-	return TUniquePtr<FTriangleMeshImplicitObject>(new FTriangleMeshImplicitObject(MoveTemp(ParticlesCopy), MoveTemp(ElementsCopy), MoveTemp(MaterialIndicesCopy), BVH, MoveTemp(ExternalFaceIndexMapCopy), MoveTemp(ExternalVertexIndexMapCopy), bCullsBackFaceRaycast));
+	return TUniquePtr<FTriangleMeshImplicitObject>(new FTriangleMeshImplicitObject(MoveTemp(ParticlesCopy), MoveTemp(ElementsCopy), MoveTemp(MaterialIndicesCopy), MoveTemp(ExternalFaceIndexMapCopy), MoveTemp(ExternalVertexIndexMapCopy), bCullsBackFaceRaycast));
 }
 
 TUniquePtr<FTriangleMeshImplicitObject> FTriangleMeshImplicitObject::CopySlow() const
@@ -1324,7 +1389,7 @@ const FTrimeshIndexBuffer& FTriangleMeshImplicitObject::Elements() const
 }
 
 template <typename IdxType>
-void FTriangleMeshImplicitObject::RebuildBVImp(const TArray<TVec3<IdxType>>& Elements)
+void FTriangleMeshImplicitObject::RebuildBVImp(const TArray<TVec3<IdxType>>& Elements, BVHType& TreeBVH)
 {
 	const int32 NumTris = Elements.Num();
 	TArray<FBvEntry<sizeof(IdxType) == sizeof(FTrimeshIndexBuffer::LargeIdxType)>> BVEntries;
@@ -1334,21 +1399,24 @@ void FTriangleMeshImplicitObject::RebuildBVImp(const TArray<TVec3<IdxType>>& Ele
 	{
 		BVEntries.Add({this, Tri});
 	}
-	BVH.Reinitialize(BVEntries);
+	TreeBVH.Reinitialize(BVEntries);
 }
 
 FTriangleMeshImplicitObject::~FTriangleMeshImplicitObject() = default;
 
-void Chaos::FTriangleMeshImplicitObject::RebuildBV()
+
+void FTriangleMeshImplicitObject::RebuildFastBVH()
 {
+	BVHType TreeBVH;
 	if (MElements.RequiresLargeIndices())
 	{
-		RebuildBVImp(MElements.GetLargeIndexBuffer());
+		RebuildBVImp(MElements.GetLargeIndexBuffer(), TreeBVH);
 	}
 	else
 	{
-		RebuildBVImp(MElements.GetSmallIndexBuffer());
+		RebuildBVImp(MElements.GetSmallIndexBuffer(), TreeBVH);
 	}
+	RebuildFastBVHFromTree(TreeBVH);
 }
 
 void FTriangleMeshImplicitObject::UpdateVertices(const TArray<FVector>& NewPositions)
@@ -1371,8 +1439,136 @@ void FTriangleMeshImplicitObject::UpdateVertices(const TArray<FVector>& NewPosit
 		}
 	}
 
-	RebuildBV();
+	RebuildFastBVH();
 }
+	
+void FTriangleMeshImplicitObject::RebuildFastBVHFromTree(const BVHType& TreeBVH)
+{
+	using NodeType = TAABBTreeNode<FRealSingle>;
+	using LeafType = TAABBTreeLeafArray<int32, /*bComputeBounds=*/ false, FRealSingle>;
+	const TArray<NodeType>& Nodes = TreeBVH.GetNodes();
+	const TArray<LeafType>& Leaves = TreeBVH.GetLeaves();
+
+	FastBVH.Nodes.Reset();
+	FastBVH.Faces.Reset();
+	FastBVH.FaceBounds.Reset();
+
+	// since we do skip leaf nodes, we need to handle the case where we have only one node that will be a leaf by default
+	if (Nodes.Num() == 1)
+	{
+		const NodeType& RootNode = Nodes[0];
+		ensure(RootNode.bLeaf);
+		ensure(RootNode.ChildrenNodes[1] == INDEX_NONE);
+
+		// the leaf index is stored in the first  ChildrenNodes for leaf type nodes
+		const int32 LeafIndex = RootNode.ChildrenNodes[0];
+		const LeafType& Leaf = Leaves[LeafIndex];
+
+		// make the node
+		FTrimeshBVH::FNode& NewNode = FastBVH.Nodes.Emplace_GetRef();
+		NewNode.Children[0].ChildOrFaceIndex = 0;
+		NewNode.Children[0].FaceCount = Leaf.Elems.Num();
+		NewNode.Children[0].Bounds = { BoundingBox().Min(), BoundingBox().Max() };
+		for (const TPayloadBoundsElement<int32, FRealSingle>& LeafPayload: Leaf.Elems)
+		{
+			FastBVH.FaceBounds.Add(LeafPayload.Bounds);
+			FastBVH.Faces.Add(LeafPayload.Payload);
+		}
+		return;
+	}
+	
+	// map leaf type nodes index (in Nodes array) to leaf index (in Leaves array) 
+	TMap<int32, int32> ChildIndexToLeafIndexMap;
+	for (int32 NodeIndex = 0; NodeIndex < Nodes.Num(); ++NodeIndex)
+	{
+		const NodeType& Node = Nodes[NodeIndex];
+		if (Node.bLeaf)
+		{
+			// the leaf index is stored in the first  ChildrenNodes for leaf type nodes
+			const int32 LeafIndex =  Node.ChildrenNodes[0];
+			ensure(Leaves.IsValidIndex(LeafIndex));
+			ChildIndexToLeafIndexMap.Emplace(NodeIndex, LeafIndex);
+		}
+	}
+
+	// let's walk the tree and build the optimized data
+	TArray<int32> NodeIndexStack;
+	NodeIndexStack.Push(0);
+
+	// used to remap indices at the end of the process
+	TMap<int32, int32> BVHToFastBVHNodeIndexMap;
+	
+	while (NodeIndexStack.Num())
+	{
+		const int32 NodeIndex = NodeIndexStack.Pop();
+		const NodeType& Node = Nodes[NodeIndex];
+
+		// leaf nodes are being trimmed and their info will be compacted in parents
+		if (!Node.bLeaf)
+		{
+			BVHToFastBVHNodeIndexMap.Emplace(NodeIndex, FastBVH.Nodes.Num());
+			FTrimeshBVH::FNode& NewNode = FastBVH.Nodes.Emplace_GetRef();
+
+			// go backward to simulate a depth first walk to keep the same order of the original tree 
+			// for consistent results to return the same triangle when a query hit right in on a share vertex
+			for (int32 ChildIndex=1; ChildIndex>=0; --ChildIndex)
+			{
+				// common infos
+				FTrimeshBVH::FChildData& ChildData = NewNode.Children[ChildIndex];
+				ChildData.Bounds = Node.ChildrenBounds[ChildIndex];
+				const int32 ChildNodeIndex = Node.ChildrenNodes[ChildIndex];
+				// index in the original BVH space, remapping is done at the end of the process
+				// this may be overwritten if the child is a leaf
+				ChildData.ChildOrFaceIndex = ChildNodeIndex;
+
+				if (Nodes.IsValidIndex(ChildNodeIndex))
+				{
+					// let's pull the face (leaves) data for leaf child nodes
+					const int32* LeafIndex = ChildIndexToLeafIndexMap.Find(ChildNodeIndex);
+					if (LeafIndex)
+					{
+						const LeafType& Leaf = Leaves[*LeafIndex];
 
+						// store face range in the node 
+						ChildData.ChildOrFaceIndex = FastBVH.Faces.Num();
+						ChildData.FaceCount = Leaf.Elems.Num();
+						check(ChildData.FaceCount > 0);
+
+						TMap<int32, int32> VertexReuse;
+						// copy indices in the linear face array
+						for (const auto& LeafPayload: Leaf.Elems)
+						{
+							FastBVH.FaceBounds.Add(LeafPayload.Bounds);
+							FastBVH.Faces.Add(LeafPayload.Payload);
+						}
+					}
+					// push for further processing
+					NodeIndexStack.Push(ChildNodeIndex);
+				}
+			}
+		}
+	}
 
+	// remap child node indices from original BVH node array space to fast BVH node array space
+	for (int32 NodeIndex = 0; NodeIndex < FastBVH.Nodes.Num(); ++NodeIndex)
+	{
+		FTrimeshBVH::FNode& Node = FastBVH.Nodes[NodeIndex];
+		for (int32 ChildIndex=0; ChildIndex<2; ++ChildIndex)
+		{
+			FTrimeshBVH::FChildData& ChildData = Node.Children[ChildIndex];
+			if (!ChildData.HasFaces())
+			{
+				const int32 ChildNodeIndex = ChildData.ChildOrFaceIndex;
+				const int32* FixedChildNodeIndex = BVHToFastBVHNodeIndexMap.Find(ChildNodeIndex);
+				if (ensure(FixedChildNodeIndex))
+				{
+					ChildData.ChildOrFaceIndex = *FixedChildNodeIndex;
+				}
+			}
+		}
+	}
+	
 }
+	
+}
+
diff --git a/Engine/Source/Runtime/Experimental/Chaos/Public/Chaos/AABBTree.h b/Engine/Source/Runtime/Experimental/Chaos/Public/Chaos/AABBTree.h
--- a/Engine/Source/Runtime/Experimental/Chaos/Public/Chaos/AABBTree.h
+++ b/Engine/Source/Runtime/Experimental/Chaos/Public/Chaos/AABBTree.h
@@ -2174,6 +2174,9 @@ public:
 		}
 	}
 
+	const TArray<TAABBTreeNode<T>>& GetNodes() const { return Nodes; }
+	const TArray<TLeafType>& GetLeaves() const { return Leaves; }
+
 private:
 
 	using FElement = TPayloadBoundsElement<TPayloadType, T>;
diff --git a/Engine/Source/Runtime/Experimental/Chaos/Public/Chaos/Triangle.h b/Engine/Source/Runtime/Experimental/Chaos/Public/Chaos/Triangle.h
--- a/Engine/Source/Runtime/Experimental/Chaos/Public/Chaos/Triangle.h
+++ b/Engine/Source/Runtime/Experimental/Chaos/Public/Chaos/Triangle.h
@@ -258,6 +258,60 @@ namespace Chaos
 		return Ar;
 	}
 
+	/**
+	* Ray / triangle  intersection
+	* this provides a double sided test
+	* note : this method assumes that the triangle formed by A,B and C is well formed
+	*/
+	FORCEINLINE_DEBUGGABLE bool RayTriangleIntersection(
+		const FVec3& RayStart, const FVec3& RayDir, FReal RayLength,
+		const FVec3& A, const FVec3& B, const FVec3& C,
+		FReal& OutT, FVec3& OutN
+	)
+	{
+		const FVec3 AB = B - A; // edge 1
+		const FVec3 AC = C - A; // edge 2
+		const FVec3 Normal = FVec3::CrossProduct(AB, AC);
+		const FVec3 NegRayDir = -RayDir;
+
+		FReal Den = FVec3::DotProduct(NegRayDir, Normal);
+		if (FMath::Abs(Den) < SMALL_NUMBER)
+		{
+			// ray is parallel or away to the triangle plane it is a miss
+			return false;
+		}
+
+		const FReal InvDen = (FReal)1 / Den;
+
+		// let's compute the time to intersection
+		const FVec3 RayToA = RayStart - A;
+		const FReal Time = FVec3::DotProduct(RayToA, Normal) * InvDen;
+		if (Time < (FReal)0 || Time > RayLength)
+		{
+			return false;
+		}
+
+		// now compute baricentric coordinates
+		const FVec3 RayToACrossNegDir = FVec3::CrossProduct(NegRayDir, RayToA);
+		constexpr FReal Epsilon = SMALL_NUMBER;
+		const FReal UU = FVec3::DotProduct(AC, RayToACrossNegDir) * InvDen;
+		if (UU < -Epsilon || UU >(1 + Epsilon))
+		{
+			return false; // outside of the triangle
+		}
+		const FReal VV = -FVec3::DotProduct(AB, RayToACrossNegDir) * InvDen;
+		if (VV < -Epsilon || (VV + UU) >(1 + Epsilon))
+		{
+			return false; // outside of the triangle
+		}
+
+		// point is within the triangle, let's compute 
+		OutT = Time;
+		OutN = Normal.GetSafeNormal();
+		OutN *= FMath::Sign(Den);
+		return true;
+	}
+
 	template<typename T> using TTriangle = FTriangle;
 	
 	template<typename T>
diff --git a/Engine/Source/Runtime/Experimental/Chaos/Public/Chaos/TriangleMeshImplicitObject.h b/Engine/Source/Runtime/Experimental/Chaos/Public/Chaos/TriangleMeshImplicitObject.h
--- a/Engine/Source/Runtime/Experimental/Chaos/Public/Chaos/TriangleMeshImplicitObject.h
+++ b/Engine/Source/Runtime/Experimental/Chaos/Public/Chaos/TriangleMeshImplicitObject.h
@@ -29,6 +29,228 @@ namespace Chaos
 	struct FMTDInfo;
 	class FContactPoint;
 
+	struct CHAOS_API FTrimeshBVH
+	{
+		using FAABBType = TAABB<FRealSingle, 3>;
+
+		enum class EVisitorResult
+		{
+			Stop = 0,
+			Continue,
+		};
+
+		enum class EFilterResult
+		{
+			Skip = 0,
+			Keep,
+		};
+		
+		struct FChildData
+		{
+			FChildData()
+				: ChildOrFaceIndex(INDEX_NONE)
+				, FaceCount(0)
+			{}
+			FAABBType Bounds;
+			int32 ChildOrFaceIndex;
+			int32 FaceCount;
+
+			bool HasFaces() const { return (FaceCount > 0); }
+
+			void Serialize(FArchive& Ar)
+			{
+				TBox<FRealSingle, 3>::SerializeAsAABB(Ar, Bounds);
+
+				Ar << ChildOrFaceIndex;
+				Ar << FaceCount;
+			}
+		};
+		
+		struct FNode
+		{
+			FNode() {}
+
+			void Serialize(FArchive& Ar)
+			{
+				Children[0].Serialize(Ar);
+				Children[1].Serialize(Ar);
+			}
+			FChildData Children[2];
+		};
+
+		template <typename SQVisitor>
+		FORCEINLINE_DEBUGGABLE void Raycast(const FVec3& Start, const FVec3& Dir, const FReal Length, SQVisitor& Visitor) const
+		{
+			FQueryFastData QueryFastData(Dir, Length);
+			const auto BoundsFilter = [&QueryFastData, &Start](const FAABBType& Bounds) -> EFilterResult
+			{
+				FReal HitTime;
+				FVec3 HitPosition;
+
+				bool bCulled = true;
+				const bool bHit = Bounds.RaycastFast(Start, QueryFastData.Dir, QueryFastData.InvDir, QueryFastData.bParallel, QueryFastData.CurrentLength, QueryFastData.InvCurrentLength, HitTime, HitPosition);
+				if (bHit)
+				{
+					// avoid going down teh nodes that are further than the current hit distance
+					bCulled = (QueryFastData.CurrentLength < HitTime);
+				}
+				
+				return (!bCulled) ? EFilterResult::Keep: EFilterResult::Skip;
+			};
+
+			const auto FaceVisitor = [&Visitor, &QueryFastData](int32 FaceIndex)
+			{
+				const bool bContinueVisiting = Visitor.VisitRaycast(FaceIndex, QueryFastData);
+				return bContinueVisiting? EVisitorResult::Continue: EVisitorResult::Stop;
+			};
+
+			VisitTree(BoundsFilter, FaceVisitor);
+		}
+
+		template <typename SQVisitor>
+		FORCEINLINE_DEBUGGABLE void Sweep(const FVec3& Start, const FVec3& Dir, const FReal Length, const FVec3& QueryHalfExtents, SQVisitor& Visitor) const
+		{
+			FQueryFastData QueryFastData(Dir, Length);
+			const auto BoundsFilter = [&QueryFastData, &Start, &QueryHalfExtents](const FAABBType& Bounds) -> EFilterResult
+			{
+				FReal HitTime;
+				FVec3 HitPosition;
+
+				bool bCulled = true;
+				const FAABB3 SweepBounds(Bounds.Min() - QueryHalfExtents, Bounds.Max() + QueryHalfExtents);
+				const bool bHit = SweepBounds.RaycastFast(Start, QueryFastData.Dir, QueryFastData.InvDir, QueryFastData.bParallel, QueryFastData.CurrentLength, QueryFastData.InvCurrentLength, HitTime, HitPosition);
+				if (bHit)
+				{
+					// avoid going down the nodes that are further than the current hit distance
+					bCulled = (QueryFastData.CurrentLength < HitTime);
+				}
+
+				return (!bCulled) ? EFilterResult::Keep : EFilterResult::Skip;
+			};
+
+			const auto FaceVisitor = [&Visitor, &QueryFastData](int32 FaceIndex)
+			{
+				const bool bContinueVisiting = Visitor.VisitSweep(FaceIndex, QueryFastData);
+				return bContinueVisiting ? EVisitorResult::Continue : EVisitorResult::Stop;
+			};
+
+			VisitTree(BoundsFilter, FaceVisitor);
+		}
+
+		template <typename SQVisitor>
+		FORCEINLINE_DEBUGGABLE void Overlap(const FAABB3& AABB, SQVisitor& Visitor) const
+		{
+			const auto BoundsFilter = [&AABB](const FAABBType& Bounds) -> EFilterResult
+			{
+				const bool bHit = Bounds.Intersects(AABB);
+				return bHit ? EFilterResult::Keep : EFilterResult::Skip;
+			};
+
+			const auto FaceVisitor = [&Visitor](int32 FaceIndex)
+			{
+				const bool bContinueVisiting = Visitor.VisitOverlap(FaceIndex);
+				return bContinueVisiting ? EVisitorResult::Continue : EVisitorResult::Stop;
+			};
+
+			VisitTree(BoundsFilter, FaceVisitor);
+		}
+
+		TArray<int32> FindAllIntersections(const FAABB3& Intersection) const;
+
+
+		template <typename BoundsFilterType, typename FaceVisitorType>
+		FORCEINLINE_DEBUGGABLE EVisitorResult VisitFaces(int32 StartIndex, int32 IndexCount, BoundsFilterType& BoundsFilter, FaceVisitorType& FaceVisitor) const
+		{
+			const int32 EndIndex = (StartIndex + IndexCount);  
+			for (int32 FaceIndex = StartIndex; FaceIndex < EndIndex; ++FaceIndex)
+			{
+				if (BoundsFilter(FaceBounds[FaceIndex]) == EFilterResult::Keep)
+				{
+					if (FaceVisitor(Faces[FaceIndex]) == EVisitorResult::Stop)
+					{
+						return EVisitorResult::Stop;
+					}
+				}
+			}
+			return EVisitorResult::Continue;
+		}
+		
+		template <typename BoundsFilterType, typename FaceVisitorType>
+		FORCEINLINE_DEBUGGABLE void VisitTree(BoundsFilterType& BoundsFilter, FaceVisitorType& FaceVisitor) const
+		{
+			if (Nodes.Num() == 0)
+			{
+				return;
+			}
+			TArray<int32> NodeIndexStack;
+			NodeIndexStack.Reserve(16);
+			NodeIndexStack.Push(0);
+			while (NodeIndexStack.Num())
+			{
+				const int32 NodeIndex = NodeIndexStack.Pop(false);
+				check(Nodes.IsValidIndex(NodeIndex));
+				const FNode& Node = Nodes[NodeIndex];
+
+				for (int32 ChildIndex = 0; ChildIndex < 2; ++ChildIndex)
+				{
+					const FChildData& ChildData = Node.Children[ChildIndex]; 
+					if (ChildData.ChildOrFaceIndex != INDEX_NONE)
+					{
+						if (BoundsFilter(ChildData.Bounds) == EFilterResult::Keep)
+						{
+							if (ChildData.HasFaces())
+							{
+								if (EVisitorResult::Stop == VisitFaces(ChildData.ChildOrFaceIndex, ChildData.FaceCount, BoundsFilter, FaceVisitor))
+								{
+									return;
+								}
+							}
+							else
+							{
+								NodeIndexStack.Push(ChildData.ChildOrFaceIndex);
+							}
+						}
+					}
+				}
+			}
+		}
+
+		void Serialize(FChaosArchive& Ar)
+		{
+			Ar << Nodes;
+			Ar << FaceBounds;
+			Ar << Faces;
+		}
+		
+		TArray<FNode> Nodes;
+		TArray<FAABBType> FaceBounds;
+		TArray<int32> Faces;
+	};
+
+	FORCEINLINE_DEBUGGABLE FChaosArchive& operator<<(FChaosArchive& Ar, FTrimeshBVH::FChildData& ChildData)
+	{
+		ChildData.Serialize(Ar);
+		return Ar;
+	}
+	
+	FORCEINLINE_DEBUGGABLE FChaosArchive& operator<<(FChaosArchive& Ar, FTrimeshBVH::FNode& Node)
+	{
+		Node.Serialize(Ar);
+		return Ar;
+	}
+
+	FORCEINLINE_DEBUGGABLE FChaosArchive& operator<<(FChaosArchive& Ar, FTrimeshBVH::FAABBType& Bounds)
+	{
+		TBox<FRealSingle, 3>::SerializeAsAABB(Ar, Bounds);
+		return Ar;
+	}
+
+	FORCEINLINE_DEBUGGABLE FChaosArchive& operator<<(FChaosArchive& Ar, FTrimeshBVH& TrimeshBVH)
+	{
+		TrimeshBVH.Serialize(Ar);
+		return Ar;
+	}
+	
 	class CHAOS_API FTrimeshIndexBuffer
 	{
 	public:
@@ -126,7 +348,7 @@ namespace Chaos
 		bool bRequiresLargeIndices;
 	};
 
-	FORCEINLINE FArchive& operator<<(FArchive& Ar, FTrimeshIndexBuffer& Buffer)
+	FORCEINLINE_DEBUGGABLE FArchive& operator<<(FArchive& Ar, FTrimeshIndexBuffer& Buffer)
 	{
 		Buffer.Serialize(Ar);
 		return Ar;
@@ -191,7 +413,7 @@ namespace Chaos
 				}
 			}
 			
-			RebuildBV();
+			RebuildFastBVH();
 		}
 
 		FTriangleMeshImplicitObject(const FTriangleMeshImplicitObject& Other) = delete;
@@ -287,18 +509,31 @@ namespace Chaos
 			if (Ar.CustomVer(FExternalPhysicsCustomObjectVersion::GUID) < FExternalPhysicsCustomObjectVersion::TrimeshSerializesBV)
 			{
 				// Should now only hit when loading older trimeshes
-				RebuildBV();
+				RebuildFastBVH();
 			}
 			else if (Ar.CustomVer(FExternalPhysicsCustomObjectVersion::GUID) < FExternalPhysicsCustomObjectVersion::TrimeshSerializesAABBTree)
 			{
 				TBoundingVolume<int32> Dummy;
 				Ar << Dummy;
-				RebuildBV();
+				RebuildFastBVH();
 			}
-			else
+			else if (Ar.CustomVer(FExternalPhysicsCustomObjectVersion::GUID) < FUE5MainStreamObjectVersion::UseTriangleMeshBVH)
 			{
 				// Serialize acceleration
+				BVHType BVH;
 				Ar << BVH;
+				if (BVH.GetNodes().IsEmpty())
+				{
+					RebuildFastBVH();
+				}
+				else
+				{
+					RebuildFastBVHFromTree(BVH);
+				}
+			}
+			else
+			{
+				Ar << FastBVH;
 			}
 
 			if (Ar.CustomVer(FExternalPhysicsCustomObjectVersion::GUID) >= FExternalPhysicsCustomObjectVersion::AddTrimeshMaterialIndices)
@@ -395,12 +630,13 @@ namespace Chaos
 		*/
 		void FindOverlappingTriangles(const FAABB3& QueryBounds, TArray<int32>& OutTriangleIndices) const
 		{
-			OutTriangleIndices = BVH.FindAllIntersections(QueryBounds);
+			OutTriangleIndices = FastBVH.FindAllIntersections(QueryBounds);
 		}
 
 	private:
-
-		void RebuildBV();
+		using BVHType = TAABBTree<int32, TAABBTreeLeafArray<int32, /*bComputeBounds=*/ false, FRealSingle>, /*bMutable=*/false, FRealSingle>;
+		void RebuildFastBVHFromTree(const BVHType& BVH);
+		void RebuildFastBVH();
 
 		ParticlesType MParticles;
 		FTrimeshIndexBuffer MElements;
@@ -410,8 +646,6 @@ namespace Chaos
 		TUniquePtr<TArray<int32>> ExternalVertexIndexMap;
 		bool bCullsBackFaceRaycast;
 
-		using BVHType = TAABBTree<int32, TAABBTreeLeafArray<int32, /*bComputeBounds=*/false, FRealSingle>, /*bMutable=*/false, FRealSingle>;
-
 		// Initialising constructor privately declared for use in CopySlow to copy the underlying BVH
 		template <typename IdxType>
 		FTriangleMeshImplicitObject(ParticlesType&& Particles, TArray<TVec3<IdxType>>&& Elements, TArray<uint16>&& InMaterialIndices, const BVHType& InBvhToCopy, TUniquePtr<TArray<int32>>&& InExternalFaceIndexMap = nullptr, TUniquePtr<TArray<int32>>&& InExternalVertexIndexMap = nullptr, const bool bInCullsBackFaceRaycast = false)
@@ -443,7 +677,7 @@ namespace Chaos
 				}
 			}
 			
-			BVH.CopyFrom(InBvhToCopy);
+			RebuildFastBVHFromTree(InBvhToCopy);
 		}
 
 		template<typename InStorageType, typename InRealType>
@@ -491,7 +725,7 @@ namespace Chaos
 			}
 		};
 
-		BVHType BVH;
+		FTrimeshBVH FastBVH;
 
 		template<typename Geom, typename IdxType>
 		friend struct FTriangleMeshSweepVisitor;
@@ -524,7 +758,7 @@ namespace Chaos
 		int32 FindMostOpposingFace(const TArray<TVec3<IdxType>>& Elements, const FVec3& Position, const FVec3& UnitDir, int32 HintFaceIndex, FReal SearchDist) const;
 
 		template <typename IdxType>
-		void RebuildBVImp(const TArray<TVec3<IdxType>>& Elements);
+		void RebuildBVImp(const TArray<TVec3<IdxType>>& Elements, BVHType& BVH);
 
 		template <typename IdxType>
 		TUniquePtr<FTriangleMeshImplicitObject> CopySlowImpl(const TArray < TVector<IdxType, 3>>& InElements) const;
diff --git a/Engine/Source/Runtime/Experimental/ChaosCore/Public/Chaos/Core.h b/Engine/Source/Runtime/Experimental/ChaosCore/Public/Chaos/Core.h
--- a/Engine/Source/Runtime/Experimental/ChaosCore/Public/Chaos/Core.h
+++ b/Engine/Source/Runtime/Experimental/ChaosCore/Public/Chaos/Core.h
@@ -42,6 +42,6 @@ namespace Chaos
 	// NOTE: if you get a merge conflict on the GUID, you must replace it with a new GUID - do not accept the source or target
 	// or you will likely get DDC version conflicts resulting in crashes during load.
 	// Core version string for Chaos data. Any DDC builder dependent on Chaos for serialization should depend on this version
-	static const FString ChaosVersionString = TEXT("34C4ADF6-BD70-4A04-A811-E9A24C643F32");
+	static const FString ChaosVersionString = TEXT("0C23CAD8-19DF-40B2-9CF1-4ECA57EB7BD4");
 
 }
