Fix cooking issue with mesh

diff --git a/Engine/Source/Runtime/Experimental/Chaos/Private/Chaos/TriangleMeshImplicitObject.cpp b/Engine/Source/Runtime/Experimental/Chaos/Private/Chaos/TriangleMeshImplicitObject.cpp
--- a/Engine/Source/Runtime/Experimental/Chaos/Private/Chaos/TriangleMeshImplicitObject.cpp
+++ b/Engine/Source/Runtime/Experimental/Chaos/Private/Chaos/TriangleMeshImplicitObject.cpp
@@ -1458,7 +1458,6 @@ void FTriangleMeshImplicitObject::RebuildFastBVHFromTree(const BVHType& TreeBVH)
 	{
 		const NodeType& RootNode = Nodes[0];
 		ensure(RootNode.bLeaf);
-		ensure(RootNode.ChildrenNodes[1] == INDEX_NONE);
 
 		// the leaf index is stored in the first  ChildrenNodes for leaf type nodes
 		const int32 LeafIndex = RootNode.ChildrenNodes[0];
diff --git a/Engine/Source/Runtime/Experimental/Chaos/Public/Chaos/TriangleMeshImplicitObject.h b/Engine/Source/Runtime/Experimental/Chaos/Public/Chaos/TriangleMeshImplicitObject.h
--- a/Engine/Source/Runtime/Experimental/Chaos/Public/Chaos/TriangleMeshImplicitObject.h
+++ b/Engine/Source/Runtime/Experimental/Chaos/Public/Chaos/TriangleMeshImplicitObject.h
@@ -494,6 +494,7 @@ namespace Chaos
 		void SerializeImp(FChaosArchive& Ar)
 		{
 			Ar.UsingCustomVersion(FExternalPhysicsCustomObjectVersion::GUID);
+			Ar.UsingCustomVersion(FUE5MainStreamObjectVersion::GUID);
 
 			FImplicitObject::SerializeImp(Ar);
 			Ar << MParticles;
@@ -517,7 +518,7 @@ namespace Chaos
 				Ar << Dummy;
 				RebuildFastBVH();
 			}
-			else if (Ar.CustomVer(FExternalPhysicsCustomObjectVersion::GUID) < FUE5MainStreamObjectVersion::UseTriangleMeshBVH)
+			else if (Ar.CustomVer(FUE5MainStreamObjectVersion::GUID) < FUE5MainStreamObjectVersion::UseTriangleMeshBVH)
 			{
 				// Serialize acceleration
 				BVHType BVH;
diff --git a/Engine/Source/Runtime/Experimental/ChaosCore/Public/Chaos/Core.h b/Engine/Source/Runtime/Experimental/ChaosCore/Public/Chaos/Core.h
--- a/Engine/Source/Runtime/Experimental/ChaosCore/Public/Chaos/Core.h
+++ b/Engine/Source/Runtime/Experimental/ChaosCore/Public/Chaos/Core.h
@@ -42,6 +42,6 @@ namespace Chaos
 	// NOTE: if you get a merge conflict on the GUID, you must replace it with a new GUID - do not accept the source or target
 	// or you will likely get DDC version conflicts resulting in crashes during load.
 	// Core version string for Chaos data. Any DDC builder dependent on Chaos for serialization should depend on this version
-	static const FString ChaosVersionString = TEXT("0C23CAD8-19DF-40B2-9CF1-4ECA57EB7BD4");
+	static const FString ChaosVersionString = TEXT("D89F8E66-8E2B-4C04-83B8-FB8761F99CF4");
 
 }
