Chaos - fix static analysis (null support in solver bodies)

diff --git a/Engine/Source/Runtime/Experimental/Chaos/Private/Chaos/Collision/PBDCollisionSolver.h b/Engine/Source/Runtime/Experimental/Chaos/Private/Chaos/Collision/PBDCollisionSolver.h
--- a/Engine/Source/Runtime/Experimental/Chaos/Private/Chaos/Collision/PBDCollisionSolver.h
+++ b/Engine/Source/Runtime/Experimental/Chaos/Private/Chaos/Collision/PBDCollisionSolver.h
@@ -207,8 +207,9 @@ namespace Chaos
 
 		void SetSolverBodies(FSolverBody* SolverBody0, FSolverBody* SolverBody1)
 		{
-			State.SolverBodies[0] = *SolverBody0;
-			State.SolverBodies[1] = *SolverBody1;
+			// @todo(chaos): change FConstraintSolverBody to take a pointer because it supports null
+			State.SolverBodies[0] = (SolverBody0 != nullptr) ? FConstraintSolverBody(*SolverBody0) : FConstraintSolverBody();
+			State.SolverBodies[1] = (SolverBody1 != nullptr) ? FConstraintSolverBody(*SolverBody1) : FConstraintSolverBody();
 		}
 
 		void ResetSolverBodies()
