CCD update to ignore disabled constraints

diff --git a/Engine/Source/Runtime/Experimental/Chaos/Private/Chaos/CCDUtilities.cpp b/Engine/Source/Runtime/Experimental/Chaos/Private/Chaos/CCDUtilities.cpp
--- a/Engine/Source/Runtime/Experimental/Chaos/Private/Chaos/CCDUtilities.cpp
+++ b/Engine/Source/Runtime/Experimental/Chaos/Private/Chaos/CCDUtilities.cpp
@@ -103,6 +103,11 @@ namespace Chaos
 		bool bNeedCCDSolve = false;
 		for (FPBDCollisionConstraint* Constraint : SweptConstraints)
 		{   
+			if (!Constraint->IsEnabled())
+			{
+				continue;
+			}
+
 			// Create CCDParticle for all dynamic particles affected by swept constraints (UseCCD() could be either true or false). For static or kinematic particles, this pointer remains to be nullptr.
 			FCCDParticle* CCDParticlePair[2] = {nullptr, nullptr};
 			bool IsDynamic[2] = {false, false};
@@ -569,6 +574,11 @@ namespace Chaos
 		// @todo(chaos): These SweptConstraints might contain non-CCD particles and those non-CCD particles might collide with other non-CCD particles, which are modeled in normal collision constraints. Those normal collision constraints might need to be updated as well.
 		for (FPBDCollisionConstraint* SweptConstraint : SweptConstraints)
 		{
+			if (!SweptConstraint->IsEnabled())
+			{
+				continue;
+			}
+
 			const FConstGenericParticleHandle P0 = FConstGenericParticleHandle(SweptConstraint->GetParticle0());
 			const FConstGenericParticleHandle P1 = FConstGenericParticleHandle(SweptConstraint->GetParticle1());
 			SweptConstraint->ResetManifold();
