Chaos - fix unitialized ObjectType enum that can cause unit test to fail

diff --git a/Engine/Source/Runtime/Experimental/Chaos/Public/GeometryCollection/GeometryCollectionSimulationCoreTypes.h b/Engine/Source/Runtime/Experimental/Chaos/Public/GeometryCollection/GeometryCollectionSimulationCoreTypes.h
--- a/Engine/Source/Runtime/Experimental/Chaos/Public/GeometryCollection/GeometryCollectionSimulationCoreTypes.h
+++ b/Engine/Source/Runtime/Experimental/Chaos/Public/GeometryCollection/GeometryCollectionSimulationCoreTypes.h
@@ -186,6 +186,7 @@ struct FSimulationParameters
 		, CacheBeginTime(0.0f)
 		, ReverseCacheBeginTime(0.0f)
 		, bClearCache(false)
+		, ObjectType(EObjectStateTypeEnum::Chaos_NONE)
 		, bGenerateBreakingData(false)
 		, bGenerateCollisionData(false)
 		, bGenerateTrailingData(false)
