Chaos - fix potential Nan when using float in the collision solver and contact velocities are getting really small

diff --git a/Engine/Source/Runtime/Experimental/Chaos/Private/Chaos/Collision/PBDCollisionSolver.h b/Engine/Source/Runtime/Experimental/Chaos/Private/Chaos/Collision/PBDCollisionSolver.h
--- a/Engine/Source/Runtime/Experimental/Chaos/Private/Chaos/Collision/PBDCollisionSolver.h
+++ b/Engine/Source/Runtime/Experimental/Chaos/Private/Chaos/Collision/PBDCollisionSolver.h
@@ -584,7 +584,7 @@ namespace Chaos
 			const FSolverReal MaxImpulseTangent = FMath::Max(FSolverReal(0), DynamicFriction * (ManifoldPoint.NetImpulseNormal + ImpulseNormal + ManifoldPoint.NetPushOutNormal / Dt));
 			const FSolverReal MaxImpulseTangentSq = FMath::Square(MaxImpulseTangent);
 			const FSolverReal ImpulseTangentSq = FMath::Square(ImpulseTangentU) + FMath::Square(ImpulseTangentV);
-			if (ImpulseTangentSq > MaxImpulseTangentSq)
+			if (ImpulseTangentSq > (MaxImpulseTangentSq + SMALL_NUMBER))
 			{
 				const FSolverReal ImpulseTangentScale = MaxImpulseTangent * FMath::InvSqrt(ImpulseTangentSq);
 				ImpulseTangentU *= ImpulseTangentScale;
