Add UPackage meta data for static and skeletal meshes and for anim sequence.
#rnx

diff --git a/Engine/Plugins/Experimental/Interchange/Runtime/Source/Import/Private/Mesh/InterchangeSkeletonHelper.cpp b/Engine/Plugins/Experimental/Interchange/Runtime/Source/Import/Private/Mesh/InterchangeSkeletonHelper.cpp
--- a/Engine/Plugins/Experimental/Interchange/Runtime/Source/Import/Private/Mesh/InterchangeSkeletonHelper.cpp
+++ b/Engine/Plugins/Experimental/Interchange/Runtime/Source/Import/Private/Mesh/InterchangeSkeletonHelper.cpp
@@ -9,7 +9,7 @@
 #include "Rendering/SkeletalMeshLODImporterData.h"
 #include "Nodes/InterchangeBaseNode.h"
 #include "Nodes/InterchangeBaseNodeContainer.h"
-
+#include "Nodes/InterchangeUserDefinedAttribute.h"
 
 namespace UE::Interchange::Private
 {
@@ -155,6 +155,24 @@ namespace UE::Interchange::Private
 		return (NumOfBoneMatches > 0);
 	}
 
+	void FSkeletonHelper::RecursiveAddSkeletonMetaDataValues(UInterchangeBaseNodeContainer* NodeContainer, UInterchangeBaseNode* DestinationNode, const FString& JointUid)
+	{
+		const UInterchangeSceneNode* SceneNode = Cast<UInterchangeSceneNode>(NodeContainer->GetNode(JointUid));
+		if (!SceneNode || !SceneNode->IsSpecializedTypeContains(FSceneNodeStaticData::GetJointSpecializeTypeString()))
+		{
+			return;
+		}
+		constexpr bool bAddSourceNodeName = true;
+		UInterchangeUserDefinedAttributesAPI::DuplicateAllUserDefinedAttribute(SceneNode, DestinationNode, bAddSourceNodeName);
+
+		//Iterate childrens
+		const TArray<FString> ChildrenIds = NodeContainer->GetNodeChildrenUids(JointUid);
+		for (int32 ChildIndex = 0; ChildIndex < ChildrenIds.Num(); ++ChildIndex)
+		{
+			RecursiveAddSkeletonMetaDataValues(NodeContainer, DestinationNode, ChildrenIds[ChildIndex]);
+		}
+	}
+
 	void FSkeletonHelper::RecursiveAddBones(const UInterchangeBaseNodeContainer* NodeContainer, const FString& JointNodeId, TArray <FJointInfo>& JointInfos, int32 ParentIndex, TArray<SkeletalMeshImportData::FBone>& RefBonesBinary, const bool bUseTimeZeroAsBindPose, bool& bOutDiffPose)
 	{
 		const UInterchangeSceneNode* JointNode = Cast<UInterchangeSceneNode>(NodeContainer->GetNode(JointNodeId));
diff --git a/Engine/Plugins/Experimental/Interchange/Runtime/Source/Import/Public/InterchangeSkeletonHelper.h b/Engine/Plugins/Experimental/Interchange/Runtime/Source/Import/Public/InterchangeSkeletonHelper.h
--- a/Engine/Plugins/Experimental/Interchange/Runtime/Source/Import/Public/InterchangeSkeletonHelper.h
+++ b/Engine/Plugins/Experimental/Interchange/Runtime/Source/Import/Public/InterchangeSkeletonHelper.h
@@ -9,6 +9,7 @@
 //#include "ReferenceSkeleton.h"
 
 struct FReferenceSkeleton;
+class UInterchangeBaseNode;
 class UInterchangeBaseNodeContainer;
 class USkeleton;
 struct FMeshBoneInfo;
@@ -28,7 +29,7 @@ namespace UE::Interchange::Private
 	public:
 		static bool ProcessImportMeshSkeleton(const USkeleton* SkeletonAsset, FReferenceSkeleton& RefSkeleton, int32& SkeletalDepth, const UInterchangeBaseNodeContainer* NodeContainer, const FString& RootJointNodeId, TArray<SkeletalMeshImportData::FBone>& RefBonesBinary, const bool bUseTimeZeroAsBindPose, bool& bOutDiffPose);
 		static bool IsCompatibleSkeleton(const USkeleton* Skeleton, const FString RootJoinUid, const UInterchangeBaseNodeContainer* BaseNodeContainer);
-			
+		static void RecursiveAddSkeletonMetaDataValues(UInterchangeBaseNodeContainer* NodeContainer, UInterchangeBaseNode* DestinationNode, const FString& JointUid);
 	private:
 		static void RecursiveAddBones(const UInterchangeBaseNodeContainer* NodeContainer, const FString& JointNodeId, TArray <FJointInfo>& JointInfos, int32 ParentIndex, TArray<SkeletalMeshImportData::FBone>& RefBonesBinary, const bool bUseTimeZeroAsBindPose, bool& bOutDiffPose);
 		static FName SkeletalLodGetBoneName(const TArray<FMeshBoneInfo>& SkeletalLodRawInfos, int32 BoneIndex);
diff --git a/Engine/Plugins/Experimental/Interchange/Runtime/Source/Pipelines/Private/InterchangeGenericAnimationPipeline.cpp b/Engine/Plugins/Experimental/Interchange/Runtime/Source/Pipelines/Private/InterchangeGenericAnimationPipeline.cpp
--- a/Engine/Plugins/Experimental/Interchange/Runtime/Source/Pipelines/Private/InterchangeGenericAnimationPipeline.cpp
+++ b/Engine/Plugins/Experimental/Interchange/Runtime/Source/Pipelines/Private/InterchangeGenericAnimationPipeline.cpp
@@ -205,6 +205,16 @@ void UInterchangeGenericAnimationPipeline::ExecutePreImportPipeline(UInterchange
 		//USkeleton cannot be created without a valid skeletalmesh
 		const FString SkeletonUid = SkeletonFactoryNode->GetUniqueID();
 		AnimSequenceFactoryNode->AddFactoryDependencyUid(SkeletonUid);
+
+		FString RootJointUid;
+		if(SkeletonFactoryNode->GetCustomRootJointUid(RootJointUid))
+		{
+#if WITH_EDITOR
+			//Iterate all joints to set the meta data value in the anim sequence factory node
+			UE::Interchange::Private::FSkeletonHelper::RecursiveAddSkeletonMetaDataValues(BaseNodeContainer, AnimSequenceFactoryNode, RootJointUid);
+#endif //WITH_EDITOR
+		}
+
 		const bool bImportOnlyAnimation = CommonSkeletalMeshesAndAnimationsProperties->bImportOnlyAnimations;
 		
 		//Iterate dependencies
@@ -231,8 +241,6 @@ void UInterchangeGenericAnimationPipeline::ExecutePreImportPipeline(UInterchange
 
 		if (!CommonSkeletalMeshesAndAnimationsProperties->Skeleton.IsNull())
 		{
-			FString RootJointUid;
-			SkeletonFactoryNode->GetCustomRootJointUid(RootJointUid);
 			bool bSkeletonCompatible = true;
 
 			//TODO: support skeleton helper in runtime
diff --git a/Engine/Plugins/Experimental/Interchange/Runtime/Source/Pipelines/Private/InterchangeGenericAssetsPipeline.cpp b/Engine/Plugins/Experimental/Interchange/Runtime/Source/Pipelines/Private/InterchangeGenericAssetsPipeline.cpp
--- a/Engine/Plugins/Experimental/Interchange/Runtime/Source/Pipelines/Private/InterchangeGenericAssetsPipeline.cpp
+++ b/Engine/Plugins/Experimental/Interchange/Runtime/Source/Pipelines/Private/InterchangeGenericAssetsPipeline.cpp
@@ -23,9 +23,11 @@
 #include "Misc/Paths.h"
 #include "Nodes/InterchangeBaseNode.h"
 #include "Nodes/InterchangeBaseNodeContainer.h"
+#include "Nodes/InterchangeUserDefinedAttribute.h"
 #include "Templates/SubclassOf.h"
 #include "UObject/Object.h"
 #include "UObject/ObjectMacros.h"
+#include "UObject/MetaData.h"
 
 UInterchangeGenericAssetsPipeline::UInterchangeGenericAssetsPipeline()
 {
@@ -212,6 +214,228 @@ void UInterchangeGenericAssetsPipeline::ExecutePostImportPipeline(const UInterch
 	{
 		AnimationPipeline->ScriptedExecutePostImportPipeline(InBaseNodeContainer, NodeKey, CreatedAsset, bIsAReimport);
 	}
+
+	const FString InterchangeMetaDataPrefix = TEXT("INTERCHANGE.");
+
+	if (const UInterchangeBaseNode* Node = InBaseNodeContainer->GetNode(NodeKey))
+	{
+		//Add UObject package meta data
+		if (UMetaData* MetaData = CreatedAsset->GetOutermost()->GetMetaData())
+		{
+			//Cleanup existing INTERCHANGE_ prefix metadata name for this object (in casse we re-import)
+			{
+				TArray<FName> InterchangeMetaDataKeys;
+				if(TMap<FName, FString>* MetaDataMapPtr = MetaData->GetMapForObject(CreatedAsset))
+				{
+					for (const TPair<FName, FString>& ObjectMetadata : *MetaDataMapPtr)
+					{
+						if (ObjectMetadata.Key.ToString().StartsWith(InterchangeMetaDataPrefix))
+						{
+							InterchangeMetaDataKeys.Add(ObjectMetadata.Key);
+						}
+					}
+					for (const FName& MetaDataKey : InterchangeMetaDataKeys)
+					{
+						MetaData->RemoveValue(CreatedAsset, MetaDataKey);
+					}
+				}
+			}
+			TArray<FInterchangeUserDefinedAttributeInfo> UserAttributeInfos;
+			UInterchangeUserDefinedAttributesAPI::GetUserDefinedAttributeInfos(Node, UserAttributeInfos);
+			//We must convert all different type to String since meta data only support string
+			for (const FInterchangeUserDefinedAttributeInfo& UserAttributeInfo : UserAttributeInfos)
+			{
+				if (UserAttributeInfo.PayloadKey.IsSet())
+				{
+					//Skip animated attributes
+					continue;
+				}
+				TOptional<FString> MetaDataValue;
+				TOptional<FString> PayloadKey;
+				switch (UserAttributeInfo.Type)
+				{
+					case UE::Interchange::EAttributeTypes::Bool:
+					{
+						bool Value = false;
+						if(UInterchangeUserDefinedAttributesAPI::GetUserDefinedAttribute(Node, UserAttributeInfo.Name, Value, PayloadKey))
+						{
+							MetaDataValue = UE::Interchange::AttributeValueToString(Value);
+						}
+					}
+					break;
+					case UE::Interchange::EAttributeTypes::Int8:
+					{
+						int8 Value = 0;
+						if (UInterchangeUserDefinedAttributesAPI::GetUserDefinedAttribute(Node, UserAttributeInfo.Name, Value, PayloadKey))
+						{
+							MetaDataValue = UE::Interchange::AttributeValueToString(Value);
+						}
+					}
+					break;
+					case UE::Interchange::EAttributeTypes::Int16:
+					{
+						int16 Value = 0;
+						if (UInterchangeUserDefinedAttributesAPI::GetUserDefinedAttribute(Node, UserAttributeInfo.Name, Value, PayloadKey))
+						{
+							MetaDataValue = UE::Interchange::AttributeValueToString(Value);
+						}
+					}
+					break;
+					case UE::Interchange::EAttributeTypes::Int32:
+					{
+						int32 Value = 0;
+						if (UInterchangeUserDefinedAttributesAPI::GetUserDefinedAttribute(Node, UserAttributeInfo.Name, Value, PayloadKey))
+						{
+							MetaDataValue = UE::Interchange::AttributeValueToString(Value);
+						}
+					}
+					break;
+					case UE::Interchange::EAttributeTypes::Int64:
+					{
+						int64 Value = 0;
+						if (UInterchangeUserDefinedAttributesAPI::GetUserDefinedAttribute(Node, UserAttributeInfo.Name, Value, PayloadKey))
+						{
+							MetaDataValue = UE::Interchange::AttributeValueToString(Value);
+						}
+					}
+					break;
+					case UE::Interchange::EAttributeTypes::UInt8:
+					{
+						uint8 Value = 0;
+						if (UInterchangeUserDefinedAttributesAPI::GetUserDefinedAttribute(Node, UserAttributeInfo.Name, Value, PayloadKey))
+						{
+							MetaDataValue = UE::Interchange::AttributeValueToString(Value);
+						}
+					}
+					break;
+					case UE::Interchange::EAttributeTypes::UInt16:
+					{
+						uint16 Value = 0;
+						if (UInterchangeUserDefinedAttributesAPI::GetUserDefinedAttribute(Node, UserAttributeInfo.Name, Value, PayloadKey))
+						{
+							MetaDataValue = UE::Interchange::AttributeValueToString(Value);
+						}
+					}
+					break;
+					case UE::Interchange::EAttributeTypes::UInt32:
+					{
+						uint32 Value = 0;
+						if (UInterchangeUserDefinedAttributesAPI::GetUserDefinedAttribute(Node, UserAttributeInfo.Name, Value, PayloadKey))
+						{
+							MetaDataValue = UE::Interchange::AttributeValueToString(Value);
+						}
+					}
+					break;
+					case UE::Interchange::EAttributeTypes::UInt64:
+					{
+						uint64 Value = 0;
+						if (UInterchangeUserDefinedAttributesAPI::GetUserDefinedAttribute(Node, UserAttributeInfo.Name, Value, PayloadKey))
+						{
+							MetaDataValue = UE::Interchange::AttributeValueToString(Value);
+						}
+					}
+					break;
+					case UE::Interchange::EAttributeTypes::Float:
+					{
+						float Value = 0.0f;
+						if (UInterchangeUserDefinedAttributesAPI::GetUserDefinedAttribute(Node, UserAttributeInfo.Name, Value, PayloadKey))
+						{
+							MetaDataValue = UE::Interchange::AttributeValueToString(Value);
+						}
+					}
+					break;
+					case UE::Interchange::EAttributeTypes::Float16:
+					{
+						FFloat16 Value = 0.0f;
+						if (UInterchangeUserDefinedAttributesAPI::GetUserDefinedAttribute(Node, UserAttributeInfo.Name, Value, PayloadKey))
+						{
+							MetaDataValue = UE::Interchange::AttributeValueToString(Value);
+						}
+					}
+					break;
+					case UE::Interchange::EAttributeTypes::Vector2f:
+					{
+						FVector2f Value(0.0f);
+						if (UInterchangeUserDefinedAttributesAPI::GetUserDefinedAttribute(Node, UserAttributeInfo.Name, Value, PayloadKey))
+						{
+							MetaDataValue = UE::Interchange::AttributeValueToString(Value);
+						}
+					}
+					break;
+					case UE::Interchange::EAttributeTypes::Vector3f:
+					{
+						FVector3f Value(0.0f);
+						if (UInterchangeUserDefinedAttributesAPI::GetUserDefinedAttribute(Node, UserAttributeInfo.Name, Value, PayloadKey))
+						{
+							MetaDataValue = UE::Interchange::AttributeValueToString(Value);
+						}
+					}
+					break;
+					case UE::Interchange::EAttributeTypes::Vector4f:
+					{
+						FVector4f Value(0.0f);
+						if (UInterchangeUserDefinedAttributesAPI::GetUserDefinedAttribute(Node, UserAttributeInfo.Name, Value, PayloadKey))
+						{
+							MetaDataValue = UE::Interchange::AttributeValueToString(Value);
+						}
+					}
+					break;
+					case UE::Interchange::EAttributeTypes::Double:
+					{
+						double Value = 0.0;
+						if (UInterchangeUserDefinedAttributesAPI::GetUserDefinedAttribute(Node, UserAttributeInfo.Name, Value, PayloadKey))
+						{
+							MetaDataValue = UE::Interchange::AttributeValueToString(Value);
+						}
+					}
+					break;
+					case UE::Interchange::EAttributeTypes::Vector2d:
+					{
+						FVector2D Value(0.0);
+						if (UInterchangeUserDefinedAttributesAPI::GetUserDefinedAttribute(Node, UserAttributeInfo.Name, Value, PayloadKey))
+						{
+							MetaDataValue = UE::Interchange::AttributeValueToString(Value);
+						}
+					}
+					break;
+					case UE::Interchange::EAttributeTypes::Vector3d:
+					{
+						FVector3d Value(0.0);
+						if (UInterchangeUserDefinedAttributesAPI::GetUserDefinedAttribute(Node, UserAttributeInfo.Name, Value, PayloadKey))
+						{
+							MetaDataValue = UE::Interchange::AttributeValueToString(Value);
+						}
+					}
+					break;
+					case UE::Interchange::EAttributeTypes::Vector4d:
+					{
+						FVector4d Value(0.0);
+						if (UInterchangeUserDefinedAttributesAPI::GetUserDefinedAttribute(Node, UserAttributeInfo.Name, Value, PayloadKey))
+						{
+							MetaDataValue = UE::Interchange::AttributeValueToString(Value);
+						}
+					}
+					break;
+					case UE::Interchange::EAttributeTypes::String:
+					{
+						FString Value;
+						if (UInterchangeUserDefinedAttributesAPI::GetUserDefinedAttribute(Node, UserAttributeInfo.Name, Value, PayloadKey))
+						{
+							MetaDataValue = UE::Interchange::AttributeValueToString(Value);
+						}
+					}
+					break;
+				}
+				if (MetaDataValue.IsSet())
+				{
+					const FString& MetaDataStringValue = MetaDataValue.GetValue();
+					const FName& MetaDataKey = FName(InterchangeMetaDataPrefix + UserAttributeInfo.Name);
+					//SetValue either add the key or set the new value
+					MetaData->SetValue(CreatedAsset, MetaDataKey, *MetaDataStringValue);
+				}
+			}
+		}
+	}
 }
 
 void UInterchangeGenericAssetsPipeline::SetReimportSourceIndex(UClass* ReimportObjectClass, const int32 SourceFileIndex)
diff --git a/Engine/Plugins/Experimental/Interchange/Runtime/Source/Pipelines/Private/InterchangeGenericSkeletalMeshPipeline.cpp b/Engine/Plugins/Experimental/Interchange/Runtime/Source/Pipelines/Private/InterchangeGenericSkeletalMeshPipeline.cpp
--- a/Engine/Plugins/Experimental/Interchange/Runtime/Source/Pipelines/Private/InterchangeGenericSkeletalMeshPipeline.cpp
+++ b/Engine/Plugins/Experimental/Interchange/Runtime/Source/Pipelines/Private/InterchangeGenericSkeletalMeshPipeline.cpp
@@ -20,6 +20,7 @@
 #include "Misc/Paths.h"
 #include "Nodes/InterchangeBaseNode.h"
 #include "Nodes/InterchangeBaseNodeContainer.h"
+#include "Nodes/InterchangeUserDefinedAttribute.h"
 #if WITH_EDITOR
 #include "PhysicsAssetUtils.h"
 #endif //WITH_EDITOR
@@ -294,6 +295,11 @@ UInterchangeSkeletonFactoryNode* UInterchangeGenericMeshPipeline::CreateSkeleton
 		SkeletonFactoryNode->SetEnabled(false);
 		SkeletonFactoryNode->ReferenceObject = CommonSkeletalMeshesAndAnimationsProperties->Skeleton;
 	}
+#if WITH_EDITOR
+	//Iterate all joints to set the meta data value in the skeleton node
+	UE::Interchange::Private::FSkeletonHelper::RecursiveAddSkeletonMetaDataValues(BaseNodeContainer, SkeletonFactoryNode, RootJointUid);
+#endif //WITH_EDITOR
+
 	return SkeletonFactoryNode;
 }
 
@@ -527,6 +533,7 @@ void UInterchangeGenericMeshPipeline::AddLodDataToSkeletalMesh(const UInterchang
 			LodDataNode->SetCustomSkeletonUid(SkeletonUid);
 			SkeletalMeshFactoryNode->AddLodDataUniqueId(SkeletalMeshLodDataUniqueID);
 		}
+		constexpr bool bAddSourceNodeName = true;
 		for (const FString& NodeUid : NodeUids)
 		{
 			TArray<FString> MaterialDependencies;
@@ -537,6 +544,7 @@ void UInterchangeGenericMeshPipeline::AddLodDataToSkeletalMesh(const UInterchang
 				if (BaseNodeContainer->IsNodeUidValid(MeshDependency))
 				{
 					const UInterchangeMeshNode* MeshDependencyNode = Cast<UInterchangeMeshNode>(BaseNodeContainer->GetNode(MeshDependency));
+					UInterchangeUserDefinedAttributesAPI::DuplicateAllUserDefinedAttribute(MeshDependencyNode, SkeletalMeshFactoryNode, bAddSourceNodeName);
 					SkeletalMeshFactoryNode->AddTargetNodeUid(MeshDependency);
 					MeshDependencyNode->AddTargetNodeUid(SkeletalMeshFactoryNode->GetUniqueID());
 					MeshDependencyNode->GetMaterialDependencies(MaterialDependencies);
@@ -545,9 +553,11 @@ void UInterchangeGenericMeshPipeline::AddLodDataToSkeletalMesh(const UInterchang
 				{
 					SceneNode->GetMaterialDependencyUids(MaterialDependencies);
 				}
+				UInterchangeUserDefinedAttributesAPI::DuplicateAllUserDefinedAttribute(SceneNode, SkeletalMeshFactoryNode, bAddSourceNodeName);
 			}
 			else if (const UInterchangeMeshNode* MeshNode = Cast<UInterchangeMeshNode>(BaseNodeContainer->GetNode(NodeUid)))
 			{
+				UInterchangeUserDefinedAttributesAPI::DuplicateAllUserDefinedAttribute(MeshNode, SkeletalMeshFactoryNode, bAddSourceNodeName);
 				SkeletalMeshFactoryNode->AddTargetNodeUid(NodeUid);
 				MeshNode->AddTargetNodeUid(SkeletalMeshFactoryNode->GetUniqueID());
 				MeshNode->GetMaterialDependencies(MaterialDependencies);
diff --git a/Engine/Plugins/Experimental/Interchange/Runtime/Source/Pipelines/Private/InterchangeGenericStaticMeshPipeline.cpp b/Engine/Plugins/Experimental/Interchange/Runtime/Source/Pipelines/Private/InterchangeGenericStaticMeshPipeline.cpp
--- a/Engine/Plugins/Experimental/Interchange/Runtime/Source/Pipelines/Private/InterchangeGenericStaticMeshPipeline.cpp
+++ b/Engine/Plugins/Experimental/Interchange/Runtime/Source/Pipelines/Private/InterchangeGenericStaticMeshPipeline.cpp
@@ -15,6 +15,7 @@
 #include "Misc/Paths.h"
 #include "Nodes/InterchangeBaseNode.h"
 #include "Nodes/InterchangeBaseNodeContainer.h"
+#include "Nodes/InterchangeUserDefinedAttribute.h"
 #include "Tasks/Task.h"
 #include "UObject/Object.h"
 #include "UObject/ObjectMacros.h"
@@ -466,7 +467,7 @@ void UInterchangeGenericMeshPipeline::AddLodDataToStaticMesh(UInterchangeStaticM
 			LodDataNode->SetParentUid(StaticMeshFactoryUid);
 			StaticMeshFactoryNode->AddLodDataUniqueId(StaticMeshLodDataUniqueID);
 		}
-
+		constexpr bool bAddSourceNodeName = true;
 		for (const FString& NodeUid : NodeUids)
 		{
 			TArray<FString> MaterialDependencies;
@@ -477,6 +478,7 @@ void UInterchangeGenericMeshPipeline::AddLodDataToStaticMesh(UInterchangeStaticM
 				if (BaseNodeContainer->IsNodeUidValid(MeshDependency))
 				{
 					const UInterchangeMeshNode* MeshDependencyNode = Cast<UInterchangeMeshNode>(BaseNodeContainer->GetNode(MeshDependency));
+					UInterchangeUserDefinedAttributesAPI::DuplicateAllUserDefinedAttribute(MeshDependencyNode, StaticMeshFactoryNode, bAddSourceNodeName);
 					StaticMeshFactoryNode->AddTargetNodeUid(MeshDependency);
 					MeshDependencyNode->AddTargetNodeUid(StaticMeshFactoryNode->GetUniqueID());
 					MeshDependencyNode->GetMaterialDependencies(MaterialDependencies);
@@ -485,9 +487,11 @@ void UInterchangeGenericMeshPipeline::AddLodDataToStaticMesh(UInterchangeStaticM
 				{
 					SceneNode->GetMaterialDependencyUids(MaterialDependencies);
 				}
+				UInterchangeUserDefinedAttributesAPI::DuplicateAllUserDefinedAttribute(SceneNode, StaticMeshFactoryNode, bAddSourceNodeName);
 			}
 			else if (const UInterchangeMeshNode* MeshNode = Cast<UInterchangeMeshNode>(BaseNodeContainer->GetNode(NodeUid)))
 			{
+				UInterchangeUserDefinedAttributesAPI::DuplicateAllUserDefinedAttribute(MeshNode, StaticMeshFactoryNode, bAddSourceNodeName);
 				StaticMeshFactoryNode->AddTargetNodeUid(NodeUid);
 				MeshNode->AddTargetNodeUid(StaticMeshFactoryNode->GetUniqueID());
 				MeshNode->GetMaterialDependencies(MaterialDependencies);
diff --git a/Engine/Source/Runtime/Experimental/Interchange/Core/Private/Nodes/InterchangeUserDefinedAttribute.cpp b/Engine/Source/Runtime/Experimental/Interchange/Core/Private/Nodes/InterchangeUserDefinedAttribute.cpp
--- a/Engine/Source/Runtime/Experimental/Interchange/Core/Private/Nodes/InterchangeUserDefinedAttribute.cpp
+++ b/Engine/Source/Runtime/Experimental/Interchange/Core/Private/Nodes/InterchangeUserDefinedAttribute.cpp
@@ -139,27 +139,134 @@ bool UInterchangeUserDefinedAttributesAPI::GetUserDefinedAttribute_FString(const
 	return bResult;
 }
 
-TArray<FString> UInterchangeUserDefinedAttributesAPI::GetUserDefinedAttributeNames(const UInterchangeBaseNode* InterchangeNode)
+TArray<FInterchangeUserDefinedAttributeInfo> UInterchangeUserDefinedAttributesAPI::GetUserDefinedAttributeInfos(const UInterchangeBaseNode* InterchangeNode)
 {
 	check(InterchangeNode);
 	TArray<UE::Interchange::FAttributeKey> AttributeKeys;
 	InterchangeNode->GetAttributeKeys(AttributeKeys);
-	TArray<FString> UserDefinedAttributeNames;
+	TArray<FInterchangeUserDefinedAttributeInfo> UserDefinedAttributeInfos;
 	int32 RightChopIndex = UserDefinedAttributeBaseKey.Len();
 	int32 LeftChopIndex = UserDefinedAttributeValuePostKey.Len();
 	for (UE::Interchange::FAttributeKey& AttributeKey : AttributeKeys)
 	{
 		if (AttributeKey.Key.StartsWith(UserDefinedAttributeBaseKey) && AttributeKey.Key.EndsWith(UserDefinedAttributeValuePostKey))
 		{
-			FString UserDefineAttributeName = AttributeKey.Key.RightChop(RightChopIndex).LeftChop(LeftChopIndex);
-			UserDefinedAttributeNames.Add(UserDefineAttributeName);
+			FString UserDefinedAttributeName = AttributeKey.Key.RightChop(RightChopIndex).LeftChop(LeftChopIndex);
+			FInterchangeUserDefinedAttributeInfo& UserDefinedAttributeInfo = UserDefinedAttributeInfos.AddDefaulted_GetRef();
+			UserDefinedAttributeInfo.Type = InterchangeNode->GetAttributeType(AttributeKey);
+			UserDefinedAttributeInfo.Name = UserDefinedAttributeName;
+
+			//Get the optional payload key
+			const FString StorageBaseKey = UserDefinedAttributeBaseKey + UserDefinedAttributeName;
+			const UE::Interchange::FAttributeKey UserDefinedPayloadKey = UE::Interchange::FAttributeKey(StorageBaseKey + UserDefinedAttributePayLoadPostKey);
+			if (InterchangeNode->HasAttribute(UserDefinedPayloadKey))
+			{
+				FString PayloadKey;
+				InterchangeNode->GetStringAttribute(UserDefinedPayloadKey.Key, PayloadKey);
+				UserDefinedAttributeInfo.PayloadKey = PayloadKey;
+			}
 		}
 	}
-	return UserDefinedAttributeNames;
+	return UserDefinedAttributeInfos;
 }
 
-void UInterchangeUserDefinedAttributesAPI::GetUserDefinedAttributeNames(const UInterchangeBaseNode* InterchangeNode, TArray<FString>& UserDefinedAttributeNames)
+void UInterchangeUserDefinedAttributesAPI::GetUserDefinedAttributeInfos(const UInterchangeBaseNode* InterchangeNode, TArray<FInterchangeUserDefinedAttributeInfo>& UserDefinedAttributeInfos)
 {
 	check(InterchangeNode);
-	UserDefinedAttributeNames = GetUserDefinedAttributeNames(InterchangeNode);
+	UserDefinedAttributeInfos = GetUserDefinedAttributeInfos(InterchangeNode);
+}
+
+void UInterchangeUserDefinedAttributesAPI::DuplicateAllUserDefinedAttribute(const UInterchangeBaseNode* InterchangeSourceNode, UInterchangeBaseNode* InterchangeDestinationNode, bool bAddSourceNodeName)
+{
+	check(InterchangeSourceNode);
+	check(InterchangeDestinationNode);
+	TArray<UE::Interchange::FAttributeKey> AttributeKeys;
+	InterchangeSourceNode->GetAttributeKeys(AttributeKeys);
+	int32 RightChopIndex = UserDefinedAttributeBaseKey.Len();
+	int32 LeftChopIndex = UserDefinedAttributeValuePostKey.Len();
+	for (UE::Interchange::FAttributeKey& AttributeKey : AttributeKeys)
+	{
+		if (AttributeKey.Key.StartsWith(UserDefinedAttributeBaseKey) && AttributeKey.Key.EndsWith(UserDefinedAttributeValuePostKey))
+		{
+			FString UserDefinedAttributeName = AttributeKey.Key.RightChop(RightChopIndex).LeftChop(LeftChopIndex);
+			//Get the optional payload key
+			const FString StorageBaseKey = UserDefinedAttributeBaseKey + UserDefinedAttributeName;
+
+			UE::Interchange::EAttributeTypes Type = InterchangeSourceNode->GetAttributeType(AttributeKey);
+			FString DuplicateName = UserDefinedAttributeName;
+			if (bAddSourceNodeName)
+			{
+				DuplicateName = InterchangeSourceNode->GetDisplayLabel() + TEXT(".") + UserDefinedAttributeName;
+			}
+
+			const UE::Interchange::FAttributeKey UserDefinedPayloadKey = UE::Interchange::FAttributeKey(StorageBaseKey + UserDefinedAttributePayLoadPostKey);
+			TOptional<FString> PayloadKey;
+			if (InterchangeSourceNode->HasAttribute(UserDefinedPayloadKey))
+			{
+				FString PayloadKeyValue;
+				InterchangeSourceNode->GetStringAttribute(UserDefinedPayloadKey.Key, PayloadKeyValue);
+				PayloadKey = PayloadKeyValue;
+			}
+			switch (Type)
+			{
+				case UE::Interchange::EAttributeTypes::Bool:
+					CreateUserDefinedAttribute(InterchangeDestinationNode, DuplicateName, InterchangeSourceNode->GetAttributeChecked<bool>(AttributeKey.Key), PayloadKey);
+					break;
+				case UE::Interchange::EAttributeTypes::Int8:
+					CreateUserDefinedAttribute(InterchangeDestinationNode, DuplicateName, InterchangeSourceNode->GetAttributeChecked<int8>(AttributeKey.Key), PayloadKey);
+					break;
+				case UE::Interchange::EAttributeTypes::Int16:
+					CreateUserDefinedAttribute(InterchangeDestinationNode, DuplicateName, InterchangeSourceNode->GetAttributeChecked<int16>(AttributeKey.Key), PayloadKey);
+					break;
+				case UE::Interchange::EAttributeTypes::Int32:
+					CreateUserDefinedAttribute(InterchangeDestinationNode, DuplicateName, InterchangeSourceNode->GetAttributeChecked<int32>(AttributeKey.Key), PayloadKey);
+					break;
+				case UE::Interchange::EAttributeTypes::Int64:
+					CreateUserDefinedAttribute(InterchangeDestinationNode, DuplicateName, InterchangeSourceNode->GetAttributeChecked<int64>(AttributeKey.Key), PayloadKey);
+					break;
+				case UE::Interchange::EAttributeTypes::UInt8:
+					CreateUserDefinedAttribute(InterchangeDestinationNode, DuplicateName, InterchangeSourceNode->GetAttributeChecked<uint8>(AttributeKey.Key), PayloadKey);
+					break;
+				case UE::Interchange::EAttributeTypes::UInt16:
+					CreateUserDefinedAttribute(InterchangeDestinationNode, DuplicateName, InterchangeSourceNode->GetAttributeChecked<uint16>(AttributeKey.Key), PayloadKey);
+					break;
+				case UE::Interchange::EAttributeTypes::UInt32:
+					CreateUserDefinedAttribute(InterchangeDestinationNode, DuplicateName, InterchangeSourceNode->GetAttributeChecked<uint32>(AttributeKey.Key), PayloadKey);
+					break;
+				case UE::Interchange::EAttributeTypes::UInt64:
+					CreateUserDefinedAttribute(InterchangeDestinationNode, DuplicateName, InterchangeSourceNode->GetAttributeChecked<uint64>(AttributeKey.Key), PayloadKey);
+					break;
+				case UE::Interchange::EAttributeTypes::Float:
+					CreateUserDefinedAttribute(InterchangeDestinationNode, DuplicateName, InterchangeSourceNode->GetAttributeChecked<float>(AttributeKey.Key), PayloadKey);
+					break;
+				case UE::Interchange::EAttributeTypes::Float16:
+					CreateUserDefinedAttribute(InterchangeDestinationNode, DuplicateName, InterchangeSourceNode->GetAttributeChecked<FFloat16>(AttributeKey.Key), PayloadKey);
+					break;
+				case UE::Interchange::EAttributeTypes::Vector2f:
+					CreateUserDefinedAttribute(InterchangeDestinationNode, DuplicateName, InterchangeSourceNode->GetAttributeChecked<FVector2f>(AttributeKey.Key), PayloadKey);
+					break;
+				case UE::Interchange::EAttributeTypes::Vector3f:
+					CreateUserDefinedAttribute(InterchangeDestinationNode, DuplicateName, InterchangeSourceNode->GetAttributeChecked<FVector3f>(AttributeKey.Key), PayloadKey);
+					break;
+				case UE::Interchange::EAttributeTypes::Vector4f:
+					CreateUserDefinedAttribute(InterchangeDestinationNode, DuplicateName, InterchangeSourceNode->GetAttributeChecked<FVector4f>(AttributeKey.Key), PayloadKey);
+					break;
+				case UE::Interchange::EAttributeTypes::Double:
+					CreateUserDefinedAttribute(InterchangeDestinationNode, DuplicateName, InterchangeSourceNode->GetAttributeChecked<double>(AttributeKey.Key), PayloadKey);
+					break;
+				case UE::Interchange::EAttributeTypes::Vector2d:
+					CreateUserDefinedAttribute(InterchangeDestinationNode, DuplicateName, InterchangeSourceNode->GetAttributeChecked<FVector2D>(AttributeKey.Key), PayloadKey);
+					break;
+				case UE::Interchange::EAttributeTypes::Vector3d:
+					CreateUserDefinedAttribute(InterchangeDestinationNode, DuplicateName, InterchangeSourceNode->GetAttributeChecked<FVector3d>(AttributeKey.Key), PayloadKey);
+					break;
+				case UE::Interchange::EAttributeTypes::Vector4d:
+					CreateUserDefinedAttribute(InterchangeDestinationNode, DuplicateName, InterchangeSourceNode->GetAttributeChecked<FVector4d>(AttributeKey.Key), PayloadKey);
+					break;
+				case UE::Interchange::EAttributeTypes::String:
+					CreateUserDefinedAttribute(InterchangeDestinationNode, DuplicateName, InterchangeSourceNode->GetAttributeChecked<FString>(AttributeKey.Key), PayloadKey);
+					break;
+			}
+		}
+	}
 }
diff --git a/Engine/Source/Runtime/Experimental/Interchange/Core/Public/Nodes/InterchangeBaseNode.h b/Engine/Source/Runtime/Experimental/Interchange/Core/Public/Nodes/InterchangeBaseNode.h
--- a/Engine/Source/Runtime/Experimental/Interchange/Core/Public/Nodes/InterchangeBaseNode.h
+++ b/Engine/Source/Runtime/Experimental/Interchange/Core/Public/Nodes/InterchangeBaseNode.h
@@ -390,6 +390,17 @@ public:
 	UFUNCTION(BlueprintCallable, Category = "Interchange | Node")
 	bool GetLinearColorAttribute(const FString& NodeAttributeKey, FLinearColor& OutValue) const;
 
+	template<typename AttributeType>
+	AttributeType GetAttributeChecked(const FString& NodeAttributeKey) const
+	{
+		AttributeType Value;
+		check(HasAttribute(UE::Interchange::FAttributeKey(NodeAttributeKey)));
+		UE::Interchange::FAttributeStorage::TAttributeHandle<AttributeType> Handle = GetAttributeHandle<AttributeType>(UE::Interchange::FAttributeKey(NodeAttributeKey));
+		check(Handle.IsValid());
+		check(Handle.Get(Value) == UE::Interchange::EAttributeStorageResult::Operation_Success);
+		return Value;
+	}
+
 	template<typename AttributeType>
 	bool GetAttribute(const FString& NodeAttributeKey, AttributeType& OutValue) const
 	{
diff --git a/Engine/Source/Runtime/Experimental/Interchange/Core/Public/Nodes/InterchangeUserDefinedAttribute.h b/Engine/Source/Runtime/Experimental/Interchange/Core/Public/Nodes/InterchangeUserDefinedAttribute.h
--- a/Engine/Source/Runtime/Experimental/Interchange/Core/Public/Nodes/InterchangeUserDefinedAttribute.h
+++ b/Engine/Source/Runtime/Experimental/Interchange/Core/Public/Nodes/InterchangeUserDefinedAttribute.h
@@ -2,6 +2,7 @@
 
 #pragma once
 
+#include "Types/AttributeStorage.h"
 #include "Misc/Optional.h"
 #include "Nodes/InterchangeBaseNode.h"
 #include "UObject/Class.h"
@@ -10,6 +11,19 @@
 
 #include "InterchangeUserDefinedAttribute.generated.h"
 
+USTRUCT(BlueprintType)
+struct FInterchangeUserDefinedAttributeInfo
+{
+	GENERATED_BODY()
+
+	UPROPERTY(EditAnywhere, BlueprintReadWrite, Category = "Interchange | Node | UserDefinedAttributeInfo")
+	FString Name;
+
+	UE::Interchange::EAttributeTypes Type;
+
+	TOptional<FString> PayloadKey;
+};
+
 /**
  * UInterchangeUserDefinedAttributesAPI is used to store and retrieve user defined attributes (i.e. DCC node attributes, pipelines will have access to those attributes)
  * Any user defined attribute have: name, value and a optional AnimationPayloadKey (FRichCurve which is a float curve).
@@ -84,11 +98,14 @@ public:
 	UFUNCTION(BlueprintCallable, Category = "Interchange | Node | UserDefinedAttribute")
 	static bool GetUserDefinedAttribute_FString(const UInterchangeBaseNode* InterchangeNode, const FString& UserDefinedAttributeName, FString& OutValue, FString& OutPayloadKey);
 
-	static TArray<FString> GetUserDefinedAttributeNames(const UInterchangeBaseNode* InterchangeNode);
+	static TArray<FInterchangeUserDefinedAttributeInfo> GetUserDefinedAttributeInfos(const UInterchangeBaseNode* InterchangeNode);
 
 	UFUNCTION(BlueprintCallable, Category = "Interchange | Node | UserDefinedAttribute")
-	static void GetUserDefinedAttributeNames(const UInterchangeBaseNode* InterchangeNode, TArray<FString>& UserDefinedAttributeNames);
+	static void GetUserDefinedAttributeInfos(const UInterchangeBaseNode* InterchangeNode, TArray<FInterchangeUserDefinedAttributeInfo>& UserDefinedAttributeInfos);
  
+	UFUNCTION(BlueprintCallable, Category = "Interchange | Node | UserDefinedAttribute")
+	static void DuplicateAllUserDefinedAttribute(const UInterchangeBaseNode* InterchangeSourceNode, UInterchangeBaseNode* InterchangeDestinationNode, bool bAddSourceNodeName);
+
 private:
 	static const FString UserDefinedAttributeBaseKey;
 	static const FString UserDefinedAttributeValuePostKey;
diff --git a/Engine/Source/Runtime/Experimental/Interchange/Core/Public/Types/AttributeStorage.h b/Engine/Source/Runtime/Experimental/Interchange/Core/Public/Types/AttributeStorage.h
--- a/Engine/Source/Runtime/Experimental/Interchange/Core/Public/Types/AttributeStorage.h
+++ b/Engine/Source/Runtime/Experimental/Interchange/Core/Public/Types/AttributeStorage.h
@@ -228,8 +228,20 @@ namespace UE
 				static_assert(!sizeof(T), "Attribute type trait must be specialized for this type.");
 				return EAttributeTypes::None;
 			}
+
+			template<typename ValueType>
+			static FString ToString(const ValueType& Value)
+			{
+				return TEXT("Unknown type");
+			}
 		};
 
+		template<typename ValueType>
+		FString AttributeValueToString(const ValueType& Value)
+		{
+			return TAttributeTypeTraits<ValueType>::ToString(Value);
+		}
+
 		/**
 		 * Enum to return complete status of a storage operation
 		 * It support success with additional information
@@ -1027,6 +1039,7 @@ namespace UE
 		template<> struct TAttributeTypeTraits<bool>
 		{
 			static CONSTEXPR EAttributeTypes GetType() { return EAttributeTypes::Bool; }
+			static FString ToString(const bool& Value) { return Value ? TEXT("true") : TEXT("false"); }
 		};
 
 
@@ -1040,12 +1053,14 @@ namespace UE
 		template<> struct TAttributeTypeTraits<TArray<uint8> >
 		{
 			static CONSTEXPR EAttributeTypes GetType() { return EAttributeTypes::ByteArray; }
+			static FString ToString(const TArray<uint8>& Value) {return TEXT("Array<uint8>"); }
 		};
 
 		/** Implements variant type traits for byte array64s. */
 		template<> struct TAttributeTypeTraits<TArray64<uint8> >
 		{
 			static CONSTEXPR EAttributeTypes GetType() { return EAttributeTypes::ByteArray64; }
+			static FString ToString(const TArray64<uint8>& Value) { return TEXT("Array64<uint8>"); }
 		};
 
 
@@ -1053,6 +1068,7 @@ namespace UE
 		template<> struct TAttributeTypeTraits<FColor>
 		{
 			static CONSTEXPR EAttributeTypes GetType() { return EAttributeTypes::Color; }
+			static FString ToString(const FColor& Value) { return Value.ToString(); }
 		};
 
 
@@ -1060,6 +1076,7 @@ namespace UE
 		template<> struct TAttributeTypeTraits<FDateTime>
 		{
 			static CONSTEXPR EAttributeTypes GetType() { return EAttributeTypes::DateTime; }
+			static FString ToString(const FDateTime& Value) { return Value.ToString(); }
 		};
 
 
@@ -1067,6 +1084,12 @@ namespace UE
 		template<> struct TAttributeTypeTraits<double>
 		{
 			static CONSTEXPR EAttributeTypes GetType() { return EAttributeTypes::Double; }
+			static FString ToString(const double& Value)
+			{
+				FStringFormatOrderedArguments OrderedArguments;
+				OrderedArguments.Add(FStringFormatArg(Value));
+				return FString::Format(TEXT("{0}"), OrderedArguments);
+			}
 		};
 
 
@@ -1074,6 +1097,13 @@ namespace UE
 		template<typename EnumType> struct TAttributeTypeTraits<TEnumAsByte<EnumType> >
 		{
 			static CONSTEXPR EAttributeTypes GetType() { return EAttributeTypes::Enum; }
+			static FString ToString(const TEnumAsByte<EnumType>& Value)
+			{
+				uint32 ValueConv = static_cast<uint32>(static_cast<uint8>(Value));
+				FStringFormatOrderedArguments OrderedArguments;
+				OrderedArguments.Add(FStringFormatArg(ValueConv));
+				return FString::Format(TEXT("{0}"), OrderedArguments);
+			}
 		};
 
 
@@ -1081,6 +1111,12 @@ namespace UE
 		template<> struct TAttributeTypeTraits<float>
 		{
 			static CONSTEXPR EAttributeTypes GetType() { return EAttributeTypes::Float; }
+			static FString ToString(const float& Value)
+			{
+				FStringFormatOrderedArguments OrderedArguments;
+				OrderedArguments.Add(FStringFormatArg(Value));
+				return FString::Format(TEXT("{0}"), OrderedArguments);
+			}
 		};
 
 
@@ -1088,6 +1124,7 @@ namespace UE
 		template<> struct TAttributeTypeTraits<FGuid>
 		{
 			static CONSTEXPR EAttributeTypes GetType() { return EAttributeTypes::Guid; }
+			static FString ToString(const FGuid& Value) { return Value.ToString(EGuidFormats::DigitsLower); }
 		};
 
 
@@ -1095,6 +1132,7 @@ namespace UE
 		template<> struct TAttributeTypeTraits<int8>
 		{
 			static CONSTEXPR EAttributeTypes GetType() { return EAttributeTypes::Int8; }
+			static FString ToString(const int8& Value) { return FString::FromInt(Value); }
 		};
 
 
@@ -1102,6 +1140,7 @@ namespace UE
 		template<> struct TAttributeTypeTraits<int16>
 		{
 			static CONSTEXPR EAttributeTypes GetType() { return EAttributeTypes::Int16; }
+			static FString ToString(const int16& Value) { return FString::FromInt(Value); }
 		};
 
 
@@ -1109,6 +1148,7 @@ namespace UE
 		template<> struct TAttributeTypeTraits<int32>
 		{
 			static CONSTEXPR EAttributeTypes GetType() { return EAttributeTypes::Int32; }
+			static FString ToString(const int32& Value) { return FString::FromInt(Value); }
 		};
 
 
@@ -1116,12 +1156,19 @@ namespace UE
 		template<> struct TAttributeTypeTraits<int64>
 		{
 			static CONSTEXPR EAttributeTypes GetType() { return EAttributeTypes::Int64; }
+			static FString ToString(const int64& Value)
+			{
+				FStringFormatOrderedArguments OrderedArguments;
+				OrderedArguments.Add(FStringFormatArg(Value));
+				return FString::Format(TEXT("{0}"), OrderedArguments);
+			}
 		};
 
 		/** Implements variant type traits for the built-in FIntRect type. */
 		template<> struct TAttributeTypeTraits<FIntRect>
 		{
 			static CONSTEXPR EAttributeTypes GetType() { return EAttributeTypes::IntRect; }
+			static FString ToString(const FIntRect& Value) { return Value.ToString(); }
 		};
 
 
@@ -1129,6 +1176,7 @@ namespace UE
 		template<> struct TAttributeTypeTraits<FLinearColor>
 		{
 			static CONSTEXPR EAttributeTypes GetType() { return EAttributeTypes::LinearColor; }
+			static FString ToString(const FLinearColor& Value) { return Value.ToString(); }
 		};
 
 
@@ -1138,6 +1186,7 @@ namespace UE
 		template<> struct TAttributeTypeTraits<FName>
 		{
 			static CONSTEXPR EAttributeTypes GetType() { return EAttributeTypes::Name; }
+			static FString ToString(const FName& Value) { return Value.ToString(); }
 		};
 
 		/** See FPlane4f/FPlane4d variant type traits for the built-in FPlane type. */
@@ -1149,6 +1198,7 @@ namespace UE
 		template<> struct TAttributeTypeTraits<FRandomStream>
 		{
 			static CONSTEXPR EAttributeTypes GetType() { return EAttributeTypes::RandomStream; }
+			static FString ToString(const FRandomStream& Value) { return Value.ToString(); }
 		};
 
 
@@ -1159,6 +1209,7 @@ namespace UE
 		template<> struct TAttributeTypeTraits<FString>
 		{
 			static CONSTEXPR EAttributeTypes GetType() { return EAttributeTypes::String; }
+			static FString ToString(const FString& Value) { return Value; }
 		};
 
 
@@ -1166,6 +1217,7 @@ namespace UE
 		template<> struct TAttributeTypeTraits<FTimespan>
 		{
 			static CONSTEXPR EAttributeTypes GetType() { return EAttributeTypes::Timespan; }
+			static FString ToString(const FTimespan& Value) { return Value.ToString(); }
 		};
 
 
@@ -1176,6 +1228,7 @@ namespace UE
 		template<> struct TAttributeTypeTraits<FTwoVectors>
 		{
 			static CONSTEXPR EAttributeTypes GetType() { return EAttributeTypes::TwoVectors; }
+			static FString ToString(const FTwoVectors& Value) { return Value.ToString(); }
 		};
 
 
@@ -1183,6 +1236,13 @@ namespace UE
 		template<> struct TAttributeTypeTraits<uint8>
 		{
 			static CONSTEXPR EAttributeTypes GetType() { return EAttributeTypes::UInt8; }
+			static FString ToString(const uint8& Value)
+			{
+				uint32 ValueConv = Value;
+				FStringFormatOrderedArguments OrderedArguments;
+				OrderedArguments.Add(FStringFormatArg(ValueConv));
+				return FString::Format(TEXT("{0}"), OrderedArguments);
+			}
 		};
 
 
@@ -1190,6 +1250,13 @@ namespace UE
 		template<> struct TAttributeTypeTraits<uint16>
 		{
 			static CONSTEXPR EAttributeTypes GetType() { return EAttributeTypes::UInt16; }
+			static FString ToString(const uint16& Value)
+			{
+				uint32 ValueConv = Value;
+				FStringFormatOrderedArguments OrderedArguments;
+				OrderedArguments.Add(FStringFormatArg(ValueConv));
+				return FString::Format(TEXT("{0}"), OrderedArguments);
+			}
 		};
 
 
@@ -1197,6 +1264,12 @@ namespace UE
 		template<> struct TAttributeTypeTraits<uint32>
 		{
 			static CONSTEXPR EAttributeTypes GetType() { return EAttributeTypes::UInt32; }
+			static FString ToString(const uint32& Value)
+			{
+				FStringFormatOrderedArguments OrderedArguments;
+				OrderedArguments.Add(FStringFormatArg(Value));
+				return FString::Format(TEXT("{0}"), OrderedArguments);
+			}
 		};
 
 
@@ -1204,6 +1277,12 @@ namespace UE
 		template<> struct TAttributeTypeTraits<uint64>
 		{
 			static CONSTEXPR EAttributeTypes GetType() { return EAttributeTypes::UInt64; }
+			static FString ToString(const uint64& Value)
+			{
+				FStringFormatOrderedArguments OrderedArguments;
+				OrderedArguments.Add(FStringFormatArg(Value));
+				return FString::Format(TEXT("{0}"), OrderedArguments);
+			}
 		};
 
 
@@ -1214,6 +1293,7 @@ namespace UE
 		template<> struct TAttributeTypeTraits<FVector2D>
 		{
 			static CONSTEXPR EAttributeTypes GetType() { return EAttributeTypes::Vector2d; }
+			static FString ToString(const FVector2D& Value) { return Value.ToString(); }
 		};
 
 		/** See FVector4f/FVector4d variant type traits for the built-in FVector4 type. */
@@ -1222,30 +1302,41 @@ namespace UE
 		template<> struct TAttributeTypeTraits<FIntPoint>
 		{
 			static CONSTEXPR EAttributeTypes GetType() { return EAttributeTypes::IntPoint; }
+			static FString ToString(const FIntPoint& Value) { return Value.ToString(); }
 		};
 
 		/** Implements variant type traits for the built-in Vector2DHalf type. */
 		template<> struct TAttributeTypeTraits<FIntVector>
 		{
 			static CONSTEXPR EAttributeTypes GetType() { return EAttributeTypes::IntVector; }
+			static FString ToString(const FIntVector& Value) { return Value.ToString(); }
 		};
 
 		/** Implements variant type traits for the built-in Vector2DHalf type. */
 		template<> struct TAttributeTypeTraits<FVector2DHalf>
 		{
 			static CONSTEXPR EAttributeTypes GetType() { return EAttributeTypes::Vector2DHalf; }
+			static FString ToString(const FVector2DHalf& Value) { return Value.ToString(); }
 		};
 
 		/** Implements variant type traits for the built-in Float16 type. */
 		template<> struct TAttributeTypeTraits<FFloat16>
 		{
 			static CONSTEXPR EAttributeTypes GetType() { return EAttributeTypes::Float16; }
+			static FString ToString(const FFloat16& Value)
+			{
+				float ValueConvert = Value.GetFloat();
+				FStringFormatOrderedArguments OrderedArguments;
+				OrderedArguments.Add(FStringFormatArg(ValueConvert));
+				return FString::Format(TEXT("{0}"), OrderedArguments);
+			}
 		};
 
 		/** Implements variant type traits for the built-in OrientedBox type. */
 		template<> struct TAttributeTypeTraits<FOrientedBox>
 		{
 			static CONSTEXPR EAttributeTypes GetType() { return EAttributeTypes::OrientedBox; }
+			static FString ToString(const FOrientedBox& Value) { return TEXT("FOrientedBox"); }
 		};
 
 		/** See FSphere3f/FSphere3d variant type traits for the built-in Sphere type. */
@@ -1254,12 +1345,20 @@ namespace UE
 		template<> struct TAttributeTypeTraits<FFrameNumber>
 		{
 			static CONSTEXPR EAttributeTypes GetType() { return EAttributeTypes::FrameNumber; }
+			static FString ToString(const FFrameNumber& Value)
+			{
+				int32 ValueConvert = Value.Value;
+				FStringFormatOrderedArguments OrderedArguments;
+				OrderedArguments.Add(FStringFormatArg(ValueConvert));
+				return FString::Format(TEXT("{0}"), OrderedArguments);
+			}
 		};
 
 		/** Implements variant type traits for the built-in FrameRate type. */
 		template<> struct TAttributeTypeTraits<FFrameRate>
 		{
 			static CONSTEXPR EAttributeTypes GetType() { return EAttributeTypes::FrameRate; }
+			static FString ToString(const FFrameRate& Value) { return Value.ToPrettyText().ToString(); }
 		};
 
 		/** Implements variant type traits for the built-in FrameTime type. */
@@ -1272,144 +1371,168 @@ namespace UE
 		template<> struct TAttributeTypeTraits<FSoftObjectPath>
 		{
 			static CONSTEXPR EAttributeTypes GetType() { return EAttributeTypes::SoftObjectPath; }
+			static FString ToString(const FSoftObjectPath& Value) { return Value.ToString(); }
 		};
 
 		/** Implements variant type traits for the built-in Matrix44f type. */
 		template<> struct TAttributeTypeTraits<FMatrix44f>
 		{
 			static CONSTEXPR EAttributeTypes GetType() { return EAttributeTypes::Matrix44f; }
+			static FString ToString(const FMatrix44f& Value) { return Value.ToString(); }
 		};
 
 		/** Implements variant type traits for the built-in Matrix44f type. */
 		template<> struct TAttributeTypeTraits<FMatrix44d>
 		{
 			static CONSTEXPR EAttributeTypes GetType() { return EAttributeTypes::Matrix44d; }
+			static FString ToString(const FMatrix44d& Value) { return Value.ToString(); }
 		};
 
 		/** Implements variant type traits for the built-in Plane4f type. */
 		template<> struct TAttributeTypeTraits<FPlane4f>
 		{
 			static CONSTEXPR EAttributeTypes GetType() { return EAttributeTypes::Plane4f; }
+			static FString ToString(const FPlane4f& Value) { return Value.ToString(); }
 		};
 
 		/** Implements variant type traits for the built-in Plane4d type. */
 		template<> struct TAttributeTypeTraits<FPlane4d>
 		{
 			static CONSTEXPR EAttributeTypes GetType() { return EAttributeTypes::Plane4d; }
+			static FString ToString(const FPlane4d& Value) { return Value.ToString(); }
 		};
 
 		/** Implements variant type traits for the built-in Quat4f type. */
 		template<> struct TAttributeTypeTraits<FQuat4f>
 		{
 			static CONSTEXPR EAttributeTypes GetType() { return EAttributeTypes::Quat4f; }
+			static FString ToString(const FQuat4f& Value) { return Value.ToString(); }
 		};
 
 		/** Implements variant type traits for the built-in Quat4d type. */
 		template<> struct TAttributeTypeTraits<FQuat4d>
 		{
 			static CONSTEXPR EAttributeTypes GetType() { return EAttributeTypes::Quat4d; }
+			static FString ToString(const FQuat4d& Value) { return Value.ToString(); }
 		};
 
 		/** Implements variant type traits for the built-in Rotator3f type. */
 		template<> struct TAttributeTypeTraits<FRotator3f>
 		{
 			static CONSTEXPR EAttributeTypes GetType() { return EAttributeTypes::Rotator3f; }
+			static FString ToString(const FRotator3f& Value) { return Value.ToString(); }
 		};
 
 		/** Implements variant type traits for the built-in Rotator3d type. */
 		template<> struct TAttributeTypeTraits<FRotator3d>
 		{
 			static CONSTEXPR EAttributeTypes GetType() { return EAttributeTypes::Rotator3d; }
+			static FString ToString(const FRotator3d& Value) { return Value.ToString(); }
 		};
 
 		/** Implements variant type traits for the built-in Rotator3f type. */
 		template<> struct TAttributeTypeTraits<FTransform3f>
 		{
 			static CONSTEXPR EAttributeTypes GetType() { return EAttributeTypes::Transform3f; }
+			static FString ToString(const FTransform3f& Value) { return Value.ToString(); }
 		};
 
 		/** Implements variant type traits for the built-in Transform3d type. */
 		template<> struct TAttributeTypeTraits<FTransform3d>
 		{
 			static CONSTEXPR EAttributeTypes GetType() { return EAttributeTypes::Transform3d; }
+			static FString ToString(const FTransform3d& Value) { return Value.ToString(); }
 		};
 
 		/** Implements variant type traits for the built-in Vector3f type. */
 		template<> struct TAttributeTypeTraits<FVector3f>
 		{
 			static CONSTEXPR EAttributeTypes GetType() { return EAttributeTypes::Vector3f; }
+			static FString ToString(const FVector3f& Value) { return Value.ToString(); }
 		};
 
 		/** Implements variant type traits for the built-in Vector3d type. */
 		template<> struct TAttributeTypeTraits<FVector3d>
 		{
 			static CONSTEXPR EAttributeTypes GetType() { return EAttributeTypes::Vector3d; }
+			static FString ToString(const FVector3d& Value) { return Value.ToString(); }
 		};
 
 		/** Implements variant type traits for the built-in Vector2f type. */
 		template<> struct TAttributeTypeTraits<FVector2f>
 		{
 			static CONSTEXPR EAttributeTypes GetType() { return EAttributeTypes::Vector2f; }
+			static FString ToString(const FVector2f& Value) { return Value.ToString(); }
 		};
 
 		/** Implements variant type traits for the built-in Vector4f type. */
 		template<> struct TAttributeTypeTraits<FVector4f>
 		{
 			static CONSTEXPR EAttributeTypes GetType() { return EAttributeTypes::Vector4f; }
+			static FString ToString(const FVector4f& Value) { return Value.ToString(); }
 		};
 
 		/** Implements variant type traits for the built-in Vector4d type. */
 		template<> struct TAttributeTypeTraits<FVector4d>
 		{
 			static CONSTEXPR EAttributeTypes GetType() { return EAttributeTypes::Vector4d; }
+			static FString ToString(const FVector4d& Value) { return Value.ToString(); }
 		};
 
 		/** Implements variant type traits for the built-in Box2f type. */
 		template<> struct TAttributeTypeTraits<FBox2f>
 		{
 			static CONSTEXPR EAttributeTypes GetType() { return EAttributeTypes::Box2f; }
+			static FString ToString(const FBox2f& Value) { return Value.ToString(); }
 		};
 
 		/** Implements variant type traits for the built-in Box2D type. */
 		template<> struct TAttributeTypeTraits<FBox2D>
 		{
 			static CONSTEXPR EAttributeTypes GetType() { return EAttributeTypes::Box2D; }
+			static FString ToString(const FBox2D& Value) { return Value.ToString(); }
 		};
 
 		/** Implements variant type traits for the built-in Box3f type. */
 		template<> struct TAttributeTypeTraits<FBox3f>
 		{
 			static CONSTEXPR EAttributeTypes GetType() { return EAttributeTypes::Box3f; }
+			static FString ToString(const FBox3f& Value) { return Value.ToString(); }
 		};
 
 		/** Implements variant type traits for the built-in Box3d type. */
 		template<> struct TAttributeTypeTraits<FBox3d>
 		{
 			static CONSTEXPR EAttributeTypes GetType() { return EAttributeTypes::Box3d; }
+			static FString ToString(const FBox3d& Value) { return Value.ToString(); }
 		};
 
 		/** Implements variant type traits for the built-in BoxSphereBounds3f type. */
 		template<> struct TAttributeTypeTraits<FBoxSphereBounds3f>
 		{
 			static CONSTEXPR EAttributeTypes GetType() { return EAttributeTypes::BoxSphereBounds3f; }
+			static FString ToString(const FBoxSphereBounds3f& Value) { return Value.ToString(); }
 		};
 
 		/** Implements variant type traits for the built-in BoxSphereBounds3d type. */
 		template<> struct TAttributeTypeTraits<FBoxSphereBounds3d>
 		{
 			static CONSTEXPR EAttributeTypes GetType() { return EAttributeTypes::BoxSphereBounds3d; }
+			static FString ToString(const FBoxSphereBounds3d& Value) { return Value.ToString(); }
 		};
 
 		/** Implements variant type traits for the built-in Sphere3f type. */
 		template<> struct TAttributeTypeTraits<FSphere3f>
 		{
 			static CONSTEXPR EAttributeTypes GetType() { return EAttributeTypes::Sphere3f; }
+			static FString ToString(const FSphere3f& Value) { return TEXT("FSphere3f"); }
 		};
 
 		/** Implements variant type traits for the built-in Sphere3d type. */
 		template<> struct TAttributeTypeTraits<FSphere3d>
 		{
 			static CONSTEXPR EAttributeTypes GetType() { return EAttributeTypes::Sphere3d; }
+			static FString ToString(const FSphere3d& Value) { return TEXT("FSphere3d"); }
 		};
 	} //ns interchange
 }
