Interchange - Fixing import task depencencies sort.

The predicate given to the sort algorithm was not transitive and so the sort was not giving proper results.

A helper struct was added to cache the whole dependency chain for each node.

diff --git a/Engine/Source/Runtime/Experimental/Interchange/Engine/Private/Tasks/InterchangeTaskParsing.cpp b/Engine/Source/Runtime/Experimental/Interchange/Engine/Private/Tasks/InterchangeTaskParsing.cpp
--- a/Engine/Source/Runtime/Experimental/Interchange/Engine/Private/Tasks/InterchangeTaskParsing.cpp
+++ b/Engine/Source/Runtime/Experimental/Interchange/Engine/Private/Tasks/InterchangeTaskParsing.cpp
@@ -25,6 +25,53 @@
 #include "UObject/UObjectGlobals.h"
 #include "UObject/WeakObjectPtrTemplates.h"
 
+/**
+ * For the Dependency sort to work the predicate must be transitive ( A > B > C implying A > C).
+ * That means we must take into account the whole dependency chain, not just the immediate dependencies.
+ * 
+ * This is a helper struct to recursively accumulate the dependencies chain of a node quickly using a Dynamic Programming cache.
+ */
+struct FNodeDependencyCache
+{
+	const TSet<FString>& GetAccumulatedDependencies(const UInterchangeBaseNodeContainer* NodeContainer, const FString& NodeID)
+	{
+		if (TSet<FString>* DependenciesPtr = CachedDependencies.Find(NodeID))
+		{
+			return *DependenciesPtr;
+		}
+
+		TSet<FString>& Dependencies = CachedDependencies.Add(NodeID);
+		AccumulateDependencies(NodeContainer, NodeID, Dependencies);
+		return Dependencies;
+	}
+
+private:
+
+	void AccumulateDependencies(const UInterchangeBaseNodeContainer* NodeContainer, const FString& NodeID, TSet<FString>& OutDependenciesSet)
+	{
+		const UInterchangeBaseNode* Node = NodeContainer->GetNode(NodeID);
+		if (!Node)
+		{
+			return;
+		}
+
+		TArray<FString> FactoryDependencies;
+		Node->GetFactoryDependencies(FactoryDependencies);
+		for (const FString& DependencyID : FactoryDependencies)
+		{
+			bool bAlreadyInSet = false;
+			OutDependenciesSet.Add(DependencyID, &bAlreadyInSet);
+			// Avoid infinite recursion.
+			if (!bAlreadyInSet)
+			{
+				AccumulateDependencies(NodeContainer, DependencyID, OutDependenciesSet);
+			}
+		}
+	}
+
+	TMap<FString, TSet<FString>> CachedDependencies;
+};
+
 
 void UE::Interchange::FTaskParsing::DoTask(ENamedThreads::Type CurrentThread, const FGraphEventRef& MyCompletionGraphEvent)
 {
@@ -57,6 +104,9 @@ void UE::Interchange::FTaskParsing::DoTask(ENamedThreads::Type CurrentThread, co
 	{
 		for (int32 SourceIndex = 0; SourceIndex < AsyncHelper->SourceDatas.Num(); ++SourceIndex)
 		{
+			TArray<FTaskData> SourceAssetTaskDatas;
+			TArray<FTaskData> SourceSceneTaskDatas;
+
 			if (!AsyncHelper->BaseNodeContainers.IsValidIndex(SourceIndex))
 			{
 				continue;
@@ -93,7 +143,7 @@ void UE::Interchange::FTaskParsing::DoTask(ENamedThreads::Type CurrentThread, co
 						return;
 					}
 
-					FTaskData& NodeTaskData = bIsAsset ? AssetTaskDatas.AddDefaulted_GetRef() : SceneTaskDatas.AddDefaulted_GetRef();
+					FTaskData& NodeTaskData = bIsAsset ? SourceAssetTaskDatas.AddDefaulted_GetRef() : SourceSceneTaskDatas.AddDefaulted_GetRef();
 					NodeTaskData.UniqueID = Node->GetUniqueID();
 					NodeTaskData.SourceIndex = SourceIndex;
 					NodeTaskData.Nodes.Add(Node);
@@ -101,26 +151,37 @@ void UE::Interchange::FTaskParsing::DoTask(ENamedThreads::Type CurrentThread, co
 					NodeTaskData.FactoryClass = RegisteredFactoryClass;
 				}
 			});
-		}
-	}
-	//Sort per dependencies
-	auto SortByDependencies =
-		[](const FTaskData& A, const FTaskData& B)
-		{
-			//if A is a dependency of B then return true to do A before B
-			if (B.Dependencies.Contains(A.UniqueID))
-			{
-				return true;
-			}
-			if (A.Dependencies.Contains(B.UniqueID))
+
 			{
-				return false;
+				FNodeDependencyCache DependencyCache;
+
+				//Sort per dependencies
+				auto SortByDependencies =
+					[&BaseNodeContainer, &DependencyCache](const FTaskData& A, const FTaskData& B)
+				{
+					const TSet<FString>& BDependencies = DependencyCache.GetAccumulatedDependencies(BaseNodeContainer, B.UniqueID);
+					//if A is a dependency of B then return true to do A before B
+					if (BDependencies.Contains(A.UniqueID))
+					{
+						return true;
+					}
+					const TSet<FString>& ADependencies = DependencyCache.GetAccumulatedDependencies(BaseNodeContainer, A.UniqueID);
+					if (ADependencies.Contains(B.UniqueID))
+					{
+						return false;
+					}
+					return ADependencies.Num() <= BDependencies.Num();
+				};
+
+				// Nodes cannot depend on a node from another source, so it's faster to sort the dependencies per-source and then append those to the TaskData arrays.
+				SourceAssetTaskDatas.Sort(SortByDependencies);
+				SourceSceneTaskDatas.Sort(SortByDependencies);
 			}
-			return A.Dependencies.Num() <= B.Dependencies.Num();
-		};
 
-	AssetTaskDatas.Sort(SortByDependencies);
-	SceneTaskDatas.Sort(SortByDependencies);
+			AssetTaskDatas.Append(MoveTemp(SourceAssetTaskDatas));
+			SceneTaskDatas.Append(MoveTemp(SourceSceneTaskDatas));
+		}
+	}
 
 	auto CreateTasksForEachTaskData = [](TArray<FTaskData>& TaskDatas, TFunction<FGraphEventRef(FTaskData&)> CreateTasksFunc) -> FGraphEventArray
 	{
@@ -135,7 +196,7 @@ void UE::Interchange::FTaskParsing::DoTask(ENamedThreads::Type CurrentThread, co
 				//Search the previous node to find the dependence
 				for (int32 DepTaskIndex = 0; DepTaskIndex < TaskIndex; ++DepTaskIndex)
 				{
-					if (TaskData.Dependencies.Contains(TaskDatas[DepTaskIndex].UniqueID))
+					if (ensure(TaskData.Dependencies.Contains(TaskDatas[DepTaskIndex].UniqueID)))
 					{
 						//Add has prerequisite
 						TaskData.Prerequisites.Add(TaskDatas[DepTaskIndex].GraphEventRef);
