Fix children reset + shared node queue to avoid reallocation + valid numconstraints to build groups

diff --git a/Engine/Source/Runtime/Experimental/Chaos/Private/Chaos/Island/IslandGraph.cpp b/Engine/Source/Runtime/Experimental/Chaos/Private/Chaos/Island/IslandGraph.cpp
--- a/Engine/Source/Runtime/Experimental/Chaos/Private/Chaos/Island/IslandGraph.cpp
+++ b/Engine/Source/Runtime/Experimental/Chaos/Private/Chaos/Island/IslandGraph.cpp
@@ -337,18 +337,11 @@ void FIslandGraph<NodeType, EdgeType, IslandType>::MergeIslands(const int32 Pare
 	{
 		ChildQueue.Dequeue(CurrentIndex);
 	
-		if (GraphIslands[CurrentIndex].IslandCounter != GraphCounter && ParentIndex != CurrentIndex)
+		if (GraphIslands.IsValidIndex(CurrentIndex) && GraphIslands[CurrentIndex].IslandCounter != GraphCounter && ParentIndex != CurrentIndex )
 		{
 			GraphIslands[CurrentIndex].IslandCounter = GraphCounter;
 			GraphIslands[CurrentIndex].ParentIsland = ParentIndex;
 	
-			// Add the number of nodes/edges to the parent one and set to 0 the children ones
-			GraphIslands[ParentIndex].NumEdges += GraphIslands[CurrentIndex].NumEdges;
-			GraphIslands[CurrentIndex].NumEdges = 0;
-	
-			GraphIslands[ParentIndex].NumNodes += GraphIslands[CurrentIndex].NumNodes;
-			GraphIslands[CurrentIndex].NumNodes = 0;
-	
 			// Recursively iterate over all the children ones
 			for (auto& MergedIsland : GraphIslands[CurrentIndex].ChildrenIslands)
 			{
@@ -432,7 +425,7 @@ void FIslandGraph<NodeType, EdgeType, IslandType>::InitSorting()
 }
 	
 template<typename NodeType, typename EdgeType, typename IslandType>
-void FIslandGraph<NodeType, EdgeType, IslandType>::UpdateLevels(const int32 NodeIndex, const int32 ContainerId, TQueue<int32>& NodeQueue)
+void FIslandGraph<NodeType, EdgeType, IslandType>::UpdateLevels(const int32 NodeIndex, const int32 ContainerId)
 {
 	if(GraphNodes.IsValidIndex(NodeIndex))
 	{
@@ -456,7 +449,7 @@ void FIslandGraph<NodeType, EdgeType, IslandType>::UpdateLevels(const int32 Node
 				if (GraphNodes.IsValidIndex(OtherIndex) && GraphNodes[OtherIndex].bValidNode && GraphNodes[OtherIndex].LevelIndex == INDEX_NONE)
 				{
 					GraphNodes[OtherIndex].LevelIndex = GraphEdge.LevelIndex+1;
-					NodeQueue.Enqueue(OtherIndex);
+					NodeQueue.Push(OtherIndex);
 				}
 			}
 		}
@@ -467,7 +460,7 @@ template<typename NodeType, typename EdgeType, typename IslandType>
 void FIslandGraph<NodeType, EdgeType, IslandType>::ComputeLevels(const int32 ContainerId)
 {
 	// First enqueue all the static/kinematic nodes for levels 0
-	TQueue<int32> NodeQueue;
+	NodeQueue.Reset();
 	for (int32 NodeIndex = 0, NumNodes = GraphNodes.GetMaxIndex(); NodeIndex < NumNodes; ++NodeIndex)
 	{
 		if(GraphNodes.IsValidIndex(NodeIndex))
@@ -476,7 +469,7 @@ void FIslandGraph<NodeType, EdgeType, IslandType>::ComputeLevels(const int32 Con
 			if(!GraphNode.bValidNode)
 			{
 				GraphNode.LevelIndex = 0;
-				UpdateLevels(NodeIndex, ContainerId, NodeQueue);
+				UpdateLevels(NodeIndex, ContainerId);
 			}
 		}
 	}
@@ -485,8 +478,8 @@ void FIslandGraph<NodeType, EdgeType, IslandType>::ComputeLevels(const int32 Con
 	int32 NodeIndex = INDEX_NONE;
 	while(!NodeQueue.IsEmpty())
 	{
-		NodeQueue.Dequeue(NodeIndex);
-		UpdateLevels(NodeIndex, ContainerId, NodeQueue);
+		NodeIndex = NodeQueue.Pop(false);
+		UpdateLevels(NodeIndex, ContainerId);
 	}
 
 	// An isolated island that is only dynamics will not have been processed above, put everything without a level into level zero
@@ -500,7 +493,7 @@ void FIslandGraph<NodeType, EdgeType, IslandType>::ComputeLevels(const int32 Con
 }
 	
 template<typename NodeType, typename EdgeType, typename IslandType>
-int32 FIslandGraph<NodeType, EdgeType, IslandType>::PickColor(const FGraphNode& GraphNode, const int32 OtherIndex, TQueue<int32>& NodeQueue)
+int32 FIslandGraph<NodeType, EdgeType, IslandType>::PickColor(const FGraphNode& GraphNode, const int32 OtherIndex)
 {
 	int32 ColorToUse = 0;
 	if (GraphNodes.IsValidIndex(OtherIndex) && GraphNodes[OtherIndex].bValidNode)
@@ -515,7 +508,7 @@ int32 FIslandGraph<NodeType, EdgeType, IslandType>::PickColor(const FGraphNode&
 		OtherNode.ColorIndices.Add(ColorToUse);
 		if (OtherNode.NodeCounter != GraphCounter)
 		{
-			NodeQueue.Enqueue(OtherIndex);
+			NodeQueue.Push(OtherIndex);
 		}
 	}
 	else
@@ -530,7 +523,7 @@ int32 FIslandGraph<NodeType, EdgeType, IslandType>::PickColor(const FGraphNode&
 }
 
 template<typename NodeType, typename EdgeType, typename IslandType>
-void FIslandGraph<NodeType, EdgeType, IslandType>::UpdateColors(const int32 NodeIndex, const int32 ContainerId, TQueue<int32>& NodeQueue, const int32 MinEdges)
+void FIslandGraph<NodeType, EdgeType, IslandType>::UpdateColors(const int32 NodeIndex, const int32 ContainerId, const int32 MinEdges)
 {
 	if (GraphNodes.IsValidIndex(NodeIndex))
 	{
@@ -550,7 +543,7 @@ void FIslandGraph<NodeType, EdgeType, IslandType>::UpdateColors(const int32 Node
 					GraphEdge.SecondNode : GraphEdge.FirstNode;
 
 				// Get the first available color to be used by the edge
-				const int32 ColorToUse = PickColor(GraphNode, OtherIndex, NodeQueue);
+				const int32 ColorToUse = PickColor(GraphNode, OtherIndex);
 				
 				GraphNode.ColorIndices.Add(ColorToUse);
 				GraphEdge.ColorIndex = ColorToUse;
@@ -567,7 +560,7 @@ void FIslandGraph<NodeType, EdgeType, IslandType>::ComputeColors(const int32 Con
 {
 	GraphCounter = (GraphCounter + 1) % MaxCount;
 
-	TQueue<int32> NodeQueue;
+	NodeQueue.Reset();
 	int32 NodeIndex = INDEX_NONE;
 
 	// We first loop over all the nodes that have not been processed and valid (dynamic/sleeping)
@@ -578,11 +571,11 @@ void FIslandGraph<NodeType, EdgeType, IslandType>::ComputeColors(const int32 Con
 			FGraphNode& GraphNode = GraphNodes[RootIndex];
 			if (GraphNode.NodeCounter != GraphCounter && GraphNode.bValidNode)
 			{
-				NodeQueue.Enqueue(RootIndex);
+				NodeQueue.Push(RootIndex);
 				while (!NodeQueue.IsEmpty())
 				{
-					NodeQueue.Dequeue(NodeIndex);
-					UpdateColors(NodeIndex, ContainerId, NodeQueue, MinEdges);
+					NodeIndex = NodeQueue.Pop(false);
+					UpdateColors(NodeIndex, ContainerId, MinEdges);
 				}
 			}
 		}
@@ -590,14 +583,14 @@ void FIslandGraph<NodeType, EdgeType, IslandType>::ComputeColors(const int32 Con
 }
 
 template<typename NodeType, typename EdgeType, typename IslandType>
-void FIslandGraph<NodeType, EdgeType, IslandType>::SplitIsland(TQueue<int32>& NodeQueue, const int32 RootIndex, const int32 IslandIndex)
+void FIslandGraph<NodeType, EdgeType, IslandType>::SplitIsland(const int32 RootIndex, const int32 IslandIndex)
 {
-	NodeQueue.Enqueue(RootIndex);
+	NodeQueue.Push(RootIndex);
 	int32 NodeIndex = RootIndex;
 					
 	while (!NodeQueue.IsEmpty())
 	{
-		NodeQueue.Dequeue(NodeIndex);
+		NodeIndex = NodeQueue.Pop(false);
 		FGraphNode& GraphNode = GraphNodes[NodeIndex];
 						
 		// Graph counter is there to avoid processing multiple times the same node/edge
@@ -621,7 +614,7 @@ void FIslandGraph<NodeType, EdgeType, IslandType>::SplitIsland(TQueue<int32>& No
 				// Only the valid nodes are allowed to continue the graph traversal
 				if (GraphNodes.IsValidIndex(OtherIndex) && GraphNodes[OtherIndex].NodeCounter != GraphCounter && GraphNodes[OtherIndex].bValidNode)
 				{
-					NodeQueue.Enqueue(OtherIndex);
+					NodeQueue.Push(OtherIndex);
 				}
 			}
 		}
@@ -634,7 +627,7 @@ void FIslandGraph<NodeType, EdgeType, IslandType>::SplitIslands()
 	SCOPE_CYCLE_COUNTER(STAT_SplitIslandGraph);
 
 	GraphCounter = (GraphCounter + 1) % MaxCount;
-	TQueue<int32> NodeQueue;
+	NodeQueue.Reset();
 	for (int32 RootIndex = 0, NumNodes = GraphNodes.GetMaxIndex(); RootIndex < NumNodes; ++RootIndex)
 	{
 		// We pick all the nodes that are inside an island
@@ -659,7 +652,7 @@ void FIslandGraph<NodeType, EdgeType, IslandType>::SplitIslands()
 						
 						GraphIslands[CurrentIsland].IslandCounter = GraphCounter;
 
-						SplitIsland(NodeQueue, RootIndex, CurrentIsland);
+						SplitIsland(RootIndex, CurrentIsland);
 					}
 				}
 				else
@@ -782,6 +775,8 @@ void FIslandGraph<NodeType, EdgeType, IslandType>::InitIslands()
 	for (auto& GraphIsland : GraphIslands)
 	{
 		GraphIsland.bIsSleeping = true;
+		GraphIsland.ChildrenIslands.Reset();
+		GraphIsland.ParentIsland = INDEX_NONE;
 	}
 }
 
diff --git a/Engine/Source/Runtime/Experimental/Chaos/Private/Chaos/Island/IslandManager.cpp b/Engine/Source/Runtime/Experimental/Chaos/Private/Chaos/Island/IslandManager.cpp
--- a/Engine/Source/Runtime/Experimental/Chaos/Private/Chaos/Island/IslandManager.cpp
+++ b/Engine/Source/Runtime/Experimental/Chaos/Private/Chaos/Island/IslandManager.cpp
@@ -695,9 +695,15 @@ void FPBDIslandManager::BuildGroups(const int32 NumContainers)
 	TSparseArray<TUniquePtr<FPBDIslandSolver>>& LocalIslands = IslandSolvers;
 	SortedIslands.Sort([&LocalIslands](const int32 IslandIndexA, const int32 IslandIndexB) -> bool
 		{ return LocalIslands[IslandIndexA]->NumConstraints() > LocalIslands[IslandIndexB]->NumConstraints();});
+
+	int32 NumConstraints = 0;
+	for (TUniquePtr<FPBDIslandSolver>& IslandSolver : IslandSolvers)
+	{
+		NumConstraints += IslandSolver->NumConstraints();
+	}
 	
 	const int32 NumGroups = IslandGroups.Num();
-	const int32 GroupSize = IslandGraph->NumEdges() / NumGroups + 1;
+	const int32 GroupSize = NumConstraints / NumGroups + 1;
 
 	for(TUniquePtr<FPBDIslandGroup>& IslandGroup : IslandGroups)
 	{
diff --git a/Engine/Source/Runtime/Experimental/Chaos/Public/Chaos/Island/IslandGraph.h b/Engine/Source/Runtime/Experimental/Chaos/Public/Chaos/Island/IslandGraph.h
--- a/Engine/Source/Runtime/Experimental/Chaos/Public/Chaos/Island/IslandGraph.h
+++ b/Engine/Source/Runtime/Experimental/Chaos/Public/Chaos/Island/IslandGraph.h
@@ -253,7 +253,7 @@ public:
 	* @param RootIndex root index of the island graph
 	* @param IslandIndex index of the island we are splitting
 	*/
-	void SplitIsland(TQueue<int32>& NodeQueue, const int32 RootIndex, const int32 IslandIndex);
+	void SplitIsland(const int32 RootIndex, const int32 IslandIndex);
 
 	/**
 	* Split all the graph islands if not sleeping
@@ -281,9 +281,8 @@ public:
 	* into the node queue
 	* @param NodeIndex Node index that we are currently iterating over
 	* @param ContainerId Constraint container id that will be compared to the item container
-	* @param NodeQueue Node queue in which all the connected nodes not processed yet will be added
 	*/
-	void UpdateLevels(const int32 NodeIndex, const int32 ContainerId, TQueue<int32>& NodeQueue);
+	void UpdateLevels(const int32 NodeIndex, const int32 ContainerId);
 
 	/**
 	* Given a constraint containerId compute all the islands/edges/nodes colors
@@ -297,19 +296,17 @@ public:
 	* into the node queue
 	* @param NodeIndex Node index that we are currently iterating over
 	* @param ContainerId Constraint container id that will be compared to the item container
-	* @param NodeQueue Node queue in which all the connected nodes not processed yet will be added
 	* @param MinEdges Minimum number of edges to compute coloring
 	*/
-	void UpdateColors(const int32 NodeIndex, const int32 ContainerId, TQueue<int32>& NodeQueue, const int32 MinEdges);
+	void UpdateColors(const int32 NodeIndex, const int32 ContainerId,  const int32 MinEdges);
 
 	/**
 	* Pick the first available color that is not used yet by the current graph node and the edge opposite one
 	* @param GraphNode Graph node that we currently iterate over
 	* @param OtherIndex Index of the edge opposite node
-	* @param NodeQueue Node queue in which all the connected nodes not processed yet will be added
 	* @return First available color 
 	*/
-	int32 PickColor(const FGraphNode& GraphNode, const int32 OtherIndex, TQueue<int32>& NodeQueue);
+	int32 PickColor(const FGraphNode& GraphNode, const int32 OtherIndex);
 
 	/**
 	* Check if at least one of the 2 edge nodes is moving
@@ -346,6 +343,9 @@ public:
 
 	/** Graph counter used for graph traversal to check if an edge/node/island has already been processed */
 	int32 GraphCounter = 0;
+
+	/** List of nodes to be processed in the graph functions */
+	TArray<int32> NodeQueue;
 };
 
 }
\ No newline at end of file
