Chaos - Fix swept collision detection sweep length (CCD issue)

diff --git a/Engine/Source/Runtime/Experimental/Chaos/Private/Chaos/Collision/ParticlePairMidPhase.cpp b/Engine/Source/Runtime/Experimental/Chaos/Private/Chaos/Collision/ParticlePairMidPhase.cpp
--- a/Engine/Source/Runtime/Experimental/Chaos/Private/Chaos/Collision/ParticlePairMidPhase.cpp
+++ b/Engine/Source/Runtime/Experimental/Chaos/Private/Chaos/Collision/ParticlePairMidPhase.cpp
@@ -822,9 +822,9 @@ namespace Chaos
 
 			FReal LengthCCD = 0;
 			FVec3 DirCCD = FVec3(0);
-			const FVec3 StartX0 = (ConstParticle0->ObjectState() == EObjectStateType::Kinematic) ? ConstParticle0->P() - ConstParticle0->V() * Dt : ConstParticle0->X();
-			const FVec3 StartX1 = (ConstParticle1->ObjectState() == EObjectStateType::Kinematic) ? ConstParticle1->P() - ConstParticle1->V() * Dt : ConstParticle1->X();
-			const bool bUseCCD = Collisions::ShouldUseCCD(Particle0, StartX0, Particle1, StartX1, DirCCD, LengthCCD, false);
+			const FVec3 DeltaX0 = (ConstParticle0->ObjectState() == EObjectStateType::Kinematic) ? (ConstParticle0->V() * Dt) : (ConstParticle0->P() - ConstParticle0->X());
+			const FVec3 DeltaX1 = (ConstParticle1->ObjectState() == EObjectStateType::Kinematic) ? (ConstParticle1->V() * Dt) : (ConstParticle1->P() - ConstParticle1->X());
+			const bool bUseCCD = Collisions::ShouldUseCCD(Particle0, DeltaX0, Particle1, DeltaX1, DirCCD, LengthCCD, false);
 
 			return bUseCCD;
 		}
diff --git a/Engine/Source/Runtime/Experimental/Chaos/Private/Chaos/CollisionResolution.cpp b/Engine/Source/Runtime/Experimental/Chaos/Private/Chaos/CollisionResolution.cpp
--- a/Engine/Source/Runtime/Experimental/Chaos/Private/Chaos/CollisionResolution.cpp
+++ b/Engine/Source/Runtime/Experimental/Chaos/Private/Chaos/CollisionResolution.cpp
@@ -153,7 +153,7 @@ namespace Chaos
 		* Now think about two boxes. If P-X <= 0.5 * BoundingBox().Extents().Min() both boxes, it is impossible for one box to tunnel through the other (it is still possible for a box's corner to move through another box's corner).
 		* So far this is sufficient for us to determine CCD constraints for our primitives. But the caveat is that for arbitrary objects (think about a rod that is not axis aligned), min of a bounding box extent can be too large and tunneling could still happen if P-X falls below this threshold.
 		*/
-		bool ShouldUseCCD(const FGeometryParticleHandle* Particle0, const FVec3& StartX0, const FGeometryParticleHandle* Particle1, const FVec3& StartX1, FVec3& Dir, FReal& Length, const bool bForceDisableCCD)
+		bool ShouldUseCCD(const FGeometryParticleHandle* Particle0, const FVec3& DeltaX0, const FGeometryParticleHandle* Particle1, const FVec3& DeltaX1, FVec3& Dir, FReal& Length, const bool bForceDisableCCD)
 		{
 			if (CCDManualForceDisable)
 			{
@@ -169,6 +169,7 @@ namespace Chaos
 			{
 				return false;
 			}
+
 			const TPBDRigidParticleHandle<FReal, 3>* Rigid0 = Particle0->CastToRigidParticle();
 			const bool bIsCCDEnabled0 = Rigid0 && Rigid0->CCDEnabled();
 			const TPBDRigidParticleHandle<FReal, 3>* Rigid1 = Particle1->CastToRigidParticle();
@@ -179,15 +180,13 @@ namespace Chaos
 			}
 
 			bool bUseCCD = false;
-			const FVec3 D0 = Rigid0 ? Rigid0->P() - StartX0 : FVec3(0.);
-			const FVec3 D1 = Rigid1 ? Rigid1->P() - StartX1 : FVec3(0.);
 
 			// bUseCCD is determined by the difference between P and X rather than V.
 			// This implies that the smaller Dt is, the less likely CCD will be triggered.
 			// This is because CCD is enabled to prevent tunneling. When Dt it small, tunneling is less likely to happen and therefore fewer constraints need to be processed with CCD.
 			if (bIsCCDEnabled0)
 			{
-				const FReal DSizeSquared0 = D0.SizeSquared();
+				const FReal DSizeSquared0 = DeltaX0.SizeSquared();
 				const FReal MinExtent0 = Particle0->LocalBounds().Extents().Min();
 				const FReal CCDThreshold0 = MinExtent0 * CVars::CCDEnableThresholdBoundsScale;
 				if (DSizeSquared0 > CCDThreshold0 * CCDThreshold0)
@@ -197,7 +196,7 @@ namespace Chaos
 			}
 			if (!bUseCCD && bIsCCDEnabled1)
 			{
-				const FReal DSizeSquared1 = D1.SizeSquared();
+				const FReal DSizeSquared1 = DeltaX1.SizeSquared();
 				const FReal MinExtent1 = Particle1->LocalBounds().Extents().Min();
 				const FReal CCDThreshold1 = MinExtent1 * CVars::CCDEnableThresholdBoundsScale;
 				if (DSizeSquared1 > CCDThreshold1 * CCDThreshold1)
@@ -208,13 +207,15 @@ namespace Chaos
 
 			if (bUseCCD)
 			{
-				Dir = D0 - D1;
+				Dir = DeltaX0 - DeltaX1;
 				Length = Dir.SafeNormalize();
 				if (Length < SMALL_NUMBER) // This is the case where both particles' velocities are high but the relative velocity is low. 
 				{
 					bUseCCD = false;
 				}
+
 				// Do not perform CCD if the vector is not close to unit length to prevent getting caught in a large or infinite loop when raycasting.
+				// @todo(chaos): What's this? Dir is normalized above so this seems unnecessary unless it is checking for 0 length, not unit length - remove it or fix
 				if (!FMath::IsNearlyEqual((float)Dir.SizeSquared(), 1.f, KINDA_SMALL_NUMBER))
 				{
 					bUseCCD = false;
@@ -1769,7 +1770,9 @@ namespace Chaos
 
 			FReal LengthCCD = 0.0f;
 			FVec3 DirCCD(0.0f);
-			bool bUseCCD = ShouldUseCCD(Particle0, WorldTransform0.GetLocation(), Particle1, WorldTransform1.GetLocation(), DirCCD, LengthCCD, false);
+			const FVec3 DeltaX0 = Constraint.GetShapeWorldTransform0().GetTranslation() - WorldTransform0.GetLocation();
+			const FVec3 DeltaX1 = Constraint.GetShapeWorldTransform1().GetTranslation() - WorldTransform1.GetLocation();
+			bool bUseCCD = ShouldUseCCD(Particle0, DeltaX0, Particle1, DeltaX1, DirCCD, LengthCCD, false);
 
 			if (bUseCCD)
 			{
@@ -1888,10 +1891,13 @@ namespace Chaos
 			FVec3 DirCCD(0.0f);
 			FConstGenericParticleHandle ConstParticle0 = FConstGenericParticleHandle(Particle0);
 			FConstGenericParticleHandle ConstParticle1 = FConstGenericParticleHandle(Particle1);
+
 			// X is end-frame position for kinematic particles. To get start-frame position, we need to use P - V * Dt.
-			const FVec3 StartX0 = ConstParticle0->ObjectState() == EObjectStateType::Kinematic ? ConstParticle0->P() - ConstParticle0->V() * Dt : ConstParticle0->X();
-			const FVec3 StartX1 = ConstParticle1->ObjectState() == EObjectStateType::Kinematic ? ConstParticle1->P() - ConstParticle1->V() * Dt : ConstParticle1->X();
-			bool bUseCCD = ShouldUseCCD(Particle0, StartX0, Particle1, StartX1, DirCCD, LengthCCD, Context.bForceDisableCCD);
+			const FVec3 StartX0 = (ConstParticle0->ObjectState() == EObjectStateType::Kinematic) ? (ConstParticle0->P() - ConstParticle0->V() * Dt) : ConstParticle0->X();
+			const FVec3 StartX1 = (ConstParticle1->ObjectState() == EObjectStateType::Kinematic) ? (ConstParticle1->P() - ConstParticle1->V() * Dt) : ConstParticle1->X();
+			const FVec3 EndX0 = ConstParticle0->P();
+			const FVec3 EndX1 = ConstParticle1->P();
+			bool bUseCCD = ShouldUseCCD(Particle0, EndX0 - StartX0, Particle1, EndX1 - StartX1, DirCCD, LengthCCD, Context.bForceDisableCCD);
 			const bool bUseGenericSweptConstraints = bUseCCD && (CCDUseGenericSweptConvexConstraints > 0) && bIsConvex0 && bIsConvex1;
 #if CHAOS_COLLISION_CREATE_BOUNDSCHECK
 			if ((Implicit0 != nullptr) && (Implicit1 != nullptr))
diff --git a/Engine/Source/Runtime/Experimental/Chaos/Public/Chaos/CollisionResolution.h b/Engine/Source/Runtime/Experimental/Chaos/Public/Chaos/CollisionResolution.h
--- a/Engine/Source/Runtime/Experimental/Chaos/Public/Chaos/CollisionResolution.h
+++ b/Engine/Source/Runtime/Experimental/Chaos/Public/Chaos/CollisionResolution.h
@@ -58,7 +58,7 @@ namespace Chaos
 		/**
 		 * @brief Whether CCD should be enabled for a contact given the current particle velocities etc
 		*/
-		bool CHAOS_API ShouldUseCCD(const FGeometryParticleHandle* Particle0, const FVec3& StartX0, const FGeometryParticleHandle* Particle1, const FVec3& StartX1, FVec3& Dir, FReal& Length, const bool bForceDisableCCD);
+		bool CHAOS_API ShouldUseCCD(const FGeometryParticleHandle* Particle0, const FVec3& DeltaX0, const FGeometryParticleHandle* Particle1, const FVec3& DeltaX1, FVec3& Dir, FReal& Length, const bool bForceDisableCCD);
 
 		// Update the constraint by re-running collision detection on the shape pair.
 		template<ECollisionUpdateType UpdateType>
