Remove unneded ensure

diff --git a/Engine/Source/Runtime/Experimental/Interchange/Engine/Private/Tasks/InterchangeTaskParsing.cpp b/Engine/Source/Runtime/Experimental/Interchange/Engine/Private/Tasks/InterchangeTaskParsing.cpp
--- a/Engine/Source/Runtime/Experimental/Interchange/Engine/Private/Tasks/InterchangeTaskParsing.cpp
+++ b/Engine/Source/Runtime/Experimental/Interchange/Engine/Private/Tasks/InterchangeTaskParsing.cpp
@@ -196,7 +196,7 @@ void UE::Interchange::FTaskParsing::DoTask(ENamedThreads::Type CurrentThread, co
 				//Search the previous node to find the dependence
 				for (int32 DepTaskIndex = 0; DepTaskIndex < TaskIndex; ++DepTaskIndex)
 				{
-					if (ensure(TaskData.Dependencies.Contains(TaskDatas[DepTaskIndex].UniqueID)))
+					if (TaskData.Dependencies.Contains(TaskDatas[DepTaskIndex].UniqueID))
 					{
 						//Add has prerequisite
 						TaskData.Prerequisites.Add(TaskDatas[DepTaskIndex].GraphEventRef);
