Chaos : GeometryCollection
- Added FIntVector4 array attribute support.

diff --git a/Engine/Source/Runtime/Experimental/Chaos/Public/GeometryCollection/ManagedArrayTypeValues.inl b/Engine/Source/Runtime/Experimental/Chaos/Public/GeometryCollection/ManagedArrayTypeValues.inl
--- a/Engine/Source/Runtime/Experimental/Chaos/Public/GeometryCollection/ManagedArrayTypeValues.inl
+++ b/Engine/Source/Runtime/Experimental/Chaos/Public/GeometryCollection/ManagedArrayTypeValues.inl
@@ -42,6 +42,7 @@ MANAGED_ARRAY_TYPE(Chaos::FPBDRigidClusteredParticleHandle*, TPBDRigidClusteredP
 MANAGED_ARRAY_TYPE(TUniquePtr<Chaos::FConvex>, FConvexUniquePtr)
 MANAGED_ARRAY_TYPE(TArray<FVector2f>, Vector2DArray)
 MANAGED_ARRAY_TYPE(double, Double)
+MANAGED_ARRAY_TYPE(FIntVector4, IntVector4)
 
 // NOTE: new types must be added at the bottom to keep serialization from breaking
 
