Chaos : GeometryCollection
- Added static Init to expose initialization on existing collection.

diff --git a/Engine/Source/Runtime/Experimental/Chaos/Private/GeometryCollection/GeometryCollection.cpp b/Engine/Source/Runtime/Experimental/Chaos/Private/GeometryCollection/GeometryCollection.cpp
--- a/Engine/Source/Runtime/Experimental/Chaos/Private/GeometryCollection/GeometryCollection.cpp
+++ b/Engine/Source/Runtime/Experimental/Chaos/Private/GeometryCollection/GeometryCollection.cpp
@@ -1341,123 +1341,127 @@ bool FGeometryCollection::IsVisible(int32 Element) const
 	return false;;
 }
 
-
 FGeometryCollection* FGeometryCollection::NewGeometryCollection(const TArray<float>& RawVertexArray, const TArray<int32>& RawIndicesArray, bool ReverseVertexOrder)
 {
+	FGeometryCollection* Collection = new FGeometryCollection();
+	FGeometryCollection::Init(Collection, RawVertexArray, RawIndicesArray, ReverseVertexOrder);
+	return Collection;
+}
 
-	FGeometryCollection* RestCollection = new FGeometryCollection();
-
-	int NumNewVertices = RawVertexArray.Num() / 3;
-	int VerticesIndex = RestCollection->AddElements(NumNewVertices, FGeometryCollection::VerticesGroup);
-	
-	int NumNewIndices = RawIndicesArray.Num() / 3;
-	int IndicesIndex = RestCollection->AddElements(NumNewIndices, FGeometryCollection::FacesGroup);
-	
-	int NumNewParticles = 1; // 1 particle for this geometry structure
-	int ParticlesIndex = RestCollection->AddElements(NumNewParticles, FGeometryCollection::TransformGroup);
-
-	TManagedArray<FVector3f>& Vertices = RestCollection->Vertex;
-	TManagedArray<FVector3f>&  Normals = RestCollection->Normal;
-	TManagedArray<FVector3f>&  TangentU = RestCollection->TangentU;
-	TManagedArray<FVector3f>&  TangentV = RestCollection->TangentV;
-	TManagedArray<TArray<FVector2f>>& UVs = RestCollection->UVs;
-	TManagedArray<FLinearColor>&  Colors = RestCollection->Color;
-	TManagedArray<FIntVector>&  Indices = RestCollection->Indices;
-	TManagedArray<bool>&  Visible = RestCollection->Visible;
-	TManagedArray<int32>&  MaterialID = RestCollection->MaterialID;
-	TManagedArray<int32>&  MaterialIndex = RestCollection->MaterialIndex;
-	TManagedArray<FTransform>&  Transform = RestCollection->Transform;
-
-	// set the vertex information
-	FVector3d TempVertices(0.f, 0.f, 0.f);
-	for (int32 Idx = 0; Idx < NumNewVertices; ++Idx)
+void FGeometryCollection::Init(FGeometryCollection* Collection, const TArray<float>& RawVertexArray, const TArray<int32>& RawIndicesArray, bool ReverseVertexOrder)
+{
+	if (Collection)
 	{
-		Vertices[Idx] = FVector3f(RawVertexArray[3 * Idx], RawVertexArray[3 * Idx + 1], RawVertexArray[3 * Idx + 2]);
-		TempVertices += FVector3d(Vertices[Idx]);
+		int NumNewVertices = RawVertexArray.Num() / 3;
+		int VerticesIndex = Collection->AddElements(NumNewVertices, FGeometryCollection::VerticesGroup);
 
-		UVs[Idx].SetNumZeroed(GeometryCollectionUV::MAX_NUM_UV_CHANNELS);
-		Colors[Idx] = FLinearColor::White;
-	}
+		int NumNewIndices = RawIndicesArray.Num() / 3;
+		int IndicesIndex = Collection->AddElements(NumNewIndices, FGeometryCollection::FacesGroup);
 
-	// set the particle information
-	TempVertices /= (float)NumNewVertices;
-	Transform[0] = FTransform(TempVertices);
-	Transform[0].NormalizeRotation();
+		int NumNewParticles = 1; // 1 particle for this geometry structure
+		int ParticlesIndex = Collection->AddElements(NumNewParticles, FGeometryCollection::TransformGroup);
 
-	// set the index information
-	TArray<FVector3f> FaceNormals;
-	FaceNormals.SetNum(NumNewIndices);
-	for (int32 Idx = 0; Idx < NumNewIndices; ++Idx)
-	{
-		int32 VertexIdx1, VertexIdx2, VertexIdx3;
-		if (!ReverseVertexOrder)
-		{
-			VertexIdx1 = RawIndicesArray[3 * Idx];
-			VertexIdx2 = RawIndicesArray[3 * Idx + 1];
-			VertexIdx3 = RawIndicesArray[3 * Idx + 2];
-		}
-		else
+		TManagedArray<FVector3f>& Vertices = Collection->Vertex;
+		TManagedArray<FVector3f>& Normals = Collection->Normal;
+		TManagedArray<FVector3f>& TangentU = Collection->TangentU;
+		TManagedArray<FVector3f>& TangentV = Collection->TangentV;
+		TManagedArray<TArray<FVector2f>>& UVs = Collection->UVs;
+		TManagedArray<FLinearColor>& Colors = Collection->Color;
+		TManagedArray<FIntVector>& Indices = Collection->Indices;
+		TManagedArray<bool>& Visible = Collection->Visible;
+		TManagedArray<int32>& MaterialID = Collection->MaterialID;
+		TManagedArray<int32>& MaterialIndex = Collection->MaterialIndex;
+		TManagedArray<FTransform>& Transform = Collection->Transform;
+
+		// set the vertex information
+		FVector3d TempVertices(0.f, 0.f, 0.f);
+		for (int32 Idx = 0; Idx < NumNewVertices; ++Idx)
 		{
-			VertexIdx1 = RawIndicesArray[3 * Idx];
-			VertexIdx2 = RawIndicesArray[3 * Idx + 2];
-			VertexIdx3 = RawIndicesArray[3 * Idx + 1];
+			Vertices[Idx] = FVector3f(RawVertexArray[3 * Idx], RawVertexArray[3 * Idx + 1], RawVertexArray[3 * Idx + 2]);
+			TempVertices += FVector3d(Vertices[Idx]);
+
+			UVs[Idx].SetNumZeroed(GeometryCollectionUV::MAX_NUM_UV_CHANNELS);
+			Colors[Idx] = FLinearColor::White;
 		}
 
-		Indices[Idx] = FIntVector(VertexIdx1, VertexIdx2, VertexIdx3);
-		Visible[Idx] = true;
-		MaterialID[Idx] = 0;
-		MaterialIndex[Idx] = Idx;
+		// set the particle information
+		TempVertices /= (float)NumNewVertices;
+		Transform[0] = FTransform(TempVertices);
+		Transform[0].NormalizeRotation();
 
-		const FVector3f Edge1 = Vertices[VertexIdx1] - Vertices[VertexIdx2];
-		const FVector3f Edge2 = Vertices[VertexIdx1] - Vertices[VertexIdx3];
-		FaceNormals[Idx] = (Edge2 ^ Edge1).GetSafeNormal();
-	}
+		// set the index information
+		TArray<FVector3f> FaceNormals;
+		FaceNormals.SetNum(NumNewIndices);
+		for (int32 Idx = 0; Idx < NumNewIndices; ++Idx)
+		{
+			int32 VertexIdx1, VertexIdx2, VertexIdx3;
+			if (!ReverseVertexOrder)
+			{
+				VertexIdx1 = RawIndicesArray[3 * Idx];
+				VertexIdx2 = RawIndicesArray[3 * Idx + 1];
+				VertexIdx3 = RawIndicesArray[3 * Idx + 2];
+			}
+			else
+			{
+				VertexIdx1 = RawIndicesArray[3 * Idx];
+				VertexIdx2 = RawIndicesArray[3 * Idx + 2];
+				VertexIdx3 = RawIndicesArray[3 * Idx + 1];
+			}
 
-	// Compute vertexNormals
-	TArray<FVector3f> VertexNormals;
-	VertexNormals.SetNum(NumNewVertices);
-	for (int32 Idx = 0; Idx < NumNewVertices; ++Idx)
-	{
-		VertexNormals[Idx] = FVector3f(0.f, 0.f, 0.f);
-	}
+			Indices[Idx] = FIntVector(VertexIdx1, VertexIdx2, VertexIdx3);
+			Visible[Idx] = true;
+			MaterialID[Idx] = 0;
+			MaterialIndex[Idx] = Idx;
 
-	for (int32 Idx = 0; Idx < NumNewIndices; ++Idx)
-	{
-		VertexNormals[Indices[Idx][0]] += FaceNormals[Idx];
-		VertexNormals[Indices[Idx][1]] += FaceNormals[Idx];
-		VertexNormals[Indices[Idx][2]] += FaceNormals[Idx];
-	}
+			const FVector3f Edge1 = Vertices[VertexIdx1] - Vertices[VertexIdx2];
+			const FVector3f Edge2 = Vertices[VertexIdx1] - Vertices[VertexIdx3];
+			FaceNormals[Idx] = (Edge2 ^ Edge1).GetSafeNormal();
+		}
 
-	for (int32 Idx = 0; Idx < NumNewVertices; ++Idx)
-	{
-		Normals[Idx] = (VertexNormals[Idx] / 3.f).GetSafeNormal();
-	}
+		// Compute vertexNormals
+		TArray<FVector3f> VertexNormals;
+		VertexNormals.SetNum(NumNewVertices);
+		for (int32 Idx = 0; Idx < NumNewVertices; ++Idx)
+		{
+			VertexNormals[Idx] = FVector3f(0.f, 0.f, 0.f);
+		}
 
-	for (int IndexIdx = 0; IndexIdx < NumNewIndices; IndexIdx++)
-	{
-		FIntVector Tri = Indices[IndexIdx];
-		for (int idx = 0; idx < 3; idx++)
+		for (int32 Idx = 0; Idx < NumNewIndices; ++Idx)
 		{
-			const FVector3f Normal = Normals[Tri[idx]];
-			const FVector3f Edge = (Vertices[Tri[(idx + 1) % 3]] - Vertices[Tri[idx]]);
-			TangentU[Tri[idx]] = (Edge ^ Normal).GetSafeNormal();
-			TangentV[Tri[idx]] = (Normal ^ TangentU[Tri[idx]]).GetSafeNormal();
+			VertexNormals[Indices[Idx][0]] += FaceNormals[Idx];
+			VertexNormals[Indices[Idx][1]] += FaceNormals[Idx];
+			VertexNormals[Indices[Idx][2]] += FaceNormals[Idx];
 		}
-	}
 
-	// Build the Geometry Group
-	GeometryCollection::AddGeometryProperties(RestCollection);
+		for (int32 Idx = 0; Idx < NumNewVertices; ++Idx)
+		{
+			Normals[Idx] = (VertexNormals[Idx] / 3.f).GetSafeNormal();
+		}
 
-	// add a material section
-	TManagedArray<FGeometryCollectionSection>&  Sections = RestCollection->Sections;
-	int Element = RestCollection->AddElements(1, FGeometryCollection::MaterialGroup);
-	Sections[Element].MaterialID = 0;
-	Sections[Element].FirstIndex = 0;
-	Sections[Element].NumTriangles = Indices.Num();
-	Sections[Element].MinVertexIndex = 0;
-	Sections[Element].MaxVertexIndex = Vertices.Num() - 1;
+		for (int IndexIdx = 0; IndexIdx < NumNewIndices; IndexIdx++)
+		{
+			FIntVector Tri = Indices[IndexIdx];
+			for (int idx = 0; idx < 3; idx++)
+			{
+				const FVector3f Normal = Normals[Tri[idx]];
+				const FVector3f Edge = (Vertices[Tri[(idx + 1) % 3]] - Vertices[Tri[idx]]);
+				TangentU[Tri[idx]] = (Edge ^ Normal).GetSafeNormal();
+				TangentV[Tri[idx]] = (Normal ^ TangentU[Tri[idx]]).GetSafeNormal();
+			}
+		}
 
-	return RestCollection;
+		// Build the Geometry Group
+		GeometryCollection::AddGeometryProperties(Collection);
+
+		// add a material section
+		TManagedArray<FGeometryCollectionSection>& Sections = Collection->Sections;
+		int Element = Collection->AddElements(1, FGeometryCollection::MaterialGroup);
+		Sections[Element].MaterialID = 0;
+		Sections[Element].FirstIndex = 0;
+		Sections[Element].NumTriangles = Indices.Num();
+		Sections[Element].MinVertexIndex = 0;
+		Sections[Element].MaxVertexIndex = Vertices.Num() - 1;
+	}
 }
 
 void FGeometryCollection::WriteDataToHeaderFile(const FString &Name, const FString &Path)
diff --git a/Engine/Source/Runtime/Experimental/Chaos/Public/GeometryCollection/GeometryCollection.h b/Engine/Source/Runtime/Experimental/Chaos/Public/GeometryCollection/GeometryCollection.h
--- a/Engine/Source/Runtime/Experimental/Chaos/Public/GeometryCollection/GeometryCollection.h
+++ b/Engine/Source/Runtime/Experimental/Chaos/Public/GeometryCollection/GeometryCollection.h
@@ -126,6 +126,7 @@ public:
 	 * Create a GeometryCollection from Vertex and Indices arrays
 	 */
 	static FGeometryCollection* NewGeometryCollection(const TArray<float>& RawVertexArray, const TArray<int32>& RawIndicesArray, bool ReverseVertexOrder = true);
+	static void Init(FGeometryCollection* Collection, const TArray<float>& RawVertexArray, const TArray<int32>& RawIndicesArray, bool ReverseVertexOrder = true);
 
 	/**
 	* Create a GeometryCollection from Vertex, Indices, BoneMap, Transform, BoneHierarchy arrays
