Fix UInterchangeBaseNodeContainer::GetNode() const-correctness.

diff --git a/Engine/Source/Runtime/Experimental/Interchange/Core/Private/Nodes/InterchangeBaseNodeContainer.cpp b/Engine/Source/Runtime/Experimental/Interchange/Core/Private/Nodes/InterchangeBaseNodeContainer.cpp
--- a/Engine/Source/Runtime/Experimental/Interchange/Core/Private/Nodes/InterchangeBaseNodeContainer.cpp
+++ b/Engine/Source/Runtime/Experimental/Interchange/Core/Private/Nodes/InterchangeBaseNodeContainer.cpp
@@ -99,7 +99,7 @@ void UInterchangeBaseNodeContainer::GetNodes(const UClass* ClassNode, TArray<FSt
 	});
 }
 
-UInterchangeBaseNode* UInterchangeBaseNodeContainer::GetNode(const FString& NodeUniqueID)
+UInterchangeBaseNode* UInterchangeBaseNodeContainer::GetNode(const FString& NodeUniqueID) const
 {
 	if (NodeUniqueID == UInterchangeBaseNode::InvalidNodeUid())
 	{
@@ -113,11 +113,6 @@ UInterchangeBaseNode* UInterchangeBaseNodeContainer::GetNode(const FString& Node
 	return Node;
 }
 
-const UInterchangeBaseNode* UInterchangeBaseNodeContainer::GetNode(const FString& NodeUniqueID) const
-{
-	return const_cast<UInterchangeBaseNodeContainer*>(this)->GetNode(NodeUniqueID);
-}
-
 bool UInterchangeBaseNodeContainer::SetNodeParentUid(const FString& NodeUniqueID, const FString& NewParentNodeUid)
 {
 	if (!Nodes.Contains(NodeUniqueID))
diff --git a/Engine/Source/Runtime/Experimental/Interchange/Core/Public/Nodes/InterchangeBaseNodeContainer.h b/Engine/Source/Runtime/Experimental/Interchange/Core/Public/Nodes/InterchangeBaseNodeContainer.h
--- a/Engine/Source/Runtime/Experimental/Interchange/Core/Public/Nodes/InterchangeBaseNodeContainer.h
+++ b/Engine/Source/Runtime/Experimental/Interchange/Core/Public/Nodes/InterchangeBaseNodeContainer.h
@@ -56,10 +56,7 @@ public:
 
 	/** Get an node pointer */
 	UFUNCTION(BlueprintCallable, Category = "Interchange | Node Container")
-	UInterchangeBaseNode* GetNode(const FString& NodeUniqueID);
-
-	/** Get an node pointer */
-	const UInterchangeBaseNode* GetNode(const FString& NodeUniqueID) const;
+	UInterchangeBaseNode* GetNode(const FString& NodeUniqueID) const;
 
 	/** Set node ParentUid */
 	UFUNCTION(BlueprintCallable, Category = "Interchange | Node Container")
