Chaos : fix issue where adding an element to the size specific array from the UI adds sometimes two

diff --git a/Engine/Source/Runtime/Experimental/GeometryCollectionEngine/Private/GeometryCollection/GeometryCollectionObject.cpp b/Engine/Source/Runtime/Experimental/GeometryCollectionEngine/Private/GeometryCollection/GeometryCollectionObject.cpp
--- a/Engine/Source/Runtime/Experimental/GeometryCollectionEngine/Private/GeometryCollection/GeometryCollectionObject.cpp
+++ b/Engine/Source/Runtime/Experimental/GeometryCollectionEngine/Private/GeometryCollection/GeometryCollectionObject.cpp
@@ -62,6 +62,8 @@ namespace GeometryCollectionCookStats
 }
 #endif
 
+static constexpr float DefaultMaxSizeValue = 99999.9;
+
 UGeometryCollection::UGeometryCollection(const FObjectInitializer& ObjectInitializer)
 	: Super(ObjectInitializer)
 	#if WITH_EDITOR
@@ -128,7 +130,7 @@ FGeometryCollectionCollisionTypeData::FGeometryCollectionCollisionTypeData()
 }
 
 FGeometryCollectionSizeSpecificData::FGeometryCollectionSizeSpecificData()
-	: MaxSize(99999.9)
+	: MaxSize(DefaultMaxSizeValue)
 	, CollisionShapes({ FGeometryCollectionCollisionTypeData()})
 #if WITH_EDITORONLY_DATA
 	, CollisionType_DEPRECATED(ECollisionTypeEnum::Chaos_Volumetric)
@@ -215,7 +217,7 @@ FGeometryCollectionSizeSpecificData UGeometryCollection::GeometryCollectionSizeS
 {
 	FGeometryCollectionSizeSpecificData Data;
 
-	Data.MaxSize = 99999.9;
+	Data.MaxSize = DefaultMaxSizeValue;
 	if (Data.CollisionShapes.Num())
 	{
 		Data.CollisionShapes[0].CollisionType = ECollisionTypeEnum::Chaos_Volumetric;
@@ -240,7 +242,7 @@ void UGeometryCollection::ValidateSizeSpecificDataDefaults()
 	{
 		for (const FGeometryCollectionSizeSpecificData& Data : DatasIn)
 		{
-			if (Data.MaxSize >= 99999.9)
+			if (Data.MaxSize >= DefaultMaxSizeValue)
 			{
 				return true;
 			}
@@ -329,7 +331,7 @@ void UGeometryCollection::GetSharedSimulationParams(FSharedSimulationParameters&
 		InfSize.CollisionShapes[0].CollisionParticles.CollisionParticlesFraction = SizeSpecificDefault.CollisionShapes[0].CollisionParticles.CollisionParticlesFraction;
 		InfSize.CollisionShapes[0].CollisionParticles.MaximumCollisionParticles = SizeSpecificDefault.CollisionShapes[0].CollisionParticles.MaximumCollisionParticles;
 	}
-	InfSize.MaxSize = 99999.9;
+	InfSize.MaxSize = DefaultMaxSizeValue;
 	OutParams.SizeSpecificData.SetNum(SizeSpecificData.Num() + 1);
 	FillSharedSimulationSizeSpecificData(OutParams.SizeSpecificData[0], InfSize);
 
