Chaos : fix shape generation not properly using relative size for size specific data when clustering is on

diff --git a/Engine/Source/Runtime/Experimental/Chaos/Private/PhysicsProxy/GeometryCollectionPhysicsProxy.cpp b/Engine/Source/Runtime/Experimental/Chaos/Private/PhysicsProxy/GeometryCollectionPhysicsProxy.cpp
--- a/Engine/Source/Runtime/Experimental/Chaos/Private/PhysicsProxy/GeometryCollectionPhysicsProxy.cpp
+++ b/Engine/Source/Runtime/Experimental/Chaos/Private/PhysicsProxy/GeometryCollectionPhysicsProxy.cpp
@@ -2540,7 +2540,17 @@ void FGeometryCollectionPhysicsProxy::InitializeSharedCollisionStructures(
 				CollectionInertiaTensor[ClusterTransformIdx] = FVector3f(InertiaDiagonal);	// LWC_TODO: Precision loss
 				CollectionMass[ClusterTransformIdx] = (FRealSingle)CollectionSpaceParticles->M(ClusterTransformIdx);
 
-				const int32 SizeSpecificIdx = GeometryCollection::SizeSpecific::FindIndexForVolume(SharedParams.SizeSpecificData, InstanceBoundingBox);
+
+				int32 SizeSpecificIdx;
+				if (bUseRelativeSize)
+				{
+					const TManagedArray<float>& RelativeSize = RestCollection.GetAttribute<float>(TEXT("Size"), FTransformCollection::TransformGroup);
+					SizeSpecificIdx = GeometryCollection::SizeSpecific::FindIndexForVolume(SharedParams.SizeSpecificData, RelativeSize[TransformGroupIndex]);
+				}
+				else
+				{
+					SizeSpecificIdx = GeometryCollection::SizeSpecific::FindIndexForVolume(SharedParams.SizeSpecificData, InstanceBoundingBox);	
+				}
 				const FSharedSimulationSizeSpecificData& SizeSpecificData = SharedParams.SizeSpecificData[SizeSpecificIdx];
 
 				if (SizeSpecificData.CollisionShapesData.Num())
