Chaos Cloth-cloth collision improvements. Self intersection detection and resolution, improved proximity repulsions with friction.

==OVERVIEW==
The core of this change is the addition of PBDTriangleMeshCollisions and PBDTriangleMeshIntersections. PBDTriangleMeshCollisions does "Init" setup for cloth-cloth collisions, both proximity repulsions and intersections. PBDTriangleMeshIntersections is a new constraint that does Contour Minimization to resolve certain types of intersections.
Additionally, proximity repulsions have been updated to take normals into account, this time with greater stability (more below).
Cloth constraint rules have been reordered--self collision repulsions now run last, AFTER kinematic collisions. This is typically more desirable in practice since often the character body is invisible in areas where there should be clothes.

==PERF AND MEMORY==
When doing self collisions, we now store the BVH. Additionally, per vertex and per triangle GIAColor(8 bytes) arrays are stored.
BVH building per frame is expensive (in Development editor, 650us for 2800 tris). There should be follow up work to make this more efficient. This is not new to this CL--we use the same BVH for intersection and proximity detection.
Intersection detection costs 1.5-2x as much as proximity detection (basically just the difference between number of edges vs number of vertices).
Cost of Intersection analysis depends on how many intersections are actually detected. If repulsions are doing a good job, you should not have a lot. In tests with a fairly large number of intersections (but not a "failure" case where everything has tangled/exploded), intersection analysis costs about the same amount as intersection detection. This should be considered a "typical" upper bound on cost.
In all, a rough heuristic is that old self collision repulsions (i.e, NOT this CL) increase the cost of cloth by a factor of 2-3x. Adding the new intersections on top of this adds an additional .3-.5 cost (e.g., if no collisions costs 1ms per frame, repulsions cost 1-2ms per frame, adding new intersections cost .3-.5ms). Some of these are fixed costs (e.g., building the BVH), some of these depend on how many proximities and intersections are actually detected.

==DETAILED CHANGES==
PBDTriangleMeshCollisions:
* Build a single BVH which can be used for intersection and proximities based on beginning of timestep positions (Particles.X)
* Detect any edge-face intersections
* Construct global contours from the intersections
** Contours which produce closed regions are used to flip the normals for proximity repulsions (i.e., make them attractions rather than repulsions)
** Open and loop contours are handled via the Contour Minimization. PBDTriangleMeshIntersections applies these deltas.
* Do GIA floodfill of closed contours
* Calculate Contour Minimization gradients for open and loop contours

PBDTriangleMeshIntersections:
* Based on the precomputed gradient (calculated as part of PBDTriangleMeshCollisions init), apply delta

Changes to PBDCollisionSpringConstraintsBase:
* Delta now is applied in the normal direction (similar to before most recent collision spring stabilization pass).
** Major difference is that normals are calculated per iteration rather than per subframe. This produces significantly more stable result.
** Friction response has also been added. This follows existing collision friction.
* Main Init method now takes in the BVH and GIAColors calculated by PBDTriangleMeshCollisions
** GIA data is used to determine value of new FlipNormals constraint entry. We calculate the normal per iteration, but we previous frame positions + GIA information to determine whether or not the repulsion should push to the front or backface of the triangle.
* Original Init method continues to exist but is deprecated. It does not support GIA and builds its own BVH.

ChaosClothingSimulation, ChaoSimulationEditorExtender
Added debug draw for intersections

ChaosClothConfig, ChaosClothingSimulationCloth
Added new SelfCollisionFriction and bUseSelfIntersections fields. Default is off for both (retain existing behavior)

ChaosClothConstraints
Added new constraints and updated property setters. All new properties have default parameters which attempt to preserve old behavior where possible for back compatibility.
Added PostCollisionConstraintRule list so self collisions can run after kinematic collisions.

Triangle.h
Made Triangle templated on floating point type. Typedef'ed FTriangle to continue to correspond with FReal.
Added new LineIntersection method which just calls existing RayTriangleIntersection method with different epsilons.

TriangleMesh
Added new EdgeIntersectionQuery.
Made the SegmentMesh related getters const and underlying lazily calculated data mutable. This is consistent with the MPointToTriangleMap and MPointToNeighborMap

diff --git a/Engine/Plugins/Experimental/ChaosCloth/Source/Chaos/Private/ChaosCloth/ChaosClothConstraints.cpp b/Engine/Plugins/Experimental/ChaosCloth/Source/Chaos/Private/ChaosCloth/ChaosClothConstraints.cpp
--- a/Engine/Plugins/Experimental/ChaosCloth/Source/Chaos/Private/ChaosCloth/ChaosClothConstraints.cpp
+++ b/Engine/Plugins/Experimental/ChaosCloth/Source/Chaos/Private/ChaosCloth/ChaosClothConstraints.cpp
@@ -11,6 +11,8 @@
 #include "Chaos/PBDAnimDriveConstraint.h"
 #include "Chaos/PBDShapeConstraints.h"
 #include "Chaos/PBDCollisionSpringConstraints.h"
+#include "Chaos/PBDTriangleMeshCollisions.h"
+#include "Chaos/PBDTriangleMeshIntersections.h"
 #include "Chaos/PBDEvolution.h"
 
 namespace Chaos {
@@ -24,8 +26,10 @@ FClothConstraints::FClothConstraints()
 	, NumParticles(0)
 	, ConstraintInitOffset(INDEX_NONE)
 	, ConstraintRuleOffset(INDEX_NONE)
+	, PostCollisionConstraintRuleOffset(INDEX_NONE)
 	, NumConstraintInits(0)
 	, NumConstraintRules(0)
+	, NumPostCollisionConstraintRules(0)
 {
 }
 
@@ -60,6 +64,10 @@ void FClothConstraints::Enable(bool bEnable)
 	{
 		Evolution->ActivateConstraintRuleRange(ConstraintRuleOffset, bEnable);
 	}
+	if (PostCollisionConstraintRuleOffset != INDEX_NONE)
+	{
+		Evolution->ActivatePostCollisionConstraintRuleRange(PostCollisionConstraintRuleOffset, bEnable);
+	}
 }
 
 void FClothConstraints::CreateRules()
@@ -75,12 +83,19 @@ void FClothConstraints::CreateRules()
 	{
 		ConstraintRuleOffset = Evolution->AddConstraintRuleRange(NumConstraintRules, false);
 	}
+	check(PostCollisionConstraintRuleOffset == INDEX_NONE);
+	if (NumPostCollisionConstraintRules)
+	{
+		PostCollisionConstraintRuleOffset = Evolution->AddPostCollisionConstraintRuleRange(NumPostCollisionConstraintRules, false);
+	}
 
 	TFunction<void(Softs::FSolverParticles&, const Softs::FSolverReal)>* const ConstraintInits = Evolution->ConstraintInits().GetData() + ConstraintInitOffset;
 	TFunction<void(Softs::FSolverParticles&, const Softs::FSolverReal)>* const ConstraintRules = Evolution->ConstraintRules().GetData() + ConstraintRuleOffset;
+	TFunction<void(Softs::FSolverParticles&, const Softs::FSolverReal)>* const PostCollisionConstraintRules = Evolution->PostCollisionConstraintRules().GetData() + PostCollisionConstraintRuleOffset;
 
 	int32 ConstraintInitIndex = 0;
 	int32 ConstraintRuleIndex = 0;
+	int32 PostCollisionConstraintRuleIndex = 0;
 
 	if (XEdgeConstraints)
 	{
@@ -188,25 +203,6 @@ void FClothConstraints::CreateRules()
 				VolumeConstraints->Apply(Particles, Dt);
 			};
 	}
-	if (XLongRangeConstraints)
-	{
-		ConstraintInits[ConstraintInitIndex++] =
-			[this](Softs::FSolverParticles& Particles, const Softs::FSolverReal Dt)
-			{
-				XLongRangeConstraints->Init();
-				XLongRangeConstraints->ApplyProperties(Dt, Evolution->GetIterations());
-				XLongRangeConstraints->Apply(Particles, Dt);  // Run the LRA constraint only once per timestep
-			};
-	}
-	if (LongRangeConstraints)
-	{
-		ConstraintInits[ConstraintInitIndex++] =
-			[this](Softs::FSolverParticles& Particles, const Softs::FSolverReal Dt)
-			{
-				LongRangeConstraints->ApplyProperties(Dt, Evolution->GetIterations());
-				LongRangeConstraints->Apply(Particles, Dt);  // Run the LRA constraint only once per timestep
-			};
-	}
 	if (MaximumDistanceConstraints)
 	{
 		ConstraintRules[ConstraintRuleIndex++] =
@@ -245,21 +241,58 @@ void FClothConstraints::CreateRules()
 				ShapeConstraints->Apply(Particles, Dt);
 			};
 	}
-	if (SelfCollisionConstraints)
+
+	if (SelfCollisionInit && SelfCollisionConstraints)
 	{
 		ConstraintInits[ConstraintInitIndex++] =
 			[this](Softs::FSolverParticles& Particles, const Softs::FSolverReal /*Dt*/)
 			{
-				SelfCollisionConstraints->Init(Particles);
+				SelfCollisionInit->Init(Particles);
+				SelfCollisionConstraints->Init(Particles, SelfCollisionInit->GetBVH(), SelfCollisionInit->GetVertexGIAColors(), SelfCollisionInit->GetTriangleGIAColors());
 			};
-		ConstraintRules[ConstraintRuleIndex++] =
+
+		PostCollisionConstraintRules[PostCollisionConstraintRuleIndex++] =
 			[this](Softs::FSolverParticles& Particles, const Softs::FSolverReal Dt)
 			{
 				SelfCollisionConstraints->Apply(Particles, Dt);
 			};
 	}
+
+	// The following constraints only run once per subframe, so we do their Apply as part of the Init() which modifies P
+	// To avoid possible dependency order issues, add them last
+	if (SelfCollisionInit && SelfIntersectionConstraints)
+	{
+		ConstraintInits[ConstraintInitIndex++] =
+			[this](Softs::FSolverParticles& Particles, const Softs::FSolverReal Dt)
+			{
+				SelfIntersectionConstraints->Apply(Particles, SelfCollisionInit->GetContourMinimizationIntersections(), Dt);
+			};
+	}
+
+	// Long range constraints modify particle P as part of Init. To avoid possible dependency order issues,
+	// add them last
+	if (XLongRangeConstraints)
+	{
+		ConstraintInits[ConstraintInitIndex++] =
+			[this](Softs::FSolverParticles& Particles, const Softs::FSolverReal Dt)
+			{
+				XLongRangeConstraints->Init();
+				XLongRangeConstraints->ApplyProperties(Dt, Evolution->GetIterations());
+				XLongRangeConstraints->Apply(Particles, Dt);  // Run the LRA constraint only once per timestep
+			};
+	}
+	if (LongRangeConstraints)
+	{
+		ConstraintInits[ConstraintInitIndex++] =
+			[this](Softs::FSolverParticles& Particles, const Softs::FSolverReal Dt)
+			{
+				LongRangeConstraints->ApplyProperties(Dt, Evolution->GetIterations());
+				LongRangeConstraints->Apply(Particles, Dt);  // Run the LRA constraint only once per timestep
+			};
+	}
 	check(ConstraintInitIndex == NumConstraintInits);
 	check(ConstraintRuleIndex == NumConstraintRules);
+	check(PostCollisionConstraintRuleIndex == NumPostCollisionConstraintRules);
 }
 
 void FClothConstraints::SetEdgeConstraints(const TArray<TVec3<int32>>& SurfaceElements, const TConstArrayView<FRealSingle>& StiffnessMultipliers, bool bUseXPBDConstraints)
@@ -467,8 +500,15 @@ void FClothConstraints::SetShapeTargetConstraints(Softs::FSolverReal ShapeTarget
 	++NumConstraintRules;
 }
 
-void FClothConstraints::SetSelfCollisionConstraints(const FTriangleMesh& TriangleMesh, TSet<TVec2<int32>>&& DisabledCollisionElements, Softs::FSolverReal SelfCollisionThickness)
+void FClothConstraints::SetSelfCollisionConstraints(const FTriangleMesh& TriangleMesh, TSet<TVec2<int32>>&& DisabledCollisionElements, Softs::FSolverReal SelfCollisionThickness, Softs::FSolverReal SelfCollisionFrictionCoefficient, bool bGlobalIntersectionAnalysis, bool bContourMinimization)
 {
+	SelfCollisionInit = MakeShared<Softs::FPBDTriangleMeshCollisions>(
+		ParticleOffset,
+		NumParticles,
+		TriangleMesh,
+		bGlobalIntersectionAnalysis,
+		bContourMinimization);
+
 	SelfCollisionConstraints = MakeShared<Softs::FPBDCollisionSpringConstraints>(
 		ParticleOffset,
 		NumParticles,
@@ -476,9 +516,17 @@ void FClothConstraints::SetSelfCollisionConstraints(const FTriangleMesh& Triangl
 		AnimationPositions,
 		MoveTemp(DisabledCollisionElements),
 		SelfCollisionThickness,
-		/*Stiffness =*/ (Softs::FSolverReal)1.);
-	++NumConstraintInits;  // Self collision has an init
-	++NumConstraintRules;  // and a rule
+		Softs::FPBDCollisionSpringConstraintsBase::BackCompatStiffness,
+		SelfCollisionFrictionCoefficient);
+
+	++NumConstraintInits;
+	++NumPostCollisionConstraintRules;
+
+	SelfIntersectionConstraints = MakeShared<Softs::FPBDTriangleMeshIntersections>(
+		ParticleOffset,
+		NumParticles,
+		TriangleMesh);
+	++NumConstraintInits;
 }
 
 void FClothConstraints::SetEdgeProperties(const Softs::FSolverVec2& EdgeStiffness)
@@ -563,11 +611,17 @@ void FClothConstraints::SetAnimDriveProperties(const Softs::FSolverVec2& AnimDri
 	}
 }
 
-void FClothConstraints::SetSelfCollisionProperties(Softs::FSolverReal SelfCollisionThickness)
+void FClothConstraints::SetSelfCollisionProperties(Softs::FSolverReal SelfCollisionThickness, Softs::FSolverReal SelfCollisionFrictionCoefficient, bool bGlobalIntersectionAnalysis, bool bContourMinimization)
 {
+	if (SelfCollisionInit)
+	{
+		SelfCollisionInit->SetGlobalIntersectionAnalysis(bGlobalIntersectionAnalysis);
+		SelfCollisionInit->SetContourMinimization(bContourMinimization);
+	}
 	if (SelfCollisionConstraints)
 	{
 		SelfCollisionConstraints->SetThickness(SelfCollisionThickness);
+		SelfCollisionConstraints->SetFrictionCoefficient(SelfCollisionFrictionCoefficient);
 	}
 }
 
diff --git a/Engine/Plugins/Experimental/ChaosCloth/Source/Chaos/Private/ChaosCloth/ChaosClothingSimulation.cpp b/Engine/Plugins/Experimental/ChaosCloth/Source/Chaos/Private/ChaosCloth/ChaosClothingSimulation.cpp
--- a/Engine/Plugins/Experimental/ChaosCloth/Source/Chaos/Private/ChaosCloth/ChaosClothingSimulation.cpp
+++ b/Engine/Plugins/Experimental/ChaosCloth/Source/Chaos/Private/ChaosCloth/ChaosClothingSimulation.cpp
@@ -24,11 +24,13 @@
 #include "Chaos/Capsule.h"
 #include "Chaos/TaperedCapsule.h"
 #include "Chaos/Convex.h"
+#include "Chaos/Triangle.h"
 #include "Chaos/PBDSphericalConstraint.h"
 #include "Chaos/PBDAnimDriveConstraint.h"
 #include "Chaos/PBDLongRangeConstraints.h"
 #include "Chaos/PBDSpringConstraints.h"
 #include "Chaos/PBDCollisionSpringConstraints.h"
+#include "Chaos/PBDTriangleMeshCollisions.h"
 #include "Chaos/VelocityField.h"
 #endif  // #if WITH_EDITOR || CHAOS_DEBUG_DRAW
 
@@ -102,6 +104,7 @@ namespace ClothingSimulationCVar
 	TAutoConsoleVariable<bool> DebugDrawLongRangeConstraint (TEXT("p.ChaosCloth.DebugDrawLongRangeConstraint" ), false, TEXT("Whether to debug draw the Chaos Cloth long range constraint (aka tether constraint)"), ECVF_Cheat);
 	TAutoConsoleVariable<bool> DebugDrawWindForces          (TEXT("p.ChaosCloth.DebugDrawWindForces"          ), false, TEXT("Whether to debug draw the Chaos Cloth wind forces"), ECVF_Cheat);
 	TAutoConsoleVariable<bool> DebugDrawSelfCollision       (TEXT("p.ChaosCloth.DebugDrawSelfCollision"       ), false, TEXT("Whether to debug draw the Chaos Cloth self collision information"), ECVF_Cheat);
+	TAutoConsoleVariable<bool> DebugDrawSelfIntersection    (TEXT("p.ChaosCloth.DebugDrawSelfIntersection"    ), false, TEXT("Whether to debug draw the Chaos Cloth self intersection information"), ECVF_Cheat);
 }
 #endif  // #if CHAOS_DEBUG_DRAW
 
@@ -420,6 +423,8 @@ void FClothingSimulation::CreateActor(USkeletalMeshComponent* InOwnerComponent,
 		ClothConfig->bUseCCD,
 		ClothConfig->bUseSelfCollisions,
 		ClothConfig->SelfCollisionThickness,
+		ClothConfig->SelfCollisionFriction,
+		ClothConfig->bUseSelfIntersections,
 		ClothConfig->bUseLegacyBackstop,
 		/*bUseLODIndexOverride =*/ false,
 		/*LODIndexOverride =*/ INDEX_NONE));
@@ -594,6 +599,7 @@ void FClothingSimulation::Simulate(IClothingSimulationContext* InContext)
 	if (ClothingSimulationCVar::DebugDrawLongRangeConstraint .GetValueOnAnyThread()) { DebugDrawLongRangeConstraint (); }
 	if (ClothingSimulationCVar::DebugDrawWindForces          .GetValueOnAnyThread()) { DebugDrawWindForces          (); }
 	if (ClothingSimulationCVar::DebugDrawSelfCollision       .GetValueOnAnyThread()) { DebugDrawSelfCollision       (); }
+	if (ClothingSimulationCVar::DebugDrawSelfIntersection    .GetValueOnAnyThread()) { DebugDrawSelfIntersection    (); }
 #endif  // #if CHAOS_DEBUG_DRAW
 }
 
@@ -841,6 +847,8 @@ void FClothingSimulation::RefreshClothConfig(const IClothingSimulationContext* I
 			ClothConfig->bUseCCD,
 			ClothConfig->bUseSelfCollisions,
 			ClothConfig->SelfCollisionThickness,
+			ClothConfig->SelfCollisionFriction,
+			ClothConfig->bUseSelfIntersections,
 			ClothConfig->bUseLegacyBackstop,
 			/*bUseLODIndexOverride =*/ false,
 			/*LODIndexOverride =*/ INDEX_NONE);
@@ -1064,12 +1072,12 @@ void FClothingSimulation::DebugDrawMaxDistanceValues(FCanvas* Canvas, const FSce
 #endif  // #if WITH_EDITOR
 
 #if WITH_EDITOR || CHAOS_DEBUG_DRAW
-static void DrawPoint(FPrimitiveDrawInterface* PDI, const FVector& Pos, const FLinearColor& Color, UMaterial* DebugClothMaterialVertex)  // Use color or material
+static void DrawPoint(FPrimitiveDrawInterface* PDI, const FVector& Pos, const FLinearColor& Color, UMaterial* DebugClothMaterialVertex, const float Thickness = 1.f)  // Use color or material
 {
 #if CHAOS_DEBUG_DRAW
 	if (!PDI)
 	{
-		FDebugDrawQueue::GetInstance().DrawDebugPoint(Pos, Color.ToFColor(true), false, KINDA_SMALL_NUMBER, SDPG_Foreground, 1.f);
+		FDebugDrawQueue::GetInstance().DrawDebugPoint(Pos, Color.ToFColor(true), false, KINDA_SMALL_NUMBER, SDPG_Foreground, Thickness);
 		return;
 	}
 #endif
@@ -1081,6 +1089,10 @@ static void DrawPoint(FPrimitiveDrawInterface* PDI, const FVector& Pos, const FL
 		const FVector YAxis = ViewMatrix.GetColumn(1);
 		DrawDisc(PDI, Pos, XAxis, YAxis, Color.ToFColor(true), 0.5f, 10, DebugClothMaterialVertex->GetRenderProxy(), SDPG_World);
 	}
+	else
+	{
+		PDI->DrawPoint(Pos, Color, Thickness, SDPG_World);
+	}
 #endif
 }
 
@@ -1987,6 +1999,7 @@ void FClothingSimulation::DebugDrawSelfCollision(FPrimitiveDrawInterface* PDI) c
 			const TArray<Softs::FSolverVec3>& Barys = SelfCollisionConstraints->GetBarys();
 			const FReal Thickness = (FReal)SelfCollisionConstraints->GetThickness();
 			const FReal Height = Thickness + Thickness;
+			const TArray<bool>& FlipNormals = SelfCollisionConstraints->GetFlipNormals();
 
 			for (int32 Index = 0; Index < Constraints.Num(); ++Index)
 			{
@@ -2002,21 +2015,133 @@ void FClothingSimulation::DebugDrawSelfCollision(FPrimitiveDrawInterface* PDI) c
 				const FVector Pos0 = P0 * Bary[0] + P1 * Bary[1] + P2 * Bary[2];
 
 				static const FLinearColor Brown(0.1f, 0.05f, 0.f);
+				static const FLinearColor Red(0.3f, 0.f, 0.f);
+				const FTriangle Triangle(P0, P1, P2);
+				const FVector Normal = FlipNormals[Index] ? -Triangle.GetNormal() : Triangle.GetNormal();
 
-				// Draw barycentric coordinate
-				DrawLine(PDI, Pos0, P0, Brown);
-				DrawLine(PDI, Pos0, P1, Brown);
-				DrawLine(PDI, Pos0, P2, Brown);
+				// Draw point to surface line (=normal)
+				const FVector Pos1 = Pos0 + Height * Normal;
+				DrawPoint(PDI, Pos0, Brown, nullptr, 2.f);
+				DrawLine(PDI, Pos0, Pos1, FlipNormals[Index] ? Red : Brown);
 
-				// Draw connection to point
+				// Draw pushup to point
 				static const FLinearColor Orange(0.3f, 0.15f, 0.f);
-				DrawLine(PDI, Pos0, P, Orange);
+				DrawPoint(PDI, P, Orange, nullptr, 2.f);
+				DrawLine(PDI, Pos1, P, Orange);
+			}
+		}
+	}
+}
 
-				// Draw hit triangle (with thickness and added colliding particle's thickness)
-				static const FLinearColor Red(0.3f, 0.f, 0.f);
-				DrawLine(PDI, P0, P1, Red);
-				DrawLine(PDI, P1, P2, Red);
-				DrawLine(PDI, P2, P0, Red);
+void FClothingSimulation::DebugDrawSelfIntersection(FPrimitiveDrawInterface* PDI) const
+{
+	const FVec3& LocalSpaceLocation = Solver->GetLocalSpaceLocation();
+
+	for (const FClothingSimulationCloth* const Cloth : Solver->GetCloths())
+	{
+		const int32 Offset = Cloth->GetOffset(Solver.Get());
+		if (Offset == INDEX_NONE)
+		{
+			continue;
+		}
+
+		const FClothConstraints& ClothConstraints = Solver->GetClothConstraints(Offset);
+
+		static const FLinearColor Red(1.f, 0.f, 0.f);
+		static const FLinearColor White(1.f, 1.f, 1.f);
+		static const FLinearColor Black(0.f, 0.f, 0.f);
+		static const FLinearColor Teal(0.f, 0.5f, 0.5f);
+		static const FLinearColor Green(0.f, 1.f, 0.f);
+
+		if (const Softs::FPBDTriangleMeshCollisions* const SelfCollisionInit = ClothConstraints.GetSelfCollisionInit().Get())
+		{
+			const TConstArrayView<Softs::FSolverVec3> Positions = Cloth->GetParticlePositions(Solver.Get());
+			const FTriangleMesh& TriangleMesh = Cloth->GetTriangleMesh(Solver.Get());
+
+			// Draw contours
+			const TArray<TArray<Softs::FPBDTriangleMeshCollisions::FBarycentricPoint>>& ContourPoints = SelfCollisionInit->GetIntersectionContourPoints();
+			const TArray<Softs::FPBDTriangleMeshCollisions::FContourType>& ContourTypes = SelfCollisionInit->GetIntersectionContourTypes();
+			check(ContourPoints.Num() == ContourTypes.Num());
+
+			static const FLinearColor ColorsForType[(int8)Softs::FPBDTriangleMeshCollisions::FContourType::Count] =
+			{
+				Teal,
+				Red,
+				White,
+				Black
+			};
+			for( int32 ContourIndex = 0; ContourIndex < ContourPoints.Num(); ++ContourIndex)
+			{
+				const TArray<Softs::FPBDTriangleMeshCollisions::FBarycentricPoint>& Contour = ContourPoints[ContourIndex];
+				const FLinearColor& ContourColor = ColorsForType[(int8)ContourTypes[ContourIndex]];
+				for (int32 PointIdx = 0; PointIdx < Contour.Num() - 1; ++PointIdx)
+				{
+					const Softs::FPBDTriangleMeshCollisions::FBarycentricPoint& Point0 = Contour[PointIdx];
+					const FVector EndPoint0 = LocalSpaceLocation + (1.f - Point0.Bary[0] - Point0.Bary[1]) * Positions[Point0.Vertices[0] - Offset] + Point0.Bary[0] * Positions[Point0.Vertices[1] - Offset] + Point0.Bary[1] * Positions[Point0.Vertices[2] - Offset];
+					const Softs::FPBDTriangleMeshCollisions::FBarycentricPoint& Point1 = Contour[PointIdx+1];
+					const FVector EndPoint1 = LocalSpaceLocation + (1.f - Point1.Bary[0] - Point1.Bary[1]) * Positions[Point1.Vertices[0] - Offset] + Point1.Bary[0] * Positions[Point1.Vertices[1] - Offset] + Point1.Bary[1] * Positions[Point1.Vertices[2] - Offset];
+					DrawLine(PDI, EndPoint0, EndPoint1, ContourColor);
+					DrawPoint(PDI, EndPoint0, ContourColor, nullptr, 1.f);
+					DrawPoint(PDI, EndPoint1, ContourColor, nullptr, 1.f);
+				}
+			}
+
+			// Draw GIA colors
+			const TConstArrayView<Softs::FPBDTriangleMeshCollisions::FGIAColor>& VertexGIAColors = SelfCollisionInit->GetVertexGIAColors();
+			static const FLinearColor Gray(0.5f, 0.5f, 0.5f);
+			if (VertexGIAColors.Num())
+			{
+				const TArray<TVec3<int32>>& Elements = TriangleMesh.GetElements();
+				for (int32 ParticleIdx = Offset; ParticleIdx < VertexGIAColors.Num(); ++ParticleIdx)
+				{
+					if (VertexGIAColors[ParticleIdx].ContourIndexBits)
+					{
+						const bool bIsLoop = VertexGIAColors[ParticleIdx].IsLoop();
+						const bool bAnyWhite = (VertexGIAColors[ParticleIdx].ContourIndexBits & ~VertexGIAColors[ParticleIdx].ColorBits);
+						const bool bAnyBlack = (VertexGIAColors[ParticleIdx].ContourIndexBits & VertexGIAColors[ParticleIdx].ColorBits);
+						const FLinearColor& VertColor = bIsLoop ? Red : (bAnyWhite && bAnyBlack) ? Gray : bAnyWhite ? White : Black;
+						
+						DrawPoint(PDI, LocalSpaceLocation + Positions[ParticleIdx - Offset], VertColor, nullptr, 5.f);
+					}
+				}
+			}
+			const TArray<Softs::FPBDTriangleMeshCollisions::FGIAColor>& TriangleGIAColors = SelfCollisionInit->GetTriangleGIAColors();
+			if (TriangleGIAColors.Num() == TriangleMesh.GetNumElements())
+			{
+				const TArray<TVec3<int32>>& Elements = TriangleMesh.GetElements();
+				for (int32 TriangleIdx = 0; TriangleIdx < TriangleGIAColors.Num(); ++TriangleIdx)
+				{
+					if (TriangleGIAColors[TriangleIdx].ContourIndexBits)
+					{
+						const bool bIsLoop = TriangleGIAColors[TriangleIdx].IsLoop();
+						const bool bAnyWhite = (TriangleGIAColors[TriangleIdx].ContourIndexBits & ~TriangleGIAColors[TriangleIdx].ColorBits);
+						const bool bAnyBlack = (TriangleGIAColors[TriangleIdx].ContourIndexBits & TriangleGIAColors[TriangleIdx].ColorBits);
+						const FLinearColor& TriColor = bIsLoop ? Red : (bAnyWhite && bAnyBlack) ? Gray : bAnyWhite ? White : Black;
+						DrawLine(PDI, LocalSpaceLocation + Positions[Elements[TriangleIdx][0] - Offset], LocalSpaceLocation + Positions[Elements[TriangleIdx][1] - Offset], TriColor);
+						DrawLine(PDI, LocalSpaceLocation + Positions[Elements[TriangleIdx][1] - Offset], LocalSpaceLocation + Positions[Elements[TriangleIdx][2] - Offset], TriColor);
+						DrawLine(PDI, LocalSpaceLocation + Positions[Elements[TriangleIdx][0] - Offset], LocalSpaceLocation + Positions[Elements[TriangleIdx][2] - Offset], TriColor);
+					}
+				}
+			}
+
+			// Draw contour minimization gradients
+			const TArray<Softs::FPBDTriangleMeshCollisions::FContourMinimizationIntersection>& ContourMinimizationIntersections = SelfCollisionInit->GetContourMinimizationIntersections();
+			constexpr FReal MaxDrawImpulse = 1.;
+			constexpr FReal RegularizeEpsilonSq = 1.;
+			for (const Softs::FPBDTriangleMeshCollisions::FContourMinimizationIntersection& Intersection : ContourMinimizationIntersections)
+			{
+				Softs::FSolverReal GradientLength;
+				Softs::FSolverVec3 GradientDir;
+				Intersection.GlobalGradientVector.ToDirectionAndLength(GradientDir, GradientLength);
+				const FVector Delta = FVector(GradientDir) * MaxDrawImpulse * GradientLength * FMath::InvSqrt(GradientLength * GradientLength + RegularizeEpsilonSq);
+
+				const FVector EdgeCenter = LocalSpaceLocation + .5 * (Positions[Intersection.EdgeVertices[0] - Offset] + Positions[Intersection.EdgeVertices[1] - Offset]);
+				const FVector TriCenter = LocalSpaceLocation + (Positions[Intersection.FaceVertices[0] - Offset] + Positions[Intersection.FaceVertices[1] - Offset] + Positions[Intersection.FaceVertices[2] - Offset]) / 3.;
+
+				DrawPoint(PDI, EdgeCenter, Green, nullptr, 2.f);
+				DrawLine(PDI, EdgeCenter, EdgeCenter + Delta, Green);
+				DrawPoint(PDI, TriCenter, Green, nullptr, 2.f);
+				DrawLine(PDI, TriCenter, TriCenter - Delta, Green);
 			}
 		}
 	}
diff --git a/Engine/Plugins/Experimental/ChaosCloth/Source/Chaos/Private/ChaosCloth/ChaosClothingSimulationCloth.cpp b/Engine/Plugins/Experimental/ChaosCloth/Source/Chaos/Private/ChaosCloth/ChaosClothingSimulationCloth.cpp
--- a/Engine/Plugins/Experimental/ChaosCloth/Source/Chaos/Private/ChaosCloth/ChaosClothingSimulationCloth.cpp
+++ b/Engine/Plugins/Experimental/ChaosCloth/Source/Chaos/Private/ChaosCloth/ChaosClothingSimulationCloth.cpp
@@ -121,7 +121,7 @@ void FClothingSimulationCloth::FLODData::Add(FClothingSimulationSolver* Solver,
 				DisabledCollisionElements.Emplace(TVec2<int32>(Element, Index));
 			}
 		}
-		ClothConstraints.SetSelfCollisionConstraints(TriangleMesh, MoveTemp(DisabledCollisionElements), (Softs::FSolverReal)Cloth->SelfCollisionThickness);
+		ClothConstraints.SetSelfCollisionConstraints(TriangleMesh, MoveTemp(DisabledCollisionElements), (Softs::FSolverReal)Cloth->SelfCollisionThickness, (Softs::FSolverReal)Cloth->SelfCollisionFrictionCoefficient, Cloth->bUseSelfIntersections, Cloth->bUseSelfIntersections);
 	}
 
 	// Edge constraints
@@ -286,7 +286,7 @@ void FClothingSimulationCloth::FLODData::Update(FClothingSimulationSolver* Solve
 	ClothConstraints.SetLongRangeAttachmentProperties(
 		Softs::FSolverVec2((Softs::FSolverReal)Cloth->TetherStiffness[0], (Softs::FSolverReal)Cloth->TetherStiffness[1]),
 		Softs::FSolverVec2((Softs::FSolverReal)Cloth->TetherScale[0], (Softs::FSolverReal)Cloth->TetherScale[1]) * MeshScale);
-	ClothConstraints.SetSelfCollisionProperties((Softs::FSolverReal)Cloth->SelfCollisionThickness);
+	ClothConstraints.SetSelfCollisionProperties((Softs::FSolverReal)Cloth->SelfCollisionThickness, (Softs::FSolverReal)Cloth->SelfCollisionFrictionCoefficient, Cloth->bUseSelfIntersections, Cloth->bUseSelfIntersections);
 	ClothConstraints.SetAnimDriveProperties(
 		Softs::FSolverVec2((Softs::FSolverReal)Cloth->AnimDriveStiffness[0], (Softs::FSolverReal)Cloth->AnimDriveStiffness[1]),
 		Softs::FSolverVec2((Softs::FSolverReal)Cloth->AnimDriveDamping[0], (Softs::FSolverReal)Cloth->AnimDriveDamping[1]));
@@ -369,6 +369,8 @@ FClothingSimulationCloth::FClothingSimulationCloth(
 	bool bInUseCCD,
 	bool bInUseSelfCollisions,
 	FRealSingle InSelfCollisionThickness,
+	FRealSingle InSelfCollisionFrictionCoefficient,
+	bool bInUseSelfIntersections,
 	bool bInUseLegacyBackstop,
 	bool bInUseLODIndexOverride, 
 	int32 InLODIndexOverride)
@@ -410,6 +412,8 @@ FClothingSimulationCloth::FClothingSimulationCloth(
 	, bUseCCD(bInUseCCD)
 	, bUseSelfCollisions(bInUseSelfCollisions)
 	, SelfCollisionThickness(InSelfCollisionThickness)
+	, SelfCollisionFrictionCoefficient(InSelfCollisionFrictionCoefficient)
+	, bUseSelfIntersections(bInUseSelfIntersections)
 	, bEnableBackstop(true)  // Set by clothing interactor
 	, bUseLegacyBackstop(bInUseLegacyBackstop)
 	, bUseLODIndexOverride(bInUseLODIndexOverride)
diff --git a/Engine/Plugins/Experimental/ChaosCloth/Source/Chaos/Public/ChaosCloth/ChaosClothConfig.h b/Engine/Plugins/Experimental/ChaosCloth/Source/Chaos/Public/ChaosCloth/ChaosClothConfig.h
--- a/Engine/Plugins/Experimental/ChaosCloth/Source/Chaos/Public/ChaosCloth/ChaosClothConfig.h
+++ b/Engine/Plugins/Experimental/ChaosCloth/Source/Chaos/Public/ChaosCloth/ChaosClothConfig.h
@@ -202,14 +202,22 @@ public:
 	UPROPERTY(EditAnywhere, Category = "Collision Properties")
 	bool bUseCCD = false;
 
-	/** Enable self collision. */
-	UPROPERTY(EditAnywhere, Category = "Collision Properties", meta = (InlineEditConditionToggle))
+	/** Enable self collision repulsion forces (point-face). */
+	UPROPERTY(EditAnywhere, Category = "Collision Properties")
 	bool bUseSelfCollisions = false;
 
-	/** The radius of the spheres used in self collision. */
+	/** The radius of the spheres used in self collision. (i.e., offset per side. total thickness of cloth is 2x this value) */
 	UPROPERTY(EditAnywhere, Category = "Collision Properties", meta = (UIMin = "0", UIMax = "100", ClampMin = "0", ClampMax = "1000", EditCondition = "bUseSelfCollisions"))
 	float SelfCollisionThickness = 2.0f;
 
+	/**Friction coefficient for cloth - cloth interaction. */
+	UPROPERTY(EditAnywhere, Category = "Collision Properties", meta = (UIMin = "0", UIMax = "1", ClampMin = "0", ClampMax = "10", EditCondition = "bUseSelfCollisions"))
+	float SelfCollisionFriction = 0.0f;
+
+	/** Enable self intersection resolution. This will try to fix any cloth intersections that are not handled by collision repulsions. */
+	UPROPERTY(EditAnywhere, Category = "Collision Properties", meta = (EditCondition = "bUseSelfCollisions"))
+	bool bUseSelfIntersections = false;
+
 	/**
 	 * This parameter is automatically set by the migration code. It can be overridden here to use the old way of authoring the backstop distances.
 	 * The legacy backstop requires the sphere radius to be included within the painted distance mask, making it difficult to author correctly. In this case the backstop distance is the distance from the animated mesh to the center of the corresponding backstop collision sphere.
diff --git a/Engine/Plugins/Experimental/ChaosCloth/Source/Chaos/Public/ChaosCloth/ChaosClothConstraints.h b/Engine/Plugins/Experimental/ChaosCloth/Source/Chaos/Public/ChaosCloth/ChaosClothConstraints.h
--- a/Engine/Plugins/Experimental/ChaosCloth/Source/Chaos/Public/ChaosCloth/ChaosClothConstraints.h
+++ b/Engine/Plugins/Experimental/ChaosCloth/Source/Chaos/Public/ChaosCloth/ChaosClothConstraints.h
@@ -6,6 +6,7 @@
 #include "Templates/SharedPointer.h"
 #include "Chaos/PBDSoftsEvolutionFwd.h"
 #include "Chaos/PBDLongRangeConstraints.h"
+#include "Chaos/PBDCollisionSpringConstraintsBase.h"
 
 namespace Chaos
 {
@@ -41,7 +42,9 @@ namespace Chaos
 		void SetBackstopConstraints(const TConstArrayView<FRealSingle>& BackstopDistances, const TConstArrayView<FRealSingle>& BackstopRadiuses, bool bUseLegacyBackstop);
 		void SetAnimDriveConstraints(const TConstArrayView<FRealSingle>& AnimDriveStiffnessMultipliers, const TConstArrayView<FRealSingle>& AnimDriveDampingMultipliers);
 		void SetShapeTargetConstraints(Softs::FSolverReal ShapeTargetStiffness);
-		void SetSelfCollisionConstraints(const class FTriangleMesh& TriangleMesh, TSet<TVec2<int32>>&& DisabledCollisionElements, Softs::FSolverReal SelfCollisionThickness);
+		void SetSelfCollisionConstraints(const class FTriangleMesh& TriangleMesh, TSet<TVec2<int32>>&& DisabledCollisionElements, Softs::FSolverReal SelfCollisionThickness,
+			Softs::FSolverReal SelfCollisionFriction = Softs::FPBDCollisionSpringConstraintsBase::BackCompatFrictionCoefficient, bool bGlobalIntersectionAnalysis = false,
+			bool bContourMinimization = false);
 
 		void CreateRules();
 		void Enable(bool bEnable);
@@ -54,7 +57,8 @@ namespace Chaos
 		void SetLongRangeAttachmentProperties(const Softs::FSolverVec2& TetherStiffness, const Softs::FSolverVec2& TetherScale);
 		void SetMaximumDistanceProperties(Softs::FSolverReal MaxDistancesMultiplier);
 		void SetAnimDriveProperties(const Softs::FSolverVec2& AnimDriveStiffness, const Softs::FSolverVec2& AnimDriveDamping);
-		void SetSelfCollisionProperties(Softs::FSolverReal SelfCollisionThickness);
+		void SetSelfCollisionProperties(Softs::FSolverReal SelfCollisionThickness, Softs::FSolverReal SelfCollisionFriction = Softs::FPBDCollisionSpringConstraintsBase::BackCompatFrictionCoefficient, 
+			bool bGlobalIntersectionAnalysis = false, bool bContourMinimization = false);
 		UE_DEPRECATED(5.0, "Use SetBackstopProperties(bool, FSolverReal) instead.")
 		void SetBackstopProperties(bool bEnabled) { SetBackstopProperties(bEnabled, (Softs::FSolverReal)1.); }
 		void SetBackstopProperties(bool bEnabled, Softs::FSolverReal BackstopDistancesMultiplier);
@@ -77,6 +81,8 @@ namespace Chaos
 		const TSharedPtr<Softs::FPBDAnimDriveConstraint>& GetAnimDriveConstraints() const { return AnimDriveConstraints; }
 		const TSharedPtr<Softs::FPBDShapeConstraints>& GetShapeConstraints() const { return ShapeConstraints; }
 		const TSharedPtr<Softs::FPBDCollisionSpringConstraints>& GetSelfCollisionConstraints() const { return SelfCollisionConstraints; }
+		const TSharedPtr<Softs::FPBDTriangleMeshIntersections>& GetSelfIntersectionConstraints() const { return SelfIntersectionConstraints; }
+		const TSharedPtr<Softs::FPBDTriangleMeshCollisions>& GetSelfCollisionInit() const { return SelfCollisionInit; }
 		// ---- End of debug functions ----
 
 	private:
@@ -95,7 +101,9 @@ namespace Chaos
 		TSharedPtr<Softs::FPBDSphericalBackstopConstraint> BackstopConstraints;
 		TSharedPtr<Softs::FPBDAnimDriveConstraint> AnimDriveConstraints;
 		TSharedPtr<Softs::FPBDShapeConstraints> ShapeConstraints;
+		TSharedPtr<Softs::FPBDTriangleMeshCollisions> SelfCollisionInit;
 		TSharedPtr<Softs::FPBDCollisionSpringConstraints> SelfCollisionConstraints;
+		TSharedPtr<Softs::FPBDTriangleMeshIntersections> SelfIntersectionConstraints;
 		
 		Softs::FPBDEvolution* Evolution;
 		const TArray<Softs::FSolverVec3>* AnimationPositions;
@@ -106,7 +114,9 @@ namespace Chaos
 		int32 NumParticles;
 		int32 ConstraintInitOffset;
 		int32 ConstraintRuleOffset;
+		int32 PostCollisionConstraintRuleOffset;
 		int32 NumConstraintInits;
 		int32 NumConstraintRules;
+		int32 NumPostCollisionConstraintRules;
 	};
 } // namespace Chaos
diff --git a/Engine/Plugins/Experimental/ChaosCloth/Source/Chaos/Public/ChaosCloth/ChaosClothingSimulation.h b/Engine/Plugins/Experimental/ChaosCloth/Source/Chaos/Public/ChaosCloth/ChaosClothingSimulation.h
--- a/Engine/Plugins/Experimental/ChaosCloth/Source/Chaos/Public/ChaosCloth/ChaosClothingSimulation.h
+++ b/Engine/Plugins/Experimental/ChaosCloth/Source/Chaos/Public/ChaosCloth/ChaosClothingSimulation.h
@@ -118,6 +118,7 @@ namespace Chaos
 		CHAOSCLOTH_API void DebugDrawWindForces(FPrimitiveDrawInterface* PDI = nullptr) const;
 		CHAOSCLOTH_API void DebugDrawLocalSpace(FPrimitiveDrawInterface* PDI = nullptr) const;
 		CHAOSCLOTH_API void DebugDrawSelfCollision(FPrimitiveDrawInterface* PDI = nullptr) const;
+		CHAOSCLOTH_API void DebugDrawSelfIntersection(FPrimitiveDrawInterface* PDI = nullptr) const;
 #endif  // #if WITH_EDITOR || CHAOS_DEBUG_DRAW
 
 	private:
diff --git a/Engine/Plugins/Experimental/ChaosCloth/Source/Chaos/Public/ChaosCloth/ChaosClothingSimulationCloth.h b/Engine/Plugins/Experimental/ChaosCloth/Source/Chaos/Public/ChaosCloth/ChaosClothingSimulationCloth.h
--- a/Engine/Plugins/Experimental/ChaosCloth/Source/Chaos/Public/ChaosCloth/ChaosClothingSimulationCloth.h
+++ b/Engine/Plugins/Experimental/ChaosCloth/Source/Chaos/Public/ChaosCloth/ChaosClothingSimulationCloth.h
@@ -63,6 +63,8 @@ namespace Chaos
 			bool bInUseCCD,
 			bool bInUseSelfCollisions,
 			FRealSingle InSelfCollisionThickness,
+			FRealSingle InSelfCollisionFrictionCoefficient,
+			bool bInUseSelfIntersections,
 			bool bInUseLegacyBackstop,
 			bool bInUseLODIndexOverride,
 			int32 InLODIndexOverride);
@@ -223,6 +225,8 @@ namespace Chaos
 		bool bUseCCD;
 		bool bUseSelfCollisions;
 		FRealSingle SelfCollisionThickness;
+		FRealSingle SelfCollisionFrictionCoefficient;
+		bool bUseSelfIntersections;
 		bool bEnableBackstop;
 		bool bUseLegacyBackstop;
 		bool bUseLODIndexOverride;
diff --git a/Engine/Plugins/Experimental/ChaosClothEditor/Source/Chaos/Private/ChaosClothEditor/ChaosSimulationEditorExtender.cpp b/Engine/Plugins/Experimental/ChaosClothEditor/Source/Chaos/Private/ChaosClothEditor/ChaosSimulationEditorExtender.cpp
--- a/Engine/Plugins/Experimental/ChaosClothEditor/Source/Chaos/Private/ChaosClothEditor/ChaosSimulationEditorExtender.cpp
+++ b/Engine/Plugins/Experimental/ChaosClothEditor/Source/Chaos/Private/ChaosClothEditor/ChaosSimulationEditorExtender.cpp
@@ -83,6 +83,7 @@ const FVisualizationOption FVisualizationOption::OptionData[] =
 	FVisualizationOption(&Chaos::FClothingSimulation::DebugDrawWindForces          , TEXT("p.ChaosClothEditor.DebugDrawWindForces"          ), LOCTEXT("ChaosVisName_WindForces"          , "Wind Aerodynamic Forces"    ), LOCTEXT("ChaosVisName_Wind_Tooltip"                , "Draws the Wind drag and lift forces")),
 	FVisualizationOption(&Chaos::FClothingSimulation::DebugDrawLocalSpace          , TEXT("p.ChaosClothEditor.DebugDrawLocalSpace"          ), LOCTEXT("ChaosVisName_LocalSpace"          , "Local Space Reference Bone" ), LOCTEXT("ChaosVisName_LocalSpace_Tooltip"          , "Draws the local space reference bone")),
 	FVisualizationOption(&Chaos::FClothingSimulation::DebugDrawSelfCollision       , TEXT("p.ChaosClothEditor.DebugDrawSelfCollision"       ), LOCTEXT("ChaosVisName_SelfCollision"       , "Self Collision"             ), LOCTEXT("ChaosVisName_SelfCollision_Tooltip"       , "Draws the self collision thickness/debugging information")),
+	FVisualizationOption(&Chaos::FClothingSimulation::DebugDrawSelfIntersection    , TEXT("p.ChaosClothEditor.DebugDrawSelfIntersection"    ), LOCTEXT("ChaosVisName_SelfIntersection"    , "Self Intersection"          ), LOCTEXT("ChaosVisName_SelfIntersection_Tooltip"    , "Draws the self intersection contour/region information")),
 };
 const uint32 FVisualizationOption::Count = sizeof(OptionData) / sizeof(FVisualizationOption);
 
diff --git a/Engine/Source/Runtime/Experimental/Chaos/Private/Chaos/CollisionOneShotManifolds.cpp b/Engine/Source/Runtime/Experimental/Chaos/Private/Chaos/CollisionOneShotManifolds.cpp
--- a/Engine/Source/Runtime/Experimental/Chaos/Private/Chaos/CollisionOneShotManifolds.cpp
+++ b/Engine/Source/Runtime/Experimental/Chaos/Private/Chaos/CollisionOneShotManifolds.cpp
@@ -1744,7 +1744,7 @@ namespace Chaos
 		// @todo(chaos): all ConstructConvexConvexOneShotManifold taking a triangle are still used in FHeightField::ContactManifoldImp and the 
 		// TriMesh version, but they should changed to use ConstructPlanarConvexTriangleOneShotManifold
 		template
-		void ConstructConvexConvexOneShotManifold<class FCapsule, FTriangle>(
+		void ConstructConvexConvexOneShotManifold<FCapsule, FTriangle>(
 			const FCapsule& Implicit1,
 			const FRigidTransform3& Convex1Transform, //world
 			const FTriangle& Implicit2,
@@ -1753,7 +1753,7 @@ namespace Chaos
 			FPBDCollisionConstraint& Constraint);
 
 		template
-		void ConstructConvexConvexOneShotManifold<TImplicitObjectScaled<class FCapsule, 1>, class FTriangle>(
+		void ConstructConvexConvexOneShotManifold<TImplicitObjectScaled<FCapsule, 1>, FTriangle>(
 			const TImplicitObjectScaled<class FCapsule, 1>& Implicit1,
 			const FRigidTransform3& Convex1Transform, //world
 			const FTriangle& Implicit2,
diff --git a/Engine/Source/Runtime/Experimental/Chaos/Private/Chaos/PBDCollisionSpringConstraintsBase.cpp b/Engine/Source/Runtime/Experimental/Chaos/Private/Chaos/PBDCollisionSpringConstraintsBase.cpp
--- a/Engine/Source/Runtime/Experimental/Chaos/Private/Chaos/PBDCollisionSpringConstraintsBase.cpp
+++ b/Engine/Source/Runtime/Experimental/Chaos/Private/Chaos/PBDCollisionSpringConstraintsBase.cpp
@@ -1,6 +1,7 @@
 // Copyright Epic Games, Inc. All Rights Reserved.
 #include "Chaos/PBDCollisionSpringConstraintsBase.h"
 #include "Chaos/Plane.h"
+#include "Chaos/Triangle.h"
 #include "Chaos/TriangleCollisionPoint.h"
 #include "Chaos/TriangleMesh.h"
 #include "Chaos/PBDSoftsSolverParticles.h"
@@ -18,7 +19,8 @@ FPBDCollisionSpringConstraintsBase::FPBDCollisionSpringConstraintsBase(
 	const TArray<FSolverVec3>* InReferencePositions,
 	TSet<TVec2<int32>>&& InDisabledCollisionElements,
 	const FSolverReal InThickness,
-	const FSolverReal InStiffness)
+	const FSolverReal InStiffness,
+	const FSolverReal InFrictionCoefficient)
 	: TriangleMesh(InTriangleMesh)
 	, Elements(InTriangleMesh.GetSurfaceElements())
 	, ReferencePositions(InReferencePositions)
@@ -27,6 +29,8 @@ FPBDCollisionSpringConstraintsBase::FPBDCollisionSpringConstraintsBase(
 	, NumParticles(InNumParticles)
 	, Thickness(InThickness)
 	, Stiffness(InStiffness)
+	, FrictionCoefficient(InFrictionCoefficient)
+	, bGlobalIntersectionAnalysis(false)
 {
 }
 
@@ -34,6 +38,9 @@ void FPBDCollisionSpringConstraintsBase::Init(const FSolverParticles& Particles)
 {
 	if (!Elements.Num())
 	{
+		Constraints.Reset();
+		Barys.Reset();
+		FlipNormal.Reset();
 		return;
 	}
 
@@ -42,33 +49,57 @@ void FPBDCollisionSpringConstraintsBase::Init(const FSolverParticles& Particles)
 		TRACE_CPUPROFILER_EVENT_SCOPE(ChaosPBDCollisionSpring_BuildBVH);
 		TriangleMesh.BuildBVH(static_cast<const TArrayView<const FSolverVec3>&>(Particles.XArray()), BVH);
 	}
+	TArray<FPBDTriangleMeshCollisions::FGIAColor> EmptyGIAColors;
+	return Init(Particles, BVH, static_cast<TConstArrayView<FPBDTriangleMeshCollisions::FGIAColor>>(EmptyGIAColors), EmptyGIAColors);
+}
+
+void FPBDCollisionSpringConstraintsBase::Init(const FSolverParticles& Particles, const FTriangleMesh::TBVHType<FSolverReal>& BVH, 
+	const TConstArrayView<FPBDTriangleMeshCollisions::FGIAColor>& VertexGIAColors, const TArray<FPBDTriangleMeshCollisions::FGIAColor>& TriangleGIAColors)
+{
+	if (!Elements.Num())
+	{
+		Constraints.Reset();
+		Barys.Reset();
+		FlipNormal.Reset();
+		return;
+	}
 	{
+		bGlobalIntersectionAnalysis = VertexGIAColors.Num() == NumParticles + Offset && TriangleGIAColors.Num() == Elements.Num();
 		TRACE_CPUPROFILER_EVENT_SCOPE(ChaosPBDCollisionSpring_ProximityQuery);
 
 		// Preallocate enough space for all possible connections.
 		constexpr int32 MaxConnectionsPerPoint = 3;
 		Constraints.SetNum(NumParticles * MaxConnectionsPerPoint);
 		Barys.SetNum(NumParticles * MaxConnectionsPerPoint);
+		FlipNormal.SetNum(NumParticles * MaxConnectionsPerPoint);
 
 		std::atomic<int32> ConstraintIndex(0);
 
 		const FSolverReal HeightSq = FMath::Square(Thickness + Thickness);
+
 		PhysicsParallelFor(NumParticles,
-			[this, &BVH, &Particles, &ConstraintIndex, HeightSq, MaxConnectionsPerPoint](int32 i)
+			[this, &BVH, &Particles, &ConstraintIndex, HeightSq, MaxConnectionsPerPoint, &VertexGIAColors, &TriangleGIAColors](int32 i)
 			{
 				const int32 Index = i + Offset;
+				constexpr FSolverReal ExtraThicknessMult = 1.5f;
 
 				TArray< TTriangleCollisionPoint<FSolverReal> > Result;
-				if (TriangleMesh.PointProximityQuery(BVH, static_cast<const TArrayView<const FSolverVec3>&>(Particles.XArray()), Index, Particles.X(Index), Thickness, Thickness,
-					[this](const int32 PointIndex, const int32 TriangleIndex)->bool
+				if (TriangleMesh.PointProximityQuery(BVH, static_cast<const TArrayView<const FSolverVec3>&>(Particles.XArray()), Index, Particles.X(Index), Thickness * ExtraThicknessMult, Thickness * ExtraThicknessMult,
+					[this, &VertexGIAColors, &TriangleGIAColors](const int32 PointIndex, const int32 TriangleIndex)->bool
 					{
 						const TVector<int32, 3>& Elem = Elements[TriangleIndex];
+						if (bGlobalIntersectionAnalysis && VertexGIAColors[PointIndex].IsLoop() && (VertexGIAColors[Elem[0]].IsLoop() || VertexGIAColors[Elem[1]].IsLoop() || VertexGIAColors[Elem[2]].IsLoop() || TriangleGIAColors[TriangleIndex].IsLoop()))
+						{
+							return false;
+						}
+
 						if (DisabledCollisionElements.Contains({ PointIndex, Elem[0] }) ||
 							DisabledCollisionElements.Contains({ PointIndex, Elem[1] }) ||
 							DisabledCollisionElements.Contains({ PointIndex, Elem[2] }))
 						{
 							return false;
 						}
+
 						return true;
 					},
 					Result))
@@ -101,10 +132,27 @@ void FPBDCollisionSpringConstraintsBase::Init(const FSolverParticles& Particles)
 								continue;
 							}
 						}
+
+						// NOTE: CollisionPoint.Normal has already been flipped to point toward the Point, so need to recalculate here.
+						const TTriangle<FSolverReal> Triangle(Particles.X(Elem[0]), Particles.X(Elem[1]), Particles.X(Elem[2]));
+						bool bFlipNormal = (Particles.X(Index) - CollisionPoint.Location).Dot(Triangle.GetNormal()) < 0; // Is Point currently behind Triangle?
+
+						// Doing a check against ANY (plus the TriangleGIAColors which captures sub-triangle intersections) seems to work better than checking against ALL vertex colors where the triangle must agree.
+						// In particular, it's better at handling thin regions of intersection where a single vertex or line of vertices intersect through faces.
+						if (bGlobalIntersectionAnalysis &&
+							(FPBDTriangleMeshCollisions::FGIAColor::ShouldFlipNormal(VertexGIAColors[Index], VertexGIAColors[Elem[0]]) ||
+								FPBDTriangleMeshCollisions::FGIAColor::ShouldFlipNormal(VertexGIAColors[Index], VertexGIAColors[Elem[1]]) ||
+								FPBDTriangleMeshCollisions::FGIAColor::ShouldFlipNormal(VertexGIAColors[Index], VertexGIAColors[Elem[2]]) ||
+								FPBDTriangleMeshCollisions::FGIAColor::ShouldFlipNormal(VertexGIAColors[Index], TriangleGIAColors[CollisionPoint.Indices[1]])))
+						{
+							// Want Point to push to opposite side of triangle
+							bFlipNormal = !bFlipNormal;
+						}
 						const int32 IndexToWrite = ConstraintIndex.fetch_add(1);
 
 						Constraints[IndexToWrite] = { Index, Elem[0], Elem[1], Elem[2] };
 						Barys[IndexToWrite] = { CollisionPoint.Bary[1], CollisionPoint.Bary[2], CollisionPoint.Bary[3] };
+						FlipNormal[IndexToWrite] = bFlipNormal;
 					}
 				}
 			}
@@ -114,6 +162,7 @@ void FPBDCollisionSpringConstraintsBase::Init(const FSolverParticles& Particles)
 		const int32 ConstraintNum = ConstraintIndex.load();
 		Constraints.SetNum(ConstraintNum, /*bAllowShrinking*/ true);
 		Barys.SetNum(ConstraintNum, /*bAllowShrinking*/ true);
+		FlipNormal.SetNum(ConstraintNum, /*bAllowShrinking*/ true);
 	}
 }
 
@@ -125,10 +174,12 @@ FSolverVec3 FPBDCollisionSpringConstraintsBase::GetDelta(const FSolverParticles&
 	const int32 i3 = Constraint[2];
 	const int32 i4 = Constraint[3];
 
-	const FSolverReal CombinedMass = Particles.InvM(i1) +
+	const FSolverReal TrianglePointInvM =
 		Particles.InvM(i2) * Barys[i][0] +
 		Particles.InvM(i3) * Barys[i][1] +
 		Particles.InvM(i4) * Barys[i][2];
+
+	const FSolverReal CombinedMass = Particles.InvM(i1) + TrianglePointInvM;
 	if (CombinedMass <= (FSolverReal)1e-7)
 	{
 		return FSolverVec3(0);
@@ -142,13 +193,36 @@ FSolverVec3 FPBDCollisionSpringConstraintsBase::GetDelta(const FSolverParticles&
 	const FSolverReal Height = Thickness + Thickness;
 	const FSolverVec3 P = Barys[i][0] * P2 + Barys[i][1] * P3 + Barys[i][2] * P4;
 	const FSolverVec3 Difference = P1 - P;
-	const FSolverReal DistSq = Difference.SizeSquared();
-	if (DistSq > Height * Height)
+
+	// Normal repulsion with friction
+	const TTriangle<FSolverReal> Triangle(P2, P3, P4);
+	const FSolverVec3 Normal = FlipNormal[i] ? -Triangle.GetNormal() : Triangle.GetNormal();
+
+	const FSolverReal NormalDifference = Difference.Dot(Normal);
+	if (NormalDifference > Height)
 	{
 		return FSolverVec3(0);
 	}
-	const FSolverVec3 Delta = Difference * Height * FMath::InvSqrt(DistSq) - Difference;
-	return Stiffness * Delta / CombinedMass;
+
+	const FSolverReal NormalDelta = Height - NormalDifference;
+	const FSolverVec3 RepulsionDelta = Stiffness * NormalDelta * Normal / CombinedMass;
+
+	if (FrictionCoefficient > 0)
+	{
+		const FSolverVec3& X1 = Particles.X(i1);
+		const FSolverVec3 X = Barys[i][0] * Particles.X(i2) + Barys[i][1] * Particles.X(i3) + Barys[i][2] * Particles.X(i4);
+		const FSolverVec3 RelativeDisplacement = (P1 - X1) - (P - X) + (Particles.InvM(i1) - TrianglePointInvM) * RepulsionDelta;
+		const FSolverVec3 RelativeDisplacementTangent = RelativeDisplacement - RelativeDisplacement.Dot(Normal) * Normal;
+		const FSolverReal RelativeDisplacementTangentLength = RelativeDisplacementTangent.Length();
+		const FSolverReal PositionCorrection = FMath::Min(NormalDelta * FrictionCoefficient, RelativeDisplacementTangentLength);
+		const FSolverReal CorrectionRatio = RelativeDisplacementTangentLength < SMALL_NUMBER ? 0.f : PositionCorrection / RelativeDisplacementTangentLength;
+		const FSolverVec3 FrictionDelta = -CorrectionRatio * RelativeDisplacementTangent / CombinedMass;
+		return RepulsionDelta + FrictionDelta;
+	}
+	else
+	{
+		return RepulsionDelta;
+	}
 }
 
 }  // End namespace Chaos::Softs
diff --git a/Engine/Source/Runtime/Experimental/Chaos/Private/Chaos/PBDEvolution.cpp b/Engine/Source/Runtime/Experimental/Chaos/Private/Chaos/PBDEvolution.cpp
--- a/Engine/Source/Runtime/Experimental/Chaos/Private/Chaos/PBDEvolution.cpp
+++ b/Engine/Source/Runtime/Experimental/Chaos/Private/Chaos/PBDEvolution.cpp
@@ -2,6 +2,7 @@
 #include "Chaos/PBDEvolution.h"
 
 #include "Chaos/Framework/Parallel.h"
+#include "Chaos/PBDTriangleMeshIntersections.h"
 #include "Chaos/PerParticleDampVelocity.h"
 #include "Chaos/PerParticleGravity.h"
 #include "Chaos/PerParticlePBDCCDCollisionConstraint.h"
@@ -22,6 +23,7 @@ DECLARE_CYCLE_STAT(TEXT("Chaos PBD Integrate"), STAT_ChaosClothSolverIntegrate,
 DECLARE_CYCLE_STAT(TEXT("Chaos PBD Iteration Loop"), STAT_ChaosPBDIterationLoop, STATGROUP_Chaos);
 DECLARE_CYCLE_STAT(TEXT("Chaos PBD Post Iteration Updates"), STAT_ChaosPBDPostIterationUpdates, STATGROUP_Chaos);
 DECLARE_CYCLE_STAT(TEXT("Chaos PBD Constraint Rules"), STAT_ChaosPBDConstraintRule, STATGROUP_Chaos);
+DECLARE_CYCLE_STAT(TEXT("Chaos PBD Post Collision Constraint Rules"), STAT_ChaosPBDPostCollisionConstraintRule, STATGROUP_Chaos);
 DECLARE_CYCLE_STAT(TEXT("Chaos PBD Self Collision"), STAT_ChaosPBDSelfCollisionRule, STATGROUP_Chaos);
 DECLARE_CYCLE_STAT(TEXT("Chaos PBD Collision Rule"), STAT_ChaosPBDCollisionRule, STATGROUP_Chaos);
 DECLARE_CYCLE_STAT(TEXT("Chaos PBD Collider Friction"), STAT_ChaosPBDCollisionRuleFriction, STATGROUP_Chaos);
@@ -56,7 +58,6 @@ void FPBDEvolution::AddGroups(int32 NumGroups)
 	{
 		MGroupGravityAccelerations[GroupId] = MGravity;
 		MGroupCollisionThicknesses[GroupId] = MCollisionThickness;
-		MGroupSelfCollisionThicknesses[GroupId] = MSelfCollisionThickness;
 		MGroupCoefficientOfFrictions[GroupId] = MCoefficientOfFriction;
 		MGroupDampings[GroupId] = MDamping;
 		MGroupLocalDampings[GroupId] = MLocalDamping;
@@ -85,10 +86,10 @@ FPBDEvolution::FPBDEvolution(
 	, MCollisionParticlesActiveView(MCollisionParticles)
 	, MConstraintInitsActiveView(MConstraintInits)
 	, MConstraintRulesActiveView(MConstraintRules)
+	, MPostCollisionConstraintRulesActiveView(MPostCollisionConstraintRules)
 	, MNumIterations(NumIterations)
 	, MGravity(FSolverVec3((FSolverReal)0., (FSolverReal)0., (FSolverReal)-980.665))
 	, MCollisionThickness(CollisionThickness)
-	, MSelfCollisionThickness(SelfCollisionThickness)
 	, MCoefficientOfFriction(CoefficientOfFriction)
 	, MDamping(Damping)
 	, MLocalDamping(LocalDamping)
@@ -100,7 +101,6 @@ FPBDEvolution::FPBDEvolution(
 	TArrayCollection::AddArray(&MGroupVelocityFields);
 	TArrayCollection::AddArray(&MGroupForceRules);
 	TArrayCollection::AddArray(&MGroupCollisionThicknesses);
-	TArrayCollection::AddArray(&MGroupSelfCollisionThicknesses);
 	TArrayCollection::AddArray(&MGroupCoefficientOfFrictions);
 	TArrayCollection::AddArray(&MGroupDampings);
 	TArrayCollection::AddArray(&MGroupLocalDampings);
@@ -199,6 +199,15 @@ int32 FPBDEvolution::AddConstraintRuleRange(int32 NumConstraints, bool bActivate
 	return MConstraintRulesActiveView.AddRange(NumConstraints, bActivate);
 }
 
+int32 FPBDEvolution::AddPostCollisionConstraintRuleRange(int32 NumConstraints, bool bActivate)
+{
+	// Add new constraint rule functions
+	MPostCollisionConstraintRules.AddDefaulted(NumConstraints);
+
+	// Add range
+	return MPostCollisionConstraintRulesActiveView.AddRange(NumConstraints, bActivate);
+}
+
 template<bool bForceRule, bool bVelocityField, bool bDampVelocityRule>
 void FPBDEvolution::PreIterationUpdate(
 	const FSolverReal Dt,
@@ -480,6 +489,16 @@ void FPBDEvolution::AdvanceOneTimeStep(const FSolverReal Dt, const bool bSmoothD
 						}
 					}, bUseSingleThreadedRange);
 			}
+			MPostCollisionConstraintRulesActiveView.RangeFor(
+				[this, Dt](TArray<TFunction<void(FSolverParticles&, const FSolverReal)>>& ConstraintRules, int32 Offset, int32 Range)
+				{
+					TRACE_CPUPROFILER_EVENT_SCOPE(ChaosPBDPostCollisionConstraintRule);
+					SCOPE_CYCLE_COUNTER(STAT_ChaosPBDPostCollisionConstraintRule);
+					for (int32 ConstraintIndex = Offset; ConstraintIndex < Range; ++ConstraintIndex)
+					{
+						ConstraintRules[ConstraintIndex](MParticles, Dt); // P +/-= ...
+					}
+				}, bUseSingleThreadedRange);
 		}
 
 		{
diff --git a/Engine/Source/Runtime/Experimental/Chaos/Private/Chaos/PBDTriangleMeshCollisions.cpp b/Engine/Source/Runtime/Experimental/Chaos/Private/Chaos/PBDTriangleMeshCollisions.cpp
new file mode 100644
--- /dev/null
+++ b/Engine/Source/Runtime/Experimental/Chaos/Private/Chaos/PBDTriangleMeshCollisions.cpp
@@ -0,0 +1,1211 @@
+// Copyright Epic Games, Inc. All Rights Reserved.
+#include "Chaos/PBDTriangleMeshCollisions.h"
+#include "Chaos/Plane.h"
+#include "Chaos/PBDSoftsSolverParticles.h"
+#include "Chaos/Triangle.h"
+#include "Chaos/TriangleCollisionPoint.h"
+#include "Chaos/TriangleMesh.h"
+#include "Chaos/Framework/Parallel.h"
+#include <atomic>
+
+namespace Chaos::Softs {
+
+struct FEdgeFaceIntersection
+{
+	int32 EdgeIndex;
+	int32 FaceIndex;
+	FSolverReal EdgeCoordinate;
+	FSolverVec2 FaceCoordinate;
+	FSolverVec3 FaceNormal;
+	FSolverVec3 IntersectionPoint;
+	bool bFlipGradient = false; // Only used when building global contours
+};
+
+// Returned array has NOT been shrunk
+static TArray<FEdgeFaceIntersection> FindEdgeFaceIntersections(const FTriangleMesh& TriangleMesh, const FTriangleMesh::TBVHType<FSolverReal>& BVH, const FSolverParticles& Particles)
+{
+	TRACE_CPUPROFILER_EVENT_SCOPE(ChaosFPBDTriangleMeshCollisions_IntersectionQuery);
+
+	const FSegmentMesh& SegmentMesh = TriangleMesh.GetSegmentMesh();
+	const TArray<TVec2<int32>>& EdgeToFaces = TriangleMesh.GetEdgeToFaces();
+	const TArray<TVec3<int32>>& Elements = TriangleMesh.GetElements();
+
+	TArray<FEdgeFaceIntersection> Intersections;
+
+	// Preallocate enough space for (more than) typical number of expected intersections.
+	constexpr int32 PreallocatedIntersectionsPerEdge = 3;
+	Intersections.SetNum(PreallocatedIntersectionsPerEdge * SegmentMesh.GetNumElements());
+	std::atomic<int32> IntersectionIndex(0);
+
+	// Extra intersections that require a lock to write to if you have more than PreallocatedIntersectionsPerEdge 
+	TArray<FEdgeFaceIntersection> ExtraIntersections;
+	FCriticalSection CriticalSection;
+	PhysicsParallelFor(SegmentMesh.GetNumElements(),
+		[&BVH, &TriangleMesh, &Particles, &SegmentMesh, &EdgeToFaces, &Elements, &IntersectionIndex, &Intersections, PreallocatedIntersectionsPerEdge, &ExtraIntersections, &CriticalSection](int32 EdgeIndex)
+		{
+
+			TArray< TTriangleCollisionPoint<FSolverReal> > Result;
+
+			const int32 EdgePointIndex0 = SegmentMesh.GetElements()[EdgeIndex][0];
+			const int32 EdgePointIndex1 = SegmentMesh.GetElements()[EdgeIndex][1];
+			const FSolverVec3& EdgePosition0 = Particles.X(EdgePointIndex0);
+			const FSolverVec3& EdgePosition1 = Particles.X(EdgePointIndex1);
+
+
+			if (TriangleMesh.EdgeIntersectionQuery(BVH, static_cast<const TArrayView<const FSolverVec3>&>(Particles.XArray()), EdgeIndex, EdgePosition0, EdgePosition1,
+				[&Elements, EdgePointIndex0, EdgePointIndex1](int32 EdgeIndex, int32 TriangleIndex)
+				{
+					if (EdgePointIndex0 == Elements[TriangleIndex][0] || EdgePointIndex0 == Elements[TriangleIndex][1] || EdgePointIndex0 == Elements[TriangleIndex][2] ||
+						EdgePointIndex1 == Elements[TriangleIndex][0] || EdgePointIndex1 == Elements[TriangleIndex][1] || EdgePointIndex1 == Elements[TriangleIndex][2])
+					{
+						return false;
+					}
+					return true;
+				},
+				Result))
+			{
+
+				for (int32 CollisionPointIndex = 0; CollisionPointIndex < Result.Num(); ++CollisionPointIndex)
+				{
+					const TTriangleCollisionPoint<FSolverReal>& CollisionPoint = Result[CollisionPointIndex];
+
+					check(CollisionPoint.ContactType == TTriangleCollisionPoint<FSolverReal>::EContactType::EdgeFace);
+					FEdgeFaceIntersection Intersection;
+					Intersection.EdgeIndex = CollisionPoint.Indices[0];
+					Intersection.FaceIndex = CollisionPoint.Indices[1];
+					Intersection.EdgeCoordinate = CollisionPoint.Bary[0];
+					Intersection.FaceCoordinate = { CollisionPoint.Bary[2], CollisionPoint.Bary[3] };
+					Intersection.FaceNormal = CollisionPoint.Normal;
+					Intersection.IntersectionPoint = CollisionPoint.Location;
+
+					if (CollisionPointIndex < PreallocatedIntersectionsPerEdge)
+					{
+						const int32 IndexToWrite = IntersectionIndex.fetch_add(1);
+						Intersections[IndexToWrite] = Intersection;
+					}
+					else
+					{
+						CriticalSection.Lock();
+						ExtraIntersections.Add(Intersection);
+						CriticalSection.Unlock();
+					}
+				}
+			}
+		}
+	);
+
+	// Set Intersections Num to actual number found.
+	const int32 IntersectionNum = IntersectionIndex.load();
+	Intersections.SetNum(IntersectionNum, false /*bAllowShrinking*/);
+
+	// Append any ExtraIntersections
+	Intersections.Append(ExtraIntersections);
+	return Intersections;
+}
+
+// Global intersection analysis (identifying global contours, flood filling)
+namespace GIA
+{
+	struct FIntersectionContourTriangleSection
+	{
+		int32 LoopVertexLocalIndex = INDEX_NONE;
+		int32 TriangleIndex = INDEX_NONE;
+
+		TVec2<int32> CrossingEdgeLocalIndex = { INDEX_NONE,INDEX_NONE };
+
+		FIntersectionContourTriangleSection() = default;
+		FIntersectionContourTriangleSection(int32 InTriangleIndex)
+			: TriangleIndex(InTriangleIndex)
+		{}
+	};
+
+	using FIntersectionContour = TArray<FIntersectionContourTriangleSection>;
+
+	struct FIntersectionContourPair
+	{
+		bool bIsClosed = false;
+		int8 LoopVertexCount = 0;
+		int8 BoundaryEdgeCount = 0;
+
+		FIntersectionContour ColorContours[2];
+		TArray<FEdgeFaceIntersection> Intersections;
+		TArray<int32> ContourPointCurves[2]; // Track which curves correspond with which Contours. (Sometimes we generate multiple curves per contour)
+	};
+
+	// methods for building contours
+	namespace ContourBuilding
+	{
+		// internal helper methods
+		namespace __internal
+		{
+			static inline int32 GetLocalEdgeIndex(const TArray<TVec3<int32>>& FaceToEdges, const int32 FaceIndex, const int32 EdgeIndex)
+			{
+				for (int32 LocalIndex = 0; LocalIndex < 3; ++LocalIndex)
+				{
+					if (FaceToEdges[FaceIndex][LocalIndex] == EdgeIndex)
+					{
+						return LocalIndex;
+					}
+				}
+				check(false); // we should never get here.
+				return INDEX_NONE;
+			}
+
+			// returns local vertex index within FaceVertices
+			// returns INDEX_NONE if no loop vertex found
+			static inline int32 FindLoopVertex(const TVec3<int32>& FaceVertices, int32 OppositeVertexIndex)
+			{
+				for (int32 FaceLocalVertexIndex = 0; FaceLocalVertexIndex < 3; ++FaceLocalVertexIndex)
+				{
+					if (FaceVertices[FaceLocalVertexIndex] == OppositeVertexIndex)
+					{
+						return FaceLocalVertexIndex;
+					}
+				}
+				return INDEX_NONE;
+			}
+
+			static inline bool IsNextIntersectionEdgeFaceFirstIntersection(const FEdgeFaceIntersection& FirstIntersection, const int32 EdgeFaceCrossingEdgeLocalIndex, const TVec3<int32>& FaceToEdges_EdgeFace, const int32 CurrIntersectionFaceIndex, int32& OutLocalEdgeIndex)
+			{
+				for (int32 LocalEdgeIndexOffset = 1; LocalEdgeIndexOffset < 3; ++LocalEdgeIndexOffset)
+				{
+					const int32 LocalEdgeIndex = (LocalEdgeIndexOffset + EdgeFaceCrossingEdgeLocalIndex) % 3;
+					if (FirstIntersection.EdgeIndex == FaceToEdges_EdgeFace[LocalEdgeIndex] && FirstIntersection.FaceIndex == CurrIntersectionFaceIndex)
+					{
+						OutLocalEdgeIndex = LocalEdgeIndex;
+						return true;
+					}
+				}
+
+				return false;
+			}
+
+			static inline bool IsNextIntersectionFaceFirstIntersection(const FEdgeFaceIntersection& FirstIntersection, const TVec3<int32>& FaceToEdges_Face, const int32 EdgeFace, int32& OutLocalEdgeIndex)
+			{
+				for (int32 LocalEdgeIndex = 0; LocalEdgeIndex < 3; ++LocalEdgeIndex)
+				{
+					if (FirstIntersection.EdgeIndex == FaceToEdges_Face[LocalEdgeIndex] && FirstIntersection.FaceIndex == EdgeFace)
+					{
+						OutLocalEdgeIndex = LocalEdgeIndex;
+						return true;
+					}
+				}
+
+				return false;
+			}
+
+			static inline bool FindNextIntersectionEdgeFace(TMap<TVec2<int32>, FEdgeFaceIntersection>& IntersectionMap, const int32 EdgeFaceCrossingEdgeLocalIndex, const TVec3<int32>& FaceToEdges_EdgeFace, const int32 CurrIntersectionFaceIndex, FEdgeFaceIntersection& NextIntersection, int32& OutLocalEdgeIndex)
+			{
+				for (int32 LocalEdgeIndexOffset = 1; LocalEdgeIndexOffset < 3; ++LocalEdgeIndexOffset)
+				{
+					const int32 LocalEdgeIndex = (LocalEdgeIndexOffset + EdgeFaceCrossingEdgeLocalIndex) % 3;
+					if (IntersectionMap.RemoveAndCopyValue({ FaceToEdges_EdgeFace[LocalEdgeIndex], CurrIntersectionFaceIndex }, NextIntersection))
+					{
+						OutLocalEdgeIndex = LocalEdgeIndex;
+						return true;
+					}
+				}
+				return false;
+			}
+
+			static inline bool FindNextIntersectionFace(TMap<TVec2<int32>, FEdgeFaceIntersection>& IntersectionMap, const TVec3<int32>& FaceToEdges_Face, const int32 EdgeFace, FEdgeFaceIntersection& NextIntersection, int32& OutLocalEdgeIndex)
+			{
+				for (int32 LocalEdgeIndex = 0; LocalEdgeIndex < 3; ++LocalEdgeIndex)
+				{
+					if (IntersectionMap.RemoveAndCopyValue({ FaceToEdges_Face[LocalEdgeIndex], EdgeFace }, NextIntersection))
+					{
+						OutLocalEdgeIndex = LocalEdgeIndex;
+						return true;
+					}
+				}
+				return false;
+			}
+		} // namespace __internal
+
+		// Build intersection contour data from edge-face intersection data.
+		static TArray<FIntersectionContourPair> BuildIntersectionContours(const FTriangleMesh& TriangleMesh, const TArray<FEdgeFaceIntersection>& IntersectionArray, TArray<TArray<FPBDTriangleMeshCollisions::FBarycentricPoint>>& ContourPoints)
+		{
+			TRACE_CPUPROFILER_EVENT_SCOPE(ChaosFPBDTriangleMeshCollisions_BuildIntersectionContours);
+
+			// Put intersections in a Map so we can look them up by Edge, Face
+			// Key is (EdgeIndex, FaceIndex)
+			TMap<TVec2<int32>, FEdgeFaceIntersection> IntersectionMap;
+			IntersectionMap.Reserve(IntersectionArray.Num());
+			for (const FEdgeFaceIntersection& Intersection : IntersectionArray)
+			{
+				IntersectionMap.Add({ Intersection.EdgeIndex, Intersection.FaceIndex }, Intersection);
+			}
+
+			const FSegmentMesh& SegmentMesh = TriangleMesh.GetSegmentMesh();
+			const TArray<TVec2<int32>>& EdgeToFaces = TriangleMesh.GetEdgeToFaces();
+			const TArray<TVec3<int32>>& FaceToEdges = TriangleMesh.GetFaceToEdges();
+			const TArray<TVec3<int32>>& Elements = TriangleMesh.GetElements();
+
+			TArray<FIntersectionContourPair> Contours;
+
+			while (!IntersectionMap.IsEmpty())
+			{
+				// Add a new contour
+				FIntersectionContourPair& ContourPair = Contours.AddDefaulted_GetRef();
+
+				// Set up first intersection
+				TMap<TVec2<int32>, FEdgeFaceIntersection>::TConstIterator Iter = IntersectionMap.CreateConstIterator();
+				const FEdgeFaceIntersection FirstIntersection = IntersectionMap.FindAndRemoveChecked(Iter.Key());
+
+				// Track which intersections are used by which contour mostly so we can build contour minimization contours with these.
+				ContourPair.Intersections.Add_GetRef(FirstIntersection).bFlipGradient = false;
+
+				// Arbitrarily set colors to start contour
+				constexpr int32 InitialFaceContourIndex = 0;
+				constexpr int32 InitialEdgeContourIndex = 1;
+
+				FIntersectionContour* FaceContour = &ContourPair.ColorContours[InitialFaceContourIndex];
+				FIntersectionContour* EdgeContour = &ContourPair.ColorContours[InitialEdgeContourIndex];
+				ContourPair.ContourPointCurves[InitialFaceContourIndex].Add(ContourPoints.Num());
+				TArray<FPBDTriangleMeshCollisions::FBarycentricPoint>* FaceContourPoints = &ContourPoints.AddDefaulted_GetRef();
+				ContourPair.ContourPointCurves[InitialEdgeContourIndex].Add(ContourPoints.Num());
+				TArray<FPBDTriangleMeshCollisions::FBarycentricPoint>* EdgeContourPoints = &ContourPoints.AddDefaulted_GetRef();
+
+				const int32 FirstEdgeFace = EdgeToFaces[FirstIntersection.EdgeIndex][0];
+				check(FirstEdgeFace != -1); // Each Edge should be connected to at least one face
+
+				// We will rely on finding the First FaceSection and EdgeFaceSection as the [0] element in these arrays.
+				check(FaceContour->Num() == 0);
+				check(EdgeContour->Num() == 0);
+
+
+				// Start main loop to consume intersections to build this contour
+				bool bReverseDirection = false; // when not reversing, write to CrossingEdgeLocalIndex[0] first. when reversing, writing to CrossingEdgeLocalIndex[1] first.
+				FEdgeFaceIntersection CurrIntersection = FirstIntersection;
+				int32 EdgeFace = FirstEdgeFace;
+				FIntersectionContourTriangleSection* FaceSection = &FaceContour->Add_GetRef(FIntersectionContourTriangleSection(FirstIntersection.FaceIndex));
+				FIntersectionContourTriangleSection* EdgeFaceSection = &EdgeContour->Add_GetRef(FIntersectionContourTriangleSection(FirstEdgeFace));
+
+				// Setup first EdgeFaceSection crossing
+				EdgeFaceSection->CrossingEdgeLocalIndex[0] = __internal::GetLocalEdgeIndex(FaceToEdges, FirstEdgeFace, FirstIntersection.EdgeIndex);
+				while (1)
+				{
+					const TVec3<int32>& FaceVertices = Elements[CurrIntersection.FaceIndex];
+					const TVec3<int32>& EdgeFaceVertices = Elements[EdgeFace];
+
+					FaceContourPoints->Add({ CurrIntersection.FaceCoordinate, FaceVertices });
+					EdgeContourPoints->Add({ {CurrIntersection.EdgeCoordinate, 0.f}, {EdgeFaceVertices[0], EdgeFaceVertices[1], EdgeFaceVertices[1]} });
+
+					// Walk the intersection contour. Next point in contour is either
+					// 1) Loop vertex
+					// 2) Intersection between another edge on EdgeFace and Face
+					//   a) This is FirstIntersection (i.e., we closed the loop)
+					//   b) This is a previously untouched intersection
+					// 3) Intersection between an edge on Face with EdgeFace
+					//   a) This is FirstIntersection (i.e., we closed the loop)
+					//   b) This is a previously untouched intersection
+
+					// General note on limitations: this algorithm can fail to correctly reconstruct complex contours with combinations of loop vertices and sliver regions that double back on themselves.
+					// These contours will be handled by contour minimization.
+
+					// Check 1)
+					const int32 EdgeFaceCrossingEdgeLocalIndex = bReverseDirection ? EdgeFaceSection->CrossingEdgeLocalIndex[1] : EdgeFaceSection->CrossingEdgeLocalIndex[0];
+					check(EdgeFaceCrossingEdgeLocalIndex != INDEX_NONE);
+					const int32 FaceLoopVertexLocalIndex = __internal::FindLoopVertex(FaceVertices, EdgeFaceVertices[(EdgeFaceCrossingEdgeLocalIndex + 2) % 3]);
+					if (FaceLoopVertexLocalIndex != INDEX_NONE)
+					{
+						FaceSection->LoopVertexLocalIndex = FaceLoopVertexLocalIndex;
+						EdgeFaceSection->LoopVertexLocalIndex = (EdgeFaceCrossingEdgeLocalIndex + 2) % 3;
+						++ContourPair.LoopVertexCount;
+
+						if (ContourPair.LoopVertexCount + ContourPair.BoundaryEdgeCount == 2)
+						{
+							// We've found both ends of this contour.
+							ContourPair.bIsClosed = ContourPair.LoopVertexCount == 2;
+							break;
+						}
+						check(ContourPair.LoopVertexCount + ContourPair.BoundaryEdgeCount == 1);
+						// We hit one loop vertex. Pick up at the FirstIntersection and move the opposite direction if possible.
+
+						EdgeFace = EdgeToFaces[FirstIntersection.EdgeIndex][1];
+						if (EdgeFace == -1)
+						{
+							// First intersection edge was a boundary, so this contour is done. We're NOT closed.
+							++ContourPair.BoundaryEdgeCount;
+							break;
+						}
+
+						check(bReverseDirection == false);
+						bReverseDirection = true;
+						CurrIntersection = FirstIntersection;
+						FaceContour = &ContourPair.ColorContours[InitialFaceContourIndex];
+						EdgeContour = &ContourPair.ColorContours[InitialEdgeContourIndex];
+
+						// ContourPoints are currently just used for debug drawing. Start new contours for reverse section
+						ContourPair.ContourPointCurves[InitialFaceContourIndex].Add(ContourPoints.Num());
+						FaceContourPoints = &ContourPoints.AddDefaulted_GetRef();
+						ContourPair.ContourPointCurves[InitialEdgeContourIndex].Add(ContourPoints.Num());
+						EdgeContourPoints = &ContourPoints.AddDefaulted_GetRef();
+
+						FaceSection = &ContourPair.ColorContours[InitialFaceContourIndex][0];
+						EdgeFaceSection = &EdgeContour->Add_GetRef(FIntersectionContourTriangleSection(EdgeFace));
+						EdgeFaceSection->CrossingEdgeLocalIndex[1] = __internal::GetLocalEdgeIndex(FaceToEdges, EdgeFace, CurrIntersection.EdgeIndex);
+						continue;
+					}
+
+					// Check 2a)
+					int32 LocalEdgeIndex;
+					if (!bReverseDirection && __internal::IsNextIntersectionEdgeFaceFirstIntersection(FirstIntersection, EdgeFaceCrossingEdgeLocalIndex, FaceToEdges[EdgeFace], CurrIntersection.FaceIndex, LocalEdgeIndex))
+					{
+						// We can close the contour
+						EdgeFaceSection->CrossingEdgeLocalIndex[1] = LocalEdgeIndex;
+
+						FIntersectionContourTriangleSection& FirstFaceSection = ContourPair.ColorContours[InitialFaceContourIndex][0];
+
+						// Merge current FaceSection and FirstFaceSection (if they're not already the same)
+						check(FaceContour == &ContourPair.ColorContours[InitialFaceContourIndex]);
+						check(FaceSection->TriangleIndex == FirstFaceSection.TriangleIndex);
+						if (FaceSection != &FirstFaceSection)
+						{
+							check(FaceSection->CrossingEdgeLocalIndex[0] != INDEX_NONE);
+							FirstFaceSection.CrossingEdgeLocalIndex[0] = FaceSection->CrossingEdgeLocalIndex[0];
+							check(FaceSection == &FaceContour->Last());
+							FaceContour->RemoveAt(FaceContour->Num() - 1, 1, false);
+						}
+
+						ContourPair.bIsClosed = true;
+
+						// Repeat first point in contour points for ease of drawing closed loop
+						FaceContourPoints->Add(FPBDTriangleMeshCollisions::FBarycentricPoint((*FaceContourPoints)[0]));
+						EdgeContourPoints->Add(FPBDTriangleMeshCollisions::FBarycentricPoint((*EdgeContourPoints)[0]));
+						break;
+					}
+
+					// Check 3a)
+					if (!bReverseDirection && __internal::IsNextIntersectionFaceFirstIntersection(FirstIntersection, FaceToEdges[CurrIntersection.FaceIndex], EdgeFace, LocalEdgeIndex))
+					{
+						// We can close the contour
+						FaceSection->CrossingEdgeLocalIndex[1] = LocalEdgeIndex;
+
+						// Merge current EdgeFaceSection with FirstFaceSection
+						FIntersectionContourTriangleSection& FirstFaceSection = ContourPair.ColorContours[InitialFaceContourIndex][0];
+						check(EdgeContour == &ContourPair.ColorContours[InitialFaceContourIndex]);
+						check(EdgeFaceSection->TriangleIndex == FirstFaceSection.TriangleIndex);
+						check(EdgeFaceSection->CrossingEdgeLocalIndex[0] != INDEX_NONE);
+						FirstFaceSection.CrossingEdgeLocalIndex[0] = EdgeFaceSection->CrossingEdgeLocalIndex[0];
+
+						check(EdgeFaceSection == &EdgeContour->Last());
+						EdgeContour->RemoveAt(EdgeContour->Num() - 1, 1, false);
+
+						ContourPair.bIsClosed = true;
+
+						// Repeat first point in contour points for ease of drawing closed loop
+						FaceContourPoints->Add(FPBDTriangleMeshCollisions::FBarycentricPoint((*FaceContourPoints)[0]));
+						EdgeContourPoints->Add(FPBDTriangleMeshCollisions::FBarycentricPoint((*EdgeContourPoints)[0]));
+						break;
+					}
+
+					// Check 2b)
+					FEdgeFaceIntersection NextIntersection;
+					if (__internal::FindNextIntersectionEdgeFace(IntersectionMap, EdgeFaceCrossingEdgeLocalIndex, FaceToEdges[EdgeFace], CurrIntersection.FaceIndex, NextIntersection, LocalEdgeIndex))
+					{
+						// We found the second crossing for EdgeFaceSection. Mark it. EdgeFaceSection is now complete!
+						int32& SecondCrossingEdgeLocalIndex = bReverseDirection ? EdgeFaceSection->CrossingEdgeLocalIndex[0] : EdgeFaceSection->CrossingEdgeLocalIndex[1];
+						check(SecondCrossingEdgeLocalIndex == INDEX_NONE);
+						SecondCrossingEdgeLocalIndex = LocalEdgeIndex;
+
+						// Add this intersection to the list of intersections consumed by this contour
+						const bool bFlipGradient = EdgeContour == &ContourPair.ColorContours[InitialFaceContourIndex]; // Flip gradient if the edge of this intersection belongs to the contour that started with the face of the first intersection
+						ContourPair.Intersections.Add_GetRef(NextIntersection).bFlipGradient = bFlipGradient;
+
+
+						// Add new FaceEdge section by choosing the face edge that doesn't correspond with our current faceEdge (if it exists)
+
+						const int32 NextEdgeFace = EdgeToFaces[NextIntersection.EdgeIndex][0] == EdgeFace ? EdgeToFaces[NextIntersection.EdgeIndex][1] : EdgeToFaces[NextIntersection.EdgeIndex][0];
+						if (NextEdgeFace == -1)
+						{
+							// We've hit a boundary.
+							++ContourPair.BoundaryEdgeCount;
+
+							if (ContourPair.LoopVertexCount + ContourPair.BoundaryEdgeCount == 2)
+							{
+								// We've hit both ends of this contour. Since (at least) one end is a boundary, we won't be able to form a closed contour
+								break;
+							}
+
+							check(ContourPair.LoopVertexCount + ContourPair.BoundaryEdgeCount == 1);
+
+							// Pick up at the FirstIntersection and move the opposite direction if possible.
+							EdgeFace = EdgeToFaces[FirstIntersection.EdgeIndex][1];
+							if (EdgeFace == -1)
+							{
+								// First intersection edge was a boundary, so this contour is done. We're NOT closed.
+								++ContourPair.BoundaryEdgeCount;
+								break;
+							}
+
+							check(bReverseDirection == false);
+							bReverseDirection = true;
+							CurrIntersection = FirstIntersection;
+							FaceContour = &ContourPair.ColorContours[0];
+							EdgeContour = &ContourPair.ColorContours[1];
+
+							// ContourPoints are currently just used for debug drawing. Start new contours for reverse section
+							ContourPair.ContourPointCurves[InitialFaceContourIndex].Add(ContourPoints.Num());
+							FaceContourPoints = &ContourPoints.AddDefaulted_GetRef();
+							ContourPair.ContourPointCurves[InitialEdgeContourIndex].Add(ContourPoints.Num());
+							EdgeContourPoints = &ContourPoints.AddDefaulted_GetRef();
+
+							FaceSection = &ContourPair.ColorContours[InitialFaceContourIndex][0];
+							EdgeFaceSection = &EdgeContour->Add_GetRef(FIntersectionContourTriangleSection(EdgeFace));
+							EdgeFaceSection->CrossingEdgeLocalIndex[1] = __internal::GetLocalEdgeIndex(FaceToEdges, EdgeFace, CurrIntersection.EdgeIndex);
+
+							continue;
+						}
+
+						// Setup for next iteration of loop
+						// No change in which contour in the contour pair represents the "face" vs "faceEdge" contour (face section hasn't changed)
+						CurrIntersection = NextIntersection;
+						EdgeFace = NextEdgeFace;
+
+						EdgeFaceSection = &EdgeContour->Add_GetRef(FIntersectionContourTriangleSection(EdgeFace));
+						int32& FirstCrossingLocalIndex = bReverseDirection ? EdgeFaceSection->CrossingEdgeLocalIndex[1] : EdgeFaceSection->CrossingEdgeLocalIndex[0];
+						FirstCrossingLocalIndex = __internal::GetLocalEdgeIndex(FaceToEdges, EdgeFace, CurrIntersection.EdgeIndex);
+						continue;
+					}
+
+					// Check 3b)
+					if (__internal::FindNextIntersectionFace(IntersectionMap, FaceToEdges[CurrIntersection.FaceIndex], EdgeFace, NextIntersection, LocalEdgeIndex))
+					{
+						// We found the second crossing for FaceSection. Mark it.
+						int32& SecondCrossingEdgeLocalIndex = bReverseDirection ? FaceSection->CrossingEdgeLocalIndex[0] : FaceSection->CrossingEdgeLocalIndex[1];
+						check(SecondCrossingEdgeLocalIndex == INDEX_NONE);
+						SecondCrossingEdgeLocalIndex = LocalEdgeIndex;
+
+						// Add this intersection to the list of intersections consumed by this contour
+						const bool bFlipGradient = EdgeContour == &ContourPair.ColorContours[InitialEdgeContourIndex]; // Flip gradient if the edge of this intersection belongs to the contour that started with the face of the first intersection
+						ContourPair.Intersections.Add_GetRef(NextIntersection).bFlipGradient = bFlipGradient;
+
+						// Add new FaceEdge section by choosing the face edge that doesn't correspond with our current faceEdge (if it exists)
+						const int32 NextEdgeFace = EdgeToFaces[NextIntersection.EdgeIndex][0] == CurrIntersection.FaceIndex ? EdgeToFaces[NextIntersection.EdgeIndex][1] : EdgeToFaces[NextIntersection.EdgeIndex][0];
+						if (NextEdgeFace == -1)
+						{
+							// We've hit a boundary.
+							++ContourPair.BoundaryEdgeCount;
+
+							if (ContourPair.LoopVertexCount + ContourPair.BoundaryEdgeCount == 2)
+							{
+								// We've hit both ends of this contour. Since (at least) one end is a boundary, we won't be able to form a closed contour
+								break;
+							}
+
+							check(ContourPair.LoopVertexCount + ContourPair.BoundaryEdgeCount == 1);
+
+							// Pick up at the FirstIntersection and move the opposite direction if possible.
+							EdgeFace = EdgeToFaces[FirstIntersection.EdgeIndex][1];
+							if (EdgeFace == -1)
+							{
+								// First intersection edge was a boundary, so this contour is done. We're NOT closed.
+								++ContourPair.BoundaryEdgeCount;
+								break;
+							}
+
+							check(bReverseDirection == false);
+							bReverseDirection = true;
+							CurrIntersection = FirstIntersection;
+							FaceContour = &ContourPair.ColorContours[0];
+							EdgeContour = &ContourPair.ColorContours[1];
+
+							// ContourPoints are currently just used for debug drawing. Start new contours for reverse section
+							ContourPair.ContourPointCurves[InitialFaceContourIndex].Add(ContourPoints.Num());
+							FaceContourPoints = &ContourPoints.AddDefaulted_GetRef();
+							ContourPair.ContourPointCurves[InitialEdgeContourIndex].Add(ContourPoints.Num());
+							EdgeContourPoints = &ContourPoints.AddDefaulted_GetRef();
+
+							FaceSection = &ContourPair.ColorContours[InitialFaceContourIndex][0];
+							EdgeFaceSection = &EdgeContour->Add_GetRef(FIntersectionContourTriangleSection(EdgeFace));
+							EdgeFaceSection->CrossingEdgeLocalIndex[1] = __internal::GetLocalEdgeIndex(FaceToEdges, EdgeFace, CurrIntersection.EdgeIndex);
+
+							continue;
+						}
+
+						// Setup for next iteration of loop
+						CurrIntersection = NextIntersection;
+
+						// FaceContour and EdgeContour swap
+						FIntersectionContour* TmpContourSwap = EdgeContour;
+						EdgeContour = FaceContour;
+						FaceContour = TmpContourSwap;
+						TArray<FPBDTriangleMeshCollisions::FBarycentricPoint>* TmpPointsSwap = EdgeContourPoints;
+						EdgeContourPoints = FaceContourPoints;
+						FaceContourPoints = TmpPointsSwap;
+
+						FaceSection = EdgeFaceSection;
+
+						EdgeFace = NextEdgeFace;
+						EdgeFaceSection = &EdgeContour->Add_GetRef(FIntersectionContourTriangleSection(EdgeFace));
+						int32& FirstCrossingLocalIndex = bReverseDirection ? EdgeFaceSection->CrossingEdgeLocalIndex[1] : EdgeFaceSection->CrossingEdgeLocalIndex[0];
+						FirstCrossingLocalIndex = __internal::GetLocalEdgeIndex(FaceToEdges, EdgeFace, CurrIntersection.EdgeIndex);
+						continue;
+					}
+
+					// We failed to find the next intersection in this contour. 
+					break;
+				}
+			}
+			return Contours;
+		}
+	} // namespace ContourBuilding
+
+	// methods for Flood-filling contours built by ContourBuilding
+	namespace FloodFill
+	{
+		enum struct EFloodFillRegion : int8
+		{
+			Untouched = 0,
+			Loop,
+			Region0,
+			Region1,
+		};
+		static constexpr EFloodFillRegion FloodFillRegions[2] = { EFloodFillRegion::Region0, EFloodFillRegion::Region1 };
+
+		// internal helper methods
+		namespace __internal
+		{
+			static bool SimplifyContourSegments(const FTriangleMesh& TriangleMesh, const TMultiMap<int32 /*TriangleIndex*/, const FIntersectionContourTriangleSection*>& ContourSegments, TMap<int32, FIntersectionContourTriangleSection>& SimplifiedContourSegments, TArray<int32>& TrianglesWithThinFeatures,
+				TArray<int32>& CandidateSeedTriangles, TVec2<int32>& LoopVertices)
+			{
+				// For purposes of vertex-based flood fill, simplify ContourSegments so there is at most one "segment" per triangle.
+				// Collapse edge crossings such that two crossing cancel out (i.e., even number = no crossing, odd number = crossing).
+				// Keep a list of triangles where there are only even crossing counts--these have thin intersection features and 
+				// are handled with the Triangle-based GIAColors (a black vertex is attracted to a white triangle, but the vertices of
+				// the white triangle are not attracted to black triangles).
+				// This collapsing can effectively generate multiple closed contours (that were originally connected by thin regions).
+				// We will need to flood fill them separately. Track the Segments that were formed by pinching off thin regions. These
+				// are good seeds for flood filling to ensure all of the separate contours are filled.		
+				int32 ContourLoopVertexCount = 0;
+				SimplifiedContourSegments.Reset();
+				CandidateSeedTriangles.Reset();
+				TrianglesWithThinFeatures.Reset();
+				LoopVertices = { INDEX_NONE, INDEX_NONE };
+
+				const TArray<TVec3<int32>>& Elements = TriangleMesh.GetElements();
+
+				TArray<int32> UniqueTriangles;
+				ContourSegments.GetKeys(UniqueTriangles);
+
+				SimplifiedContourSegments.Reserve(UniqueTriangles.Num());
+
+				for (const int32 TriangleIndex : UniqueTriangles)
+				{
+					TVec3<int32> CrossingCounts(0);
+					TVec3<int32> LoopVertexCounts(0);
+					TMultiMap<int32 /*TriangleIndex*/, const FIntersectionContourTriangleSection*>::TConstKeyIterator TriangleIter = ContourSegments.CreateConstKeyIterator(TriangleIndex);
+					for (; TriangleIter; ++TriangleIter)
+					{
+						if (TriangleIter.Value()->LoopVertexLocalIndex != INDEX_NONE)
+						{
+							++(LoopVertexCounts[TriangleIter.Value()->LoopVertexLocalIndex]);
+						}
+						if (TriangleIter.Value()->CrossingEdgeLocalIndex[0] != INDEX_NONE)
+						{
+							++(CrossingCounts[TriangleIter.Value()->CrossingEdgeLocalIndex[0]]);
+						}
+						if (TriangleIter.Value()->CrossingEdgeLocalIndex[1] != INDEX_NONE)
+						{
+							++(CrossingCounts[TriangleIter.Value()->CrossingEdgeLocalIndex[1]]);
+						}
+					}
+
+					const TVec3<int32> SimplifiedCrossingCounts = { CrossingCounts[0] & 1, CrossingCounts[1] & 1, CrossingCounts[2] & 1 };
+					const int32 NonZeroCrossingEdgeCount = SimplifiedCrossingCounts[0] + SimplifiedCrossingCounts[1] + SimplifiedCrossingCounts[2];
+					if (LoopVertexCounts.Max() > 1)
+					{
+						// This can occur in pathological loop cases where we have lots of doubling back. These contours aren't actually closed.
+						return false;
+					}
+					const int32 NonZeroLoopVertexCount = LoopVertexCounts[0] + LoopVertexCounts[1] + LoopVertexCounts[2];
+					if ((NonZeroLoopVertexCount + NonZeroCrossingEdgeCount) % 2 != 0)
+					{
+						// Similarly, this can occur with slivers doubling back. These contours aren't actually closed.
+						return false;
+					}
+
+					if (NonZeroCrossingEdgeCount + NonZeroLoopVertexCount == 0)
+					{
+						TrianglesWithThinFeatures.Add(TriangleIndex);
+						continue;
+					}
+
+					if (NonZeroLoopVertexCount > 0)
+					{
+						// Triangle with loop vertex
+						FIntersectionContourTriangleSection SimplifiedSegment(TriangleIndex);
+						for (int32 LocalVertexIdx = 0; LocalVertexIdx < 3; ++LocalVertexIdx)
+						{
+							if (LoopVertexCounts[LocalVertexIdx])
+							{
+								const int32 LoopVertexGlobalIdx = Elements[TriangleIndex][LocalVertexIdx];
+								if (LoopVertices[0] != LoopVertexGlobalIdx && LoopVertices[1] != LoopVertexGlobalIdx)
+								{
+									check(ContourLoopVertexCount < 2);
+									LoopVertices[ContourLoopVertexCount++] = LoopVertexGlobalIdx;
+								}
+								SimplifiedSegment.LoopVertexLocalIndex = LocalVertexIdx; // It doesn't matter in the Simplified context if there's more than one loop vertex and we only record one because we just care if a loop vertex exists or not.
+							}
+						}
+						SimplifiedContourSegments.Add(TriangleIndex, MoveTemp(SimplifiedSegment));
+					}
+					else
+					{
+						// Crossings
+						FIntersectionContourTriangleSection SimplifiedSegment(TriangleIndex);
+						check(NonZeroCrossingEdgeCount <= 2);
+						int32 CrossingEdgeIdx = 0;
+						for (int32 LocalEdgeIdx = 0; LocalEdgeIdx < 3; ++LocalEdgeIdx)
+						{
+							if (SimplifiedCrossingCounts[LocalEdgeIdx])
+							{
+								SimplifiedSegment.CrossingEdgeLocalIndex[CrossingEdgeIdx++] = LocalEdgeIdx;
+							}
+						}
+						SimplifiedContourSegments.Add(TriangleIndex, MoveTemp(SimplifiedSegment));
+						if (CrossingCounts[0] && CrossingCounts[1] && CrossingCounts[2])
+						{
+							// This used to have crossings on 3 sides, and one was pinched off. This is a candidate seed triangle
+							CandidateSeedTriangles.Add(TriangleIndex);
+						}
+					}
+				}
+				return true;
+			}
+
+			// Returns success
+			static bool OneFloodFillStep(const TMap<int32 /*TriangleIndex*/, FIntersectionContourTriangleSection>& ContourSegments, const TArray<TVec3<int32>>& Elements, const TConstArrayView<TArray<int32>>& PointToTriangleMap, TArray<int32>& Queue, int32& RegionSize, TArrayView<EFloodFillRegion>& VertexRegions,
+				const EFloodFillRegion RegionColor, const EFloodFillRegion OtherRegionColor)
+			{
+				const int32 CurrVertex = Queue.Pop(false);
+				for (int32 NeighborTri : PointToTriangleMap[CurrVertex])
+				{
+					if (const FIntersectionContourTriangleSection* TriSection = ContourSegments.Find(NeighborTri))
+					{
+						const int32 CurrVertexLocalIndex = CurrVertex == Elements[NeighborTri][0] ? 0 : CurrVertex == Elements[NeighborTri][1] ? 1 : CurrVertex == Elements[NeighborTri][2] ? 2 : INDEX_NONE;
+						check(CurrVertexLocalIndex != INDEX_NONE);
+
+						if (TriSection->LoopVertexLocalIndex != INDEX_NONE)
+						{
+							// Triangles with loop vertices don't connect other vertices within the triangle
+							// We will mark loop vertices elsewhere
+							continue;
+						}
+
+						check(TriSection->CrossingEdgeLocalIndex[0] != INDEX_NONE && TriSection->CrossingEdgeLocalIndex[1] != INDEX_NONE);
+						check(TriSection->CrossingEdgeLocalIndex[0] != TriSection->CrossingEdgeLocalIndex[1]);
+
+						// Order such that SecondCrossEdge = (FirstCrossEdge + 1)%3
+						const bool bSwapEdgeOrder = TriSection->CrossingEdgeLocalIndex[1] != (TriSection->CrossingEdgeLocalIndex[0] + 1) % 3;
+						const int32 FirstCrossEdge = bSwapEdgeOrder ? TriSection->CrossingEdgeLocalIndex[1] : TriSection->CrossingEdgeLocalIndex[0];
+						// const int32 SecondCrossEdge = bSwapEdgeOrder ? TriSection->CrossingEdgeLocalIndex[0] : TriSection->CrossingEdgeLocalIndex[1];
+
+						if (FirstCrossEdge == CurrVertexLocalIndex)
+						{
+							// CurrVertex is on the same side as CurrVertex - 1
+							const int32 TriVertex = Elements[NeighborTri][(CurrVertexLocalIndex + 2) % 3];
+							if (VertexRegions[TriVertex] == EFloodFillRegion::Untouched)
+							{
+								VertexRegions[TriVertex] = RegionColor;
+								++RegionSize;
+								Queue.Add(TriVertex);
+							}
+							else if (VertexRegions[TriVertex] == OtherRegionColor)
+							{
+								// We don't have a closed contour after all.
+								return false;
+							}
+						}
+						else if (FirstCrossEdge == (CurrVertexLocalIndex + 1) % 3)
+						{
+							// CurrVertex is on the same side as CurrVertex + 1
+							const int32 TriVertex = Elements[NeighborTri][(CurrVertexLocalIndex + 1) % 3];
+							if (VertexRegions[TriVertex] == EFloodFillRegion::Untouched)
+							{
+								VertexRegions[TriVertex] = RegionColor;
+								++RegionSize;
+								Queue.Add(TriVertex);
+							}
+							else if (VertexRegions[TriVertex] == OtherRegionColor)
+							{
+								// We don't have a closed contour after all.
+								return false;
+							}
+						}
+						// else CurrVertex is not on the same side with any other vertices in this triangle
+					}
+					else
+					{
+						// Contour doesn't pass through this tri. Just add all (untouched) points to this Region
+						for (int32 VertexLocalIndex = 0; VertexLocalIndex < 3; ++VertexLocalIndex)
+						{
+							const int32 TriVertex = Elements[NeighborTri][VertexLocalIndex];
+							if (VertexRegions[TriVertex] == EFloodFillRegion::Untouched)
+							{
+								VertexRegions[TriVertex] = RegionColor;
+								++RegionSize;
+								Queue.Add(TriVertex);
+							}
+							else if (VertexRegions[TriVertex] == OtherRegionColor)
+							{
+								// We don't have a closed contour after all.
+								return false;
+							}
+						}
+					}
+				}
+				return true;
+			}
+
+			// Flood fill a single contour
+			static bool FloodFillContourColor(const FTriangleMesh& TriangleMesh, const int32 NumParticles, const int32 Offset, const TMultiMap<int32 /*TriangleIndex*/, const FIntersectionContourTriangleSection*>& ContourSegments, int32 ContourIndex, bool bIsColorB,
+				TArray<FPBDTriangleMeshCollisions::FGIAColor>& VertexGIAColors, TArray<FPBDTriangleMeshCollisions::FGIAColor>& TriangleGIAColors)
+			{
+				// Simplify Contour 
+				TMap<int32, FIntersectionContourTriangleSection> SimplifiedContourSegments;
+				TArray<int32> TrianglesWithThinFeatures;
+				TArray<int32> CandidateSeedTriangles;
+				TVec2<int32> LoopVertices;
+				if (!SimplifyContourSegments(TriangleMesh, ContourSegments, SimplifiedContourSegments, TrianglesWithThinFeatures, CandidateSeedTriangles, LoopVertices))
+				{
+					return false;
+				}
+
+				const TArray<TVec3<int32>>& Elements = TriangleMesh.GetElements();
+				const TConstArrayView<TArray<int32>>& PointToTrianglemMap = TriangleMesh.GetPointToTriangleMap();
+
+				// Main contour flood fill
+				if (SimplifiedContourSegments.Num())
+				{
+					// This collapsing can effectively generate multiple closed contours (that were originally connected by thin regions).
+					// We will need to flood fill them separately.
+					if (CandidateSeedTriangles.Num() == 0)
+					{
+						// If there weren't any candidates found by pinching regions, just start with first non-loop segment
+						TMap<int32, FIntersectionContourTriangleSection>::TConstIterator Iter = SimplifiedContourSegments.CreateConstIterator();
+						for (; Iter; ++Iter)
+						{
+							const FIntersectionContourTriangleSection& FirstSection = Iter.Value();
+							if (FirstSection.LoopVertexLocalIndex != INDEX_NONE)
+							{
+								continue;
+							}
+
+							const int32 FirstVertex0 = FirstSection.CrossingEdgeLocalIndex[0];
+							const int32 FirstVertex1 = (FirstSection.CrossingEdgeLocalIndex[0] + 1) % 3;
+							if (Elements[Iter.Key()][FirstVertex0] == LoopVertices[0] || Elements[Iter.Key()][FirstVertex0] == LoopVertices[1] ||
+								Elements[Iter.Key()][FirstVertex1] == LoopVertices[0] || Elements[Iter.Key()][FirstVertex1] == LoopVertices[1])
+							{
+								continue;
+							}
+							CandidateSeedTriangles.Add(Iter.Key());
+							break;
+						}
+					}
+
+					for (const int32 SeedTriangle : CandidateSeedTriangles)
+					{
+						// Initialize queue for flood fill with first Section in Contour
+						const FIntersectionContourTriangleSection& FirstSection = SimplifiedContourSegments.FindChecked(SeedTriangle);
+						check(FirstSection.CrossingEdgeLocalIndex[0] != INDEX_NONE);
+						const int32 FirstVertex0 = FirstSection.CrossingEdgeLocalIndex[0];
+						const int32 FirstVertex1 = (FirstSection.CrossingEdgeLocalIndex[0] + 1) % 3;
+
+						// Check if this Triangle already has been colored on either side.
+						if (VertexGIAColors[Elements[FirstSection.TriangleIndex][FirstVertex0] - Offset].HasContourColorSet(ContourIndex)
+							|| VertexGIAColors[Elements[FirstSection.TriangleIndex][FirstVertex1] - Offset].HasContourColorSet(ContourIndex))
+						{
+							continue;
+						}
+
+						// Flood fill both sides of contour. Stop when one side is complete -- we want the smaller region.
+						// For now, assume triangle areas are similar, and just calculate "smaller" based on number of vertices
+
+						TArray<int32 /*VertexIndex*/> Queue[2];
+						int32 RegionSizes[2] = { 0, 0 };
+						TArray<EFloodFillRegion> VertexRegionsNoOffset;
+						static_assert((int8)EFloodFillRegion::Untouched == 0);
+						VertexRegionsNoOffset.SetNumZeroed(NumParticles);
+
+						TArrayView<EFloodFillRegion> VertexRegions(VertexRegionsNoOffset.GetData() - Offset, Offset + NumParticles);
+
+						Queue[0].Add(Elements[FirstSection.TriangleIndex][FirstVertex0]);
+						Queue[1].Add(Elements[FirstSection.TriangleIndex][FirstVertex1]);
+
+						VertexRegions[Elements[FirstSection.TriangleIndex][FirstVertex0]] = FloodFillRegions[0];
+						VertexRegions[Elements[FirstSection.TriangleIndex][FirstVertex1]] = FloodFillRegions[1];
+						++RegionSizes[0];
+						++RegionSizes[1];
+
+						// Set Loop vertices so we don't cross them when flooding
+						if (LoopVertices[0] != INDEX_NONE)
+						{
+							VertexRegions[LoopVertices[0]] = EFloodFillRegion::Loop;
+						}
+						if (LoopVertices[1] != INDEX_NONE)
+						{
+							VertexRegions[LoopVertices[1]] = EFloodFillRegion::Loop;
+						}
+
+						while (Queue[0].Num() > 0 && Queue[1].Num() > 0)
+						{
+							for (int32 RegionIndex = 0; RegionIndex < 2; ++RegionIndex)
+							{
+								if (!OneFloodFillStep(SimplifiedContourSegments, Elements, PointToTrianglemMap, Queue[RegionIndex], RegionSizes[RegionIndex], VertexRegions, FloodFillRegions[RegionIndex], FloodFillRegions[RegionIndex ^ 1]))
+								{
+									return false;
+								}
+							}
+						}
+
+						const int32 CompletedQueueIndex = Queue[0].Num() == 0 ? 0 : 1;
+						if (RegionSizes[CompletedQueueIndex] >= RegionSizes[CompletedQueueIndex ^ 1])
+						{
+							// Run FloodFill on incomplete queue until complete/larger than completed queue
+							while (Queue[CompletedQueueIndex ^ 1].Num() > 0 && (RegionSizes[CompletedQueueIndex] >= RegionSizes[CompletedQueueIndex ^ 1]))
+							{
+								if (!OneFloodFillStep(SimplifiedContourSegments, Elements, PointToTrianglemMap, Queue[CompletedQueueIndex ^ 1], RegionSizes[CompletedQueueIndex ^ 1], VertexRegions, FloodFillRegions[CompletedQueueIndex ^ 1], FloodFillRegions[CompletedQueueIndex]))
+								{
+									return false;
+								}
+							}
+						}
+
+						if (RegionSizes[0] == RegionSizes[1])
+						{
+							// Both regions are the same size. Can't determine inside.
+							return false;
+						}
+
+						// Mark VertexGIAColors for smaller region
+						const EFloodFillRegion SmallerRegionColor = RegionSizes[0] < RegionSizes[1] ? FloodFillRegions[0] : FloodFillRegions[1];
+						for (int ParticleIndexNoOffset = 0; ParticleIndexNoOffset < NumParticles; ++ParticleIndexNoOffset)
+						{
+							if (VertexRegionsNoOffset[ParticleIndexNoOffset] == SmallerRegionColor)
+							{
+								VertexGIAColors[ParticleIndexNoOffset].SetContourColor(ContourIndex, bIsColorB);
+							}
+						}
+					}
+				}
+
+				// Loop vertex handling
+				if (ContourIndex == FPBDTriangleMeshCollisions::FGIAColor::LoopContourIndex)
+				{
+					if (LoopVertices[0] != INDEX_NONE)
+					{
+						VertexGIAColors[LoopVertices[0] - Offset].SetContourColor(ContourIndex, bIsColorB);
+					}
+					if (LoopVertices[1] != INDEX_NONE)
+					{
+						VertexGIAColors[LoopVertices[1] - Offset].SetContourColor(ContourIndex, bIsColorB);
+					}
+				}
+
+				// Thin region triangles
+				for (int32 TriangleIndex : TrianglesWithThinFeatures)
+				{
+					TriangleGIAColors[TriangleIndex].SetContourColor(ContourIndex, bIsColorB);
+				}
+				return true;
+			}
+		} // namespace __internal
+
+		// Flood fill contours built by ContourBuilding to generate VertexGIAColors. Identify triangles with thin intersections (TriangleGIAColors)
+		static void FloodFillContours(const FTriangleMesh& TriangleMesh, const int32 NumParticles, const int32 Offset, TArray<FIntersectionContourPair>& Contours, TArray<FPBDTriangleMeshCollisions::FGIAColor>& VertexGIAColors, TArray<FPBDTriangleMeshCollisions::FGIAColor>& TriangleGIAColors)
+		{
+			TRACE_CPUPROFILER_EVENT_SCOPE(ChaosFPBDTriangleMeshCollisions_FloodFillContours);
+			const TArray<TVec3<int32>>& Elements = TriangleMesh.GetElements();
+			const TConstArrayView<TArray<int32>>& PointToTrianglemMap = TriangleMesh.GetPointToTriangleMap();
+
+			// If we used a critsect/made writing to the GIAColors threadsafe with atomics, we parallelize by contour. Not sure if it's worth the overhead though.
+
+			int32 CurrentContourIndex = 1; // ContourIndex 0 is reserved for loop 
+			for (const FIntersectionContourPair& Contour : Contours)
+			{
+				if (!Contour.bIsClosed)
+				{
+					continue;
+				}
+
+				TMultiMap<int32 /*TriangleIndex*/, const FIntersectionContourTriangleSection*> ContourSegments;
+				const bool bCombineContours = Contour.LoopVertexCount > 0; // Loop contours were built still as a contour pair, but really represent a single contour doubling back on itself.
+				const int32 ContourSegmentReserveNum = bCombineContours ? Contour.ColorContours[0].Num() + Contour.ColorContours[1].Num() : FMath::Max(Contour.ColorContours[0].Num(), Contour.ColorContours[1].Num());
+				ContourSegments.Reserve(ContourSegmentReserveNum);
+
+				for (const FIntersectionContourTriangleSection& Section : Contour.ColorContours[0])
+				{
+					ContourSegments.Emplace(Section.TriangleIndex, &Section);
+				}
+
+				if (bCombineContours)
+				{
+					for (const FIntersectionContourTriangleSection& Section : Contour.ColorContours[1])
+					{
+						ContourSegments.Emplace(Section.TriangleIndex, &Section);
+					}
+				}
+
+				const int32 ContourIndex = bCombineContours ? FPBDTriangleMeshCollisions::FGIAColor::LoopContourIndex : (((CurrentContourIndex++) - 1) % 31) + 1;
+
+				if (!__internal::FloodFillContourColor(TriangleMesh, NumParticles, Offset, ContourSegments, ContourIndex, bCombineContours, VertexGIAColors, TriangleGIAColors))
+				{
+					const_cast<FIntersectionContourPair&>(Contour).bIsClosed = false;
+					continue;
+				}
+
+				if (!bCombineContours)
+				{
+					ContourSegments.Reset();
+
+					for (const FIntersectionContourTriangleSection& Section : Contour.ColorContours[1])
+					{
+						ContourSegments.Emplace(Section.TriangleIndex, &Section);
+					}
+
+					if (!__internal::FloodFillContourColor(TriangleMesh, NumParticles, Offset, ContourSegments, ContourIndex, true, VertexGIAColors, TriangleGIAColors))
+					{
+						const_cast<FIntersectionContourPair&>(Contour).bIsClosed = false;
+						continue;
+					}
+				}
+			}
+		}
+
+		// Once FloodFill has completed, we can mark the IntersectionContourTypes for debug drawing ContourPoints.
+		static void AssignContourPointTypes(const TArray<FIntersectionContourPair>& Contours, const TArray<TArray<FPBDTriangleMeshCollisions::FBarycentricPoint>>& ContourPoints, TArray<FPBDTriangleMeshCollisions::FContourType>& IntersectionContourTypes)
+		{
+			IntersectionContourTypes.Reset();
+			IntersectionContourTypes.SetNumZeroed(ContourPoints.Num());
+			for (const FIntersectionContourPair& ContourPair : Contours)
+			{
+				if (ContourPair.bIsClosed)
+				{
+					if (ContourPair.LoopVertexCount > 0)
+					{
+						for (const int32 PointIndex : ContourPair.ContourPointCurves[0])
+						{
+							IntersectionContourTypes[PointIndex] = FPBDTriangleMeshCollisions::FContourType::Loop;
+						}
+						for (const int32 PointIndex : ContourPair.ContourPointCurves[1])
+						{
+							IntersectionContourTypes[PointIndex] = FPBDTriangleMeshCollisions::FContourType::Loop;
+						}
+					}
+					else
+					{
+						for (const int32 PointIndex : ContourPair.ContourPointCurves[0])
+						{
+							IntersectionContourTypes[PointIndex] = FPBDTriangleMeshCollisions::FContourType::Contour0;
+						}
+						for (const int32 PointIndex : ContourPair.ContourPointCurves[1])
+						{
+							IntersectionContourTypes[PointIndex] = FPBDTriangleMeshCollisions::FContourType::Contour1;
+						}
+					}
+				}
+			}
+		}
+	} // namespace FloodFill
+} // namespace GIA
+
+// Calculating gradient direction to minimize contour length.
+namespace ContourMinimization
+{
+	static void BuildLocalContourMinimizationIntersection(const FEdgeFaceIntersection& EdgeFaceIntersection, const FTriangleMesh& TriangleMesh, const FSegmentMesh& SegmentMesh, const FSolverParticles& Particles, FPBDTriangleMeshCollisions::FContourMinimizationIntersection& ContourIntersection)
+	{
+		ContourIntersection.EdgeVertices = SegmentMesh.GetElements()[EdgeFaceIntersection.EdgeIndex];
+		ContourIntersection.FaceVertices = TriangleMesh.GetElements()[EdgeFaceIntersection.FaceIndex];
+		ContourIntersection.LocalGradientVector = FSolverVec3(0.f);
+
+		// Calculate local gradient
+		const TArray<TVec2<int32>>& EdgeToFaces = TriangleMesh.GetEdgeToFaces();
+		const TArray<TVec3<int32>>& Elements = TriangleMesh.GetElements();
+
+		const FSolverVec3& EdgePosition0 = Particles.X(ContourIntersection.EdgeVertices[0]);
+		const FSolverVec3& EdgePosition1 = Particles.X(ContourIntersection.EdgeVertices[1]);
+		const FSolverVec3 EdgeDir = EdgePosition1 - EdgePosition0; // Unnormalized
+		const FSolverReal EdgeDirDotNormal = FSolverVec3::DotProduct(EdgeDir, EdgeFaceIntersection.FaceNormal);
+		if (FMath::Abs(EdgeDirDotNormal) >= SMALL_NUMBER)
+		{
+			const FSolverVec3 NOverEDotN = EdgeFaceIntersection.FaceNormal * ((FSolverReal)1.f / EdgeDirDotNormal);
+
+			for (int32 LocalFaceIndex = 0; LocalFaceIndex < 2; ++LocalFaceIndex)
+			{
+				const int32 EdgeFace_i = EdgeToFaces[EdgeFaceIntersection.EdgeIndex][LocalFaceIndex];
+				if (EdgeFace_i == -1)
+				{
+					continue;
+				}
+
+				const FSolverVec3& P0 = Particles.X(Elements[EdgeFace_i][0]);
+				const FSolverVec3& P1 = Particles.X(Elements[EdgeFace_i][1]);
+				const FSolverVec3& P2 = Particles.X(Elements[EdgeFace_i][2]);
+
+				// B_i
+				const TTriangle<FSolverReal> Triangle_i(P0, P1, P2);
+
+				// N_i
+				const FSolverVec3 Normal_i = Triangle_i.GetNormal();
+
+				// R_i
+				FSolverVec3 IntersectionVector_i = FSolverVec3::CrossProduct(EdgeFaceIntersection.FaceNormal, Normal_i);
+				IntersectionVector_i.Normalize();
+				// This should point inward from the Edge. Flip if necessary.
+
+				// SideDirection = Normal_i ^ (EdgeDir) // NOTE: Need to calculate EdgeDir based on vertex order within the triangle, not what's in the SegmentMesh to ensure SideDirection points INTO from triangle.
+				// If Dot(IntersectionVector_i, SideDirection) < 0, need to flip IntersectionVector_i
+				// Dot(IntersectionVector_i, SideDirection) = Dot((CollisionPoint.Normal ^ Normal_i), (Normal_i ^ EdgeDir))
+				// Using triple product rules, this is
+				// Dot(CollisionPoint.Normal, Normal_i) * Dot(Normal_i, EdgeDir) - Dot(CollisionPoint.Normal, EdgeDir) * Dot(Normal_i, Normal_i)
+				const int32 EdgePointLocalIndex0 = ContourIntersection.EdgeVertices[0] == Elements[EdgeFace_i][0] ? 0 : ContourIntersection.EdgeVertices[0] == Elements[EdgeFace_i][1] ? 1 : ContourIntersection.EdgeVertices[0] == Elements[EdgeFace_i][2] ? 2 : INDEX_NONE;
+				check(EdgePointLocalIndex0 != INDEX_NONE); // Otherwise the EdgeToFaces is messed up because this edge doesn't below to this face
+				check(ContourIntersection.EdgeVertices[1] == Elements[EdgeFace_i][(EdgePointLocalIndex0 + 1) % 3] || ContourIntersection.EdgeVertices[1] == Elements[EdgeFace_i][(EdgePointLocalIndex0 + 2) % 3]);
+				const bool bEdgePointsFlipped = ContourIntersection.EdgeVertices[1] != Elements[EdgeFace_i][(EdgePointLocalIndex0 + 1) % 3];
+				const FSolverReal TestFlipRValue = (FSolverVec3::DotProduct(EdgeFaceIntersection.FaceNormal, Normal_i) * FSolverVec3::DotProduct(Normal_i, EdgeDir) - FSolverVec3::DotProduct(EdgeFaceIntersection.FaceNormal, EdgeDir)) * (bEdgePointsFlipped ? (FSolverReal)-1.f : (FSolverReal)1.f);
+				if (TestFlipRValue < 0)
+				{
+					IntersectionVector_i *= (FSolverReal)-1.f;
+				}
+
+				ContourIntersection.LocalGradientVector += IntersectionVector_i - (FSolverReal)2.f * FSolverVec3::DotProduct(EdgeDir, IntersectionVector_i) * NOverEDotN;
+			}
+		}
+
+		ContourIntersection.GlobalGradientVector = ContourIntersection.LocalGradientVector;
+	}
+
+	static void BuildLocalContourMinimizationIntersections(const FTriangleMesh& TriangleMesh, const FSolverParticles& Particles, const TArray<FEdgeFaceIntersection>& Intersections, TArray<FPBDTriangleMeshCollisions::FContourMinimizationIntersection>& ContourMinimizationIntersections)
+	{
+		TRACE_CPUPROFILER_EVENT_SCOPE(ChaosFPBDTriangleMeshCollisions_BuildLocalContourMinimizationIntersections);
+		const FSegmentMesh& SegmentMesh = TriangleMesh.GetSegmentMesh();
+
+		ContourMinimizationIntersections.SetNum(Intersections.Num());
+		PhysicsParallelFor(Intersections.Num(),
+			[&TriangleMesh, &SegmentMesh, &Particles, &Intersections, &ContourMinimizationIntersections](int32 IntersectionIndex)
+			{
+				const FEdgeFaceIntersection& EdgeFaceIntersection = Intersections[IntersectionIndex];
+				FPBDTriangleMeshCollisions::FContourMinimizationIntersection& ContourIntersection = ContourMinimizationIntersections[IntersectionIndex];
+				BuildLocalContourMinimizationIntersection(EdgeFaceIntersection, TriangleMesh, SegmentMesh, Particles, ContourIntersection);
+			}
+		);
+	}
+
+	static void BuildGlobalContourMinimizationIntersections(const FTriangleMesh& TriangleMesh, const FSolverParticles& Particles, const TArray<GIA::FIntersectionContourPair>& IntersectionContours, TArray<FPBDTriangleMeshCollisions::FContourMinimizationIntersection>& ContourMinimizationIntersections)
+	{
+		TRACE_CPUPROFILER_EVENT_SCOPE(ChaosFPBDTriangleMeshCollisions_BuildGlobalContourMinimizationIntersections);
+		const FSegmentMesh& SegmentMesh = TriangleMesh.GetSegmentMesh();
+		// Only do contour minimization for non-closed contours. Closed contours will be handled via GIA
+
+		// Determine total number of non-closed contours and number of segments
+		struct FContourRange
+		{
+			int32 ContourPairIndex;
+			int32 GlobalIndexStart;
+		};
+
+		TArray<FContourRange> ContourRanges;
+		ContourRanges.Reset(IntersectionContours.Num());
+		int32 IndexStart = 0;
+		for (int32 ContourPairIndex = 0; ContourPairIndex < IntersectionContours.Num(); ++ContourPairIndex)
+		{
+			if (IntersectionContours[ContourPairIndex].bIsClosed && IntersectionContours[ContourPairIndex].LoopVertexCount == 0)
+			{
+				continue;
+			}
+			ContourRanges.Add({ ContourPairIndex, IndexStart });
+			IndexStart += IntersectionContours[ContourPairIndex].Intersections.Num();
+		}
+
+		if (ContourRanges.Num() == 0)
+		{
+			return;
+		}
+
+		ContourMinimizationIntersections.SetNum(IndexStart);
+
+		// For now, parallelize by contour. Might want to balance by contour length?
+		PhysicsParallelFor(ContourRanges.Num(),
+			[&TriangleMesh, &SegmentMesh, &Particles, &ContourRanges, &IntersectionContours, &ContourMinimizationIntersections](int32 ContourRangeIndex)
+			{ 
+				const TArray<FEdgeFaceIntersection>& Intersections = IntersectionContours[ContourRanges[ContourRangeIndex].ContourPairIndex].Intersections;
+				const int32 GlobalIndexStart = ContourRanges[ContourRangeIndex].GlobalIndexStart;
+				FSolverVec3 GlobalGradient(0.f);
+				for (int32 SegmentIndex = 0; SegmentIndex < Intersections.Num(); ++SegmentIndex)
+				{
+					const FEdgeFaceIntersection& EdgeFaceIntersection = Intersections[SegmentIndex];
+					FPBDTriangleMeshCollisions::FContourMinimizationIntersection& ContourIntersection = ContourMinimizationIntersections[GlobalIndexStart + SegmentIndex];
+					BuildLocalContourMinimizationIntersection(EdgeFaceIntersection, TriangleMesh, SegmentMesh, Particles, ContourIntersection);
+					GlobalGradient += ContourIntersection.LocalGradientVector * (EdgeFaceIntersection.bFlipGradient ? -1.f : 1.f);
+				}
+
+				for (int32 SegmentIndex = 0; SegmentIndex < Intersections.Num(); ++SegmentIndex)
+				{
+					const FEdgeFaceIntersection& EdgeFaceIntersection = Intersections[SegmentIndex];
+					FPBDTriangleMeshCollisions::FContourMinimizationIntersection& ContourIntersection = ContourMinimizationIntersections[GlobalIndexStart + SegmentIndex];
+					ContourIntersection.GlobalGradientVector = GlobalGradient * (EdgeFaceIntersection.bFlipGradient ? -1.f : 1.f);
+				}
+			}
+		);		
+	}
+} // namespace ContourMinimization
+
+void FPBDTriangleMeshCollisions::Init(const FSolverParticles& Particles)
+{
+	if (TriangleMesh.GetNumElements() == 0)
+	{
+		return;
+	}
+
+	{
+		TRACE_CPUPROFILER_EVENT_SCOPE(ChaosFPBDTriangleMeshCollisions_BuildBVH);
+		TriangleMesh.BuildBVH(static_cast<const TArrayView<const FSolverVec3>&>(Particles.XArray()), BVH);
+	}
+
+	ContourMinimizationIntersections.Reset();
+	VertexGIAColors.Reset();
+	VertexGIAColors.SetNumZeroed(NumParticles);
+	TriangleGIAColors.Reset();
+	TriangleGIAColors.SetNumZeroed(TriangleMesh.GetNumElements());
+	IntersectionContourPoints.Reset();
+	IntersectionContourTypes.Reset();
+
+	if (!bGlobalIntersectionAnalysis && !bContourMinimization)
+	{
+		return;
+	}
+
+	// Detect all EdgeFace Intersections
+	TArray<FEdgeFaceIntersection> Intersections = FindEdgeFaceIntersections(TriangleMesh, BVH, Particles);
+	if (Intersections.Num() == 0)
+	{
+		return;
+	}
+
+	TArray<GIA::FIntersectionContourPair> IntersectionContours;
+	if (bGlobalIntersectionAnalysis)
+	{ 
+		// Walk EdgeFace intersections to build global contours
+		IntersectionContours = GIA::ContourBuilding::BuildIntersectionContours(TriangleMesh, Intersections, IntersectionContourPoints);
+		// Flood fill global contours to determine intersecting regions.
+		GIA::FloodFill::FloodFillContours(TriangleMesh, NumParticles, Offset, IntersectionContours, VertexGIAColors, TriangleGIAColors);
+		GIA::FloodFill::AssignContourPointTypes(IntersectionContours, IntersectionContourPoints, IntersectionContourTypes);
+	}
+
+	if (bContourMinimization)
+	{
+		if (bGlobalIntersectionAnalysis)
+		{
+			// Global contours which are non-closed or loop are handled via ContourMinimization impulses. Build global gradient (by adding contribution across all intersections per contour).
+			ContourMinimization::BuildGlobalContourMinimizationIntersections(TriangleMesh, Particles, IntersectionContours, ContourMinimizationIntersections);
+		}
+		else
+		{
+			// Just build local gradient for all Intersections
+			ContourMinimization::BuildLocalContourMinimizationIntersections(TriangleMesh, Particles, Intersections, ContourMinimizationIntersections);
+		}
+	}
+}
+}  // End namespace Chaos::Softs
diff --git a/Engine/Source/Runtime/Experimental/Chaos/Private/Chaos/PBDTriangleMeshIntersections.cpp b/Engine/Source/Runtime/Experimental/Chaos/Private/Chaos/PBDTriangleMeshIntersections.cpp
new file mode 100644
--- /dev/null
+++ b/Engine/Source/Runtime/Experimental/Chaos/Private/Chaos/PBDTriangleMeshIntersections.cpp
@@ -0,0 +1,51 @@
+// Copyright Epic Games, Inc. All Rights Reserved.
+#include "Chaos/PBDTriangleMeshIntersections.h"
+#include "Chaos/PBDSoftsSolverParticles.h"
+
+float Chaos_TriangleIntersections_MaxDelta = 0.01f;
+FAutoConsoleVariableRef CVarChaosTriangleIntersectionMaxImpulse(TEXT("p.Chaos.TriangleIntersections.MaxDelta"), Chaos_TriangleIntersections_MaxDelta, TEXT("Maximum delta position applied to resolve triangle intersections."));
+
+namespace Chaos::Softs {
+
+static inline FSolverVec3 GetDelta(const FPBDTriangleMeshCollisions::FContourMinimizationIntersection& Intersection, const FSolverReal MaxDelta, const FSolverReal RegularizeEpsilonSq)
+{
+	// Using pre-calculated "GlobalGradientVector". Whether this is actually globally calculated or a copy of the local gradient depends on settings in PBDTriangleMeshCollisions (default is global)
+	FSolverReal GradientLength;
+	FSolverVec3 GradientDir;
+	Intersection.GlobalGradientVector.ToDirectionAndLength(GradientDir, GradientLength);
+
+	return MaxDelta * GradientLength * FMath::InvSqrt(GradientLength * GradientLength + RegularizeEpsilonSq) * GradientDir;
+}
+
+
+void FPBDTriangleMeshIntersections::Apply(FSolverParticles& Particles, const TArray<FPBDTriangleMeshCollisions::FContourMinimizationIntersection>& Intersections, const FSolverReal Dt) const
+{
+	TRACE_CPUPROFILER_EVENT_SCOPE(ChaosFPBDTriangleMeshIntersections_Apply);
+	static FSolverReal RegularizeEpsilonSq = 1.f;
+	for( const FPBDTriangleMeshCollisions::FContourMinimizationIntersection& Intersection : Intersections )
+	{
+		const FSolverVec3 Delta = GetDelta(Intersection, Chaos_TriangleIntersections_MaxDelta, RegularizeEpsilonSq);
+		if (Particles.InvM(Intersection.EdgeVertices[0]) > 0)
+		{
+			Particles.P(Intersection.EdgeVertices[0]) += Delta;
+		}
+		if (Particles.InvM(Intersection.EdgeVertices[1]) > 0)
+		{
+			Particles.P(Intersection.EdgeVertices[1]) += Delta;
+		}
+		if (Particles.InvM(Intersection.FaceVertices[0]) > 0)
+		{
+			Particles.P(Intersection.FaceVertices[0]) -= Delta;
+		}
+		if (Particles.InvM(Intersection.FaceVertices[1]) > 0)
+		{
+			Particles.P(Intersection.FaceVertices[1]) -= Delta;
+		}
+		if (Particles.InvM(Intersection.FaceVertices[2]) > 0)
+		{
+			Particles.P(Intersection.FaceVertices[2]) -= Delta;
+		}
+	}
+}
+
+}  // End namespace Chaos::Softs
diff --git a/Engine/Source/Runtime/Experimental/Chaos/Private/Chaos/TriangleMesh.cpp b/Engine/Source/Runtime/Experimental/Chaos/Private/Chaos/TriangleMesh.cpp
--- a/Engine/Source/Runtime/Experimental/Chaos/Private/Chaos/TriangleMesh.cpp
+++ b/Engine/Source/Runtime/Experimental/Chaos/Private/Chaos/TriangleMesh.cpp
@@ -4,6 +4,7 @@
 #include "Chaos/Box.h"
 #include "Chaos/Defines.h"
 #include "Chaos/Plane.h"
+#include "Chaos/Triangle.h"
 #include "Chaos/TriangleCollisionPoint.h"
 #include "HAL/IConsoleManager.h"
 #include "Math/NumericLimits.h"
@@ -684,7 +685,7 @@ struct OrderedEdgeKeyFuncs : BaseKeyFuncs<TVec2<int32>, TVec2<int32>, false>
 	}
 };
 
-FSegmentMesh& FTriangleMesh::GetSegmentMesh()
+const FSegmentMesh& FTriangleMesh::GetSegmentMesh() const
 {
 	if (MSegmentMesh.GetNumElements() != 0)
 	{
@@ -759,13 +760,13 @@ FSegmentMesh& FTriangleMesh::GetSegmentMesh()
 	return MSegmentMesh;
 }
 
-const TArray<TVec3<int32>>& FTriangleMesh::GetFaceToEdges()
+const TArray<TVec3<int32>>& FTriangleMesh::GetFaceToEdges() const
 {
 	GetSegmentMesh();
 	return MFaceToEdges;
 }
 
-const TArray<TVec2<int32>>& FTriangleMesh::GetEdgeToFaces()
+const TArray<TVec2<int32>>& FTriangleMesh::GetEdgeToFaces() const
 {
 	GetSegmentMesh();
 	return MEdgeToFaces;
@@ -774,7 +775,7 @@ const TArray<TVec2<int32>>& FTriangleMesh::GetEdgeToFaces()
 
 TSet<int32> FTriangleMesh::GetBoundaryPoints()
 {
-	FSegmentMesh& SegmentMesh = GetSegmentMesh();
+	const FSegmentMesh& SegmentMesh = const_cast<const FTriangleMesh*>(this)->GetSegmentMesh();
 	const TArray<TVec2<int32>>& Edges = SegmentMesh.GetElements();
 	const TArray<TVec2<int32>>& EdgeToFaces = GetEdgeToFaces();
 	TSet<int32> OpenBoundaryPoints;
@@ -913,7 +914,7 @@ TArray<FReal> FTriangleMesh::GetCurvatureOnEdges(const TArray<FVec3>& FaceNormal
 {
 	const int32 NumNormals = FaceNormals.Num();
 	check(NumNormals == MElements.Num());
-	const FSegmentMesh& SegmentMesh = GetSegmentMesh(); // builds MEdgeToFaces
+	const FSegmentMesh& SegmentMesh = const_cast<const FTriangleMesh*>(this)->GetSegmentMesh(); // builds MEdgeToFaces
 	TArray<FReal> EdgeAngles;
 	EdgeAngles.SetNumZeroed(MEdgeToFaces.Num());
 	for (int32 EdgeId = 0; EdgeId < MEdgeToFaces.Num(); EdgeId++)
@@ -940,7 +941,7 @@ TArray<FReal> FTriangleMesh::GetCurvatureOnEdges(const TConstArrayView<FVec3>& P
 
 TArray<FReal> FTriangleMesh::GetCurvatureOnPoints(const TArray<FReal>& EdgeCurvatures)
 {
-	const FSegmentMesh& SegmentMesh = GetSegmentMesh();
+	const FSegmentMesh& SegmentMesh = const_cast<const FTriangleMesh*>(this)->GetSegmentMesh();
 	const TArray<TVec2<int32>>& Segments = SegmentMesh.GetElements();
 	check(EdgeCurvatures.Num() == Segments.Num());
 
@@ -1468,8 +1469,8 @@ void FTriangleMesh::BuildBVH(const TConstArrayView<TVec3<T>>& Points, TBVHType<T
 	}
 	BVH.Reinitialize(BVEntries);
 }
-template CHAOS_API void FTriangleMesh::BuildBVH<FRealSingle>(const TConstArrayView<TVec3<FRealSingle>>& Points, TBVHType<FRealSingle>& BVH) const;
-template CHAOS_API void FTriangleMesh::BuildBVH<FRealDouble>(const TConstArrayView<TVec3<FRealDouble>>& Points, TBVHType<FRealDouble>& BVH) const;
+template void FTriangleMesh::BuildBVH<FRealSingle>(const TConstArrayView<TVec3<FRealSingle>>& Points, TBVHType<FRealSingle>& BVH) const;
+template void FTriangleMesh::BuildBVH<FRealDouble>(const TConstArrayView<TVec3<FRealDouble>>& Points, TBVHType<FRealDouble>& BVH) const;
 
 template<typename T>
 bool FTriangleMesh::PointProximityQuery(const TBVHType<T>& BVH, const TConstArrayView<TVec3<T>>& Points, const int32 PointIndex, const TVec3<T>& PointPosition, const T PointThickness, const T ThisThickness, 
@@ -1519,7 +1520,52 @@ bool FTriangleMesh::PointProximityQuery(const TBVHType<T>& BVH, const TConstArra
 	}
 	return Result.Num() > 0;
 }
-template CHAOS_API bool FTriangleMesh::PointProximityQuery<FRealSingle>(const TBVHType<FRealSingle>& BVH, const TConstArrayView<TVector<FRealSingle, 3>>& Points, const int32 PointIndex, const TVector<FRealSingle, 3>& PointPosition, const FRealSingle PointThickness, const FRealSingle ThisThickness, TFunctionRef<bool(const int32 PointIndex, const int32 TriangleIndex)> BroadphaseTest, TArray<TTriangleCollisionPoint<FRealSingle>>& Result) const;
-template CHAOS_API bool FTriangleMesh::PointProximityQuery<FRealDouble>(const TBVHType<FRealDouble>& BVH, const TConstArrayView<TVector<FRealDouble, 3>>& Points, const int32 PointIndex, const TVector<FRealDouble, 3>& PointPosition, const FRealDouble PointThickness, const FRealDouble ThisThickness, TFunctionRef<bool(const int32 PointIndex, const int32 TriangleIndex)> BroadphaseTest, TArray<TTriangleCollisionPoint<FRealDouble>>& Result) const;
+template bool FTriangleMesh::PointProximityQuery<FRealSingle>(const TBVHType<FRealSingle>& BVH, const TConstArrayView<TVector<FRealSingle, 3>>& Points, const int32 PointIndex, const TVector<FRealSingle, 3>& PointPosition, const FRealSingle PointThickness, const FRealSingle ThisThickness, TFunctionRef<bool(const int32 PointIndex, const int32 TriangleIndex)> BroadphaseTest, TArray<TTriangleCollisionPoint<FRealSingle>>& Result) const;
+template bool FTriangleMesh::PointProximityQuery<FRealDouble>(const TBVHType<FRealDouble>& BVH, const TConstArrayView<TVector<FRealDouble, 3>>& Points, const int32 PointIndex, const TVector<FRealDouble, 3>& PointPosition, const FRealDouble PointThickness, const FRealDouble ThisThickness, TFunctionRef<bool(const int32 PointIndex, const int32 TriangleIndex)> BroadphaseTest, TArray<TTriangleCollisionPoint<FRealDouble>>& Result) const;
 
+template<typename T>
+bool FTriangleMesh::EdgeIntersectionQuery(const TBVHType<T>& BVH, const TConstArrayView<TVec3<T>>& Points, const int32 EdgeIndex, const TVec3<T>& EdgePosition1, const TVec3<T>& EdgePosition2,
+	TFunctionRef<bool(const int32 EdgeIndex, const int32 TriangleIndex)> BroadphaseTest, TArray<TTriangleCollisionPoint<T>>& Result) const
+{
+	FAABB3 QueryBounds(EdgePosition1, EdgePosition1);
+	QueryBounds.GrowToInclude(EdgePosition2);
+
+	const TArray<int32> PotentialIntersections = BVH.FindAllIntersections(QueryBounds);
+
+	Result.Reset(PotentialIntersections.Num());
+
+	for (int32 TriIdx : PotentialIntersections)
+	{
+		if (!BroadphaseTest(EdgeIndex, TriIdx))
+		{
+			continue;
+		}
+
+		const TVec3<T>& A = Points[MElements[TriIdx][0]];
+		const TVec3<T>& B = Points[MElements[TriIdx][1]];
+		const TVec3<T>& C = Points[MElements[TriIdx][2]];
+
+		const TTriangle<T> Triangle(A, B, C);
+
+		T Time;
+		TVector<T, 2> Bary;
+		if (Triangle.LineIntersection(EdgePosition1, EdgePosition2, Bary, Time))
+		{
+			TTriangleCollisionPoint<T> CollisionPoint;
+			CollisionPoint.ContactType = TTriangleCollisionPoint<T>::EContactType::EdgeFace;
+			CollisionPoint.Indices[0] = EdgeIndex;
+			CollisionPoint.Indices[1] = TriIdx;
+			CollisionPoint.Bary = TVec4<T>(Time, (T)1. - Bary.X - Bary.Y, Bary.X, Bary.Y);
+			CollisionPoint.Location = ((T)1. - Time) * EdgePosition1 + Time * EdgePosition2;
+			CollisionPoint.Normal = Triangle.GetNormal();
+			CollisionPoint.Phi = 0;
+			Result.Add(CollisionPoint);
+		}
+	}
+	return Result.Num() > 0;
+}
+template bool FTriangleMesh::EdgeIntersectionQuery<FRealSingle>(const TBVHType<FRealSingle>& BVH, const TConstArrayView<TVec3<FRealSingle>>& Points, const int32 EdgeIndex, const TVec3<FRealSingle>& EdgePosition1, const TVec3<FRealSingle>& EdgePosition2,
+	TFunctionRef<bool(const int32 EdgeIndex, const int32 TriangleIndex)> BroadphaseTest, TArray<TTriangleCollisionPoint<FRealSingle>>& Result) const;
+template bool FTriangleMesh::EdgeIntersectionQuery<FRealDouble>(const TBVHType<FRealDouble>& BVH, const TConstArrayView<TVec3<FRealDouble>>& Points, const int32 EdgeIndex, const TVec3<FRealDouble>& EdgePosition1, const TVec3<FRealDouble>& EdgePosition2,
+	TFunctionRef<bool(const int32 EdgeIndex, const int32 TriangleIndex)> BroadphaseTest, TArray<TTriangleCollisionPoint<FRealDouble>>& Result) const;
 }  // End namespace Chaos
diff --git a/Engine/Source/Runtime/Experimental/Chaos/Public/Chaos/HeightField.h b/Engine/Source/Runtime/Experimental/Chaos/Public/Chaos/HeightField.h
--- a/Engine/Source/Runtime/Experimental/Chaos/Public/Chaos/HeightField.h
+++ b/Engine/Source/Runtime/Experimental/Chaos/Public/Chaos/HeightField.h
@@ -17,8 +17,6 @@ namespace Chaos
 {
 	class FHeightfieldRaycastVisitor;
 	class FContactPoint;
-	class FConvex;
-	class FTriangle;
 	struct FMTDInfo;
 }
 
diff --git a/Engine/Source/Runtime/Experimental/Chaos/Public/Chaos/ImplicitFwd.h b/Engine/Source/Runtime/Experimental/Chaos/Public/Chaos/ImplicitFwd.h
--- a/Engine/Source/Runtime/Experimental/Chaos/Public/Chaos/ImplicitFwd.h
+++ b/Engine/Source/Runtime/Experimental/Chaos/Public/Chaos/ImplicitFwd.h
@@ -7,7 +7,8 @@ namespace Chaos
 	class FImplicitObject;
 	class FConvex;
 	class FCapsule;
-	class FTriangle;
+	template<class T> class TTriangle;
+	using FTriangle = TTriangle<FReal>;
 
 	template<class T, int d> class TBox;
 	template<class T, int d> class TPlane;
diff --git a/Engine/Source/Runtime/Experimental/Chaos/Public/Chaos/PBDCollisionSpringConstraints.h b/Engine/Source/Runtime/Experimental/Chaos/Public/Chaos/PBDCollisionSpringConstraints.h
--- a/Engine/Source/Runtime/Experimental/Chaos/Public/Chaos/PBDCollisionSpringConstraints.h
+++ b/Engine/Source/Runtime/Experimental/Chaos/Public/Chaos/PBDCollisionSpringConstraints.h
@@ -20,8 +20,9 @@ public:
 		const FTriangleMesh& InTriangleMesh,
 		const TArray<FSolverVec3>* InRestPositions,
 		TSet<TVec2<int32>>&& InDisabledCollisionElements,
-		const FSolverReal InThickness = (FSolverReal)1.,
-		const FSolverReal InStiffness = (FSolverReal)1.)
+		const FSolverReal InThickness = Base::BackCompatThickness,
+		const FSolverReal InStiffness = Base::BackCompatStiffness,
+		const FSolverReal InFrictionCoefficient = Base::BackCompatFrictionCoefficient)
 		: Base(InOffset, InNumParticles, InTriangleMesh, InRestPositions, MoveTemp(InDisabledCollisionElements), InThickness, InStiffness)
 	{}
 
@@ -38,22 +39,21 @@ public:
 		const int32 i3 = Constraint[2];
 		const int32 i4 = Constraint[3];
 		const FSolverVec3 Delta = Base::GetDelta(Particles, i);
-		static const FSolverReal Multiplier = (FSolverReal)0.5;  // TODO(mlentine): Figure out what the best multiplier here is
 		if (Particles.InvM(i1) > 0)
 		{
-			Particles.P(i1) += Multiplier * Particles.InvM(i1) * Delta;
+			Particles.P(i1) += Particles.InvM(i1) * Delta;
 		}
 		if (Particles.InvM(i2) > (FSolverReal)0.)
 		{
-			Particles.P(i2) -= Multiplier * Particles.InvM(i2) * Barys[i][0] * Delta;
+			Particles.P(i2) -= Particles.InvM(i2) * Barys[i][0] * Delta;
 		}
 		if (Particles.InvM(i3) > (FSolverReal)0.)
 		{
-			Particles.P(i3) -= Multiplier * Particles.InvM(i3) * Barys[i][1] * Delta;
+			Particles.P(i3) -= Particles.InvM(i3) * Barys[i][1] * Delta;
 		}
 		if (Particles.InvM(i4) > (FSolverReal)0.)
 		{
-			Particles.P(i4) -= Multiplier * Particles.InvM(i4) * Barys[i][2] * Delta;
+			Particles.P(i4) -= Particles.InvM(i4) * Barys[i][2] * Delta;
 		}
 	}
 
diff --git a/Engine/Source/Runtime/Experimental/Chaos/Public/Chaos/PBDCollisionSpringConstraintsBase.h b/Engine/Source/Runtime/Experimental/Chaos/Public/Chaos/PBDCollisionSpringConstraintsBase.h
--- a/Engine/Source/Runtime/Experimental/Chaos/Public/Chaos/PBDCollisionSpringConstraintsBase.h
+++ b/Engine/Source/Runtime/Experimental/Chaos/Public/Chaos/PBDCollisionSpringConstraintsBase.h
@@ -5,6 +5,7 @@
 #include "Chaos/Core.h"
 #include "Chaos/PBDSoftsEvolutionFwd.h"
 #include "Chaos/PBDSoftsSolverParticles.h"
+#include "Chaos/PBDTriangleMeshCollisions.h"
 #include "Containers/Set.h"
 
 
@@ -20,6 +21,9 @@ namespace Chaos::Softs
 class CHAOS_API FPBDCollisionSpringConstraintsBase
 {
 public:
+	static constexpr FSolverReal BackCompatThickness = (FSolverReal)1.f;
+	static constexpr FSolverReal BackCompatStiffness = (FSolverReal)0.5f;
+	static constexpr FSolverReal BackCompatFrictionCoefficient = (FSolverReal)0.f;
 
 	FPBDCollisionSpringConstraintsBase(
 		const int32 InOffset,
@@ -27,24 +31,32 @@ public:
 		const FTriangleMesh& InTriangleMesh,
 		const TArray<FSolverVec3>* InReferencePositions,
 		TSet<TVec2<int32>>&& InDisabledCollisionElements,
-		const FSolverReal InThickness = (FSolverReal)1.,
-		const FSolverReal InStiffness = (FSolverReal)1.);
+		const FSolverReal InThickness = BackCompatThickness,
+		const FSolverReal InStiffness = BackCompatStiffness,
+		const FSolverReal InFrictionCoefficient = BackCompatFrictionCoefficient);
 
 	virtual ~FPBDCollisionSpringConstraintsBase() {}
 
+	UE_DEPRECATED(5.0, "Use Init(Particles, BVH, GIAColors) instead.")
 	void Init(const FSolverParticles& Particles);
 
+	void Init(const FSolverParticles& Particles, const FTriangleMesh::TBVHType<FSolverReal>& BVH, const TConstArrayView<FPBDTriangleMeshCollisions::FGIAColor>& VertexGIAColors, const TArray<FPBDTriangleMeshCollisions::FGIAColor>& TriangleGIAColors);
+
 	FSolverVec3 GetDelta(const FSolverParticles& InParticles, const int32 i) const;
 
 	const TArray<TVec4<int32>>& GetConstraints() const { return Constraints;  }
 	const TArray<FSolverVec3>& GetBarys() const { return Barys; }
 	FSolverReal GetThickness() const { return Thickness; }
+	bool GetGlobalIntersectionAnalysis() const { return bGlobalIntersectionAnalysis; }
+	const TArray<bool>& GetFlipNormals() const { return FlipNormal; }
 
 	void SetThickness(FSolverReal InThickness) { Thickness = FMath::Max(InThickness, (FSolverReal)0.);  }
+	void SetFrictionCoefficient(FSolverReal InFrictionCoefficient) { FrictionCoefficient = InFrictionCoefficient; }
 
 protected:
 	TArray<TVec4<int32>> Constraints;
 	TArray<FSolverVec3> Barys;
+	TArray<bool> FlipNormal;
 
 private:
 	const FTriangleMesh& TriangleMesh;
@@ -56,6 +68,8 @@ private:
 	int32 NumParticles;
 	FSolverReal Thickness;
 	FSolverReal Stiffness;
+	FSolverReal FrictionCoefficient;
+	bool bGlobalIntersectionAnalysis; // This is set based on which Init is called.
 };
 
 }  // End namespace Chaos::Softs
diff --git a/Engine/Source/Runtime/Experimental/Chaos/Public/Chaos/PBDEvolution.h b/Engine/Source/Runtime/Experimental/Chaos/Public/Chaos/PBDEvolution.h
--- a/Engine/Source/Runtime/Experimental/Chaos/Public/Chaos/PBDEvolution.h
+++ b/Engine/Source/Runtime/Experimental/Chaos/Public/Chaos/PBDEvolution.h
@@ -70,25 +70,38 @@ class CHAOS_API FPBDEvolution : public TArrayCollection
 	const TPBDActiveView<FSolverRigidParticles>& CollisionParticlesActiveView() { return MCollisionParticlesActiveView; }
 
 	// Reset all constraint init and rule functions.
-	void ResetConstraintRules() { MConstraintInits.Reset(); MConstraintRules.Reset(); MConstraintInitsActiveView.Reset(); MConstraintRulesActiveView.Reset();  };
+	void ResetConstraintRules() 
+	{ 
+		MConstraintInits.Reset(); 
+		MConstraintRules.Reset(); 
+		MPostCollisionConstraintRules.Reset();
+		MConstraintInitsActiveView.Reset(); 
+		MConstraintRulesActiveView.Reset();  
+		MPostCollisionConstraintRulesActiveView.Reset();
+	}
 
 	// Add constraints. Return the index of the first added constraint.
 	int32 AddConstraintInitRange(int32 NumConstraints, bool bActivate);
 	int32 AddConstraintRuleRange(int32 NumConstraints, bool bActivate);
+	int32 AddPostCollisionConstraintRuleRange(int32 NumConstraints, bool bActivate);
 	
 	// Return the number of particles of the block starting at Offset
 	int32 GetConstraintInitRangeSize(int32 Offset) const { return MConstraintInitsActiveView.GetRangeSize(Offset); }
 	int32 GetConstraintRuleRangeSize(int32 Offset) const { return MConstraintRulesActiveView.GetRangeSize(Offset); }
+	int32 GetPostCollisionConstraintRuleRangeSize(int32 Offset) const { return MPostCollisionConstraintRulesActiveView.GetRangeSize(Offset); }
 
 	// Set a block of constraints active or inactive, using the index of the first added particle to identify the block.
 	void ActivateConstraintInitRange(int32 Offset, bool bActivate) { MConstraintInitsActiveView.ActivateRange(Offset, bActivate); }
 	void ActivateConstraintRuleRange(int32 Offset, bool bActivate) { MConstraintRulesActiveView.ActivateRange(Offset, bActivate); }
+	void ActivatePostCollisionConstraintRuleRange(int32 Offset, bool bActivate) { MPostCollisionConstraintRulesActiveView.ActivateRange(Offset, bActivate); }
 
 	// Constraint accessors
 	const TArray<TFunction<void(FSolverParticles&, const FSolverReal)>>& ConstraintInits() const { return MConstraintInits; }
 	TArray<TFunction<void(FSolverParticles&, const FSolverReal)>>& ConstraintInits() { return MConstraintInits; }
 	const TArray<TFunction<void(FSolverParticles&, const FSolverReal)>>& ConstraintRules() const { return MConstraintRules; }
 	TArray<TFunction<void(FSolverParticles&, const FSolverReal)>>& ConstraintRules() { return MConstraintRules; }
+	const TArray<TFunction<void(FSolverParticles&, const FSolverReal)>>& PostCollisionConstraintRules() const { return MPostCollisionConstraintRules; }
+	TArray<TFunction<void(FSolverParticles&, const FSolverReal)>>& PostCollisionConstraintRules() { return MPostCollisionConstraintRules; }
 	
 	void SetKinematicUpdateFunction(TFunction<void(FSolverParticles&, const FSolverReal, const FSolverReal, const int32)> KinematicUpdate) { MKinematicUpdate = KinematicUpdate; }
 	void SetCollisionKinematicUpdateFunction(TFunction<void(FSolverRigidParticles&, const FSolverReal, const FSolverReal, const int32)> KinematicUpdate) { MCollisionKinematicUpdate = KinematicUpdate; }
@@ -105,9 +118,6 @@ class CHAOS_API FPBDEvolution : public TArrayCollection
 	int32 GetIterations() const { return MNumIterations; }
 	void SetIterations(const int32 Iterations) { MNumIterations = Iterations; }
 
-	FSolverReal GetSelfCollisionThickness(const uint32 GroupId = 0) const { check(GroupId < TArrayCollection::Size()); return MGroupSelfCollisionThicknesses[GroupId]; }
-	void SetSelfCollisionThickness(const FSolverReal SelfCollisionThickness, const uint32 GroupId = 0) { check(GroupId < TArrayCollection::Size()); MGroupSelfCollisionThicknesses[GroupId] = SelfCollisionThickness; }
-
 	FSolverReal GetCollisionThickness(const uint32 GroupId = 0) const { check(GroupId < TArrayCollection::Size()); return MGroupCollisionThicknesses[GroupId]; }
 	void SetCollisionThickness(const FSolverReal CollisionThickness, const uint32 GroupId = 0) { check(GroupId < TArrayCollection::Size()); MGroupCollisionThicknesses[GroupId] = CollisionThickness; }
 
@@ -158,7 +168,6 @@ private:
 	TArrayCollectionArray<FVelocityField> MGroupVelocityFields;
 	TArrayCollectionArray<TFunction<void(FSolverParticles&, const FSolverReal, const int32)>> MGroupForceRules;
 	TArrayCollectionArray<FSolverReal> MGroupCollisionThicknesses;
-	TArrayCollectionArray<FSolverReal> MGroupSelfCollisionThicknesses;
 	TArrayCollectionArray<FSolverReal> MGroupCoefficientOfFrictions;
 	TArrayCollectionArray<FSolverReal> MGroupDampings;
 	TArrayCollectionArray<FSolverReal> MGroupLocalDampings;
@@ -168,6 +177,8 @@ private:
 	TPBDActiveView<TArray<TFunction<void(FSolverParticles&, const FSolverReal)>>> MConstraintInitsActiveView;
 	TArray<TFunction<void(FSolverParticles&, const FSolverReal)>> MConstraintRules;
 	TPBDActiveView<TArray<TFunction<void(FSolverParticles&, const FSolverReal)>>> MConstraintRulesActiveView;
+	TArray<TFunction<void(FSolverParticles&, const FSolverReal)>> MPostCollisionConstraintRules;
+	TPBDActiveView<TArray<TFunction<void(FSolverParticles&, const FSolverReal)>>> MPostCollisionConstraintRulesActiveView;
 
 	TFunction<void(FSolverParticles&, const FSolverReal, const FSolverReal, const int32)> MKinematicUpdate;
 	TFunction<void(FSolverRigidParticles&, const FSolverReal, const FSolverReal, const int32)> MCollisionKinematicUpdate;
@@ -175,7 +186,6 @@ private:
 	int32 MNumIterations;
 	FSolverVec3 MGravity;
 	FSolverReal MCollisionThickness;
-	FSolverReal MSelfCollisionThickness;
 	FSolverReal MCoefficientOfFriction;
 	FSolverReal MDamping;
 	FSolverReal MLocalDamping;
diff --git a/Engine/Source/Runtime/Experimental/Chaos/Public/Chaos/PBDSoftsEvolutionFwd.h b/Engine/Source/Runtime/Experimental/Chaos/Public/Chaos/PBDSoftsEvolutionFwd.h
--- a/Engine/Source/Runtime/Experimental/Chaos/Public/Chaos/PBDSoftsEvolutionFwd.h
+++ b/Engine/Source/Runtime/Experimental/Chaos/Public/Chaos/PBDSoftsEvolutionFwd.h
@@ -57,6 +57,8 @@ namespace Chaos
 		class FPBDAnimDriveConstraint;
 		class FPBDShapeConstraints;
 		class FPBDCollisionSpringConstraints;
+		class FPBDTriangleMeshIntersections;
+		class FPBDTriangleMeshCollisions;
 
 		// Softs solver forces
 		class FVelocityField;
diff --git a/Engine/Source/Runtime/Experimental/Chaos/Public/Chaos/PBDTriangleMeshCollisions.h b/Engine/Source/Runtime/Experimental/Chaos/Public/Chaos/PBDTriangleMeshCollisions.h
new file mode 100644
--- /dev/null
+++ b/Engine/Source/Runtime/Experimental/Chaos/Public/Chaos/PBDTriangleMeshCollisions.h
@@ -0,0 +1,144 @@
+// Copyright Epic Games, Inc. All Rights Reserved.
+#pragma once
+
+#include "Chaos/Core.h"
+#include "Chaos/PBDSoftsEvolutionFwd.h"
+#include "Chaos/TriangleMesh.h"
+
+
+// This does initialization for PBDCollisionSpringConstraints and PBDTriangleMeshIntersections, 
+// including intersection detection and global intersection analysis
+//
+namespace Chaos::Softs
+{
+class CHAOS_API FPBDTriangleMeshCollisions
+{
+public:
+	FPBDTriangleMeshCollisions(
+		const int32 InOffset,
+		const int32 InNumParticles,
+		const FTriangleMesh& InTriangleMesh,
+		bool bInGlobalIntersectionAnalysis,
+		bool bInContourMinimization
+	)
+		:TriangleMesh(InTriangleMesh)
+		, Offset(InOffset)
+		, NumParticles(InNumParticles)
+		, bGlobalIntersectionAnalysis(bInGlobalIntersectionAnalysis)
+		, bContourMinimization(bInContourMinimization)
+	{}
+
+	~FPBDTriangleMeshCollisions() = default;
+
+	void Init(const FSolverParticles& Particles);
+
+	struct FContourMinimizationIntersection
+	{
+		TVec2<int32> EdgeVertices;
+		TVec3<int32> FaceVertices;
+		FSolverVec3 LocalGradientVector;
+		FSolverVec3 GlobalGradientVector;
+	};
+	
+	// Vertices (and Triangles) are assigned FGIAColors by flood-filling global intersection contours.
+	// ContourIndex represents which pair of global contours, up to 31 unique contour pairs--then the contour index will reused. 
+	// In practice this seems to be enough unique contours. It's OK if multiple contours have the same index as long
+	// as their intersecting regions don't interact, which is unlikely.
+	// Contours come in pairs representing two regions intersecting (except for "loop" contours where a region of cloth intersects itself--think pinch regions). 
+	// These are "Colors" represented by setting the ColorBit to 0 or 1 corresponding with the given contour index.
+	struct FGIAColor
+	{
+		int32 ContourIndexBits = 0;
+		int32 ColorBits = 0;
+
+		static constexpr int32 LoopContourIndex = 0;
+		static constexpr int32 LoopBits = 1 << LoopContourIndex;
+		static constexpr int32 NonLoopMask = ~LoopBits;
+		
+		void SetContourColor(int32 ContourIndex, bool bIsColorB)
+		{
+			check(ContourIndex < 32);
+			ContourIndexBits |= 1 << ContourIndex;
+			if (bIsColorB)
+			{
+				ColorBits |= 1 << ContourIndex;
+			}
+			else
+			{
+				ColorBits &= ~(1 << ContourIndex);
+			}
+		}
+
+		bool HasContourColorSet(int32 ContourIndex) const
+		{
+			check(ContourIndex < 32);
+			return ContourIndexBits & (1 << ContourIndex);
+		}
+
+		bool IsLoop() const
+		{
+			return (ContourIndexBits & LoopBits) && (ColorBits & LoopBits);
+		}
+
+		void SetLoop()
+		{
+			SetContourColor(LoopContourIndex, true);
+		}
+
+		// Because they are opposite colors with a shared contour index. This will cause repulsion forces to attract and fix the intersection.
+		// NOTE: We flip normals if ANY of the TriVertColors are opposite the PointColor or if the TriColor (used for thin regions is flipped). This does a better job for thin features
+		// than only flipping normals if ALL TriVertColors are opposite.
+		static bool ShouldFlipNormal(const FGIAColor& Color0, const FGIAColor& Color1)
+		{
+			const int32 SharedContourBits = Color0.ContourIndexBits & Color1.ContourIndexBits & NonLoopMask;
+			const int32 FlippedColorBits = Color0.ColorBits ^ Color1.ColorBits;
+			return FlippedColorBits & SharedContourBits;
+		}
+	};
+
+
+	// Debug display of intersection contours
+	struct FBarycentricPoint
+	{
+		FSolverVec2 Bary;
+		TVec3<int32> Vertices;
+	};
+
+	enum struct FContourType : int8
+	{
+		Open = 0,
+		Loop,
+		Contour0,
+		Contour1,
+		Count
+	};
+
+	void SetGlobalIntersectionAnalysis(bool bInGlobalIntersectionAnalysis) { bGlobalIntersectionAnalysis = bInGlobalIntersectionAnalysis; }
+	void SetContourMinimization(bool bInContourMinimization) { bContourMinimization = bInContourMinimization; }
+
+	const FTriangleMesh::TBVHType<FSolverReal>& GetBVH() const { return BVH; }
+	const TArray<FContourMinimizationIntersection>& GetContourMinimizationIntersections() const { return ContourMinimizationIntersections; }
+	const TConstArrayView<FGIAColor> GetVertexGIAColors() const { return bGlobalIntersectionAnalysis && VertexGIAColors.Num() == NumParticles ? TConstArrayView<FGIAColor>(VertexGIAColors.GetData() - Offset, NumParticles + Offset) : TConstArrayView<FGIAColor>(); }
+	const TArray<FGIAColor>& GetTriangleGIAColors() const { return TriangleGIAColors; }
+	const TArray<TArray<FBarycentricPoint>>& GetIntersectionContourPoints() const { return IntersectionContourPoints; }
+	const TArray<FContourType>& GetIntersectionContourTypes() const { return IntersectionContourTypes; }
+private:
+
+	const FTriangleMesh& TriangleMesh;
+	int32 Offset;
+	int32 NumParticles;
+	bool bGlobalIntersectionAnalysis;
+	bool bContourMinimization;
+
+	FTriangleMesh::TBVHType<FSolverReal> BVH;
+	TArray<FContourMinimizationIntersection> ContourMinimizationIntersections;
+	TArray<FGIAColor> VertexGIAColors;
+	TArray<FGIAColor> TriangleGIAColors;
+
+	// Debug display of intersection contours
+	TArray<TArray<FBarycentricPoint>> IntersectionContourPoints;
+	TArray<FContourType> IntersectionContourTypes;
+
+};
+
+}  // End namespace Chaos::Softs
diff --git a/Engine/Source/Runtime/Experimental/Chaos/Public/Chaos/PBDTriangleMeshIntersections.h b/Engine/Source/Runtime/Experimental/Chaos/Public/Chaos/PBDTriangleMeshIntersections.h
new file mode 100644
--- /dev/null
+++ b/Engine/Source/Runtime/Experimental/Chaos/Public/Chaos/PBDTriangleMeshIntersections.h
@@ -0,0 +1,34 @@
+// Copyright Epic Games, Inc. All Rights Reserved.
+#pragma once
+
+#include "Chaos/Core.h"
+#include "Chaos/PBDSoftsEvolutionFwd.h"
+#include "Chaos/PBDTriangleMeshCollisions.h"
+
+// Resolve intersections via contour minimization. Contour Minimization data is calculated in PBDTriangleMeshCollisions.
+namespace Chaos::Softs
+{
+class CHAOS_API FPBDTriangleMeshIntersections
+{
+public:
+	FPBDTriangleMeshIntersections(
+		const int32 InOffset,
+		const int32 InNumParticles,
+		const FTriangleMesh& InTriangleMesh
+	)
+		:TriangleMesh(InTriangleMesh)
+		, Offset(InOffset)
+		, NumParticles(InNumParticles)
+	{}
+
+	~FPBDTriangleMeshIntersections() = default;
+
+	void Apply(FSolverParticles& Particles, const TArray<FPBDTriangleMeshCollisions::FContourMinimizationIntersection>& Intersections, const FSolverReal Dt) const;
+
+private:
+	const FTriangleMesh& TriangleMesh;
+	int32 Offset;
+	int32 NumParticles;
+};
+
+}  // End namespace Chaos::Softs
diff --git a/Engine/Source/Runtime/Experimental/Chaos/Public/Chaos/Triangle.h b/Engine/Source/Runtime/Experimental/Chaos/Public/Chaos/Triangle.h
--- a/Engine/Source/Runtime/Experimental/Chaos/Public/Chaos/Triangle.h
+++ b/Engine/Source/Runtime/Experimental/Chaos/Public/Chaos/Triangle.h
@@ -10,54 +10,55 @@
 
 namespace Chaos
 {
-	class FTriangle
+	template<typename T>
+	class TTriangle
 	{
 	public:
-		FTriangle()
+		TTriangle()
 		{
 		}
 
-		FTriangle(const FVec3& InA, const FVec3& InB, const FVec3& InC)
+		TTriangle(const TVec3<T>& InA, const TVec3<T>& InB, const TVec3<T>& InC)
 			: ABC{ InA, InB, InC }
 		{
 		}
 
-		FORCEINLINE FVec3& operator[](uint32 InIndex)
+		FORCEINLINE TVec3<T>& operator[](uint32 InIndex)
 		{
 			checkSlow(InIndex < 3);
 			return ABC[InIndex];
 		}
 
-		FORCEINLINE const FVec3& operator[](uint32 InIndex) const
+		FORCEINLINE const TVec3<T>& operator[](uint32 InIndex) const
 		{
 			checkSlow(InIndex < 3);
 			return ABC[InIndex];
 		}
 
-		FORCEINLINE const FVec3& GetVertex(const int32 InIndex) const
+		FORCEINLINE const TVec3<T>& GetVertex(const int32 InIndex) const
 		{
 			checkSlow(InIndex < 3);
 			return ABC[InIndex];
 		}
 
-		FORCEINLINE FVec3 GetNormal() const
+		FORCEINLINE TVec3<T> GetNormal() const
 		{
-			return FVec3::CrossProduct(ABC[1] - ABC[0], ABC[2] - ABC[0]).GetSafeNormal();
+			return TVec3<T>::CrossProduct(ABC[1] - ABC[0], ABC[2] - ABC[0]).GetSafeNormal();
 		}
 
-		FORCEINLINE TPlane<FReal, 3> GetPlane() const
+		FORCEINLINE TPlane<T, 3> GetPlane() const
 		{
-			return TPlane<FReal, 3>(ABC[0], GetNormal());
+			return TPlane<T, 3>(ABC[0], GetNormal());
 		}
 
 		// Face index is ignored since we only have one face
 		// Used for manifold generation
-		FORCEINLINE TPlaneConcrete<FReal, 3> GetPlane(int32 FaceIndex) const
+		FORCEINLINE TPlaneConcrete<T, 3> GetPlane(int32 FaceIndex) const
 		{
-			return TPlaneConcrete < FReal, 3> (ABC[0], GetNormal());
+			return TPlaneConcrete <T, 3> (ABC[0], GetNormal());
 		}
 
-		FORCEINLINE void GetPlaneNX(const int32 FaceIndex, FVec3& OutN, FVec3& OutX) const
+		FORCEINLINE void GetPlaneNX(const int32 FaceIndex, TVec3<T>& OutN, TVec3<T>& OutX) const
 		{
 			OutN = GetNormal();
 			OutX = ABC[0];
@@ -65,20 +66,20 @@ namespace Chaos
 
 		// Get the nearest point on an edge and the edge vertices
 		// Used for manifold generation
-		FVec3 GetClosestEdge(int32 PlaneIndexHint, const FVec3& Position, FVec3& OutEdgePos0, FVec3& OutEdgePos1) const
+		TVec3<T> GetClosestEdge(int32 PlaneIndexHint, const TVec3<T>& Position, TVec3<T>& OutEdgePos0, TVec3<T>& OutEdgePos1) const
 		{
-			FVec3 ClosestEdgePosition = FVec3(0);
-			FReal ClosestDistanceSq = TNumericLimits<FReal>::Max();
+			TVec3<T> ClosestEdgePosition = TVec3<T>(0);
+			T ClosestDistanceSq = TNumericLimits<T>::Max();
 
 			int32 PlaneVerticesNum = 3;
 			
-			FVec3 P0 = ABC[2];
+			TVec3<T> P0 = ABC[2];
 			for (int32 PlaneVertexIndex = 0; PlaneVertexIndex < PlaneVerticesNum; ++PlaneVertexIndex)
 			{
-				const TVector<FReal, 3>& P1 = GetVertex(PlaneVertexIndex);
+				const TVector<T, 3>& P1 = GetVertex(PlaneVertexIndex);
 				
-				const FVec3 EdgePosition = FMath::ClosestPointOnLine(P0, P1, Position);
-				const FReal EdgeDistanceSq = (EdgePosition - Position).SizeSquared();
+				const TVec3<T> EdgePosition = FMath::ClosestPointOnLine(P0, P1, Position);
+				const T EdgeDistanceSq = (EdgePosition - Position).SizeSquared();
 
 				if (EdgeDistanceSq < ClosestDistanceSq)
 				{
@@ -96,9 +97,9 @@ namespace Chaos
 
 		// Get the nearest point on an edge
 		// Used for manifold generation
-		FVec3 GetClosestEdgePosition(int32 PlaneIndexHint, const FVec3& Position) const
+		TVec3<T> GetClosestEdgePosition(int32 PlaneIndexHint, const TVec3<T>& Position) const
 		{
-			FVec3 Unused0, Unused1;
+			TVec3<T> Unused0, Unused1;
 			return GetClosestEdge(PlaneIndexHint, Position, Unused0, Unused1);
 		}
 
@@ -112,7 +113,7 @@ namespace Chaos
 
 		// Returns a winding order multiplier used in the manifold clipping and required when we have negative scales (See ImplicitObjectScaled)
 		// Used for manifold generation
-		FORCEINLINE FReal GetWindingOrder() const
+		FORCEINLINE T GetWindingOrder() const
 		{
 			return 1.0f;
 		}
@@ -137,7 +138,7 @@ namespace Chaos
 		}
 		
 		// Get the index of the plane that most opposes the normal
-		int32 GetMostOpposingPlane(const FVec3& Normal) const
+		int32 GetMostOpposingPlane(const TVec3<T>& Normal) const
 		{
 			return 0; // Only have one plane
 		}
@@ -153,18 +154,18 @@ namespace Chaos
 		// Used for manifold generation
 		int32 NumPlanes() const { return 1; }
 
-		FORCEINLINE FReal PhiWithNormal(const FVec3& InSamplePoint, FVec3& OutNormal) const
+		FORCEINLINE T PhiWithNormal(const TVec3<T>& InSamplePoint, TVec3<T>& OutNormal) const
 		{
 			OutNormal = GetNormal();
-			FVec3 ClosestPoint = FindClosestPointOnTriangle(GetPlane(), ABC[0], ABC[1], ABC[2], InSamplePoint);
-			return FVec3::DotProduct((InSamplePoint - ClosestPoint), OutNormal);
+			TVec3<T> ClosestPoint = FindClosestPointOnTriangle(GetPlane(), ABC[0], ABC[1], ABC[2], InSamplePoint);
+			return TVec3<T>::DotProduct((InSamplePoint - ClosestPoint), OutNormal);
 		}
 
-		FORCEINLINE FVec3 Support(const FVec3& Direction, const FReal Thickness, int32& VertexIndex) const
+		FORCEINLINE TVec3<T> Support(const TVec3<T>& Direction, const T Thickness, int32& VertexIndex) const
 		{
-			const FReal DotA = FVec3::DotProduct(ABC[0], Direction);
-			const FReal DotB = FVec3::DotProduct(ABC[1], Direction);
-			const FReal DotC = FVec3::DotProduct(ABC[2], Direction);
+			const T DotA = TVec3<T>::DotProduct(ABC[0], Direction);
+			const T DotB = TVec3<T>::DotProduct(ABC[1], Direction);
+			const T DotC = TVec3<T>::DotProduct(ABC[2], Direction);
 
 			if(DotA >= DotB && DotA >= DotC)
 			{
@@ -192,12 +193,12 @@ namespace Chaos
 			return ABC[2];
 		}
 
-		FORCEINLINE_DEBUGGABLE FVec3 SupportCore(const FVec3& Direction, const FReal InMargin, FReal* OutSupportDelta,int32& VertexIndex) const
+		FORCEINLINE_DEBUGGABLE TVec3<T> SupportCore(const TVec3<T>& Direction, const T InMargin, T* OutSupportDelta,int32& VertexIndex) const
 		{
 			// Note: assumes margin == 0
-			const FReal DotA = FVec3::DotProduct(ABC[0], Direction);
-			const FReal DotB = FVec3::DotProduct(ABC[1], Direction);
-			const FReal DotC = FVec3::DotProduct(ABC[2], Direction);
+			const T DotA = TVec3<T>::DotProduct(ABC[0], Direction);
+			const T DotB = TVec3<T>::DotProduct(ABC[1], Direction);
+			const T DotC = TVec3<T>::DotProduct(ABC[2], Direction);
 
 			if (DotA >= DotB && DotA >= DotC)
 			{
@@ -213,30 +214,30 @@ namespace Chaos
 			return ABC[2];
 		}
 
-		FORCEINLINE FVec3 SupportCoreScaled(const FVec3& Direction, FReal InMargin, const FVec3& Scale, FReal* OutSupportDelta, int32& VertexIndex) const
+		FORCEINLINE TVec3<T> SupportCoreScaled(const TVec3<T>& Direction, T InMargin, const TVec3<T>& Scale, T* OutSupportDelta, int32& VertexIndex) const
 		{
 			// Note: ignores InMargin, assumed 0 (triangles cannot have a margin as they are zero thickness)
 			return SupportCore(Direction * Scale, 0.0f, OutSupportDelta, VertexIndex) * Scale;
 		}
 
-		FORCEINLINE FReal GetMargin() const { return 0; }
-		FORCEINLINE FReal GetRadius() const { return 0; }
+		FORCEINLINE T GetMargin() const { return 0; }
+		FORCEINLINE T GetRadius() const { return 0; }
 
-		FORCEINLINE bool Raycast(const FVec3& StartPoint, const FVec3& Dir, const FReal Length, const FReal Thickness, FReal& OutTime, FVec3& OutPosition, FVec3& OutNormal, int32& OutFaceIndex) const
+		FORCEINLINE bool Raycast(const TVec3<T>& StartPoint, const TVec3<T>& Dir, const T Length, const T Thickness, T& OutTime, TVec3<T>& OutPosition, TVec3<T>& OutNormal, int32& OutFaceIndex) const
 		{
 			// No face as this is only one triangle
 			OutFaceIndex = INDEX_NONE;
 
 			// Pass through GJK #BGTODO Maybe specialise if it's possible to be faster
 			const FRigidTransform3 StartTM(StartPoint, FRotation3::FromIdentity());
-			const TSphere<FReal, 3> Sphere(FVec3(0), Thickness);
+			const TSphere<T, 3> Sphere(TVec3<T>(0), Thickness);
 			return GJKRaycast(*this, Sphere, StartTM, Dir, Length, OutTime, OutPosition, OutNormal);
 		}
 
-		FORCEINLINE bool Overlap(const FVec3& Point, const FReal Thickness) const
+		FORCEINLINE bool Overlap(const TVec3<T>& Point, const T Thickness) const
 		{
-			const FVec3 ClosestPoint = FindClosestPointOnTriangle(GetPlane(), ABC[0], ABC[1], ABC[2], Point);
-			const FReal AdjustedThickness = FMath::Max(Thickness, KINDA_SMALL_NUMBER);
+			const TVec3<T> ClosestPoint = FindClosestPointOnTriangle(GetPlane(), ABC[0], ABC[1], ABC[2], Point);
+			const T AdjustedThickness = FMath::Max(Thickness, KINDA_SMALL_NUMBER);
 			return (Point - ClosestPoint).SizeSquared() <= (AdjustedThickness * AdjustedThickness);
 		}
 
@@ -245,74 +246,123 @@ namespace Chaos
 			return true;
 		}
 
+
+		struct FLineIntersectionToleranceProvider
+		{
+			static constexpr T ParallelTolerance = (T)SMALL_NUMBER;
+			static constexpr T BaryTolerance = (T)0;
+		};
+		FORCEINLINE bool LineIntersection(const TVec3<T>& StartPoint, const TVec3<T>& EndPoint, TVector<T, 2>& OutBary, T& OutTime) const;
+
 	private:
 
-		friend FChaosArchive& operator<<(FChaosArchive& Ar, FTriangle& Value);
+		friend FChaosArchive& operator<<(FChaosArchive& Ar, TTriangle& Value);
 
-		FVec3 ABC[3];
+		TVec3<T> ABC[3];
 	};
 
-	inline FChaosArchive& operator<<(FChaosArchive& Ar, FTriangle& Value)
+	using FTriangle = TTriangle<FReal>;
+
+	template<typename T>
+	inline FChaosArchive& operator<<(FChaosArchive& Ar, TTriangle<T>& Value)
 	{
 		Ar << Value.ABC[0] << Value.ABC[1] << Value.ABC[2];
 		return Ar;
 	}
 
-	/**
-	* Ray / triangle  intersection
-	* this provides a double sided test
-	* note : this method assumes that the triangle formed by A,B and C is well formed
-	*/
-	FORCEINLINE_DEBUGGABLE bool RayTriangleIntersection(
-		const FVec3& RayStart, const FVec3& RayDir, FReal RayLength,
-		const FVec3& A, const FVec3& B, const FVec3& C,
-		FReal& OutT, FVec3& OutN
-	)
+	template<typename T>
+	struct TRayTriangleIntersectionDefaultToleranceProvider
 	{
-		const FVec3 AB = B - A; // edge 1
-		const FVec3 AC = C - A; // edge 2
-		const FVec3 Normal = FVec3::CrossProduct(AB, AC);
-		const FVec3 NegRayDir = -RayDir;
+		static constexpr T ParallelTolerance = (T)SMALL_NUMBER;
+		static constexpr T BaryTolerance = (T)SMALL_NUMBER;
+	};
 
-		FReal Den = FVec3::DotProduct(NegRayDir, Normal);
-		if (FMath::Abs(Den) < SMALL_NUMBER)
+	template<typename T, typename ToleranceProvider = TRayTriangleIntersectionDefaultToleranceProvider<T> >
+	FORCEINLINE_DEBUGGABLE bool RayTriangleIntersectionAndBary(const TVec3<T>& RayStart, const TVec3<T>& RayDir, T RayLength,
+		const TVec3<T>& A, const TVec3<T>& B, const TVec3<T>& C, T& OutT, TVec2<T>& OutBary, TVec3<T>& OutN)
+	{
+		const TVec3<T> AB = B - A; // edge 1
+		const TVec3<T> AC = C - A; // edge 2
+		const TVec3<T> Normal = TVec3<T>::CrossProduct(AB, AC);
+		const TVec3<T> NegRayDir = -RayDir;
+
+		const T Den = TVec3<T>::DotProduct(NegRayDir, Normal);
+		if (FMath::Abs(Den) < ToleranceProvider::ParallelTolerance)
 		{
 			// ray is parallel or away to the triangle plane it is a miss
 			return false;
 		}
 
-		const FReal InvDen = (FReal)1 / Den;
+		const T InvDen = (T)1 / Den;
 
 		// let's compute the time to intersection
-		const FVec3 RayToA = RayStart - A;
-		const FReal Time = FVec3::DotProduct(RayToA, Normal) * InvDen;
-		if (Time < (FReal)0 || Time > RayLength)
+		const TVec3<T> RayToA = RayStart - A;
+		const T Time = TVec3<T>::DotProduct(RayToA, Normal) * InvDen;
+		if (Time < (T)0 || Time > RayLength)
 		{
 			return false;
 		}
 
-		// now compute baricentric coordinates
-		const FVec3 RayToACrossNegDir = FVec3::CrossProduct(NegRayDir, RayToA);
-		constexpr FReal Epsilon = SMALL_NUMBER;
-		const FReal UU = FVec3::DotProduct(AC, RayToACrossNegDir) * InvDen;
-		if (UU < -Epsilon || UU >(1 + Epsilon))
+		// now compute barycentric coordinates
+		const TVec3<T> RayToACrossNegDir = FVec3::CrossProduct(NegRayDir, RayToA);
+		const T UU = TVec3<T>::DotProduct(AC, RayToACrossNegDir) * InvDen;
+		if (UU < -ToleranceProvider::BaryTolerance || UU >(1 + ToleranceProvider::BaryTolerance))
 		{
 			return false; // outside of the triangle
 		}
-		const FReal VV = -FVec3::DotProduct(AB, RayToACrossNegDir) * InvDen;
-		if (VV < -Epsilon || (VV + UU) >(1 + Epsilon))
+		const T VV = -TVec3<T>::DotProduct(AB, RayToACrossNegDir) * InvDen;
+		if (VV < -ToleranceProvider::BaryTolerance || (VV + UU) >(1 + ToleranceProvider::BaryTolerance))
 		{
 			return false; // outside of the triangle
 		}
 
 		// point is within the triangle, let's compute 
 		OutT = Time;
+		OutBary = { UU, VV };
 		OutN = Normal.GetSafeNormal();
 		OutN *= FMath::Sign(Den);
 		return true;
 	}
 
-	template<typename T> using TTriangle = FTriangle;
+	/**
+	* Ray / triangle  intersection
+	* this provides a double sided test
+	* note : this method assumes that the triangle formed by A,B and C is well formed
+	*/
+	template<typename T>
+	FORCEINLINE bool RayTriangleIntersection(
+		const TVec3<T>& RayStart, const TVec3<T>& RayDir, T RayLength,
+		const TVec3<T>& A, const TVec3<T>& B, const TVec3<T>& C,
+		T& OutT, TVec3<T>& OutN
+	)
+	{
+		TVec2<T> BaryUnused;
+		return RayTriangleIntersectionAndBary(RayStart, RayDir, RayLength, A, B, C, OutT, BaryUnused, OutN);
+	}
+
+	template<typename T>
+	FORCEINLINE bool TTriangle<T>::LineIntersection(const TVec3<T>& StartPoint, const TVec3<T>& EndPoint, TVector<T, 2>& OutBary, T& OutTime) const
+	{
+		const TVec3<T> StartToEnd = EndPoint - StartPoint;
+		const T SegmentLenSq = StartToEnd.SizeSquared();
+		if (SegmentLenSq < SMALL_NUMBER)
+		{
+			return false;
+		}
+		const T SegmentLen = FMath::Sqrt(SegmentLenSq);
+		const T OneOverSegmentLen = (T)1 / SegmentLen;
+		const TVec3<T> Ray = StartToEnd * OneOverSegmentLen;
+
+		TVec3<T> NormalUnused;
+		if (RayTriangleIntersectionAndBary<T, FLineIntersectionToleranceProvider>(StartPoint, Ray, SegmentLen, ABC[0], ABC[1], ABC[2], OutTime, OutBary, NormalUnused))
+		{
+			// OutTime is between 0 and SegmentLen. Convert to 0 to 1
+			OutTime *= OneOverSegmentLen;
+			return true;
+		}
+
+		return false;
+	}
 	
 	template<typename T>
 	class UE_DEPRECATED(4.27, "Deprecated. this class is to be deleted, use other triangle based ImplicitObjects") TImplicitTriangle final : public FImplicitObject
diff --git a/Engine/Source/Runtime/Experimental/Chaos/Public/Chaos/TriangleCollisionPoint.h b/Engine/Source/Runtime/Experimental/Chaos/Public/Chaos/TriangleCollisionPoint.h
--- a/Engine/Source/Runtime/Experimental/Chaos/Public/Chaos/TriangleCollisionPoint.h
+++ b/Engine/Source/Runtime/Experimental/Chaos/Public/Chaos/TriangleCollisionPoint.h
@@ -20,7 +20,8 @@ namespace Chaos
 		{
 			Invalid = 0,
 			PointFace,
-			EdgeEdge
+			EdgeEdge,
+			EdgeFace
 		} ContactType;
 
 		// Index into Point and Face arrays or Edge arrays
@@ -29,6 +30,7 @@ namespace Chaos
 		// Barycentric coordinates for contacts
 		// Point-Face: Bary[0] = 1.0 (Point), Bary[1:3] correspond with coordinates on Face
 		// Edge-Edge: Bary[0:1] = coordinates on first edge, [2:3] coordinates on second edge
+		// Edge-Face: Bary[0] = second coordinate on edge (first can be calculated as 1-Bary[0], Bary[1:3] correspond with coordinates on Face
 		TVec4<T> Bary;
 
 		// World space location (which point depends on query type)
diff --git a/Engine/Source/Runtime/Experimental/Chaos/Public/Chaos/TriangleMesh.h b/Engine/Source/Runtime/Experimental/Chaos/Public/Chaos/TriangleMesh.h
--- a/Engine/Source/Runtime/Experimental/Chaos/Public/Chaos/TriangleMesh.h
+++ b/Engine/Source/Runtime/Experimental/Chaos/Public/Chaos/TriangleMesh.h
@@ -167,13 +167,23 @@ namespace Chaos
 		{ return GetConvexHullFromParticles(InParticles.X()); }
 
 		/**
+		 * @brief Note that the SegmentMesh is lazily calculated (this method is not threadsafe unless it is known that the SegmentMesh is already up to date)
 		 * @ret The connectivity of this mesh represented as a collection of unique segments.
 		 */
-		CHAOS_API FSegmentMesh& GetSegmentMesh();
-		/** @ret A map from all face indices, to the indices of their associated edges. */
-		CHAOS_API const TArray<TVec3<int32>>& GetFaceToEdges();
-		/** @ret A map from all edge indices, to the indices of their containing faces. */
-		CHAOS_API const TArray<TVec2<int32>>& GetEdgeToFaces();
+		CHAOS_API const FSegmentMesh& GetSegmentMesh() const;
+		/**
+		 * @brief Note that this data is lazily calculated with the SegmentMesh (this method is not threadsafe unless it is known that the SegmentMesh is already up to date)
+		 * @ret A map from all face indices, to the indices of their associated edges.
+		 */
+		CHAOS_API const TArray<TVec3<int32>>& GetFaceToEdges() const;
+		/**
+		 * @brief Note that this data is lazily calculated with the SegmentMesh (this method is not threadsafe unless it is known that the SegmentMesh is already up to date)
+		 * @ret A map from all edge indices, to the indices of their containing faces. 
+		 */
+		CHAOS_API const TArray<TVec2<int32>>& GetEdgeToFaces() const;
+
+		UE_DEPRECATED(5.1, "Non-const access to GetSegmentMesh will be removed. Use const version instead.")
+		CHAOS_API FSegmentMesh& GetSegmentMesh() { return const_cast<FSegmentMesh&>(const_cast<const FTriangleMesh*>(this)->GetSegmentMesh()); }
 
 		/**
 		 * @ret Curvature between adjacent faces, specified on edges in radians.
@@ -271,12 +281,16 @@ namespace Chaos
 		using TBVHType = TAABBTree<int32, TAABBTreeLeafArray<int32, /*bComputeBounds=*/false, T>, /*bMutable=*/true, T>;
 
 		template<typename T>
-		CHAOS_API void BuildBVH(const TConstArrayView<TVec3<T>>& Points, TBVHType<T>& BVH) const;
+		void BuildBVH(const TConstArrayView<TVec3<T>>& Points, TBVHType<T>& BVH) const;
 
 		// NOTE: This method assumes the BVH has already been built/fitted to Points.
 		template<typename T>
-		CHAOS_API bool PointProximityQuery(const TBVHType<T>& BVH, const TConstArrayView<TVec3<T>>& Points, const int32 PointIndex, const TVec3<T>& PointPosition, const T PointThickness, const T ThisThickness, 
+		bool PointProximityQuery(const TBVHType<T>& BVH, const TConstArrayView<TVec3<T>>& Points, const int32 PointIndex, const TVec3<T>& PointPosition, const T PointThickness, const T ThisThickness, 
 			TFunctionRef<bool (const int32 PointIndex, const int32 TriangleIndex)> BroadphaseTest, TArray<TTriangleCollisionPoint<T>>& Result) const;
+
+		template<typename T>
+		bool EdgeIntersectionQuery(const TBVHType<T>& BVH, const TConstArrayView<TVec3<T>>& Points, const int32 EdgeIndex, const TVec3<T>& EdgePosition1, const TVec3<T>& EdgePosition2,
+			TFunctionRef<bool(const int32 EdgeIndex, const int32 TriangleIndex)> BroadphaseTest, TArray<TTriangleCollisionPoint<T>>& Result) const;
 		
 	private:
 		CHAOS_API void InitHelper(const int32 StartIdx, const int32 EndIdx, const bool CullDegenerateElements=true);
@@ -300,9 +314,9 @@ namespace Chaos
 		mutable TArray<TArray<int32>> MPointToTriangleMap;  // !! Unlike the TArrayView returned by GetPointToTriangleMap, this array starts at 0 for the point of index MStartIdx. Use GlobalToLocal to access with a global index. Note that this array's content is always indexed in global index.
 		mutable TMap<int32, TSet<int32>> MPointToNeighborsMap;
 
-		FSegmentMesh MSegmentMesh;
-		TArray<TVec3<int32>> MFaceToEdges;
-		TArray<TVec2<int32>> MEdgeToFaces;
+		mutable FSegmentMesh MSegmentMesh;
+		mutable TArray<TVec3<int32>> MFaceToEdges;
+		mutable TArray<TVec2<int32>> MEdgeToFaces;
 
 		int32 MStartIdx;
 		int32 MNumIndices;
diff --git a/Engine/Source/Runtime/Experimental/Chaos/Public/Chaos/TriangleMeshImplicitObject.h b/Engine/Source/Runtime/Experimental/Chaos/Public/Chaos/TriangleMeshImplicitObject.h
--- a/Engine/Source/Runtime/Experimental/Chaos/Public/Chaos/TriangleMeshImplicitObject.h
+++ b/Engine/Source/Runtime/Experimental/Chaos/Public/Chaos/TriangleMeshImplicitObject.h
@@ -3,6 +3,7 @@
 
 #include "Chaos/Array.h"
 #include "Chaos/GeometryParticles.h"
+#include "Chaos/ImplicitFwd.h"
 #include "Chaos/ImplicitObjectScaled.h"
 #include "Chaos/SegmentMesh.h"
 #include "Chaos/Triangle.h"
@@ -22,10 +23,6 @@ namespace Chaos
 {
 	extern CHAOS_API bool TriMeshPerPolySupport;
 
-	class FCapsule;
-	class FTriangle;
-	class FConvex;
-	class FTriangle;
 	struct FMTDInfo;
 	class FContactPoint;
 
