BVH vectorization optimization for Triangle Mesh

diff --git a/Engine/Source/Runtime/Experimental/Chaos/Private/Chaos/TriangleMeshImplicitObject.cpp b/Engine/Source/Runtime/Experimental/Chaos/Private/Chaos/TriangleMeshImplicitObject.cpp
--- a/Engine/Source/Runtime/Experimental/Chaos/Private/Chaos/TriangleMeshImplicitObject.cpp
+++ b/Engine/Source/Runtime/Experimental/Chaos/Private/Chaos/TriangleMeshImplicitObject.cpp
@@ -66,8 +66,8 @@ TImplicitObjectScaled<QueryGeomType> ScaleGeomIntoWorldHelper(const TImplicitObj
 	return MakeScaledHelper(QueryGeom, TriMeshScale);
 }
 
-void TransformSweepOutputsHelper(FVec3 TriMeshScale, const FVec3& HitNormal, const FVec3& HitPosition, const FReal LengthScale,
-	const FReal Time, FVec3& OutNormal, FVec3& OutPosition, FReal& OutTime)
+void TransformSweepOutputsHelper(FVec3 TriMeshScale, const FVec3& HitNormal, const FVec3& HitPosition, const FRealSingle LengthScale,
+	const FRealSingle Time, FVec3& OutNormal, FVec3& OutPosition, FRealSingle& OutTime)
 {
 	if (ensure(TriMeshScale != FVec3(0.0f)))
 	{
@@ -138,7 +138,7 @@ struct FTriangleMeshRaycastVisitor
 	 * find the intersection of a ray and the triangle at index TriIdx
 	 * @return true if the search should continue otherwise false to stop the search through the mesh 
 	 */
-	bool VisitRaycast(TSpatialVisitorData<int32> TriIdx, FQueryFastData& CurData)
+	bool VisitRaycast(TSpatialVisitorData<int32> TriIdx, FRealSingle& CurDataLength)
 	{
 		constexpr FReal Epsilon = 1e-4f;
 
@@ -165,19 +165,19 @@ struct FTriangleMeshRaycastVisitor
 
 		FVec3 HitNormal;
 		FReal HitTime;
-		if (RayTriangleIntersection(StartPoint, Dir, CurData.CurrentLength, A, B, C, HitTime, HitNormal))
+		if (RayTriangleIntersection(StartPoint, Dir, static_cast<FReal>(CurDataLength), A, B, C, HitTime, HitNormal))
 		{
 			OutPosition = StartPoint + (Dir * HitTime);
 			OutNormal = HitNormal;
 			OutTime = HitTime;
 			OutFaceIndex = FaceIndex;
-			CurData.SetLength(HitTime);	//prevent future rays from going any farther
+			CurDataLength = static_cast<FRealSingle>(HitTime); //prevent future rays from going any farther
 		}
 		// continue the visit
 		return true;
 	}
 	
-	bool VisitSweep(TSpatialVisitorData<int32> TriIdx, FQueryFastData& CurData)
+	bool VisitSweep(TSpatialVisitorData<int32> TriIdx, FRealSingle& CurDataLength)
 	{
 		constexpr FReal Epsilon = 1e-4f;
 		constexpr FReal Epsilon2 = Epsilon * Epsilon;
@@ -217,7 +217,7 @@ struct FTriangleMeshRaycastVisitor
 
 		//Check if we even intersect with triangle plane
 		int32 DummyFaceIndex;
-		if (TriPlane.Raycast(StartPoint, Dir, CurData.CurrentLength, Thickness, Time, RaycastPosition, RaycastNormal, DummyFaceIndex))
+		if (TriPlane.Raycast(StartPoint, Dir, CurDataLength, Thickness, Time, RaycastPosition, RaycastNormal, DummyFaceIndex))
 		{
 			FVec3 IntersectionPosition = RaycastPosition;
 			FVec3 IntersectionNormal = RaycastNormal;
@@ -252,19 +252,19 @@ struct FTriangleMeshRaycastVisitor
 				{
 					FVec3 ABCapsuleAxis = B - A;
 					FReal ABHeight = ABCapsuleAxis.SafeNormalize();
-					bBorderIntersections[0] = FCapsule::RaycastFast(Thickness, ABHeight, ABCapsuleAxis, A, B, StartPoint, Dir, CurData.CurrentLength, 0, BorderTimes[0], BorderPositions[0], BorderNormals[0], DummyFaceIndex);
+					bBorderIntersections[0] = FCapsule::RaycastFast(Thickness, ABHeight, ABCapsuleAxis, A, B, StartPoint, Dir, CurDataLength, 0, BorderTimes[0], BorderPositions[0], BorderNormals[0], DummyFaceIndex);
 				}
 				
 				{
 					FVec3 BCCapsuleAxis = C - B;
 					FReal BCHeight = BCCapsuleAxis.SafeNormalize();
-					bBorderIntersections[1] = FCapsule::RaycastFast(Thickness, BCHeight, BCCapsuleAxis, B, C, StartPoint, Dir, CurData.CurrentLength, 0, BorderTimes[1], BorderPositions[1], BorderNormals[1], DummyFaceIndex);
+					bBorderIntersections[1] = FCapsule::RaycastFast(Thickness, BCHeight, BCCapsuleAxis, B, C, StartPoint, Dir, CurDataLength, 0, BorderTimes[1], BorderPositions[1], BorderNormals[1], DummyFaceIndex);
 				}
 				
 				{
 					FVec3 ACCapsuleAxis = C - A;
 					FReal ACHeight = ACCapsuleAxis.SafeNormalize();
-					bBorderIntersections[2] = FCapsule::RaycastFast(Thickness, ACHeight, ACCapsuleAxis, A, C, StartPoint, Dir, CurData.CurrentLength, 0, BorderTimes[2], BorderPositions[2], BorderNormals[2], DummyFaceIndex);
+					bBorderIntersections[2] = FCapsule::RaycastFast(Thickness, ACHeight, ACCapsuleAxis, A, C, StartPoint, Dir, CurDataLength, 0, BorderTimes[2], BorderPositions[2], BorderNormals[2], DummyFaceIndex);
 				}
 
 				int32 MinBorderIdx = INDEX_NONE;
@@ -307,7 +307,7 @@ struct FTriangleMeshRaycastVisitor
 					OutPosition = IntersectionPosition;
 					OutNormal = RaycastNormal;	//We use the plane normal even when hitting triangle edges. This is to deal with triangles that approximate a single flat surface.
 					OutTime = Time;
-					CurData.SetLength(Time);	//prevent future rays from going any farther
+					CurDataLength = static_cast<FRealSingle>(Time);	//prevent future rays from going any farther
 					OutFaceIndex = FaceIndex;
 				}
 			}
@@ -343,12 +343,12 @@ struct FTriangleMeshOverlapVisitor
 		CollectedResults.Add(Instance);
 		return true;
 	}
-	bool VisitSweep(int32 Instance, FQueryFastData& CurData)
+	bool VisitSweep(int32 Instance, FRealSingle& CurDataLength)
 	{
 		check(false);
 		return true;
 	}
-	bool VisitRaycast(int32 Instance, FQueryFastData& CurData)
+	bool VisitRaycast(int32 Instance, FRealSingle& CurDataLength)
 	{
 		check(false);
 		return true;
@@ -361,7 +361,7 @@ TArray<int32> FTrimeshBVH::FindAllIntersections (const FAABB3& Intersection) con
 {
 	TArray<int32> Results;
 	FTriangleMeshOverlapVisitor Visitor(Results);
-	Overlap(Intersection, Visitor);
+	Overlap(FAABBVectorized(Intersection), Visitor);
 
 	return Results;
 }
@@ -948,25 +948,31 @@ bool FTriangleMeshImplicitObject::OverlapGeom(const TImplicitObjectScaled<FConve
 template <typename QueryGeomType, typename IdxType>
 struct FTriangleMeshSweepVisitor
 {
-	FTriangleMeshSweepVisitor(const FTriangleMeshImplicitObject& InTriMesh, const TArray<TVec3<IdxType>>& InElements, const QueryGeomType& InQueryGeom, const TRigidTransform<FReal,3>& InStartTM, const FVec3& InDir,
-		const FVec3& InScaledDirNormalized, const FReal InLengthScale, const FRigidTransform3& InScaledStartTM, const FReal InThickness, const bool InComputeMTD, FVec3 InTriMeshScale, FReal InCullsBackFaceSweepsCode)
+	FTriangleMeshSweepVisitor(const FTriangleMeshImplicitObject& InTriMesh, const TArray<TVec3<IdxType>>& InElements, const QueryGeomType& InQueryGeom,
+		const FVec3& InScaledDirNormalized, const FReal InLengthScale, const FRigidTransform3& InScaledStartTM, const bool InComputeMTD, FVec3 InTriMeshScale, FReal InCullsBackFaceSweepsCode)
 	: TriMesh(InTriMesh)
 	, Elements(InElements)
-	, StartTM(InStartTM)
 	, QueryGeom(InQueryGeom)
-	, Dir(InDir)
-	, Thickness(InThickness)
 	, bComputeMTD(InComputeMTD)
 	, CullsBackFaceSweepsCode(InCullsBackFaceSweepsCode)
 	, ScaledDirNormalized(InScaledDirNormalized)
-	, LengthScale(InLengthScale)
-	, ScaledStartTM(InScaledStartTM)
-	, OutTime(TNumericLimits<FReal>::Max())
+	, LengthScale(static_cast<FRealSingle>(InLengthScale))
+	, OutTime(TNumericLimits<FRealSingle>::Max())
 	, OutFaceIndex(INDEX_NONE)
 	, TriMeshScale(InTriMeshScale)
 	{
 		VectorScaledDirNormalized = MakeVectorRegisterFloatFromDouble(MakeVectorRegister(InScaledDirNormalized.X, InScaledDirNormalized.Y, InScaledDirNormalized.Z, 0.0));
 		VectorCullsBackFaceSweepsCode = MakeVectorRegisterFloatFromDouble(VectorLoadFloat1(&InCullsBackFaceSweepsCode));
+
+		const UE::Math::TQuat<FReal>& RotationDouble = InScaledStartTM.GetRotation();
+		RotationSimd = MakeVectorRegisterFloatFromDouble(MakeVectorRegister(RotationDouble.X, RotationDouble.Y, RotationDouble.Z, RotationDouble.W));
+
+		const UE::Math::TVector<FReal>& TranslationDouble = InScaledStartTM.GetTranslation();
+		TranslationSimd = MakeVectorRegisterFloatFromDouble(MakeVectorRegister(TranslationDouble.X, TranslationDouble.Y, TranslationDouble.Z, 0.0));
+
+		// Normalize rotation
+		RotationSimd = VectorNormalizeSafe(RotationSimd, GlobalVectorConstants::Float0001);
+		RayDirSimd = MakeVectorRegisterFloatFromDouble(MakeVectorRegister(ScaledDirNormalized[0], ScaledDirNormalized[1], ScaledDirNormalized[2], 0.0));
 	}
 
 	const void* GetQueryData() const { return nullptr; }
@@ -984,20 +990,16 @@ struct FTriangleMeshSweepVisitor
 		return true;
 	}
 
-	bool VisitRaycast(const TSpatialVisitorData<int32>& VisitData, FQueryFastData& CurData)
+	bool VisitRaycast(const TSpatialVisitorData<int32>& VisitData, FRealSingle& CurDataLength)
 	{
 		check(false);
 		return true;
 	}
 
-	bool VisitSweep(const TSpatialVisitorData<int32>& VisitData, FQueryFastData& CurData)
+	bool VisitSweep(const TSpatialVisitorData<int32>& VisitData, FRealSingle& CurDataLength)
 	{
 		const int32 TriIdx = VisitData.Payload;
 
-		FReal Time;
-		FVec3 HitPosition;
-		FVec3 HitNormal;
-
 		const VectorRegister4Float TriMeshScaleVector = MakeVectorRegisterFloatFromDouble(MakeVectorRegister(TriMeshScale.X, TriMeshScale.Y, TriMeshScale.Z, 0.0));
 
 		const TParticles<FRealSingle, 3>& Particles = TriMesh.MParticles;
@@ -1025,12 +1027,25 @@ struct FTriangleMeshSweepVisitor
 				return true;
 			}
 		}
-
-		if(GJKRaycast2<FReal>(Tri, QueryGeom, ScaledStartTM, ScaledDirNormalized, LengthScale * CurData.CurrentLength, Time, HitPosition, HitNormal, Thickness, bComputeMTD))
+		FRealSingle Time;
+		if(GJKRaycast2ImplSimd(Tri, QueryGeom, RotationSimd, TranslationSimd, RayDirSimd, LengthScale * CurDataLength, Time, OutPositionSimd, OutNormalSimd, bComputeMTD, GlobalVectorConstants::Float1000))
 		{
 			// Time is world scale, OutTime is local scale.
 			if(Time < LengthScale * OutTime)
 			{
+
+				FVec3 HitPosition;
+				FVec3 HitNormal;
+				alignas(16) FRealSingle OutFloat[4];
+				VectorStoreAligned(OutNormalSimd, OutFloat);
+				HitNormal.X = OutFloat[0];
+				HitNormal.Y = OutFloat[1];
+				HitNormal.Z = OutFloat[2];
+
+				VectorStoreAligned(OutPositionSimd, OutFloat);
+				HitPosition.X = OutFloat[0];
+				HitPosition.Y = OutFloat[1];
+				HitPosition.Z = OutFloat[2];
 				TransformSweepOutputsHelper(TriMeshScale, HitNormal, HitPosition, LengthScale, Time, OutNormal, OutPosition, OutTime);
 
 				OutFaceIndex = TriIdx;
@@ -1038,13 +1053,13 @@ struct FTriangleMeshSweepVisitor
 
 				if(Time <= 0)	//MTD or initial overlap
 				{
-					CurData.SetLength(0);
+					CurDataLength = 0.0f;
 
 					//initial overlap, no one will beat this
 					return false;
 				}
 
-				CurData.SetLength(OutTime);
+				CurDataLength = static_cast<FRealSingle>(OutTime);
 			}
 		}
 
@@ -1055,8 +1070,6 @@ struct FTriangleMeshSweepVisitor
 	const TArray<TVec3<IdxType>>& Elements;
 	const FRigidTransform3 StartTM;
 	const QueryGeomType& QueryGeom;
-	const FVec3& Dir;
-	const FReal Thickness;
 	const bool bComputeMTD;
 	const FReal CullsBackFaceSweepsCode; // 0: no culling, 1/-1: winding order
 	VectorRegister4Float VectorCullsBackFaceSweepsCode; // 0: no culling, 1/-1: winding order
@@ -1064,16 +1077,21 @@ struct FTriangleMeshSweepVisitor
 	// Cache these values for Scaled Triangle Mesh, as they are needed for transformation when sweeping against triangles.
 	FVec3 ScaledDirNormalized;
 	VectorRegister4Float VectorScaledDirNormalized;
-	FReal LengthScale;
-	FRigidTransform3 ScaledStartTM;
+	FRealSingle LengthScale;
 
-	FReal OutTime;
+	FRealSingle OutTime;
 	FVec3 OutPosition;
 	FVec3 OutNormal;
 	int32 OutFaceIndex;
 	FVec3 OutFaceNormal;
 
 	FVec3 TriMeshScale;
+
+	VectorRegister4Float RotationSimd;
+	VectorRegister4Float TranslationSimd;
+	VectorRegister4Float RayDirSimd;
+	VectorRegister4Float OutPositionSimd, OutNormalSimd;
+
 };
 
 void ComputeScaledSweepInputs(FVec3 TriMeshScale, const FRigidTransform3& StartTM, const FVec3& Dir, const FReal Length,
@@ -1127,7 +1145,7 @@ bool FTriangleMeshImplicitObject::SweepGeomImp(const QueryGeomType& QueryGeom, c
 	{
 		const FReal CullsBackFaceRaycastCode = bCullsBackFaceRaycast ? GetWindingOrder(TriMeshScale) : 0.f;
 		using VisitorType = FTriangleMeshSweepVisitor<QueryGeomType,decltype(Elements[0][0])>;
-		VisitorType SQVisitor(*this,Elements, QueryGeom,StartTM,Dir,ScaledDirNormalized,LengthScale,ScaledStartTM,Thickness,bComputeMTD, TriMeshScale, CullsBackFaceRaycastCode);
+		VisitorType SQVisitor(*this, Elements, QueryGeom, ScaledDirNormalized, LengthScale, ScaledStartTM, bComputeMTD, TriMeshScale, CullsBackFaceRaycastCode);
 
 		const FAABB3 QueryBounds = QueryGeom.BoundingBox().TransformedAABB(FRigidTransform3(FVec3::ZeroVector,StartTM.GetRotation()));
 		const FVec3 InvTriMeshScale = SafeInvScale(TriMeshScale);
@@ -1465,12 +1483,12 @@ void FTriangleMeshImplicitObject::RebuildFastBVHFromTree(const BVHType& TreeBVH)
 
 		// make the node
 		FTrimeshBVH::FNode& NewNode = FastBVH.Nodes.Emplace_GetRef();
-		NewNode.Children[0].ChildOrFaceIndex = 0;
-		NewNode.Children[0].FaceCount = Leaf.Elems.Num();
-		NewNode.Children[0].Bounds = { BoundingBox().Min(), BoundingBox().Max() };
+		NewNode.Children[0].SetChildOrFaceIndex(0);
+		NewNode.Children[0].SetFaceCount(Leaf.Elems.Num());
+		NewNode.Children[0].SetBounds(BoundingBox());
 		for (const TPayloadBoundsElement<int32, FRealSingle>& LeafPayload: Leaf.Elems)
 		{
-			FastBVH.FaceBounds.Add(LeafPayload.Bounds);
+			FastBVH.FaceBounds.Add(FAABBVectorized(LeafPayload.Bounds));
 			FastBVH.Faces.Add(LeafPayload.Payload);
 		}
 		return;
@@ -1514,11 +1532,11 @@ void FTriangleMeshImplicitObject::RebuildFastBVHFromTree(const BVHType& TreeBVH)
 			{
 				// common infos
 				FTrimeshBVH::FChildData& ChildData = NewNode.Children[ChildIndex];
-				ChildData.Bounds = Node.ChildrenBounds[ChildIndex];
+				ChildData.SetBounds(Node.ChildrenBounds[ChildIndex]);
 				const int32 ChildNodeIndex = Node.ChildrenNodes[ChildIndex];
 				// index in the original BVH space, remapping is done at the end of the process
 				// this may be overwritten if the child is a leaf
-				ChildData.ChildOrFaceIndex = ChildNodeIndex;
+				ChildData.SetChildOrFaceIndex(ChildNodeIndex);
 
 				if (Nodes.IsValidIndex(ChildNodeIndex))
 				{
@@ -1529,15 +1547,15 @@ void FTriangleMeshImplicitObject::RebuildFastBVHFromTree(const BVHType& TreeBVH)
 						const LeafType& Leaf = Leaves[*LeafIndex];
 
 						// store face range in the node 
-						ChildData.ChildOrFaceIndex = FastBVH.Faces.Num();
-						ChildData.FaceCount = Leaf.Elems.Num();
-						check(ChildData.FaceCount > 0);
+						ChildData.SetChildOrFaceIndex(FastBVH.Faces.Num());
+						ChildData.SetFaceCount(Leaf.Elems.Num());
+						check(ChildData.GetFaceCount() > 0);
 
 						TMap<int32, int32> VertexReuse;
 						// copy indices in the linear face array
 						for (const auto& LeafPayload: Leaf.Elems)
 						{
-							FastBVH.FaceBounds.Add(LeafPayload.Bounds);
+							FastBVH.FaceBounds.Add(FAABBVectorized(LeafPayload.Bounds));
 							FastBVH.Faces.Add(LeafPayload.Payload);
 						}
 					}
@@ -1555,13 +1573,14 @@ void FTriangleMeshImplicitObject::RebuildFastBVHFromTree(const BVHType& TreeBVH)
 		for (int32 ChildIndex=0; ChildIndex<2; ++ChildIndex)
 		{
 			FTrimeshBVH::FChildData& ChildData = Node.Children[ChildIndex];
-			if (!ChildData.HasFaces())
+			const bool bHasFaces = ChildData.GetFaceCount() > 0;
+			if (!bHasFaces)
 			{
-				const int32 ChildNodeIndex = ChildData.ChildOrFaceIndex;
+				const int32 ChildNodeIndex = ChildData.GetChildOrFaceIndex();
 				const int32* FixedChildNodeIndex = BVHToFastBVHNodeIndexMap.Find(ChildNodeIndex);
 				if (ensure(FixedChildNodeIndex))
 				{
-					ChildData.ChildOrFaceIndex = *FixedChildNodeIndex;
+					ChildData.SetChildOrFaceIndex( *FixedChildNodeIndex);
 				}
 			}
 		}
diff --git a/Engine/Source/Runtime/Experimental/Chaos/Public/Chaos/AABBTree.h b/Engine/Source/Runtime/Experimental/Chaos/Public/Chaos/AABBTree.h
--- a/Engine/Source/Runtime/Experimental/Chaos/Public/Chaos/AABBTree.h
+++ b/Engine/Source/Runtime/Experimental/Chaos/Public/Chaos/AABBTree.h
@@ -152,8 +152,8 @@ struct TAABBTreeIntersectionHelper
 		check(false);
 		return true;
 	}
-	static bool IntersectsVectorized(const VectorRegister4Float& Start, FQueryFastData& QueryFastData, VectorRegister4Float& TOI, VectorRegister4Float& OutPosition,
-		const FAABBVectorized& Bounds, const FAABBVectorized& QueryBounds, const VectorRegister4Float& QueryHalfExtents, const VectorRegister4Float& Dir, const VectorRegister4Float InvDir, const bool bParallel[3])
+	FORCEINLINE_DEBUGGABLE static bool IntersectsVectorized(const VectorRegister4Float& Start, FQueryFastData& QueryFastData, const FAABBVectorized& Bounds, const FAABBVectorized& QueryBounds,
+		const VectorRegister4Float& QueryHalfExtents, const VectorRegister4Float& InvDir, const VectorRegister4Float& Parallel)
 	{
 		check(false);
 		return true;
@@ -169,8 +169,8 @@ struct TAABBTreeIntersectionHelper<FQueryFastData, EAABBQueryType::Raycast>
 		return Bounds.RaycastFast(Start, Dir, InvDir, bParallel, QueryFastData.CurrentLength, QueryFastData.InvCurrentLength, TOI, OutPosition);
 	}
 
-	FORCEINLINE_DEBUGGABLE static bool IntersectsVectorized(const VectorRegister4Float& Start, FQueryFastData& QueryFastData, VectorRegister4Float& TOI, VectorRegister4Float& OutPosition,
-		const FAABBVectorized& Bounds, const FAABBVectorized& QueryBounds, const VectorRegister4Float& QueryHalfExtents, const VectorRegister4Float& Dir, const VectorRegister4Float InvDir, const bool bParallel[3])
+	FORCEINLINE_DEBUGGABLE static bool IntersectsVectorized(const VectorRegister4Float& Start, FQueryFastData& QueryFastData, const FAABBVectorized& Bounds, const FAABBVectorized& QueryBounds,
+		const VectorRegister4Float& QueryHalfExtents, const VectorRegister4Float& InvDir, const VectorRegister4Float& Parallel)
 	{
 		check(false);
 		return true;
@@ -187,13 +187,12 @@ struct TAABBTreeIntersectionHelper<FQueryFastData, EAABBQueryType::Sweep>
 		return SweepBounds.RaycastFast(Start, Dir, InvDir, bParallel, QueryFastData.CurrentLength, QueryFastData.InvCurrentLength, TOI, OutPosition);
 	}
 
-	FORCEINLINE_DEBUGGABLE static bool IntersectsVectorized(const VectorRegister4Float& Start, FQueryFastData& QueryFastData, VectorRegister4Float& TOI, VectorRegister4Float& OutPosition,
-		const FAABBVectorized& Bounds, const FAABBVectorized& QueryBounds, const VectorRegister4Float& QueryHalfExtents, const VectorRegister4Float& Dir, const VectorRegister4Float InvDir, const bool bParallel[3])
+	FORCEINLINE_DEBUGGABLE static bool IntersectsVectorized(const VectorRegister4Float& Start, FQueryFastData& QueryFastData, const FAABBVectorized& Bounds, const FAABBVectorized& QueryBounds,
+															const VectorRegister4Float& QueryHalfExtents, const VectorRegister4Float& InvDir, const VectorRegister4Float& Parallel)
 	{
 		VectorRegister4Float CurrentLength = MakeVectorRegisterFloatFromDouble(VectorLoadDouble1(&QueryFastData.CurrentLength));
-		VectorRegister4Float InvCurrentLength = MakeVectorRegisterFloatFromDouble(VectorLoadDouble1(&QueryFastData.InvCurrentLength));
 		FAABBVectorized SweepBounds(VectorSubtract(Bounds.GetMin(), QueryHalfExtents), VectorAdd(Bounds.GetMax(), QueryHalfExtents));
-		return SweepBounds.RaycastFast(Start, Dir, InvDir, bParallel, CurrentLength, InvCurrentLength, TOI, OutPosition);
+		return SweepBounds.RaycastFast(Start, InvDir, Parallel, CurrentLength);
 	}
 
 };
@@ -339,6 +338,7 @@ struct TAABBTreeLeafArray : public TBoundsWrapperHelper<TPayloadType, T, bComput
 		const VectorRegister4Float DirSimd = MakeVectorRegisterFloatFromDouble(MakeVectorRegisterDouble(Dir.X, Dir.Y, Dir.Z, 0.0));
 		const VectorRegister4Float InvDirSimd = MakeVectorRegisterFloatFromDouble(MakeVectorRegisterDouble(InvDir.X, InvDir.Y, InvDir.Z, 0.0));
 		const VectorRegister4Float QueryHalfExtentsSimd = MakeVectorRegisterFloatFromDouble(MakeVectorRegisterDouble(QueryHalfExtents.X, QueryHalfExtents.Y, QueryHalfExtents.Z, 0.0));
+		const VectorRegister4Float Parallel = VectorCompareGT(GlobalVectorConstants::SmallNumber, VectorAbs(DirSimd));
 		
 		for (const auto& Elem : Elems)
 		{
@@ -353,12 +353,10 @@ struct TAABBTreeLeafArray : public TBoundsWrapperHelper<TPayloadType, T, bComput
 				const VectorRegister4Float Min = MakeVectorRegisterFloatFromDouble(MakeVectorRegisterDouble(InstanceBounds.Min().X, InstanceBounds.Min().Y, InstanceBounds.Min().Z, 0.0));
 				const VectorRegister4Float Max = MakeVectorRegisterFloatFromDouble(MakeVectorRegisterDouble(InstanceBounds.Max().X, InstanceBounds.Max().Y, InstanceBounds.Max().Z, 0.0));
 				const FAABBVectorized InstanceBoundsSimd(Min, Max);
-				VectorRegister4Float TmpPosition;
-				VectorRegister4Float TOI;
 				VectorRegister4Float StartSimd = MakeVectorRegisterFloatFromDouble(MakeVectorRegisterDouble(Start.X, Start.Y, Start.Z, 0.0));
 
 				if (TAABBTreeIntersectionHelper<TQueryFastData, EAABBQueryType::Sweep>::IntersectsVectorized(
-					StartSimd, QueryFastData, TOI, TmpPosition, InstanceBoundsSimd, FAABBVectorized(), QueryHalfExtentsSimd, DirSimd, InvDirSimd, bParallel))
+					StartSimd, QueryFastData, InstanceBoundsSimd, FAABBVectorized(), QueryHalfExtentsSimd, InvDirSimd, Parallel))
 				{
 					TSpatialVisitorData<TPayloadType> VisitData(Elem.Payload, true, InstanceBounds);
 					const bool bContinue = (bSweep && Visitor.VisitSweep(VisitData, QueryFastData)) || (!bSweep && Visitor.VisitRaycast(VisitData, QueryFastData));
diff --git a/Engine/Source/Runtime/Experimental/Chaos/Public/Chaos/AABBVectorized.h b/Engine/Source/Runtime/Experimental/Chaos/Public/Chaos/AABBVectorized.h
--- a/Engine/Source/Runtime/Experimental/Chaos/Public/Chaos/AABBVectorized.h
+++ b/Engine/Source/Runtime/Experimental/Chaos/Public/Chaos/AABBVectorized.h
@@ -1,6 +1,7 @@
 // Copyright Epic Games, Inc. All Rights Reserved.
 #pragma once
 
+#include "Chaos/AABB.h"
 #include "Chaos/Core.h"
 #include "ChaosArchive.h"
 #include "Math/VectorRegister.h"
@@ -11,72 +12,85 @@
 namespace Chaos
 {
 
-class FAABBVectorized
-{
-public:
-
-	FORCEINLINE FAABBVectorized()
+	class alignas(16) FAABBVectorized
 	{
-		Min = GlobalVectorConstants::BigNumber;
-		Max = VectorNegate(GlobalVectorConstants::BigNumber);
-	}
+	public:
 
-	FORCEINLINE FAABBVectorized(const VectorRegister4Float& Min, const VectorRegister4Float& Max)
-		: Min(Min)
-		, Max(Max)
-	{
-	}
-
-	FORCEINLINE const VectorRegister4Float& GetMin() const { return Min; }
-	FORCEINLINE const VectorRegister4Float& GetMax() const { return Max; }
+		FORCEINLINE_DEBUGGABLE FAABBVectorized()
+		{
+			Min = GlobalVectorConstants::BigNumber;
+			Max = VectorNegate(GlobalVectorConstants::BigNumber);
+		}
 
-	FORCEINLINE_DEBUGGABLE bool RaycastFast(const VectorRegister4Float& StartPoint, const VectorRegister4Float& Dir, const VectorRegister4Float& InvDir, const bool* bParallel, const VectorRegister4Float& Length, const VectorRegister4Float& InvLength, VectorRegister4Float& OutEntryTime, VectorRegister4Float& OutExitTime) const
-	{
-		const VectorRegister4Float StartToMin = VectorSubtract(Min, StartPoint);
-		const VectorRegister4Float StartToMax = VectorSubtract(Max, StartPoint);
-		const VectorRegister4Float Parallel = MakeVectorRegisterFloat(~uint32(0) * static_cast<uint32>(bParallel[0]), ~uint32(0) * static_cast<uint32>(bParallel[1]), ~uint32(0) * static_cast<uint32>(bParallel[2]), 0);
-		const VectorRegister4Float StarToMinGTZero = VectorCompareGT(StartToMin, VectorZero());
-		const VectorRegister4Float ZeroGTStarToMax = VectorCompareGT(VectorZero(), StartToMax);
-		VectorRegister4Float IsFalse = VectorBitwiseAnd(VectorBitwiseOr(StarToMinGTZero, ZeroGTStarToMax), Parallel);
-
-		if (VectorMaskBits(IsFalse))
+		FORCEINLINE_DEBUGGABLE FAABBVectorized(const VectorRegister4Float& InMin, const VectorRegister4Float& InMax)
+			: Min(InMin)
+			, Max(InMax)
 		{
-			return false;	//parallel and outside
+		}
+		template<typename T>
+		FORCEINLINE_DEBUGGABLE explicit FAABBVectorized(const TAABB<T, 3>& AABB)
+		{
+			const VectorRegister4Float MinSmid = MakeVectorRegisterFloatFromDouble(MakeVectorRegisterDouble(AABB.Min().X, AABB.Min().Y, AABB.Min().Z, 0.0));
+			const VectorRegister4Float MaxSmid = MakeVectorRegisterFloatFromDouble(MakeVectorRegisterDouble(AABB.Max().X, AABB.Max().Y, AABB.Max().Z, 0.0));
+			Min = MinSmid;
+			Max = MaxSmid;
 		}
 
-		const VectorRegister4Float StartToMinInvDir = VectorMultiply(StartToMin, InvDir);
-		const VectorRegister4Float StartToMaxInvDir = VectorMultiply(StartToMax, InvDir);
+		FORCEINLINE_DEBUGGABLE const VectorRegister4Float& GetMin() const { return Min; }
+		FORCEINLINE_DEBUGGABLE const VectorRegister4Float& GetMax() const { return Max; }
 
-		VectorRegister4Float Time1 = VectorSelect(Parallel, VectorZero(), StartToMinInvDir);
-		VectorRegister4Float Time2 = VectorSelect(Parallel, Length, StartToMaxInvDir);
-		const VectorRegister4Float Time1GTTime2 = VectorCompareGT(Time1, Time2);
+		FORCEINLINE_DEBUGGABLE void SetMin(const VectorRegister4Float& InMin) { Min = InMin; }
+		FORCEINLINE_DEBUGGABLE void SetMax(const VectorRegister4Float& InMax) { Max = InMax; }
 
-		// Maybe could be optimized with a xor
-		const VectorRegister4Float TmpTime1 = VectorSelect(Time1GTTime2, Time2, Time1);
-		Time2 = VectorSelect(Time1GTTime2, Time1, Time2);
-		Time1 = TmpTime1;
+		FORCEINLINE_DEBUGGABLE bool RaycastFast(const VectorRegister4Float& StartPoint, const VectorRegister4Float& InvDir, const VectorRegister4Float& Parallel, const VectorRegister4Float& Length) const
+		{
+			const VectorRegister4Float StarToMinGTZero = VectorCompareGT(Min, StartPoint);
+			const VectorRegister4Float ZeroGTStarToMax = VectorCompareGT(StartPoint, Max);
+			VectorRegister4Float IsFalse = VectorBitwiseAnd(VectorBitwiseOr(StarToMinGTZero, ZeroGTStarToMax), Parallel);
 
-		VectorRegister4Float LatestStartTime = VectorMax(Time1, VectorSwizzle(Time1, 1, 2, 0, 3));
-		LatestStartTime = VectorMax(LatestStartTime, VectorSwizzle(Time1, 2, 0, 1, 3));
-		LatestStartTime = VectorMax(LatestStartTime, VectorZero());
+			if (VectorMaskBits(IsFalse))
+			{
+				return false;	//parallel and outside
+			}
 
-		VectorRegister4Float EarliestEndTime = VectorMin(Time2, VectorSwizzle(Time2, 1, 2, 0, 3));
-		EarliestEndTime = VectorMin(EarliestEndTime, VectorSwizzle(Time2, 2, 0, 1, 3));
-		EarliestEndTime = VectorMin(EarliestEndTime, Length);
+			const VectorRegister4Float StartToMin = VectorSubtract(Min, StartPoint);
+			const VectorRegister4Float StartToMax = VectorSubtract(Max, StartPoint);
+			const VectorRegister4Float StartToMinInvDir = VectorMultiply(StartToMin, InvDir);
+			const VectorRegister4Float StartToMaxInvDir = VectorMultiply(StartToMax, InvDir);
 
-		//Outside of slab before entering another
-		IsFalse = VectorCompareGT(LatestStartTime, EarliestEndTime);
+			const VectorRegister4Float Time1 = VectorBitwiseNotAnd(Parallel, StartToMinInvDir);
+			const VectorRegister4Float Time2 = VectorSelect(Parallel, Length, StartToMaxInvDir);
 
-		if (VectorMaskBits(IsFalse))
+			const VectorRegister4Float SortedTime1 = VectorMin(Time1, Time2);
+			const VectorRegister4Float SortedTime2 = VectorMax(Time1, Time2);
+
+			VectorRegister4Float LatestStartTime = VectorMax(SortedTime1, VectorSwizzle(SortedTime1, 1, 2, 0, 3));
+			LatestStartTime = VectorMax(LatestStartTime, VectorSwizzle(SortedTime1, 2, 0, 1, 3));
+			LatestStartTime = VectorMax(LatestStartTime, VectorZero());
+
+			VectorRegister4Float EarliestEndTime = VectorMin(SortedTime2, VectorSwizzle(SortedTime2, 1, 2, 0, 3));
+			EarliestEndTime = VectorMin(EarliestEndTime, VectorSwizzle(SortedTime2, 2, 0, 1, 3));
+			EarliestEndTime = VectorMin(EarliestEndTime, Length);
+
+			//Outside of slab before entering another
+			IsFalse = VectorCompareGT(LatestStartTime, EarliestEndTime);
+
+			if (VectorMaskBits(IsFalse))
+			{
+				return false;
+			}
+
+			return true;
+		}
+
+		FORCEINLINE_DEBUGGABLE bool Intersects(const FAABBVectorized& Other) const
 		{
-			return false; 
+			VectorRegister4Float IsFalse = VectorBitwiseOr(VectorCompareGT(Min, Other.GetMax()), VectorCompareGT(Other.GetMin(), Max));
+
+			return !static_cast<bool>(VectorMaskBits(IsFalse));
 		}
-		OutEntryTime = LatestStartTime;
-		OutExitTime = EarliestEndTime;
-		return true;
-	}
-
-private:
-	VectorRegister4Float Min, Max;
-};
+
+	private:
+		VectorRegister4Float Min, Max;
+	};
 }
\ No newline at end of file
diff --git a/Engine/Source/Runtime/Experimental/Chaos/Public/Chaos/GJK.h b/Engine/Source/Runtime/Experimental/Chaos/Public/Chaos/GJK.h
--- a/Engine/Source/Runtime/Experimental/Chaos/Public/Chaos/GJK.h
+++ b/Engine/Source/Runtime/Experimental/Chaos/Public/Chaos/GJK.h
@@ -1699,7 +1699,7 @@ namespace Chaos
 
 	template <typename TGeometryA, typename TGeometryB, typename T>
 	bool GJKRaycast2ImplSimd(const TGeometryA& A, const TGeometryB& B, const VectorRegister4Float& BToARotation, const VectorRegister4Float& StartPoint, const VectorRegister4Float& RayDir, const T RayLength,
-		T& OutTime, VectorRegister4Float& OutPosition, VectorRegister4Float& OutNormal, bool bComputeMTD, const VectorRegister4Float& InitialDir, const TRigidTransform<double, 3>& StartTM)
+		T& OutTime, VectorRegister4Float& OutPosition, VectorRegister4Float& OutNormal, bool bComputeMTD, const VectorRegister4Float& InitialDir)
 	{
 		ensure(RayLength > 0);
 
@@ -2076,7 +2076,7 @@ namespace Chaos
 
 		FRealSingle OutTimeFloat;
 		VectorRegister4Float OutPositionSimd, OutNormalSimd;
-		bool result = GJKRaycast2ImplSimd(A, B, Rotation, Translation, RayDirSimd, static_cast<FRealSingle>(RayLength), OutTimeFloat, OutPositionSimd, OutNormalSimd, bComputeMTD, InitialDirSimd, StartTM);
+		bool result = GJKRaycast2ImplSimd(A, B, Rotation, Translation, RayDirSimd, static_cast<FRealSingle>(RayLength), OutTimeFloat, OutPositionSimd, OutNormalSimd, bComputeMTD, InitialDirSimd);
 
 		OutTime = static_cast<double>(OutTimeFloat);
 
diff --git a/Engine/Source/Runtime/Experimental/Chaos/Public/Chaos/TriangleMeshImplicitObject.h b/Engine/Source/Runtime/Experimental/Chaos/Public/Chaos/TriangleMeshImplicitObject.h
--- a/Engine/Source/Runtime/Experimental/Chaos/Public/Chaos/TriangleMeshImplicitObject.h
+++ b/Engine/Source/Runtime/Experimental/Chaos/Public/Chaos/TriangleMeshImplicitObject.h
@@ -25,6 +25,7 @@ namespace Chaos
 
 	struct FMTDInfo;
 	class FContactPoint;
+	FArchive& operator<<(FArchive& Ar, FAABBVectorized& Bounds);
 
 	struct CHAOS_API FTrimeshBVH
 	{
@@ -42,25 +43,56 @@ namespace Chaos
 			Keep,
 		};
 		
-		struct FChildData
+		struct alignas(16) FChildData
 		{
 			FChildData()
 				: ChildOrFaceIndex(INDEX_NONE)
 				, FaceCount(0)
-			{}
-			FAABBType Bounds;
-			int32 ChildOrFaceIndex;
-			int32 FaceCount;
+			{
+			}
+
+			FORCEINLINE int32 GetChildOrFaceIndex() const
+			{
+				return ChildOrFaceIndex;
+			}
 
-			bool HasFaces() const { return (FaceCount > 0); }
+			FORCEINLINE int32 GetFaceCount() const
+			{
+				return FaceCount;
+			}
 
-			void Serialize(FArchive& Ar)
+			FORCEINLINE void SetChildOrFaceIndex(int32 InChildOrFaceIndex)
 			{
-				TBox<FRealSingle, 3>::SerializeAsAABB(Ar, Bounds);
+				ChildOrFaceIndex = InChildOrFaceIndex;
+			}
 
+			FORCEINLINE void SetFaceCount(int32 InFaceCount)
+			{
+				FaceCount = InFaceCount;
+			}
+
+			FORCEINLINE void SetBounds(const FAABB3& AABB)
+			{
+				Bounds = FAABBVectorized(AABB);
+			}
+			FORCEINLINE void SetBounds(const TAABB<FRealSingle, 3>& AABB)
+			{
+				Bounds = FAABBVectorized(AABB);
+			}
+
+			const FAABBVectorized& GetBounds() const { return Bounds; }
+
+			void Serialize(FArchive& Ar)
+			{
+				Ar << Bounds;
 				Ar << ChildOrFaceIndex;
 				Ar << FaceCount;
 			}
+
+		private:
+			FAABBVectorized Bounds;
+			int32 ChildOrFaceIndex;
+			int32 FaceCount;
 		};
 		
 		struct FNode
@@ -78,26 +110,25 @@ namespace Chaos
 		template <typename SQVisitor>
 		FORCEINLINE_DEBUGGABLE void Raycast(const FVec3& Start, const FVec3& Dir, const FReal Length, SQVisitor& Visitor) const
 		{
-			FQueryFastData QueryFastData(Dir, Length);
-			const auto BoundsFilter = [&QueryFastData, &Start](const FAABBType& Bounds) -> EFilterResult
+			FRealSingle CurrentLength = static_cast<FRealSingle>(Length);
+			const VectorRegister4Float StartSimd = MakeVectorRegisterFloatFromDouble(MakeVectorRegisterDouble(Start.X, Start.Y, Start.Z, 0.0));
+			const VectorRegister4Float DirSimd = MakeVectorRegisterFloatFromDouble(MakeVectorRegisterDouble(Dir.X, Dir.Y, Dir.Z, 0.0));
+			const VectorRegister4Float Parallel = VectorCompareGT(GlobalVectorConstants::SmallNumber, VectorAbs(DirSimd));
+			const VectorRegister4Float InvDirSimd = VectorBitwiseNotAnd(Parallel, VectorDivide(VectorOne(), DirSimd));
+
+			const auto BoundsFilter = [&StartSimd, &InvDirSimd, &Parallel, &CurrentLength](const FAABBVectorized& Bounds) -> EFilterResult
 			{
-				FReal HitTime;
-				FVec3 HitPosition;
+				const VectorRegister4Float CurDataLength = VectorLoadFloat1(&CurrentLength);
 
-				bool bCulled = true;
-				const bool bHit = Bounds.RaycastFast(Start, QueryFastData.Dir, QueryFastData.InvDir, QueryFastData.bParallel, QueryFastData.CurrentLength, QueryFastData.InvCurrentLength, HitTime, HitPosition);
-				if (bHit)
-				{
-					// avoid going down teh nodes that are further than the current hit distance
-					bCulled = (QueryFastData.CurrentLength < HitTime);
-				}
-				
-				return (!bCulled) ? EFilterResult::Keep: EFilterResult::Skip;
+				FAABBVectorized CleanBounds(Bounds.GetMin(), Bounds.GetMax());
+				const bool bHit = CleanBounds.RaycastFast(StartSimd, InvDirSimd, Parallel, CurDataLength);
+
+				return bHit ? EFilterResult::Keep : EFilterResult::Skip;
 			};
 
-			const auto FaceVisitor = [&Visitor, &QueryFastData](int32 FaceIndex)
+			const auto FaceVisitor = [&Visitor, &CurrentLength](int32 FaceIndex)
 			{
-				const bool bContinueVisiting = Visitor.VisitRaycast(FaceIndex, QueryFastData);
+				const bool bContinueVisiting = Visitor.VisitRaycast(FaceIndex, CurrentLength);
 				return bContinueVisiting? EVisitorResult::Continue: EVisitorResult::Stop;
 			};
 
@@ -107,27 +138,25 @@ namespace Chaos
 		template <typename SQVisitor>
 		FORCEINLINE_DEBUGGABLE void Sweep(const FVec3& Start, const FVec3& Dir, const FReal Length, const FVec3& QueryHalfExtents, SQVisitor& Visitor) const
 		{
-			FQueryFastData QueryFastData(Dir, Length);
-			const auto BoundsFilter = [&QueryFastData, &Start, &QueryHalfExtents](const FAABBType& Bounds) -> EFilterResult
-			{
-				FReal HitTime;
-				FVec3 HitPosition;
+			FRealSingle CurrentLength = static_cast<FRealSingle>(Length);
+			const VectorRegister4Float StartSimd = MakeVectorRegisterFloatFromDouble(MakeVectorRegisterDouble(Start.X, Start.Y, Start.Z, 0.0));
+			const VectorRegister4Float DirSimd = MakeVectorRegisterFloatFromDouble(MakeVectorRegisterDouble(Dir.X, Dir.Y, Dir.Z, 0.0));
+			const VectorRegister4Float Parallel = VectorCompareGT(GlobalVectorConstants::SmallNumber, VectorAbs(DirSimd));
+			const VectorRegister4Float InvDirSimd = VectorBitwiseNotAnd(Parallel, VectorDivide(VectorOne(), DirSimd));
 
-				bool bCulled = true;
-				const FAABB3 SweepBounds(Bounds.Min() - QueryHalfExtents, Bounds.Max() + QueryHalfExtents);
-				const bool bHit = SweepBounds.RaycastFast(Start, QueryFastData.Dir, QueryFastData.InvDir, QueryFastData.bParallel, QueryFastData.CurrentLength, QueryFastData.InvCurrentLength, HitTime, HitPosition);
-				if (bHit)
-				{
-					// avoid going down the nodes that are further than the current hit distance
-					bCulled = (QueryFastData.CurrentLength < HitTime);
-				}
+			const VectorRegister4Float QueryHalfExtentsSimd = MakeVectorRegisterFloatFromDouble(MakeVectorRegisterDouble(QueryHalfExtents.X, QueryHalfExtents.Y, QueryHalfExtents.Z, 0.0));
+			const auto BoundsFilter = [&StartSimd, &InvDirSimd, &Parallel, &CurrentLength, &QueryHalfExtentsSimd](const FAABBVectorized& Bounds) -> EFilterResult
+			{
+				const VectorRegister4Float CurDataLength = VectorLoadFloat1(&CurrentLength);
+				const FAABBVectorized SweepBounds(VectorSubtract(Bounds.GetMin(), QueryHalfExtentsSimd), VectorAdd(Bounds.GetMax(), QueryHalfExtentsSimd));
+				const bool bHit = SweepBounds.RaycastFast(StartSimd, InvDirSimd, Parallel, CurDataLength);
 
-				return (!bCulled) ? EFilterResult::Keep : EFilterResult::Skip;
+				return bHit ? EFilterResult::Keep : EFilterResult::Skip;
 			};
 
-			const auto FaceVisitor = [&Visitor, &QueryFastData](int32 FaceIndex)
+			const auto FaceVisitor = [&Visitor, &CurrentLength](int32 FaceIndex)
 			{
-				const bool bContinueVisiting = Visitor.VisitSweep(FaceIndex, QueryFastData);
+				const bool bContinueVisiting = Visitor.VisitSweep(FaceIndex, CurrentLength);
 				return bContinueVisiting ? EVisitorResult::Continue : EVisitorResult::Stop;
 			};
 
@@ -135,11 +164,13 @@ namespace Chaos
 		}
 
 		template <typename SQVisitor>
-		FORCEINLINE_DEBUGGABLE void Overlap(const FAABB3& AABB, SQVisitor& Visitor) const
+		FORCEINLINE_DEBUGGABLE void Overlap(const FAABBVectorized& AABB, SQVisitor& Visitor) const
 		{
-			const auto BoundsFilter = [&AABB](const FAABBType& Bounds) -> EFilterResult
+			CSV_SCOPED_TIMING_STAT(PhysicsVerbose, BVHOverlap);
+			const auto BoundsFilter = [&AABB](const FAABBVectorized& Bounds) -> EFilterResult
 			{
-				const bool bHit = Bounds.Intersects(AABB);
+				FAABBVectorized CleanBounds(Bounds.GetMin(), Bounds.GetMax());
+				const bool bHit = CleanBounds.Intersects(AABB);
 				return bHit ? EFilterResult::Keep : EFilterResult::Skip;
 			};
 
@@ -171,7 +202,7 @@ namespace Chaos
 			}
 			return EVisitorResult::Continue;
 		}
-		
+
 		template <typename BoundsFilterType, typename FaceVisitorType>
 		FORCEINLINE_DEBUGGABLE void VisitTree(BoundsFilterType& BoundsFilter, FaceVisitorType& FaceVisitor) const
 		{
@@ -191,20 +222,22 @@ namespace Chaos
 				for (int32 ChildIndex = 0; ChildIndex < 2; ++ChildIndex)
 				{
 					const FChildData& ChildData = Node.Children[ChildIndex]; 
-					if (ChildData.ChildOrFaceIndex != INDEX_NONE)
+					const int32 FaceIndex = ChildData.GetChildOrFaceIndex();
+					if (FaceIndex != INDEX_NONE)
 					{
-						if (BoundsFilter(ChildData.Bounds) == EFilterResult::Keep)
+						if (BoundsFilter(ChildData.GetBounds()) == EFilterResult::Keep)
 						{
-							if (ChildData.HasFaces())
+							const uint32 FaceCount = ChildData.GetFaceCount();
+							if (FaceCount > 0)
 							{
-								if (EVisitorResult::Stop == VisitFaces(ChildData.ChildOrFaceIndex, ChildData.FaceCount, BoundsFilter, FaceVisitor))
+								if (EVisitorResult::Stop == VisitFaces(FaceIndex, FaceCount, BoundsFilter, FaceVisitor))
 								{
 									return;
 								}
 							}
 							else
 							{
-								NodeIndexStack.Push(ChildData.ChildOrFaceIndex);
+								NodeIndexStack.Push(FaceIndex);
 							}
 						}
 					}
@@ -220,7 +253,7 @@ namespace Chaos
 		}
 		
 		TArray<FNode> Nodes;
-		TArray<FAABBType> FaceBounds;
+		TArray<FAABBVectorized> FaceBounds;
 		TArray<int32> Faces;
 	};
 
@@ -242,6 +275,32 @@ namespace Chaos
 		return Ar;
 	}
 
+	FORCEINLINE_DEBUGGABLE FChaosArchive& operator<<(FChaosArchive& Ar, FAABBVectorized& Bounds)
+	{
+		alignas(16) FRealSingle Floats[4];
+		VectorStoreAligned(Bounds.GetMin(), Floats);
+		TVector<FRealSingle, 3>  Min(Floats[0], Floats[1], Floats[2]);
+		VectorStoreAligned(Bounds.GetMax(), Floats);
+		TVector<FRealSingle, 3>  Max(Floats[0], Floats[1], Floats[2]);
+		FTrimeshBVH::FAABBType BoundsToSerialize(Min, Max);
+		TBox<FRealSingle, 3>::SerializeAsAABB(Ar, BoundsToSerialize);
+		Bounds = FAABBVectorized(BoundsToSerialize);
+		return Ar;
+	}
+
+	FORCEINLINE_DEBUGGABLE FArchive& operator<<(FArchive& Ar, FAABBVectorized& Bounds)
+	{
+		alignas(16) FRealSingle Floats[4];
+		VectorStoreAligned(Bounds.GetMin(), Floats);
+		TVector<FRealSingle, 3>  Min(Floats[0], Floats[1], Floats[2]);
+		VectorStoreAligned(Bounds.GetMax(), Floats);
+		TVector<FRealSingle, 3>  Max(Floats[0], Floats[1], Floats[2]);
+		FTrimeshBVH::FAABBType BoundsToSerialize(Min, Max);
+		TBox<FRealSingle, 3>::SerializeAsAABB(Ar, BoundsToSerialize);
+		Bounds = FAABBVectorized(BoundsToSerialize);
+		return Ar;
+	}
+
 	FORCEINLINE_DEBUGGABLE FChaosArchive& operator<<(FChaosArchive& Ar, FTrimeshBVH& TrimeshBVH)
 	{
 		TrimeshBVH.Serialize(Ar);
diff --git a/Engine/Source/Runtime/Experimental/Chaos/Public/Chaos/VectorUtility.h b/Engine/Source/Runtime/Experimental/Chaos/Public/Chaos/VectorUtility.h
--- a/Engine/Source/Runtime/Experimental/Chaos/Public/Chaos/VectorUtility.h
+++ b/Engine/Source/Runtime/Experimental/Chaos/Public/Chaos/VectorUtility.h
@@ -66,6 +66,18 @@ FORCEINLINE VectorRegister4Float VectorMoveLh(const VectorRegister4Float& A, con
 	return vzip1q_f64(A, B);
 }
 
+/**
+ * Combines two vectors using bitwise NOT AND (treating each vector as a 128 bit field)
+ *
+ * @param Vec1	1st vector
+ * @param Vec2	2nd vector
+ * @return		VectorRegister4Float( for each bit i: !Vec1[i] & Vec2[i] )
+ */
+FORCEINLINE VectorRegister4Float VectorBitwiseNotAnd(const VectorRegister4Float& A, const VectorRegister4Float& B)
+{
+	return (VectorRegister4Float)vandq_u32(vmvnq_u32((VectorRegister4Int)A), (VectorRegister4Int)B);
+}
+
 /**
  * Return square root.
  *
@@ -104,6 +116,19 @@ FORCEINLINE VectorRegister4Float VectorMoveLh(const VectorRegister4Float& A, con
 	return _mm_movelh_ps(A, B);
 }
 
+
+/**
+ * Combines two vectors using bitwise NOT AND (treating each vector as a 128 bit field)
+ *
+ * @param Vec1	1st vector
+ * @param Vec2	2nd vector
+ * @return		VectorRegister4Float( for each bit i: !Vec1[i] & Vec2[i] )
+ */
+FORCEINLINE VectorRegister4Float VectorBitwiseNotAnd(const VectorRegister4Float& A, const VectorRegister4Float& B)
+{
+	return _mm_andnot_ps(A, B);
+}
+
 /**
  * Return square root.
  *
