Chaos : Fix edge case where sweeping elements in the hashed grid can lead to an infinite loop
- also provided an optimized path for mostly vertical sweeps

diff --git a/Engine/Source/Programs/HeadlessChaos/Private/HeadlessChaosApp.cpp b/Engine/Source/Programs/HeadlessChaos/Private/HeadlessChaosApp.cpp
--- a/Engine/Source/Programs/HeadlessChaos/Private/HeadlessChaosApp.cpp
+++ b/Engine/Source/Programs/HeadlessChaos/Private/HeadlessChaosApp.cpp
@@ -235,6 +235,7 @@ TEST(EPA, EPATests) {
 }
 
 TEST(BP, BroadphaseTests) {
+	ChaosTest::AABBTreeDirtyGridFunctionsWithEdgeCase();
 	ChaosTest::GridBPTest();
 	ChaosTest::GridBPTest2();
 	ChaosTest::AABBTreeTest();
diff --git a/Engine/Source/Programs/HeadlessChaos/Private/HeadlessChaosTestBroadphase.cpp b/Engine/Source/Programs/HeadlessChaos/Private/HeadlessChaosTestBroadphase.cpp
--- a/Engine/Source/Programs/HeadlessChaos/Private/HeadlessChaosTestBroadphase.cpp
+++ b/Engine/Source/Programs/HeadlessChaos/Private/HeadlessChaosTestBroadphase.cpp
@@ -1,5 +1,6 @@
 // Copyright Epic Games, Inc. All Rights Reserved.
 
+
 #include "HeadlessChaosTestBroadphase.h"
 
 #include "HeadlessChaos.h"
@@ -852,6 +853,34 @@ namespace ChaosTest
 		}
 	}
 
+	// specific  sweep query seems to generate infinite loop
+	// the values are outside of the function to avoid optimizing away the values
+	// (the issue would only appear when optimization is enable)
+	namespace EdgeCase
+	{
+		FVec3 QueryHalfExtents{ 5.51277924f, 4.77557945f, 4.96569443f };
+		FVec3 StartPoint{ -40.7950134f, -4.77560043f, -11.2947388f };
+		FVec3 Dir{ 0, 0, -1 };
+		FReal CurrentLength = 146.779785f;
+	}
+
+	void AABBTreeDirtyGridFunctionsWithEdgeCase()
+	{
+		constexpr FReal DirtyElementGridCellSize = 1000.0;
+		constexpr FReal DirtyElementGridCellSizeInv = 1.0 / DirtyElementGridCellSize;
+		constexpr int32 DirtyElementMaxGridCellQueryCount = 340;
+
+		TArray<TVec2<FReal>> VisitedCells;
+		// this should report only be as much as 4 hits as teh grid is 2D and the ray downward vertical and the halfextends are smaller than a cell size
+		DoForSweepIntersectCells(EdgeCase::QueryHalfExtents, EdgeCase::StartPoint, EdgeCase::Dir, EdgeCase::CurrentLength, DirtyElementGridCellSize, DirtyElementGridCellSizeInv,
+			[&VisitedCells](FReal X, FReal Y)
+			{
+				VisitedCells.Add({ X,Y });
+				EXPECT_TRUE(VisitedCells.Num() <= 4 );
+			});
+		EXPECT_TRUE(VisitedCells.Num() <= 4);
+	}
+	
 	void BroadphaseCollectionTest()
 	{
 		using TreeType = TAABBTree<int32, TAABBTreeLeafArray<int32>>;
@@ -1084,4 +1113,4 @@ namespace ChaosTest
 		}
 	}
 
-}
\ No newline at end of file
+}
diff --git a/Engine/Source/Programs/HeadlessChaos/Public/HeadlessChaosTestBroadphase.h b/Engine/Source/Programs/HeadlessChaos/Public/HeadlessChaosTestBroadphase.h
--- a/Engine/Source/Programs/HeadlessChaos/Public/HeadlessChaosTestBroadphase.h
+++ b/Engine/Source/Programs/HeadlessChaos/Public/HeadlessChaosTestBroadphase.h
@@ -21,6 +21,8 @@ namespace ChaosTest {
 
 	void DoForSweepIntersectCellsImpTest();
 
+	void AABBTreeDirtyGridFunctionsWithEdgeCase();
+
 	void BroadphaseCollectionTest();
 
 	void BoundingVolumeNoBoundsTest();
diff --git a/Engine/Source/Runtime/Experimental/Chaos/Public/Chaos/AABBTreeDirtyGridUtils.h b/Engine/Source/Runtime/Experimental/Chaos/Public/Chaos/AABBTreeDirtyGridUtils.h
--- a/Engine/Source/Runtime/Experimental/Chaos/Public/Chaos/AABBTreeDirtyGridUtils.h
+++ b/Engine/Source/Runtime/Experimental/Chaos/Public/Chaos/AABBTreeDirtyGridUtils.h
@@ -324,6 +324,14 @@ namespace Chaos
 		int64 LastCellIndexX = GetDirtyCellIndexFromWorldCoordinate(XEndPointExpanded + LocalCoordinatesOrigin.X, DirtyElementGridCellSizeInv);
 		int64 LastCellIndexY = GetDirtyCellIndexFromWorldCoordinate(YEndPointExpanded + LocalCoordinatesOrigin.Y, DirtyElementGridCellSizeInv);
 
+		// Because of floating point math precision there's case where we can get LastCellIndexY to be smaller than CurrentCellIndexY0 when the ray is stright down
+		// that would cause the while loop to go on forever as it relies on an strict equality
+		// we then need to adjust the value of LastCellIndexY accordingly
+		if (DxTooSmall)
+		{
+			LastCellIndexY = FMath::Max(CurrentCellIndexY0, LastCellIndexY);
+		}
+
 		bool Done = false;
 		while (!Done)
 		{
@@ -414,21 +422,43 @@ namespace Chaos
 	}
 
 	template <typename FunctionType>
-	FORCEINLINE_DEBUGGABLE void DoForSweepIntersectCells(const FVec3 QueryHalfExtents, const FVec3& StartPoint, const FVec3& Dir, FReal Length, FReal DirtyElementGridCellSize, FReal DirtyElementGridCellSizeInv, FunctionType InFunction)
+	void DoForSweepIntersectCells(const FVec3 QueryHalfExtents, const FVec3& StartPoint, const FVec3& Dir, FReal Length, FReal DirtyElementGridCellSize, FReal DirtyElementGridCellSizeInv, FunctionType InFunction)
 	{
-		FReal AbsDx = FMath::Abs(Dir.X);
-		FReal AbsDy = FMath::Abs(Dir.Y);
+		FReal AbsDx = FMath::Abs(Dir.X * Length);
+		FReal AbsDy = FMath::Abs(Dir.Y * Length);
 
 		bool XDirectionDominant = AbsDx >= AbsDy;
 
-		if (XDirectionDominant)
+		// if the ray is mostly vertical then we can default to an simple overlap 
+		if (AbsDx <= SMALL_NUMBER && AbsDy <= SMALL_NUMBER)
 		{
-			DoForSweepIntersectCellsImp(QueryHalfExtents.X, QueryHalfExtents.Y, StartPoint.X, StartPoint.Y, Dir.X * Length, Dir.Y * Length, DirtyElementGridCellSize, DirtyElementGridCellSizeInv, InFunction);
+			// no need to account for the ray length as we only collect 2 dimensional cell coordinates and the ray is already proven to be vertical
+			const FAABB3 QueryBounds(StartPoint- QueryHalfExtents, StartPoint + QueryHalfExtents);
+
+			const int64 CellStartX = GetDirtyCellIndexFromWorldCoordinate(QueryBounds.Min().X, DirtyElementGridCellSizeInv);
+			const int64 CellStartY = GetDirtyCellIndexFromWorldCoordinate(QueryBounds.Min().Y, DirtyElementGridCellSizeInv);
+			const int64 CellEndX = GetDirtyCellIndexFromWorldCoordinate(QueryBounds.Max().X, DirtyElementGridCellSizeInv);
+			const int64 CellEndY = GetDirtyCellIndexFromWorldCoordinate(QueryBounds.Max().Y, DirtyElementGridCellSizeInv);
+			
+			for (int64 X = CellStartX; X <= CellEndX; X++)
+			{
+				for (int64 Y = CellStartY; Y <= CellEndY; Y++)
+				{
+					InFunction((FReal)X * DirtyElementGridCellSize, (FReal)Y * DirtyElementGridCellSize);
+				}
+			}
 		}
 		else
 		{
-			// Swap Y and X
-			DoForSweepIntersectCellsImp(QueryHalfExtents.Y, QueryHalfExtents.X, StartPoint.Y, StartPoint.X, Dir.Y * Length, Dir.X * Length, DirtyElementGridCellSize, DirtyElementGridCellSizeInv, [&](auto X, auto Y) {InFunction(Y, X); });
+			if (XDirectionDominant)
+			{
+				DoForSweepIntersectCellsImp(QueryHalfExtents.X, QueryHalfExtents.Y, StartPoint.X, StartPoint.Y, Dir.X * Length, Dir.Y * Length, DirtyElementGridCellSize, DirtyElementGridCellSizeInv, InFunction);
+			}
+			else
+			{
+				// Swap Y and X
+				DoForSweepIntersectCellsImp(QueryHalfExtents.Y, QueryHalfExtents.X, StartPoint.Y, StartPoint.X, Dir.Y * Length, Dir.X * Length, DirtyElementGridCellSize, DirtyElementGridCellSizeInv, [&](auto X, auto Y) {InFunction(Y, X); });
+			}
 		}
 
 	}
