Fix invalid node order because of the tarray queue

diff --git a/Engine/Source/Runtime/Experimental/Chaos/Private/Chaos/Island/IslandGraph.cpp b/Engine/Source/Runtime/Experimental/Chaos/Private/Chaos/Island/IslandGraph.cpp
--- a/Engine/Source/Runtime/Experimental/Chaos/Private/Chaos/Island/IslandGraph.cpp
+++ b/Engine/Source/Runtime/Experimental/Chaos/Private/Chaos/Island/IslandGraph.cpp
@@ -449,7 +449,7 @@ void FIslandGraph<NodeType, EdgeType, IslandType>::UpdateLevels(const int32 Node
 				if (GraphNodes.IsValidIndex(OtherIndex) && GraphNodes[OtherIndex].bValidNode && GraphNodes[OtherIndex].LevelIndex == INDEX_NONE)
 				{
 					GraphNodes[OtherIndex].LevelIndex = GraphEdge.LevelIndex+1;
-					NodeQueue.Push(OtherIndex);
+					NodeQueue.PushLast(OtherIndex);
 				}
 			}
 		}
@@ -478,7 +478,9 @@ void FIslandGraph<NodeType, EdgeType, IslandType>::ComputeLevels(const int32 Con
 	int32 NodeIndex = INDEX_NONE;
 	while(!NodeQueue.IsEmpty())
 	{
-		NodeIndex = NodeQueue.Pop(false);
+		NodeIndex = NodeQueue.First();
+		NodeQueue.PopFirst();
+
 		UpdateLevels(NodeIndex, ContainerId);
 	}
 
@@ -508,7 +510,7 @@ int32 FIslandGraph<NodeType, EdgeType, IslandType>::PickColor(const FGraphNode&
 		OtherNode.ColorIndices.Add(ColorToUse);
 		if (OtherNode.NodeCounter != GraphCounter)
 		{
-			NodeQueue.Push(OtherIndex);
+			NodeQueue.PushLast(OtherIndex);
 		}
 	}
 	else
@@ -571,10 +573,12 @@ void FIslandGraph<NodeType, EdgeType, IslandType>::ComputeColors(const int32 Con
 			FGraphNode& GraphNode = GraphNodes[RootIndex];
 			if (GraphNode.NodeCounter != GraphCounter && GraphNode.bValidNode)
 			{
-				NodeQueue.Push(RootIndex);
+				NodeQueue.PushLast(RootIndex);
 				while (!NodeQueue.IsEmpty())
 				{
-					NodeIndex = NodeQueue.Pop(false);
+					NodeIndex = NodeQueue.First();
+					NodeQueue.PopFirst();
+
 					UpdateColors(NodeIndex, ContainerId, MinEdges);
 				}
 			}
@@ -585,12 +589,14 @@ void FIslandGraph<NodeType, EdgeType, IslandType>::ComputeColors(const int32 Con
 template<typename NodeType, typename EdgeType, typename IslandType>
 void FIslandGraph<NodeType, EdgeType, IslandType>::SplitIsland(const int32 RootIndex, const int32 IslandIndex)
 {
-	NodeQueue.Push(RootIndex);
+	NodeQueue.PushLast(RootIndex);
 	int32 NodeIndex = RootIndex;
 					
 	while (!NodeQueue.IsEmpty())
 	{
-		NodeIndex = NodeQueue.Pop(false);
+		NodeIndex = NodeQueue.First();
+		NodeQueue.PopFirst();
+
 		FGraphNode& GraphNode = GraphNodes[NodeIndex];
 						
 		// Graph counter is there to avoid processing multiple times the same node/edge
@@ -603,7 +609,7 @@ void FIslandGraph<NodeType, EdgeType, IslandType>::SplitIsland(const int32 RootI
 			for (int32& EdgeIndex : GraphNode.NodeEdges)
 			{
 				FGraphEdge& GraphEdge = GraphEdges[EdgeIndex];
-				if(GraphEdge.EdgeCounter != GraphCounter)
+				if (GraphEdge.EdgeCounter != GraphCounter)
 				{
 					GraphEdge.EdgeCounter = GraphCounter;
 					GraphEdge.IslandIndex = IslandIndex;
@@ -614,7 +620,7 @@ void FIslandGraph<NodeType, EdgeType, IslandType>::SplitIsland(const int32 RootI
 				// Only the valid nodes are allowed to continue the graph traversal
 				if (GraphNodes.IsValidIndex(OtherIndex) && GraphNodes[OtherIndex].NodeCounter != GraphCounter && GraphNodes[OtherIndex].bValidNode)
 				{
-					NodeQueue.Push(OtherIndex);
+					NodeQueue.PushLast(OtherIndex);
 				}
 			}
 		}
diff --git a/Engine/Source/Runtime/Experimental/Chaos/Public/Chaos/Island/IslandGraph.h b/Engine/Source/Runtime/Experimental/Chaos/Public/Chaos/Island/IslandGraph.h
--- a/Engine/Source/Runtime/Experimental/Chaos/Public/Chaos/Island/IslandGraph.h
+++ b/Engine/Source/Runtime/Experimental/Chaos/Public/Chaos/Island/IslandGraph.h
@@ -3,6 +3,7 @@
 
 #include "Containers/Map.h"
 #include "Containers/Queue.h"
+#include "Containers/Deque.h"
 #include "Containers/SparseArray.h"
 
 namespace Chaos
@@ -345,7 +346,7 @@ public:
 	int32 GraphCounter = 0;
 
 	/** List of nodes to be processed in the graph functions */
-	TArray<int32> NodeQueue;
+	TDeque<int32> NodeQueue;
 };
 
 }
\ No newline at end of file
