Specified trailing return type for lambda in TThreadedSingleParticlePhysicsProxyBase::InvM() to match ::M().

diff --git a/Engine/Source/Runtime/Experimental/Chaos/Public/PhysicsProxy/SingleParticlePhysicsProxy.h b/Engine/Source/Runtime/Experimental/Chaos/Public/PhysicsProxy/SingleParticlePhysicsProxy.h
--- a/Engine/Source/Runtime/Experimental/Chaos/Public/PhysicsProxy/SingleParticlePhysicsProxy.h
+++ b/Engine/Source/Runtime/Experimental/Chaos/Public/PhysicsProxy/SingleParticlePhysicsProxy.h
@@ -724,7 +724,7 @@ public:
 
 	const FReal InvM() const
 	{
-		return Read([](auto* Particle)
+		return Read([](auto* Particle) -> FReal
 		{
 			if (auto Rigid = Particle->CastToRigidParticle())
 			{
