Chaos
- fix merging of islands when bodies are stationary
- improved natvis for clearer debugging of constraint lists, island graph, etc
- add validadity check (cvar enabled) for island edges/nodes

diff --git a/Engine/Source/Runtime/Experimental/Chaos/Private/Chaos/Island/IslandGraph.cpp b/Engine/Source/Runtime/Experimental/Chaos/Private/Chaos/Island/IslandGraph.cpp
--- a/Engine/Source/Runtime/Experimental/Chaos/Private/Chaos/Island/IslandGraph.cpp
+++ b/Engine/Source/Runtime/Experimental/Chaos/Private/Chaos/Island/IslandGraph.cpp
@@ -367,22 +367,17 @@ void FIslandGraph<NodeType, EdgeType, IslandType>::MergeIslands()
 	}
 
 	// We loop over all the islands and if they have children
-	// we merge them onto the parent one
+	// we recusrively merge them onto the parent one
 	for (int32 IslandIndex = 0, NumIslands = GraphIslands.GetMaxIndex(); IslandIndex < NumIslands; ++IslandIndex)
 	{
 		if (GraphIslands.IsValidIndex(IslandIndex))
 		{
 			FGraphIsland& GraphIsland = GraphIslands[IslandIndex];
-			// If the graph island is persistent we don't need to merge its children
-			// and update the parent index
-			if (!GraphIsland.bIsPersistent)
+			for (const int32& MergedIsland : GraphIsland.ChildrenIslands)
 			{
-				for (const int32& MergedIsland : GraphIsland.ChildrenIslands)
-				{
-					MergeIslands(IslandIndex, MergedIsland);
-				}
-				GraphIsland.ChildrenIslands.Reset();
+				MergeIslands(IslandIndex, MergedIsland);
 			}
+			GraphIsland.ChildrenIslands.Reset();
 		}
 	}
 
diff --git a/Engine/Source/Runtime/Experimental/Chaos/Private/Chaos/Island/IslandManager.cpp b/Engine/Source/Runtime/Experimental/Chaos/Private/Chaos/Island/IslandManager.cpp
--- a/Engine/Source/Runtime/Experimental/Chaos/Private/Chaos/Island/IslandManager.cpp
+++ b/Engine/Source/Runtime/Experimental/Chaos/Private/Chaos/Island/IslandManager.cpp
@@ -13,8 +13,8 @@
 #include "Misc/App.h"
 
 /** Cvar to enable/disable the island sleeping */
-bool ChaosSolverSleepEnabled = true;
-FAutoConsoleVariableRef CVarChaosSolverSleepEnabled(TEXT("p.Chaos.Solver.SleepEnabled"), ChaosSolverSleepEnabled, TEXT(""));
+bool bChaosSolverSleepEnabled = true;
+FAutoConsoleVariableRef CVarChaosSolverSleepEnabled(TEXT("p.Chaos.Solver.SleepEnabled"), bChaosSolverSleepEnabled, TEXT(""));
 
 /** Cvar to control the number of island groups used in the solver. The total number will be NumThreads * IslandGroupsMultiplier */
 Chaos::FRealSingle GChaosSolverIslandGroupsMultiplier = 1;
@@ -32,6 +32,9 @@ FAutoConsoleVariableRef CVarChaosSolverCollisionDefaultLinearSleepThreshold(TEXT
 Chaos::FRealSingle ChaosSolverCollisionDefaultAngularSleepThresholdCVar = 0.0087f;  //~1/2 unit mass degree
 FAutoConsoleVariableRef CVarChaosSolverCollisionDefaultAngularSleepThreshold(TEXT("p.ChaosSolverCollisionDefaultAngularSleepThreshold"), ChaosSolverCollisionDefaultAngularSleepThresholdCVar, TEXT("Default angular threshold for sleeping.[def:0.0087]"));
 
+bool bChaosSolverValidateGraph = false;
+FAutoConsoleVariableRef CVarChaosSolverValidateGraph(TEXT("p.Chaos.Solver.ValidateGraph"), bChaosSolverValidateGraph, TEXT(""));
+
 namespace Chaos
 {
 	extern int32 GSingleThreadedPhysics;
@@ -460,11 +463,11 @@ int32 FPBDIslandManager::AddConstraint(const uint32 ContainerId, FConstraintHand
 {
 	if (ConstraintHandle)
 	{
+		// Are the particles dynamic (including asleep)?
 		const bool bValidParticle0 = ConstrainedParticles[0] && IsDynamicParticle(ConstrainedParticles[0]);
 		const bool bValidParticle1 = ConstrainedParticles[1] && IsDynamicParticle(ConstrainedParticles[1]);
 		
 		// We are checking if one of the 2 particle is dynamic to add the constraint to the graph
-		// Will discard constraints in between 2 sleeping particles. Huge gain but potential side effect?
 		if(bValidParticle0 || bValidParticle1)
 		{
 			const int32 NodeIndex0 = AddParticle(ConstrainedParticles[0], INDEX_NONE, false);
@@ -749,6 +752,26 @@ void FPBDIslandManager::UpdateIslands(const TParticleView<FPBDRigidParticles>& P
 	
 	// Sync the graph islands with the solver islands objects
 	SyncIslands(Particles, NumContainers);
+
+	ValidateIslands();
+}
+
+void FPBDIslandManager::ValidateIslands() const
+{
+	if (!bChaosSolverValidateGraph)
+	{
+		return;
+	}
+
+	// Check that nodes joined by an edge are in the same island as the edge
+	for (const FGraphEdge& Edge : IslandGraph->GraphEdges)
+	{
+		const FGraphNode* Node0 = (Edge.FirstNode != INDEX_NONE) ? &IslandGraph->GraphNodes[Edge.FirstNode] : nullptr;
+		const FGraphNode* Node1 = (Edge.SecondNode != INDEX_NONE) ? &IslandGraph->GraphNodes[Edge.SecondNode] : nullptr;
+
+		ensure((Node0 == nullptr) || !Node0->bValidNode || (Node0->IslandIndex == Edge.IslandIndex));
+		ensure((Node1 == nullptr) || !Node1->bValidNode || (Node1->IslandIndex == Edge.IslandIndex));
+	}
 }
 
 bool FPBDIslandManager::SleepInactive(const int32 IslandIndex,
@@ -757,7 +780,7 @@ bool FPBDIslandManager::SleepInactive(const int32 IslandIndex,
 {
 	// Only the persistent islands could start sleeping
 	const int32 GraphIndex = GetGraphIndex(IslandIndex);
-	if (!ChaosSolverSleepEnabled || !IslandSolvers.IsValidIndex(GraphIndex) ||
+	if (!bChaosSolverSleepEnabled || !IslandSolvers.IsValidIndex(GraphIndex) ||
 		(IslandSolvers.IsValidIndex(GraphIndex) && !IslandSolvers[GraphIndex]->IsPersistent()) )
 	{
 		return false;
diff --git a/Engine/Source/Runtime/Experimental/Chaos/Public/Chaos/Island/IslandGraph.h b/Engine/Source/Runtime/Experimental/Chaos/Public/Chaos/Island/IslandGraph.h
--- a/Engine/Source/Runtime/Experimental/Chaos/Public/Chaos/Island/IslandGraph.h
+++ b/Engine/Source/Runtime/Experimental/Chaos/Public/Chaos/Island/IslandGraph.h
@@ -8,6 +8,110 @@
 
 namespace Chaos
 {
+	/** Graph node structure */
+	template<typename NodeType>
+	struct CHAOS_API TIslandGraphNode
+	{
+		/** List of edges connected to the node */
+		// @todo(chaos): should we really have a sparse array here? The implementation probably wastes a lot of memory
+		TSparseArray<int32> NodeEdges;
 
+		/** Node Island Index (for static/kinematic particles could belong to several islands) */
+		int32 IslandIndex = INDEX_NONE;
+
+		/** List of islands in which the node is referenced */
+		// @todo(chaos): should we really have a set here? The implementation probably wastes a lot of memory
+		TSet<int32> NodeIslands;
+
+		/** Check if a node is valid (checked for graph partitioning) */
+		bool bValidNode = true;
+
+		/** Check if a node is steady */
+		bool bStationaryNode = true;
+
+		/** Node counter to filter nodes already processed */
+		int32 NodeCounter = 0;
+
+		/** Node item that is stored per node */
+		NodeType NodeItem;
+
+		/** Node level index */
+		int32 LevelIndex = INDEX_NONE;
+
+		/** Set of used colors */
+		TSet<int32> ColorIndices;
+	};
+
+	/** Graph edge structure */
+	template<typename EdgeType>
+	struct CHAOS_API TIslandGraphEdge
+	{
+		/** First node of the edge */
+		int32 FirstNode = INDEX_NONE;
+
+		/** Second node of the edge*/
+		int32 SecondNode = INDEX_NONE;
+
+		/** Current edge index in the list of edges of the first node */
+		int32 FirstEdge = INDEX_NONE;
+
+		/** Current edge index in the list of edges of the second node  */
+		int32 SecondEdge = INDEX_NONE;
+
+		/** Unique edge island index */
+		int32 IslandIndex = INDEX_NONE;
+
+		/** Edge counter to filter edges already processed */
+		int32 EdgeCounter = 0;
+
+		/** Edge item that is stored per node */
+		EdgeType EdgeItem;
+
+		/** Item Container Id  */
+		int32 ItemContainer = 0;
+
+		/** Edge level index */
+		int32 LevelIndex = INDEX_NONE;
+
+		/** Edge Color index */
+		int32 ColorIndex = INDEX_NONE;
+	};
+
+	/** Graph island structure */
+	template<typename IslandType>
+	struct CHAOS_API TIslandGraphIsland
+	{
+		/** Number of edges per islands*/
+		int32 NumEdges = 0;
+
+		/** Number of valid nodes per islands (should be less than the ones in solver islands since it is only including the valid ones)*/
+		int32 NumNodes = 0;
+
+		/** Island counter to filter islands already processed */
+		int32 IslandCounter = 0;
+
+		/** Boolean to check if an island is persistent or not */
+		bool bIsPersistent = true;
+
+		/** Boolean to check if an island is sleeping or not */
+		bool bIsSleeping = true;
+
+		/** List of children islands  to be merged */
+		// @todo(chaos): should we really have a set here? The implementation probably wastes a lot of memory
+		TSet<int32> ChildrenIslands;
+
+		/** Parent Island */
+		int32 ParentIsland = INDEX_NONE;
+
+		/** Island Item that is stored per island*/
+		IslandType IslandItem;
+
+		/** Max levels per island */
+		int32 MaxLevels = INDEX_NONE;
+
+		/** Max color per island */
+		int32 MaxColors = INDEX_NONE;
+	};
+
 /**
  * Island Graph.
@@ -22,6 +126,11 @@
 class CHAOS_API FIslandGraph
 {
 public:
+	// NOTE: IslandGraph internal structs are non-member structs for better natvis behaviour/debugging. Natvis 
+	// has trouble with member classes of templated classes
+	using FGraphNode = TIslandGraphNode<NodeType>;
+	using FGraphEdge = TIslandGraphEdge<EdgeType>;
+	using FGraphIsland = TIslandGraphIsland<IslandType>;
 
 	static constexpr const int32 MaxCount = 100000;
 
@@ -123,105 +232,6 @@
 	 */
 	FORCEINLINE int32 NumIslands() const { return GraphIslands.GetMaxIndex(); }
 
-	/** Graph node structure */
-	struct FGraphNode 
-	{
-		/** List of edges connected to the node */
-		TSparseArray<int32> NodeEdges;
-
-		/** Node Island Index (for static/kinematic particles could belong to several islands) */
-		int32 IslandIndex = INDEX_NONE;
-		
-		/** List of islands in which the node is referenced */
-		TSet<int32> NodeIslands;
-
-		/** Check if a node is valid (checked for graph partitioning) */
-		bool bValidNode = true;
-
-		/** Check if a node is steady */
-		bool bStationaryNode = true;
-
-		/** Node counter to filter nodes already processed */
-		int32 NodeCounter = 0;
-
-		/** Node item that is stored per node */
-		NodeType NodeItem;
-		
-		/** Node level index */
-		int32 LevelIndex = INDEX_NONE;
-
-		/** Set of used colors */
-		TSet<int32> ColorIndices;
-	};
-
-	/** Graph edge structure */
-	struct FGraphEdge 
-	{
-		/** First node of the edge */
-		int32 FirstNode = INDEX_NONE;
-
-		/** Second node of the edge*/
-		int32 SecondNode = INDEX_NONE;
-
-		/** Current edge index in the list of edges of the first node */
-		int32 FirstEdge = INDEX_NONE;
-
-		/** Current edge index in the list of edges of the second node  */
-		int32 SecondEdge = INDEX_NONE;
-
-		/** Unique edge island index */
-		int32 IslandIndex = INDEX_NONE;
-
-		/** Edge counter to filter edges already processed */
-		int32 EdgeCounter = 0;
-
-		/** Edge item that is stored per node */
-		EdgeType EdgeItem;
-
-		/** Item Container Id  */
-		int32 ItemContainer = 0;
-
-		/** Edge level index */
-		int32 LevelIndex = INDEX_NONE;
-
-		/** Edge Color index */
-		int32 ColorIndex = INDEX_NONE;
-	};
-
-	/** Graph island structure */
-	struct FGraphIsland
-	{
-		/** Number of edges per islands*/
-		int32 NumEdges = 0;
-
-		/** Number of valid nodes per islands (should be less than the ones in solver islands since it is only including the valid ones)*/
-		int32 NumNodes = 0;
-
-		/** Island counter to filter islands already processed */
-		int32 IslandCounter = 0;
-		
-		/** Boolean to check if an island is persistent or not */
-		bool bIsPersistent = true;
-
-		/** Boolean to check if an island is sleeping or not */
-		bool bIsSleeping = true;
-
-		/** List of children islands  to be merged */
-		TSet<int32> ChildrenIslands;
-
-		/** Parent Island */
-		int32 ParentIsland = INDEX_NONE;
-		
-		/** Island Item that is stored per island*/
-		IslandType IslandItem;
-
-		/** Max levels per island */
-		int32 MaxLevels = INDEX_NONE;
-
-		/** Max color per island */
-		int32 MaxColors = INDEX_NONE;
-	};
-
 	/**
 	 * Link two islands to each other for future island traversal
 	 * @param FirstIsland First island index to be linked
@@ -319,10 +329,9 @@
 		if(GraphEdges.IsValidIndex(EdgeIndex))
 		{
 			const FGraphEdge& GraphEdge = GraphEdges[EdgeIndex];
-			if(GraphNodes.IsValidIndex(GraphEdge.FirstNode) && GraphNodes.IsValidIndex(GraphEdge.SecondNode))
-			{
-				return !GraphNodes[GraphEdge.FirstNode].bStationaryNode || !GraphNodes[GraphEdge.SecondNode].bStationaryNode;
-			}
+			const bool bFirstNodeMoving = GraphNodes.IsValidIndex(GraphEdge.FirstNode) && !GraphNodes[GraphEdge.FirstNode].bStationaryNode;
+			const bool bSecondNodeMoving = GraphNodes.IsValidIndex(GraphEdge.SecondNode) && !GraphNodes[GraphEdge.SecondNode].bStationaryNode;
+			return bFirstNodeMoving || bSecondNodeMoving;
 		}
 		return false;
 	}
diff --git a/Engine/Source/Runtime/Experimental/Chaos/Public/Chaos/Island/IslandManager.h b/Engine/Source/Runtime/Experimental/Chaos/Public/Chaos/Island/IslandManager.h
--- a/Engine/Source/Runtime/Experimental/Chaos/Public/Chaos/Island/IslandManager.h
+++ b/Engine/Source/Runtime/Experimental/Chaos/Public/Chaos/Island/IslandManager.h
@@ -316,6 +316,8 @@ protected:
 	*/
 	void SyncIslands(FPBDRigidsSOAs& Particles, const int32 NumContainers);
 
+	void ValidateIslands() const;
+
 	/** List of solver islands in the manager */
 	TSparseArray<TUniquePtr<FPBDIslandSolver>> IslandSolvers;
 
