Chaos
- option to remove sub-surface contacts ( experimental - disabled by default)
- option to ignore shallow CCD contacts (experimental - disabled by default)
- ability to disable individual manifold points in a collision constraint
- debug draw fixes

diff --git a/Engine/Source/Runtime/Experimental/Chaos/Private/Chaos/CCDUtilities.cpp b/Engine/Source/Runtime/Experimental/Chaos/Private/Chaos/CCDUtilities.cpp
--- a/Engine/Source/Runtime/Experimental/Chaos/Private/Chaos/CCDUtilities.cpp
+++ b/Engine/Source/Runtime/Experimental/Chaos/Private/Chaos/CCDUtilities.cpp
@@ -6,6 +6,8 @@
 #include "Chaos/PBDRigidsEvolutionGBF.h"
 #include "Chaos/CollisionResolution.h"
 
+// @todo(chaos): these cvars should all be settings per particle
+
 bool bChaosCollisionCCDEnableResweep = true;
 FAutoConsoleVariableRef CVarChaosCollisionCCDEnableResweep(TEXT("p.Chaos.Collision.CCD.EnableResweep"), bChaosCollisionCCDEnableResweep, TEXT("Enable resweep for CCD. Resweeping allows CCD to catch more secondary collisions but also is more costly. Default is true."));
 
@@ -15,6 +17,8 @@ FAutoConsoleVariableRef CVarChaosCollisionCCDAllowClipping(TEXT("p.Chaos.Collisi
 int32 ChaosCollisionCCDConstraintMaxProcessCount = 1;
 FAutoConsoleVariableRef CVarChaosCollisionCCDConstraintMaxProcessCount(TEXT("p.Chaos.Collision.CCD.ConstraintMaxProcessCount"), ChaosCollisionCCDConstraintMaxProcessCount, TEXT("The max number of times each constraint can be resolved when applying CCD constraints. Default is 2. The larger this number is, the more fully CCD constraints are resolved."));
 
+bool bChaosCollisionCCDEnablePhiCheck = false;
+FAutoConsoleVariableRef CVarChaosCollisionCCDEnablePhiCheck(TEXT("p.Chaos.Collision.CCD.EnablePhiCheck"), bChaosCollisionCCDEnablePhiCheck, TEXT("Enable phi check for CCD. If disabled we skip CCD processing of contacts if the penetration depth is below the CCD threshold at TOI=1."));
 
 namespace Chaos
 {
@@ -146,7 +150,7 @@ namespace Chaos
 			const FReal CCDThreshold0 = GetParticleCCDThreshold(Constraint->GetImplicit0());
 			const FReal CCDThreshold1 = GetParticleCCDThreshold(Constraint->GetImplicit1());
 			const FReal CCDConstraintThreshold = CCDThreshold0 + CCDThreshold1;
-			if ((Displacements[1] - Displacements[0]).SizeSquared() > CCDConstraintThreshold * CCDConstraintThreshold)
+			if ((Displacements[1] - Displacements[0]).SizeSquared() > FMath::Square(CCDConstraintThreshold))
 			{
 				bNeedCCDSolve = true;
 			}
@@ -154,7 +158,8 @@ namespace Chaos
 			// make sure we ignore pairs that don't include any dynamics
 			if (CCDParticlePair[0] != nullptr || CCDParticlePair[1] != nullptr)
 			{
-				CCDConstraints.Add(FCCDConstraint(Constraint, CCDParticlePair, Displacements));
+				const FReal PhiThreshold = -CCDConstraintThreshold;
+				CCDConstraints.Add(FCCDConstraint(Constraint, CCDParticlePair, Displacements, PhiThreshold));
 				for (int32 i = 0; i < 2; i++)
 				{
 					if (CCDParticlePair[i] != nullptr)
@@ -288,38 +293,59 @@ namespace Chaos
 		while (ConstraintIndex < ConstraintEnd) 
 		{
 			FCCDConstraint *CCDConstraint = SortedCCDConstraints[ConstraintIndex];
+			FCCDParticle* CCDParticle0 = CCDConstraint->Particle[0];
+			FCCDParticle* CCDParticle1 = CCDConstraint->Particle[1];
+
 			IslandTOI = CCDConstraint->SweptConstraint->TimeOfImpact;
 
-			// Constraints whose TOIs are in the range of [0, 1) are resolved for this frame. TOI = 1 means that the two particles just start touching at the end of the frame and therefore cannot have tunneling this frame. So this TOI = 1 can be left to normal collisions or CCD in next frame.
+			// Constraints whose TOIs are in the range of [0, 1) are resolved for this frame. TOI = 1 means that the two 
+			// particles just start touching at the end of the frame and therefore cannot have tunneling this frame. 
+			// So this TOI = 1 can be left to normal collisions or CCD in next frame.
 			if (IslandTOI > 1) 
 			{
 				break;
 			}
 
-			FCCDParticle **CCDConstraintParticles = CCDConstraint->Particle;
-
 			// If both particles are marked Done (due to clipping), continue
-			if (bChaosCollisionCCDAllowClipping && (!CCDConstraintParticles[0] || CCDConstraintParticles[0]->Done) && (!CCDConstraintParticles[1] || CCDConstraintParticles[1]->Done))
+			if (bChaosCollisionCCDAllowClipping && (!CCDParticle0 || CCDParticle0->Done) && (!CCDParticle1 || CCDParticle1->Done))
 			{
-				ConstraintIndex ++;
+				ConstraintIndex++;
 				continue;
 			}
 
+			// Skip CCD resolution if the depth at TOI=1 is low and let the standard contact resolution handle it.
+			// This can be important because CCD contact resolution uses point-mass physics and does not apply rotation 
+			// so we must leave some work for the regular solver to do to get decent behaviour when sliding fast along
+			// a surface for example.
+			// @todo(chaos): enabling this can lead to tunnelling problems when you have a jumble of CCD objects.
+			// Possibly because if we need to re-examin this CCD contact because of an impulse from another CCD contact,
+			// we may be initially penetrating. Investigate and aee if we can come up with a more robust solution. 
+			if (bChaosCollisionCCDEnablePhiCheck && (CCDConstraint->SweptConstraint->GetPhi() > CCDConstraint->PhiThreshold))
+			{
+				ConstraintIndex++;
+				continue;
+			}
+			
 			ensure(CCDConstraint->ProcessedCount < ChaosCollisionCCDConstraintMaxProcessCount);
 
-			// In UpdateConstraintFromGeometrySwept, InitManifoldPoint requires P, Q to be at TOI=1., but the input of UpdateConstraintFromGeometrySwept requires transforms at current TOI. So instead of rewinding P, Q, we advance X, R to current TOI and keep P, Q at TOI=1.
-			if (CCDConstraintParticles[0] && !CCDConstraintParticles[0]->Done)
+			// In UpdateConstraintFromGeometrySwept, InitManifoldPoint requires P, Q to be at TOI=1., but the input of 
+			// UpdateConstraintFromGeometrySwept requires transforms at current TOI. So instead of rewinding P, Q, we 
+			// advance X, R to current TOI and keep P, Q at TOI=1.
+			if (CCDParticle0 && !CCDParticle0->Done)
 			{
-				AdvanceParticleXToTOI(CCDConstraintParticles[0], IslandTOI, Dt);
+				AdvanceParticleXToTOI(CCDParticle0, IslandTOI, Dt);
 			}
-			if (CCDConstraintParticles[1] && !CCDConstraintParticles[1]->Done)
+			if (CCDParticle1 && !CCDParticle1->Done)
 			{
-				AdvanceParticleXToTOI(CCDConstraintParticles[1], IslandTOI, Dt);
+				AdvanceParticleXToTOI(CCDParticle1, IslandTOI, Dt);
 			}
 
 			ApplyImpulse(CCDConstraint);
-			CCDConstraint->ProcessedCount ++;
-			// After applying impulse, constraint TOI need be updated to reflect the new velocities. Usually the new velocities are separating, and therefore TOI should be infinity.
+			CCDConstraint->ProcessedCount++;
+
+			// After applying impulse, constraint TOI need be updated to reflect the new velocities. 
+			// Usually the new velocities are separating, and therefore TOI should be infinity.
+			// See resweep below which (optionally) updates TOI for all other contacts as a result of handling this one
 			CCDConstraint->SweptConstraint->TimeOfImpact = TNumericLimits<FReal>::Max();
 
 			if (CCDConstraint->ProcessedCount >= ChaosCollisionCCDConstraintMaxProcessCount)
@@ -335,47 +361,49 @@ namespace Chaos
 				*/ 
 				if (bChaosCollisionCCDAllowClipping)
 				{
-					if (CCDConstraintParticles[0])
+					if (CCDParticle0)
 					{
 						if (CCDConstraint->FastMovingKinematicIndex != INDEX_NONE)
 						{
+							// @todo(chaos): can this just get the particle from the CCD constraint?
 							const FConstGenericParticleHandle KinematicParticle = FGenericParticleHandle(CCDConstraint->SweptConstraint->GetParticle(CCDConstraint->FastMovingKinematicIndex));
 							const FVec3 Normal = CCDConstraint->SweptConstraint->CalculateWorldContactNormal();
 							const FVec3 Offset = FVec3::DotProduct(KinematicParticle->V() * ((1.f - IslandTOI) * Dt), Normal) * Normal;
-							ClipParticleP(CCDConstraintParticles[0], Offset);
+							ClipParticleP(CCDParticle0, Offset);
 						}
 						else
 						{
-							ClipParticleP(CCDConstraintParticles[0]);
+							ClipParticleP(CCDParticle0);
 						}
-						CCDConstraintParticles[0]->Done = true;
+						CCDParticle0->Done = true;
 					}
-					if (CCDConstraintParticles[1])
+					if (CCDParticle1)
 					{
 						if (CCDConstraint->FastMovingKinematicIndex != INDEX_NONE)
 						{
+							// @todo(chaos): can this just get the particle from the CCD constraint?
 							const FConstGenericParticleHandle KinematicParticle = FGenericParticleHandle(CCDConstraint->SweptConstraint->GetParticle(CCDConstraint->FastMovingKinematicIndex));
 							const FVec3 Normal = CCDConstraint->SweptConstraint->CalculateWorldContactNormal();
 							const FVec3 Offset = FVec3::DotProduct(KinematicParticle->V() * ((1.f - IslandTOI) * Dt), Normal) * Normal;
-							ClipParticleP(CCDConstraintParticles[1], Offset);
+							ClipParticleP(CCDParticle1, Offset);
 						}
 						else
 						{
-							ClipParticleP(CCDConstraintParticles[1]);
+							ClipParticleP(CCDParticle1);
 						}
-						CCDConstraintParticles[1]->Done = true;
+						CCDParticle1->Done = true;
 					}
 				}
 				// If clipping is not allowed, we update particle P (at TOI=1) based on new velocities. 
 				else
 				{
-					if (CCDConstraintParticles[0])
+					if (CCDParticle0)
 					{
-						UpdateParticleP(CCDConstraint->Particle[0], Dt);
+						UpdateParticleP(CCDParticle0, Dt);
 					}
-					if (CCDConstraintParticles[1])
+					if (CCDParticle1)
 					{
-						UpdateParticleP(CCDConstraint->Particle[1], Dt);
+						UpdateParticleP(CCDParticle1, Dt);
 					}
 				}
 				// Increment ConstraintIndex if we run out of computational budget for this constraint.
@@ -384,13 +412,13 @@ namespace Chaos
 			// If we still have computational budget for this constraint, update particle P and don't clip.
 			else
 			{
-				if (CCDConstraintParticles[0] && !CCDConstraintParticles[0]->Done)
+				if (CCDParticle0 && !CCDParticle0->Done)
 				{
-					UpdateParticleP(CCDConstraint->Particle[0], Dt);
+					UpdateParticleP(CCDParticle0, Dt);
 				}
-				if (CCDConstraintParticles[1] && !CCDConstraintParticles[1]->Done)
+				if (CCDParticle1 && !CCDParticle1->Done)
 				{
-					UpdateParticleP(CCDConstraint->Particle[1], Dt);
+					UpdateParticleP(CCDParticle1, Dt);
 				}
 			}
 
@@ -401,7 +429,7 @@ namespace Chaos
 				bool HasResweptConstraint = false;
 				for (int32 i = 0; i < 2; i++)
 				{
-					FCCDParticle *CCDParticle = CCDConstraintParticles[i];
+					FCCDParticle *CCDParticle = CCDConstraint->Particle[i];
 					if (CCDParticle != nullptr)
 					{
 						for (int32 AttachedCCDConstraintIndex = 0; AttachedCCDConstraintIndex < CCDParticle->AttachedCCDConstraints.Num(); AttachedCCDConstraintIndex++)
@@ -542,6 +570,11 @@ namespace Chaos
 		const FRigidTransform3& ShapeWorldTransform1 = Constraint->GetShapeWorldTransform1();
 		for(const FManifoldPoint &ManifoldPoint : Constraint->GetManifoldPoints())
 		{
+			if (ManifoldPoint.Flags.bDisabled)
+			{
+				continue;
+			}
+			
 			const FVec3 Normal = ShapeWorldTransform1.TransformVectorNoScale(ManifoldPoint.ContactPoint.ShapeContactNormal);
 			const FVec3 V0 = Rigid0 != nullptr ? Rigid0->V() : FVec3(0.f);
 			const FVec3 V1 = Rigid1 != nullptr ? Rigid1->V() : FVec3(0.f);
diff --git a/Engine/Source/Runtime/Experimental/Chaos/Private/Chaos/ChaosDebugDraw.cpp b/Engine/Source/Runtime/Experimental/Chaos/Private/Chaos/ChaosDebugDraw.cpp
--- a/Engine/Source/Runtime/Experimental/Chaos/Private/Chaos/ChaosDebugDraw.cpp
+++ b/Engine/Source/Runtime/Experimental/Chaos/Private/Chaos/ChaosDebugDraw.cpp
@@ -484,15 +484,13 @@ namespace Chaos
 
 		void DrawShapesImpl(const FGeometryParticleHandle* Particle, const FRigidTransform3& ShapeTransform, const FImplicitObject* Implicit, const FShapeOrShapesArray& Shapes, const FReal Margin, const FColor& Color, const FChaosDebugDrawSettings& Settings)
 		{
-
 			if (!Particle || !Implicit) return;
 
 			const EImplicitObjectType PackedType = Implicit->GetType(); // Type includes scaling and instancing data
 			const EImplicitObjectType InnerType = GetInnerType(Implicit->GetType());
 
 			// Are we within the region of interest?
-			const FReal ParticleSize = Particle->HasBounds() ? 0.5f * Particle->LocalBounds().Extents().Size() : TNumericLimits<FReal>::Max();
-			if (!FDebugDrawQueue::GetInstance().IsInRegionOfInterest(ShapeTransform.GetLocation(), ParticleSize))
+			if (Particle->HasBounds() && !FDebugDrawQueue::GetInstance().IsInRegionOfInterest(Particle->WorldSpaceInflatedBounds()))
 			{
 				return;
 			}
@@ -834,13 +832,11 @@ namespace Chaos
 				const FRigidTransform3 WorldActorTransform1 = FParticleUtilities::GetActorWorldTransform(Particle1);
 
 				// Are we within the region of interest?
-				const FReal Particle0Size = Contact.GetParticle0()->HasBounds() ? 0.5f * Contact.GetParticle0()->LocalBounds().Extents().Size() : TNumericLimits<FReal>::Max();
-				const FReal Particle1Size = Contact.GetParticle1()->HasBounds() ? 0.5f * Contact.GetParticle1()->LocalBounds().Extents().Size() : TNumericLimits<FReal>::Max();
-				if (!FDebugDrawQueue::GetInstance().IsInRegionOfInterest(SpaceTransform.TransformPositionNoScale(WorldActorTransform0.GetLocation()), Particle0Size))
+				if (Contact.GetParticle0()->HasBounds() && !FDebugDrawQueue::GetInstance().IsInRegionOfInterest(Contact.GetParticle0()->WorldSpaceInflatedBounds().TransformedAABB(SpaceTransform)))
 				{
 					return;
 				}
-				if (!FDebugDrawQueue::GetInstance().IsInRegionOfInterest(SpaceTransform.TransformPositionNoScale(WorldActorTransform1.GetLocation()), Particle1Size))
+				if (Contact.GetParticle1()->HasBounds() && !FDebugDrawQueue::GetInstance().IsInRegionOfInterest(Contact.GetParticle1()->WorldSpaceInflatedBounds().TransformedAABB(SpaceTransform)))
 				{
 					return;
 				}
@@ -853,6 +849,8 @@ namespace Chaos
 						continue;
 					}
 
+					const bool bPruned = ManifoldPoint.Flags.bDisabled;
+
 					const int32 ContactPlaneOwner = 1;
 					const int32 ContactPointOwner = 1 - ContactPlaneOwner;
 					const FRigidTransform3& PlaneTransform = (ContactPlaneOwner == 0) ? Contact.GetShapeRelativeTransform0() * WorldActorTransform0 : Contact.GetShapeRelativeTransform1() * WorldActorTransform1;
@@ -870,8 +868,7 @@ namespace Chaos
 					FColor PlaneNormalColor = FColor(200, 0, 0);
 					FColor EdgeNormalColor = FColor(200, 150, 0);
 					FColor ImpulseColor = FColor(0, 0, 200);
-					FColor PushOutColor = FColor(200, 200, 0);
-					FColor PushOutImpusleColor = FColor(0, 200, 200);
+					FColor PushOutColor = FColor(0, 200, 200);
 					if (ManifoldPoint.Flags.bInsideStaticFrictionCone)
 					{
 						DiscColor = FColor(150, 200, 0);
@@ -882,6 +879,11 @@ namespace Chaos
 						PlaneNormalColor = FColor(100, 0, 0);
 						EdgeNormalColor = FColor(100, 80, 0);
 					}
+					if (bPruned)
+					{
+						PlaneNormalColor = FColor(200, 0, 200);
+						EdgeNormalColor = FColor(200, 0, 200);
+					}
 
 					const FVec3 WorldPointLocation = SpaceTransform.TransformPosition(PointLocation);
 					const FVec3 WorldPlaneLocation = SpaceTransform.TransformPosition(PlaneLocation);
@@ -892,7 +894,7 @@ namespace Chaos
 					// Pushout
 					if ((Settings.PushOutScale > 0) && !ManifoldPoint.NetPushOut.IsNearlyZero())
 					{
-						FColor Color = (ColorScale * PushOutImpusleColor).ToFColor(false);
+						FColor Color = (ColorScale * PushOutColor).ToFColor(false);
 						FDebugDrawQueue::GetInstance().DrawDebugLine(WorldPointPlaneLocation, WorldPointPlaneLocation + Settings.DrawScale * Settings.PushOutScale * SpaceTransform.TransformVectorNoScale(ManifoldPoint.NetPushOut), Color, false, KINDA_SMALL_NUMBER, Settings.DrawPriority, Settings.LineThickness);
 					}
 					if ((Settings.ImpulseScale > 0) && !ManifoldPoint.NetImpulse.IsNearlyZero())
@@ -974,6 +976,7 @@ namespace Chaos
 		void DrawCollidingShapesImpl(const FRigidTransform3& SpaceTransform, const FPBDCollisionConstraints& Collisions, FRealSingle ColorScale, const FChaosDebugDrawSettings& Settings)
 		{
 			TArray<const FImplicitObject*> Implicits;
+			TArray<const FPerShapeData*> Shapes;
 			TArray<FConstGenericParticleHandle> ShapeParticles;
 			TArray<FRigidTransform3> ShapeTransforms;
 
@@ -992,10 +995,13 @@ namespace Chaos
 						const FRigidTransform3 WorldActorTransform1 = FParticleUtilities::GetActorWorldTransform(Particle1);
 						const FRigidTransform3 ShapeWorldTransform0 = PointConstraint->GetShapeRelativeTransform0() * WorldActorTransform0;
 						const FRigidTransform3 ShapeWorldTransform1 = PointConstraint->GetShapeRelativeTransform1() * WorldActorTransform1;
+						const FPerShapeData* Shape0 = PointConstraint->GetShape0();
+						const FPerShapeData* Shape1 = PointConstraint->GetShape1();
 
 						if (!Implicits.Contains(Implicit0))
 						{
 							Implicits.Add(Implicit0);
+							Shapes.Add(Shape0);
 							ShapeParticles.Add(Particle0);
 							ShapeTransforms.Add(ShapeWorldTransform0);
 						}
@@ -1003,6 +1009,7 @@ namespace Chaos
 						if (!Implicits.Contains(Implicit1))
 						{
 							Implicits.Add(Implicit1);
+							Shapes.Add(Shape1);
 							ShapeParticles.Add(Particle1);
 							ShapeTransforms.Add(ShapeWorldTransform1);
 						}
@@ -1016,7 +1023,7 @@ namespace Chaos
 					ShapeParticles[ShapeIndex]->Handle(), 
 					ShapeTransforms[ShapeIndex], 
 					Implicits[ShapeIndex],
-					FShapeOrShapesArray(ShapeParticles[ShapeIndex]->Handle()),
+					Shapes[ShapeIndex],
 					0.0f, 
 					ShapeParticles[ShapeIndex]->IsDynamic() ? FColor::Yellow : FColor::Red, 
 					Settings);
diff --git a/Engine/Source/Runtime/Experimental/Chaos/Private/Chaos/Collision/CollisionPruning.cpp b/Engine/Source/Runtime/Experimental/Chaos/Private/Chaos/Collision/CollisionPruning.cpp
new file mode 100644
--- /dev/null
+++ b/Engine/Source/Runtime/Experimental/Chaos/Private/Chaos/Collision/CollisionPruning.cpp
@@ -0,0 +1,290 @@
+// Copyright Epic Games, Inc. All Rights Reserved.
+#include "Chaos/Collision/CollisionPruning.h"
+
+#include "Chaos/Collision/CollisionConstraintAllocator.h"
+#include "Chaos/Collision/ParticleCollisions.h"
+#include "Chaos/Collision/PBDCollisionConstraint.h"
+#include "Chaos/ParticleHandle.h"
+
+namespace Chaos
+{
+	extern FRealSingle Chaos_Collision_EdgePrunePlaneDistance;
+
+
+	// Loop over edge collisions, then all plane collisions and remove the edge collision if it is hidden by a plane collision
+	// NOTE: We only look at plane collisions where the other shape owns the plane.
+	// @todo(chaos): this should probably only disable individual manifold points
+	// @todo(chaos): we should probably only reject edges if the plane contact is also close to the edge contact
+	// @todo(chaos): we should also try to eliminate face contacts from sub-surface faces
+	// @todo(chaos): perf issue: this processes contacts in world space, but we don't calculated that data until Gather. Fix this.
+	void FParticleEdgeCollisionPruner::Prune()
+	{
+		// Chaos_Collision_EdgePrunePlaneDistance is prune distance for unit sized object (100cm) so scale to the size of this object
+		const FReal EdgePlaneTolerancePerCm = FReal(Chaos_Collision_EdgePrunePlaneDistance) * FReal(0.01);
+		const FReal EdgePlaneToleranceSize = Particle->HasBounds() ? Particle->LocalBounds().Extents().Min() : FReal(1);
+		const FReal EdgePlaneTolerance = EdgePlaneToleranceSize * EdgePlaneTolerancePerCm;
+
+		FParticleCollisions& ParticleCollisions = Particle->ParticleCollisions();
+
+		ParticleCollisions.VisitCollisions(
+			[this, &ParticleCollisions, EdgePlaneTolerance](FPBDCollisionConstraint& EdgeCollision)
+			{
+				if (EdgeCollision.IsEnabled())
+				{
+					const int32 EdgeOtherShapeIndex = (EdgeCollision.GetParticle0() == Particle) ? 1 : 0;
+					const EContactPointType VertexContactType = (EdgeOtherShapeIndex == 0) ? EContactPointType::VertexPlane : EContactPointType::PlaneVertex;
+
+					for (const FManifoldPoint& EdgeManifoldPoint : EdgeCollision.GetManifoldPoints())
+					{
+						const bool bIsEdgeContact = (EdgeManifoldPoint.ContactPoint.ContactType == EContactPointType::EdgeEdge);
+
+						if (bIsEdgeContact)
+						{
+							const FRigidTransform3& EdgeTransform = (EdgeOtherShapeIndex == 0) ? EdgeCollision.GetShapeWorldTransform0() : EdgeCollision.GetShapeWorldTransform1();
+							const FVec3 EdgePos = EdgeTransform.TransformPositionNoScale(EdgeManifoldPoint.ContactPoint.ShapeContactPoints[EdgeOtherShapeIndex]);
+
+							// Loop over plane collisions
+							ECollisionVisitorResult PlaneResult = ParticleCollisions.VisitConstCollisions(
+								[this, &EdgeCollision, &EdgePos, EdgePlaneTolerance](const FPBDCollisionConstraint& PlaneCollision)
+								{
+									if ((&PlaneCollision != &EdgeCollision) && PlaneCollision.IsEnabled())
+									{
+										const int32 PlaneOtherShapeIndex = (PlaneCollision.GetParticle0() == Particle) ? 1 : 0;
+										const EContactPointType PlaneContactType = (PlaneOtherShapeIndex == 0) ? EContactPointType::PlaneVertex : EContactPointType::VertexPlane;
+										const FRigidTransform3& PlaneTransform = (PlaneOtherShapeIndex == 0) ? PlaneCollision.GetShapeWorldTransform0() : PlaneCollision.GetShapeWorldTransform1();
+
+										for (const FManifoldPoint& PlaneManifoldPoint : PlaneCollision.GetManifoldPoints())
+										{
+											if (PlaneManifoldPoint.ContactPoint.ContactType == PlaneContactType)
+											{
+												// If the edge position is in the plane, disable it
+												const FVec3 PlanePos = PlaneTransform.TransformPositionNoScale(PlaneManifoldPoint.ContactPoint.ShapeContactPoints[PlaneOtherShapeIndex]);
+												const FVec3 PlaneNormal = PlaneCollision.GetShapeWorldTransform1().TransformVectorNoScale(PlaneManifoldPoint.ContactPoint.ShapeContactNormal);
+
+												const FVec3 EdgePlaneDelta = EdgePos - PlanePos;
+												const FReal EdgePlaneDistance = FVec3::DotProduct(EdgePlaneDelta, PlaneNormal);
+												if (FMath::Abs(EdgePlaneDistance) < EdgePlaneTolerance)
+												{
+													// The edge contact is hidden by a plane contact so disable it and stop the inner loop
+													EdgeCollision.SetDisabled(true);
+													return ECollisionVisitorResult::Stop;
+												}
+											}
+										}
+									}
+									return ECollisionVisitorResult::Continue;
+								});
+
+							// If we disabled this constraint, move to the next one and ignore remaining manifold points
+							if (PlaneResult == ECollisionVisitorResult::Stop)
+							{
+								break;
+							}
+						}
+					}
+				}
+				return ECollisionVisitorResult::Continue;
+			});
+	}
+
+	/**
+	 * @brief Information about a contact point held in a ContactGroup (points in the same plane)
+	*/
+	class FContactGroupManifoldPoint
+	{
+	public:
+		FPBDCollisionConstraint* Collision;
+		int32 ManifoldPointIndex;
+		FVec3 ContactNormal;
+		FVec3 ContactPosition;
+	};
+
+	/**
+	 * @brief A set of contacts in the same plane
+	*/
+	class FContactGroup
+	{
+	public:
+		FContactGroup()
+			: PlanePosition()
+			, PlaneNormal()
+			, UpDotNormal(0)
+			, bActive(false)
+		{
+		}
+
+		void Add(FPBDCollisionConstraint& Collision, const int32 ManifoldPointIndex, const FVec3& InContactPosition, const FVec3& InContactNormal, const FReal InUpDotNormal)
+		{
+			if (ManifoldPoints.Num() == 0)
+			{
+				PlanePosition = InContactPosition;
+				PlaneNormal = InContactNormal;
+				UpDotNormal = InUpDotNormal;
+				bActive = true;
+			}
+
+			FContactGroupManifoldPoint& Point = ManifoldPoints[ManifoldPoints.AddDefaulted()];
+			Point.Collision = &Collision;
+			Point.ManifoldPointIndex = ManifoldPointIndex;
+			Point.ContactPosition = InContactPosition;
+			Point.ContactNormal = InContactNormal;
+		}
+
+		TArray<FContactGroupManifoldPoint> ManifoldPoints;
+		FVec3 PlanePosition;
+		FVec3 PlaneNormal;
+		FReal UpDotNormal;
+		bool bActive;
+	};
+
+	/**
+	 * @brief All the contacts on a body, held in sets where all contacts in each set are in the same plane
+	*/
+	class FContactGroupContainer
+	{
+	public:
+		FContactGroupContainer(const FVec3& InUpVector)
+			: UpVector(InUpVector)
+		{
+		}
+		
+		void Add(FPBDCollisionConstraint& Collision, const int32 ManifoldPointIndex, const FVec3& ContactPos, const FVec3& ContactNormal)
+		{
+			FContactGroup* ContactGroup = FindContactGroup(ContactPos, ContactNormal);
+			if (ContactGroup == nullptr)
+			{
+				ContactGroup = &ContactGroups[ContactGroups.AddDefaulted()];
+			}
+
+			const FReal UpDotNormal = FVec3::DotProduct(UpVector, ContactNormal);
+			ContactGroup->Add(Collision, ManifoldPointIndex, ContactPos, ContactNormal, UpDotNormal);
+		}
+
+		FContactGroup* FindContactGroup(const FVec3& ContactPos, const FVec3& ContactNormal)
+		{
+			const FReal NormalProjectionThreshold = FReal(0.99);	// 8deg
+			const FReal PlaneDistanceThreshold = FReal(5);
+
+			for (FContactGroup& ContactGroup : ContactGroups)
+			{
+				const FReal NormalProjection = FVec3::DotProduct(ContactNormal, ContactGroup.PlaneNormal);
+				if (NormalProjection > NormalProjectionThreshold)
+				{
+					const FReal PlaneDistance = FVec3::DotProduct(ContactPos - ContactGroup.PlanePosition, ContactGroup.PlaneNormal);
+					if (FMath::Abs(PlaneDistance) < PlaneDistanceThreshold)
+					{
+						return &ContactGroup;
+					}
+				}
+			}
+			return nullptr;
+		}
+
+		void SortContactGroups()
+		{
+			ContactGroups.Sort([](const FContactGroup& L, const FContactGroup& R)
+			{
+				return L.UpDotNormal > R.UpDotNormal;
+			});
+		}
+
+		TArray<FContactGroup, TInlineAllocator<4>> ContactGroups;
+		FVec3 UpVector;
+	};
+
+
+	void FParticleSubSurfaceCollisionPruner::Prune(const FVec3& UpVector)
+	{
+		// Chaos_Collision_EdgePrunePlaneDistance is prune distance for unit sized object (100cm) so scale by size of this object
+		const FReal PlaneDistanceTolerancePerCm = FReal(Chaos_Collision_EdgePrunePlaneDistance) * FReal(0.01);
+		const FReal PlaneDistanceSize = Particle->HasBounds() ? Particle->LocalBounds().Extents().Min() : FReal(1);
+		const FReal PlaneDistanceTolerance = PlaneDistanceSize * PlaneDistanceTolerancePerCm;
+		const FReal PlaneNormalTolerance = FReal(0.95);
+
+		FContactGroupContainer ContactGroupContainer(UpVector);
+
+		FParticleCollisions& ParticleCollisions = Particle->ParticleCollisions();
+
+		ParticleCollisions.VisitCollisions(
+			[this, &ContactGroupContainer](FPBDCollisionConstraint& Collision)
+			{
+				if (Collision.IsEnabled())
+				{
+					const int32 OtherShapeIndex = (Collision.GetParticle0() == Particle) ? 1 : 0;
+					const FRigidTransform3& ContactTransform = (OtherShapeIndex == 0) ? Collision.GetShapeWorldTransform0() : Collision.GetShapeWorldTransform1();
+					FConstGenericParticleHandle P0 = Collision.GetParticle0();
+					FConstGenericParticleHandle P1 = Collision.GetParticle1();
+					FConstGenericParticleHandle OtherParticle = (OtherShapeIndex == 0) ? Collision.GetParticle0() : Collision.GetParticle1();
+
+					// Only reject sub-surface contacts against static or kinematic objects
+					if (!OtherParticle->IsDynamic())
+					{
+						for (int32 ManifoldPointIndex = 0; ManifoldPointIndex < Collision.NumManifoldPoints(); ++ManifoldPointIndex)
+						{
+							const FManifoldPoint& ManifoldPoint = Collision.GetManifoldPoint(ManifoldPointIndex);
+
+							const FVec3 ContactPos = ContactTransform.TransformPositionNoScale(ManifoldPoint.ContactPoint.ShapeContactPoints[OtherShapeIndex]);
+							const FVec3 ContactNormal = Collision.GetShapeWorldTransform1().TransformVectorNoScale(ManifoldPoint.ContactPoint.ShapeContactNormal);
+
+							ContactGroupContainer.Add(Collision, ManifoldPointIndex, ContactPos, ContactNormal);
+						}
+					}
+				}
+				return ECollisionVisitorResult::Continue;
+			});
+
+		// Sort the groups by the most vertical normal so that we preferentially keep horizontal planes (wrt the UpVector)
+		ContactGroupContainer.SortContactGroups();
+
+		// If all contacts in one group lie beneath the plane of another group, reject it
+		// Loop over all groups (i.e., contact planes)
+		for (int32 ContactGroupIndex = 0; ContactGroupIndex < ContactGroupContainer.ContactGroups.Num(); ++ContactGroupIndex)
+		{
+			const FContactGroup& ContactGroup = ContactGroupContainer.ContactGroups[ContactGroupIndex];
+
+			if (ContactGroup.PlaneNormal.Z < FReal(0.5))
+			{
+				static int32 sHere = 0;
+				++sHere;
+			}
+
+			// @todo(chaos): assuming here that the contacts are not all in a line...do we need to check for that?
+			if (!ContactGroup.bActive)
+			{
+				continue;
+			}
+			
+			// Loop over all other contact groups
+			for (int32 RejectCandidateContactGroupIndex = ContactGroupIndex + 1; RejectCandidateContactGroupIndex < ContactGroupContainer.ContactGroups.Num(); ++RejectCandidateContactGroupIndex)
+			{
+				FContactGroup& RejectCandidateContactGroup = ContactGroupContainer.ContactGroups[RejectCandidateContactGroupIndex];
+				if (!RejectCandidateContactGroup.bActive)
+				{
+					continue;
+				}
+				
+				// Are any point in the candidate group above the plane?
+				bool bPointsAbovePlane = false;
+				for (FContactGroupManifoldPoint& RejectCandidateManifoldPoint : RejectCandidateContactGroup.ManifoldPoints)
+				{
+					const FReal PlaneDistance = FVec3::DotProduct(RejectCandidateManifoldPoint.ContactPosition - ContactGroup.PlanePosition, ContactGroup.PlaneNormal);
+					if (PlaneDistance > PlaneDistanceTolerance)
+					{
+						bPointsAbovePlane = true;
+						break;
+					}
+				}
+
+				// If all point in the candidate group are below the plane, reject them
+				if (!bPointsAbovePlane)
+				{
+					for (FContactGroupManifoldPoint& RejectCandidateManifoldPoint : RejectCandidateContactGroup.ManifoldPoints)
+					{
+						RejectCandidateManifoldPoint.Collision->GetManifoldPoint(RejectCandidateManifoldPoint.ManifoldPointIndex).Flags.bDisabled = true;
+					}
+					RejectCandidateContactGroup.bActive = false;
+				}
+			}
+		}
+	}
+}
\ No newline at end of file
diff --git a/Engine/Source/Runtime/Experimental/Chaos/Private/Chaos/Collision/PBDCollisionConstraint.cpp b/Engine/Source/Runtime/Experimental/Chaos/Private/Chaos/Collision/PBDCollisionConstraint.cpp
--- a/Engine/Source/Runtime/Experimental/Chaos/Private/Chaos/Collision/PBDCollisionConstraint.cpp
+++ b/Engine/Source/Runtime/Experimental/Chaos/Private/Chaos/Collision/PBDCollisionConstraint.cpp
@@ -194,6 +194,7 @@ namespace Chaos
 		, LastShapeWorldPositionDelta()
 		, LastShapeWorldRotationDelta()
 		, SolverBodies{ nullptr, nullptr }
+		, SolverIndex(INDEX_NONE)
 		, GJKWarmStartData()
 		, SavedManifoldPoints()
 		, ManifoldPoints()
@@ -232,6 +233,7 @@ namespace Chaos
 		, LastShapeWorldPositionDelta()
 		, LastShapeWorldRotationDelta()
 		, SolverBodies{ nullptr, nullptr }
+		, SolverIndex(INDEX_NONE)
 		, GJKWarmStartData()
 		, SavedManifoldPoints()
 		, ManifoldPoints()
diff --git a/Engine/Source/Runtime/Experimental/Chaos/Private/Chaos/Collision/SolverCollisionContainer.cpp b/Engine/Source/Runtime/Experimental/Chaos/Private/Chaos/Collision/SolverCollisionContainer.cpp
--- a/Engine/Source/Runtime/Experimental/Chaos/Private/Chaos/Collision/SolverCollisionContainer.cpp
+++ b/Engine/Source/Runtime/Experimental/Chaos/Private/Chaos/Collision/SolverCollisionContainer.cpp
@@ -162,6 +162,10 @@ namespace Chaos
 			{
 				TArrayView<FManifoldPoint> ManifoldPoints = Constraint->GetManifoldPoints();
 				FManifoldPoint& ManifoldPoint = ManifoldPoints[ManifoldPointIndex];
+				if (ManifoldPoint.Flags.bDisabled)
+				{
+					continue;
+				}
 
 				const FVec3 WorldContactPoint0 = ShapeWorldTransform0.TransformPositionNoScale(ManifoldPoint.ContactPoint.ShapeContactPoints[0]);
 				const FVec3 WorldContactPoint1 = ShapeWorldTransform1.TransformPositionNoScale(ManifoldPoint.ContactPoint.ShapeContactPoints[1]);
@@ -257,6 +261,11 @@ namespace Chaos
 			for (int32 ManifoldPointIndex = BeginPointIndex; ManifoldPointIndex < EndPointIndex; ++ManifoldPointIndex)
 			{
 				FManifoldPoint& ManifoldPoint = ManifoldPoints[ManifoldPointIndex];
+				if (ManifoldPoint.Flags.bDisabled)
+				{
+					continue;
+				}
+
 				const FSavedManifoldPoint* SavedManifoldPoint = Constraint->FindSavedManifoldPoint(ManifoldPoint);
 
 				const VectorRegister ShapeContactPos0 = VectorLoadFloat3(&ManifoldPoint.ContactPoint.ShapeContactPoints[0]);
@@ -291,14 +300,14 @@ namespace Chaos
 				}
 				else
 				{
-					//const VectorRegister BodyV0 = VectorLoadFloat3(&Body0->V());
-					//const VectorRegister BodyW0 = VectorLoadFloat3(&Body0->W());
-					//const VectorRegister BodyV1 = VectorLoadFloat3(&Body1->V());
-					//const VectorRegister BodyW1 = VectorLoadFloat3(&Body1->W());
-					//const VectorRegister ContactVel0 = VectorAdd(BodyV0, VectorCross(BodyW0, RelativeContactPosition0));
-					//const VectorRegister ContactVel1 = VectorAdd(BodyV1, VectorCross(BodyW1, RelativeContactPosition1));
-					//const SolverVectorRegister ContactVel = MakeVectorRegisterFloatFromDouble(VectorSubtract(ContactVel0, ContactVel1));
-					//WorldFrictionDelta = VectorMultiply(ContactVel, VectorSetFloat1(Dt));
+					const VectorRegister BodyV0 = VectorLoadFloat3(&Body0->V());
+					const VectorRegister BodyW0 = VectorLoadFloat3(&Body0->W());
+					const VectorRegister BodyV1 = VectorLoadFloat3(&Body1->V());
+					const VectorRegister BodyW1 = VectorLoadFloat3(&Body1->W());
+					const VectorRegister ContactVel0 = VectorAdd(BodyV0, VectorCross(BodyW0, RelativeContactPosition0));
+					const VectorRegister ContactVel1 = VectorAdd(BodyV1, VectorCross(BodyW1, RelativeContactPosition1));
+					const SolverVectorRegister ContactVel = MakeVectorRegisterFloatFromDouble(VectorSubtract(ContactVel0, ContactVel1));
+					WorldFrictionDelta = VectorMultiply(ContactVel, VectorSetFloat1(Dt));
 
 					ManifoldPoint.ShapeAnchorPoints[0] = ManifoldPoint.ContactPoint.ShapeContactPoints[0];
 					ManifoldPoint.ShapeAnchorPoints[1] = ManifoldPoint.ContactPoint.ShapeContactPoints[1];
@@ -353,28 +362,44 @@ namespace Chaos
 
 			Constraint->ResetSolverResults();
 
-			for (int32 PointIndex = 0; PointIndex < Solver.NumManifoldPoints(); ++PointIndex)
+			// NOTE: We only put the non-pruned manifold points into the solver so the ManifoldPointIndex and
+			// SolverManifoldPointIndex do not necessarily match. See GatherManifoldPoints
+			int32 SolverManifoldPointIndex = 0;
+			for (int32 ManifoldPointIndex = 0; ManifoldPointIndex < Constraint->NumManifoldPoints(); ++ManifoldPointIndex)
 			{
-				const FPBDCollisionSolverManifoldPoint& SolverManifoldPoint = Solver.GetManifoldPoint(PointIndex);
+				FSolverVec3 NetPushOut = FVec3(0);
+				FSolverVec3 NetImpulse = FVec3(0);
+				FReal StaticFrictionRatio = FReal(0);
 
-				const FSolverVec3 NetPushOut = 
-					SolverManifoldPoint.NetPushOutNormal * SolverManifoldPoint.WorldContactNormal +
-					SolverManifoldPoint.NetPushOutTangentU * SolverManifoldPoint.WorldContactTangentU +
-					SolverManifoldPoint.NetPushOutTangentV * SolverManifoldPoint.WorldContactTangentV;
+				if (!Constraint->GetManifoldPoint(ManifoldPointIndex).Flags.bDisabled)
+				{
+					const FPBDCollisionSolverManifoldPoint& SolverManifoldPoint = Solver.GetManifoldPoint(SolverManifoldPointIndex);
+
+					NetPushOut = 
+						SolverManifoldPoint.NetPushOutNormal * SolverManifoldPoint.WorldContactNormal +
+						SolverManifoldPoint.NetPushOutTangentU * SolverManifoldPoint.WorldContactTangentU +
+						SolverManifoldPoint.NetPushOutTangentV * SolverManifoldPoint.WorldContactTangentV;
+
+					NetImpulse =
+						SolverManifoldPoint.NetImpulseNormal * SolverManifoldPoint.WorldContactNormal +
+						SolverManifoldPoint.NetImpulseTangentU * SolverManifoldPoint.WorldContactTangentU +
+						SolverManifoldPoint.NetImpulseTangentV * SolverManifoldPoint.WorldContactTangentV;
 
-				const FSolverVec3 NetImpulse =
-					SolverManifoldPoint.NetImpulseNormal * SolverManifoldPoint.WorldContactNormal +
-					SolverManifoldPoint.NetImpulseTangentU * SolverManifoldPoint.WorldContactTangentU +
-					SolverManifoldPoint.NetImpulseTangentV * SolverManifoldPoint.WorldContactTangentV;
+					StaticFrictionRatio = SolverManifoldPoint.StaticFrictionRatio;
+
+					++SolverManifoldPointIndex;
+				}
 
-				Constraint->SetSolverResults(PointIndex, 
+				// NOTE: We call this even for points we did not run the solver for (but with zero results)
+				Constraint->SetSolverResults(ManifoldPointIndex,
 					NetPushOut, 
 					NetImpulse, 
-					SolverManifoldPoint.StaticFrictionRatio,
+					StaticFrictionRatio,
 					Dt);
 			}
 
 			Constraint->SetSolverBodies(nullptr, nullptr);
+			Constraint->SetSolverIndex(INDEX_NONE);
 			Constraint = nullptr;
 			Solver.ResetSolverBodies();
 		}
@@ -413,29 +438,29 @@ namespace Chaos
 		CollisionSolvers.SetNum(MaxCollisions, false);
 	}
 
-	void FPBDCollisionSolverContainer::PreAddConstraintSolver(FPBDCollisionConstraint& Constraint, FSolverBodyContainer& SolverBodyContainer, int32& ConstraintIndex)
+	void FPBDCollisionSolverContainer::PreAddConstraintSolver(FPBDCollisionConstraint& Constraint, FSolverBodyContainer& SolverBodyContainer, int32& SolverIndex)
 	{
 		// This container is required to allocate pointers that are valid for the whole tick,
 		// so we cannot allow the container to resize during the tick. See Reset()
-		check(ConstraintIndex < CollisionSolvers.Num());
-		Constraint.SetSolverBodyContainerIndex(ConstraintIndex);
+		check(SolverIndex < CollisionSolvers.Num());
+		Constraint.SetSolverIndex(SolverIndex);
 
-		FPBDCollisionSolverAdapter& CollisionSolver = CollisionSolvers[ConstraintIndex];
+		FPBDCollisionSolverAdapter& CollisionSolver = CollisionSolvers[SolverIndex];
 		CollisionSolver.GetSolver().Reset();
 
 		CollisionSolver.PreGatherInput(Constraint, SolverBodyContainer);
 
-		++ConstraintIndex;
+		++SolverIndex;
 	}
 
 	void FPBDCollisionSolverContainer::AddConstraintSolver(FReal Dt, FPBDCollisionConstraint& Constraint, const int32 Particle0Level, const int32 Particle1Level, FSolverBodyContainer& SolverBodyContainer, const FPBDCollisionSolverSettings& SolverSettings)
 	{
 		// This container is required to allocate pointers that are valid for the whole tick,
 		// so we cannot allow the container to resize during the tick. See Reset()
-		int32 ConstraintIndex = Constraint.GetSolverBodyContainerIndex();
-		check(ConstraintIndex < CollisionSolvers.Num());
+		const int32 SolverIndex = Constraint.GetSolverIndex();
+		check(SolverIndex < CollisionSolvers.Num());
 
-		FPBDCollisionSolverAdapter& CollisionSolver = CollisionSolvers[ConstraintIndex];
+		FPBDCollisionSolverAdapter& CollisionSolver = CollisionSolvers[SolverIndex];
 
 		CollisionSolver.GatherInput(Dt, Constraint, Particle0Level, Particle1Level, SolverBodyContainer, SolverSettings);
 
diff --git a/Engine/Source/Runtime/Experimental/Chaos/Private/Chaos/PBDCollisionConstraints.cpp b/Engine/Source/Runtime/Experimental/Chaos/Private/Chaos/PBDCollisionConstraints.cpp
--- a/Engine/Source/Runtime/Experimental/Chaos/Private/Chaos/PBDCollisionConstraints.cpp
+++ b/Engine/Source/Runtime/Experimental/Chaos/Private/Chaos/PBDCollisionConstraints.cpp
@@ -2,32 +2,21 @@
 
 #include "Chaos/PBDCollisionConstraints.h"
 
-#include "Chaos/Capsule.h"
 #include "Chaos/ChaosPerfTest.h"
-#include "Chaos/ChaosDebugDraw.h"
 #include "Chaos/ContactModification.h"
 #include "Chaos/PBDCollisionConstraintsContact.h"
-#include "Chaos/CollisionResolutionUtil.h"
 #include "Chaos/CollisionResolution.h"
-#include "Chaos/Collision/CollisionContext.h"
+#include "Chaos/Collision/CollisionPruning.h"
 #include "Chaos/Collision/SolverCollisionContainer.h"
 #include "Chaos/Defines.h"
 #include "Chaos/Evolution/SolverBodyContainer.h"
 #include "Chaos/GeometryQueries.h"
-#include "Chaos/ImplicitObjectUnion.h"
-#include "Chaos/ImplicitObjectScaled.h"
 #include "Chaos/SpatialAccelerationCollection.h"
-#include "Chaos/Levelset.h"
-#include "Chaos/Pair.h"
-#include "Chaos/PBDCollisionConstraintsContact.h"
 #include "Chaos/PBDRigidsSOAs.h"
-#include "Chaos/Sphere.h"
-#include "Chaos/Transform.h"
 #include "Chaos/CastingUtilities.h"
 #include "ChaosLog.h"
 #include "ChaosStats.h"
 #include "Chaos/Evolution/SolverDatas.h"
-#include "Containers/Queue.h"
 #include "ProfilingDebugging/ScopedTimers.h"
 #include "Algo/Sort.h"
 #include "Algo/StableSort.h"
@@ -43,8 +32,6 @@
 
 namespace Chaos
 {
-	extern FRealSingle Chaos_Collision_EdgePrunePlaneDistance;
-
 	int32 CollisionParticlesBVHDepth = 4;
 	FAutoConsoleVariableRef CVarCollisionParticlesBVHDepth(TEXT("p.CollisionParticlesBVHDepth"), CollisionParticlesBVHDepth, TEXT("The maximum depth for collision particles bvh"));
 
@@ -84,6 +71,9 @@ namespace Chaos
 	bool CollisionsAllowParticleTracking = true;
 	FAutoConsoleVariableRef CVarCollisionsAllowParticleTracking(TEXT("p.Chaos.Collision.AllowParticleTracking"), CollisionsAllowParticleTracking, TEXT("Allow particles to track their collisions constraints when their DoBufferCollisions flag is enable [def:true]"));
 
+	bool bCollisionsEnableSubSurfaceCollisionPruning = false;
+	FAutoConsoleVariableRef CVarCollisionsEnableSubSurfaceCollisionPruning(TEXT("p.Chaos.Collision.EnableSubSurfaceCollisionPruning"), bCollisionsEnableSubSurfaceCollisionPruning, TEXT(""));
+	
 	DECLARE_CYCLE_STAT(TEXT("Collisions::Reset"), STAT_Collisions_Reset, STATGROUP_ChaosCollision);
 	DECLARE_CYCLE_STAT(TEXT("Collisions::UpdatePointConstraints"), STAT_Collisions_UpdatePointConstraints, STATGROUP_ChaosCollision);
 	DECLARE_CYCLE_STAT(TEXT("Collisions::BeginDetect"), STAT_Collisions_BeginDetect, STATGROUP_ChaosCollision);
@@ -686,83 +676,18 @@ namespace Chaos
 			{
 				if ((ParticleHandle.CollisionConstraintFlags() & (uint32)ECollisionConstraintFlags::CCF_SmoothEdgeCollisions) != 0)
 				{
-					PruneParticleEdgeCollisions(ParticleHandle.Handle());
-				}
-			}
-		}
-	}
+					FParticleEdgeCollisionPruner EdgePruner(ParticleHandle.Handle());
+					EdgePruner.Prune();
 
-	void FPBDCollisionConstraints::PruneParticleEdgeCollisions(FGeometryParticleHandle* Particle)
-	{
-		FParticleCollisions& ParticleCollision = Particle->ParticleCollisions();
-
-		const FReal EdgePlaneTolerance = Chaos_Collision_EdgePrunePlaneDistance;
-
-		// Loop over edge collisions, then all plane collisions and remove the edge collision if it is hidden by a plane collision
-		// NOTE: We only look at plane collisions where the other shape owns the plane.
-		// @todo(chaos): this should probably only disable individual manifold points
-		// @todo(chaos): we should probably only reject edges if the plane contact is also close to the edge contact
-		// @todo(chaos): we should also try to eliminate face contacts from sub-surface faces
-		// @todo(chaos): perf issue: this processes contacts in world space, but we don't calculated that data until Gather. Fix this.
-		ParticleCollision.VisitCollisions(
-			[Particle, &ParticleCollision, EdgePlaneTolerance](FPBDCollisionConstraint& EdgeCollision)
-			{
-				if (EdgeCollision.IsEnabled())
-				{
-					const int32 EdgeOtherShapeIndex = (EdgeCollision.GetParticle0() == Particle) ? 1 : 0;
-					const EContactPointType VertexContactType = (EdgeOtherShapeIndex == 0) ? EContactPointType::VertexPlane : EContactPointType::PlaneVertex;
-
-					for (const FManifoldPoint& EdgeManifoldPoint : EdgeCollision.GetManifoldPoints())
+					if (bCollisionsEnableSubSurfaceCollisionPruning)
 					{
-						const bool bIsEdgeContact = (EdgeManifoldPoint.ContactPoint.ContactType == EContactPointType::EdgeEdge);
-
-						if (bIsEdgeContact)
-						{
-							const FRigidTransform3& EdgeTransform = (EdgeOtherShapeIndex == 0) ? EdgeCollision.GetShapeWorldTransform0() : EdgeCollision.GetShapeWorldTransform1();
-							const FVec3 EdgePos = EdgeTransform.TransformPositionNoScale(EdgeManifoldPoint.ContactPoint.ShapeContactPoints[EdgeOtherShapeIndex]);
-
-							// Loop over plane collisions
-							ECollisionVisitorResult PlaneResult = ParticleCollision.VisitConstCollisions(
-								[Particle, &EdgeCollision, &EdgePos, EdgePlaneTolerance](const FPBDCollisionConstraint& PlaneCollision)
-								{
-									if ((&PlaneCollision != &EdgeCollision) && PlaneCollision.IsEnabled())
-									{
-										const int32 PlaneOtherShapeIndex = (PlaneCollision.GetParticle0() == Particle) ? 1 : 0;
-										const EContactPointType PlaneContactType = (PlaneOtherShapeIndex == 0) ? EContactPointType::PlaneVertex : EContactPointType::VertexPlane;
-										const FRigidTransform3& PlaneTransform = (PlaneOtherShapeIndex == 0) ? PlaneCollision.GetShapeWorldTransform0() : PlaneCollision.GetShapeWorldTransform1();
-
-										for (const FManifoldPoint& PlaneManifoldPoint : PlaneCollision.GetManifoldPoints())
-										{
-											if (PlaneManifoldPoint.ContactPoint.ContactType == PlaneContactType)
-											{
-												// If the edge position is in the plane, disable it
-												const FVec3 PlanePos = PlaneTransform.TransformPositionNoScale(PlaneManifoldPoint.ContactPoint.ShapeContactPoints[PlaneOtherShapeIndex]);
-												const FVec3 PlaneNormal = PlaneCollision.GetShapeWorldTransform1().TransformVectorNoScale(PlaneManifoldPoint.ContactPoint.ShapeContactNormal);
-
-												const FVec3 EdgePlaneDelta = EdgePos - PlanePos;
-												const FReal EdgePlaneDistance = FVec3::DotProduct(EdgePlaneDelta, PlaneNormal);
-												if (FMath::Abs(EdgePlaneDistance) < EdgePlaneTolerance)
-												{
-													// The edge contact is hidden by a plane contact so disable it and stop the inner loop
-													EdgeCollision.SetDisabled(true);
-													return ECollisionVisitorResult::Stop;
-												}
-											}
-										}
-									}
-									return ECollisionVisitorResult::Continue;
-								});
-
-							// If we disabled this constraint, move to the next one and ignore remaining manifold points
-							if (PlaneResult == ECollisionVisitorResult::Stop)
-							{
-								break;
-							}
-						}
+						const FVec3 UpVector = ParticleHandle.R().GetAxisZ();
+						FParticleSubSurfaceCollisionPruner SubSurfacePruner(ParticleHandle.Handle());
+						SubSurfacePruner.Prune(UpVector);
 					}
 				}
-				return ECollisionVisitorResult::Continue;
-			});
+			}
+		}
 	}
 
 }
diff --git a/Engine/Source/Runtime/Experimental/Chaos/Public/Chaos/CCDUtilities.h b/Engine/Source/Runtime/Experimental/Chaos/Public/Chaos/CCDUtilities.h
--- a/Engine/Source/Runtime/Experimental/Chaos/Public/Chaos/CCDUtilities.h
+++ b/Engine/Source/Runtime/Experimental/Chaos/Public/Chaos/CCDUtilities.h
@@ -54,6 +54,7 @@ namespace Chaos
 	{
 		FPBDCollisionConstraint* SweptConstraint;
 		FCCDParticle* Particle[2];
+		FReal PhiThreshold;
 		int32 Island;
 		int32 ProcessedCount; // The number of times this constraint is processed in the current frame.
 		int32 FastMovingKinematicIndex;
@@ -61,14 +62,16 @@ namespace Chaos
 		FCCDConstraint()
 			: SweptConstraint(nullptr)
 			, Particle{ nullptr, nullptr }
+			, PhiThreshold(0)
 			, Island(INDEX_NONE)
 			, ProcessedCount(0)
 			, FastMovingKinematicIndex(INDEX_NONE)
 		{}
 
-		FCCDConstraint(FPBDCollisionConstraint* const InConstraint, FCCDParticle* InParticle[], const FVec3 Displacements[])
+		FCCDConstraint(FPBDCollisionConstraint* const InConstraint, FCCDParticle* InParticle[], const FVec3 Displacements[], const FReal InPhiThreshold)
 			: SweptConstraint(InConstraint)
 			, Particle{ InParticle[0], InParticle[1] }
+			, PhiThreshold(InPhiThreshold)
 			, Island(INDEX_NONE)
 			, ProcessedCount(0)
 			, FastMovingKinematicIndex(GetFastMovingKinematicIndex(InConstraint, Displacements))
diff --git a/Engine/Source/Runtime/Experimental/Chaos/Public/Chaos/Collision/CollisionPruning.h b/Engine/Source/Runtime/Experimental/Chaos/Public/Chaos/Collision/CollisionPruning.h
new file mode 100644
--- /dev/null
+++ b/Engine/Source/Runtime/Experimental/Chaos/Public/Chaos/Collision/CollisionPruning.h
@@ -0,0 +1,49 @@
+// Copyright Epic Games, Inc. All Rights Reserved.
+#pragma once
+
+#include "CoreMinimal.h"
+
+#include "Chaos/Core.h"
+#include "Chaos/CollisionResolutionTypes.h"
+#include "Chaos/Collision/PBDCollisionConstraint.h"
+#include "Chaos/ParticleHandleFwd.h"
+
+namespace Chaos
+{
+	/**
+	 * \brief Utility to remove edge contacts on a particle if those edge contacts are "hidden" by a face contact
+	 */
+	class FParticleEdgeCollisionPruner
+	{
+	public:
+		FParticleEdgeCollisionPruner(FGeometryParticleHandle* InParticle)
+			: Particle(InParticle)
+		{
+		}
+
+		void Prune();
+
+	private:
+		FGeometryParticleHandle* Particle;
+	};
+
+	/**
+	 * \brief Utility to attempt to remove sub-surface contacts on a particle when those edge contacts are "hidden" by other contacts
+	 * This is not 100% effective and may occasionally remove contacts that it shouldn't! Ideally you would just use
+	 * this pruning policy in specific locations in the world where sub-surface contacts are a problem 
+	 */
+	class FParticleSubSurfaceCollisionPruner
+	{
+	public:
+		FParticleSubSurfaceCollisionPruner(FGeometryParticleHandle* InParticle)
+			: Particle(InParticle)
+		{
+		}
+
+		void Prune(const FVec3& UpVector);
+
+	private:
+		FGeometryParticleHandle* Particle;
+	};
+
+}
\ No newline at end of file
diff --git a/Engine/Source/Runtime/Experimental/Chaos/Public/Chaos/Collision/PBDCollisionConstraint.h b/Engine/Source/Runtime/Experimental/Chaos/Public/Chaos/Collision/PBDCollisionConstraint.h
--- a/Engine/Source/Runtime/Experimental/Chaos/Public/Chaos/Collision/PBDCollisionConstraint.h
+++ b/Engine/Source/Runtime/Experimental/Chaos/Public/Chaos/Collision/PBDCollisionConstraint.h
@@ -77,9 +77,10 @@ namespace Chaos
 
 			struct
 			{
-				uint8 bWasRestored : 1;
-				uint8 bWasReplaced : 1;
-				uint8 bWasFrictionRestored : 1;
+				uint8 bDisabled : 1;						// Whether the point was disabled by edge pruning etc
+				uint8 bWasRestored : 1;						// Whether the point was retored from the previous frame due to lack of movement
+				uint8 bWasReplaced : 1;						// @todo(chaos): remove this
+				uint8 bWasFrictionRestored : 1;				// Whether we used data from the previous frame to set up the contacts to maintain static friction
 				uint8 bInsideStaticFrictionCone : 1;		// Whether we are inside the static friction cone
 			};
 			uint8 Bits;
@@ -475,10 +476,10 @@ namespace Chaos
 			SolverBodies[1] = InSolverBody1;
 		}
 
-		int32 GetSolverBodyContainerIndex() const { return ConstraintIndex; }
-		void SetSolverBodyContainerIndex(int32 InConstraintIndex)
+		int32 GetSolverIndex() const { return SolverIndex; }
+		void SetSolverIndex(int32 InSolverIndex)
 		{
-			ConstraintIndex = InConstraintIndex;
+			SolverIndex = InSolverIndex;
 		}
 
 		/**
@@ -676,7 +677,7 @@ namespace Chaos
 		FSolverBody* SolverBodies[2];
 
 		// Stores the index into the solver container for this constraint
-		int32 ConstraintIndex;
+		int32 SolverIndex;
 
 		// Simplex data from the last call to GJK, used to warm-start GJK
 		FGJKSimplexData GJKWarmStartData;
diff --git a/Engine/Source/Runtime/Experimental/Chaos/Public/Chaos/Collision/SolverCollisionContainer.h b/Engine/Source/Runtime/Experimental/Chaos/Public/Chaos/Collision/SolverCollisionContainer.h
--- a/Engine/Source/Runtime/Experimental/Chaos/Public/Chaos/Collision/SolverCollisionContainer.h
+++ b/Engine/Source/Runtime/Experimental/Chaos/Public/Chaos/Collision/SolverCollisionContainer.h
@@ -66,7 +66,7 @@ namespace Chaos
 		void SetNum(const int32 MaxCollisions);
 
 		// Add a solver constraint for the specified constraint and gather the required solver data
-		void PreAddConstraintSolver(FPBDCollisionConstraint& Constraint, FSolverBodyContainer& SolverBodyContainer, int32& ConstraintIndex);
+		void PreAddConstraintSolver(FPBDCollisionConstraint& Constraint, FSolverBodyContainer& SolverBodyContainer, int32& SolverIndex);
 		void AddConstraintSolver(FReal Dt, FPBDCollisionConstraint& Constraint, const int32 Particle0Level, const int32 Particle1Level, FSolverBodyContainer& SolverBodyContainer, const FPBDCollisionSolverSettings& SolverSettings);
 
 		bool SolvePositionSerial(const FReal Dt, const int32 It, const int32 NumIts, const int32 BeginIndex, const int32 EndIndex, const FPBDCollisionSolverSettings& SolverSettings);
diff --git a/Engine/Source/Runtime/Experimental/Chaos/Public/Chaos/DebugDrawQueue.h b/Engine/Source/Runtime/Experimental/Chaos/Public/Chaos/DebugDrawQueue.h
--- a/Engine/Source/Runtime/Experimental/Chaos/Public/Chaos/DebugDrawQueue.h
+++ b/Engine/Source/Runtime/Experimental/Chaos/Public/Chaos/DebugDrawQueue.h
@@ -429,6 +429,11 @@ public:
 		return (RadiusOfInterest <= 0.0f) || ((Pos - CenterOfInterest).SizeSquared() < (RadiusOfInterest + Radius) * (RadiusOfInterest + Radius));
 	}
 
+	bool IsInRegionOfInterest(FAABB3 Bounds) const
+	{
+		return Bounds.ThickenSymmetrically(FVec3(RadiusOfInterest)).Contains(CenterOfInterest);
+	}
+
 	void SetConsumerActive(void* Consumer, bool bConsumerActive);
 
 	static FDebugDrawQueue& GetInstance();
diff --git a/Engine/Source/Runtime/Experimental/Chaos/Public/Chaos/PBDCollisionConstraints.h b/Engine/Source/Runtime/Experimental/Chaos/Public/Chaos/PBDCollisionConstraints.h
--- a/Engine/Source/Runtime/Experimental/Chaos/Public/Chaos/PBDCollisionConstraints.h
+++ b/Engine/Source/Runtime/Experimental/Chaos/Public/Chaos/PBDCollisionConstraints.h
@@ -319,9 +319,6 @@ protected:
 	// Call PruneParticleEdgeCollisions on all particles with ECollisionConstraintFlags::CCF_SmoothEdgeCollisions set in CollisionFlags
 	void PruneEdgeCollisions();
 
-	// Remove all edge collision on a particle that cannot be hit because a plane contact would prevent it
-	void PruneParticleEdgeCollisions(FGeometryParticleHandle* Particle);
-
 	// The "Legacy" functions handle the older solver types (GbfPbd and StandardPbd)
 	// @todo(chaos): remove legacy methods when the new solver is fully operational and used everywhere (RBAN)
 	void LegacyGatherInput(const FReal Dt, FPBDCollisionConstraint& Constraint, const int32 Particle0Level, const int32 Particle1Level, FPBDIslandSolverData& SolverData);
diff --git a/Engine/Source/Runtime/Experimental/ChaosSolverEngine/Private/Chaos/ChaosDebugDrawComponent.cpp b/Engine/Source/Runtime/Experimental/ChaosSolverEngine/Private/Chaos/ChaosDebugDrawComponent.cpp
--- a/Engine/Source/Runtime/Experimental/ChaosSolverEngine/Private/Chaos/ChaosDebugDrawComponent.cpp
+++ b/Engine/Source/Runtime/Experimental/ChaosSolverEngine/Private/Chaos/ChaosDebugDrawComponent.cpp
@@ -44,13 +44,13 @@ int bChaosDebugDraw_DrawMode = 0;
 FAutoConsoleVariableRef CVarArrowSize(TEXT("p.Chaos.DebugDraw.Mode"), bChaosDebugDraw_DrawMode, TEXT("Where to send debug draw commands. 0 = UE Debug Draw; 1 = VisLog; 2 = Both"));
 
 
-float CommandLifeTime(const Chaos::FLatentDrawCommand& Command)
+float CommandLifeTime(const Chaos::FLatentDrawCommand& Command, const bool bIsPaused)
 {
 	// @todo(chaos): remove lifetime from the system
-	return 0.0f;
+	return ((Command.LifeTime > 0.0f) && !bIsPaused) ? Command.LifeTime : 0.0f;
 }
 
-void DebugDrawChaos(AActor* DebugDrawActor, const TArray<Chaos::FLatentDrawCommand>& DrawCommands)
+void DebugDrawChaos(AActor* DebugDrawActor, const TArray<Chaos::FLatentDrawCommand>& DrawCommands, bool bIsPaused)
 {
 	using namespace Chaos;
 
@@ -89,32 +89,32 @@ void DebugDrawChaos(AActor* DebugDrawActor, const TArray<Chaos::FLatentDrawComma
 			switch (Command.Type)
 			{
 			case FLatentDrawCommand::EDrawType::Point:
-				DrawDebugPoint(World, Command.LineStart, Command.Thickness, Command.Color, Command.bPersistentLines, CommandLifeTime(Command), Command.DepthPriority);
+				DrawDebugPoint(World, Command.LineStart, Command.Thickness, Command.Color, Command.bPersistentLines, CommandLifeTime(Command, bIsPaused), Command.DepthPriority);
 				break;
 			case FLatentDrawCommand::EDrawType::Line:
-				DrawDebugLine(World, Command.LineStart, Command.LineEnd, Command.Color, Command.bPersistentLines, CommandLifeTime(Command), Command.DepthPriority, Command.Thickness);
+				DrawDebugLine(World, Command.LineStart, Command.LineEnd, Command.Color, Command.bPersistentLines, CommandLifeTime(Command, bIsPaused), Command.DepthPriority, Command.Thickness);
 				break;
 			case FLatentDrawCommand::EDrawType::DirectionalArrow:
-				DrawDebugDirectionalArrow(World, Command.LineStart, Command.LineEnd, Command.ArrowSize, Command.Color, Command.bPersistentLines, CommandLifeTime(Command), Command.DepthPriority, Command.Thickness);
+				DrawDebugDirectionalArrow(World, Command.LineStart, Command.LineEnd, Command.ArrowSize, Command.Color, Command.bPersistentLines, CommandLifeTime(Command, bIsPaused), Command.DepthPriority, Command.Thickness);
 				break;
 			case FLatentDrawCommand::EDrawType::Sphere:
-				DrawDebugSphere(World, Command.LineStart, Command.Radius, Command.Segments, Command.Color, Command.bPersistentLines, CommandLifeTime(Command), Command.DepthPriority, Command.Thickness);
+				DrawDebugSphere(World, Command.LineStart, Command.Radius, Command.Segments, Command.Color, Command.bPersistentLines, CommandLifeTime(Command, bIsPaused), Command.DepthPriority, Command.Thickness);
 				break;
 			case FLatentDrawCommand::EDrawType::Box:
-				DrawDebugBox(World, Command.Center, Command.Extent, Command.Rotation, Command.Color, Command.bPersistentLines, CommandLifeTime(Command), Command.DepthPriority, Command.Thickness);
+				DrawDebugBox(World, Command.Center, Command.Extent, Command.Rotation, Command.Color, Command.bPersistentLines, CommandLifeTime(Command, bIsPaused), Command.DepthPriority, Command.Thickness);
 				break;
 			case FLatentDrawCommand::EDrawType::String:
-				DrawDebugString(World, Command.TextLocation, Command.Text, Command.TestBaseActor, Command.Color, CommandLifeTime(Command), Command.bDrawShadow, Command.FontScale);
+				DrawDebugString(World, Command.TextLocation, Command.Text, Command.TestBaseActor, Command.Color, CommandLifeTime(Command, bIsPaused), Command.bDrawShadow, Command.FontScale);
 				break;
 			case FLatentDrawCommand::EDrawType::Circle:
 			{
 				FMatrix M = FRotationMatrix::MakeFromYZ(Command.YAxis, Command.ZAxis);
 				M.SetOrigin(Command.Center);
-				DrawDebugCircle(World, M, Command.Radius, Command.Segments, Command.Color, Command.bPersistentLines, CommandLifeTime(Command), Command.DepthPriority, Command.Thickness, Command.bDrawAxis);
+				DrawDebugCircle(World, M, Command.Radius, Command.Segments, Command.Color, Command.bPersistentLines, CommandLifeTime(Command, bIsPaused), Command.DepthPriority, Command.Thickness, Command.bDrawAxis);
 				break;
 			}
 			case FLatentDrawCommand::EDrawType::Capsule:
-				DrawDebugCapsule(World, Command.Center, Command.HalfHeight, Command.Radius, Command.Rotation, Command.Color, Command.bPersistentLines, CommandLifeTime(Command), Command.DepthPriority, Command.Thickness);
+				DrawDebugCapsule(World, Command.Center, Command.HalfHeight, Command.Radius, Command.Rotation, Command.Color, Command.bPersistentLines, CommandLifeTime(Command, bIsPaused), Command.DepthPriority, Command.Thickness);
 			default:
 				break;
 			}
@@ -281,7 +281,7 @@ void UChaosDebugDrawComponent::TickComponent(float DeltaTime, enum ELevelTick Ti
 			FDebugDrawQueue::GetInstance().ExtractAllElements(DrawCommands);
 		}
 
-		DebugDrawChaos(GetOwner(), DrawCommands);
+		DebugDrawChaos(GetOwner(), DrawCommands, World->IsPaused());
 	}
 #endif
 }
