Chaos - CCD adds to AccumulatedImpulse on the contact
- Add debug draw for AccumulatedImpulse

diff --git a/Engine/Source/Runtime/Experimental/Chaos/Private/Chaos/CCDUtilities.cpp b/Engine/Source/Runtime/Experimental/Chaos/Private/Chaos/CCDUtilities.cpp
--- a/Engine/Source/Runtime/Experimental/Chaos/Private/Chaos/CCDUtilities.cpp
+++ b/Engine/Source/Runtime/Experimental/Chaos/Private/Chaos/CCDUtilities.cpp
@@ -505,6 +505,11 @@ namespace Chaos
 				}
 			}
 		}
+		
+		for (FCCDConstraint* CCDConstraint : SortedCCDConstraints)
+		{ 
+			CCDConstraint->SweptConstraint->SetCCDResults(CCDConstraint->NetImpulse);
+		}
 	}
 
 
@@ -596,6 +601,8 @@ namespace Chaos
 				{
 					Rigid1->V() -= Impulse * InvM1;
 				}
+
+				CCDConstraint->NetImpulse += Impulse;
 			}
 		}
 	}
diff --git a/Engine/Source/Runtime/Experimental/Chaos/Private/Chaos/ChaosDebugDraw.cpp b/Engine/Source/Runtime/Experimental/Chaos/Private/Chaos/ChaosDebugDraw.cpp
--- a/Engine/Source/Runtime/Experimental/Chaos/Private/Chaos/ChaosDebugDraw.cpp
+++ b/Engine/Source/Runtime/Experimental/Chaos/Private/Chaos/ChaosDebugDraw.cpp
@@ -953,6 +953,15 @@ namespace Chaos
 						}
 					}
 				}
+
+				// AccumulatedImpulse
+				if ((Settings.ImpulseScale > 0) && !Contact.GetAccumulatedImpulse().IsNearlyZero())
+				{
+					FColor Color = (ColorScale * FColor::White).ToFColor(false);
+					const FVec3 ImpulsePos = SpaceTransform.TransformPosition(WorldActorTransform0.GetTranslation());
+					FDebugDrawQueue::GetInstance().DrawDebugLine(ImpulsePos, ImpulsePos + Settings.DrawScale * Settings.ImpulseScale * SpaceTransform.TransformVectorNoScale(Contact.GetAccumulatedImpulse()), Color, false, KINDA_SMALL_NUMBER, Settings.DrawPriority, Settings.LineThickness);
+				}
+
 			}
 			if (Settings.ContactOwnerWidth > 0)
 			{
diff --git a/Engine/Source/Runtime/Experimental/Chaos/Public/Chaos/CCDUtilities.h b/Engine/Source/Runtime/Experimental/Chaos/Public/Chaos/CCDUtilities.h
--- a/Engine/Source/Runtime/Experimental/Chaos/Public/Chaos/CCDUtilities.h
+++ b/Engine/Source/Runtime/Experimental/Chaos/Public/Chaos/CCDUtilities.h
@@ -58,6 +58,7 @@ namespace Chaos
 		int32 Island;
 		int32 ProcessedCount; // The number of times this constraint is processed in the current frame.
 		int32 FastMovingKinematicIndex;
+		FVec3 NetImpulse;
 
 		FCCDConstraint()
 			: SweptConstraint(nullptr)
@@ -66,6 +67,7 @@ namespace Chaos
 			, Island(INDEX_NONE)
 			, ProcessedCount(0)
 			, FastMovingKinematicIndex(INDEX_NONE)
+			, NetImpulse(0)
 		{}
 
 		FCCDConstraint(FPBDCollisionConstraint* const InConstraint, FCCDParticle* InParticle[], const FVec3 Displacements[], const FReal InPhiThreshold)
@@ -75,6 +77,7 @@ namespace Chaos
 			, Island(INDEX_NONE)
 			, ProcessedCount(0)
 			, FastMovingKinematicIndex(GetFastMovingKinematicIndex(InConstraint, Displacements))
+			, NetImpulse(0)
 		{}
 
 		int32 GetFastMovingKinematicIndex(const FPBDCollisionConstraint* Constraint, const FVec3 Displacements[]) const;
diff --git a/Engine/Source/Runtime/Experimental/Chaos/Public/Chaos/Collision/CollisionConstraintAllocator.h b/Engine/Source/Runtime/Experimental/Chaos/Public/Chaos/Collision/CollisionConstraintAllocator.h
--- a/Engine/Source/Runtime/Experimental/Chaos/Public/Chaos/Collision/CollisionConstraintAllocator.h
+++ b/Engine/Source/Runtime/Experimental/Chaos/Public/Chaos/Collision/CollisionConstraintAllocator.h
@@ -421,6 +421,9 @@ namespace Chaos
 			}
 
 			Cookie.LastUsedEpoch = CurrentEpoch;
+
+			// An opportunity for the constraint to initialize data before CCD or solver
+			CollisionConstraint->Activate();
 		}
 
 		// All of the overlapping particle pairs in the scene
diff --git a/Engine/Source/Runtime/Experimental/Chaos/Public/Chaos/Collision/PBDCollisionConstraint.h b/Engine/Source/Runtime/Experimental/Chaos/Public/Chaos/Collision/PBDCollisionConstraint.h
--- a/Engine/Source/Runtime/Experimental/Chaos/Public/Chaos/Collision/PBDCollisionConstraint.h
+++ b/Engine/Source/Runtime/Experimental/Chaos/Public/Chaos/Collision/PBDCollisionConstraint.h
@@ -325,6 +325,12 @@ namespace Chaos
 		const FReal GetCollisionRadius0() const { return (Flags.bIsQuadratic0) ? CollisionMargins[0] : FReal(0); }
 		const FReal GetCollisionRadius1() const { return (Flags.bIsQuadratic1) ? CollisionMargins[1] : FReal(0); }
 
+		/** \brief Called each frame when the constraint is active after primary collision detection (but not per incremental collision detection call if enabled) */
+		void Activate()
+		{
+			AccumulatedImpulse = FVec3(0);
+		}
+
 		// @todo(chaos): half of this API is wrong for the new multi-point manifold constraints. Remove it
 
 		void ResetPhi(FReal InPhi) { ClosestManifoldPointIndex = INDEX_NONE; }
@@ -499,7 +505,14 @@ namespace Chaos
 		void ResetSolverResults()
 		{
 			ResetSavedManifoldPoints();
-			AccumulatedImpulse = FVec3(0);
+		}
+
+		/**
+		 * \brief Store the results of CCD contact resolution, if active
+		 */
+		void SetCCDResults(const FVec3& InNetImpulse)
+		{
+			AccumulatedImpulse += InNetImpulse;
 		}
 
 		/**
