Fixing bug in FNodeDependencyCache causing it to not properly use its cache.

diff --git a/Engine/Source/Runtime/Experimental/Interchange/Engine/Private/Tasks/InterchangeTaskParsing.cpp b/Engine/Source/Runtime/Experimental/Interchange/Engine/Private/Tasks/InterchangeTaskParsing.cpp
--- a/Engine/Source/Runtime/Experimental/Interchange/Engine/Private/Tasks/InterchangeTaskParsing.cpp
+++ b/Engine/Source/Runtime/Experimental/Interchange/Engine/Private/Tasks/InterchangeTaskParsing.cpp
@@ -29,7 +29,7 @@
  * For the Dependency sort to work the predicate must be transitive ( A > B > C implying A > C).
  * That means we must take into account the whole dependency chain, not just the immediate dependencies.
  * 
- * This is a helper struct to recursively accumulate the dependencies chain of a node quickly using a Dynamic Programming cache.
+ * This is a helper struct to quickly create the dependencies chain of a node using a cache.
  */
 struct FNodeDependencyCache
 {
@@ -40,9 +40,9 @@ struct FNodeDependencyCache
 			return *DependenciesPtr;
 		}
 
-		TSet<FString>& Dependencies = CachedDependencies.Add(NodeID);
+		TSet<FString> Dependencies;
 		AccumulateDependencies(NodeContainer, NodeID, Dependencies);
-		return Dependencies;
+		return CachedDependencies.Add(NodeID, MoveTemp(Dependencies));
 	}
 
 private:
@@ -64,7 +64,7 @@ private:
 			// Avoid infinite recursion.
 			if (!bAlreadyInSet)
 			{
-				AccumulateDependencies(NodeContainer, DependencyID, OutDependenciesSet);
+				OutDependenciesSet.Append(GetAccumulatedDependencies(NodeContainer, DependencyID));
 			}
 		}
 	}
