Chaos - remove dead collision solver code

diff --git a/Engine/Source/Runtime/Experimental/Chaos/Private/Chaos/Collision/CollisionSolver.cpp b/Engine/Source/Runtime/Experimental/Chaos/Private/Chaos/Collision/CollisionSolver.cpp
--- a/Engine/Source/Runtime/Experimental/Chaos/Private/Chaos/Collision/CollisionSolver.cpp
+++ b/Engine/Source/Runtime/Experimental/Chaos/Private/Chaos/Collision/CollisionSolver.cpp
@@ -1,699 +1,7 @@
 // Copyright Epic Games, Inc. All Rights Reserved.
 #include "Chaos/Collision/CollisionSolver.h"
 
-#include "Chaos/PBDCollisionConstraintsContact.h"
-#include "Chaos/CollisionResolution.h"
-#include "Chaos/CollisionResolutionUtil.h"
-#include "Chaos/Defines.h"
-#include "Chaos/Particle/ParticleUtilities.h"
-#include "Chaos/Utilities.h"
-
-#if INTEL_ISPC
-#include "PBDCollisionConstraints.ispc.generated.h"
-#endif
-
-//PRAGMA_DISABLE_OPTIMIZATION
-
-CHAOS_API DECLARE_LOG_CATEGORY_EXTERN(LogChaosCollisionSolver, Log, All);
-DEFINE_LOG_CATEGORY(LogChaosCollisionSolver);
-
-namespace Chaos
-{
-	namespace CVars
-	{
-		extern int32 Chaos_Collision_UseShockPropagation;
-		extern FRealSingle Chaos_Collision_CollisionClipTolerance;
-
-		bool bChaos_Manifold_Apply_AllowNegativeIncrementalImpulse = true;
-		FAutoConsoleVariableRef CVarChaos_Manifold_NegativeApply(TEXT("p.Chaos.Collision.Manifold.Apply.NegativeIncrementalImpulse"), bChaos_Manifold_Apply_AllowNegativeIncrementalImpulse, TEXT(""));
-
-		bool Chaos_Manifold_PushOut_NegativePushOut = false;
-		bool Chaos_Manifold_PushOut_StaticFriction = true;
-		bool Chaos_Manifold_PushOut_Restitution = false;
-		bool Chaos_Manifold_PushOut_PositionCorrection = true;
-		int32 Chaos_Manifold_PushOut_VelocityCorrection = 2;
-		FAutoConsoleVariableRef CVarChaos_Manifold_PushOut_NegativePushOut(TEXT("p.Chaos.Collision.Manifold.PushOut.NegativePushOut"), Chaos_Manifold_PushOut_NegativePushOut, TEXT(""));
-		FAutoConsoleVariableRef CVarChaos_Manifold_PushOut_StaticFriction(TEXT("p.Chaos.Collision.Manifold.PushOut.StaticFriction"), Chaos_Manifold_PushOut_StaticFriction, TEXT(""));
-		FAutoConsoleVariableRef CVarChaos_Manifold_PushOut_Restitution(TEXT("p.Chaos.Collision.Manifold.PushOut.Restitution"), Chaos_Manifold_PushOut_Restitution, TEXT(""));
-		FAutoConsoleVariableRef CVarChaos_Manifold_PushOut_PositionCorrection(TEXT("p.Chaos.Collision.Manifold.PushOut.PositionCorrection"), Chaos_Manifold_PushOut_PositionCorrection, TEXT(""));
-		FAutoConsoleVariableRef CVarChaos_Manifold_PushOut_VelocityCorrection(TEXT("p.Chaos.Collision.Manifold.PushOut.VelocityCorrectionMode"), Chaos_Manifold_PushOut_VelocityCorrection, TEXT("0 = No Velocity Correction; 1 = Normal Velocity Correction; 2 = Normal + Tangential Velocity Correction"));
-
-		// Note: Velocity solve requires full stiffness for restitution to work correctly so the max stiffness should be 1
-		FRealSingle Chaos_Manifold_Apply_MinStiffness = 0.25f;
-		FRealSingle Chaos_Manifold_Apply_MaxStiffness = 1.0f;
-		FAutoConsoleVariableRef CVarChaos_Manifold_Apply_MinStiffness(TEXT("p.Chaos.Collision.Manifold.Apply.MinStiffness"), Chaos_Manifold_Apply_MinStiffness, TEXT(""));
-		FAutoConsoleVariableRef CVarChaos_Manifold_Apply_MaxStiffness(TEXT("p.Chaos.Collision.Manifold.Apply.MaxStiffness"), Chaos_Manifold_Apply_MaxStiffness, TEXT(""));
-
-		FRealSingle Chaos_Manifold_PushOut_MinStiffness = 0.25f;
-		FRealSingle Chaos_Manifold_PushOut_MaxStiffness = 0.5f;
-		FAutoConsoleVariableRef CVarChaos_Manifold_PushOut_MinStiffness(TEXT("p.Chaos.Collision.Manifold.PushOut.MinStiffness"), Chaos_Manifold_PushOut_MinStiffness, TEXT(""));
-		FAutoConsoleVariableRef CVarChaos_Manifold_PushOut_MaxStiffness(TEXT("p.Chaos.Collision.Manifold.PushOut.MaxStiffness"), Chaos_Manifold_PushOut_MaxStiffness, TEXT(""));
-
-		float Chaos_ConstraintStiffness_Modifier = 5.0;
-		FAutoConsoleVariableRef CVarChaos_ConstraintStiffness_Modifier(TEXT("p.Chaos.Collision.ConstraintStiffnessModifier"), Chaos_ConstraintStiffness_Modifier, TEXT(""));
-
-		FRealSingle Chaos_Manifold_Apply_ImpulseTolerance = 1.e-4f;
-		FAutoConsoleVariableRef CVarChaos_Manifold_Apply_ImpulseTolerance(TEXT("p.Chaos.Collision.Manifold.Apply.ImpulseTolerance"), Chaos_Manifold_Apply_ImpulseTolerance, TEXT(""));
-
-		float Chaos_Manifold_PushOut_PositionTolerance = 1.e-4f;
-		FAutoConsoleVariableRef CVarChaos_Manifold_PushOut_PositionTolerance(TEXT("p.Chaos.Collision.Manifold.PushOut.PositionTolerance"), Chaos_Manifold_PushOut_PositionTolerance, TEXT(""));
-	}
-	using namespace CVars;
-
-	namespace Collisions
-	{
-		inline FReal GetConstraintStiffness(const FReal InStiffness, const FContactIterationParameters& IterationParameters)
-		{
-			FReal Stiffness = FMath::Clamp(InStiffness, (FReal)0.0, (FReal)1.0);
-			if (Stiffness < 1.0)
-			{
-				FReal Factor = (FReal)Chaos_ConstraintStiffness_Modifier * (FReal)IterationParameters.NumIterations * (FReal)IterationParameters.NumPairIterations;
-				if (!FMath::IsNearlyZero(Factor))
-				{
-					return Stiffness / Factor;
-				}
-			}
-			return (FReal)1;
-		}
-
-		inline FReal GetApplyStiffness(int32 It, int32 NumIts)
-		{
-			const FReal Interpolant = (FReal)(It + 1) / (FReal)NumIts;
-			return FMath::Lerp(Chaos_Manifold_Apply_MinStiffness, Chaos_Manifold_Apply_MaxStiffness, Interpolant);
-		}
-
-		inline FReal GetPushOutStiffness(int32 It, int32 NumIts)
-		{
-			const FReal Interpolant = (FReal)(It + 1) / (FReal)NumIts;
-			return FMath::Lerp(Chaos_Manifold_PushOut_MinStiffness, Chaos_Manifold_PushOut_MaxStiffness, Interpolant);
-		}
-
-
-		// Calculate the impulse to drive normal contact velocity to zero (or positive for restitution), ignoring friction
-		// Some state on the in/out Manifold point is also modified for use in the push out phase.
-		void CalculateManifoldNormalVelocityCorrection(
-			const FReal Stiffness,
-			const FContactIterationParameters& IterationParameters,
-			const FContactParticleParameters& ParticleParameters,
-			const FReal Restitution,
-			FSolverBody& Body0,
-			FSolverBody& Body1,
-			FManifoldPoint& ManifoldPoint)
-		{
-		#if 0
-			// Velocity correction uses the average contact point, and not the deepest point on each body
-			const FVec3 RelativeContactPoint0 = ManifoldPoint.ContactPoint.Location - Body0.P();
-			const FVec3 RelativeContactPoint1 = ManifoldPoint.ContactPoint.Location - Body1.P();
-			const FVec3 ContactNormal = ManifoldPoint.ContactPoint.Normal;
-			const FReal ContactPhi = ManifoldPoint.ContactPoint.Phi;
-
-			// Reject non-contact points unless the point has previously been processed - we may want to undo some of the previous work
-			// The tolerance is so that we catch collisions that were resolved last frame, but resulted in slightly positive separations. However,
-			// it also means we are effectively padding objects by this amount.
-			if ((ContactPhi > Chaos_Collision_CollisionClipTolerance) && (ManifoldPoint.InitialPhi > Chaos_Collision_CollisionClipTolerance) && !ManifoldPoint.bActive)
-			{
-				return;
-			}
-
-			const FVec3 ContactVelocity0 = Body0.V() + FVec3::CrossProduct(Body0.W(), RelativeContactPoint0);
-			const FVec3 ContactVelocity1 = Body1.V() + FVec3::CrossProduct(Body1.W(), RelativeContactPoint1);
-			const FVec3 ContactVelocity = ContactVelocity0 - ContactVelocity1;
-			const FReal ContactVelocityNormalLen = FVec3::DotProduct(ContactVelocity, ContactNormal);
-
-			// Reject contacts moving apart unless the point has previously been processed, in which case we may want to undo some of the previous work
-			const bool bAllowNegativeImpulse = ManifoldPoint.bActive && bChaos_Manifold_Apply_AllowNegativeIncrementalImpulse;
-			if ((ContactVelocityNormalLen > 0.0f) && !bAllowNegativeImpulse)
-			{
-				return;
-			}
-
-			// Target normal velocity, including restitution
-			const bool bApplyRestitution = (Restitution > 0.0f) && (ManifoldPoint.InitialContactVelocity < -ParticleParameters.RestitutionVelocityThreshold);
-			FReal ContactVelocityTargetNormal = 0.0f;
-			if (bApplyRestitution)
-			{
-				ContactVelocityTargetNormal = FMath::Max((FReal)0., -Restitution * ManifoldPoint.InitialContactVelocity);
-			}
-
-			// Calculate constraint-space mass
-			const FMatrix33 ConstraintMassInv =
-				(Body0.IsDynamic() ? ComputeFactorMatrix3(RelativeContactPoint0, Body0.InvI(), Body0.InvM()) : FMatrix33(0)) +
-				(Body1.IsDynamic() ? ComputeFactorMatrix3(RelativeContactPoint1, Body1.InvI(), Body1.InvM()) : FMatrix33(0));
-
-			// Calculate the impulse required to drive normal contact velocity to zero
-			const FReal ContactVelocityNormalChange = (ContactVelocityTargetNormal - ContactVelocityNormalLen);
-			const FReal ConstraintMassNormalInv = FVec3::DotProduct(ContactNormal, ConstraintMassInv * ContactNormal);
-			if (ConstraintMassNormalInv < SMALL_NUMBER)
-			{
-				return;
-			}
-			FVec3 Impulse = (Stiffness * ContactVelocityNormalChange / ConstraintMassNormalInv) * ContactNormal;
-
-			// Clip the impulse so that the accumulated impulse is not in the wrong direction
-			FVec3 NetImpulse = ManifoldPoint.NetImpulse + Impulse;
-
-			// Clamp the total impulse to be positive along the normal
-			// Note: if we go inactive we also undo all friction
-			const bool bActive = FVec3::DotProduct(NetImpulse, ContactNormal) > 0.0f;
-			if (!bActive)
-			{
-				Impulse = -ManifoldPoint.NetImpulse;
-				NetImpulse = FVec3(0);
-			}
-
-			// Calculate the velocity deltas from the impulse
-			if (Body0.IsDynamic())
-			{
-				const FVec3 AngularImpulse = FVec3::CrossProduct(RelativeContactPoint0, Impulse);
-				const FVec3 DV0 = Body0.InvM() * Impulse;
-				const FVec3 DW0 = Body0.InvI() * AngularImpulse;
-				Body0.ApplyVelocityDelta(DV0, DW0);
-			}
-			if (Body1.IsDynamic())
-			{
-				const FVec3 AngularImpulse = FVec3::CrossProduct(RelativeContactPoint1, Impulse);
-				const FVec3 DV1 = Body1.InvM() * -Impulse;
-				const FVec3 DW1 = Body1.InvI() * -AngularImpulse;
-				Body1.ApplyVelocityDelta(DV1, DW1);
-			}
-
-			// PushOut needs to know if we applied restitution and static friction
-			ManifoldPoint.bActive = bActive;
-			ManifoldPoint.bInsideStaticFrictionCone = bActive && ManifoldPoint.bInsideStaticFrictionCone;					// Reset when not active
-			ManifoldPoint.bRestitutionEnabled = bActive && (bApplyRestitution || ManifoldPoint.bRestitutionEnabled);		// Latches to on-state when active
-			ManifoldPoint.NetImpulse = NetImpulse;
-
-			// If we applied any additional impulse, we need to go again next iteration
-			const FReal ImpulseTolerance = Chaos_Manifold_Apply_ImpulseTolerance * Chaos_Manifold_Apply_ImpulseTolerance;
-			*IterationParameters.NeedsAnotherIteration |= (Impulse.SizeSquared() > ImpulseTolerance);
-#endif
-		}
-
-		void CalculateManifoldFrictionVelocityCorrection(
-			const FReal Stiffness,
-			const FContactIterationParameters& IterationParameters,
-			const FContactParticleParameters& ParticleParameters,
-			const FReal StaticFriction,
-			const FReal DynamicFriction,
-			FSolverBody& Body0,
-			FSolverBody& Body1,
-			FManifoldPoint& ManifoldPoint)
-		{
-#if 0
-			// Reject inactive contacts (no normal impulse)
-			if (!ManifoldPoint.bActive || (StaticFriction == 0.0f))
-			{
-				return;
-			}
-
-			// Velocity correction uses the average contact point, and not the deepest point on each body
-			const FVec3 RelativeContactPoint0 = ManifoldPoint.ContactPoint.Location - Body0.P();
-			const FVec3 RelativeContactPoint1 = ManifoldPoint.ContactPoint.Location - Body1.P();
-			const FVec3 ContactNormal = ManifoldPoint.ContactPoint.Normal;
-
-			const FVec3 ContactVelocity0 = Body0.V() + FVec3::CrossProduct(Body0.W(), RelativeContactPoint0);
-			const FVec3 ContactVelocity1 = Body1.V() + FVec3::CrossProduct(Body1.W(), RelativeContactPoint1);
-			const FVec3 ContactVelocity = ContactVelocity0 - ContactVelocity1;
-			const FReal ContactVelocityNormalLen = FVec3::DotProduct(ContactVelocity, ContactNormal);
-			const FVec3 ContactVelocityTangent = ContactVelocity - ContactVelocityNormalLen * ContactNormal;
-			const FReal ContactVelocityTangentLen = ContactVelocityTangent.Size();
-				
-			// Note: We cannot early-out at low lateral contact velocity because we may have reduced the normal
-			// impulse in a later iteration, which means we may need to reduce the friction impulse to get
-			// it back into the friction cone. However, we can skip the incremental impulse calculation
-			// and go straight to the friction cone clamping (on net impulse).
-			FVec3 Impulse = FVec3(0);
-			const FReal FrictionVelocityThreshold = 1.e-5f;
-			if (ContactVelocityTangentLen > FrictionVelocityThreshold)
-			{
-				const FVec3 ContactTangent = ContactVelocityTangent / ContactVelocityTangentLen;
-
-				// Calculate constraint-space mass
-				const FMatrix33 ConstraintMassInv =
-					(Body0.IsDynamic() ? ComputeFactorMatrix3(RelativeContactPoint0, Body0.InvI(), Body0.InvM()) : FMatrix33(0)) +
-					(Body1.IsDynamic() ? ComputeFactorMatrix3(RelativeContactPoint1, Body1.InvI(), Body1.InvM()) : FMatrix33(0));
-
-				// Calculate the impulse required to drive contact velocity to zero along the tangent
-				const FReal ContactVelocityTangentChange = -ContactVelocityTangentLen;
-				const FReal ConstraintMassTangentInv = FVec3::DotProduct(ContactTangent, ConstraintMassInv * ContactTangent);
-				const FReal ImpulseTangentLen = (Stiffness * ContactVelocityTangentChange / ConstraintMassTangentInv);
-				Impulse = ImpulseTangentLen * ContactTangent;
-			}
-
-			// Clamp the tangential impulse to the friction cone (looking at net impulse over all iterations)
-			bool bInsideStaticFrictionCone = true;
-			const FReal NetImpulseNormalLen = FVec3::DotProduct(ManifoldPoint.NetImpulse, ContactNormal);
-			const FVec3 NetImpulseTangent = ManifoldPoint.NetImpulse - NetImpulseNormalLen * ContactNormal;
-			const FVec3 NewNetImpulseTangent = NetImpulseTangent + Impulse;
-			const FReal NewNetImpulseTangentLen = NewNetImpulseTangent.Size();
-			const FReal MaxStaticImpulseTangentLen = StaticFriction * NetImpulseNormalLen;
-			if (NewNetImpulseTangentLen > MaxStaticImpulseTangentLen)
-			{
-				const FReal MaxDynamicImpulseTangentLen = DynamicFriction * NetImpulseNormalLen;
-				const FVec3 ClampedNewNetImpulseTangent = (MaxDynamicImpulseTangentLen / NewNetImpulseTangentLen) * NewNetImpulseTangent;
-
-				Impulse = ClampedNewNetImpulseTangent - NetImpulseTangent;
-				bInsideStaticFrictionCone = false;
-			}
-			const FVec3 NetImpulse = ManifoldPoint.NetImpulse + Impulse;
-
-			// Calculate the velocity deltas from the impulse
-			if (Body0.IsDynamic())
-			{
-				const FVec3 AngularImpulse = FVec3::CrossProduct(RelativeContactPoint0, Impulse);
-				const FVec3 DV0 = Body0.InvM() * Impulse;
-				const FVec3 DW0 = Body0.InvI() * AngularImpulse;
-				Body0.ApplyVelocityDelta(DV0, DW0);
-			}
-			if (Body1.IsDynamic())
-			{
-				const FVec3 AngularImpulse = FVec3::CrossProduct(RelativeContactPoint1, Impulse);
-				const FVec3 DV1 = Body1.InvM() * -Impulse;
-				const FVec3 DW1 = Body1.InvI() * -AngularImpulse;
-				Body1.ApplyVelocityDelta(DV1, DW1);
-			}
-
-			// PushOut needs to know if we applied restitution and static friction
-			ManifoldPoint.bInsideStaticFrictionCone = bInsideStaticFrictionCone;
-			ManifoldPoint.NetImpulse = NetImpulse;
-
-			// If we applied any additional impulse, we need to go again next iteration
-			const FReal ImpulseTolerance = Chaos_Manifold_Apply_ImpulseTolerance * Chaos_Manifold_Apply_ImpulseTolerance;
-			*IterationParameters.NeedsAnotherIteration |= (Impulse.SizeSquared() > ImpulseTolerance);
-#endif
-		}
-
-
-		// Apply a position correction to the bodies so that the manifold points is not penetrating.
-		// This modifies the in/out CoM position and rotations.
-		void ApplyManifoldPushOutCorrection(
-			const FReal Stiffness,
-			const FPBDCollisionConstraint& Constraint,
-			const FContactIterationParameters& IterationParameters,
-			const FContactParticleParameters& ParticleParameters,
-			FConstraintSolverBody& Body0,
-			FConstraintSolverBody& Body1,
-			FManifoldPoint& ManifoldPoint)
-		{
-#if 0
-			const FReal Margin0 = ManifoldPoint.ContactPoint.ShapeMargins[0];
-			const FReal Margin1 = ManifoldPoint.ContactPoint.ShapeMargins[1];
-
-			// Calculate the position error we need to correct, including static friction and restitution
-			// Position correction uses the deepest point on each body (see velocity correction which uses average contact)
-
-			const FVec3 ContactNormal = ManifoldPoint.ManifoldContactNormal;
-
-			const bool bApplyStaticFriction = (ManifoldPoint.bActive && ManifoldPoint.bInsideStaticFrictionCone && ManifoldPoint.bPotentialRestingContact && Chaos_Manifold_PushOut_StaticFriction);
-			FVec3 LocalContactPoint0 = ManifoldPoint.CoMContactPoints[0];
-			FVec3 LocalContactPoint1 = ManifoldPoint.CoMContactPoints[1];
-
-			// If we are applying static friction, we are attempting to move the contact points at their intitial positions back together
-			if (bApplyStaticFriction)
-			{
-				Constraint.CalculatePrevCoMContactPoints(Body0.SolverBody(), Body1.SolverBody(), ManifoldPoint, IterationParameters.Dt, LocalContactPoint0, LocalContactPoint1);
-			}
-
-			// We could push out to the PBD distance that would give an implicit velocity equal to -(1+e).Vin
-			// but the values involved are not very stable from frame to frame, so instead we actually pull
-			// objects together so that Phi = 0 after pushout if the contact applied a contact impulse
-			const FReal PhiPadding = 0;
-
-			// Contact points on each body adjusted so that the points end up separated by TargetPhi
-			const FVec3 RelativeContactPoint0 = Body0.Q() * LocalContactPoint0 - (PhiPadding + Margin0) * ContactNormal;
-			const FVec3 RelativeContactPoint1 = Body1.Q() * LocalContactPoint1 + (PhiPadding + Margin1) * ContactNormal;
-
-			// Net error we need to correct, including lateral movement to correct for friction
-			FVec3 ContactError = (Body1.P() + RelativeContactPoint1) - (Body0.P() + RelativeContactPoint0);
-
-			FReal ContactErrorNormal = FVec3::DotProduct(ContactError, ContactNormal);
-
-			// Don't allow objects to be pulled together, but we may still have to correct static friction drift
-			// @todo(chaos): allow undoing of previously added normal pushout
-			if ((ContactErrorNormal < 0.0f) && !Chaos_Manifold_PushOut_NegativePushOut)
-			{
-				// If we have no static friction and no normal separation to enforce, we have nothing to do
-				if (!bApplyStaticFriction)
-				{
-					return;
-				}
-
-				ContactError = ContactError - ContactErrorNormal * ContactNormal;
-				ContactErrorNormal = 0.0f;
-			}
-
-			if (ContactError.SizeSquared() < Chaos_Manifold_PushOut_PositionTolerance * Chaos_Manifold_PushOut_PositionTolerance)
-			{
-				return;
-			}
-
-			// Calculate joint-space mass matrix (J.M.Jt)
-			const FMatrix33 ContactMassInv =
-				(Body0.IsDynamic() ? ComputeFactorMatrix3(RelativeContactPoint0, Body0.InvI(), Body0.InvM()) : FMatrix33(0)) +
-				(Body1.IsDynamic() ? ComputeFactorMatrix3(RelativeContactPoint1, Body1.InvI(), Body1.InvM()) : FMatrix33(0));
-
-			// Calculate pushout
-			// - If we were inside the static friction cone during the apply step, correct positions so that
-			// the relative contact points at the start of the frame are coincident.
-			// - If we were outside the static friction cone, just push out along the normal 
-			// i.e., ignore dynamic friction during the pushout step
-			FVec3 PushOut = FVec3(0);
-			if (bApplyStaticFriction)
-			{
-				const FMatrix33 ContactMass = ContactMassInv.Inverse();
-				PushOut = Stiffness * ContactMass * ContactError;
-			}
-			else
-			{
-				const FReal PushOutDenominator = FVec3::DotProduct(ContactNormal, ContactMassInv * ContactNormal);
-				if (PushOutDenominator > SMALL_NUMBER)
-				{
-					PushOut = (Stiffness * ContactErrorNormal / PushOutDenominator) * ContactNormal;
-				}
-			}
-
-			FVec3 NetPushOut = ManifoldPoint.NetPushOut + PushOut;
-
-			if (Body0.IsDynamic())
-			{
-				const FVec3 AngularPushOut = FVec3::CrossProduct(RelativeContactPoint0, PushOut);
-				const FVec3 DX0 = Body0.InvM() * PushOut;
-				const FVec3 DR0 = Body0.InvI() * AngularPushOut;
-				Body0.ApplyTransformDelta(DX0, DR0);
-				Body0.UpdateRotationDependentState();
-			}
-
-			if (Body1.IsDynamic())
-			{
-				const FVec3 AngularPushOut = FVec3::CrossProduct(RelativeContactPoint1, PushOut);
-				const FVec3 DX1 = Body1.InvM() * -PushOut;
-				const FVec3 DR1 = Body1.InvI() * -AngularPushOut;
-				Body1.ApplyTransformDelta(DX1, DR1);
-				Body1.UpdateRotationDependentState();
-			}
-
-			ManifoldPoint.NetPushOut = NetPushOut;
-			*IterationParameters.NeedsAnotherIteration = true;
-#endif
-		}
-
-
-		// Apply a velocity correction to the Manifold point for use in the PushOut phase.
-		void CalculateManifoldPushOutVelocityCorrection(
-			const FReal Stiffness,
-			const FContactIterationParameters& IterationParameters,
-			const FContactParticleParameters& ParticleParameters,
-			const FReal Restitution,
-			const FReal StaticFriction,
-			const FReal DynamicFriction,
-			FConstraintSolverBody& Body0,
-			FConstraintSolverBody& Body1,
-			FManifoldPoint& ManifoldPoint)
-		{
-#if 0
-			// Velocity correction uses the average contact point, and not the deepest point on each body
-			const FVec3 RelativeContactPoint0 = ManifoldPoint.ContactPoint.Location - Body0.P();
-			const FVec3 RelativeContactPoint1 = ManifoldPoint.ContactPoint.Location - Body1.P();
-			const FVec3 ContactNormal = ManifoldPoint.ContactPoint.Normal;
-			FVec3 ContactTangent = FVec3(0);
-
-			if (Chaos_Manifold_PushOut_VelocityCorrection == 2)
-			{
-				ContactTangent = (ManifoldPoint.NetPushOut - FVec3::DotProduct(ManifoldPoint.NetPushOut, ContactNormal) * ContactNormal).GetSafeNormal();
-			}
-
-			const FVec3 ContactVelocity0 = Body0.V() + FVec3::CrossProduct(Body0.W(), RelativeContactPoint0);
-			const FVec3 ContactVelocity1 = Body1.V() + FVec3::CrossProduct(Body1.W(), RelativeContactPoint1);
-			const FVec3 ContactVelocity = ContactVelocity0 - ContactVelocity1;
-			const FReal ContactVelocityNormal = FVec3::DotProduct(ContactVelocity, ContactNormal);
-			const FReal ContactVelocityTangent = FVec3::DotProduct(ContactVelocity, ContactTangent);
-
-			// If we are already moving away we are probably done. However, we may want to undo some
-			// of the previously applied velocity correct if there is some
-			if ((ContactVelocityNormal > 0.0f) && FMath::IsNearlyZero(ManifoldPoint.NetPushOutImpulseNormal))
-			{
-				return;
-			}
-
-			// Calculate constraint-space inverse mass
-			const FMatrix33 ContactMassInv =
-				(Body0.IsDynamic() ? ComputeFactorMatrix3(RelativeContactPoint0, Body0.InvI(), Body0.InvM()) : FMatrix33(0)) +
-				(Body1.IsDynamic() ? ComputeFactorMatrix3(RelativeContactPoint1, Body1.InvI(), Body1.InvM()) : FMatrix33(0));
-
-			// Calculate the impulse to get the desired target normal velocity
-			// We are ignoring both static and dynamic friction here
-			FReal ImpulseNormal = 0.0f;
-			FReal NetImpulseNormal = 0.0f;
-			{
-				// If we applied restitution in the velocity solve step, we also apply it here
-				FReal TargetVelocityNormal = 0.0f;
-				if (ManifoldPoint.bRestitutionEnabled && Chaos_Manifold_PushOut_Restitution)
-				{
-					TargetVelocityNormal = FMath::Max((FReal)0., -Restitution * ManifoldPoint.InitialContactVelocity);
-				}
-
-				const FReal ImpulseNormalDenominator = FVec3::DotProduct(ContactNormal, ContactMassInv * ContactNormal);
-				if (FMath::Abs(ImpulseNormalDenominator) > SMALL_NUMBER)
-				{
-					const FReal ContactVelocityError = TargetVelocityNormal - ContactVelocityNormal;
-					ImpulseNormal = (Stiffness * ContactVelocityError / ImpulseNormalDenominator);
-				}
-
-				// If we applied a negative impulse this iteration, make sure the total impulse is not negative
-				NetImpulseNormal = ManifoldPoint.NetPushOutImpulseNormal + ImpulseNormal;
-				if (NetImpulseNormal < 0.0f)
-				{
-					ImpulseNormal = -ManifoldPoint.NetPushOutImpulseNormal;
-					NetImpulseNormal = 0.0f;
-				}
-			}
-
-			FReal ImpulseTangent = 0.0f;
-			FReal NetImpulseTangent = 0.0f;
-			const bool bApplyStaticFriction = ManifoldPoint.bInsideStaticFrictionCone && ManifoldPoint.bPotentialRestingContact && Chaos_Manifold_PushOut_StaticFriction && (Chaos_Manifold_PushOut_VelocityCorrection == 2);
-			if (bApplyStaticFriction)
-			{
-				FReal TargetVelocityTangent = 0.0f;
-
-				const FReal ImpulseTangentDenominator = FVec3::DotProduct(ContactTangent, ContactMassInv * ContactTangent);
-				if (FMath::Abs(ImpulseTangentDenominator) > SMALL_NUMBER)
-				{
-					const FReal ContactVelocityError = TargetVelocityTangent - ContactVelocityTangent;
-					ImpulseTangent = (Stiffness * ContactVelocityError / ImpulseTangentDenominator);
-				}
-
-				// If we applied a negative impulse this iteration, make sure the total impulse is not negative
-				NetImpulseTangent = ManifoldPoint.NetPushOutImpulseTangent + ImpulseTangent;
-				if (NetImpulseTangent < 0.0f)
-				{
-					ImpulseTangent = -ManifoldPoint.NetPushOutImpulseTangent;
-					NetImpulseTangent = 0.0f;
-				}
-			}
-
-
-			const FVec3 Impulse = ImpulseNormal * ContactNormal + ImpulseTangent * ContactTangent;
-			if (Body0.IsDynamic())
-			{
-				const FVec3 AngularImpulse = FVec3::CrossProduct(RelativeContactPoint0, Impulse);
-				const FVec3 DV0 = Body0.InvM() * Impulse;
-				const FVec3 DW0 = Body0.InvI() * AngularImpulse;
-				Body0.ApplyVelocityDelta(DV0, DW0);
-			}
-
-			if (Body1.IsDynamic())
-			{
-				const FVec3 AngularImpulse = FVec3::CrossProduct(RelativeContactPoint1, Impulse);
-				const FVec3 DV1 = Body1.InvM() * -Impulse;
-				const FVec3 DW1 = Body1.InvI() * -AngularImpulse;
-				Body1.ApplyVelocityDelta(DV1, DW1);
-			}
-
-			ManifoldPoint.NetPushOutImpulseNormal = NetImpulseNormal;
-			ManifoldPoint.NetPushOutImpulseTangent = NetImpulseTangent;
-#endif
-		}
-
-
-		// Velocity solver loop for a single contact manifold.
-		void ApplyContactManifold(
-			FPBDCollisionConstraint& Constraint,
-			const FContactIterationParameters& IterationParameters,
-			const FContactParticleParameters& ParticleParameters)
-		{
-#if 0
-			FSolverBody& Body0 = *Constraint.GetSolverBody0();
-			FSolverBody& Body1 = *Constraint.GetSolverBody1();
-
-			Constraint.AccumulatedImpulse = FVec3(0);
-
-			const FReal ConstraintStiffness = GetConstraintStiffness(Constraint.GetStiffness(), IterationParameters);
-			const FReal Stiffness = ConstraintStiffness * GetApplyStiffness(IterationParameters.Iteration, IterationParameters.NumIterations);
-
-			const FVec3 InitialV0 = Body0.V();
-			const FVec3 InitialW0 = Body0.W();
-			const FVec3 InitialV1 = Body1.V();
-			const FVec3 InitialW1 = Body1.W();
-
-			// Iterate over the manifold and accumulate velocity corrections
-			TArrayView<FManifoldPoint> ManifoldPoints = Constraint.GetManifoldPoints();
-			for (int32 PointIndex = 0; PointIndex < ManifoldPoints.Num(); ++PointIndex)
-			{
-				FManifoldPoint& ManifoldPoint = Constraint.SetActiveManifoldPoint(PointIndex, Body0.P(), Body0.Q(), Body1.P(), Body1.Q());
-
-				CalculateManifoldNormalVelocityCorrection(
-					Stiffness,
-					IterationParameters,
-					ParticleParameters,
-					Constraint.Manifold.Restitution,
-					Body0, 
-					Body1,
-					ManifoldPoint);
-
-				Constraint.AccumulatedImpulse += ManifoldPoint.NetImpulse;
-			}
-
-			const FReal StaticFriction = FMath::Max(Constraint.Manifold.Friction, Constraint.Manifold.AngularFriction);
-			const FReal DynamicFriction = Constraint.Manifold.Friction;
-			if (StaticFriction > 0.0f)
-			{
-				for (int32 PointIndex = 0; PointIndex < ManifoldPoints.Num(); ++PointIndex)
-				{
-					FManifoldPoint& ManifoldPoint = Constraint.SetActiveManifoldPoint(PointIndex, Body0.P(), Body0.Q(), Body1.P(), Body1.Q());
-
-					CalculateManifoldFrictionVelocityCorrection(
-						Stiffness,
-						IterationParameters,
-						ParticleParameters,
-						StaticFriction,
-						DynamicFriction,
-						Body0, 
-						Body1,
-						ManifoldPoint);
-
-					Constraint.AccumulatedImpulse += ManifoldPoint.NetImpulse;
-				}
-			}
-
-			// Apply a position correction so that we get the desired velocity in the implicit velocity step
-			if (Body0.IsDynamic())
-			{
-				const FVec3 DV0 = Body0.V() - InitialV0;
-				const FVec3 DW0 = Body0.W() - InitialW0;
-				const FVec3 DX0 = DV0 * IterationParameters.Dt;
-				const FVec3 DR0 = DW0 * IterationParameters.Dt;
-				Body0.ApplyTransformDelta(DX0, DR0);
-				Body0.UpdateRotationDependentState();
-			}
-			if (Body1.IsDynamic())
-			{
-				const FVec3 DV1 = Body1.V() - InitialV1;
-				const FVec3 DW1 = Body1.W() - InitialW1;
-				const FVec3 DX1 = DV1 * IterationParameters.Dt;
-				const FVec3 DR1 = DW1 * IterationParameters.Dt;
-				Body1.ApplyTransformDelta(DX1, DR1);
-				Body1.UpdateRotationDependentState();
-			}
-#endif
-		}
-
-
-		// Pushout solver loop for a single contafct manifold.
-		void ApplyPushOutManifold(
-			FPBDCollisionConstraint& Constraint,
-			const FContactIterationParameters& IterationParameters,
-			const FContactParticleParameters& ParticleParameters)
-		{
-#if 0
-			TArrayView<FManifoldPoint> ManifoldPoints = Constraint.GetManifoldPoints();
-
-			FConstraintSolverBody Body0 = *Constraint.GetSolverBody0();
-			FConstraintSolverBody Body1 = *Constraint.GetSolverBody1();
-
-			// How much shock propagation should we apply - currently 100% on final iteration, otherwise none
-			FReal ShockPropagationInvMassScale = (IterationParameters.Iteration < (IterationParameters.NumIterations - 1)) ? 1.0f : 0.0f;
-			if ((ShockPropagationInvMassScale < 1.0f) && Chaos_Collision_UseShockPropagation)
-			{
-				// Lock one of the bodies if want shock propagation and they are at different shock levels
-				if (Body0.Level() < Body1.Level())
-				{
-					// Body0 is lower level and should be frozen
-					Body0.SetInvMassScale(ShockPropagationInvMassScale);
-				}
-				else if (Body0.Level() > Body1.Level())
-				{
-					// Body1 is lower level and should be frozen
-					Body1.SetInvMassScale(ShockPropagationInvMassScale);
-				}
-				else
-				{
-					// Both bodies are at the same level so neither is frozen
-				}
-			}
-
-			// If both bodies are frozen or kinematic just skip
-			// @todo(chaos): this should not happen - if one body is kinematic, the other should have no mass scaling. Check and remove
-			if (!Body0.IsDynamic() && !Body1.IsDynamic())
-			{
-				return;
-			}
-
-			// Gradually increase position correction through iterations (optional based on cvars)
-			const FReal Stiffness = GetPushOutStiffness(IterationParameters.Iteration, IterationParameters.NumIterations);
-
-			// Apply the position correction so that all contacts have zero separation
-			if (Chaos_Manifold_PushOut_PositionCorrection)
-			{
-				for (int32 PointIndex = 0; PointIndex < ManifoldPoints.Num(); ++PointIndex)
-				{
-					FManifoldPoint& ManifoldPoint = Constraint.SetActiveManifoldPoint(PointIndex, Body0.P(), Body0.Q(), Body1.P(), Body1.Q());
-
-					// The distance we should pushout to
-					FReal PushOutSeparation = 0.0f;
-
-					if (ManifoldPoint.bActive || (ManifoldPoint.ContactPoint.Phi < PushOutSeparation))
-					{
-						ApplyManifoldPushOutCorrection(
-							Stiffness,
-							Constraint,
-							IterationParameters,
-							ParticleParameters,
-							Body0,
-							Body1,
-							ManifoldPoint);
-					}
-				}
-			}
-
-			// Apply a velocity correction so that all contacts have non-negative contact velocity
-			if (Chaos_Manifold_PushOut_VelocityCorrection > 0)
-			{
-				const FReal StaticFriction = FMath::Max(Constraint.Manifold.Friction, Constraint.Manifold.AngularFriction);
-				const FReal DynamicFriction = Constraint.Manifold.Friction;
-
-				for (int32 PointIndex = 0; PointIndex < ManifoldPoints.Num(); ++PointIndex)
-				{
-					FManifoldPoint& ManifoldPoint = Constraint.SetActiveManifoldPoint(PointIndex, Body0.P(), Body0.Q(), Body1.P(), Body1.Q());
-
-					if (ManifoldPoint.bActive || !ManifoldPoint.NetPushOut.IsNearlyZero(KINDA_SMALL_NUMBER))
-					{
-						CalculateManifoldPushOutVelocityCorrection(
-							Stiffness,
-							IterationParameters,
-							ParticleParameters,
-							Constraint.Manifold.Restitution,
-							StaticFriction,
-							DynamicFriction,
-							Body0,
-							Body1,
-							ManifoldPoint);
-					}
-
-					//Constraint.AccumulatedImpulse += ManifoldPoint.NetPushOutImpulseNormal * ManifoldPoint.ContactPoint.Normal;
-				}
-			}
-#endif
-		}
-
-	}
-}
\ No newline at end of file
+/**
+ * @brief This file used to contain collision solver code for the GbfPbd solver (which was no longer functional)
+ * @see PBDCollisionSolver.cpp for the new solver
+*/
diff --git a/Engine/Source/Runtime/Experimental/Chaos/Private/Chaos/Collision/CollisionSolver.h b/Engine/Source/Runtime/Experimental/Chaos/Private/Chaos/Collision/CollisionSolver.h
--- a/Engine/Source/Runtime/Experimental/Chaos/Private/Chaos/Collision/CollisionSolver.h
+++ b/Engine/Source/Runtime/Experimental/Chaos/Private/Chaos/Collision/CollisionSolver.h
@@ -1,24 +1,8 @@
 // Copyright Epic Games, Inc. All Rights Reserved.
 #pragma once
 
-#include "Chaos/ExternalCollisionData.h"
-#include "Chaos/PBDCollisionConstraints.h"
-#include "Chaos/CollisionResolutionTypes.h"
-#include "Chaos/Collision/CollisionApplyType.h"
+/**
+ * @brief This file used to contain collision solver code for the GbfPbd solver (which was no longer functional)
+ * @see PBDCollisionSolver.h for the new solver
+*/
 
-namespace Chaos
-{
-	namespace Collisions
-	{
-		void ApplyContactManifold(
-			FPBDCollisionConstraint& Constraint,
-			const FContactIterationParameters& IterationParameters,
-			const FContactParticleParameters& ParticleParameters);
-
-		void ApplyPushOutManifold(
-			FPBDCollisionConstraint& Constraint,
-			const FContactIterationParameters& IterationParameters,
-			const FContactParticleParameters& ParticleParameters);
-
-	}
-}
diff --git a/Engine/Source/Runtime/Experimental/Chaos/Private/Chaos/Collision/PBDCollisionConstraint.cpp b/Engine/Source/Runtime/Experimental/Chaos/Private/Chaos/Collision/PBDCollisionConstraint.cpp
--- a/Engine/Source/Runtime/Experimental/Chaos/Private/Chaos/Collision/PBDCollisionConstraint.cpp
+++ b/Engine/Source/Runtime/Experimental/Chaos/Private/Chaos/Collision/PBDCollisionConstraint.cpp
@@ -176,7 +176,7 @@ namespace Chaos
 		, Implicit{ nullptr, nullptr }
 		, Shape{ nullptr, nullptr }
 		, Simplicial{ nullptr, nullptr }
-		, Manifold()
+		, Material()
 		, Stiffness(1)
 		, AccumulatedImpulse(0)
 		, TimeOfImpact(0)
@@ -215,7 +215,7 @@ namespace Chaos
 		, Implicit{ Implicit0, Implicit1 }
 		, Shape{ Shape0, Shape1 }
 		, Simplicial{ Simplicial0, Simplicial1 }
-		, Manifold()
+		, Material()
 		, Stiffness(1)
 		, AccumulatedImpulse(0)
 		, TimeOfImpact(0)
@@ -506,7 +506,7 @@ namespace Chaos
 
 		Flags.bDisabled = false;
 		ClosestManifoldPointIndex = INDEX_NONE;
-		Manifold.Reset();
+		Material.Reset();
 
 		for (int32 ManifoldPointIndex = 0; ManifoldPointIndex < ManifoldPoints.Num(); ManifoldPointIndex++)
 		{
@@ -573,7 +573,7 @@ namespace Chaos
 	void FPBDCollisionConstraint::ResetActiveManifoldContacts()
 	{
 		ClosestManifoldPointIndex = INDEX_NONE;
-		Manifold.Reset();
+		Material.Reset();
 		ManifoldPoints.Reset();
 		ExpectedNumManifoldPoints = 0;
 		Flags.bDisabled = false;
@@ -591,7 +591,7 @@ namespace Chaos
 		// Reset current closest point
 		ClosestManifoldPointIndex = INDEX_NONE;
 		Flags.bDisabled = false;
-		Manifold.Reset();
+		Material.Reset();
 
 		// How many manifold points we expect. E.g., for Box-box this will be 4 or 1 depending on whether
 		// we have a face or edge contact. We don't reuse the manifold if we lose points after culling
diff --git a/Engine/Source/Runtime/Experimental/Chaos/Private/Chaos/PBDCollisionConstraints.cpp b/Engine/Source/Runtime/Experimental/Chaos/Private/Chaos/PBDCollisionConstraints.cpp
--- a/Engine/Source/Runtime/Experimental/Chaos/Private/Chaos/PBDCollisionConstraints.cpp
+++ b/Engine/Source/Runtime/Experimental/Chaos/Private/Chaos/PBDCollisionConstraints.cpp
@@ -186,59 +186,58 @@ namespace Chaos
 		const FChaosPhysicsMaterial* PhysicsMaterial0 = GetPhysicsMaterial(Constraint.Particle[0], Constraint.Implicit[0], MPhysicsMaterials, MPerParticlePhysicsMaterials, SimMaterials);
 		const FChaosPhysicsMaterial* PhysicsMaterial1 = GetPhysicsMaterial(Constraint.Particle[1], Constraint.Implicit[1], MPhysicsMaterials, MPerParticlePhysicsMaterials, SimMaterials);
 
-		FCollisionContact& Contact = Constraint.Manifold;
+		FPBDCollisionConstraintMaterial& CollisionMaterial = Constraint.Material;
 		if (PhysicsMaterial0 && PhysicsMaterial1)
 		{
-			const FChaosPhysicsMaterial::ECombineMode RestitutionCombineMode = FChaosPhysicsMaterial::ChooseCombineMode(PhysicsMaterial0->RestitutionCombineMode,PhysicsMaterial1->RestitutionCombineMode);
-			Contact.Restitution = FChaosPhysicsMaterial::CombineHelper(PhysicsMaterial0->Restitution, PhysicsMaterial1->Restitution, RestitutionCombineMode);
+			const FChaosPhysicsMaterial::ECombineMode RestitutionCombineMode = FChaosPhysicsMaterial::ChooseCombineMode(PhysicsMaterial0->RestitutionCombineMode, PhysicsMaterial1->RestitutionCombineMode);
+			CollisionMaterial.Restitution = FChaosPhysicsMaterial::CombineHelper(PhysicsMaterial0->Restitution, PhysicsMaterial1->Restitution, RestitutionCombineMode);
 
-			const FChaosPhysicsMaterial::ECombineMode FrictionCombineMode = FChaosPhysicsMaterial::ChooseCombineMode(PhysicsMaterial0->FrictionCombineMode,PhysicsMaterial1->FrictionCombineMode);
-			Contact.Friction = FChaosPhysicsMaterial::CombineHelper(PhysicsMaterial0->Friction,PhysicsMaterial1->Friction, FrictionCombineMode);
+			const FChaosPhysicsMaterial::ECombineMode FrictionCombineMode = FChaosPhysicsMaterial::ChooseCombineMode(PhysicsMaterial0->FrictionCombineMode, PhysicsMaterial1->FrictionCombineMode);
+			CollisionMaterial.Friction = FChaosPhysicsMaterial::CombineHelper(PhysicsMaterial0->Friction, PhysicsMaterial1->Friction, FrictionCombineMode);
 			const FReal StaticFriction0 = FMath::Max(PhysicsMaterial0->Friction, PhysicsMaterial0->StaticFriction);
 			const FReal StaticFriction1 = FMath::Max(PhysicsMaterial1->Friction, PhysicsMaterial1->StaticFriction);
-			Contact.AngularFriction = FChaosPhysicsMaterial::CombineHelper(StaticFriction0, StaticFriction1, FrictionCombineMode);
+			CollisionMaterial.AngularFriction = FChaosPhysicsMaterial::CombineHelper(StaticFriction0, StaticFriction1, FrictionCombineMode);
 		}
 		else if (PhysicsMaterial0)
 		{
 			const FReal StaticFriction0 = FMath::Max(PhysicsMaterial0->Friction, PhysicsMaterial0->StaticFriction);
-			Contact.Restitution = PhysicsMaterial0->Restitution;
-			Contact.Friction = PhysicsMaterial0->Friction;
-			Contact.AngularFriction = StaticFriction0;
+			CollisionMaterial.Restitution = PhysicsMaterial0->Restitution;
+			CollisionMaterial.Friction = PhysicsMaterial0->Friction;
+			CollisionMaterial.AngularFriction = StaticFriction0;
 		}
 		else if (PhysicsMaterial1)
 		{
 			const FReal StaticFriction1 = FMath::Max(PhysicsMaterial1->Friction, PhysicsMaterial1->StaticFriction);
-			Contact.Restitution = PhysicsMaterial1->Restitution;
-			Contact.Friction = PhysicsMaterial1->Friction;
-			Contact.AngularFriction = StaticFriction1;
+			CollisionMaterial.Restitution = PhysicsMaterial1->Restitution;
+			CollisionMaterial.Friction = PhysicsMaterial1->Friction;
+			CollisionMaterial.AngularFriction = StaticFriction1;
 		}
 		else
 		{
-			Contact.Friction = DefaultCollisionFriction;
-			Contact.AngularFriction = DefaultCollisionFriction;
-			Contact.Restitution = DefaultCollisionRestitution;
+			CollisionMaterial.Friction = DefaultCollisionFriction;
+			CollisionMaterial.AngularFriction = DefaultCollisionFriction;
+			CollisionMaterial.Restitution = DefaultCollisionRestitution;
 		}
 
-		Contact.RestitutionThreshold = (CollisionRestitutionThresholdOverride >= 0.0f) ? CollisionRestitutionThresholdOverride : RestitutionThreshold;
-
-		if (!bEnableRestitution)
-		{
-			Contact.Restitution = 0.0f;
-		}
+		CollisionMaterial.RestitutionThreshold = (CollisionRestitutionThresholdOverride >= 0.0f) ? CollisionRestitutionThresholdOverride : RestitutionThreshold;
 
 		// Overrides for testing
 		if (CollisionFrictionOverride >= 0)
 		{
-			Contact.Friction = CollisionFrictionOverride;
-			Contact.AngularFriction = CollisionFrictionOverride;
+			CollisionMaterial.Friction = CollisionFrictionOverride;
+			CollisionMaterial.AngularFriction = CollisionFrictionOverride;
 		}
 		if (CollisionRestitutionOverride >= 0)
 		{
-			Contact.Restitution = CollisionRestitutionOverride;
+			CollisionMaterial.Restitution = CollisionRestitutionOverride;
 		}
 		if (CollisionAngularFrictionOverride >= 0)
 		{
-			Contact.AngularFriction = CollisionAngularFrictionOverride;
+			CollisionMaterial.AngularFriction = CollisionAngularFrictionOverride;
+		}
+		if (!bEnableRestitution)
+		{
+			CollisionMaterial.Restitution = 0.0f;
 		}
 	}
 
@@ -524,10 +523,8 @@ namespace Chaos
 			FPBDCollisionSolverContainer& SolverContainer = GetConstraintSolverContainer(SolverData);
 			return SolverContainer.SolveVelocitySerial(Dt, It, NumIts, 0, SolverContainer.NumSolvers(), SolverSettings);
 		}
-		else
-		{
-			return LegacyApplyPhase2Serial(Dt, It, NumIts, 0, SolverData.GetConstraintHandles(ContainerId).Num(), SolverData);
-		}
+
+		return false;
 	}
 
 	// Island Rule version
@@ -540,10 +537,8 @@ namespace Chaos
 			FPBDCollisionSolverContainer& SolverContainer = GetConstraintSolverContainer(SolverData);
 			return SolverContainer.SolveVelocitySerial(Dt, It, NumIts, 0, SolverContainer.NumSolvers(), SolverSettings);
 		}
-		else
-		{
-			return LegacyApplyPhase2Serial(Dt, It, NumIts, 0, SolverData.GetConstraintHandles(ContainerId).Num(), SolverData);
-		}
+
+		return false;
 	}
 
 	// Color Rule version
@@ -556,10 +551,8 @@ namespace Chaos
 			FPBDCollisionSolverContainer& SolverContainer = GetConstraintSolverContainer(SolverData);
 			return SolverContainer.SolveVelocitySerial(Dt, It, NumIts, BeginIndex, EndIndex, SolverSettings);
 		}
-		else
-		{
-			return LegacyApplyPhase2Serial(Dt, It, NumIts, BeginIndex, EndIndex, SolverData);
-		}
+
+		return false;
 	}
 
 	// Color Rule version
@@ -572,10 +565,8 @@ namespace Chaos
 			FPBDCollisionSolverContainer& SolverContainer = GetConstraintSolverContainer(SolverData);
 			return SolverContainer.SolveVelocityParallel(Dt, It, NumIts, BeginIndex, EndIndex, SolverSettings);
 		}
-		else
-		{
-			return LegacyApplyPhase2Parallel(Dt, It, NumIts, BeginIndex, EndIndex, SolverData);
-		}
+
+		return false;
 	}
 
 	void FPBDCollisionConstraints::LegacyGatherInput(const FReal Dt, FPBDCollisionConstraint& Constraint, const int32 Particle0Level, const int32 Particle1Level, FPBDIslandSolverData& SolverData)
@@ -629,31 +620,6 @@ namespace Chaos
 		return LegacyApplyPhase1Serial(Dt, Iterations, NumIterations, BeginIndex, EndIndex, SolverData);
 	}
 
-	bool FPBDCollisionConstraints::LegacyApplyPhase2Serial(const FReal Dt, const int32 Iterations, const int32 NumIterations, const int32 BeginIndex, const int32 EndIndex, FPBDIslandSolverData& SolverData)
-	{
-		bool bNeedsAnotherIteration = false;
-		if (MApplyPushOutPairIterations > 0)
-		{
-			const Collisions::FContactParticleParameters ParticleParameters = GetContactParticleParameters(Dt);
-			const Collisions::FContactIterationParameters IterationParameters = GetContactIterationParameters(Dt, Iterations, NumIterations, MApplyPushOutPairIterations, bNeedsAnotherIteration);
-
-			for (int32 Index = BeginIndex; Index < EndIndex; ++Index)
-			{
-				FPBDCollisionConstraint* Constraint = SolverData.GetConstraintHandle<FPBDCollisionConstraint>(ContainerId,Index);
-				if (!Constraint->GetDisabled())
-				{
-					Collisions::ApplyPushOut(*Constraint, IterationParameters, ParticleParameters);
-				}
-			}
-		}
-		return bNeedsAnotherIteration;
-	}
-
-	bool FPBDCollisionConstraints::LegacyApplyPhase2Parallel(const FReal Dt, const int32 Iterations, const int32 NumIterations, const int32 BeginIndex, const int32 EndIndex, FPBDIslandSolverData& SolverData)
-	{
-		return LegacyApplyPhase2Serial(Dt,  Iterations, NumIterations, BeginIndex, EndIndex, SolverData);
-	}
-
 	const FPBDCollisionConstraint& FPBDCollisionConstraints::GetConstraint(int32 Index) const
 	{
 		check(Index < NumConstraints());
diff --git a/Engine/Source/Runtime/Experimental/Chaos/Private/Chaos/PBDCollisionConstraintsContact.cpp b/Engine/Source/Runtime/Experimental/Chaos/Private/Chaos/PBDCollisionConstraintsContact.cpp
--- a/Engine/Source/Runtime/Experimental/Chaos/Private/Chaos/PBDCollisionConstraintsContact.cpp
+++ b/Engine/Source/Runtime/Experimental/Chaos/Private/Chaos/PBDCollisionConstraintsContact.cpp
@@ -13,39 +13,13 @@
 #include "PBDCollisionConstraints.ispc.generated.h"
 #endif
 
-//PRAGMA_DISABLE_OPTIMIZATION
+/**
+ * @brief This file contains collision solver code for the StandardPbd solver (which will be removed at some point)
+ * @see PBDCollisionSolver.cpp for the new solver
+*/
 
 namespace Chaos
 {
-	namespace CVars
-	{
-		int32 Chaos_Collision_EnergyClampEnabled = 1;
-		FAutoConsoleVariableRef CVarChaosCollisionEnergyClampEnabled(TEXT("p.Chaos.Collision.EnergyClampEnabled"), Chaos_Collision_EnergyClampEnabled, TEXT("Whether to use energy clamping in collision apply step"));
-
-		int32 Chaos_Collision_RelaxationEnabled = 0; // TODO remove this feature soon
-		FAutoConsoleVariableRef CVarChaosCollisionRelaxationEnabled(TEXT("p.Chaos.Collision.RelaxationEnabled"), Chaos_Collision_RelaxationEnabled, TEXT("Whether to reduce applied impulses during iterations for improved solver stability but reduced convergence"));
-
-		int32 Chaos_Collision_PrevVelocityRestitutionEnabled = 0;
-		FAutoConsoleVariableRef CVarChaosCollisionPrevVelocityRestitutionEnabled(TEXT("p.Chaos.Collision.PrevVelocityRestitutionEnabled"), Chaos_Collision_PrevVelocityRestitutionEnabled, TEXT("If enabled restitution will be calculated on previous frame velocities instead of current frame velocities"));
-
-		FRealSingle Chaos_Collision_ContactMovementAllowance = 0.05f;
-		FAutoConsoleVariableRef CVarChaosCollisionContactMovementAllowance(TEXT("p.Chaos.Collision.AntiJitterContactMovementAllowance"), Chaos_Collision_ContactMovementAllowance,
-			TEXT("If a contact is close to where it was during a previous iteration, we will assume it is the same contact that moved (to reduce jitter). Expressed as the fraction of movement distance and Centre of Mass distance to the contact point"));
-
-		int32 Chaos_Collision_UseAccumulatedImpulseClipSolve = 0; // This requires multiple contact points per iteration per pair and contact points that don't move too much (in body space) to have an effect
-		FAutoConsoleVariableRef CVarChaosCollisionImpulseClipSolve(TEXT("p.Chaos.Collision.UseAccumulatedImpulseClipSolve"), Chaos_Collision_UseAccumulatedImpulseClipSolve, TEXT("Use experimental Accumulated impulse clipped contact solve"));
-
-		int32 Chaos_Collision_UseShockPropagation = 1;
-		FAutoConsoleVariableRef CVarChaosCollisionUseShockPropagation(TEXT("p.Chaos.Collision.UseShockPropagation"), Chaos_Collision_UseShockPropagation, TEXT(""));
-
-		FRealSingle Chaos_Collision_CollisionClipTolerance = 0.01f;
-		FAutoConsoleVariableRef CVarChaosCollisionClipTolerance(TEXT("p.Chaos.Collision.ClipTolerance"), Chaos_Collision_CollisionClipTolerance, TEXT(""));
-
-		bool Chaos_Collision_CheckManifoldComplete = false;
-		FAutoConsoleVariableRef CVarChaosCollisionCheckManifoldComplete(TEXT("p.Chaos.Collision.CheckManifoldComplete"), Chaos_Collision_CheckManifoldComplete, TEXT(""));
-	}
-	using namespace CVars;
-
 	namespace Collisions
 	{
 		void Update(FPBDCollisionConstraint& Constraint, const FReal Dt)
@@ -61,23 +35,6 @@ namespace Chaos
 			UpdateConstraintFromGeometry<ECollisionUpdateType::Deepest>(Constraint, Transform0, Transform1, Dt);
 		}
 
-		void UpdateSwept(FPBDCollisionConstraint& Constraint, const FReal Dt)
-		{
-			// Note: This is unusual but we are using a mix of the previous and current transform
-			// This is due to how CCD rewinds the position (not rotation) and then sweeps to find the first contact at the current orientation
-			// NOTE: These are actor transforms, not CoM transforms
-			// \todo(chaos): see if we can easily switch to CoM transforms now in collision loop (shapes are held in actor space)
-			const FSolverBody& Body0 = *Constraint.GetSolverBody0();
-			const FSolverBody& Body1 = *Constraint.GetSolverBody1();
-			FGenericParticleHandle Particle0 = FGenericParticleHandle(Constraint.GetParticle0());
-			const FRigidTransform3 TransformXQ0(Body0.X() - Body0.ActorQ().RotateVector(Particle0->CenterOfMass()), Body0.ActorQ());
-			const FRigidTransform3 Transform1 = FRigidTransform3(Body1.ActorP(), Body1.ActorQ());
-
-			Constraint.ResetPhi(Constraint.GetCullDistance());
-			UpdateConstraintFromGeometrySwept<ECollisionUpdateType::Deepest>(Constraint, TransformXQ0, Transform1, Dt);
-		}
-
-
 		// A PBD collision penetration correction.
 		// Currently only used by RBAN
 		FVec3 ApplyContact2(FPBDCollisionConstraint& Constraint,
@@ -201,9 +158,10 @@ namespace Chaos
 		}
 
 
-		template<typename T_CONSTRAINT>
-		void ApplyImpl(T_CONSTRAINT& Constraint, const FContactIterationParameters & IterationParameters, const FContactParticleParameters & ParticleParameters)
+		void Apply(FPBDCollisionConstraint& Constraint, const FContactIterationParameters& IterationParameters, const FContactParticleParameters& ParticleParameters)
 		{
+			check(IterationParameters.SolverType == EConstraintSolverType::StandardPbd);
+
 			for (int32 PairIt = 0; PairIt < IterationParameters.NumPairIterations; ++PairIt)
 			{
 				// Update the contact information based on current particles' positions
@@ -228,7 +186,7 @@ namespace Chaos
 
 				// Do not early out here in the case of Accumulated impulse solve
 				// @todo(chaos): remove this early out when we settle on manifolds
-				const bool bIsAccumulatingImpulses = Constraint.GetUseManifold() || Chaos_Collision_UseAccumulatedImpulseClipSolve;
+				const bool bIsAccumulatingImpulses = Constraint.GetUseManifold();
 				if (Constraint.GetPhi() >= 0.0f && !bIsAccumulatingImpulses)
 				{
 					return;
@@ -245,154 +203,10 @@ namespace Chaos
 					Particle1->AuxilaryValue(*ParticleParameters.Collided) = true;
 				}
 
-				// What solver algorithm should we use?
-				switch (IterationParameters.SolverType)
-				{
-				case EConstraintSolverType::GbfPbd:
-					{
-						ensure(false);	// not currently working
-						ApplyContactManifold(Constraint, IterationParameters, ParticleParameters);
-					}
-					break;
-				case EConstraintSolverType::StandardPbd:
-					{
-						Constraint.AccumulatedImpulse += ApplyContact2(Constraint, IterationParameters, ParticleParameters);
-					}
-					break;
-				case EConstraintSolverType::QuasiPbd:
-					{
-						check(false);	// does not use this path
-					}
-					break;
-				default:
-					break;
-				}
-			}
-		}
-		
-		void ApplySweptImpl(FPBDCollisionConstraint& Constraint, const FContactIterationParameters & IterationParameters, const FContactParticleParameters & ParticleParameters)
-		{
-			check(Constraint.GetCCDType() == ECollisionCCDType::Enabled);
-
-			FSolverBody& Body0 = *Constraint.GetSolverBody0();
-			FSolverBody& Body1 = *Constraint.GetSolverBody1();
-
-			Collisions::UpdateSwept(Constraint, IterationParameters.Dt);
-			
-			const FContactParticleParameters CCDParticleParamaters{ ParticleParameters.RestitutionVelocityThreshold, true, ParticleParameters.Collided };
-			if (Constraint.TimeOfImpact >= 1)
-			{
-				// If TOI = 1 (normal constraint) we don't want to split timestep at TOI.
-				ApplyImpl(Constraint, IterationParameters, CCDParticleParamaters);
-				return;
-			}
-
-			// Rebuild iteration params with partial dt, and non-zero iteration count to force update of constraint.
-			// P may have changed due to other constraints, so at TOI our manifold needs updating.
-			const FReal PartialDT = Constraint.TimeOfImpact * IterationParameters.Dt;
-			const FReal RemainingDT = (1 - Constraint.TimeOfImpact) * IterationParameters.Dt;
-			const int32 FakeIteration = IterationParameters.NumIterations / 2; // For iteration count dependent effects (like relaxation) // @todo: Do we still need this?
-			const int32 PartialPairIterations = FMath::Max(IterationParameters.NumPairIterations, 2); // Do at least 2 pair iterations // @todo: Do we still need this?
-			const FContactIterationParameters IterationParametersPartialDT{ PartialDT, FakeIteration, IterationParameters.NumIterations, PartialPairIterations, IterationParameters.SolverType, IterationParameters.NeedsAnotherIteration };
-			const FContactIterationParameters IterationParametersRemainingDT{ RemainingDT, FakeIteration, IterationParameters.NumIterations, IterationParameters.NumPairIterations, IterationParameters.SolverType, IterationParameters.NeedsAnotherIteration };
-
-			// Rewind P to TOI and Apply
-			Body0.SetP(FMath::Lerp(Body0.X(), Body0.P(), Constraint.TimeOfImpact));
-			ApplyImpl(Constraint, IterationParametersPartialDT, CCDParticleParamaters);
-
-			// @todo(chaos): Make this work properly for real Stanard and QPBD implementations (they do not alter velocity)
-			if (IterationParameters.SolverType == EConstraintSolverType::GbfPbd)
-			{
-				// Advance P to end of frame from TOI, and Apply
-				if (CCDAlwaysSweepRemainingDT || IterationParameters.Iteration + 1 < IterationParameters.NumIterations)
-				{
-					Body0.SetP(Body0.P() + Body0.V() * RemainingDT); // If we are tunneling through something else due to this, it will be resolved in the next iteration
-					ApplyImpl(Constraint, IterationParametersRemainingDT, CCDParticleParamaters);
-				}
-				else
-				{
-					// We get here if we cannot solve CCD collisions with the given number of iterations and restitution settings.
-					// So don't do the remaining dt update. This will bleed the energy! (also: Ignore rotation)
-					// To prevent this condition: increase number of iterations and/or reduce restitution and/or reduce velocities
-					if (IterationParameters.Dt > SMALL_NUMBER)
-					{
-						// Update velocity to be consistent with PBD
-						const FReal InvDt = FReal(1) / IterationParameters.Dt;
-						Body0.SetV((Body0.P() - Body0.X()) * InvDt);
-					}
-					else
-					{
-						Body0.SetV(FVec3(0));
-					}
-				}
-			}
-		}
-
-
-		void Apply(FPBDCollisionConstraint& Constraint, const FContactIterationParameters & IterationParameters, const FContactParticleParameters & ParticleParameters)
-		{
-			if (Constraint.GetCCDType() == ECollisionCCDType::Disabled)
-			{
-				ApplyImpl(Constraint, IterationParameters, ParticleParameters);
-			}
-			else if (Constraint.GetCCDType() == ECollisionCCDType::Enabled)
-			{
-				ApplySweptImpl(Constraint, IterationParameters, ParticleParameters);
+				Constraint.AccumulatedImpulse += ApplyContact2(Constraint, IterationParameters, ParticleParameters);
 			}
 		}
 
-
-		void ApplyPushOutImpl(FPBDCollisionConstraint& Constraint, const FContactIterationParameters & IterationParameters, const FContactParticleParameters & ParticleParameters)
-		{
-			for (int32 PairIt = 0; PairIt < IterationParameters.NumPairIterations; ++PairIt)
-			{
-				Update(Constraint, IterationParameters.Dt);
-
-				// Ignore contacts where the closest point is greater than cull distance
-				if (Constraint.GetPhi() >= Constraint.GetCullDistance())
-				{
-					// Optionally permanently disable the contact for the remaining iterations
-					if (ParticleParameters.bCanDisableContacts)
-					{
-						Constraint.SetDisabled(true);
-					}
-					return;
-				}
-
-				// @todo(chaos): is this supposed to run for swept contacts as well?
-				if (Constraint.GetCCDType() == ECollisionCCDType::Disabled)
-				{
-					switch (IterationParameters.SolverType)
-					{
-					case EConstraintSolverType::GbfPbd:
-						{
-							ensure(false);	// not currently working
-							ApplyPushOutManifold(Constraint, IterationParameters, ParticleParameters);
-						}
-						break;
-					case EConstraintSolverType::StandardPbd:
-						{
-							// Shouldn't have PushOut for Standard Pbd, but this is for experimentation (Collision PushOut Iterations should normally be set to 0 instead)
-							ApplyPushOutManifold(Constraint, IterationParameters, ParticleParameters);
-						}
-						break;
-					case EConstraintSolverType::QuasiPbd:
-						{
-							check(false);
-						}
-						break;
-					default:
-						break;
-					}
-				}
-			}
-		}
-
-		void ApplyPushOut(FPBDCollisionConstraint& Constraint, const FContactIterationParameters & IterationParameters, const FContactParticleParameters & ParticleParameters)
-		{
-			ApplyPushOutImpl(Constraint, IterationParameters, ParticleParameters);
-		}
-
 	} // Collisions
 
 }// Chaos
diff --git a/Engine/Source/Runtime/Experimental/Chaos/Private/Chaos/PBDCollisionConstraintsPGS.cpp b/Engine/Source/Runtime/Experimental/Chaos/Private/Chaos/PBDCollisionConstraintsPGS.cpp
--- a/Engine/Source/Runtime/Experimental/Chaos/Private/Chaos/PBDCollisionConstraintsPGS.cpp
+++ b/Engine/Source/Runtime/Experimental/Chaos/Private/Chaos/PBDCollisionConstraintsPGS.cpp
@@ -1,1295 +1,2 @@
 // Copyright Epic Games, Inc. All Rights Reserved.
-
 #include "Chaos/PBDCollisionConstraintsPGS.h"
-
-#include "Chaos/BoundingVolume.h"
-#include "Chaos/BoundingVolumeHierarchy.h"
-#include "Chaos/Defines.h"
-#include "Chaos/Pair.h"
-#include "Chaos/Sphere.h"
-#include "Chaos/Transform.h"
-
-#include "ProfilingDebugging/ScopedTimers.h"
-
-#if CHAOS_PARTICLEHANDLE_TODO
-
-using namespace Chaos;
-
-template<class T_PARTICLES>
-FVec3 GetTranslationPGS(const T_PARTICLES& InParticles)
-{
-	check(false);
-	return FRigidTransform3();
-}
-
-FVec3 GetTranslationPGS(const TRigidParticles<FReal,3>& InParticles, const int32 Index)
-{
-	return InParticles.X(Index);
-}
-
-FVec3 GetTranslationPGS(const FPBDRigidParticles& InParticles, const int32 Index)
-{
-	return InParticles.P(Index);
-}
-
-template<class T_PARTICLES>
-FRotation3 GetRotationPGS(const T_PARTICLES& InParticles)
-{
-	check(false);
-	return FRotation3();
-}
-
-FRotation3 GetRotationPGS(const TRigidParticles<FReal, 3>& InParticles, const int32 Index)
-{
-	return InParticles.R(Index);
-}
-
-FRotation3 GetRotationPGS(const FPBDRigidParticles& InParticles, const int32 Index)
-{
-	return InParticles.Q(Index);
-}
-
-FRigidTransform3 GetTransformPGS(const T_PARTICLES& InParticles)
-{
-	check(false);
-	return FRigidTransform3();
-}
-
-FRigidTransform3 GetTransformPGS(const TRigidParticles<FReal, 3>& InParticles, const int32 Index)
-{
-	return FRigidTransform3(InParticles.X(Index), InParticles.R(Index));
-}
-
-FRigidTransform3 GetTransformPGS(const FPBDRigidParticles& InParticles, const int32 Index)
-{
-	return FRigidTransform3(InParticles.P(Index), InParticles.Q(Index));
-}
-
-FPBDCollisionConstraintPGS::FPBDCollisionConstraintPGS(FPBDRigidParticles& InParticles, const TArray<int32>& InIndices, TArrayCollectionArray<bool>& Collided, TArrayCollectionArray<TSerializablePtr<FChaosPhysicsMaterial>>& InPhysicsMaterials, const FReal Thickness /*= (FReal)0*/)
-    : MCollided(Collided), MPhysicsMaterials(InPhysicsMaterials), MThickness(Thickness), Tolerance(KINDA_SMALL_NUMBER), MaxIterations(10), bUseCCD(false)
-{
-}
-
-void FPBDCollisionConstraintPGS::ComputeConstraints(const FPBDRigidParticles& InParticles, const TArray<int32>& InIndices, const FReal Dt)
-{
-	double Time = 0;
-	FDurationTimer Timer(Time);
-	// Broad phase
-	//TBoundingVolumeHierarchy<FPBDRigidParticles, T, d> Hierarchy(InParticles);
-	TBoundingVolume<FPBDRigidParticles> Hierarchy(InParticles, true, Dt);
-	Timer.Stop();
-	UE_LOG(LogChaos, Verbose, TEXT("\tPBDCollisionConstraint Construct Hierarchy %f"), Time);
-	// Narrow phase
-	FCriticalSection CriticalSection;
-	Time = 0;
-	Timer.Start();
-	PhysicsParallelFor(InIndices.Num(), [&](int32 Index) {
-		const int32 ParticleIndex = InIndices[Index];
-
-		TArray<int32> PotentialIntersections;
-		FAABB3 Box1;
-		if (InParticles.Geometry(ParticleIndex)->HasBoundingBox())
-		{
-			Box1 = Hierarchy.GetWorldSpaceBoundingBox(InParticles, ParticleIndex);
-			PotentialIntersections = Hierarchy.FindAllIntersections(Box1);
-		}
-		else
-		{
-			PotentialIntersections = Hierarchy.GlobalObjects();
-		}
-
-		for (int32 i = 0; i < PotentialIntersections.Num(); ++i)
-		{
-			int32 Body2Index = PotentialIntersections[i];
-
-			// Collision group culling...
-			// CollisionGroup == 0 : Collide_With_Everything
-			// CollisionGroup == INDEX_NONE : Disabled collisions
-			// CollisionGroup_A != CollisionGroup_B : Skip Check
-			if(InParticles.Disabled(Body2Index))
-			{
-				continue;
-			}
-				
-			if (InParticles.CollisionGroup(ParticleIndex) == INDEX_NONE || InParticles.CollisionGroup(Body2Index) == INDEX_NONE)
-			{
-				continue;
-			}
-			if (InParticles.CollisionGroup(ParticleIndex) && InParticles.CollisionGroup(Body2Index) && InParticles.CollisionGroup(ParticleIndex) != InParticles.CollisionGroup(Body2Index))
-			{
-				continue;
-			}
-
-			if (InParticles.InvM(ParticleIndex) < FLT_MIN && InParticles.InvM(Body2Index) < FLT_MIN)
-			{
-				continue;
-			}
-			if (ParticleIndex == Body2Index || ((InParticles.Geometry(ParticleIndex)->HasBoundingBox() == InParticles.Geometry(Body2Index)->HasBoundingBox()) && Body2Index > ParticleIndex))
-			{
-				continue;
-			}
-			if (InParticles.Geometry(ParticleIndex)->HasBoundingBox() && InParticles.Geometry(Body2Index)->HasBoundingBox())
-			{
-				const auto& Box2 = Hierarchy.GetWorldSpaceBoundingBox(InParticles, Body2Index);
-				if (!Box1.Intersects(Box2))
-				{
-					continue;
-				}
-			}
-			auto Constraint = ComputeConstraint(InParticles, ParticleIndex, Body2Index, MThickness);
-			CriticalSection.Lock();
-			Constraints.Add(Constraint);
-			CriticalSection.Unlock();
-		}
-	});
-
-	Timer.Stop();
-
-	UE_LOG(LogChaos, Verbose, TEXT("\tPBDCollisionConstraint Construct %d Constraints with Potential Collisions %f"), Constraints.Num(), Time);
-}
-
-void FPBDCollisionConstraintPGS::RemoveConstraints(const TSet<uint32>& RemovedParticles)
-{
-	for (int32 i = 0; i < Constraints.Num(); ++i)
-	{
-		const auto& Constraint = Constraints[i];
-		if (RemovedParticles.Contains(Constraint.ParticleIndex) || RemovedParticles.Contains(Constraint.LevelsetIndex))
-		{
-			Constraints.RemoveAtSwap(i);
-			i--;
-		}
-	}
-}
-
-void FPBDCollisionConstraintPGS::UpdateConstraints(const FPBDRigidParticles& InParticles, T Dt, const TSet<uint32>& AddedParticles, const TArray<uint32>& ActiveParticles)
-{
-	double Time = 0;
-	FDurationTimer Timer(Time);
-
-	//
-	// Broad phase
-	//
-
-	// @todo(mlentine): We only need to construct the hierarchy for the islands we care about
-	TBoundingVolume<FPBDRigidParticles> Hierarchy(InParticles, ActiveParticles, true, Dt);
-	Timer.Stop();
-	UE_LOG(LogChaos, Verbose, TEXT("\tPBDCollisionConstraint Construct Hierarchy %f"), Time);
-
-	//
-	// Narrow phase
-	//
-
-	FCriticalSection CriticalSection;
-	Time = 0;
-	Timer.Start();
-	TArray<uint32> AddedParticlesArray = AddedParticles.Array();
-	PhysicsParallelFor(AddedParticlesArray.Num(), [&](int32 Index) {
-		int32 Body1Index = AddedParticlesArray[Index];
-		if (InParticles.Disabled(Body1Index))
-			return;
-		TArray<int32> PotentialIntersections;
-		FAABB3 Box1;
-		if (InParticles.Geometry(Body1Index)->HasBoundingBox())
-		{
-			Box1 = Hierarchy.GetWorldSpaceBoundingBox(InParticles, Body1Index);
-			PotentialIntersections = Hierarchy.FindAllIntersections(Box1);
-		}
-		else
-		{
-			PotentialIntersections = Hierarchy.GlobalObjects();
-		}
-		for (int32 i = 0; i < PotentialIntersections.Num(); ++i)
-		{
-			int32 Body2Index = PotentialIntersections[i];
-
-			if (Body1Index == Body2Index || ((InParticles.Geometry(Body1Index)->HasBoundingBox() == InParticles.Geometry(Body2Index)->HasBoundingBox()) && AddedParticles.Contains(Body2Index) && AddedParticles.Contains(Body1Index) && Body2Index > Body1Index))
-			{
-				continue;
-			}
-			if (InParticles.Geometry(Body1Index)->HasBoundingBox() && InParticles.Geometry(Body2Index)->HasBoundingBox())
-			{
-				const auto& Box2 = Hierarchy.GetWorldSpaceBoundingBox(InParticles, Body2Index);
-				if (!Box1.Intersects(Box2))
-				{
-					continue;
-				}
-			}
-			
-			if (InParticles.InvM(Body1Index) && InParticles.InvM(Body2Index) && (InParticles.Island(Body1Index) != InParticles.Island(Body2Index)))	//todo(ocohen): this is a hack - we should not even consider dynamics from other islands
-			{
-				continue;
-			}
-			auto Constraint = ComputeConstraint(InParticles, Body1Index, Body2Index, MThickness);
-			CriticalSection.Lock();
-			Constraints.Add(Constraint);
-			CriticalSection.Unlock();
-		}
-	});
-	Timer.Stop();
-	UE_LOG(LogChaos, Verbose, TEXT("\tPBDCollisionConstraint Update %d Constraints with Potential Collisions %f"), Constraints.Num(), Time);
-}
-
-template<class T_PARTICLES>
-void ComputePGSPropeties(const T_PARTICLES& InParticles, const FRigidBodyContactConstraintPGS& Constraint, const int32 PointIndex, const int32 FlattenedIndex,
-	const FMatrix33& WorldSpaceInvI1, const FMatrix33& WorldSpaceInvI2, const FVec3 Direction,
-	TArray<TVec2<FVec3>>& Angulars, TArray<TVec2<FVec3>>& MassWeightedAngulars, TArray<FReal>& Multipliers)
-{
-	FVec3 VectorToPoint1 = Constraint.Location[PointIndex] - GetTranslationPGS(InParticles, Constraint.ParticleIndex);
-	FVec3 VectorToPoint2 = Constraint.Location[PointIndex] - GetTranslationPGS(InParticles, Constraint.LevelsetIndex);
-	Angulars[FlattenedIndex][0] = -FVec3::CrossProduct(VectorToPoint1, Direction);
-	Angulars[FlattenedIndex][1] = FVec3::CrossProduct(VectorToPoint2, Direction);
-	MassWeightedAngulars[FlattenedIndex][0] = WorldSpaceInvI1 * Angulars[FlattenedIndex][0];
-	MassWeightedAngulars[FlattenedIndex][1] = WorldSpaceInvI2 * Angulars[FlattenedIndex][1];
-	if (InParticles.InvM(Constraint.ParticleIndex))
-	{
-		Multipliers[FlattenedIndex] += InParticles.InvM(Constraint.ParticleIndex) + FVec3::DotProduct(Angulars[FlattenedIndex][0], MassWeightedAngulars[FlattenedIndex][0]);
-	}
-	if (InParticles.InvM(Constraint.LevelsetIndex))
-	{
-		Multipliers[FlattenedIndex] += InParticles.InvM(Constraint.LevelsetIndex) + FVec3::DotProduct(Angulars[FlattenedIndex][1], MassWeightedAngulars[FlattenedIndex][1]);
-	}
-}
-
-template<class T_PARTICLES>
-void FPBDCollisionConstraintPGS::Solve(T_PARTICLES& InParticles, const FReal Dt, const TArray<int32>& InConstraintIndices)
-{
-	TArray<FReal> Normals;
-	TArray<FVec2> Tangents;
-	TArray<FReal> Multipliers;
-	TArray<TVec2<FVec3>> Angulars, MassWeightedAngulars;
-	TArray<TVec2<FVec3>> ConstraintTangents;
-	TVec2<TArray<FReal>> TangentMultipliers;
-	TVec2<TArray<TVec2<FVec3>>> TangentAngulars, TangentMassWeightedAngulars;
-
-	int32 NumConstraints = 0;
-	for (int32 ConstraintIndex : InConstraintIndices)
-	{
-		const FRigidBodyContactConstraint& Constraint = Constraints[ConstraintIndex];
-		NumConstraints += Constraint.Phi.Num();
-	}
-
-	Normals.SetNumZeroed(NumConstraints);
-	Tangents.SetNumZeroed(NumConstraints);
-	Multipliers.SetNumZeroed(NumConstraints);
-	Angulars.SetNum(NumConstraints);
-	MassWeightedAngulars.SetNum(NumConstraints);
-	ConstraintTangents.SetNum(NumConstraints);
-	for (int32 Dimension = 0; Dimension < (d - 1); ++Dimension)
-	{
-		TangentMultipliers[Dimension].SetNumZeroed(NumConstraints);
-		TangentAngulars[Dimension].SetNum(NumConstraints);
-		TangentMassWeightedAngulars[Dimension].SetNum(NumConstraints);
-	}
-
-	int32 FlattenedIndex = 0;
-	for (int32 ConstraintIndex : InConstraintIndices)
-	{
-		FRigidBodyContactConstraint& Constraint = Constraints[ConstraintIndex];
-		FMatrix33 WorldSpaceInvI1 = (GetRotationPGS(InParticles, Constraint.ParticleIndex) * FMatrix::Identity).GetTransposed() * InParticles.InvI(Constraint.ParticleIndex) * (GetRotationPGS(InParticles, Constraint.ParticleIndex) * FMatrix::Identity);
-		FMatrix33 WorldSpaceInvI2 = (GetRotationPGS(InParticles, Constraint.LevelsetIndex) * FMatrix::Identity).GetTransposed() * InParticles.InvI(Constraint.LevelsetIndex) * (GetRotationPGS(InParticles, Constraint.LevelsetIndex) * FMatrix::Identity);
-		for (int32 PointIndex = 0; PointIndex < Constraint.Phi.Num(); ++PointIndex)
-		{
-			ComputePGSPropeties(InParticles, Constraint, PointIndex, FlattenedIndex, WorldSpaceInvI1, WorldSpaceInvI2, -Constraint.Normal[PointIndex], Angulars, MassWeightedAngulars, Multipliers);
-			// Tangents
-			{
-				T AbsX = FMath::Abs(Constraint.Normal[PointIndex][0]);
-				T AbsY = FMath::Abs(Constraint.Normal[PointIndex][1]);
-				T AbsZ = FMath::Abs(Constraint.Normal[PointIndex][2]);
-				if (AbsX < AbsY)
-				{
-					if (AbsX < AbsZ)
-					{
-						ConstraintTangents[FlattenedIndex][0] = FVec3(0, Constraint.Normal[PointIndex][2], -Constraint.Normal[PointIndex][1]);
-					}
-					else
-					{
-						ConstraintTangents[FlattenedIndex][0] = FVec3(Constraint.Normal[PointIndex][1], -Constraint.Normal[PointIndex][0], 0);
-					}
-				}
-				else
-				{
-					if (AbsY < AbsZ)
-					{
-						ConstraintTangents[FlattenedIndex][0] = FVec3(-Constraint.Normal[PointIndex][2], 0, Constraint.Normal[PointIndex][0]);
-					}
-					else
-					{
-						ConstraintTangents[FlattenedIndex][0] = FVec3(Constraint.Normal[PointIndex][1], -Constraint.Normal[PointIndex][0], 0);
-					}
-				}
-			}
-			ConstraintTangents[FlattenedIndex][0] = ConstraintTangents[FlattenedIndex][0].GetSafeNormal();
-			ConstraintTangents[FlattenedIndex][1] = FVec3::CrossProduct(-ConstraintTangents[FlattenedIndex][0], Constraint.Normal[PointIndex]);
-			ComputePGSPropeties(InParticles, Constraint, PointIndex, FlattenedIndex, WorldSpaceInvI1, WorldSpaceInvI2, -ConstraintTangents[FlattenedIndex][0], TangentAngulars[0], TangentMassWeightedAngulars[0], TangentMultipliers[0]);
-			ComputePGSPropeties(InParticles, Constraint, PointIndex, FlattenedIndex, WorldSpaceInvI1, WorldSpaceInvI2, -ConstraintTangents[FlattenedIndex][1], TangentAngulars[1], TangentMassWeightedAngulars[1], TangentMultipliers[1]);
-			FlattenedIndex++;
-		}
-	}
-
-	T Residual = 0;
-	for (int32 Iteration = 0; Iteration < MaxIterations; ++Iteration)
-	{
-		Residual = 0;
-		FlattenedIndex = 0;
-		for (int32 ConstraintIndex : InConstraintIndices)
-		{
-			FRigidBodyContactConstraint& Constraint = Constraints[ConstraintIndex];
-			for (int32 PointIndex = 0; PointIndex < Constraint.Phi.Num(); ++PointIndex)
-			{
-				T Body1NormalVelocity = FVec3::DotProduct(InParticles.V(Constraint.ParticleIndex), Constraint.Normal[PointIndex]) +
-					FVec3::DotProduct(InParticles.W(Constraint.ParticleIndex), Angulars[FlattenedIndex][0]);
-				T Body2NormalVelocity = FVec3::DotProduct(InParticles.V(Constraint.LevelsetIndex), -Constraint.Normal[PointIndex]) +
-					FVec3::DotProduct(InParticles.W(Constraint.LevelsetIndex), Angulars[FlattenedIndex][1]);
-				T RelativeNormalVelocity = Body1NormalVelocity + Body2NormalVelocity + Constraint.Phi[PointIndex] / Dt;
-				T NewResidual = FMath::Max(-RelativeNormalVelocity, RelativeNormalVelocity * Normals[FlattenedIndex]);
-				if (NewResidual > Residual)
-				{
-					Residual = NewResidual;
-				}
-				T NormalDelta = -RelativeNormalVelocity / Multipliers[FlattenedIndex];
-				// Update Normals
-				T NewNormal = Normals[FlattenedIndex] + NormalDelta;
-				if (NewNormal < 0)
-				{
-					NewNormal = 0;
-					NormalDelta = -Normals[FlattenedIndex];
-				}
-				check(RelativeNormalVelocity < 0 || NormalDelta == 0 || Iteration > 0);
-				// Velocity update
-				InParticles.V(Constraint.ParticleIndex) += NormalDelta * InParticles.InvM(Constraint.ParticleIndex) * Constraint.Normal[PointIndex];
-				InParticles.V(Constraint.LevelsetIndex) += NormalDelta * InParticles.InvM(Constraint.LevelsetIndex) * -Constraint.Normal[PointIndex];
-				InParticles.W(Constraint.ParticleIndex) += NormalDelta * MassWeightedAngulars[FlattenedIndex][0];
-				InParticles.W(Constraint.LevelsetIndex) += NormalDelta * MassWeightedAngulars[FlattenedIndex][1];
-				// Normal update
-				Normals[FlattenedIndex] = NewNormal;
-				T Friction = MPhysicsMaterials[Constraint.LevelsetIndex]
-					? FMath::Max(MPhysicsMaterials[Constraint.ParticleIndex]->Friction, MPhysicsMaterials[Constraint.LevelsetIndex]->Friction)
-					: MPhysicsMaterials[Constraint.ParticleIndex]->Friction;
-				if (Friction)
-				{
-					for (int32 Dimension = 0; Dimension < (d - 1); Dimension++)
-					{
-						T Body1TangentVelocity = FVec3::DotProduct(InParticles.V(Constraint.ParticleIndex), ConstraintTangents[PointIndex][Dimension]) +
-							FVec3::DotProduct(InParticles.W(Constraint.ParticleIndex), TangentAngulars[Dimension][FlattenedIndex][0]);
-						T Body2TangentVelocity = FVec3::DotProduct(InParticles.V(Constraint.LevelsetIndex), -ConstraintTangents[PointIndex][Dimension]) +
-							FVec3::DotProduct(InParticles.W(Constraint.LevelsetIndex), TangentAngulars[Dimension][FlattenedIndex][1]);
-						T RelativeTangentVelocity = Body1TangentVelocity + Body2TangentVelocity;
-						T TangentDelta = -RelativeTangentVelocity / TangentMultipliers[Dimension][FlattenedIndex];
-						T NewTangent = Tangents[FlattenedIndex][Dimension] + TangentDelta;
-						if (FMath::Abs(NewTangent) > Friction * NewNormal)
-						{
-							NewTangent = Friction * NewNormal;
-							if (NewTangent < 0)
-							{
-								NewTangent *= -1;
-							}
-						}
-						// Velocity update
-						InParticles.V(Constraint.ParticleIndex) += TangentDelta * InParticles.InvM(Constraint.ParticleIndex) * ConstraintTangents[PointIndex][Dimension];
-						InParticles.V(Constraint.LevelsetIndex) += TangentDelta * InParticles.InvM(Constraint.LevelsetIndex) * -ConstraintTangents[PointIndex][Dimension];
-						InParticles.W(Constraint.ParticleIndex) += TangentDelta * TangentMassWeightedAngulars[Dimension][FlattenedIndex][0];
-						InParticles.W(Constraint.LevelsetIndex) += TangentDelta * TangentMassWeightedAngulars[Dimension][FlattenedIndex][1];
-						Tangents[FlattenedIndex][Dimension] = NewTangent;
-					}
-				}
-				FlattenedIndex++;
-			}
-		}
-		UE_LOG(LogChaos, Verbose, TEXT("\tPBDCollisionConstraint Solve with Residual %f"), Residual);
-		if (Residual < Tolerance)
-		{
-			break;
-		}
-	}
-}
-
-void FPBDCollisionConstraintPGS::PrintParticles(const FPBDRigidParticles& InParticles, const TArray<int32>& InConstraintIndices)
-{
-	TSet<int32> ConstraintParticles;
-	for (int32 ConstraintIndex : InConstraintIndices)
-	{
-		const FRigidBodyContactConstraint& Constraint = Constraints[ConstraintIndex];
-		if (!ConstraintParticles.Contains(Constraint.ParticleIndex))
-		{
-			ConstraintParticles.Add(Constraint.ParticleIndex);
-		}
-		if (!ConstraintParticles.Contains(Constraint.LevelsetIndex))
-		{
-			ConstraintParticles.Add(Constraint.LevelsetIndex);
-		}
-	}
-	for (const auto& i : ConstraintParticles)
-	{
-		UE_LOG(LogChaos, Verbose, TEXT("Particle %d has X=(%f, %f, %f) and V=(%f, %f, %f)"), i, InParticles.X(i)[0], InParticles.X(i)[1], InParticles.X(i)[2], InParticles.V(i)[0], InParticles.V(i)[1], InParticles.V(i)[2]);
-	}
-}
-
-void FPBDCollisionConstraintPGS::PrintConstraints(const FPBDRigidParticles& InParticles, const TArray<int32>& InConstraintIndices)
-{
-	for (int32 ConstraintIndex : InConstraintIndices)
-	{
-		const FRigidBodyContactConstraint& Constraint = Constraints[ConstraintIndex];
-		UE_LOG(LogChaos, Verbose, TEXT("Constraint between %d and %d has %d contacts"), Constraint.ParticleIndex, Constraint.LevelsetIndex, Constraint.Phi.Num());
-		for (int32 j = 0; j < Constraint.Phi.Num(); ++j)
-		{
-			UE_LOG(LogChaos, Verbose, TEXT("Constraint has loction (%f, %f, %f) and phi %f"), Constraint.Location[j][0], Constraint.Location[j][1], Constraint.Location[j][2], Constraint.Phi[j]);
-		}
-	}
-}
-
-void FindPointsOnHull(FRigidBodyContactConstraintPGS& Constraint,
-	const FVec3& X0, const FVec3& X1, const FVec3& X2,
-	const TPlane<FReal, 3>& SplitPlane, const TArray<int32>& Indices, TSet<int32>& HullPoints)
-{
-	int32 MaxD = 0; //This doesn't need to be initialized but we need to avoid the compiler warning
-	FReal MaxDistance = 0;
-	for (int32 i = 0; i < Indices.Num(); ++i)
-	{
-		FReal Distance = SplitPlane.SignedDistance(Constraint.Location[Indices[i]]);
-		check(Distance >= 0);
-		if (Distance > MaxDistance)
-		{
-			MaxDistance = Distance;
-			MaxD = Indices[i];
-		}
-	}
-	if (MaxDistance > 0)
-	{
-		if (!HullPoints.Contains(MaxD))
-		{
-			HullPoints.Add(MaxD);
-		}
-		const FVec3& NewX = Constraint.Location[MaxD];
-		const FVec3 V1 = (NewX - X0).GetSafeNormal();
-		const FVec3 V2 = (NewX - X1).GetSafeNormal();
-		const FVec3 V3 = (NewX - X2).GetSafeNormal();
-		FVec3 Normal1 = FVec3::CrossProduct(V1, V2).GetSafeNormal();
-		if (FVec3::DotProduct(Normal1, X2 - X0) > 0)
-		{
-			Normal1 *= -1;
-		}
-		FVec3 Normal2 = FVec3::CrossProduct(V1, V3).GetSafeNormal();
-		if (FVec3::DotProduct(Normal2, X1 - X0) > 0)
-		{
-			Normal2 *= -1;
-		}
-		FVec3 Normal3 = FVec3::CrossProduct(V2, V3).GetSafeNormal();
-		if (FVec3::DotProduct(Normal3, X0 - X1) > 0)
-		{
-			Normal3 *= -1;
-		}
-		TPlane<FReal, 3> NewPlane1(NewX, Normal1);
-		TPlane<FReal, 3> NewPlane2(NewX, Normal2);
-		TPlane<FReal, 3> NewPlane3(NewX, Normal3);
-		TArray<int32> NewIndices1;
-		TArray<int32> NewIndices2;
-		TArray<int32> NewIndices3;
-		for (int32 i = 0; i < Indices.Num(); ++i)
-		{
-			if (MaxD == Indices[i])
-			{
-				continue;
-			}
-			if (NewPlane1.SignedDistance(Constraint.Location[Indices[i]]) > 0)
-			{
-				NewIndices1.Add(Indices[i]);
-			}
-			if (NewPlane2.SignedDistance(Constraint.Location[Indices[i]]) > 0)
-			{
-				NewIndices2.Add(Indices[i]);
-			}
-			if (NewPlane3.SignedDistance(Constraint.Location[Indices[i]]) > 0)
-			{
-				NewIndices3.Add(Indices[i]);
-			}
-		}
-		FindPointsOnHull(Constraint, X0, X1, NewX, NewPlane1, NewIndices1, HullPoints);
-		FindPointsOnHull(Constraint, X0, X2, NewX, NewPlane2, NewIndices2, HullPoints);
-		FindPointsOnHull(Constraint, X1, X2, NewX, NewPlane3, NewIndices3, HullPoints);
-	}
-}
-
-void RemovePointsInsideHull(FRigidBodyContactConstraintPGS& Constraint)
-{
-	if (Constraint.Location.Num() <= 2)
-	{
-		return;
-	}
-	// Find max and min x points
-	int32 MinX = 0;
-	int32 MaxX = 0;
-	int32 MinY = 0;
-	int32 MaxY = 0;
-	int32 Index1 = 0;
-	int32 Index2 = 0;
-	for (int32 i = 1; i < Constraint.Location.Num(); ++i)
-	{
-		if (Constraint.Location[i][0] > Constraint.Location[MaxX][0])
-		{
-			MaxX = i;
-		}
-		if (Constraint.Location[i][0] < Constraint.Location[MinX][0])
-		{
-			MinX = i;
-		}
-		if (Constraint.Location[i][1] > Constraint.Location[MaxY][1])
-		{
-			MaxY = i;
-		}
-		if (Constraint.Location[i][1] < Constraint.Location[MinY][1])
-		{
-			MinY = i;
-		}
-	}
-	if (MaxX == MinX && MinY == MaxY && MinX == MinY)
-	{
-		// Points are colinear so need to sort but for now do nothing
-		return;
-	}
-	// Find max distance
-	FReal DistanceY = (Constraint.Location[MaxY] - Constraint.Location[MinY]).Size();
-	FReal DistanceX = (Constraint.Location[MaxX] - Constraint.Location[MinX]).Size();
-	if (DistanceX > DistanceY)
-	{
-		Index1 = MaxX;
-		Index2 = MinX;
-	}
-	else
-	{
-		Index1 = MaxY;
-		Index2 = MinY;
-	}
-	TSet<int32> HullPoints;
-	HullPoints.Add(Index1);
-	HullPoints.Add(Index2);
-	const FVec3& X1 = Constraint.Location[Index1];
-	const FVec3& X2 = Constraint.Location[Index2];
-	FReal MaxDist = 0;
-	int32 MaxD = -1;
-	for (int32 i = 0; i < Constraint.Location.Num(); ++i)
-	{
-		if (i == Index1 || i == Index2)
-		{
-			continue;
-		}
-		const FVec3& X0 = Constraint.Location[i];
-		FReal Distance = FVec3::CrossProduct(X0 - X1, X0 - X2).Size() / (X2 - X1).Size();
-		if (Distance > MaxDist)
-		{
-			MaxDist = Distance;
-			MaxD = i;
-		}
-	}
-	if (MaxD != -1)
-	{
-		HullPoints.Add(MaxD);
-		const FVec3& X0 = Constraint.Location[MaxD];
-		FVec3 Normal = FVec3::CrossProduct((X0 - X1).GetSafeNormal(), (X0 - X2).GetSafeNormal());
-		TPlane<FReal, 3> SplitPlane(X0, Normal);
-		TPlane<FReal, 3> SplitPlaneNeg(X0, -Normal);
-		TArray<int32> Left;
-		TArray<int32> Right;
-		for (int32 i = 0; i < Constraint.Location.Num(); ++i)
-		{
-			if (i == Index1 || i == Index2 || i == MaxD)
-			{
-				continue;
-			}
-			if (SplitPlane.SignedDistance(Constraint.Location[i]) >= 0)
-			{
-				Left.Add(i);
-			}
-			else
-			{
-				Right.Add(i);
-			}
-		}
-		FindPointsOnHull(Constraint, X0, X1, X2, SplitPlane, Left, HullPoints);
-		FindPointsOnHull(Constraint, X0, X1, X2, SplitPlaneNeg, Right, HullPoints);
-	}
-	TArray<FVec3> Locations;
-	TArray<FVec3> Normals;
-	TArray<FReal> Distances;
-	for (const auto& Index : HullPoints)
-	{
-		Locations.Add(Constraint.Location[Index]);
-		Normals.Add(Constraint.Normal[Index]);
-		Distances.Add(Constraint.Phi[Index]);
-	}
-	Constraint.Location = Locations;
-	Constraint.Normal = Normals;
-	Constraint.Phi = Distances;
-}
-
-void FPBDCollisionConstraintPGS::Apply(FPBDRigidParticles& InParticles, const FReal Dt, const TArray<int32>& InConstraintIndices)
-{
-	PhysicsParallelFor(InConstraintIndices.Num(), [&](int32 ConstraintIndex) {
-		FRigidBodyContactConstraint& Constraint = Constraints[ConstraintIndex];
-		if (InParticles.Sleeping(Constraint.ParticleIndex))
-		{
-			check(InParticles.Sleeping(Constraint.LevelsetIndex) || InParticles.InvM(Constraint.LevelsetIndex) == 0);
-			return;
-		}
-		// @todo(mlentine): This is a really coarse approximation. Prune points that we know are not relevant.
-		const FReal Threshold = (InParticles.V(Constraint.ParticleIndex).Size() - InParticles.V(Constraint.LevelsetIndex).Size()) * Dt;
-		// Guessing Max is a decent approximation as with restitution 0 the difference is X between predicted and actual is Vdt
-		const FReal Thickness = MThickness + FMath::Max(InParticles.V(Constraint.ParticleIndex).Size(), InParticles.V(Constraint.LevelsetIndex).Size()) * Dt;
-		const_cast<FPBDCollisionConstraintPGS*>(this)->UpdateConstraint(static_cast<TRigidParticles<FReal, 3>&>(InParticles), Thickness + Threshold, Constraint);
-		// @todo(mlentine): Prune contact points based on convex hull
-		RemovePointsInsideHull(Constraint);
-	});
-	PrintParticles(InParticles, InConstraintIndices);
-	PrintConstraints(InParticles, InConstraintIndices);
-	Solve(static_cast<TRigidParticles<FReal, 3>&>(InParticles), Dt, InConstraintIndices);
-	PrintParticles(InParticles, InConstraintIndices);
-}
-
-void FPBDCollisionConstraintPGS::ApplyPushOut(FPBDRigidParticles& InParticles, const FReal Dt, const TArray<int32>& InConstraintIndices)
-{
-	PhysicsParallelFor(InConstraintIndices.Num(), [&](int32 ConstraintIndex) {
-		FRigidBodyContactConstraint& Constraint = Constraints[ConstraintIndex];
-		if (InParticles.Sleeping(Constraint.ParticleIndex))
-		{
-			check(InParticles.Sleeping(Constraint.LevelsetIndex) || InParticles.InvM(Constraint.LevelsetIndex) == 0);
-			return;
-		}
-		const_cast<FPBDCollisionConstraintPGS*>(this)->UpdateConstraint(InParticles, MThickness, Constraint);
-		// @todo(mlentine): Prune contact points based on convex hull
-	});
-
-	TArray<bool> Saved;
-	TArray<FVec3> SavedV, SavedW;
-	Saved.SetNum(InParticles.Size());
-	SavedV.SetNum(InParticles.Size());
-	SavedW.SetNum(InParticles.Size());
-	auto SaveParticle = [&](int32 ParticleIndex)
-	{
-		if (!Saved[ParticleIndex])
-		{
-			SavedV[ParticleIndex] = InParticles.V(ParticleIndex);
-			SavedW[ParticleIndex] = InParticles.W(ParticleIndex);
-			InParticles.V(ParticleIndex) = FVec3(0);
-			InParticles.W(ParticleIndex) = FVec3(0);
-			Saved[ParticleIndex] = true;
-		}
-	};
-	auto RestoreParticle = [&](int32 ParticleIndex)
-	{
-		if (Saved[ParticleIndex])
-		{
-			if (InParticles.InvM(ParticleIndex))
-			{
-				InParticles.P(ParticleIndex) += InParticles.V(ParticleIndex) * Dt;
-				InParticles.Q(ParticleIndex) += FRotation3(InParticles.W(ParticleIndex), 0.f) * InParticles.Q(ParticleIndex) * Dt * FReal(0.5);
-				InParticles.Q(ParticleIndex).Normalize();
-			}
-			InParticles.V(ParticleIndex) = SavedV[ParticleIndex];
-			InParticles.W(ParticleIndex) = SavedW[ParticleIndex];
-			Saved[ParticleIndex] = false;
-		}
-	};
-
-	PhysicsParallelFor(InConstraintIndices.Num(), [&](int32 ConstraintIndex) {
-		SaveParticle(Constraints[ConstraintIndex].ParticleIndex);
-		SaveParticle(Constraints[ConstraintIndex].LevelsetIndex);
-	});
-	PrintParticles(InParticles, InConstraintIndices);
-	PrintConstraints(InParticles, InConstraintIndices);
-
-	Solve(InParticles, Dt, InConstraintIndices);
-	
-	PrintParticles(InParticles, InConstraintIndices);
-	PhysicsParallelFor(InConstraintIndices.Num(), [&](int32 ConstraintIndex) {
-		RestoreParticle(Constraints[ConstraintIndex].ParticleIndex);
-		RestoreParticle(Constraints[ConstraintIndex].LevelsetIndex);
-	});
-}
-
-bool FPBDCollisionConstraintPGS::NearestPoint(TArray<Pair<FVec3, FVec3>>& Points, FVec3& Direction)
-{
-	check(Points.Num() > 1 && Points.Num() <= 4);
-	if (Points.Num() == 2)
-	{
-		TPlane<FReal, 3> LocalPlane(Points[1].First, Points[0].First - Points[1].First);
-		FVec3 Normal;
-		const auto& Phi = LocalPlane.PhiWithNormal(FVec3(0), Normal);
-		if ((FVec3::DotProduct(-Points[1].First, Normal.GetSafeNormal()) - Points[1].First.Size()) < SMALL_NUMBER)
-		{
-			FReal Alpha = Points[0].First.Size() / (Points[1].First - Points[0].First).Size();
-			return true;
-		}
-		if (Phi > 0)
-		{
-			check(Points.Num() == 2);
-			Direction = FVec3::CrossProduct(FVec3::CrossProduct(Normal, -Points[1].First), Normal);
-		}
-		else
-		{
-			Direction = -Points[1].First;
-			Points.RemoveAtSwap(0);
-			check(Points.Num() == 1);
-		}
-		check(Points.Num() > 1 && Points.Num() < 4);
-		return false;
-	}
-	if (Points.Num() == 3)
-	{
-		FVec3 TriangleNormal = FVec3::CrossProduct(Points[0].First - Points[2].First, Points[0].First - Points[1].First);
-		TPlane<FReal, 3> LocalPlane1(Points[2].First, FVec3::CrossProduct(Points[0].First - Points[2].First, TriangleNormal));
-		TPlane<FReal, 3> LocalPlane2(Points[2].First, FVec3::CrossProduct(Points[1].First - Points[2].First, TriangleNormal));
-		FVec3 Normal;
-		FReal Phi = LocalPlane1.PhiWithNormal(FVec3(0), Normal);
-		if (Phi > 0)
-		{
-			FVec3 Delta = Points[0].First - Points[2].First;
-			if (FVec3::DotProduct(-Points[2].First, Delta) > 0)
-			{
-				Direction = FVec3::CrossProduct(FVec3::CrossProduct(Delta, -Points[2].First), Delta);
-				Points.RemoveAtSwap(1);
-				check(Points.Num() == 2);
-			}
-			else
-			{
-				Delta = Points[1].First - Points[2].First;
-				if (FVec3::DotProduct(-Points[2].First, Delta) > 0)
-				{
-					Direction = FVec3::CrossProduct(FVec3::CrossProduct(Delta, -Points[2].First), Delta);
-					Points.RemoveAtSwap(0);
-					check(Points.Num() == 2);
-				}
-				else
-				{
-					Direction = -Points[2].First;
-					Points = {Points[2]};
-					check(Points.Num() == 1);
-				}
-			}
-		}
-		else
-		{
-			Phi = LocalPlane2.PhiWithNormal(FVec3(0), Normal);
-			if (Phi > 0)
-			{
-				FVec3 Delta = Points[1].First - Points[2].First;
-				if (FVec3::DotProduct(-Points[2].First, Delta) > 0)
-				{
-					Direction = FVec3::CrossProduct(FVec3::CrossProduct(Delta, -Points[2].First), Delta);
-					Points.RemoveAtSwap(0);
-					check(Points.Num() == 2);
-				}
-				else
-				{
-					Direction = -Points[2].First;
-					Points = {Points[2]};
-					check(Points.Num() == 1);
-				}
-			}
-			else
-			{
-				const auto DotResult = FVec3::DotProduct(TriangleNormal, -Points[2].First);
-				// We are inside the triangle
-				if (DotResult < SMALL_NUMBER)
-				{
-					FVec3 Bary;
-					FVec3 P10 = Points[1].First - Points[0].First;
-					FVec3 P20 = Points[2].First - Points[0].First;
-					FVec3 PP0 = -Points[0].First;
-					FReal Size10 = P10.SizeSquared();
-					FReal Size20 = P20.SizeSquared();
-					FReal ProjSides = FVec3::DotProduct(P10, P20);
-					FReal ProjP1 = FVec3::DotProduct(PP0, P10);
-					FReal ProjP2 = FVec3::DotProduct(PP0, P20);
-					FReal Denom = Size10 * Size20 - ProjSides * ProjSides;
-					Bary.Y = (Size20 * ProjP1 - ProjSides * ProjP2) / Denom;
-					Bary.Z = (Size10 * ProjP2 - ProjSides * ProjP1) / Denom;
-					Bary.X = 1.0f - Bary.Z - Bary.Y;
-					return true;
-				}
-				if (DotResult > 0)
-				{
-					Direction = TriangleNormal;
-				}
-				else
-				{
-					Direction = -TriangleNormal;
-					Points.Swap(0, 1);
-					check(Points.Num() == 3);
-				}
-			}
-		}
-		check(Points.Num() > 0 && Points.Num() < 4);
-		return false;
-	}
-	if (Points.Num() == 4)
-	{
-		FVec3 TriangleNormal = FVec3::CrossProduct(Points[1].First - Points[3].First, Points[1].First - Points[2].First);
-		if (FVec3::DotProduct(TriangleNormal, Points[0].First - Points[3].First) > 0)
-		{
-			TriangleNormal *= -1;
-		}
-		FReal DotResult = FVec3::DotProduct(TriangleNormal, -Points[3].First);
-		if (DotResult > 0)
-		{
-			Points = {Points[1], Points[2], Points[3]};
-			check(Points.Num() == 3);
-			return NearestPoint(Points, Direction);
-		}
-		TriangleNormal = FVec3::CrossProduct(Points[2].First - Points[0].First, Points[2].First - Points[3].First);
-		if (FVec3::DotProduct(TriangleNormal, Points[1].First - Points[3].First) > 0)
-		{
-			TriangleNormal *= -1;
-		}
-		DotResult = FVec3::DotProduct(TriangleNormal, -Points[3].First);
-		if (DotResult > 0)
-		{
-			Points = {Points[0], Points[2], Points[3]};
-			check(Points.Num() == 3);
-			return NearestPoint(Points, Direction);
-		}
-		TriangleNormal = FVec3::CrossProduct(Points[3].First - Points[1].First, Points[3].First - Points[0].First);
-		if (FVec3::DotProduct(TriangleNormal, Points[2].First - Points[3].First) > 0)
-		{
-			TriangleNormal *= -1;
-		}
-		DotResult = FVec3::DotProduct(TriangleNormal, -Points[3].First);
-		if (DotResult > 0)
-		{
-			Points = {Points[0], Points[1], Points[3]};
-			check(Points.Num() == 3);
-			return NearestPoint(Points, Direction);
-		}
-		TVector<FReal, 4> Bary;
-		FVec3 PP0 = -Points[0].First;
-		FVec3 PP1 = -Points[1].First;
-		FVec3 P10 = Points[1].First - Points[0].First;
-		FVec3 P20 = Points[2].First - Points[0].First;
-		FVec3 P30 = Points[3].First - Points[0].First;
-		FVec3 P21 = Points[2].First - Points[1].First;
-		FVec3 P31 = Points[3].First - Points[1].First;
-		Bary[0] = FVec3::DotProduct(PP1, FVec3::CrossProduct(P31, P21));
-		Bary[1] = FVec3::DotProduct(PP0, FVec3::CrossProduct(P20, P30));
-		Bary[2] = FVec3::DotProduct(PP0, FVec3::CrossProduct(P30, P10));
-		Bary[3] = FVec3::DotProduct(PP0, FVec3::CrossProduct(P10, P20));
-		FReal Denom = FVec3::DotProduct(P10, FVec3::CrossProduct(P20, P30));
-		return true;
-	}
-	check(Points.Num() > 1 && Points.Num() < 4);
-	return false;
-}
-
-void UpdateLevelsetConstraintHelperCCD(const TRigidParticles<FReal, 3>& InParticles, const int32 j, const FRigidTransform3& LocalToWorld1, const FRigidTransform3& LocalToWorld2, const T Thickness, FRigidBodyContactConstraintPGS& Constraint)
-{
-	if(InParticles.CollisionParticles(Constraint.ParticleIndex))
-	{
-		const FRigidTransform3 PreviousLocalToWorld1 = GetTransformPGS(InParticles, Constraint.ParticleIndex);
-		FVec3 WorldSpacePointStart = PreviousLocalToWorld1.TransformPosition(InParticles.CollisionParticles(Constraint.ParticleIndex)->X(j));
-		FVec3 WorldSpacePointEnd = LocalToWorld1.TransformPosition(InParticles.CollisionParticles(Constraint.ParticleIndex)->X(j));
-		FVec3 Body2SpacePointStart = LocalToWorld2.InverseTransformPosition(WorldSpacePointStart);
-		FVec3 Body2SpacePointEnd = LocalToWorld2.InverseTransformPosition(WorldSpacePointEnd);
-		Pair<FVec3, bool> PointPair = InParticles.Geometry(Constraint.LevelsetIndex)->FindClosestIntersection(Body2SpacePointStart, Body2SpacePointEnd, Thickness);
-		if (PointPair.Second)
-		{
-			const FVec3 WorldSpaceDelta = WorldSpacePointEnd - FVec3(LocalToWorld2.TransformPosition(PointPair.First));
-			Constraint.Phi.Add(-WorldSpaceDelta.Size());
-			Constraint.Normal.Add(LocalToWorld2.TransformVector(InParticles.Geometry(Constraint.LevelsetIndex)->Normal(PointPair.First)));
-			// @todo(mlentine): Should we be using the actual collision point or that point evolved to the current time step?
-			Constraint.Location.Add(WorldSpacePointEnd);
-		}
-	}
-}
-
-void UpdateLevelsetConstraintHelper(const TRigidParticles<FReal, 3>& InParticles, const int32 j, const FRigidTransform3& LocalToWorld1, const FRigidTransform3& LocalToWorld2, const T Thickness, FRigidBodyContactConstraintPGS& Constraint)
-{
-	if(InParticles.CollisionParticles(Constraint.ParticleIndex))
-	{
-		FVec3 WorldSpacePoint = LocalToWorld1.TransformPosition(InParticles.CollisionParticles(Constraint.ParticleIndex)->X(j));
-		FVec3 Body2SpacePoint = LocalToWorld2.InverseTransformPosition(WorldSpacePoint);
-		FVec3 LocalNormal;
-		FReal LocalPhi = InParticles.Geometry(Constraint.LevelsetIndex)->PhiWithNormal(Body2SpacePoint, LocalNormal);
-		if (LocalPhi < Thickness)
-		{
-			Constraint.Phi.Add(LocalPhi);
-			Constraint.Normal.Add(LocalToWorld2.TransformVector(LocalNormal));
-			Constraint.Location.Add(WorldSpacePoint);
-		}
-	}
-}
-
-template<class T_PARTICLES>
-void FPBDCollisionConstraintPGS::UpdateLevelsetConstraint(const T_PARTICLES& InParticles, const FReal Thickness, FRigidBodyContactConstraint& Constraint)
-{
-	Constraint.Normal.Reset();
-	Constraint.Phi.Reset();
-	Constraint.Location.Reset();
-	const FRigidTransform3 LocalToWorld1 = GetTransformPGS(InParticles, Constraint.ParticleIndex);
-	const FRigidTransform3 LocalToWorld2 = GetTransformPGS(InParticles, Constraint.LevelsetIndex);
-	if (InParticles.Geometry(Constraint.LevelsetIndex)->HasBoundingBox())
-	{
-		FAABB3 ImplicitBox = InParticles.Geometry(Constraint.LevelsetIndex)->BoundingBox().TransformedBox(LocalToWorld2 * LocalToWorld1.Inverse());
-		if (InParticles.CollisionParticles(Constraint.ParticleIndex))
-		{
-			TArray<int32> PotentialParticles = InParticles.CollisionParticles(Constraint.ParticleIndex)->FindAllIntersections(ImplicitBox);
-			for (int32 j = 0; j < PotentialParticles.Num(); ++j)
-			{
-				if (bUseCCD)
-				{
-					UpdateLevelsetConstraintHelperCCD(InParticles, PotentialParticles[j], LocalToWorld1, LocalToWorld2, Thickness, Constraint);
-				}
-				else
-				{
-					UpdateLevelsetConstraintHelper(InParticles, PotentialParticles[j], LocalToWorld1, LocalToWorld2, Thickness, Constraint);
-				}
-			}
-		}
-	}
-	else
-	{
-		if (InParticles.CollisionParticles(Constraint.ParticleIndex))
-		{
-			for (uint32 j = 0; j < InParticles.CollisionParticles(Constraint.ParticleIndex)->Size(); ++j)
-			{
-				UpdateLevelsetConstraintHelper(InParticles, j, LocalToWorld1, LocalToWorld2, Thickness, Constraint);
-			}
-		}
-	}
-}
-
-template<class T_PARTICLES>
-void FPBDCollisionConstraintPGS::UpdateLevelsetConstraintGJK(const T_PARTICLES& InParticles, const FReal Thickness, FRigidBodyContactConstraint& Constraint)
-{
-	static int32 MaxIterationsGJK = 100;
-	Constraint.Normal.Reset();
-	Constraint.Phi.Reset();
-	Constraint.Location.Reset();
-	const FRigidTransform3 LocalToWorld1 = GetTransformPGS(InParticles, Constraint.ParticleIndex);
-	const FRigidTransform3 LocalToWorld2 = GetTransformPGS(InParticles, Constraint.LevelsetIndex);
-	FVec3 Direction = LocalToWorld1.GetTranslation() - LocalToWorld2.GetTranslation();
-	FVec3 SupportA = LocalToWorld1.TransformPosition(InParticles.Geometry(Constraint.ParticleIndex)->Support(LocalToWorld1.InverseTransformVector(-Direction), Thickness));
-	FVec3 SupportB = LocalToWorld2.TransformPosition(InParticles.Geometry(Constraint.LevelsetIndex)->Support(LocalToWorld2.InverseTransformVector(Direction), Thickness));
-	FVec3 Point = SupportB - SupportA;
-	TArray<Pair<FVec3, FVec3>> Points = {MakePair(Point, SupportA)};
-	Direction = -Point;
-	for (int32 i = 0; i < MaxIterationsGJK; ++i)
-	{
-		SupportA = LocalToWorld1.TransformPosition(InParticles.Geometry(Constraint.ParticleIndex)->Support(LocalToWorld1.InverseTransformVector(-Direction), Thickness));
-		SupportB = LocalToWorld2.TransformPosition(InParticles.Geometry(Constraint.LevelsetIndex)->Support(LocalToWorld2.InverseTransformVector(Direction), Thickness));
-		Point = SupportB - SupportA;
-		if (FVec3::DotProduct(Point, Direction) < 0)
-		{
-			break;
-		}
-		Points.Add(MakePair(Point, SupportA));
-		FVec3 ClosestPoint;
-		if (NearestPoint(Points, Direction))
-		{
-			for (const auto& SinglePoint : Points)
-			{
-				FVec3 Body1Location = LocalToWorld1.InverseTransformPosition(SinglePoint.Second);
-				FVec3 Normal;
-				FReal Phi = InParticles.Geometry(Constraint.ParticleIndex)->PhiWithNormal(Body1Location, Normal);
-				Normal = LocalToWorld1.TransformVector(Normal);
-				FVec3 SurfacePoint = SinglePoint.Second - Phi * Normal;
-				Constraint.Location.Add(SurfacePoint);
-				FVec3 Body2Location = LocalToWorld2.InverseTransformPosition(SurfacePoint);
-				Constraint.Phi.Add(InParticles.Geometry(Constraint.LevelsetIndex)->PhiWithNormal(Body2Location, Normal));
-				Constraint.Normal.Add(LocalToWorld2.TransformVector(Normal));
-			}
-			break;
-		}
-	}
-}
-
-template<class T_PARTICLES>
-void FPBDCollisionConstraintPGS::UpdateBoxConstraint(const T_PARTICLES& InParticles, const T Thickness, FRigidBodyContactConstraint& Constraint)
-{
-	Constraint.Normal.Reset();
-	Constraint.Phi.Reset();
-	Constraint.Location.Reset();
-	const FRigidTransform3 Box1Transform = GetTransformPGS(InParticles, Constraint.ParticleIndex);
-	const FRigidTransform3 Box2Transform = GetTransformPGS(InParticles, Constraint.LevelsetIndex);
-	const auto& Box1 = *InParticles.Geometry(Constraint.ParticleIndex)->template GetObject<TBox<T, d>>();
-	const auto& Box2 = *InParticles.Geometry(Constraint.LevelsetIndex)->template GetObject<TBox<T, d>>();
-	auto Box2SpaceBox1 = Box1.TransformedBox(Box1Transform * Box2Transform.Inverse());
-	auto Box1SpaceBox2 = Box2.TransformedBox(Box2Transform * Box1Transform.Inverse());
-	Box2SpaceBox1.Thicken(Thickness);
-	Box1SpaceBox2.Thicken(Thickness);
-	if (Box1SpaceBox2.Intersects(Box1) && Box2SpaceBox1.Intersects(Box2))
-	{
-		const FVec3 Box1Center = (Box1Transform * Box2Transform.Inverse()).TransformPosition(Box1.Center());
-		if (Box2.SignedDistance(Box1Center) < 0)
-		{
-			TSphere<FReal, 3> Sphere1(Box1Transform.TransformPosition(Box1.Center()), Box1.Extents().Min() / 2);
-			TSphere<FReal, 3> Sphere2(Box2Transform.TransformPosition(Box2.Center()), Box2.Extents().Min() / 2);
-			const FVec3 Direction = Sphere1.GetCenter() - Sphere2.GetCenter();
-			FReal Size = Direction.Size();
-			if (Size < (Sphere1.GetRadius() + Sphere2.GetRadius()))
-			{
-				FVec3 Normal = Size > SMALL_NUMBER ? Direction / Size : FVec3(0, 0, 1);
-				Constraint.Normal.Add(Normal);
-				Constraint.Phi.Add(Size - (Sphere1.GetRadius() + Sphere2.GetRadius()));
-				Constraint.Location.Add(Sphere1.GetCenter() - Sphere1.GetRadius() * Normal);
-			}
-		}
-		if (!Constraint.Phi.Num())
-		{
-			//UpdateLevelsetConstraintGJK(InParticles, Thickness, Constraint);
-			//check(Constraint.Phi < MThickness);
-			// For now revert to doing all points vs lsv check until we can figure out a good way to get the deepest point without needing this
-			UpdateLevelsetConstraint(InParticles, Thickness, Constraint);
-		}
-	}
-}
-
-template<class T_PARTICLES>
-void FPBDCollisionConstraintPGS::UpdateBoxPlaneConstraint(const T_PARTICLES& InParticles, const FReal Thickness, FRigidBodyContactConstraint& Constraint)
-{
-	Constraint.Normal.Reset();
-	Constraint.Phi.Reset();
-	Constraint.Location.Reset();
-	const FRigidTransform3 BoxTransform = GetTransformPGS(InParticles, Constraint.ParticleIndex);
-	const FRigidTransform3 PlaneTransform = GetTransformPGS(InParticles, Constraint.LevelsetIndex);
-	const auto& ObjectBox = *InParticles.Geometry(Constraint.ParticleIndex)->template GetObject<TBox<T, d>>();
-	const auto& ObjectPlane = *InParticles.Geometry(Constraint.LevelsetIndex)->template GetObject<TPlane<FReal, 3>>();
-	const FRigidTransform3 BoxToPlaneTransform(BoxTransform * PlaneTransform.Inverse());
-	const FVec3 Extents = ObjectBox.Extents();
-	TArray<FVec3> Corners;
-	Corners.Add(BoxToPlaneTransform.TransformPosition(ObjectBox.Max()));
-	Corners.Add(BoxToPlaneTransform.TransformPosition(ObjectBox.Min()));
-	for (int32 j = 0; j < d; ++j)
-	{
-		Corners.Add(BoxToPlaneTransform.TransformPosition(ObjectBox.Min() + FVec3::AxisVector(j) * Extents));
-		Corners.Add(BoxToPlaneTransform.TransformPosition(ObjectBox.Max() - FVec3::AxisVector(j) * Extents));
-	}
-	TArray<FVec3> PotentialConstraints;
-	for (int32 i = 0; i < Corners.Num(); ++i)
-	{
-		FVec3 Normal;
-		const FReal NewPhi = ObjectPlane.PhiWithNormal(Corners[i], Normal);
-		if (NewPhi < Thickness)
-		{
-			Constraint.Phi.Add(NewPhi);
-			Constraint.Normal.Add(PlaneTransform.TransformVector(Normal));
-			Constraint.Location.Add(PlaneTransform.TransformPosition(Corners[i]));
-		}
-	}
-}
-
-template<class T_PARTICLES>
-void FPBDCollisionConstraintPGS::UpdateSphereConstraint(const T_PARTICLES& InParticles, const FReal Thickness, FRigidBodyContactConstraint& Constraint)
-{
-	Constraint.Normal.Reset();
-	Constraint.Phi.Reset();
-	Constraint.Location.Reset();
-	const FRigidTransform3 Sphere1Transform = GetTransformPGS(InParticles, Constraint.ParticleIndex);
-	const FRigidTransform3 Sphere2Transform = GetTransformPGS(InParticles, Constraint.LevelsetIndex);
-	const auto& Sphere1 = *InParticles.Geometry(Constraint.ParticleIndex)->template GetObject<TSphere<FReal, 3>>();
-	const auto& Sphere2 = *InParticles.Geometry(Constraint.LevelsetIndex)->template GetObject<TSphere<FReal, 3>>();
-	const FVec3 Center1 = Sphere1Transform.TransformPosition(Sphere1.GetCenter());
-	const FVec3 Center2 = Sphere2Transform.TransformPosition(Sphere2.GetCenter());
-	const FVec3 Direction = Center1 - Center2;
-	const FReal Size = Direction.Size();
-	if (Size < (Sphere1.GetRadius() + Sphere2.GetRadius() + Thickness))
-	{
-		FVec3 Normal = Size > SMALL_NUMBER ? Direction / Size : FVec3(0, 0, 1);
-		Constraint.Normal.Add(Normal);
-		Constraint.Phi.Add(Size - (Sphere1.GetRadius() + Sphere2.GetRadius()));
-		Constraint.Location.Add(Center1 - Sphere1.GetRadius() * Normal);
-	}
-}
-
-template<class T_PARTICLES>
-void FPBDCollisionConstraintPGS::UpdateSpherePlaneConstraint(const T_PARTICLES& InParticles, const FReal Thickness, FRigidBodyContactConstraint& Constraint)
-{
-	Constraint.Normal.Reset();
-	Constraint.Phi.Reset();
-	Constraint.Location.Reset();
-	const FRigidTransform3 SphereTransform = GetTransformPGS(InParticles, Constraint.ParticleIndex);
-	const FRigidTransform3 PlaneTransform = GetTransformPGS(InParticles, Constraint.LevelsetIndex);
-	const auto& ObjectSphere = *InParticles.Geometry(Constraint.ParticleIndex)->template GetObject<TSphere<FReal, 3>>();
-	const auto& ObjectPlane = *InParticles.Geometry(Constraint.LevelsetIndex)->template GetObject<TPlane<FReal, 3>>();
-	const FRigidTransform3 SphereToPlaneTransform(PlaneTransform.Inverse() * SphereTransform);
-	const FVec3 SphereCenter = SphereToPlaneTransform.TransformPosition(ObjectSphere.GetCenter());
-	Constraint.Normal.SetNum(1);
-	Constraint.Phi.Add(ObjectPlane.PhiWithNormal(SphereCenter, Constraint.Normal[0]));
-	Constraint.Phi[0] -= ObjectSphere.GetRadius();
-	Constraint.Location.Add(SphereCenter - Constraint.Normal[0] * ObjectSphere.GetRadius());
-}
-
-template<class T, int d>
-template<class T_PARTICLES>
-void FPBDCollisionConstraintPGS::UpdateSphereBoxConstraint(const T_PARTICLES& InParticles, const FReal Thickness, FRigidBodyContactConstraint& Constraint)
-{
-	Constraint.Normal.Reset();
-	Constraint.Phi.Reset();
-	Constraint.Location.Reset();
-	const FRigidTransform3 SphereTransform = GetTransformPGS(InParticles, Constraint.ParticleIndex);
-	const FRigidTransform3 BoxTransform = GetTransformPGS(InParticles, Constraint.LevelsetIndex);
-	const auto& ObjectSphere = *InParticles.Geometry(Constraint.ParticleIndex)->template GetObject<TSphere<FReal, 3>>();
-	const auto& ObjectBox = *InParticles.Geometry(Constraint.LevelsetIndex)->template GetObject<TBox<T, d>>();
-	const FRigidTransform3 SphereToBoxTransform(SphereTransform * BoxTransform.Inverse());
-	const FVec3 SphereCenter = SphereToBoxTransform.TransformPosition(ObjectSphere.GetCenter());
-	Constraint.Normal.SetNum(1);
-	Constraint.Phi.Add(ObjectBox.PhiWithNormal(SphereCenter, Constraint.Normal[0]));
-	Constraint.Phi[0] -= ObjectSphere.GetRadius();
-	Constraint.Location.Add(SphereCenter - Constraint.Normal[0] * ObjectSphere.GetRadius());
-}
-
-typename FPBDCollisionConstraintPGS::FRigidBodyContactConstraint FPBDCollisionConstraintPGS::ComputeLevelsetConstraint(const FPBDRigidParticles& InParticles, int32 ParticleIndex, int32 LevelsetIndex, const FReal Thickness)
-{
-	if (!InParticles.CollisionParticlesSize(ParticleIndex))
-	{
-		int32 TmpIndex = ParticleIndex;
-		ParticleIndex = LevelsetIndex;
-		LevelsetIndex = TmpIndex;
-	}
-	// Find Deepest Point
-	FRigidBodyContactConstraint Constraint;
-	Constraint.ParticleIndex = ParticleIndex;
-	Constraint.LevelsetIndex = LevelsetIndex;
-	return Constraint;
-}
-
-typename FPBDCollisionConstraintPGS::FRigidBodyContactConstraint FPBDCollisionConstraintPGS::ComputeLevelsetConstraintGJK(const FPBDRigidParticles& InParticles, int32 ParticleIndex, int32 LevelsetIndex, const FReal Thickness)
-{
-	FRigidBodyContactConstraint Constraint;
-	Constraint.ParticleIndex = ParticleIndex;
-	Constraint.LevelsetIndex = LevelsetIndex;
-	return Constraint;
-}
-
-typename FPBDCollisionConstraintPGS::FRigidBodyContactConstraint FPBDCollisionConstraintPGS::ComputeBoxConstraint(const FPBDRigidParticles& InParticles, int32 Box1Index, int32 Box2Index, const FReal Thickness)
-{
-	FRigidBodyContactConstraint Constraint;
-	Constraint.ParticleIndex = Box1Index;
-	Constraint.LevelsetIndex = Box2Index;
-	return Constraint;
-}
-
-typename FPBDCollisionConstraintPGS::FRigidBodyContactConstraint FPBDCollisionConstraintPGS::ComputeBoxPlaneConstraint(const FPBDRigidParticles& InParticles, int32 BoxIndex, int32 PlaneIndex, const FReal Thickness)
-{
-	FRigidBodyContactConstraint Constraint;
-	Constraint.ParticleIndex = BoxIndex;
-	Constraint.LevelsetIndex = PlaneIndex;
-	return Constraint;
-}
-
-typename FPBDCollisionConstraintPGS::FRigidBodyContactConstraint FPBDCollisionConstraintPGS::ComputeSphereConstraint(const FPBDRigidParticles& InParticles, int32 Sphere1Index, int32 Sphere2Index, const FReal Thickness)
-{
-	FRigidBodyContactConstraint Constraint;
-	Constraint.ParticleIndex = Sphere1Index;
-	Constraint.LevelsetIndex = Sphere2Index;
-	return Constraint;
-}
-
-typename FPBDCollisionConstraintPGS::FRigidBodyContactConstraint FPBDCollisionConstraintPGS::ComputeSpherePlaneConstraint(const FPBDRigidParticles& InParticles, int32 SphereIndex, int32 PlaneIndex, const FReal Thickness)
-{
-	FRigidBodyContactConstraint Constraint;
-	Constraint.ParticleIndex = SphereIndex;
-	Constraint.LevelsetIndex = PlaneIndex;
-	return Constraint;
-}
-
-typename FPBDCollisionConstraintPGS::FRigidBodyContactConstraint FPBDCollisionConstraintPGS::ComputeSphereBoxConstraint(const FPBDRigidParticles& InParticles, int32 SphereIndex, int32 BoxIndex, const FReal Thickness)
-{
-	FRigidBodyContactConstraint Constraint;
-	Constraint.ParticleIndex = SphereIndex;
-	Constraint.LevelsetIndex = BoxIndex;
-	return Constraint;
-}
-
-typename FPBDCollisionConstraintPGS::FRigidBodyContactConstraint FPBDCollisionConstraintPGS::ComputeConstraint(const FPBDRigidParticles& InParticles, int32 Body1Index, int32 Body2Index, const FReal Thickness)
-{
-	if (InParticles.Geometry(Body1Index)->GetType() == FAABB3::GetType() && InParticles.Geometry(Body2Index)->GetType() == FAABB3::GetType())
-	{
-		return ComputeBoxConstraint(InParticles, Body1Index, Body2Index, Thickness);
-	}
-	else if (InParticles.Geometry(Body1Index)->GetType() == TSphere<FReal, 3>::GetType() && InParticles.Geometry(Body2Index)->GetType() == TSphere<FReal, 3>::GetType())
-	{
-		return ComputeSphereConstraint(InParticles, Body1Index, Body2Index, Thickness);
-	}
-	else if (InParticles.Geometry(Body1Index)->GetType() == FAABB3::GetType() && InParticles.Geometry(Body2Index)->GetType() == TPlane<FReal, 3>::GetType())
-	{
-		return ComputeBoxPlaneConstraint(InParticles, Body1Index, Body2Index, Thickness);
-	}
-	else if (InParticles.Geometry(Body2Index)->GetType() == TPlane<FReal, 3>::GetType() && InParticles.Geometry(Body1Index)->GetType() == FAABB3::GetType())
-	{
-		return ComputeBoxPlaneConstraint(InParticles, Body2Index, Body1Index, Thickness);
-	}
-	else if (InParticles.Geometry(Body1Index)->GetType() == TSphere<FReal, 3>::GetType() && InParticles.Geometry(Body2Index)->GetType() == TPlane<FReal, 3>::GetType())
-	{
-		return ComputeSpherePlaneConstraint(InParticles, Body1Index, Body2Index, Thickness);
-	}
-	else if (InParticles.Geometry(Body2Index)->GetType() == TPlane<FReal, 3>::GetType() && InParticles.Geometry(Body1Index)->GetType() == TSphere<FReal, 3>::GetType())
-	{
-		return ComputeSpherePlaneConstraint(InParticles, Body2Index, Body1Index, Thickness);
-	}
-	else if (InParticles.Geometry(Body1Index)->GetType() == TSphere<FReal, 3>::GetType() && InParticles.Geometry(Body2Index)->GetType() == FAABB3::GetType())
-	{
-		return ComputeSphereBoxConstraint(InParticles, Body1Index, Body2Index, Thickness);
-	}
-	else if (InParticles.Geometry(Body2Index)->GetType() == FAABB3::GetType() && InParticles.Geometry(Body1Index)->GetType() == TSphere<FReal, 3>::GetType())
-	{
-		return ComputeSphereBoxConstraint(InParticles, Body2Index, Body1Index, Thickness);
-	}
-	else if (InParticles.Geometry(Body1Index)->IsConvex() && InParticles.Geometry(Body2Index)->IsConvex())
-	{
-		return ComputeLevelsetConstraintGJK(InParticles, Body1Index, Body2Index, Thickness);
-	}
-	return ComputeLevelsetConstraint(InParticles, Body1Index, Body2Index, Thickness);
-}
-
-template<class T_PARTICLES>
-void FPBDCollisionConstraintPGS::UpdateConstraint(const T_PARTICLES& InParticles, const FReal Thickness, FRigidBodyContactConstraint& Constraint)
-{
-	if (InParticles.Geometry(Constraint.ParticleIndex)->GetType() == FAABB3::GetType() && InParticles.Geometry(Constraint.LevelsetIndex)->GetType() == FAABB3::GetType())
-	{
-		UpdateBoxConstraint(InParticles, Thickness, Constraint);
-	}
-	else if (InParticles.Geometry(Constraint.ParticleIndex)->GetType() == TSphere<FReal, 3>::GetType() && InParticles.Geometry(Constraint.LevelsetIndex)->GetType() == TSphere<FReal, 3>::GetType())
-	{
-		UpdateSphereConstraint(InParticles, Thickness, Constraint);
-	}
-	else if (InParticles.Geometry(Constraint.ParticleIndex)->GetType() == FAABB3::GetType() && InParticles.Geometry(Constraint.LevelsetIndex)->GetType() == TPlane<FReal, 3>::GetType())
-	{
-		UpdateBoxPlaneConstraint(InParticles, Thickness, Constraint);
-	}
-	else if (InParticles.Geometry(Constraint.ParticleIndex)->GetType() == TSphere<FReal, 3>::GetType() && InParticles.Geometry(Constraint.LevelsetIndex)->GetType() == TPlane<FReal, 3>::GetType())
-	{
-		UpdateSpherePlaneConstraint(InParticles, Thickness, Constraint);
-	}
-	else if (InParticles.Geometry(Constraint.ParticleIndex)->GetType() == TSphere<FReal, 3>::GetType() && InParticles.Geometry(Constraint.LevelsetIndex)->GetType() == FAABB3::GetType())
-	{
-		UpdateSphereBoxConstraint(InParticles, Thickness, Constraint);
-	}
-	else if (InParticles.Geometry(Constraint.ParticleIndex)->GetType() == TPlane<FReal, 3>::GetType() && InParticles.Geometry(Constraint.LevelsetIndex)->GetType() == FAABB3::GetType())
-	{
-		int32 Tmp = Constraint.ParticleIndex;
-		Constraint.ParticleIndex = Constraint.LevelsetIndex;
-		Constraint.LevelsetIndex = Tmp;
-		UpdateBoxPlaneConstraint(InParticles, Thickness, Constraint);
-	}
-	else if (InParticles.Geometry(Constraint.ParticleIndex)->GetType() == TPlane<FReal, 3>::GetType() && InParticles.Geometry(Constraint.LevelsetIndex)->GetType() == TSphere<FReal, 3>::GetType())
-	{
-		int32 Tmp = Constraint.ParticleIndex;
-		Constraint.ParticleIndex = Constraint.LevelsetIndex;
-		Constraint.LevelsetIndex = Tmp;
-		UpdateSpherePlaneConstraint(InParticles, Thickness, Constraint);
-	}
-	else if (InParticles.Geometry(Constraint.ParticleIndex)->GetType() == FAABB3::GetType() && InParticles.Geometry(Constraint.LevelsetIndex)->GetType() == TSphere<FReal, 3>::GetType())
-	{
-		int32 Tmp = Constraint.ParticleIndex;
-		Constraint.ParticleIndex = Constraint.LevelsetIndex;
-		Constraint.LevelsetIndex = Tmp;
-		UpdateSphereBoxConstraint(InParticles, Thickness, Constraint);
-	}
-	else if (InParticles.Geometry(Constraint.ParticleIndex)->IsConvex() && InParticles.Geometry(Constraint.LevelsetIndex)->IsConvex())
-	{
-		UpdateLevelsetConstraintGJK(InParticles, Thickness, Constraint);
-	}
-	else
-	{
-		UpdateLevelsetConstraint(InParticles, Thickness, Constraint);
-	}
-}
-
-#endif
\ No newline at end of file
diff --git a/Engine/Source/Runtime/Experimental/Chaos/Public/Chaos/Collision/CollisionApplyType.h b/Engine/Source/Runtime/Experimental/Chaos/Public/Chaos/Collision/CollisionApplyType.h
--- a/Engine/Source/Runtime/Experimental/Chaos/Public/Chaos/Collision/CollisionApplyType.h
+++ b/Engine/Source/Runtime/Experimental/Chaos/Public/Chaos/Collision/CollisionApplyType.h
@@ -11,7 +11,6 @@ namespace Chaos
 	enum class EConstraintSolverType
 	{
 		None,
-		GbfPbd,						// Solve for velocity in the Apply step with PBD-style algorithm. Position correction + additional velocity correction in ApplyPushOut step (Original Chaos Main Solver)
 		StandardPbd,				// Solve for position in the Apply step with standard PBD algorithm. Calculate Implicit Velocities. Do nothing in the ApplyPushOut step (Original Chaos RBAN Solver)
 		QuasiPbd,					// Solve for position in the Apply step with standard PBD algorithm. Calculate Implicit Velocities. Solve for normal velocity in the ApplyPushOut step (New Chaos Solver)
 	};
diff --git a/Engine/Source/Runtime/Experimental/Chaos/Public/Chaos/Collision/PBDCollisionConstraint.h b/Engine/Source/Runtime/Experimental/Chaos/Public/Chaos/Collision/PBDCollisionConstraint.h
--- a/Engine/Source/Runtime/Experimental/Chaos/Public/Chaos/Collision/PBDCollisionConstraint.h
+++ b/Engine/Source/Runtime/Experimental/Chaos/Public/Chaos/Collision/PBDCollisionConstraint.h
@@ -98,12 +98,12 @@ namespace Chaos
 	};
 
 	/*
-	*
+	 * @brief Material properties for a collision constraint
 	*/
-	class CHAOS_API FCollisionContact
+	class CHAOS_API FPBDCollisionConstraintMaterial
 	{
 	public:
-		FCollisionContact()
+		FPBDCollisionConstraintMaterial()
 			: Friction(0)
 			, AngularFriction(0)
 			, Restitution(0)
@@ -119,7 +119,7 @@ namespace Chaos
 		FReal Friction;			// @todo(chaos): rename DynamicFriction
 		FReal AngularFriction;	// @todo(chaos): rename StaticFriction
 		FReal Restitution;
-		FReal RestitutionPadding; // For PBD implementation of resitution, we pad constraints on initial contact to enforce outward velocity
+		FReal RestitutionPadding; // For StandardPBD implementation of resitution, we pad constraints on initial contact to enforce outward velocity
 		FReal RestitutionThreshold;
 		FReal InvMassScale0;
 		FReal InvMassScale1;
@@ -132,6 +132,9 @@ namespace Chaos
 		}
 	};
 
+	// Renamed to FPBDCollisionConstraintMaterial
+	using FCollisionContact UE_DEPRECATED(5.1, "FCollisionContact was renamed to FPBDCollisionConstraintMaterial") = FPBDCollisionConstraintMaterial;
+
 
 	/**
 	 * @brief Information used by the constraint allocator
@@ -280,19 +283,6 @@ namespace Chaos
 		*/
 		void SetCCDEnabled(const bool bCCDEnabled) { CCDType = bCCDEnabled ? ECollisionCCDType::Enabled : ECollisionCCDType::Disabled; }
 
-		bool ContainsManifold(const FImplicitObject* A, const FBVHParticles* AS, const FImplicitObject* B, const FBVHParticles* BS) const
-		{
-			return A == Implicit[0] && B == Implicit[1] && AS == Simplicial[0] && BS == Simplicial[1];
-		}
-
-		void SetManifold(const FImplicitObject* A, const FBVHParticles* AS, const FImplicitObject* B, const FBVHParticles* BS)
-		{
-			Implicit[0] = A; 
-			Implicit[1] = B;
-			Simplicial[0] = AS; 
-			Simplicial[1] = BS;
-		}
-
 		//
 		// API
 		//
@@ -349,35 +339,35 @@ namespace Chaos
 		// @todo(chaos): remove (used by legacy RBAN collision solver)
 		FVec3 CalculateWorldContactLocation() const;
 
-		void SetInvMassScale0(const FReal InInvMassScale) { Manifold.InvMassScale0 = InInvMassScale; }
-		FReal GetInvMassScale0() const { return Manifold.InvMassScale0; }
+		void SetInvMassScale0(const FReal InInvMassScale) { Material.InvMassScale0 = InInvMassScale; }
+		FReal GetInvMassScale0() const { return Material.InvMassScale0; }
 
-		void SetInvMassScale1(const FReal InInvMassScale) { Manifold.InvMassScale1 = InInvMassScale; }
-		FReal GetInvMassScale1() const { return Manifold.InvMassScale1; }
+		void SetInvMassScale1(const FReal InInvMassScale) { Material.InvMassScale1 = InInvMassScale; }
+		FReal GetInvMassScale1() const { return Material.InvMassScale1; }
 
-		void SetInvInertiaScale0(const FReal InInvInertiaScale) { Manifold.InvInertiaScale0 = InInvInertiaScale; }
-		FReal GetInvInertiaScale0() const { return Manifold.InvInertiaScale0; }
+		void SetInvInertiaScale0(const FReal InInvInertiaScale) { Material.InvInertiaScale0 = InInvInertiaScale; }
+		FReal GetInvInertiaScale0() const { return Material.InvInertiaScale0; }
 
-		void SetInvInertiaScale1(const FReal InInvInertiaScale) { Manifold.InvInertiaScale1 = InInvInertiaScale; }
-		FReal GetInvInertiaScale1() const { return Manifold.InvInertiaScale1; }
+		void SetInvInertiaScale1(const FReal InInvInertiaScale) { Material.InvInertiaScale1 = InInvInertiaScale; }
+		FReal GetInvInertiaScale1() const { return Material.InvInertiaScale1; }
 
 		void SetStiffness(FReal InStiffness) { Stiffness = InStiffness; }
 		FReal GetStiffness() const { return Stiffness; }
 
-		void SetRestitution(const FReal InRestitution) { Manifold.Restitution = InRestitution; }
-		FReal GetRestitution() const { return Manifold.Restitution; }
+		void SetRestitution(const FReal InRestitution) { Material.Restitution = InRestitution; }
+		FReal GetRestitution() const { return Material.Restitution; }
 
-		void SetRestitutionThreshold(const FReal InRestitutionThreshold) { Manifold.RestitutionThreshold = InRestitutionThreshold; }
-		FReal GetRestitutionThreshold() const { return Manifold.RestitutionThreshold; }
+		void SetRestitutionThreshold(const FReal InRestitutionThreshold) { Material.RestitutionThreshold = InRestitutionThreshold; }
+		FReal GetRestitutionThreshold() const { return Material.RestitutionThreshold; }
 
-		void SetRestitutionPadding(const FReal InRestitutionPadding) { Manifold.RestitutionPadding = InRestitutionPadding; }
-		FReal GetRestitutionPadding() const { return Manifold.RestitutionPadding; }
+		void SetRestitutionPadding(const FReal InRestitutionPadding) { Material.RestitutionPadding = InRestitutionPadding; }
+		FReal GetRestitutionPadding() const { return Material.RestitutionPadding; }
 
-		void SetStaticFriction(const FReal InStaticFriction) { Manifold.AngularFriction = InStaticFriction; }
-		FReal GetStaticFriction() const { return FMath::Max(Manifold.AngularFriction, Manifold.Friction); }
+		void SetStaticFriction(const FReal InStaticFriction) { Material.AngularFriction = InStaticFriction; }
+		FReal GetStaticFriction() const { return FMath::Max(Material.AngularFriction, Material.Friction); }
 
-		void SetDynamicFriction(const FReal InDynamicFriction) { Manifold.Friction = InDynamicFriction; }
-		FReal GetDynamicFriction() const { return Manifold.Friction; }
+		void SetDynamicFriction(const FReal InDynamicFriction) { Material.Friction = InDynamicFriction; }
+		FReal GetDynamicFriction() const { return Material.Friction; }
 
 		EContactShapesType GetShapesType() const { return ShapesType; }
 
@@ -631,7 +621,8 @@ namespace Chaos
 		const FImplicitObject* Implicit[2];
 		const FPerShapeData* Shape[2];
 		const FBVHParticles* Simplicial[2];
-		FCollisionContact Manifold;// @todo(chaos): rename to FCollisionMaterial or something
+
+		FPBDCollisionConstraintMaterial Material;
 		FReal Stiffness;
 
 	public:
diff --git a/Engine/Source/Runtime/Experimental/Chaos/Public/Chaos/PBDCollisionConstraints.h b/Engine/Source/Runtime/Experimental/Chaos/Public/Chaos/PBDCollisionConstraints.h
--- a/Engine/Source/Runtime/Experimental/Chaos/Public/Chaos/PBDCollisionConstraints.h
+++ b/Engine/Source/Runtime/Experimental/Chaos/Public/Chaos/PBDCollisionConstraints.h
@@ -319,14 +319,12 @@ protected:
 	// Call PruneParticleEdgeCollisions on all particles with ECollisionConstraintFlags::CCF_SmoothEdgeCollisions set in CollisionFlags
 	void PruneEdgeCollisions();
 
-	// The "Legacy" functions handle the older solver types (GbfPbd and StandardPbd)
-	// @todo(chaos): remove legacy methods when the new solver is fully operational and used everywhere (RBAN)
+	// The "Legacy" functions handle the older solver types (StandardPbd)
+	// @todo(chaos): remove legacy methods when the new solver is fully operational and used everywhere (RBAN may still optionally use StandardPbd)
 	void LegacyGatherInput(const FReal Dt, FPBDCollisionConstraint& Constraint, const int32 Particle0Level, const int32 Particle1Level, FPBDIslandSolverData& SolverData);
 	void LegacyScatterOutput(const FReal Dt, const int32 BeginIndex, const int32 EndIndex, FPBDIslandSolverData& SolverData);
 	bool LegacyApplyPhase1Serial(const FReal Dt, const int32 Iterations, const int32 NumIterations, const int32 BeginIndex, const int32 EndIndex, FPBDIslandSolverData& SolverData);
 	bool LegacyApplyPhase1Parallel(const FReal Dt, const int32 Iterations, const int32 NumIterations, const int32 BeginIndex, const int32 EndIndex, FPBDIslandSolverData& SolverData);
-	bool LegacyApplyPhase2Serial(const FReal Dt, const int32 Iterations, const int32 NumIterations, const int32 BeginIndex, const int32 EndIndex, FPBDIslandSolverData& SolverData);
-	bool LegacyApplyPhase2Parallel(const FReal Dt, const int32 Iterations, const int32 NumIterations, const int32 BeginIndex, const int32 EndIndex, FPBDIslandSolverData& SolverData);
 
 private:
 
diff --git a/Engine/Source/Runtime/Experimental/Chaos/Public/Chaos/PBDCollisionConstraintsContact.h b/Engine/Source/Runtime/Experimental/Chaos/Public/Chaos/PBDCollisionConstraintsContact.h
--- a/Engine/Source/Runtime/Experimental/Chaos/Public/Chaos/PBDCollisionConstraintsContact.h
+++ b/Engine/Source/Runtime/Experimental/Chaos/Public/Chaos/PBDCollisionConstraintsContact.h
@@ -6,6 +6,11 @@
 #include "Chaos/CollisionResolutionTypes.h"
 #include "Chaos/Collision/CollisionApplyType.h"
 
+/**
+ * @brief This file contains collision solver code for the StandardPbd solver (which will be removed at some point)
+ * @see PBDCollisionSolver.h for the new solver
+*/
+
 namespace Chaos
 {
 	class FCollisionContext;
@@ -66,17 +71,12 @@ namespace Chaos
 			const int32 Iteration;
 			const int32 NumIterations;
 			const int32 NumPairIterations;
-			const EConstraintSolverType SolverType;	// @todo(chaos): a better way to customize the collision solver
+			const EConstraintSolverType SolverType;
 			bool* NeedsAnotherIteration;
 		};
 
-		// Update the constraint (re-runs collision detection for this contact)
 		extern void Update(FPBDCollisionConstraint& Constraint, const FReal Dt);
-		extern void UpdateSwept(FPBDCollisionConstraint& Constraint, const FReal Dt);
-
 		extern void Apply(FPBDCollisionConstraint& Constraint, const FContactIterationParameters& IterationParameters, const FContactParticleParameters& ParticleParameters);
-
-		extern void ApplyPushOut(FPBDCollisionConstraint& Constraint, const FContactIterationParameters& IterationParameters, const FContactParticleParameters& ParticleParameters);
 	}
 
 }
diff --git a/Engine/Source/Runtime/Experimental/Chaos/Public/Chaos/PBDCollisionConstraintsPGS.h b/Engine/Source/Runtime/Experimental/Chaos/Public/Chaos/PBDCollisionConstraintsPGS.h
--- a/Engine/Source/Runtime/Experimental/Chaos/Public/Chaos/PBDCollisionConstraintsPGS.h
+++ b/Engine/Source/Runtime/Experimental/Chaos/Public/Chaos/PBDCollisionConstraintsPGS.h
@@ -1,94 +1,3 @@
 // Copyright Epic Games, Inc. All Rights Reserved.
 #pragma once
 
-#include "Chaos/CollisionResolutionTypes.h"
-#include "Chaos/PBDConstraintGraph.h"
-
-#include <memory>
-#include <queue>
-#include <sstream>
-
-#if CHAOS_PARTICLEHANDLE_TODO
-
-namespace ChaosTest
-{
-	void CollisionPGS();
-	void CollisionPGS2();
-}
-
-namespace Chaos
-{
-class FPBDCollisionConstraintPGS
-{
-  public:
-	friend void ChaosTest::CollisionPGS();
-	friend void ChaosTest::CollisionPGS2();
-
-	typedef FRigidBodyContactConstraintPGS FRigidBodyContactConstraint;
-	typedef FPBDConstraintGraph FConstraintGraph;
-
-	FPBDCollisionConstraintPGS(FPBDRigidParticles& InParticles, const TArray<int32>& InIndices, TArrayCollectionArray<bool>& Collided, TArrayCollectionArray<TSerializablePtr<FChaosPhysicsMaterial>>& InPhysicsMaterials, const FReal Thickness = (FReal)0);
-	virtual ~FPBDCollisionConstraintPGS() {}
-
-	void ComputeConstraints(const FPBDRigidParticles& InParticles, const TArray<int32>& InIndices, const FReal Dt);
-
-	void Apply(FPBDRigidParticles& InParticles, const FReal Dt, const TArray<int32>& InConstraintIndices);
-	bool ApplyPushOut(FPBDRigidParticles& InParticles, const FReal Dt, const TArray<int32>& InConstraintIndices);
-	template<class T_PARTICLES> void Solve(T_PARTICLES& InParticles, const FReal Dt, const TArray<int32>& InConstraintIndices);
-
-	void RemoveConstraints(const TSet<uint32>& RemovedParticles);
-	void UpdateConstraints(const FPBDRigidParticles& InParticles, FReal Dt, const TSet<uint32>& AddedParticles, const TArray<uint32>& ActiveParticles);
-
-
-	static bool NearestPoint(TArray<Pair<FVec3, FVec3>>& Points, FVec3& Direction);
-
-	const TArray<FRigidBodyContactConstraint>& GetAllConstraints() const { return Constraints; }
-
-	int32 NumConstraints() const { return Constraints.Num(); }
-	TVec2<int32> ConstraintParticleIndices(int32 ContactIndex) const { return { Constraints[ContactIndex].ParticleIndex, Constraints[ContactIndex].LevelsetIndex }; }
-	void UpdatePositionBasedState(const FPBDRigidParticles& InParticles, const TArray<int32>& InIndices, const FReal Dt) { ComputeConstraints(InParticles, InIndices, Dt); }
-
-  private:
-	void PrintParticles(const FPBDRigidParticles& InParticles, const TArray<int32>& InConstraintIndices);
-	void PrintConstraints(const FPBDRigidParticles& InParticles, const TArray<int32>& InConstraintIndices);
-
-	template<class T_PARTICLES>
-	void UpdateConstraint(const T_PARTICLES& InParticles, const FReal Thickness, FRigidBodyContactConstraint& Constraint);
-	template<class T_PARTICLES>
-	void UpdateLevelsetConstraint(const T_PARTICLES& InParticles, const FReal Thickness, FRigidBodyContactConstraint& Constraint);
-	template<class T_PARTICLES>
-	void UpdateLevelsetConstraintGJK(const T_PARTICLES& InParticles, const FReal Thickness, FRigidBodyContactConstraint& Constraint);
-	template<class T_PARTICLES>
-	void UpdateBoxConstraint(const T_PARTICLES& InParticles, const FReal Thickness, FRigidBodyContactConstraint& Constraint);
-	template<class T_PARTICLES>
-	void UpdateBoxPlaneConstraint(const T_PARTICLES& InParticles, const FReal Thickness, FRigidBodyContactConstraint& Constraint);
-	template<class T_PARTICLES>
-	void UpdateSphereConstraint(const T_PARTICLES& InParticles, const FReal Thickness, FRigidBodyContactConstraint& Constraint);
-	template<class T_PARTICLES>
-	void UpdateSpherePlaneConstraint(const T_PARTICLES& InParticles, const FReal Thickness, FRigidBodyContactConstraint& Constraint);
-	template<class T_PARTICLES>
-	void UpdateSphereBoxConstraint(const T_PARTICLES& InParticles, const FReal Thickness, FRigidBodyContactConstraint& Constraint);
-
-	FRigidBodyContactConstraint ComputeLevelsetConstraint(const FPBDRigidParticles& InParticles, int32 ParticleIndex, int32 LevelsetIndex, const FReal Thickness);
-	FRigidBodyContactConstraint ComputeLevelsetConstraintGJK(const FPBDRigidParticles& InParticles, int32 ParticleIndex, int32 LevelsetIndex, const FReal Thickness);
-	FRigidBodyContactConstraint ComputeBoxConstraint(const FPBDRigidParticles& InParticles, int32 Box1Index, int32 Box2Index, const FReal Thickness);
-	FRigidBodyContactConstraint ComputeBoxPlaneConstraint(const FPBDRigidParticles& InParticles, int32 BoxIndex, int32 PlaneIndex, const FReal Thickness);
-	FRigidBodyContactConstraint ComputeSphereConstraint(const FPBDRigidParticles& InParticles, int32 Sphere1Index, int32 Sphere2Index, const FReal Thickness);
-	FRigidBodyContactConstraint ComputeSpherePlaneConstraint(const FPBDRigidParticles& InParticles, int32 SphereIndex, int32 PlaneIndex, const FReal Thickness);
-	FRigidBodyContactConstraint ComputeSphereBoxConstraint(const FPBDRigidParticles& InParticles, int32 SphereIndex, int32 BoxIndex, const FReal Thickness);
-	FRigidBodyContactConstraint ComputeConstraint(const FPBDRigidParticles& InParticles, int32 Body1Index, int32 Body2Index, const FReal Thickness);
-
-	TArray<FRigidBodyContactConstraint> Constraints;
-	TArrayCollectionArray<bool>& MCollided;
-	const TArrayCollectionArray<TSerializablePtr<FChaosPhysicsMaterial>>& MPhysicsMaterials;
-	const FReal MThickness;
-	const FReal Tolerance;
-	const int32 MaxIterations;
-	const bool bUseCCD;
-};
-
-template<class T, int d>
-using TPBDCollisionConstraintPGS UE_DEPRECATED(4.27, "Deprecated. this class is to be deleted, use FPBDCollisionConstraintPGS instead") = FPBDCollisionConstraintPGS;
-}
-
-#endif
\ No newline at end of file
