Chaos:Cloth
Updating stiffness parameter calculations to prevent floating point exception

diff --git a/Engine/Source/Runtime/Experimental/Chaos/Public/Chaos/PBDStiffness.h b/Engine/Source/Runtime/Experimental/Chaos/Public/Chaos/PBDStiffness.h
--- a/Engine/Source/Runtime/Experimental/Chaos/Public/Chaos/PBDStiffness.h
+++ b/Engine/Source/Runtime/Experimental/Chaos/Public/Chaos/PBDStiffness.h
@@ -193,6 +193,12 @@ void FPBDStiffness::ApplyValues(const FSolverReal Dt, const int32 NumIterations)
 	// Define the stiffness mapping function
 	auto SimulationValue = [this, Exponent](const FSolverReal InValue)->FSolverReal
 	{
+		// If InValue is 1 then LogValue = -inf and the output becomes -inf as well
+		// In order for this function to be continuous, We want the output to be 1 when InValue = 1
+		if (InValue == (FSolverReal)1.)
+		{
+			return (FSolverReal)1.;
+		}
 		// Get a very steep exponential curve between the [0, 1] range to make easier to set the parameter
 		// The base has been chosen empirically
 		// ParameterValue = Pow(ParameterFitBase, ParameterValue - 1)
