Chaos: Fix move constructors for all particle SOAs

diff --git a/Engine/Source/Runtime/Experimental/Chaos/Public/Chaos/GeometryParticles.h b/Engine/Source/Runtime/Experimental/Chaos/Public/Chaos/GeometryParticles.h
--- a/Engine/Source/Runtime/Experimental/Chaos/Public/Chaos/GeometryParticles.h
+++ b/Engine/Source/Runtime/Experimental/Chaos/Public/Chaos/GeometryParticles.h
@@ -429,7 +429,9 @@ namespace Chaos
 			TArrayCollection::AddArray(&MGeometry);
 			TArrayCollection::AddArray(&MSharedGeometry);
 			TArrayCollection::AddArray(&MDynamicGeometry);
+#if CHAOS_DETERMINISTIC
 			TArrayCollection::AddArray(&MParticleIDs);
+#endif
 			TArrayCollection::AddArray(&MHasCollision);
 			TArrayCollection::AddArray(&MShapesArray);
 			TArrayCollection::AddArray(&MLocalBounds);
@@ -483,11 +485,13 @@ namespace Chaos
 			, MResimType(MoveTemp(Other.MResimType))
 			, MEnabledDuringResim(MoveTemp(Other.MEnabledDuringResim))
 			, MLightWeightDisabled(MoveTemp(Other.MLightWeightDisabled))
-
-
 #if CHAOS_DETERMINISTIC
 			, MParticleIDs(MoveTemp(Other.MParticleIDs))
 #endif
+#if CHAOS_DEBUG_NAME
+			, MDebugName(MoveTemp(Other.MDebugName))
+#endif
+
 		{
 			MParticleType = EParticleType::Static;
 			TArrayCollection::AddArray(&MUniqueIdx);
@@ -495,6 +499,9 @@ namespace Chaos
 			TArrayCollection::AddArray(&MGeometry);
 			TArrayCollection::AddArray(&MSharedGeometry);
 			TArrayCollection::AddArray(&MDynamicGeometry);
+#if CHAOS_DETERMINISTIC
+			TArrayCollection::AddArray(&MParticleIDs);
+#endif
 			TArrayCollection::AddArray(&MHasCollision);
 			TArrayCollection::AddArray(&MShapesArray);
 			TArrayCollection::AddArray(&MLocalBounds);
@@ -506,10 +513,10 @@ namespace Chaos
 			TArrayCollection::AddArray(&MWeakParticleHandle);
 			TArrayCollection::AddArray(&MParticleConstraints);
 			TArrayCollection::AddArray(&MParticleCollisions);
+			TArrayCollection::AddArray(&MResimType);
+			TArrayCollection::AddArray(&MEnabledDuringResim);
+			TArrayCollection::AddArray(&MLightWeightDisabled);
 
-#if CHAOS_DETERMINISTIC
-			TArrayCollection::AddArray(&MParticleIDs);
-#endif
 #if CHAOS_DEBUG_NAME
 			TArrayCollection::AddArray(&MDebugName);
 #endif
@@ -520,10 +527,6 @@ namespace Chaos
 				TArrayCollection::AddArray(&MGeometryParticle);
 				TArrayCollection::AddArray(&MPhysicsProxy);
 			}
-
-			TArrayCollection::AddArray(&MResimType);
-			TArrayCollection::AddArray(&MEnabledDuringResim);
-			TArrayCollection::AddArray(&MLightWeightDisabled);
 		}
 
 		static constexpr bool IsRigidBodySim() { return SimType == EGeometryParticlesSimType::RigidBodySim; }
@@ -537,6 +540,9 @@ namespace Chaos
 			TArrayCollection::AddArray(&MGeometry);
 			TArrayCollection::AddArray(&MSharedGeometry);
 			TArrayCollection::AddArray(&MDynamicGeometry);
+#if CHAOS_DETERMINISTIC
+			TArrayCollection::AddArray(&MParticleIDs);
+#endif
 			TArrayCollection::AddArray(&MHasCollision);
 			TArrayCollection::AddArray(&MShapesArray);
 			TArrayCollection::AddArray(&MLocalBounds);
@@ -548,10 +554,11 @@ namespace Chaos
 			TArrayCollection::AddArray(&MWeakParticleHandle);
 			TArrayCollection::AddArray(&MParticleConstraints);
 			TArrayCollection::AddArray(&MParticleCollisions);
+			TArrayCollection::AddArray(&MResimType);
+			TArrayCollection::AddArray(&MEnabledDuringResim);
+			TArrayCollection::AddArray(&MLightWeightDisabled);
+
 
-#if CHAOS_DETERMINISTIC
-			TArrayCollection::AddArray(&MParticleIDs);
-#endif
 #if CHAOS_DEBUG_NAME
 			TArrayCollection::AddArray(&MDebugName);
 #endif
@@ -562,10 +569,6 @@ namespace Chaos
 				TArrayCollection::AddArray(&MGeometryParticle);
 				TArrayCollection::AddArray(&MPhysicsProxy);
 			}
-			
-			TArrayCollection::AddArray(&MResimType);
-			TArrayCollection::AddArray(&MEnabledDuringResim);
-			TArrayCollection::AddArray(&MLightWeightDisabled);
 		}
 
 		CHAOS_API virtual ~TGeometryParticlesImp()
diff --git a/Engine/Source/Runtime/Experimental/Chaos/Public/Chaos/KinematicGeometryParticles.h b/Engine/Source/Runtime/Experimental/Chaos/Public/Chaos/KinematicGeometryParticles.h
--- a/Engine/Source/Runtime/Experimental/Chaos/Public/Chaos/KinematicGeometryParticles.h
+++ b/Engine/Source/Runtime/Experimental/Chaos/Public/Chaos/KinematicGeometryParticles.h
@@ -22,7 +22,7 @@ class TKinematicGeometryParticlesImp : public TGeometryParticlesImp<T, d, SimTyp
 	}
 	TKinematicGeometryParticlesImp(const TKinematicGeometryParticlesImp<T, d, SimType>& Other) = delete;
 	TKinematicGeometryParticlesImp(TKinematicGeometryParticlesImp<T, d, SimType>&& Other)
-	    : TGeometryParticlesImp<T, d, SimType>(MoveTemp(Other)), MV(MoveTemp(Other.MV)), MW(MoveTemp(Other.MW))
+	    : TGeometryParticlesImp<T, d, SimType>(MoveTemp(Other)), MV(MoveTemp(Other.MV)), MW(MoveTemp(Other.MW)), KinematicTargets(MoveTemp(Other.KinematicTargets))
 	{
 		this->MParticleType = EParticleType::Kinematic;
 		TArrayCollection::AddArray(&MV);
diff --git a/Engine/Source/Runtime/Experimental/Chaos/Public/Chaos/PBDRigidClusteredParticles.h b/Engine/Source/Runtime/Experimental/Chaos/Public/Chaos/PBDRigidClusteredParticles.h
--- a/Engine/Source/Runtime/Experimental/Chaos/Public/Chaos/PBDRigidClusteredParticles.h
+++ b/Engine/Source/Runtime/Experimental/Chaos/Public/Chaos/PBDRigidClusteredParticles.h
@@ -82,6 +82,7 @@ class TPBDRigidClusteredParticles : public TPBDRigidParticles<T, d>
 	, MChildToParent(MoveTemp(Other.MChildToParent))
 	, MClusterGroupIndex(MoveTemp(Other.MClusterGroupIndex))
 	, MInternalCluster(MoveTemp(Other.MInternalCluster))
+	, MChildrenSpatial(MoveTemp(Other.MChildrenSpatial))
 	, MPhysicsProxies(MoveTemp(Other.MPhysicsProxies))
 	, MCollisionImpulses(MoveTemp(Other.MCollisionImpulses))
 	, MStrains(MoveTemp(Other.MStrains))
diff --git a/Engine/Source/Runtime/Experimental/Chaos/Public/Chaos/RigidParticles.h b/Engine/Source/Runtime/Experimental/Chaos/Public/Chaos/RigidParticles.h
--- a/Engine/Source/Runtime/Experimental/Chaos/Public/Chaos/RigidParticles.h
+++ b/Engine/Source/Runtime/Experimental/Chaos/Public/Chaos/RigidParticles.h
@@ -106,22 +106,29 @@ public:
 		, MAngularAcceleration(MoveTemp(Other.MAngularAcceleration))
 		, MLinearImpulseVelocity(MoveTemp(Other.MLinearImpulseVelocity))
 		, MAngularImpulseVelocity(MoveTemp(Other.MAngularImpulseVelocity))
-		, MI(MoveTemp(Other.MI)), MInvI(MoveTemp(Other.MInvI))
+		, MI(MoveTemp(Other.MI))
+		, MInvI(MoveTemp(Other.MInvI))
 		, MM(MoveTemp(Other.MM))
 		, MInvM(MoveTemp(Other.MInvM))
 		, MCenterOfMass(MoveTemp(Other.MCenterOfMass))
 		, MRotationOfMass(MoveTemp(Other.MRotationOfMass))
+		, MLinearEtherDrag(MoveTemp(Other.MLinearEtherDrag))
+		, MAngularEtherDrag(MoveTemp(Other.MAngularEtherDrag))
 		, MaxLinearSpeedsSq(MoveTemp(Other.MaxLinearSpeedsSq))
 		, MaxAngularSpeedsSq(MoveTemp(Other.MaxAngularSpeedsSq))
 		, MCollisionParticles(MoveTemp(Other.MCollisionParticles))
 		, MCollisionGroup(MoveTemp(Other.MCollisionGroup))
 		, MCollisionConstraintFlags(MoveTemp(Other.MCollisionConstraintFlags))
+		, MIslandIndex(MoveTemp(Other.MIslandIndex))
+		, MGraphIndex(MoveTemp(Other.MGraphIndex))
+		, MToBeRemovedOnFracture(MoveTemp(Other.MToBeRemovedOnFracture))
 		, MObjectState(MoveTemp(Other.MObjectState))
 		, MPreObjectState(MoveTemp(Other.MPreObjectState))
 		, MGravityEnabled(MoveTemp(Other.MGravityEnabled))
 		, MOneWayInteraction(MoveTemp(Other.MOneWayInteraction))
 		, MSleepType(MoveTemp(Other.MSleepType))
 		, bCCDEnabled(MoveTemp(Other.bCCDEnabled))
+		, MDisabled(MoveTemp(Other.MDisabled))
 	{
 		TArrayCollection::AddArray(&MVSmooth);
 		TArrayCollection::AddArray(&MWSmooth);
diff --git a/Engine/Source/Runtime/Experimental/ChaosCore/Public/Chaos/ArrayCollectionArray.h b/Engine/Source/Runtime/Experimental/ChaosCore/Public/Chaos/ArrayCollectionArray.h
--- a/Engine/Source/Runtime/Experimental/ChaosCore/Public/Chaos/ArrayCollectionArray.h
+++ b/Engine/Source/Runtime/Experimental/ChaosCore/Public/Chaos/ArrayCollectionArray.h
@@ -24,7 +24,7 @@ public:
 	TArrayCollectionArray()
 	    : TArray<T>() {}
 	TArrayCollectionArray(const TArrayCollectionArray<T>& Other) = delete;
-	TArrayCollectionArray(TArrayCollectionArray<T>&& Other)
+	explicit TArrayCollectionArray(TArrayCollectionArray<T>&& Other)
 	    : TArray<T>(MoveTemp(Other)) {}
 	TArrayCollectionArray& operator=(TArrayCollectionArray<T>&& Other)
 	{
