Chaos:
- Clear handle memory on destruction to force crash on dangling handles
- Add define to hunt down more complex dangling handle cases

diff --git a/Engine/Source/Runtime/Experimental/Chaos/Public/Chaos/ParticleHandle.h b/Engine/Source/Runtime/Experimental/Chaos/Public/Chaos/ParticleHandle.h
--- a/Engine/Source/Runtime/Experimental/Chaos/Public/Chaos/ParticleHandle.h
+++ b/Engine/Source/Runtime/Experimental/Chaos/Public/Chaos/ParticleHandle.h
@@ -14,6 +14,10 @@
 #include "Chaos/ChaosDebugDraw.h"
 #endif
 
+#ifndef UE_DEBUG_DANGLING_HANDLES
+#define UE_DEBUG_DANGLING_HANDLES 0 // Will deliberately cause a memory leak
+#endif
+
 class IPhysicsProxyBase;
 
 namespace Chaos
@@ -462,6 +466,12 @@ public:
 				}
 			}
 		}
+
+		// Zero the handle out to detect dangling handles and associated memory corruptions
+		HandleIdx = INDEX_NONE;
+		GeometryParticles = nullptr;
+		ParticleIdx = INDEX_NONE;
+		Type = EParticleType::Unknown;
 	}
 
 	template <typename T2, int d2>
@@ -1197,22 +1207,22 @@ public:
 };
 
 template <typename T, int d, bool bPersistent>
-TKinematicGeometryParticleHandleImp<T, d, bPersistent>* TGeometryParticleHandleImp<T, d, bPersistent>::CastToKinematicParticle() { return Type >= EParticleType::Kinematic ? static_cast<TKinematicGeometryParticleHandleImp<T, d, bPersistent>*>(this) : nullptr; }
+TKinematicGeometryParticleHandleImp<T, d, bPersistent>* TGeometryParticleHandleImp<T, d, bPersistent>::CastToKinematicParticle() { checkSlow(Type <= EParticleType::Clustered || Type == EParticleType::GeometryCollection);  return Type >= EParticleType::Kinematic ? static_cast<TKinematicGeometryParticleHandleImp<T, d, bPersistent>*>(this) : nullptr; }
 
 template <typename T, int d, bool bPersistent>
-const TKinematicGeometryParticleHandleImp<T, d, bPersistent>* TGeometryParticleHandleImp<T,d, bPersistent>::CastToKinematicParticle() const { return Type >= EParticleType::Kinematic ? static_cast<const TKinematicGeometryParticleHandleImp<T, d, bPersistent>*>(this) : nullptr; }
+const TKinematicGeometryParticleHandleImp<T, d, bPersistent>* TGeometryParticleHandleImp<T,d, bPersistent>::CastToKinematicParticle() const { checkSlow(Type <= EParticleType::Clustered || Type == EParticleType::GeometryCollection); return Type >= EParticleType::Kinematic ? static_cast<const TKinematicGeometryParticleHandleImp<T, d, bPersistent>*>(this) : nullptr; }
 
 template <typename T, int d, bool bPersistent>
-const TPBDRigidParticleHandleImp<T, d, bPersistent>* TGeometryParticleHandleImp<T, d, bPersistent>::CastToRigidParticle() const { return Type >= EParticleType::Rigid ? static_cast<const TPBDRigidParticleHandleImp<T, d, bPersistent>*>(this) : nullptr; }
+const TPBDRigidParticleHandleImp<T, d, bPersistent>* TGeometryParticleHandleImp<T, d, bPersistent>::CastToRigidParticle() const { checkSlow(Type <= EParticleType::Clustered || Type == EParticleType::GeometryCollection); return Type >= EParticleType::Rigid ? static_cast<const TPBDRigidParticleHandleImp<T, d, bPersistent>*>(this) : nullptr; }
 
 template <typename T, int d, bool bPersistent>
-TPBDRigidParticleHandleImp<T, d, bPersistent>* TGeometryParticleHandleImp<T, d, bPersistent>::CastToRigidParticle() { return Type >= EParticleType::Rigid ? static_cast<TPBDRigidParticleHandleImp<T, d, bPersistent>*>(this) : nullptr; }
+TPBDRigidParticleHandleImp<T, d, bPersistent>* TGeometryParticleHandleImp<T, d, bPersistent>::CastToRigidParticle() { checkSlow(Type <= EParticleType::Clustered || Type == EParticleType::GeometryCollection); return Type >= EParticleType::Rigid ? static_cast<TPBDRigidParticleHandleImp<T, d, bPersistent>*>(this) : nullptr; }
 
 template <typename T, int d, bool bPersistent>
-const TPBDRigidClusteredParticleHandleImp<T, d, bPersistent>* TGeometryParticleHandleImp<T, d, bPersistent>::CastToClustered() const { return Type >= EParticleType::Clustered ? static_cast<const TPBDRigidClusteredParticleHandleImp<T, d, bPersistent>*>(this) : nullptr; }
+const TPBDRigidClusteredParticleHandleImp<T, d, bPersistent>* TGeometryParticleHandleImp<T, d, bPersistent>::CastToClustered() const { checkSlow(Type <= EParticleType::Clustered || Type == EParticleType::GeometryCollection); return Type >= EParticleType::Clustered ? static_cast<const TPBDRigidClusteredParticleHandleImp<T, d, bPersistent>*>(this) : nullptr; }
 
 template <typename T, int d, bool bPersistent>
-TPBDRigidClusteredParticleHandleImp<T, d, bPersistent>* TGeometryParticleHandleImp<T, d, bPersistent>::CastToClustered() { return Type >= EParticleType::Clustered ? static_cast<TPBDRigidClusteredParticleHandleImp<T, d, bPersistent>*>(this) : nullptr; }
+TPBDRigidClusteredParticleHandleImp<T, d, bPersistent>* TGeometryParticleHandleImp<T, d, bPersistent>::CastToClustered() { checkSlow( Type <= EParticleType::Clustered || Type == EParticleType::GeometryCollection); return Type >= EParticleType::Clustered ? static_cast<TPBDRigidClusteredParticleHandleImp<T, d, bPersistent>*>(this) : nullptr; }
 
 template <typename T, int d, bool bPersistent>
 EObjectStateType TGeometryParticleHandleImp<T,d, bPersistent>::ObjectState() const
@@ -1786,6 +1796,14 @@ public:
 	void DestroyHandleSwap(TGeometryParticleHandle<T,d>* Handle)
 	{
 		const int32 UnstableIdx = Handle->HandleIdx;
+
+#if UE_DEBUG_DANGLING_HANDLES
+		// This helps to detect dangling handles and associated memory corruptions
+		// Leaking memory deliberately here, so only use this while debugging
+		Handles[UnstableIdx]->~TGeometryParticleHandleImp();
+		Handles[UnstableIdx].Release();
+#endif
+
 		RemoveAtSwapHelper(UnstableIdx);
 		if (static_cast<uint32>(UnstableIdx) < Size())
 		{
diff --git a/Engine/Source/Runtime/Experimental/ChaosCore/Public/Chaos/ArrayCollection.h b/Engine/Source/Runtime/Experimental/ChaosCore/Public/Chaos/ArrayCollection.h
--- a/Engine/Source/Runtime/Experimental/ChaosCore/Public/Chaos/ArrayCollection.h
+++ b/Engine/Source/Runtime/Experimental/ChaosCore/Public/Chaos/ArrayCollection.h
@@ -13,7 +13,14 @@ public:
 	    : MSize(0) {}
 	TArrayCollection(const TArrayCollection& Other) = delete;
 	TArrayCollection(TArrayCollection&& Other) = delete;
-	virtual ~TArrayCollection() {}
+	virtual ~TArrayCollection() 
+	{
+		// Null out to find dangling pointers
+		for (int32 Index = 0; Index < MArrays.Num(); Index++)
+		{
+			MArrays[Index] = nullptr;
+		}
+	}
 
 	int32 AddArray(TArrayCollectionArrayBase* Array)
 	{
