Fix fields applying commands to other clients in multi-client PIE session

diff --git a/Engine/Source/Runtime/Experimental/FieldSystem/Source/FieldSystemEngine/Private/Field/FieldSystemComponent.cpp b/Engine/Source/Runtime/Experimental/FieldSystem/Source/FieldSystemEngine/Private/Field/FieldSystemComponent.cpp
--- a/Engine/Source/Runtime/Experimental/FieldSystem/Source/FieldSystemEngine/Private/Field/FieldSystemComponent.cpp
+++ b/Engine/Source/Runtime/Experimental/FieldSystem/Source/FieldSystemEngine/Private/Field/FieldSystemComponent.cpp
@@ -87,7 +87,18 @@ TArray<Chaos::FPhysicsSolverBase*> UFieldSystemComponent::GetPhysicsSolvers() co
 		}
 	}
 
-	const TArray<Chaos::FPhysicsSolverBase*>& WorldSolvers = ChaosModule->GetAllSolvers();
+	TArray<Chaos::FPhysicsSolverBase*> WorldSolvers;
+
+	//Need to only grab solvers that are owned by our world. In multi-client PIE this avoids solvers for other clients
+	UWorld* World = GetWorld();
+	FPhysScene* PhysScene = World ? World->GetPhysicsScene() : nullptr;
+	if(PhysScene)
+	{
+		WorldSolvers.Add(PhysScene->GetSolver());
+		ChaosModule->GetSolversMutable(World, WorldSolvers);
+	}
+
+	
 	const int32 NumFilterSolvers = FilterSolvers.Num();
 
 	for (Chaos::FPhysicsSolverBase* Solver : WorldSolvers)
