Math macro prefix fixes for Chaos.

diff --git a/Engine/Source/Runtime/Experimental/Chaos/Private/Chaos/AABB.cpp b/Engine/Source/Runtime/Experimental/Chaos/Private/Chaos/AABB.cpp
--- a/Engine/Source/Runtime/Experimental/Chaos/Private/Chaos/AABB.cpp
+++ b/Engine/Source/Runtime/Experimental/Chaos/Private/Chaos/AABB.cpp
@@ -97,7 +97,7 @@ template <typename T, int d>
 bool TAABB<T, d>::Raycast(const TVector<FReal, d>& StartPoint, const TVector<FReal, d>& Dir, const FReal Length, const FReal Thickness, FReal& OutTime, TVector<FReal, d>& OutPosition, TVector<FReal, d>& OutNormal, int32& OutFaceIndex) const
 {
 	ensure(Length > 0);
-	ensure(FMath::IsNearlyEqual(Dir.SizeSquared(), (FReal)1, (FReal)KINDA_SMALL_NUMBER));
+	ensure(FMath::IsNearlyEqual(Dir.SizeSquared(), (FReal)1, (FReal)UE_KINDA_SMALL_NUMBER));
 
 	OutFaceIndex = INDEX_NONE;
 	const TVector<FReal, d> MinInflated = TVector<FReal, d>(MMin) - Thickness;
diff --git a/Engine/Source/Runtime/Experimental/Chaos/Private/Chaos/ChaosDebugDraw.cpp b/Engine/Source/Runtime/Experimental/Chaos/Private/Chaos/ChaosDebugDraw.cpp
--- a/Engine/Source/Runtime/Experimental/Chaos/Private/Chaos/ChaosDebugDraw.cpp
+++ b/Engine/Source/Runtime/Experimental/Chaos/Private/Chaos/ChaosDebugDraw.cpp
@@ -600,14 +600,14 @@ namespace Chaos
 			{
 				const TSphere<FReal, 3>* Sphere = Implicit->template GetObject<TSphere<FReal, 3>>();
 				const FVec3 P = ShapeTransform.TransformPosition(Sphere->GetCenter());
-				FDebugDrawQueue::GetInstance().DrawDebugSphere(P, Sphere->GetRadius(), 8, ShapeColor, false, KINDA_SMALL_NUMBER, Settings.DrawPriority, Settings.ShapeThicknesScale * Settings.LineThickness);
+				FDebugDrawQueue::GetInstance().DrawDebugSphere(P, Sphere->GetRadius(), 8, ShapeColor, false, UE_KINDA_SMALL_NUMBER, Settings.DrawPriority, Settings.ShapeThicknesScale * Settings.LineThickness);
 				break;
 			}
 			case ImplicitObjectType::Box:
 			{
 				const TBox<FReal, 3>* Box = Implicit->template GetObject<TBox<FReal, 3>>();
 				const FVec3 P = ShapeTransform.TransformPosition(Box->GetCenter());
-				FDebugDrawQueue::GetInstance().DrawDebugBox(P, (FReal)0.5 * Box->Extents(), ShapeTransform.GetRotation(), ShapeColor, false, KINDA_SMALL_NUMBER, Settings.DrawPriority, Settings.ShapeThicknesScale * Settings.LineThickness);
+				FDebugDrawQueue::GetInstance().DrawDebugBox(P, (FReal)0.5 * Box->Extents(), ShapeTransform.GetRotation(), ShapeColor, false, UE_KINDA_SMALL_NUMBER, Settings.DrawPriority, Settings.ShapeThicknesScale * Settings.LineThickness);
 				break;
 			}
 			case ImplicitObjectType::Plane:
@@ -617,7 +617,7 @@ namespace Chaos
 				const FCapsule* Capsule = Implicit->template GetObject<FCapsule>();
 				const FVec3 P = ShapeTransform.TransformPosition(Capsule->GetCenter());
 				const FRotation3 Q = ShapeTransform.GetRotation() * FRotationMatrix::MakeFromZ(Capsule->GetAxis());
-				FDebugDrawQueue::GetInstance().DrawDebugCapsule(P, (FReal)0.5 * Capsule->GetHeight() + Capsule->GetRadius(), Capsule->GetRadius(), Q, ShapeColor, false, KINDA_SMALL_NUMBER, Settings.DrawPriority, Settings.ShapeThicknesScale * Settings.LineThickness);
+				FDebugDrawQueue::GetInstance().DrawDebugCapsule(P, (FReal)0.5 * Capsule->GetHeight() + Capsule->GetRadius(), Capsule->GetRadius(), Q, ShapeColor, false, UE_KINDA_SMALL_NUMBER, Settings.DrawPriority, Settings.ShapeThicknesScale * Settings.LineThickness);
 				break;
 			}
 			case ImplicitObjectType::LevelSet:
@@ -682,7 +682,7 @@ namespace Chaos
 						for (int32 ParticleIndex = 0; ParticleIndex < (int32)Particles->Size(); ++ParticleIndex)
 						{
 							FVec3 P = ShapeTransform.TransformPosition(Particles->X(ParticleIndex));
-							FDebugDrawQueue::GetInstance().DrawDebugPoint(P, ShapeColor, false, KINDA_SMALL_NUMBER, Settings.DrawPriority, Settings.PointSize);
+							FDebugDrawQueue::GetInstance().DrawDebugPoint(P, ShapeColor, false, UE_KINDA_SMALL_NUMBER, Settings.DrawPriority, Settings.PointSize);
 						}
 					}
 				}
@@ -693,7 +693,7 @@ namespace Chaos
 				const FColor ShapeBoundsColor = FColor::Orange;
 				const FAABB3& ShapeBounds = Implicit->BoundingBox();
 				const FVec3 ShapeBoundsPos = ShapeTransform.TransformPosition(ShapeBounds.Center());
-				FDebugDrawQueue::GetInstance().DrawDebugBox(ShapeBoundsPos, 0.5f * ShapeBounds.Extents(), ShapeTransform.GetRotation(), ShapeBoundsColor, false, KINDA_SMALL_NUMBER, Settings.DrawPriority, Settings.LineThickness);
+				FDebugDrawQueue::GetInstance().DrawDebugBox(ShapeBoundsPos, 0.5f * ShapeBounds.Extents(), ShapeTransform.GetRotation(), ShapeBoundsColor, false, UE_KINDA_SMALL_NUMBER, Settings.DrawPriority, Settings.LineThickness);
 			}
 		}
 
@@ -759,7 +759,7 @@ namespace Chaos
 				}
 			}
 
-			FDebugDrawQueue::GetInstance().DrawDebugBox(P, 0.5f * Box.Extents(), Q, Color, false, KINDA_SMALL_NUMBER, Settings.DrawPriority, Settings.LineThickness);
+			FDebugDrawQueue::GetInstance().DrawDebugBox(P, 0.5f * Box.Extents(), Q, Color, false, UE_KINDA_SMALL_NUMBER, Settings.DrawPriority, Settings.LineThickness);
 
 			for (const auto& Shape : InParticle->ShapesArray())
 			{
@@ -772,7 +772,7 @@ namespace Chaos
 					const FRotation3 ShapeQ = SpaceTransform.GetRotation();
 					const FColor ShapeColor = (bChaosDebugDebugDrawColorBoundsByShapeType) ? Settings.BoundsColorsPerShapeType.GetColorFromShapeType(ShapeType) : Color;
 
-					FDebugDrawQueue::GetInstance().DrawDebugBox(ShapeP, 0.5f * ShapeBox.Extents(), ShapeQ, ShapeColor, false, KINDA_SMALL_NUMBER, Settings.DrawPriority, Settings.LineThickness);
+					FDebugDrawQueue::GetInstance().DrawDebugBox(ShapeP, 0.5f * ShapeBox.Extents(), ShapeQ, ShapeColor, false, UE_KINDA_SMALL_NUMBER, Settings.DrawPriority, Settings.LineThickness);
 				}
 			}
 		}
@@ -787,29 +787,29 @@ namespace Chaos
 			FVec3 PCOM = SpaceTransform.TransformPosition(FParticleUtilities::GetCoMWorldPosition(Particle));
 			FRotation3 QCOM = SpaceTransform.GetRotation() * FParticleUtilities::GetCoMWorldRotation(Particle);
 			FMatrix33 QCOMm = QCOM.ToMatrix();
-			FDebugDrawQueue::GetInstance().DrawDebugDirectionalArrow(PCOM, PCOM + Settings.DrawScale * Settings.BodyAxisLen * QCOMm.GetAxis(0), Settings.DrawScale * Settings.ArrowSize, Red, false, KINDA_SMALL_NUMBER, Settings.DrawPriority, Settings.LineThickness);
-			FDebugDrawQueue::GetInstance().DrawDebugDirectionalArrow(PCOM, PCOM + Settings.DrawScale * Settings.BodyAxisLen * QCOMm.GetAxis(1), Settings.DrawScale * Settings.ArrowSize, Green, false, KINDA_SMALL_NUMBER, Settings.DrawPriority, Settings.LineThickness);
-			FDebugDrawQueue::GetInstance().DrawDebugDirectionalArrow(PCOM, PCOM + Settings.DrawScale * Settings.BodyAxisLen * QCOMm.GetAxis(2), Settings.DrawScale * Settings.ArrowSize, Blue, false, KINDA_SMALL_NUMBER, Settings.DrawPriority, Settings.LineThickness);
+			FDebugDrawQueue::GetInstance().DrawDebugDirectionalArrow(PCOM, PCOM + Settings.DrawScale * Settings.BodyAxisLen * QCOMm.GetAxis(0), Settings.DrawScale * Settings.ArrowSize, Red, false, UE_KINDA_SMALL_NUMBER, Settings.DrawPriority, Settings.LineThickness);
+			FDebugDrawQueue::GetInstance().DrawDebugDirectionalArrow(PCOM, PCOM + Settings.DrawScale * Settings.BodyAxisLen * QCOMm.GetAxis(1), Settings.DrawScale * Settings.ArrowSize, Green, false, UE_KINDA_SMALL_NUMBER, Settings.DrawPriority, Settings.LineThickness);
+			FDebugDrawQueue::GetInstance().DrawDebugDirectionalArrow(PCOM, PCOM + Settings.DrawScale * Settings.BodyAxisLen * QCOMm.GetAxis(2), Settings.DrawScale * Settings.ArrowSize, Blue, false, UE_KINDA_SMALL_NUMBER, Settings.DrawPriority, Settings.LineThickness);
 
 			FColor Black = FColor::Black;
 			FColor Grey = (ColorScale * FColor(64, 64, 64)).ToFColor(false);
 			FVec3 PActor = SpaceTransform.TransformPosition(FParticleUtilities::GetActorWorldTransform(Particle).GetTranslation());
-			FDebugDrawQueue::GetInstance().DrawDebugPoint(PActor, Black, false, KINDA_SMALL_NUMBER, Settings.DrawPriority, Settings.DrawScale * Settings.PointSize);
-			FDebugDrawQueue::GetInstance().DrawDebugLine(PCOM, PActor, Grey, false, KINDA_SMALL_NUMBER, Settings.DrawPriority, Settings.LineThickness);
+			FDebugDrawQueue::GetInstance().DrawDebugPoint(PActor, Black, false, UE_KINDA_SMALL_NUMBER, Settings.DrawPriority, Settings.DrawScale * Settings.PointSize);
+			FDebugDrawQueue::GetInstance().DrawDebugLine(PCOM, PActor, Grey, false, UE_KINDA_SMALL_NUMBER, Settings.DrawPriority, Settings.LineThickness);
 		
 			if (Index >= 0)
 			{
 				//FDebugDrawQueue::GetInstance().DrawDebugString(PCOM + FontHeight * FVec3(0, 0, 1), FString::Format(TEXT("{0}{1}"), { Particle->IsKinematic()? TEXT("K"): TEXT("D"), Index }), nullptr, FColor::White, KINDA_SMALL_NUMBER, false, FontScale);
 			}
 
-			if ((Settings.VelScale > 0.0f) && (Particle->V().Size() > KINDA_SMALL_NUMBER))
+			if ((Settings.VelScale > 0.0f) && (Particle->V().Size() > UE_KINDA_SMALL_NUMBER))
 			{
-				FDebugDrawQueue::GetInstance().DrawDebugLine(PCOM, PCOM + SpaceTransform.TransformVector(Particle->V()) * Settings.VelScale, Red, false, KINDA_SMALL_NUMBER, Settings.DrawPriority, Settings.LineThickness);
+				FDebugDrawQueue::GetInstance().DrawDebugLine(PCOM, PCOM + SpaceTransform.TransformVector(Particle->V()) * Settings.VelScale, Red, false, UE_KINDA_SMALL_NUMBER, Settings.DrawPriority, Settings.LineThickness);
 				//FDebugDrawQueue::GetInstance().DrawDebugLine(PCOM, PCOM + SpaceTransform.TransformVector(Particle->VSmooth()) * Settings.VelScale, Blue, false, KINDA_SMALL_NUMBER, Settings.DrawPriority, Settings.LineThickness);
 			}
-			if ((Settings.AngVelScale > 0.0f) && (Particle->W().Size() > KINDA_SMALL_NUMBER))
+			if ((Settings.AngVelScale > 0.0f) && (Particle->W().Size() > UE_KINDA_SMALL_NUMBER))
 			{
-				FDebugDrawQueue::GetInstance().DrawDebugLine(PCOM, PCOM + SpaceTransform.TransformVector(Particle->W()) * Settings.AngVelScale, Green, false, KINDA_SMALL_NUMBER, Settings.DrawPriority, Settings.LineThickness);
+				FDebugDrawQueue::GetInstance().DrawDebugLine(PCOM, PCOM + SpaceTransform.TransformVector(Particle->W()) * Settings.AngVelScale, Green, false, UE_KINDA_SMALL_NUMBER, Settings.DrawPriority, Settings.LineThickness);
 			}
 
 			if (Settings.InertiaScale > 0.0f)
@@ -895,42 +895,42 @@ namespace Chaos
 					if ((Settings.PushOutScale > 0) && !ManifoldPoint.NetPushOut.IsNearlyZero())
 					{
 						FColor Color = (ColorScale * PushOutColor).ToFColor(false);
-						FDebugDrawQueue::GetInstance().DrawDebugLine(WorldPointPlaneLocation, WorldPointPlaneLocation + Settings.DrawScale * Settings.PushOutScale * SpaceTransform.TransformVectorNoScale(ManifoldPoint.NetPushOut), Color, false, KINDA_SMALL_NUMBER, Settings.DrawPriority, Settings.LineThickness);
+						FDebugDrawQueue::GetInstance().DrawDebugLine(WorldPointPlaneLocation, WorldPointPlaneLocation + Settings.DrawScale * Settings.PushOutScale * SpaceTransform.TransformVectorNoScale(ManifoldPoint.NetPushOut), Color, false, UE_KINDA_SMALL_NUMBER, Settings.DrawPriority, Settings.LineThickness);
 					}
 					if ((Settings.ImpulseScale > 0) && !ManifoldPoint.NetImpulse.IsNearlyZero())
 					{
 						FColor Color = (ColorScale * ImpulseColor).ToFColor(false);
-						FDebugDrawQueue::GetInstance().DrawDebugLine(WorldPointPlaneLocation, WorldPointPlaneLocation + Settings.DrawScale * Settings.ImpulseScale * SpaceTransform.TransformVectorNoScale(ManifoldPoint.NetImpulse), Color, false, KINDA_SMALL_NUMBER, Settings.DrawPriority, Settings.LineThickness);
+						FDebugDrawQueue::GetInstance().DrawDebugLine(WorldPointPlaneLocation, WorldPointPlaneLocation + Settings.DrawScale * Settings.ImpulseScale * SpaceTransform.TransformVectorNoScale(ManifoldPoint.NetImpulse), Color, false, UE_KINDA_SMALL_NUMBER, Settings.DrawPriority, Settings.LineThickness);
 					}
 
 					// Manifold plane and normal
 					if (Settings.ContactWidth > 0)
 					{
 						FColor C0 = (ColorScale * DiscColor).ToFColor(false);
-						FDebugDrawQueue::GetInstance().DrawDebugCircle(WorldPlaneLocation, Settings.DrawScale * Settings.ContactWidth, 12, C0, false, KINDA_SMALL_NUMBER, Settings.DrawPriority, Settings.LineThickness, Axes.GetUnitAxis(EAxis::Y), Axes.GetUnitAxis(EAxis::Z), false);
+						FDebugDrawQueue::GetInstance().DrawDebugCircle(WorldPlaneLocation, Settings.DrawScale * Settings.ContactWidth, 12, C0, false, UE_KINDA_SMALL_NUMBER, Settings.DrawPriority, Settings.LineThickness, Axes.GetUnitAxis(EAxis::Y), Axes.GetUnitAxis(EAxis::Z), false);
 					}
 					if (Settings.ContactLen > 0)
 					{
 						FColor NormalColor = ((ManifoldPoint.ContactPoint.ContactType != EContactPointType::EdgeEdge) ? PlaneNormalColor : EdgeNormalColor);
 						FColor C1 = (ColorScale * NormalColor).ToFColor(false);
-						FDebugDrawQueue::GetInstance().DrawDebugLine(WorldPlaneLocation, WorldPlaneLocation + Settings.DrawScale * Settings.ContactLen * WorldPlaneNormal, C1, false, KINDA_SMALL_NUMBER, Settings.DrawPriority, Settings.LineThickness);
+						FDebugDrawQueue::GetInstance().DrawDebugLine(WorldPlaneLocation, WorldPlaneLocation + Settings.DrawScale * Settings.ContactLen * WorldPlaneNormal, C1, false, UE_KINDA_SMALL_NUMBER, Settings.DrawPriority, Settings.LineThickness);
 					}
 					if (Settings.ContactPhiWidth > 0 && (ManifoldPoint.ContactPoint.Phi < FLT_MAX))
 					{
 						FColor C2 = (ColorScale * FColor(128, 128, 0)).ToFColor(false);
-						FDebugDrawQueue::GetInstance().DrawDebugCircle(WorldPlaneLocation - ManifoldPoint.ContactPoint.Phi * WorldPlaneNormal, Settings.DrawScale * Settings.ContactPhiWidth, 12, C2, false, KINDA_SMALL_NUMBER, Settings.DrawPriority, Settings.LineThickness, Axes.GetUnitAxis(EAxis::Y), Axes.GetUnitAxis(EAxis::Z), false);
+						FDebugDrawQueue::GetInstance().DrawDebugCircle(WorldPlaneLocation - ManifoldPoint.ContactPoint.Phi * WorldPlaneNormal, Settings.DrawScale * Settings.ContactPhiWidth, 12, C2, false, UE_KINDA_SMALL_NUMBER, Settings.DrawPriority, Settings.LineThickness, Axes.GetUnitAxis(EAxis::Y), Axes.GetUnitAxis(EAxis::Z), false);
 					}
 
 					// Manifold point
-					FDebugDrawQueue::GetInstance().DrawDebugCircle(WorldPointLocation, 0.5f * Settings.DrawScale * Settings.ContactWidth, 12, DiscColor, false, KINDA_SMALL_NUMBER, Settings.DrawPriority, Settings.LineThickness, Axes.GetUnitAxis(EAxis::Y), Axes.GetUnitAxis(EAxis::Z), false);
+					FDebugDrawQueue::GetInstance().DrawDebugCircle(WorldPointLocation, 0.5f * Settings.DrawScale * Settings.ContactWidth, 12, DiscColor, false, UE_KINDA_SMALL_NUMBER, Settings.DrawPriority, Settings.LineThickness, Axes.GetUnitAxis(EAxis::Y), Axes.GetUnitAxis(EAxis::Z), false);
 
 					// Previous points
 					if (bIsActive && ManifoldPoint.Flags.bWasFrictionRestored)
 					{
 						const FVec3 WorldPrevPointLocation = SpaceTransform.TransformPosition(PointTransform.TransformPosition(ManifoldPoint.ShapeAnchorPoints[ContactPointOwner]));
 						const FVec3 WorldPrevPlaneLocation = SpaceTransform.TransformPosition(PlaneTransform.TransformPosition(ManifoldPoint.ShapeAnchorPoints[ContactPlaneOwner]));
-						FDebugDrawQueue::GetInstance().DrawDebugLine(WorldPrevPointLocation, WorldPointLocation, FColor::White, false, KINDA_SMALL_NUMBER, Settings.DrawPriority, Settings.LineThickness);
-						FDebugDrawQueue::GetInstance().DrawDebugLine(WorldPrevPlaneLocation, WorldPlaneLocation, FColor::White, false, KINDA_SMALL_NUMBER, Settings.DrawPriority, Settings.LineThickness);
+						FDebugDrawQueue::GetInstance().DrawDebugLine(WorldPrevPointLocation, WorldPointLocation, FColor::White, false, UE_KINDA_SMALL_NUMBER, Settings.DrawPriority, Settings.LineThickness);
+						FDebugDrawQueue::GetInstance().DrawDebugLine(WorldPrevPlaneLocation, WorldPlaneLocation, FColor::White, false, UE_KINDA_SMALL_NUMBER, Settings.DrawPriority, Settings.LineThickness);
 					}
 
 					// Whether restored
@@ -939,17 +939,17 @@ namespace Chaos
 						if (Contact.WasManifoldRestored())
 						{
 							const FReal BoxScale = Settings.DrawScale * Settings.ContactInfoWidth;
-							FDebugDrawQueue::GetInstance().DrawDebugBox(WorldPlaneLocation, FVec3(BoxScale, BoxScale, FReal(0.01)), FRotation3(FRotationMatrix::MakeFromZ(WorldPlaneNormal)), FColor::Blue, false, KINDA_SMALL_NUMBER, Settings.DrawPriority, 0.5f * Settings.LineThickness);
+							FDebugDrawQueue::GetInstance().DrawDebugBox(WorldPlaneLocation, FVec3(BoxScale, BoxScale, FReal(0.01)), FRotation3(FRotationMatrix::MakeFromZ(WorldPlaneNormal)), FColor::Blue, false, UE_KINDA_SMALL_NUMBER, Settings.DrawPriority, 0.5f * Settings.LineThickness);
 						}
 						else if (ManifoldPoint.Flags.bWasRestored)
 						{
 							const FReal BoxScale = Settings.DrawScale * Settings.ContactInfoWidth;
-							FDebugDrawQueue::GetInstance().DrawDebugBox(WorldPlaneLocation, FVec3(BoxScale, BoxScale, FReal(0.01)), FRotation3(FRotationMatrix::MakeFromZ(WorldPlaneNormal)), FColor::Purple, false, KINDA_SMALL_NUMBER, Settings.DrawPriority, 0.5f * Settings.LineThickness);
+							FDebugDrawQueue::GetInstance().DrawDebugBox(WorldPlaneLocation, FVec3(BoxScale, BoxScale, FReal(0.01)), FRotation3(FRotationMatrix::MakeFromZ(WorldPlaneNormal)), FColor::Purple, false, UE_KINDA_SMALL_NUMBER, Settings.DrawPriority, 0.5f * Settings.LineThickness);
 						}
 						else if (ManifoldPoint.Flags.bWasReplaced)
 						{
 							const FReal BoxScale = Settings.DrawScale * Settings.ContactInfoWidth;
-							FDebugDrawQueue::GetInstance().DrawDebugBox(WorldPlaneLocation, FVec3(BoxScale, BoxScale, FReal(0.01)), FRotation3(FRotationMatrix::MakeFromZ(WorldPlaneNormal)), FColor::Orange, false, KINDA_SMALL_NUMBER, Settings.DrawPriority, 0.5f * Settings.LineThickness);
+							FDebugDrawQueue::GetInstance().DrawDebugBox(WorldPlaneLocation, FVec3(BoxScale, BoxScale, FReal(0.01)), FRotation3(FRotationMatrix::MakeFromZ(WorldPlaneNormal)), FColor::Orange, false, UE_KINDA_SMALL_NUMBER, Settings.DrawPriority, 0.5f * Settings.LineThickness);
 						}
 					}
 				}
@@ -959,7 +959,7 @@ namespace Chaos
 				{
 					FColor Color = (ColorScale * FColor::White).ToFColor(false);
 					const FVec3 ImpulsePos = SpaceTransform.TransformPosition(WorldActorTransform0.GetTranslation());
-					FDebugDrawQueue::GetInstance().DrawDebugLine(ImpulsePos, ImpulsePos + Settings.DrawScale * Settings.ImpulseScale * SpaceTransform.TransformVectorNoScale(Contact.GetAccumulatedImpulse()), Color, false, KINDA_SMALL_NUMBER, Settings.DrawPriority, Settings.LineThickness);
+					FDebugDrawQueue::GetInstance().DrawDebugLine(ImpulsePos, ImpulsePos + Settings.DrawScale * Settings.ImpulseScale * SpaceTransform.TransformVectorNoScale(Contact.GetAccumulatedImpulse()), Color, false, UE_KINDA_SMALL_NUMBER, Settings.DrawPriority, Settings.LineThickness);
 				}
 
 			}
@@ -972,8 +972,8 @@ namespace Chaos
 				const FMatrix Axes = FRotationMatrix::MakeFromX(Normal);
 				const FVec3 P0 = SpaceTransform.TransformPosition(Contact.GetParticle0()->X());
 				const FVec3 P1 = SpaceTransform.TransformPosition(Contact.GetParticle1()->X());
-				FDebugDrawQueue::GetInstance().DrawDebugLine(Location, P0, C3, false, KINDA_SMALL_NUMBER, Settings.DrawPriority, Settings.LineThickness * 0.5f);
-				FDebugDrawQueue::GetInstance().DrawDebugLine(Location, P1, C3, false, KINDA_SMALL_NUMBER, Settings.DrawPriority, Settings.LineThickness * 0.5f);
+				FDebugDrawQueue::GetInstance().DrawDebugLine(Location, P0, C3, false, UE_KINDA_SMALL_NUMBER, Settings.DrawPriority, Settings.LineThickness * 0.5f);
+				FDebugDrawQueue::GetInstance().DrawDebugLine(Location, P1, C3, false, UE_KINDA_SMALL_NUMBER, Settings.DrawPriority, Settings.LineThickness * 0.5f);
 			}
 		}
 		
@@ -1062,20 +1062,20 @@ namespace Chaos
 				// Leave a gap around the actor position so we can see where the center is
 				FVec3 Sa = Pa;
 				const FReal Lena = (Xa - Pa).Size();
-				if (Lena > KINDA_SMALL_NUMBER)
+				if (Lena > UE_KINDA_SMALL_NUMBER)
 				{
 					Sa = FMath::Lerp(Pa, Xa, FMath::Clamp<FReal>(CoMSize / Lena, 0., 1.));
 				}
 				FVec3 Sb = Pb;
 				const FReal Lenb = (Xb - Pb).Size();
-				if (Lenb > KINDA_SMALL_NUMBER)
+				if (Lenb > UE_KINDA_SMALL_NUMBER)
 				{
 					Sb = FMath::Lerp(Pb, Xb, FMath::Clamp<FReal>(CoMSize / Lena, 0., 1.));
 				}
-				FDebugDrawQueue::GetInstance().DrawDebugLine(Pa, Sa, FColor::White, false, KINDA_SMALL_NUMBER, Settings.DrawPriority, ConnectorThickness);
-				FDebugDrawQueue::GetInstance().DrawDebugLine(Pb, Sb, FColor::White, false, KINDA_SMALL_NUMBER, Settings.DrawPriority, ConnectorThickness);
-				FDebugDrawQueue::GetInstance().DrawDebugLine(Sa, Xa, R, false, KINDA_SMALL_NUMBER, Settings.DrawPriority, ConnectorThickness);
-				FDebugDrawQueue::GetInstance().DrawDebugLine(Sb, Xb, C, false, KINDA_SMALL_NUMBER, Settings.DrawPriority, ConnectorThickness);
+				FDebugDrawQueue::GetInstance().DrawDebugLine(Pa, Sa, FColor::White, false, UE_KINDA_SMALL_NUMBER, Settings.DrawPriority, ConnectorThickness);
+				FDebugDrawQueue::GetInstance().DrawDebugLine(Pb, Sb, FColor::White, false, UE_KINDA_SMALL_NUMBER, Settings.DrawPriority, ConnectorThickness);
+				FDebugDrawQueue::GetInstance().DrawDebugLine(Sa, Xa, R, false, UE_KINDA_SMALL_NUMBER, Settings.DrawPriority, ConnectorThickness);
+				FDebugDrawQueue::GetInstance().DrawDebugLine(Sb, Xb, C, false, UE_KINDA_SMALL_NUMBER, Settings.DrawPriority, ConnectorThickness);
 			}
 			if (FeatureMask.bCoMConnector)
 			{
@@ -1084,54 +1084,54 @@ namespace Chaos
 				// Leave a gap around the body position so we can see where the center is
 				FVec3 Sa = Ca;
 				const FReal Lena = (Xa - Ca).Size();
-				if (Lena > KINDA_SMALL_NUMBER)
+				if (Lena > UE_KINDA_SMALL_NUMBER)
 				{
 					Sa = FMath::Lerp(Ca, Xa, FMath::Clamp<FReal>(CoMSize / Lena, 0., 1.));
 				}
 				FVec3 Sb = Cb;
 				const FReal Lenb = (Xb - Cb).Size();
-				if (Lenb > KINDA_SMALL_NUMBER)
+				if (Lenb > UE_KINDA_SMALL_NUMBER)
 				{
 					Sb = FMath::Lerp(Cb, Xb, FMath::Clamp<FReal>(CoMSize / Lena, 0., 1.));
 				}
-				FDebugDrawQueue::GetInstance().DrawDebugLine(Ca, Sa, FColor::Black, false, KINDA_SMALL_NUMBER, Settings.DrawPriority, ConnectorThickness);
-				FDebugDrawQueue::GetInstance().DrawDebugLine(Cb, Sb, FColor::Black, false, KINDA_SMALL_NUMBER, Settings.DrawPriority, ConnectorThickness);
-				FDebugDrawQueue::GetInstance().DrawDebugLine(Sa, Xa, R, false, KINDA_SMALL_NUMBER, Settings.DrawPriority, ConnectorThickness);
-				FDebugDrawQueue::GetInstance().DrawDebugLine(Sb, Xb, C, false, KINDA_SMALL_NUMBER, Settings.DrawPriority, ConnectorThickness);
+				FDebugDrawQueue::GetInstance().DrawDebugLine(Ca, Sa, FColor::Black, false, UE_KINDA_SMALL_NUMBER, Settings.DrawPriority, ConnectorThickness);
+				FDebugDrawQueue::GetInstance().DrawDebugLine(Cb, Sb, FColor::Black, false, UE_KINDA_SMALL_NUMBER, Settings.DrawPriority, ConnectorThickness);
+				FDebugDrawQueue::GetInstance().DrawDebugLine(Sa, Xa, R, false, UE_KINDA_SMALL_NUMBER, Settings.DrawPriority, ConnectorThickness);
+				FDebugDrawQueue::GetInstance().DrawDebugLine(Sb, Xb, C, false, UE_KINDA_SMALL_NUMBER, Settings.DrawPriority, ConnectorThickness);
 			}
 			if (FeatureMask.bStretch)
 			{
 				const FRealSingle StretchThickness = 3.0f * Settings.LineThickness;
-				FDebugDrawQueue::GetInstance().DrawDebugLine(Xa, Xb, M, false, KINDA_SMALL_NUMBER, Settings.DrawPriority, StretchThickness);
+				FDebugDrawQueue::GetInstance().DrawDebugLine(Xa, Xb, M, false, UE_KINDA_SMALL_NUMBER, Settings.DrawPriority, StretchThickness);
 			}
 			if (FeatureMask.bAxes)
 			{
-				FDebugDrawQueue::GetInstance().DrawDebugDirectionalArrow(Xa, Xa + Settings.DrawScale * Settings.ConstraintAxisLen * SpaceTransform.TransformVector(Ra.GetAxis(0)), Settings.DrawScale * Settings.ArrowSize, R, false, KINDA_SMALL_NUMBER, Settings.DrawPriority, Settings.LineThickness);
-				FDebugDrawQueue::GetInstance().DrawDebugDirectionalArrow(Xa, Xa + Settings.DrawScale * Settings.ConstraintAxisLen * SpaceTransform.TransformVector(Ra.GetAxis(1)), Settings.DrawScale * Settings.ArrowSize, G, false, KINDA_SMALL_NUMBER, Settings.DrawPriority, Settings.LineThickness);
-				FDebugDrawQueue::GetInstance().DrawDebugDirectionalArrow(Xa, Xa + Settings.DrawScale * Settings.ConstraintAxisLen * SpaceTransform.TransformVector(Ra.GetAxis(2)), Settings.DrawScale * Settings.ArrowSize, B, false, KINDA_SMALL_NUMBER, Settings.DrawPriority, Settings.LineThickness);
-				FDebugDrawQueue::GetInstance().DrawDebugDirectionalArrow(Xb, Xb + Settings.DrawScale * Settings.ConstraintAxisLen * SpaceTransform.TransformVector(Rb.GetAxis(0)), Settings.DrawScale * Settings.ArrowSize, C, false, KINDA_SMALL_NUMBER, Settings.DrawPriority, Settings.LineThickness);
-				FDebugDrawQueue::GetInstance().DrawDebugDirectionalArrow(Xb, Xb + Settings.DrawScale * Settings.ConstraintAxisLen * SpaceTransform.TransformVector(Rb.GetAxis(1)), Settings.DrawScale * Settings.ArrowSize, M, false, KINDA_SMALL_NUMBER, Settings.DrawPriority, Settings.LineThickness);
-				FDebugDrawQueue::GetInstance().DrawDebugDirectionalArrow(Xb, Xb + Settings.DrawScale * Settings.ConstraintAxisLen * SpaceTransform.TransformVector(Rb.GetAxis(2)), Settings.DrawScale * Settings.ArrowSize, Y, false, KINDA_SMALL_NUMBER, Settings.DrawPriority, Settings.LineThickness);
+				FDebugDrawQueue::GetInstance().DrawDebugDirectionalArrow(Xa, Xa + Settings.DrawScale * Settings.ConstraintAxisLen * SpaceTransform.TransformVector(Ra.GetAxis(0)), Settings.DrawScale * Settings.ArrowSize, R, false, UE_KINDA_SMALL_NUMBER, Settings.DrawPriority, Settings.LineThickness);
+				FDebugDrawQueue::GetInstance().DrawDebugDirectionalArrow(Xa, Xa + Settings.DrawScale * Settings.ConstraintAxisLen * SpaceTransform.TransformVector(Ra.GetAxis(1)), Settings.DrawScale * Settings.ArrowSize, G, false, UE_KINDA_SMALL_NUMBER, Settings.DrawPriority, Settings.LineThickness);
+				FDebugDrawQueue::GetInstance().DrawDebugDirectionalArrow(Xa, Xa + Settings.DrawScale * Settings.ConstraintAxisLen * SpaceTransform.TransformVector(Ra.GetAxis(2)), Settings.DrawScale * Settings.ArrowSize, B, false, UE_KINDA_SMALL_NUMBER, Settings.DrawPriority, Settings.LineThickness);
+				FDebugDrawQueue::GetInstance().DrawDebugDirectionalArrow(Xb, Xb + Settings.DrawScale * Settings.ConstraintAxisLen * SpaceTransform.TransformVector(Rb.GetAxis(0)), Settings.DrawScale * Settings.ArrowSize, C, false, UE_KINDA_SMALL_NUMBER, Settings.DrawPriority, Settings.LineThickness);
+				FDebugDrawQueue::GetInstance().DrawDebugDirectionalArrow(Xb, Xb + Settings.DrawScale * Settings.ConstraintAxisLen * SpaceTransform.TransformVector(Rb.GetAxis(1)), Settings.DrawScale * Settings.ArrowSize, M, false, UE_KINDA_SMALL_NUMBER, Settings.DrawPriority, Settings.LineThickness);
+				FDebugDrawQueue::GetInstance().DrawDebugDirectionalArrow(Xb, Xb + Settings.DrawScale * Settings.ConstraintAxisLen * SpaceTransform.TransformVector(Rb.GetAxis(2)), Settings.DrawScale * Settings.ArrowSize, Y, false, UE_KINDA_SMALL_NUMBER, Settings.DrawPriority, Settings.LineThickness);
 			}
 			FVec3 TextPos = Xb;
 			if (FeatureMask.bLevel && (LevelIndex >= 0))
 			{
-				FDebugDrawQueue::GetInstance().DrawDebugString(TextPos, FString::Format(TEXT("{0}"), { LevelIndex }), nullptr, FColor::Red, KINDA_SMALL_NUMBER, false, Settings.FontScale);
+				FDebugDrawQueue::GetInstance().DrawDebugString(TextPos, FString::Format(TEXT("{0}"), { LevelIndex }), nullptr, FColor::Red, UE_KINDA_SMALL_NUMBER, false, Settings.FontScale);
 				TextPos += Settings.FontHeight * FVec3(0, 0, 1);
 			}
 			if (FeatureMask.bIndex && (Index >= 0))
 			{
-				FDebugDrawQueue::GetInstance().DrawDebugString(TextPos, FString::Format(TEXT("{0}"), { Index }), nullptr, FColor::Red, KINDA_SMALL_NUMBER, false, Settings.FontScale);
+				FDebugDrawQueue::GetInstance().DrawDebugString(TextPos, FString::Format(TEXT("{0}"), { Index }), nullptr, FColor::Red, UE_KINDA_SMALL_NUMBER, false, Settings.FontScale);
 				TextPos += Settings.FontHeight * FVec3(0, 0, 1);
 			}
 			if (FeatureMask.bColor && (ColorIndex >= 0))
 			{
-				FDebugDrawQueue::GetInstance().DrawDebugString(TextPos, FString::Format(TEXT("{0}"), { ColorIndex }), nullptr, FColor::Red, KINDA_SMALL_NUMBER, false, Settings.FontScale);
+				FDebugDrawQueue::GetInstance().DrawDebugString(TextPos, FString::Format(TEXT("{0}"), { ColorIndex }), nullptr, FColor::Red, UE_KINDA_SMALL_NUMBER, false, Settings.FontScale);
 				TextPos += Settings.FontHeight * FVec3(0, 0, 1);
 			}
 			if (FeatureMask.bIsland && (IslandIndex >= 0))
 			{
-				FDebugDrawQueue::GetInstance().DrawDebugString(TextPos, FString::Format(TEXT("{0}"), { IslandIndex }), nullptr, FColor::Red, KINDA_SMALL_NUMBER, false, Settings.FontScale);
+				FDebugDrawQueue::GetInstance().DrawDebugString(TextPos, FString::Format(TEXT("{0}"), { IslandIndex }), nullptr, FColor::Red, UE_KINDA_SMALL_NUMBER, false, Settings.FontScale);
 				TextPos += Settings.FontHeight * FVec3(0, 0, 1);
 			}
 		}
@@ -1159,14 +1159,14 @@ namespace Chaos
 			const FVec3 Pos = SimSpace.Transform.GetLocation();
 			const FRotation3& Rot = SimSpace.Transform.GetRotation();
 			const FMatrix33 Rotm = Rot.ToMatrix();
-			FDebugDrawQueue::GetInstance().DrawDebugDirectionalArrow(Pos, Pos + Settings.DrawScale * Settings.BodyAxisLen * Rotm.GetAxis(0), Settings.DrawScale * Settings.ArrowSize, FColor::Red, false, KINDA_SMALL_NUMBER, Settings.DrawPriority, Settings.LineThickness);
-			FDebugDrawQueue::GetInstance().DrawDebugDirectionalArrow(Pos, Pos + Settings.DrawScale * Settings.BodyAxisLen * Rotm.GetAxis(1), Settings.DrawScale * Settings.ArrowSize, FColor::Green, false, KINDA_SMALL_NUMBER, Settings.DrawPriority, Settings.LineThickness);
-			FDebugDrawQueue::GetInstance().DrawDebugDirectionalArrow(Pos, Pos + Settings.DrawScale * Settings.BodyAxisLen * Rotm.GetAxis(2), Settings.DrawScale * Settings.ArrowSize, FColor::Blue, false, KINDA_SMALL_NUMBER, Settings.DrawPriority, Settings.LineThickness);
-
-			FDebugDrawQueue::GetInstance().DrawDebugLine(Pos, Pos + Settings.VelScale * SimSpace.LinearVelocity, FColor::Cyan, false, KINDA_SMALL_NUMBER, Settings.DrawPriority, Settings.LineThickness);
-			FDebugDrawQueue::GetInstance().DrawDebugLine(Pos, Pos + Settings.AngVelScale * SimSpace.AngularVelocity, FColor::Cyan, false, KINDA_SMALL_NUMBER, Settings.DrawPriority, Settings.LineThickness);
-			FDebugDrawQueue::GetInstance().DrawDebugLine(Pos, Pos + 0.01f * Settings.VelScale * SimSpace.LinearAcceleration, FColor::Yellow, false, KINDA_SMALL_NUMBER, Settings.DrawPriority, Settings.LineThickness);
-			FDebugDrawQueue::GetInstance().DrawDebugLine(Pos, Pos + 0.01f * Settings.AngVelScale * SimSpace.AngularAcceleration, FColor::Orange, false, KINDA_SMALL_NUMBER, Settings.DrawPriority, Settings.LineThickness);
+			FDebugDrawQueue::GetInstance().DrawDebugDirectionalArrow(Pos, Pos + Settings.DrawScale * Settings.BodyAxisLen * Rotm.GetAxis(0), Settings.DrawScale * Settings.ArrowSize, FColor::Red, false, UE_KINDA_SMALL_NUMBER, Settings.DrawPriority, Settings.LineThickness);
+			FDebugDrawQueue::GetInstance().DrawDebugDirectionalArrow(Pos, Pos + Settings.DrawScale * Settings.BodyAxisLen * Rotm.GetAxis(1), Settings.DrawScale * Settings.ArrowSize, FColor::Green, false, UE_KINDA_SMALL_NUMBER, Settings.DrawPriority, Settings.LineThickness);
+			FDebugDrawQueue::GetInstance().DrawDebugDirectionalArrow(Pos, Pos + Settings.DrawScale * Settings.BodyAxisLen * Rotm.GetAxis(2), Settings.DrawScale * Settings.ArrowSize, FColor::Blue, false, UE_KINDA_SMALL_NUMBER, Settings.DrawPriority, Settings.LineThickness);
+
+			FDebugDrawQueue::GetInstance().DrawDebugLine(Pos, Pos + Settings.VelScale * SimSpace.LinearVelocity, FColor::Cyan, false, UE_KINDA_SMALL_NUMBER, Settings.DrawPriority, Settings.LineThickness);
+			FDebugDrawQueue::GetInstance().DrawDebugLine(Pos, Pos + Settings.AngVelScale * SimSpace.AngularVelocity, FColor::Cyan, false, UE_KINDA_SMALL_NUMBER, Settings.DrawPriority, Settings.LineThickness);
+			FDebugDrawQueue::GetInstance().DrawDebugLine(Pos, Pos + 0.01f * Settings.VelScale * SimSpace.LinearAcceleration, FColor::Yellow, false, UE_KINDA_SMALL_NUMBER, Settings.DrawPriority, Settings.LineThickness);
+			FDebugDrawQueue::GetInstance().DrawDebugLine(Pos, Pos + 0.01f * Settings.AngVelScale * SimSpace.AngularAcceleration, FColor::Orange, false, UE_KINDA_SMALL_NUMBER, Settings.DrawPriority, Settings.LineThickness);
 		}
 
 		void DrawConstraintGraphImpl(const FRigidTransform3& SpaceTransform, const FPBDConstraintGraph& Graph, const FChaosDebugDrawSettings& Settings)
@@ -1193,14 +1193,14 @@ namespace Chaos
 				if ((bChaosDebugDebugDrawContactGraphUsed && bIsUsed) || (bChaosDebugDebugDrawContactGraphUnused && !bIsUsed))
 				{
 					FColor Color = bIsUsed ? FColor::Green : FColor::Red;
-					FDebugDrawQueue::GetInstance().DrawDebugLine(Transform0.GetLocation(), Transform1.GetLocation(), Color, false, KINDA_SMALL_NUMBER, Settings.DrawPriority, Settings.LineThickness);
+					FDebugDrawQueue::GetInstance().DrawDebugLine(Transform0.GetLocation(), Transform1.GetLocation(), Color, false, UE_KINDA_SMALL_NUMBER, Settings.DrawPriority, Settings.LineThickness);
 				}
 
 				if (bChaosDebugDebugDrawContactGraph)
 				{
-					FDebugDrawQueue::GetInstance().DrawDebugLine(Transform0.GetLocation(), ContactPos, FColor::Red, false, KINDA_SMALL_NUMBER, Settings.DrawPriority, Settings.LineThickness);
-					FDebugDrawQueue::GetInstance().DrawDebugLine(Transform1.GetLocation(), ContactPos, FColor::White, false, KINDA_SMALL_NUMBER, Settings.DrawPriority, Settings.LineThickness);
-					FDebugDrawQueue::GetInstance().DrawDebugString(ContactPos, FString::Format(TEXT("{0}-{1}-{2}"), { LevelIndex, ColorIndex, OrderIndex }), nullptr, FColor::Yellow, KINDA_SMALL_NUMBER, false, Settings.FontScale);
+					FDebugDrawQueue::GetInstance().DrawDebugLine(Transform0.GetLocation(), ContactPos, FColor::Red, false, UE_KINDA_SMALL_NUMBER, Settings.DrawPriority, Settings.LineThickness);
+					FDebugDrawQueue::GetInstance().DrawDebugLine(Transform1.GetLocation(), ContactPos, FColor::White, false, UE_KINDA_SMALL_NUMBER, Settings.DrawPriority, Settings.LineThickness);
+					FDebugDrawQueue::GetInstance().DrawDebugString(ContactPos, FString::Format(TEXT("{0}-{1}-{2}"), { LevelIndex, ColorIndex, OrderIndex }), nullptr, FColor::Yellow, UE_KINDA_SMALL_NUMBER, false, Settings.FontScale);
 				}
 			};
 
@@ -1223,7 +1223,7 @@ namespace Chaos
 				{
 					const FColor IslandColor = GetIslandColor(IslandIndex, !Island->IsSleeping());
 					const FAABB3 Bounds = IslandAABB.TransformedAABB(SpaceTransform);
-					FDebugDrawQueue::GetInstance().DrawDebugBox(Bounds.Center(), 0.5f * Bounds.Extents(), SpaceTransform.GetRotation(), IslandColor, false, KINDA_SMALL_NUMBER, Settings.DrawPriority, 3.0f * Settings.LineThickness);
+					FDebugDrawQueue::GetInstance().DrawDebugBox(Bounds.Center(), 0.5f * Bounds.Extents(), SpaceTransform.GetRotation(), IslandColor, false, UE_KINDA_SMALL_NUMBER, Settings.DrawPriority, 3.0f * Settings.LineThickness);
 				}
 
 				if (bChaosDebugDebugDrawContactGraph || bChaosDebugDebugDrawContactGraphUnused || bChaosDebugDebugDrawContactGraphUsed)
@@ -1304,7 +1304,7 @@ namespace Chaos
 
 					FColor Color = FColor::Green; // FMath::Lerp(FColor::Green, FColor::Red, Connection.Strain / MaxStrain);
 
-					FDebugDrawQueue::GetInstance().DrawDebugLine(Pos0, Pos1, Color, false, KINDA_SMALL_NUMBER, Settings.DrawPriority, Settings.LineThickness * 3);
+					FDebugDrawQueue::GetInstance().DrawDebugLine(Pos0, Pos1, Color, false, UE_KINDA_SMALL_NUMBER, Settings.DrawPriority, Settings.LineThickness * 3);
 				}
 
 			};
@@ -1383,7 +1383,7 @@ namespace Chaos
 			FDebugDrawQueue::GetInstance().DrawDebugLine(
 				SpaceTransform.TransformPosition(PWorld), 
 				SpaceTransform.TransformPosition(PWorld + AxisLen * AxisWorld), 
-				FColor::Green, false, KINDA_SMALL_NUMBER, Settings.DrawPriority, Settings.LineThickness);
+				FColor::Green, false, UE_KINDA_SMALL_NUMBER, Settings.DrawPriority, Settings.LineThickness);
 
 			if (Settings.PushOutScale > 0)
 			{
@@ -1399,12 +1399,12 @@ namespace Chaos
 				FDebugDrawQueue::GetInstance().DrawDebugLine(
 					SpaceTransform.TransformPosition(PWorld),
 					SpaceTransform.TransformPosition(PWorld + Settings.PushOutScale * PushOutWorld),
-					SpringPushOutColor, false, KINDA_SMALL_NUMBER, Settings.DrawPriority, SpringWidthScale * Settings.LineThickness);
+					SpringPushOutColor, false, UE_KINDA_SMALL_NUMBER, Settings.DrawPriority, SpringWidthScale * Settings.LineThickness);
 
 				FDebugDrawQueue::GetInstance().DrawDebugLine(
 					SpaceTransform.TransformPosition(PWorld),
 					SpaceTransform.TransformPosition(PWorld + Settings.PushOutScale * HardStopPushOutWorld),
-					HardStopPushOutColor, false, KINDA_SMALL_NUMBER, Settings.DrawPriority, HardStopWidthScale * Settings.LineThickness);
+					HardStopPushOutColor, false, UE_KINDA_SMALL_NUMBER, Settings.DrawPriority, HardStopWidthScale * Settings.LineThickness);
 			}
 
 			if (Settings.ImpulseScale > 0)
@@ -1414,7 +1414,7 @@ namespace Chaos
 				FDebugDrawQueue::GetInstance().DrawDebugLine(
 					SpaceTransform.TransformPosition(PWorld),
 					SpaceTransform.TransformPosition(PWorld + Settings.PushOutScale * HardStopImpulseWorld),
-					HardStopImpulseColor, false, KINDA_SMALL_NUMBER, Settings.DrawPriority, Settings.LineThickness * FRealSingle(0.5));
+					HardStopImpulseColor, false, UE_KINDA_SMALL_NUMBER, Settings.DrawPriority, Settings.LineThickness * FRealSingle(0.5));
 			}
 		}
 
diff --git a/Engine/Source/Runtime/Experimental/Chaos/Private/Chaos/Collision/ContactPointsMiscShapes.cpp b/Engine/Source/Runtime/Experimental/Chaos/Private/Chaos/Collision/ContactPointsMiscShapes.cpp
--- a/Engine/Source/Runtime/Experimental/Chaos/Private/Chaos/Collision/ContactPointsMiscShapes.cpp
+++ b/Engine/Source/Runtime/Experimental/Chaos/Private/Chaos/Collision/ContactPointsMiscShapes.cpp
@@ -194,7 +194,7 @@ namespace Chaos
 		if (SizeSq < CullDistanceSq)
 		{
 			const FReal Size = FMath::Sqrt(SizeSq);
-			const FVec3 Normal = Size > SMALL_NUMBER ? Direction / Size : FVec3(0, 0, 1);
+			const FVec3 Normal = Size > UE_SMALL_NUMBER ? Direction / Size : FVec3(0, 0, 1);
 			const FReal NewPhi = Size - (R1 + R2);
 
 			Result.ShapeContactPoints[0] = Sphere1.GetCenter() - Sphere1Transform.InverseTransformVector(R1 * Normal);
@@ -259,7 +259,7 @@ namespace Chaos
 
 		FVec3 Delta = P2 - A1;
 		FReal DeltaLen = Delta.Size();
-		if (DeltaLen > KINDA_SMALL_NUMBER)
+		if (DeltaLen > UE_KINDA_SMALL_NUMBER)
 		{
 			FReal NewPhi = DeltaLen - (A.GetRadius() + B.GetRadius()) - ShapePadding;
 			FVec3 Dir = Delta / DeltaLen;
@@ -393,7 +393,7 @@ namespace Chaos
 
 		FVec3 Delta = P2 - P1;
 		FReal DeltaLen = Delta.Size();
-		if (DeltaLen > KINDA_SMALL_NUMBER)
+		if (DeltaLen > UE_KINDA_SMALL_NUMBER)
 		{
 			FReal NewPhi = DeltaLen - (A.GetRadius() + B.GetRadius()) - ShapePadding;
 			FVec3 Dir = Delta / DeltaLen;
diff --git a/Engine/Source/Runtime/Experimental/Chaos/Private/Chaos/Collision/PBDCollisionConstraint.cpp b/Engine/Source/Runtime/Experimental/Chaos/Private/Chaos/Collision/PBDCollisionConstraint.cpp
--- a/Engine/Source/Runtime/Experimental/Chaos/Private/Chaos/Collision/PBDCollisionConstraint.cpp
+++ b/Engine/Source/Runtime/Experimental/Chaos/Private/Chaos/Collision/PBDCollisionConstraint.cpp
@@ -985,7 +985,7 @@ namespace Chaos
 		FVec3 GravityDirection = ConcreteContainer()->GetGravityDirection();
 		FReal GravitySize = ConcreteContainer()->GetGravitySize();
 		// When gravity is zero, we still want to sort the constraints instead of having a random order. In this case, set gravity to default gravity.
-		if (GravitySize < SMALL_NUMBER)
+		if (GravitySize < UE_SMALL_NUMBER)
 		{
 			GravityDirection = FVec3(0, 0, -1);
 			GravitySize = 980.f;
diff --git a/Engine/Source/Runtime/Experimental/Chaos/Private/Chaos/Collision/PBDCollisionSolver.cpp b/Engine/Source/Runtime/Experimental/Chaos/Private/Chaos/Collision/PBDCollisionSolver.cpp
--- a/Engine/Source/Runtime/Experimental/Chaos/Private/Chaos/Collision/PBDCollisionSolver.cpp
+++ b/Engine/Source/Runtime/Experimental/Chaos/Private/Chaos/Collision/PBDCollisionSolver.cpp
@@ -183,7 +183,7 @@ namespace Chaos
 				WorldContactNormalAngular1 = Body1.InvI() * R1xN;
 				ContactMassInvNormal += FSolverVec3::DotProduct(R1xN, WorldContactNormalAngular1) + Body1.InvM();
 			}
-			AverageManifoldPoint.ContactMassNormal = (ContactMassInvNormal > FSolverReal(SMALL_NUMBER)) ? FSolverReal(1) / ContactMassInvNormal : FSolverReal(0);
+			AverageManifoldPoint.ContactMassNormal = (ContactMassInvNormal > FSolverReal(UE_SMALL_NUMBER)) ? FSolverReal(1) / ContactMassInvNormal : FSolverReal(0);
 			AverageManifoldPoint.WorldContactNormalAngular0 = WorldContactNormalAngular0;
 			AverageManifoldPoint.WorldContactNormalAngular1 = WorldContactNormalAngular1;
 
diff --git a/Engine/Source/Runtime/Experimental/Chaos/Private/Chaos/Collision/PBDCollisionSolver.h b/Engine/Source/Runtime/Experimental/Chaos/Private/Chaos/Collision/PBDCollisionSolver.h
--- a/Engine/Source/Runtime/Experimental/Chaos/Private/Chaos/Collision/PBDCollisionSolver.h
+++ b/Engine/Source/Runtime/Experimental/Chaos/Private/Chaos/Collision/PBDCollisionSolver.h
@@ -419,7 +419,7 @@ namespace Chaos
 		// Assume we stay in the friction cone...
 		OutStaticFrictionRatio = FSolverReal(1);
 
-		if (NetPushOutNormal < FSolverReal(SMALL_NUMBER))
+		if (NetPushOutNormal < FSolverReal(UE_SMALL_NUMBER))
 		{
 			// Note: we have already added the current iteration's PushOut to the NetPushOut but it has not been applied to the body
 			// so we must subtract it again to calculate the actual pushout we want to undo (i.e., the net pushout that has been applied to the body so far)
@@ -584,7 +584,7 @@ namespace Chaos
 			const FSolverReal MaxImpulseTangent = FMath::Max(FSolverReal(0), DynamicFriction * (ManifoldPoint.NetImpulseNormal + ImpulseNormal + ManifoldPoint.NetPushOutNormal / Dt));
 			const FSolverReal MaxImpulseTangentSq = FMath::Square(MaxImpulseTangent);
 			const FSolverReal ImpulseTangentSq = FMath::Square(ImpulseTangentU) + FMath::Square(ImpulseTangentV);
-			if (ImpulseTangentSq > (MaxImpulseTangentSq + SMALL_NUMBER))
+			if (ImpulseTangentSq > (MaxImpulseTangentSq + UE_SMALL_NUMBER))
 			{
 				const FSolverReal ImpulseTangentScale = MaxImpulseTangent * FMath::InvSqrt(ImpulseTangentSq);
 				ImpulseTangentU *= ImpulseTangentScale;
@@ -795,9 +795,9 @@ namespace Chaos
 			ContactMassInvTangentV += FSolverVec3::DotProduct(R1xV, WorldContactTangentVAngular1) + Body1.InvM();
 		}
 
-		ContactMassNormal = (ContactMassInvNormal > FSolverReal(SMALL_NUMBER)) ? FSolverReal(1) / ContactMassInvNormal : FSolverReal(0);
-		ContactMassTangentU = (ContactMassInvTangentU > FSolverReal(SMALL_NUMBER)) ? FSolverReal(1) / ContactMassInvTangentU : FSolverReal(0);
-		ContactMassTangentV = (ContactMassInvTangentV > FSolverReal(SMALL_NUMBER)) ? FSolverReal(1) / ContactMassInvTangentV : FSolverReal(0);
+		ContactMassNormal = (ContactMassInvNormal > FSolverReal(UE_SMALL_NUMBER)) ? FSolverReal(1) / ContactMassInvNormal : FSolverReal(0);
+		ContactMassTangentU = (ContactMassInvTangentU > FSolverReal(UE_SMALL_NUMBER)) ? FSolverReal(1) / ContactMassInvTangentU : FSolverReal(0);
+		ContactMassTangentV = (ContactMassInvTangentV > FSolverReal(UE_SMALL_NUMBER)) ? FSolverReal(1) / ContactMassInvTangentV : FSolverReal(0);
 	}
 
 	FORCEINLINE_DEBUGGABLE void FPBDCollisionSolverManifoldPoint::UpdateMass(
@@ -856,9 +856,9 @@ namespace Chaos
 			VectorStoreFloat3(TangentVAngular1, &WorldContactTangentVAngular1);
 		}
 
-		ContactMassNormal = (ContactMassInvNormal > FSolverReal(SMALL_NUMBER)) ? FSolverReal(1) / ContactMassInvNormal : FSolverReal(0);
-		ContactMassTangentU = (ContactMassInvTangentU > FSolverReal(SMALL_NUMBER)) ? FSolverReal(1) / ContactMassInvTangentU : FSolverReal(0);
-		ContactMassTangentV = (ContactMassInvTangentV > FSolverReal(SMALL_NUMBER)) ? FSolverReal(1) / ContactMassInvTangentV : FSolverReal(0);
+		ContactMassNormal = (ContactMassInvNormal > FSolverReal(UE_SMALL_NUMBER)) ? FSolverReal(1) / ContactMassInvNormal : FSolverReal(0);
+		ContactMassTangentU = (ContactMassInvTangentU > FSolverReal(UE_SMALL_NUMBER)) ? FSolverReal(1) / ContactMassInvTangentU : FSolverReal(0);
+		ContactMassTangentV = (ContactMassInvTangentV > FSolverReal(UE_SMALL_NUMBER)) ? FSolverReal(1) / ContactMassInvTangentV : FSolverReal(0);
 	}
 
 	FORCEINLINE_DEBUGGABLE void FPBDCollisionSolverManifoldPoint::UpdateMassNormal(const FConstraintSolverBody& Body0, const FConstraintSolverBody& Body1)
@@ -878,7 +878,7 @@ namespace Chaos
 			WorldContactNormalAngular1 = InvI1 * R1xN;
 			ContactMassInvNormal += FSolverVec3::DotProduct(R1xN, WorldContactNormalAngular1) + Body1.InvM();
 		}
-		ContactMassNormal = (ContactMassInvNormal > FSolverReal(SMALL_NUMBER)) ? FSolverReal(1) / ContactMassInvNormal : FSolverReal(0);
+		ContactMassNormal = (ContactMassInvNormal > FSolverReal(UE_SMALL_NUMBER)) ? FSolverReal(1) / ContactMassInvNormal : FSolverReal(0);
 	}
 
 	FORCEINLINE_DEBUGGABLE void FPBDCollisionSolverManifoldPoint::CalculateContactPositionErrorNormal(const FConstraintSolverBody& Body0, const FConstraintSolverBody& Body1, const FSolverReal MaxPushOut, FSolverReal& OutContactDeltaNormal) const
@@ -1053,7 +1053,7 @@ namespace Chaos
 			FSolverReal ContactDeltaNormal, ContactDeltaTangent0, ContactDeltaTangent1;
 			SolverManifoldPoint.CalculateContactPositionError(Body0.SolverBody(), Body1.SolverBody(), MaxPushOut, ContactDeltaNormal, ContactDeltaTangent0, ContactDeltaTangent1);
 
-			const bool bProcessManifoldPoint = (ContactDeltaNormal < FSolverReal(0)) || (SolverManifoldPoint.NetPushOutNormal > FSolverReal(SMALL_NUMBER));
+			const bool bProcessManifoldPoint = (ContactDeltaNormal < FSolverReal(0)) || (SolverManifoldPoint.NetPushOutNormal > FSolverReal(UE_SMALL_NUMBER));
 			if (bProcessManifoldPoint)
 			{
 				ApplyPositionCorrectionWithFriction(
@@ -1086,7 +1086,7 @@ namespace Chaos
 			FSolverReal ContactDeltaNormal;
 			SolverManifoldPoint.CalculateContactPositionErrorNormal(Body0.SolverBody(), Body1.SolverBody(), MaxPushOut, ContactDeltaNormal);
 
-			const bool bProcessManifoldPoint = (ContactDeltaNormal < FSolverReal(0)) || (SolverManifoldPoint.NetPushOutNormal > FSolverReal(SMALL_NUMBER));
+			const bool bProcessManifoldPoint = (ContactDeltaNormal < FSolverReal(0)) || (SolverManifoldPoint.NetPushOutNormal > FSolverReal(UE_SMALL_NUMBER));
 			if (bProcessManifoldPoint)
 			{
 				ApplyPositionCorrectionNoFriction(
diff --git a/Engine/Source/Runtime/Experimental/Chaos/Private/Chaos/Collision/SolverCollisionContainer.cpp b/Engine/Source/Runtime/Experimental/Chaos/Private/Chaos/Collision/SolverCollisionContainer.cpp
--- a/Engine/Source/Runtime/Experimental/Chaos/Private/Chaos/Collision/SolverCollisionContainer.cpp
+++ b/Engine/Source/Runtime/Experimental/Chaos/Private/Chaos/Collision/SolverCollisionContainer.cpp
@@ -204,7 +204,7 @@ namespace Chaos
 				// World-space contact tangents. We are treating the normal as the constraint-space Z axis
 				// and the Tangent U and V as the constraint-space X and Y axes respectively
 				FSolverVec3 WorldContactTangentU = FSolverVec3::CrossProduct(FSolverVec3(0, 1, 0), WorldContactNormal);
-				if (!WorldContactTangentU.Normalize(FSolverReal(KINDA_SMALL_NUMBER)))
+				if (!WorldContactTangentU.Normalize(FSolverReal(UE_KINDA_SMALL_NUMBER)))
 				{
 					WorldContactTangentU = FSolverVec3::CrossProduct(FSolverVec3(1, 0, 0), WorldContactNormal);
 					WorldContactTangentU = WorldContactTangentU.GetUnsafeNormal();
@@ -317,7 +317,7 @@ namespace Chaos
 				// and the Tangent U and V as the constraint-space X and Y axes respectively
 				SolverVectorRegister WorldContactTangentU = VectorCross(MakeVectorRegisterFloatConstant(0, 1, 0, 0), WorldContactNormal);
 				SolverVectorRegister WorldContactTangentULenSq = VectorDot3(WorldContactTangentU, WorldContactTangentU);
-				SolverVectorRegister WorldContactTangentUCompareMask = VectorCompareGT(WorldContactTangentULenSq, VectorSetFloat1(KINDA_SMALL_NUMBER));
+				SolverVectorRegister WorldContactTangentUCompareMask = VectorCompareGT(WorldContactTangentULenSq, VectorSetFloat1(UE_KINDA_SMALL_NUMBER));
 				if (VectorMaskBits(WorldContactTangentUCompareMask))
 				{
 					WorldContactTangentU = VectorMultiply(WorldContactTangentU, VectorReciprocalSqrt(WorldContactTangentULenSq));
diff --git a/Engine/Source/Runtime/Experimental/Chaos/Private/Chaos/CollisionOneShotManifolds.cpp b/Engine/Source/Runtime/Experimental/Chaos/Private/Chaos/CollisionOneShotManifolds.cpp
--- a/Engine/Source/Runtime/Experimental/Chaos/Private/Chaos/CollisionOneShotManifolds.cpp
+++ b/Engine/Source/Runtime/Experimental/Chaos/Private/Chaos/CollisionOneShotManifolds.cpp
@@ -105,7 +105,7 @@ namespace Chaos
 				FVec3 Result;
 
 				FReal Denominator = Point2[ClippingAxis] - Point1[ClippingAxis];  // Can be negative
-				if (FMath::Abs(Denominator) < SMALL_NUMBER)
+				if (FMath::Abs(Denominator) < UE_SMALL_NUMBER)
 				{
 					Result = Point1;
 				}
@@ -443,12 +443,12 @@ namespace Chaos
 			FReal PrevClipPointDotNormal;
 			FVec3 CurrentClipPoint = InputVertexBuffer[ClipPointCount - 1];;
 			FReal CurrentClipPointDotNormal = FVec3::DotProduct(CurrentClipPoint, ClippingPlaneNormal);
-			const FReal PlaneClipDistance = PlaneDistance + PlaneDistance * SMALL_NUMBER;
+			const FReal PlaneClipDistance = PlaneDistance + PlaneDistance * UE_SMALL_NUMBER;
 
 			auto CalculateIntersect = [&PrevClipPoint, &CurrentClipPoint, &PlaneDistance](const FReal Dot1, const FReal Dot2)
 			{
 				const FReal Denominator = Dot2 - Dot1; // Can be negative
-				if (FMath::Abs(Denominator) < SMALL_NUMBER)
+				if (FMath::Abs(Denominator) < UE_SMALL_NUMBER)
 				{
 					return PrevClipPoint;
 				}
@@ -549,7 +549,7 @@ namespace Chaos
 		{
 			check(MaxContactPointCount >= 2);
 
-			const FReal DistanceEpsilon = FReal(KINDA_SMALL_NUMBER);
+			const FReal DistanceEpsilon = FReal(UE_KINDA_SMALL_NUMBER);
 
 			// Populate the clipped vertices by the line segment vertices
 			ContactPointCount = 2;
@@ -700,7 +700,7 @@ namespace Chaos
 			{
 				// Ignore planes that do not oppose N
 				const FReal PlaneNormalDotN = FVec3::DotProduct(N, Plane.Normal());
-				if (PlaneNormalDotN <= -SMALL_NUMBER)
+				if (PlaneNormalDotN <= -UE_SMALL_NUMBER)
 				{
 					// Keep the most opposing plane
 					if (PlaneNormalDotN < BestPlaneDot)
@@ -743,7 +743,7 @@ namespace Chaos
 				{
 					// Ignore planes that do not oppose N
 					const FConvexReal PlaneNormalDotN = FConvexVec3::DotProduct(N, ScaledPlaneN);
-					if (PlaneNormalDotN <= -SMALL_NUMBER)
+					if (PlaneNormalDotN <= -UE_SMALL_NUMBER)
 					{
 						// Keep the most opposing plane
 						if (PlaneNormalDotN < InOutBestPlaneDot)
@@ -1201,7 +1201,7 @@ namespace Chaos
 					// Divisor used below to project the capsule contact onto the convex plane
 					// See FMath::RayPlaneIntersection for math, but it is separated here to avoid excessive dot products
 					const FReal ConvexDistanceDivisor = FVec3::DotProduct(CapsulePlaneNormal, ConvexPlaneNormal);
-					if (FMath::Abs(ConvexDistanceDivisor) > KINDA_SMALL_NUMBER)
+					if (FMath::Abs(ConvexDistanceDivisor) > UE_KINDA_SMALL_NUMBER)
 					{
 						const FReal ConvexDistanceMultiplier = FReal(1) / ConvexDistanceDivisor;
 
diff --git a/Engine/Source/Runtime/Experimental/Chaos/Private/Chaos/CollisionOneShotManifoldsMiscShapes.cpp b/Engine/Source/Runtime/Experimental/Chaos/Private/Chaos/CollisionOneShotManifoldsMiscShapes.cpp
--- a/Engine/Source/Runtime/Experimental/Chaos/Private/Chaos/CollisionOneShotManifoldsMiscShapes.cpp
+++ b/Engine/Source/Runtime/Experimental/Chaos/Private/Chaos/CollisionOneShotManifoldsMiscShapes.cpp
@@ -130,7 +130,7 @@ namespace Chaos
 			const FVec3 NearPos = FMath::Lerp(Capsule.GetX1(), Capsule.GetX2(), NearPosT);
 			const FVec3 NearPosDelta = SpherePos - NearPos;
 			const FReal NearPosDistanceSq = NearPosDelta.SizeSquared();
-			if (NearPosDistanceSq > SMALL_NUMBER)
+			if (NearPosDistanceSq > UE_SMALL_NUMBER)
 			{
 				if (NearPosDistanceSq < NetCullDistanceSq)
 				{
@@ -469,7 +469,7 @@ namespace Chaos
 
 			// Calculate the normal from the two closest points. Handle exact axis overlaps.
 			FVec3 ClosestNormal;
-			if (ClosestDeltaLen > KINDA_SMALL_NUMBER)
+			if (ClosestDeltaLen > UE_KINDA_SMALL_NUMBER)
 			{
 				ClosestNormal = -ClosestDelta / ClosestDeltaLen;
 			}
diff --git a/Engine/Source/Runtime/Experimental/Chaos/Private/Chaos/CollisionResolution.cpp b/Engine/Source/Runtime/Experimental/Chaos/Private/Chaos/CollisionResolution.cpp
--- a/Engine/Source/Runtime/Experimental/Chaos/Private/Chaos/CollisionResolution.cpp
+++ b/Engine/Source/Runtime/Experimental/Chaos/Private/Chaos/CollisionResolution.cpp
@@ -209,14 +209,14 @@ namespace Chaos
 			{
 				Dir = DeltaX0 - DeltaX1;
 				Length = Dir.SafeNormalize();
-				if (Length < SMALL_NUMBER) // This is the case where both particles' velocities are high but the relative velocity is low. 
+				if (Length < UE_SMALL_NUMBER) // This is the case where both particles' velocities are high but the relative velocity is low. 
 				{
 					bUseCCD = false;
 				}
 
 				// Do not perform CCD if the vector is not close to unit length to prevent getting caught in a large or infinite loop when raycasting.
 				// @todo(chaos): What's this? Dir is normalized above so this seems unnecessary unless it is checking for 0 length, not unit length - remove it or fix
-				if (!FMath::IsNearlyEqual((float)Dir.SizeSquared(), 1.f, KINDA_SMALL_NUMBER))
+				if (!FMath::IsNearlyEqual((float)Dir.SizeSquared(), 1.f, UE_KINDA_SMALL_NUMBER))
 				{
 					bUseCCD = false;
 				}
@@ -595,7 +595,7 @@ namespace Chaos
 			const FRigidTransform3 BoxToPlaneTransform(BoxTransform.GetRelativeTransform(PlaneTransform));
 			const FVec3 Extents = Box.Extents();
 			constexpr int32 NumCorners = 2 + 2 * 3;
-			constexpr FReal Epsilon = KINDA_SMALL_NUMBER;
+			constexpr FReal Epsilon = UE_KINDA_SMALL_NUMBER;
 
 			FVec3 Corners[NumCorners];
 			int32 CornerIdx = 0;
diff --git a/Engine/Source/Runtime/Experimental/Chaos/Private/Chaos/CollisionResolutionUtil.cpp b/Engine/Source/Runtime/Experimental/Chaos/Private/Chaos/CollisionResolutionUtil.cpp
--- a/Engine/Source/Runtime/Experimental/Chaos/Private/Chaos/CollisionResolutionUtil.cpp
+++ b/Engine/Source/Runtime/Experimental/Chaos/Private/Chaos/CollisionResolutionUtil.cpp
@@ -42,7 +42,7 @@ namespace Chaos
 			FVec3 Jr0, Jr1, IInvJr0, IInvJr1;
 			FReal ImpulseRatioNumerator0 = 0, ImpulseRatioNumerator1 = 0, ImpulseRatioDenom0 = 0, ImpulseRatioDenom1 = 0;
 			FReal ImpulseSizeSQ = Impulse.SizeSquared();
-			if (ImpulseSizeSQ < SMALL_NUMBER)
+			if (ImpulseSizeSQ < UE_SMALL_NUMBER)
 			{
 				return Impulse;
 			}
@@ -91,7 +91,7 @@ namespace Chaos
 			FVec3 Jr0, Jr1, IInvJr0, IInvJr1;
 			FReal ImpulseRatioNumerator0 = 0, ImpulseRatioNumerator1 = 0, ImpulseRatioDenom0 = 0, ImpulseRatioDenom1 = 0;
 			FReal ImpulseSizeSQ = Impulse.SizeSquared();
-			if (ImpulseSizeSQ < SMALL_NUMBER)
+			if (ImpulseSizeSQ < UE_SMALL_NUMBER)
 			{
 				return Impulse;
 			}
@@ -141,7 +141,7 @@ namespace Chaos
 			FReal LocalPhi = Object.PhiWithNormal(LocalPoint, LocalNormal);
 			FReal LocalThickness = LocalPhi - Thickness;
 
-			if (LocalThickness < -KINDA_SMALL_NUMBER)
+			if (LocalThickness < -UE_KINDA_SMALL_NUMBER)
 			{
 				Contact.ShapeContactPoints[0] += LocalPoint * LocalThickness;
 				TotalThickness += LocalThickness;
@@ -598,7 +598,7 @@ namespace Chaos
 
 			if (NormalAveraging)
 			{
-				if (WeightSum < -KINDA_SMALL_NUMBER)
+				if (WeightSum < -UE_KINDA_SMALL_NUMBER)
 				{
 					FVec3 LocalPoint = AvgContact.ShapeContactPoints[0] / WeightSum;
 					FVec3 LocalNormal;
diff --git a/Engine/Source/Runtime/Experimental/Chaos/Private/Chaos/Convex.cpp b/Engine/Source/Runtime/Experimental/Chaos/Private/Chaos/Convex.cpp
--- a/Engine/Source/Runtime/Experimental/Chaos/Private/Chaos/Convex.cpp
+++ b/Engine/Source/Runtime/Experimental/Chaos/Private/Chaos/Convex.cpp
@@ -422,7 +422,7 @@ namespace Chaos
 
 	void FConvex::ComputeUnitMassInertiaTensorAndRotationOfMass(const FReal InVolume)
 	{
-		if (InVolume < SMALL_NUMBER || !StructureData.IsValid())
+		if (InVolume < UE_SMALL_NUMBER || !StructureData.IsValid())
 		{
 			UnitMassInertiaTensor = FVec3{1., 1., 1.};
 			RotationOfMass = FRotation3::Identity;
diff --git a/Engine/Source/Runtime/Experimental/Chaos/Private/Chaos/Evolution/PBDMinEvolution.cpp b/Engine/Source/Runtime/Experimental/Chaos/Private/Chaos/Evolution/PBDMinEvolution.cpp
--- a/Engine/Source/Runtime/Experimental/Chaos/Private/Chaos/Evolution/PBDMinEvolution.cpp
+++ b/Engine/Source/Runtime/Experimental/Chaos/Private/Chaos/Evolution/PBDMinEvolution.cpp
@@ -164,7 +164,7 @@ namespace Chaos
 
 		PrepareTick();
 
-		if (RewindDt > SMALL_NUMBER)
+		if (RewindDt > UE_SMALL_NUMBER)
 		{
 			Rewind(StepDt, RewindDt);
 		}
@@ -423,7 +423,7 @@ namespace Chaos
 				// Target positions only need to be processed once, and we reset the velocity next frame (if no new target is set)
 				FVec3 NewX;
 				FRotation3 NewR;
-				if (FMath::IsNearlyEqual(StepFraction, (FReal)1, (FReal)KINDA_SMALL_NUMBER))
+				if (FMath::IsNearlyEqual(StepFraction, (FReal)1, (FReal)UE_KINDA_SMALL_NUMBER))
 				{
 					NewX = KinematicTarget.GetTarget().GetLocation();
 					NewR = KinematicTarget.GetTarget().GetRotation();
diff --git a/Engine/Source/Runtime/Experimental/Chaos/Private/Chaos/FFT.cpp b/Engine/Source/Runtime/Experimental/Chaos/Private/Chaos/FFT.cpp
--- a/Engine/Source/Runtime/Experimental/Chaos/Private/Chaos/FFT.cpp
+++ b/Engine/Source/Runtime/Experimental/Chaos/Private/Chaos/FFT.cpp
@@ -36,7 +36,7 @@ void EnforceSymmetry(const FFFT3::FUniformGrid& Grid, FFFT3::FArrayNDOfComplex&
 }
 void FFFT3::MakeDivergenceFree(const FUniformGrid& Grid, FArrayNDOfComplex& u, FArrayNDOfComplex& v, FArrayNDOfComplex& w)
 {
-	FVec3 Coefficients = ((FReal)2 * (FReal)PI) / Grid.DomainSize();
+	FVec3 Coefficients = ((FReal)2 * (FReal)UE_PI) / Grid.DomainSize();
 	for (int i = 1; i <= Grid.Counts()[0] / 2; ++i)
 		u(i, 0, 0) = FComplex(0, 0);
 	for (int i = 0; i < Grid.Counts()[0]; ++i)
diff --git a/Engine/Source/Runtime/Experimental/Chaos/Private/Chaos/Framework/PhysicsSolverBase.cpp b/Engine/Source/Runtime/Experimental/Chaos/Private/Chaos/Framework/PhysicsSolverBase.cpp
--- a/Engine/Source/Runtime/Experimental/Chaos/Private/Chaos/Framework/PhysicsSolverBase.cpp
+++ b/Engine/Source/Runtime/Experimental/Chaos/Private/Chaos/Framework/PhysicsSolverBase.cpp
@@ -172,7 +172,7 @@ namespace Chaos
 		, AsyncDt(InAsyncDt)
 		, AccumulatedTime(0)
 		, MMaxDeltaTime(0.0)
-		, MMinDeltaTime(SMALL_NUMBER)
+		, MMinDeltaTime(UE_SMALL_NUMBER)
 		, MMaxSubSteps(1)
 		, ExternalSteps(0)
 #if !UE_BUILD_SHIPPING
diff --git a/Engine/Source/Runtime/Experimental/Chaos/Private/Chaos/HeightField.cpp b/Engine/Source/Runtime/Experimental/Chaos/Private/Chaos/HeightField.cpp
--- a/Engine/Source/Runtime/Experimental/Chaos/Private/Chaos/HeightField.cpp
+++ b/Engine/Source/Runtime/Experimental/Chaos/Private/Chaos/HeightField.cpp
@@ -58,7 +58,7 @@ namespace Chaos
 
 			const int32 FullIndex = Payload + SubY;
 
-			const FReal Radius = Thickness + SMALL_NUMBER;
+			const FReal Radius = Thickness + UE_SMALL_NUMBER;
 			const FReal Radius2 = Radius * Radius;
 			bool bIntersection = false;
 
@@ -121,7 +121,7 @@ namespace Chaos
 
 			const int32 FullIndex = Payload + SubY;
 
-			const FReal Radius = Thickness + SMALL_NUMBER;
+			const FReal Radius = Thickness + UE_SMALL_NUMBER;
 			const FReal Radius2 = Radius * Radius;
 			bool bIntersection = false;
 
@@ -134,7 +134,7 @@ namespace Chaos
 				FVec3 Normal = FVec3::CrossProduct(AB, AC);
 				const FReal Len2 = Normal.SafeNormalize();
 
-				if(!ensure(Len2 > SMALL_NUMBER))
+				if(!ensure(Len2 > UE_SMALL_NUMBER))
 				{
 					// Bad triangle, co-linear points or very thin
 					return false;
@@ -167,7 +167,7 @@ namespace Chaos
 					{
 						const FVec3 ClosestPtOnTri = FindClosestPointOnTriangle(ResultPosition, A, B, C, ResultPosition);
 						const FReal DistToTriangle2 = (ResultPosition - ClosestPtOnTri).SizeSquared();
-						bIntersection = DistToTriangle2 <= SMALL_NUMBER;
+						bIntersection = DistToTriangle2 <= UE_SMALL_NUMBER;
 					}
 				}
 
@@ -472,7 +472,7 @@ namespace Chaos
 				OutBounds.GrowToInclude(Position * InScale);
 			}
 		}
-		OutBounds.Thicken(KINDA_SMALL_NUMBER);
+		OutBounds.Thicken(UE_KINDA_SMALL_NUMBER);
 
 		if(bHaveMaterials)
 		{
@@ -552,7 +552,7 @@ namespace Chaos
 				}
 			}
 
-			OutBounds.Thicken(KINDA_SMALL_NUMBER);
+			OutBounds.Thicken(UE_KINDA_SMALL_NUMBER);
 
 			OutData.MinValue = NewMin;
 			OutData.MaxValue = NewMax;
@@ -883,7 +883,7 @@ namespace Chaos
 		FReal RayExitTime;
 		if(CachedBounds.RaycastFast(StartPoint, Dir, InvDir, bParallel, Length, InvCurrentLength, RayEntryTime, RayExitTime))
 		{
-			CurrentLength = RayExitTime + SMALL_NUMBER; // to account for precision errors 
+			CurrentLength = RayExitTime + UE_SMALL_NUMBER; // to account for precision errors 
 			FVec3 NextStart = StartPoint + (Dir * RayEntryTime);
 
 			const FVec2 Scale2D(GeomData.Scale[0],GeomData.Scale[1]);
@@ -1281,7 +1281,7 @@ namespace Chaos
 			FVec3 Normal = FVec3::CrossProduct(AB, AC);
 			const FReal NormalLength = Normal.SafeNormalize();
 
-			if(!ensure(NormalLength > KINDA_SMALL_NUMBER))
+			if(!ensure(NormalLength > UE_KINDA_SMALL_NUMBER))
 			{
 				return false;
 			}
@@ -2068,7 +2068,7 @@ namespace Chaos
 			const FVec3 AC = C - A;
 			FVec3 Normal = FVec3::CrossProduct(AB, AC);
 			const FReal NormalLength = Normal.SafeNormalize();
-			if(!ensure(NormalLength > KINDA_SMALL_NUMBER))
+			if(!ensure(NormalLength > UE_KINDA_SMALL_NUMBER))
 			{
 				//hitting degenerate triangle - should be fixed before we get to this stage
 				return;
@@ -2134,7 +2134,7 @@ namespace Chaos
 				FVec3 Normal = FVec3::CrossProduct(AB, AC);
 
 				const FReal NormalLength = Normal.SafeNormalize();
-				if(!ensure(NormalLength > KINDA_SMALL_NUMBER))
+				if(!ensure(NormalLength > UE_KINDA_SMALL_NUMBER))
 				{
 					//hitting degenerate triangle - should be fixed before we get to this stage
 					return;
diff --git a/Engine/Source/Runtime/Experimental/Chaos/Private/Chaos/Joint/PBDJointCachedSolverGaussSeidel.cpp b/Engine/Source/Runtime/Experimental/Chaos/Private/Chaos/Joint/PBDJointCachedSolverGaussSeidel.cpp
--- a/Engine/Source/Runtime/Experimental/Chaos/Private/Chaos/Joint/PBDJointCachedSolverGaussSeidel.cpp
+++ b/Engine/Source/Runtime/Experimental/Chaos/Private/Chaos/Joint/PBDJointCachedSolverGaussSeidel.cpp
@@ -43,11 +43,11 @@ void FPBDJointCachedSolver::ComputeBodyState(const int32 BodyIndex)
 void FPBDJointCachedSolver::UpdateDerivedState()
 {
 	// Kinematic bodies will not be moved, so we don't update derived state during iterations
-	if (InvM(0) > SMALL_NUMBER)
+	if (InvM(0) > UE_SMALL_NUMBER)
 	{
 		ComputeBodyState(0);
 	}
-	if (InvM(1) > SMALL_NUMBER)
+	if (InvM(1) > UE_SMALL_NUMBER)
 	{
 		ComputeBodyState(1);
 	}
@@ -536,8 +536,8 @@ void FPBDJointCachedSolver::InitSphericalPositionConstraint(
 	FReal SphereDelta0;
 	FPBDJointUtilities::GetSphericalAxisDelta(ConnectorXs[0], ConnectorXs[1], SphereAxis0, SphereDelta0);
 	
-	const FVec3 SphereAxis1 = (FMath::Abs(FMath::Abs(SphereAxis0.Dot(FVec3(1,0,0))-1.0)) > SMALL_NUMBER) ? SphereAxis0.Cross(FVec3(1,0,0)) :
-	                          (FMath::Abs(FMath::Abs(SphereAxis0.Dot(FVec3(0,1,0))-1.0)) > SMALL_NUMBER) ? SphereAxis0.Cross(FVec3(0,1,0)) : SphereAxis0.Cross(FVec3(0,0,1)) ;
+	const FVec3 SphereAxis1 = (FMath::Abs(FMath::Abs(SphereAxis0.Dot(FVec3(1,0,0))-1.0)) > UE_SMALL_NUMBER) ? SphereAxis0.Cross(FVec3(1,0,0)) :
+	                          (FMath::Abs(FMath::Abs(SphereAxis0.Dot(FVec3(0,1,0))-1.0)) > UE_SMALL_NUMBER) ? SphereAxis0.Cross(FVec3(0,1,0)) : SphereAxis0.Cross(FVec3(0,0,1)) ;
 	const FVec3 SphereAxis2 = SphereAxis0.Cross(SphereAxis1);
 
 	// Using Connector1 for both conserves angular momentum and avoid having 
@@ -656,7 +656,7 @@ void FPBDJointCachedSolver::SolvePositionConstraintSoft(
 	const FReal TargetVel)
 {
 	FReal VelDt = 0;
-	if (PositionConstraints.ConstraintSoftDamping[ConstraintIndex] > KINDA_SMALL_NUMBER)
+	if (PositionConstraints.ConstraintSoftDamping[ConstraintIndex] > UE_KINDA_SMALL_NUMBER)
 	{
 		const FVec3 V0Dt = FVec3::CalculateVelocity(InitConnectorXs[0], ConnectorXs[0]+Body(0).DP() + FVec3::CrossProduct(Body(0).DQ(), PositionConstraints.ConstraintArms[ConstraintIndex][0]), 1.0f);
 		const FVec3 V1Dt = FVec3::CalculateVelocity(InitConnectorXs[1], ConnectorXs[1]+Body(1).DP() + FVec3::CrossProduct(Body(1).DQ(), PositionConstraints.ConstraintArms[ConstraintIndex][1]), 1.0f);
@@ -754,7 +754,7 @@ void FPBDJointCachedSolver::SolveLinearVelocityConstraint(
 
 void FPBDJointCachedSolver::ApplyAxisVelocityConstraint(const int32 ConstraintIndex)
 {
-	if(!NetLinearImpulse.IsNearlyZero() && (FMath::Abs(PositionConstraints.ConstraintLambda[ConstraintIndex]) > SMALL_NUMBER) )
+	if(!NetLinearImpulse.IsNearlyZero() && (FMath::Abs(PositionConstraints.ConstraintLambda[ConstraintIndex]) > UE_SMALL_NUMBER) )
 	{
 		FReal TargetVel = 0.0f;
 		if (PositionConstraints.MotionType[ConstraintIndex] == EJointMotionType::Limited && PositionConstraints.ConstraintRestitution[ConstraintIndex] != 0.0f)
@@ -1149,7 +1149,7 @@ void FPBDJointCachedSolver::SolveRotationConstraintSoft(
 {
 	// Damping angular velocity
 	FReal AngVelDt = 0;
-	if (RotationConstraints.ConstraintSoftDamping[ConstraintIndex] > KINDA_SMALL_NUMBER)
+	if (RotationConstraints.ConstraintSoftDamping[ConstraintIndex] > UE_KINDA_SMALL_NUMBER)
 	{
 		const FVec3 W0Dt = FVec3(Body(0).DQ()) + ConnectorWDts[0];
 		const FVec3 W1Dt = FVec3(Body(1).DQ()) + ConnectorWDts[1];
@@ -1242,7 +1242,7 @@ void FPBDJointCachedSolver::SolveAngularVelocityConstraint(
 
 void FPBDJointCachedSolver::ApplyAngularVelocityConstraint(const int32 ConstraintIndex)
 {
-	if(!NetAngularImpulse.IsNearlyZero() && (FMath::Abs(RotationConstraints.ConstraintLambda[ConstraintIndex]) > SMALL_NUMBER))
+	if(!NetAngularImpulse.IsNearlyZero() && (FMath::Abs(RotationConstraints.ConstraintLambda[ConstraintIndex]) > UE_SMALL_NUMBER))
 	{
 		FReal TargetVel = 0.0f;
 		if (RotationConstraints.ConstraintRestitution[ConstraintIndex] != 0.0f)
@@ -1505,7 +1505,7 @@ void FPBDJointCachedSolver::ApplyAxisPositionDrive(
 	const FReal DeltaPos = PositionDrives.ConstraintCX[ConstraintIndex] + FVec3::DotProduct(Delta1 - Delta0, PositionDrives.ConstraintAxis[ConstraintIndex]);
 
 	FReal VelDt = 0;
-	if (PositionDrives.ConstraintSoftDamping[ConstraintIndex] > KINDA_SMALL_NUMBER)
+	if (PositionDrives.ConstraintSoftDamping[ConstraintIndex] > UE_KINDA_SMALL_NUMBER)
 	{
 		const FVec3 V0Dt = FVec3::CalculateVelocity(InitConnectorXs[0], ConnectorXs[0]+ Delta0, 1.0f);
 		const FVec3 V1Dt = FVec3::CalculateVelocity(InitConnectorXs[1], ConnectorXs[1]+ Delta1, 1.0f);
@@ -1657,9 +1657,9 @@ void FPBDJointCachedSolver::InitSLerpDrive(
 		if (RotationDrives.ConstraintSoftStiffness[0] > 0.0f)
 		{
 			FPBDJointUtilities::GetLockedRotationAxes(ConnectorRs[0], ConnectorRs[1], Axes[0], Axes[1], Axes[2]);
-			Utilities::NormalizeSafe(Axes[0], KINDA_SMALL_NUMBER);
-			Utilities::NormalizeSafe(Axes[1], KINDA_SMALL_NUMBER);
-			Utilities::NormalizeSafe(Axes[2], KINDA_SMALL_NUMBER);
+			Utilities::NormalizeSafe(Axes[0], UE_KINDA_SMALL_NUMBER);
+			Utilities::NormalizeSafe(Axes[1], UE_KINDA_SMALL_NUMBER);
+			Utilities::NormalizeSafe(Axes[2], UE_KINDA_SMALL_NUMBER);
 		}
 		const FRotation3 R01 = ConnectorRs[0].Inverse() * ConnectorRs[1];
 		FRotation3 TargetAngPos = JointSettings.AngularDrivePositionTarget;
@@ -1689,9 +1689,9 @@ void FPBDJointCachedSolver::InitSLerpDrive(
 		FReal SLerpAngle;
 		if (DR.ToAxisAndAngleSafe(SLerpAxis, SLerpAngle, FVec3(1, 0, 0)))
 		{
-			if (SLerpAngle > (FReal)PI)
+			if (SLerpAngle > (FReal)UE_PI)
 			{
-				SLerpAngle = SLerpAngle - (FReal)2 * PI;
+				SLerpAngle = SLerpAngle - (FReal)2 * UE_PI;
 			}
 
 			if (FMath::Abs(SLerpAngle) > AngleTolerance)
@@ -1727,7 +1727,7 @@ void FPBDJointCachedSolver::ApplyAxisRotationDrive(
 
 	// Damping angular velocity
 	FReal AngVelDt = 0;
-	if (RotationDrives.ConstraintSoftDamping[ConstraintIndex] > KINDA_SMALL_NUMBER)
+	if (RotationDrives.ConstraintSoftDamping[ConstraintIndex] > UE_KINDA_SMALL_NUMBER)
 	{
 		const FVec3 W0Dt = FVec3(Body(0).DQ()) + ConnectorWDts[0];
 		const FVec3 W1Dt = FVec3(Body(1).DQ()) + ConnectorWDts[1];
diff --git a/Engine/Source/Runtime/Experimental/Chaos/Private/Chaos/Joint/PBDJointSolverGaussSeidel.cpp b/Engine/Source/Runtime/Experimental/Chaos/Private/Chaos/Joint/PBDJointSolverGaussSeidel.cpp
--- a/Engine/Source/Runtime/Experimental/Chaos/Private/Chaos/Joint/PBDJointSolverGaussSeidel.cpp
+++ b/Engine/Source/Runtime/Experimental/Chaos/Private/Chaos/Joint/PBDJointSolverGaussSeidel.cpp
@@ -57,14 +57,14 @@ namespace Chaos
 	void FPBDJointSolver::UpdateDerivedState()
 	{
 		// Kinematic bodies will not be moved, so we don't update derived state during iterations
-		if (InvM(0) > SMALL_NUMBER)
+		if (InvM(0) > UE_SMALL_NUMBER)
 		{
 			const FVec3 BodyP0 = P(0);
 			const FRotation3 BodyQ0 = Q(0);
 			ConnectorXs[0] = BodyP0 + BodyQ0 * LocalConnectorXs[0].GetTranslation();
 			ConnectorRs[0] = BodyQ0 * LocalConnectorXs[0].GetRotation();
 		}
-		if (InvM(1) > SMALL_NUMBER)
+		if (InvM(1) > UE_SMALL_NUMBER)
 		{
 			const FVec3 BodyP1 = P(1);
 			const FRotation3 BodyQ1 = Q(1);
@@ -964,7 +964,7 @@ namespace Chaos
 		const FPBDJointSettings& JointSettings,
 		const bool bUseSoftLimit)
 	{
-		if(!NetAngularImpulse.IsNearlyZero() && !bUseSoftLimit && AngularHardLambda[(int32)EJointAngularConstraintIndex::Twist] > SMALL_NUMBER)
+		if(!NetAngularImpulse.IsNearlyZero() && !bUseSoftLimit && AngularHardLambda[(int32)EJointAngularConstraintIndex::Twist] > UE_SMALL_NUMBER)
 		{
 			FVec3 TwistAxis;
 			FReal TwistAngle;
@@ -993,7 +993,7 @@ namespace Chaos
 		const FPBDJointSettings& JointSettings,
 		const bool bUseSoftLimit)
 	{
-		if(!NetAngularImpulse.IsNearlyZero() && !bUseSoftLimit && AngularHardLambda[(int32)EJointAngularConstraintIndex::Swing2] > SMALL_NUMBER)
+		if(!NetAngularImpulse.IsNearlyZero() && !bUseSoftLimit && AngularHardLambda[(int32)EJointAngularConstraintIndex::Swing2] > UE_SMALL_NUMBER)
 		{
 			FVec3 SwingAxisLocal;
 			FReal DSwingAngle = 0.0f;
@@ -1022,7 +1022,7 @@ namespace Chaos
 		const EJointAngularConstraintIndex SwingConstraintIndex,
 		const bool bUseSoftLimit)
 	{
-		if(!NetAngularImpulse.IsNearlyZero() && !bUseSoftLimit && AngularHardLambda[(int32)SwingConstraintIndex] > SMALL_NUMBER)
+		if(!NetAngularImpulse.IsNearlyZero() && !bUseSoftLimit && AngularHardLambda[(int32)SwingConstraintIndex] > UE_SMALL_NUMBER)
 		{
 			FVec3 SwingAxis;
 			FReal SwingAngle;
@@ -1042,7 +1042,7 @@ namespace Chaos
 		const EJointAngularConstraintIndex SwingConstraintIndex,
 		const bool bUseSoftLimit)
 	{
-		if(!NetAngularImpulse.IsNearlyZero() && !bUseSoftLimit && AngularHardLambda[(int32)SwingConstraintIndex] > SMALL_NUMBER)
+		if(!NetAngularImpulse.IsNearlyZero() && !bUseSoftLimit && AngularHardLambda[(int32)SwingConstraintIndex] > UE_SMALL_NUMBER)
 		{
 			FVec3 SwingAxis;
 			FReal SwingAngle;
@@ -1066,7 +1066,7 @@ namespace Chaos
 		const EJointAngularConstraintIndex SwingConstraintIndex,
 		const bool bUseSoftLimit)
 	{
-		if(!NetAngularImpulse.IsNearlyZero() && !bUseSoftLimit && AngularHardLambda[(int32)SwingConstraintIndex] > SMALL_NUMBER)
+		if(!NetAngularImpulse.IsNearlyZero() && !bUseSoftLimit && AngularHardLambda[(int32)SwingConstraintIndex] > UE_SMALL_NUMBER)
 		{
 			FVec3 SwingAxis;
 			FReal SwingAngle;
@@ -1114,14 +1114,14 @@ namespace Chaos
 			FPBDJointUtilities::GetLockedRotationAxes(ConnectorRs[0], ConnectorRs[1], Axis0, Axis1, Axis2);
 			const FRotation3 R01 = ConnectorRs[0].Inverse() * ConnectorRs[1];
 
-			if (bApplyTwist && AngularHardLambda[(int32)EJointAngularConstraintIndex::Twist] > SMALL_NUMBER)
+			if (bApplyTwist && AngularHardLambda[(int32)EJointAngularConstraintIndex::Twist] > UE_SMALL_NUMBER)
 			{
 				FReal TwistStiffness = SolverStiffness * FPBDJointUtilities::GetTwistStiffness(SolverSettings, JointSettings);
 				Axis0.SafeNormalize();
 				ApplyAngularVelocityConstraint(TwistStiffness, Axis0);
 			}
 
-			if (bApplySwing && AngularHardLambda[(int32)EJointAngularConstraintIndex::Swing2] + AngularHardLambda[(int32)EJointAngularConstraintIndex::Swing1] > SMALL_NUMBER)
+			if (bApplySwing && AngularHardLambda[(int32)EJointAngularConstraintIndex::Swing2] + AngularHardLambda[(int32)EJointAngularConstraintIndex::Swing1] > UE_SMALL_NUMBER)
 			{
 				FReal SwingStiffness = SolverStiffness * FPBDJointUtilities::GetSwingStiffness(SolverSettings, JointSettings);
 				Axis1.SafeNormalize();
@@ -1321,7 +1321,7 @@ namespace Chaos
 			const FReal II = (InvM(0) + II0 + InvM(1) + II1);
 	
 			FReal VelDt = 0;
-			if (JointDamping > KINDA_SMALL_NUMBER)
+			if (JointDamping > UE_KINDA_SMALL_NUMBER)
 			{
 				const FVec3 V0Dt = FVec3::CalculateVelocity(InitConnectorXs[0], ConnectorXs[0], 1.0f);
 				const FVec3 V1Dt = FVec3::CalculateVelocity(InitConnectorXs[1], ConnectorXs[1], 1.0f);
@@ -1460,7 +1460,7 @@ namespace Chaos
 
 			// Damping angular velocity
 			FReal AngVelDt = 0;
-			if (JointDamping > KINDA_SMALL_NUMBER)
+			if (JointDamping > UE_KINDA_SMALL_NUMBER)
 			{
 				const FVec3 W0Dt = FRotation3::CalculateAngularVelocity(InitConnectorRs[KIndex], ConnectorRs[KIndex], 1.0f);
 				const FVec3 W1Dt = FRotation3::CalculateAngularVelocity(InitConnectorRs[DIndex], ConnectorRs[DIndex], 1.0f);
@@ -1518,7 +1518,7 @@ namespace Chaos
 
 			// Damping angular velocity
 			FReal AngVelDt = 0;
-			if (JointDamping > KINDA_SMALL_NUMBER)
+			if (JointDamping > UE_KINDA_SMALL_NUMBER)
 			{
 				const FVec3 W0Dt = FRotation3::CalculateAngularVelocity(InitConnectorRs[0], ConnectorRs[0], 1.0f);
 				const FVec3 W1Dt = FRotation3::CalculateAngularVelocity(InitConnectorRs[1], ConnectorRs[1], 1.0f);
@@ -1969,9 +1969,9 @@ namespace Chaos
 			if (AngularDriveStiffness > 0.0f)
 			{
 				FPBDJointUtilities::GetLockedRotationAxes(ConnectorRs[0], ConnectorRs[1], Axes[0], Axes[1], Axes[2]);
-				Utilities::NormalizeSafe(Axes[0], KINDA_SMALL_NUMBER);
-				Utilities::NormalizeSafe(Axes[1], KINDA_SMALL_NUMBER);
-				Utilities::NormalizeSafe(Axes[2], KINDA_SMALL_NUMBER);
+				Utilities::NormalizeSafe(Axes[0], UE_KINDA_SMALL_NUMBER);
+				Utilities::NormalizeSafe(Axes[1], UE_KINDA_SMALL_NUMBER);
+				Utilities::NormalizeSafe(Axes[2], UE_KINDA_SMALL_NUMBER);
 			}
 
 			const FRotation3 R01 = ConnectorRs[0].Inverse() * ConnectorRs[1];
@@ -2004,9 +2004,9 @@ namespace Chaos
 			FReal SLerpAngle;
 			if (DR.ToAxisAndAngleSafe(SLerpAxis, SLerpAngle, FVec3(1, 0, 0)))
 			{
-				if (SLerpAngle > (FReal)PI)
+				if (SLerpAngle > (FReal)UE_PI)
 				{
-					SLerpAngle = SLerpAngle - (FReal)2 * PI;
+					SLerpAngle = SLerpAngle - (FReal)2 * UE_PI;
 				}
 
 				UE_LOG(LogChaosJoint, VeryVerbose, TEXT("      SLerpDrive Pos: %f Axis: %f %f %f"), -SLerpAngle, SLerpAxis.X, SLerpAxis.Y, SLerpAxis.Z);
@@ -2640,11 +2640,11 @@ namespace Chaos
 			// Calculate constraint correction
 			FMatrix33 M0 = FMatrix33::Zero;
 			FMatrix33 M1 = FMatrix33::Zero;
-			if (InvM(0) > SMALL_NUMBER)
+			if (InvM(0) > UE_SMALL_NUMBER)
 			{
 				M0 = Utilities::ComputeJointFactorMatrix(ConnectorXs[0] - P(0), InvI(0), InvM(0));
 			}
-			if (InvM(1) > SMALL_NUMBER)
+			if (InvM(1) > UE_SMALL_NUMBER)
 			{
 				M1 = Utilities::ComputeJointFactorMatrix(ConnectorXs[1] - P(1), InvI(1), InvM(1));
 			}
@@ -2702,7 +2702,7 @@ namespace Chaos
 			FPBDJointUtilities::GetCylindricalAxesDeltas(ConnectorRs[0], ConnectorXs[0], ConnectorXs[1], AxisIndex, Axis, AxialDelta, RadialAxis, RadialDelta);
 			const FVec3 DConnectorX = ConnectorXs[1] - ConnectorXs[0];
 
-			if (LinearHardLambda[(int32)EJointCylindricalPositionConstraintType::Axial] > SMALL_NUMBER)
+			if (LinearHardLambda[(int32)EJointCylindricalPositionConstraintType::Axial] > UE_SMALL_NUMBER)
 			{
 				const FReal Stiffness = SolverStiffness * FPBDJointUtilities::GetLinearStiffness(SolverSettings, JointSettings);
 
@@ -2728,7 +2728,7 @@ namespace Chaos
 				ApplyLinearVelocityConstraint(Stiffness, Axis, Connector0Correction, TargetVel);
 			}
 
-			if (LinearHardLambda[(int32)EJointCylindricalPositionConstraintType::Radial] > SMALL_NUMBER)
+			if (LinearHardLambda[(int32)EJointCylindricalPositionConstraintType::Radial] > UE_SMALL_NUMBER)
 			{
 				const FReal Stiffness = SolverStiffness * FPBDJointUtilities::GetLinearStiffness(SolverSettings, JointSettings);
 
diff --git a/Engine/Source/Runtime/Experimental/Chaos/Private/Chaos/Levelset.cpp b/Engine/Source/Runtime/Experimental/Chaos/Private/Chaos/Levelset.cpp
--- a/Engine/Source/Runtime/Experimental/Chaos/Private/Chaos/Levelset.cpp
+++ b/Engine/Source/Runtime/Experimental/Chaos/Private/Chaos/Levelset.cpp
@@ -628,10 +628,10 @@ void FLevelSet::ComputeConvexity(const TArray<TVec3<int32>>& InterfaceIndices)
 		FReal PhiYZ = (MPhi[MAX_CLAMP(i + ZOffset + 1, NumCells, i)] + MPhi[MIN_CLAMP(i - ZOffset - 1, 0, i)] - MPhi[RANGE_CLAMP(i - ZOffset + 1, NumCells, i)] - MPhi[RANGE_CLAMP(i + ZOffset - 1, NumCells, i)]) / (4 * MGrid.Dx()[1] * MGrid.Dx()[2]);
 
 		FReal Denom = sqrt(PhiX * PhiX + PhiY * PhiY + PhiZ * PhiZ);
-		if (Denom > SMALL_NUMBER)
+		if (Denom > UE_SMALL_NUMBER)
 		{
 			FReal curvature = -(PhiX * PhiX * PhiYY - 2 * PhiX * PhiY * PhiXY + PhiY * PhiY * PhiXX + PhiX * PhiX * PhiZZ - 2 * PhiX * PhiZ * PhiXZ + PhiZ * PhiZ * PhiXX + PhiY * PhiY * PhiZZ - 2 * PhiY * PhiZ * PhiYZ + PhiZ * PhiZ * PhiYY) / (Denom * Denom * Denom);
-			LocalSign = curvature > KINDA_SMALL_NUMBER ? 1 : (curvature < -KINDA_SMALL_NUMBER ? -1 : 0);
+			LocalSign = curvature > UE_KINDA_SMALL_NUMBER ? 1 : (curvature < -UE_KINDA_SMALL_NUMBER ? -1 : 0);
 			if (bFirst)
 			{
 				bFirst = false;
@@ -676,8 +676,8 @@ bool FLevelSet::ComputeDistancesNearZeroIsocontour(FErrorReporter& ErrorReporter
 		TriangleBounds.GrowToInclude(InParticles.X(Element[2]));
 		MOriginalLocalBoundingBox.GrowToInclude(TriangleBounds); //also save the original bounding box
 
-		TVec3<int32> StartIndex = MGrid.Cell(TriangleBounds.Min() - FVec3((0.5f + KINDA_SMALL_NUMBER) * MGrid.Dx()));
-		TVec3<int32> EndIndex = MGrid.Cell(TriangleBounds.Max() + FVec3((0.5f + KINDA_SMALL_NUMBER) * MGrid.Dx()));
+		TVec3<int32> StartIndex = MGrid.Cell(TriangleBounds.Min() - FVec3((0.5f + UE_KINDA_SMALL_NUMBER) * MGrid.Dx()));
+		TVec3<int32> EndIndex = MGrid.Cell(TriangleBounds.Max() + FVec3((0.5f + UE_KINDA_SMALL_NUMBER) * MGrid.Dx()));
 		for (int32 i = StartIndex[0]; i <= EndIndex[0]; ++i)
 		{
 			for (int32 j = StartIndex[1]; j <= EndIndex[1]; ++j)
@@ -1251,7 +1251,7 @@ void FLevelSet::ComputeNormals()
 				    (SignedDistance(X + FVec3::AxisVector(1) * Dx[1]) - SignedDistance(X - FVec3::AxisVector(1) * Dx[1])) / (2 * Dx[1]),
 				    (SignedDistance(X + FVec3::AxisVector(2) * Dx[2]) - SignedDistance(X - FVec3::AxisVector(2) * Dx[2])) / (2 * Dx[2]));
 				FReal Size = MNormals(CellIndex).Size();
-				if (Size > SMALL_NUMBER)
+				if (Size > UE_SMALL_NUMBER)
 				{
 					MNormals(CellIndex) /= Size;
 				}
@@ -1304,8 +1304,8 @@ void FLevelSet::ComputeNormals(const FParticles& InParticles, const FTriangleMes
 		TriangleBounds.GrowToInclude(InParticles.X(Element[2]));
 		MOriginalLocalBoundingBox.GrowToInclude(TriangleBounds); //also save the original bounding box
 
-		TVec3<int32> StartIndex = MGrid.Cell(TriangleBounds.Min() - FVec3((0.5f + KINDA_SMALL_NUMBER) * MGrid.Dx()));
-		TVec3<int32> EndIndex = MGrid.Cell(TriangleBounds.Max() + FVec3((0.5f + KINDA_SMALL_NUMBER) * MGrid.Dx()));
+		TVec3<int32> StartIndex = MGrid.Cell(TriangleBounds.Min() - FVec3((0.5f + UE_KINDA_SMALL_NUMBER) * MGrid.Dx()));
+		TVec3<int32> EndIndex = MGrid.Cell(TriangleBounds.Max() + FVec3((0.5f + UE_KINDA_SMALL_NUMBER) * MGrid.Dx()));
 		for (int32 i = StartIndex[0]; i <= EndIndex[0]; ++i)
 		{
 			for (int32 j = StartIndex[1]; j <= EndIndex[1]; ++j)
@@ -1413,7 +1413,7 @@ FReal FLevelSet::PhiWithNormal(const FVec3& x, FVec3& Normal) const
 	{
 		Normal = MGrid.LinearlyInterpolate(MNormals, Location);
 		FReal NormalMag = Normal.Size();
-		if (NormalMag > SMALL_NUMBER)
+		if (NormalMag > UE_SMALL_NUMBER)
 		{
 			Normal /= NormalMag;
 		}
diff --git a/Engine/Source/Runtime/Experimental/Chaos/Private/Chaos/MassProperties.cpp b/Engine/Source/Runtime/Experimental/Chaos/Private/Chaos/MassProperties.cpp
--- a/Engine/Source/Runtime/Experimental/Chaos/Private/Chaos/MassProperties.cpp
+++ b/Engine/Source/Runtime/Experimental/Chaos/Private/Chaos/MassProperties.cpp
@@ -20,13 +20,13 @@ namespace Chaos
 		FReal OffDiagSize = FMath::Square(Inertia.M[1][0]) + FMath::Square(Inertia.M[2][0]) + FMath::Square(Inertia.M[2][1]);
 		FRealDouble Trace = (Inertia.M[0][0] + Inertia.M[1][1] + Inertia.M[2][2]) / 3;
 
-		if (Trace <= SMALL_NUMBER)
+		if (Trace <= UE_SMALL_NUMBER)
 		{
 			// Tiny inertia - numerical instability would follow. We should not get this unless we have bad input.
 			return FRotation3::FromElements(FVec3(0), 1);
 		}
 
-		if ((OffDiagSize / Trace) < SMALL_NUMBER)
+		if ((OffDiagSize / Trace) < UE_SMALL_NUMBER)
 		{
 			// Almost diagonal matrix - we are already in local space.
 			return FRotation3::FromElements(FVec3(0), 1);
@@ -35,10 +35,10 @@ namespace Chaos
 		FReal Size = static_cast<FReal>(FMath::Sqrt((FMath::Square(Inertia.M[0][0] - Trace) + FMath::Square(Inertia.M[1][1] - Trace) + FMath::Square(Inertia.M[2][2] - Trace) + 2. * OffDiagSize) / 6.));
 		FMatrix33 NewMat = (Inertia - FMatrix::Identity * static_cast<FReal>(Trace)) * static_cast<FReal>((1 / Size));
 		FReal HalfDeterminant = NewMat.Determinant() / 2;
-		FReal Angle = HalfDeterminant <= -1 ? PI / 3 : (HalfDeterminant >= 1 ? 0 : FMath::Acos(HalfDeterminant) / 3);
+		FReal Angle = HalfDeterminant <= -1 ? UE_PI / 3 : (HalfDeterminant >= 1 ? 0 : FMath::Acos(HalfDeterminant) / 3);
 		
 		FReal m00 = static_cast<FReal>(Trace + 2 * Size * FMath::Cos(Angle));
-		FReal m11 = static_cast<FReal>(Trace + 2 * Size * FMath::Cos(Angle + (2 * PI / 3)));
+		FReal m11 = static_cast<FReal>(Trace + 2 * Size * FMath::Cos(Angle + (2 * UE_PI / 3)));
 		FReal m22 = static_cast<FReal>(3 * Trace - m00 - m11);
 
 		// Extract Eigenvectors
@@ -62,7 +62,7 @@ namespace Chaos
 		FReal SqrtIM1Scale1 = FMath::Sqrt(IM1Scale1);
 
 		FVec3 Eigenvector2, Eigenvector1;
-		if ((SqrtIM1Scale0 < KINDA_SMALL_NUMBER) && (SqrtIM1Scale1 < KINDA_SMALL_NUMBER))
+		if ((SqrtIM1Scale0 < UE_KINDA_SMALL_NUMBER) && (SqrtIM1Scale1 < UE_KINDA_SMALL_NUMBER))
 		{
 			Eigenvector1 = Orthogonal;
 			Eigenvector2 = FVec3::CrossProduct(Eigenvector0, Orthogonal).GetSafeNormal();
@@ -82,7 +82,7 @@ namespace Chaos
 
 		// NOTE: UE Matrix are column-major, so the PMatrix constructor is not setting eigenvectors - we need to transpose it to get a UE rotation matrix.
 		FinalRotation = FRotation3(RotationMatrix.GetTransposed());
-		if (!ensure(FMath::IsNearlyEqual((float)FinalRotation.Size(), 1.0f, KINDA_SMALL_NUMBER)))
+		if (!ensure(FMath::IsNearlyEqual((float)FinalRotation.Size(), 1.0f, UE_KINDA_SMALL_NUMBER)))
 		{
 			return FRotation3::FromElements(FVec3(0), 1);
 		}
@@ -132,7 +132,7 @@ namespace Chaos
 			VolumeTimesSum += Det * PerElementSize;
 		}
 		// @todo(mlentine): Should add suppoert for thin shell mass properties
-		if (Volume < KINDA_SMALL_NUMBER)	//handle negative volume using fallback for now. Need to investigate cases where this happens
+		if (Volume < UE_KINDA_SMALL_NUMBER)	//handle negative volume using fallback for now. Need to investigate cases where this happens
 		{
 			OutVolume = 0;
 			return;
@@ -181,7 +181,7 @@ namespace Chaos
 		CalculateVolumeAndCenterOfMass(Vertices, Surfaces, MassProperties.Volume, MassProperties.CenterOfMass);
 
 		check(Mass > 0);
-		check(MassProperties.Volume > SMALL_NUMBER);
+		check(MassProperties.Volume > UE_SMALL_NUMBER);
 		CalculateInertiaAndRotationOfMass(Vertices, Surfaces, Mass / MassProperties.Volume, MassProperties.CenterOfMass, MassProperties.InertiaTensor, MassProperties.RotationOfMass);
 		
 		return MassProperties;
@@ -209,7 +209,7 @@ namespace Chaos
 			NewMP.CenterOfMass += Child.CenterOfMass * Child.Mass;
 			NewMP.Mass += Child.Mass;
 		}
-		check(NewMP.Mass > SMALL_NUMBER);
+		check(NewMP.Mass > UE_SMALL_NUMBER);
 		NewMP.CenterOfMass /= NewMP.Mass;
 		for (const FMassProperties& Child : MPArray)
 		{
diff --git a/Engine/Source/Runtime/Experimental/Chaos/Private/Chaos/PBDCollisionConstraintsContact.cpp b/Engine/Source/Runtime/Experimental/Chaos/Private/Chaos/PBDCollisionConstraintsContact.cpp
--- a/Engine/Source/Runtime/Experimental/Chaos/Private/Chaos/PBDCollisionConstraintsContact.cpp
+++ b/Engine/Source/Runtime/Experimental/Chaos/Private/Chaos/PBDCollisionConstraintsContact.cpp
@@ -60,7 +60,7 @@ namespace Chaos
 
 				bool bApplyResitution = (ContactRestitution > 0.0f);
 				bool bHaveRestitutionPadding = (ContactRestitutionPadding > 0.0f);
-				bool bApplyFriction = (ContactFriction > 0) && (IterationParameters.Dt > SMALL_NUMBER);
+				bool bApplyFriction = (ContactFriction > 0) && (IterationParameters.Dt > UE_SMALL_NUMBER);
 
 				// If we have restitution, pad the constraint by an amount that enforces the outgoing velocity constraint
 				// Really this should be per contact point, not per constraint.
@@ -115,7 +115,7 @@ namespace Chaos
 					{
 						FVec3 CVLateral = CV - CVNormal * ContactNormal;
 						FReal CVLateralMag = CVLateral.Size();
-						if (CVLateralMag > KINDA_SMALL_NUMBER)
+						if (CVLateralMag > UE_KINDA_SMALL_NUMBER)
 						{
 							FVec3 DirLateral = CVLateral / CVLateralMag;
 							FVec3 LateralImpulseNumerator = -CVLateral * IterationParameters.Dt;
diff --git a/Engine/Source/Runtime/Experimental/Chaos/Private/Chaos/PBDCollisionSpringConstraintsBase.cpp b/Engine/Source/Runtime/Experimental/Chaos/Private/Chaos/PBDCollisionSpringConstraintsBase.cpp
--- a/Engine/Source/Runtime/Experimental/Chaos/Private/Chaos/PBDCollisionSpringConstraintsBase.cpp
+++ b/Engine/Source/Runtime/Experimental/Chaos/Private/Chaos/PBDCollisionSpringConstraintsBase.cpp
@@ -215,7 +215,7 @@ FSolverVec3 FPBDCollisionSpringConstraintsBase::GetDelta(const FSolverParticles&
 		const FSolverVec3 RelativeDisplacementTangent = RelativeDisplacement - RelativeDisplacement.Dot(Normal) * Normal;
 		const FSolverReal RelativeDisplacementTangentLength = RelativeDisplacementTangent.Length();
 		const FSolverReal PositionCorrection = FMath::Min(NormalDelta * FrictionCoefficient, RelativeDisplacementTangentLength);
-		const FSolverReal CorrectionRatio = RelativeDisplacementTangentLength < SMALL_NUMBER ? 0.f : PositionCorrection / RelativeDisplacementTangentLength;
+		const FSolverReal CorrectionRatio = RelativeDisplacementTangentLength < UE_SMALL_NUMBER ? 0.f : PositionCorrection / RelativeDisplacementTangentLength;
 		const FSolverVec3 FrictionDelta = -CorrectionRatio * RelativeDisplacementTangent / CombinedMass;
 		return RepulsionDelta + FrictionDelta;
 	}
diff --git a/Engine/Source/Runtime/Experimental/Chaos/Private/Chaos/PBDEvolution.cpp b/Engine/Source/Runtime/Experimental/Chaos/Private/Chaos/PBDEvolution.cpp
--- a/Engine/Source/Runtime/Experimental/Chaos/Private/Chaos/PBDEvolution.cpp
+++ b/Engine/Source/Runtime/Experimental/Chaos/Private/Chaos/PBDEvolution.cpp
@@ -243,11 +243,11 @@ void FPBDEvolution::PreIterationUpdate(
 		const FSolverReal Damping = FMath::Clamp(MGroupDampings[ParticleGroupId], (FSolverReal)0., (FSolverReal)1.);
 		FSolverReal DampingPowDt;
 		FSolverReal DampingIntegrated;
-		if (Damping > (FSolverReal)1. - (FSolverReal)KINDA_SMALL_NUMBER)
+		if (Damping > (FSolverReal)1. - (FSolverReal)UE_KINDA_SMALL_NUMBER)
 		{
 			DampingIntegrated = DampingPowDt = (FSolverReal)0.;
 		}
-		else if (Damping > (FSolverReal)SMALL_NUMBER)
+		else if (Damping > (FSolverReal)UE_SMALL_NUMBER)
 		{
 			const FSolverReal LogValueByFrequency = FMath::Loge((FSolverReal)1. - Damping) * DampingFrequency;
 
diff --git a/Engine/Source/Runtime/Experimental/Chaos/Private/Chaos/PBDJointConstraintUtilities.cpp b/Engine/Source/Runtime/Experimental/Chaos/Private/Chaos/PBDJointConstraintUtilities.cpp
--- a/Engine/Source/Runtime/Experimental/Chaos/Private/Chaos/PBDJointConstraintUtilities.cpp
+++ b/Engine/Source/Runtime/Experimental/Chaos/Private/Chaos/PBDJointConstraintUtilities.cpp
@@ -105,7 +105,7 @@ namespace Chaos
 
 		const FVec3 DX = X1 - X0;
 		const FReal DXLen = DX.Size();
-		if (DXLen > KINDA_SMALL_NUMBER)
+		if (DXLen > UE_KINDA_SMALL_NUMBER)
 		{
 			Axis = DX / DXLen;
 			Delta = DXLen;
@@ -136,7 +136,7 @@ namespace Chaos
 
 		DX = DX - CylinderDelta * CylinderAxis;
 		const FReal DXLen = DX.Size();
-		if (DXLen > KINDA_SMALL_NUMBER)
+		if (DXLen > UE_KINDA_SMALL_NUMBER)
 		{
 			RadialAxis = DX / DXLen;
 			RadialDelta = DXLen;
@@ -178,9 +178,9 @@ namespace Chaos
 		FRotation3 Twist = InTwist.GetNormalized();
 		ensure(FMath::Abs(Twist.W) <= 1.0f);
 		FReal Angle = Twist.GetAngle();
-		if (Angle > PI)
+		if (Angle > UE_PI)
 		{
-			Angle = Angle - (FReal)2 * PI;
+			Angle = Angle - (FReal)2 * UE_PI;
 		}
 		if (Twist.X < 0.0f)
 		{
@@ -216,9 +216,9 @@ namespace Chaos
 		FPBDJointUtilities::DecomposeSwingTwistLocal(R0, R1, R01Swing, R01Twist);
 
 		R01Swing.ToAxisAndAngleSafe(AxisLocal, Angle, FJointConstants::Swing1Axis(), AngleTolerance);
-		if (Angle > PI)
+		if (Angle > UE_PI)
 		{
-			Angle = Angle - (FReal)2 * PI;
+			Angle = Angle - (FReal)2 * UE_PI;
 		}
 	}
 
@@ -314,7 +314,7 @@ namespace Chaos
 		FVec3 Swing0 = R0 * FJointConstants::OtherSwingAxis(SwingConstraintIndex);
 		Axis = FVec3::CrossProduct(Swing0, Twist1);
 		Angle = 0.0f;
-		if (Utilities::NormalizeSafe(Axis, KINDA_SMALL_NUMBER))
+		if (Utilities::NormalizeSafe(Axis, UE_KINDA_SMALL_NUMBER))
 		{
 			FReal SwingTwistDot = FVec3::DotProduct(Swing0, Twist1);
 			Angle = FMath::Asin(FMath::Clamp(-SwingTwistDot, FReal(-1), FReal(1)));
@@ -357,9 +357,9 @@ namespace Chaos
 		Axis2 = 0.5f * (V0 * V1.Z + V1 * V0.Z + FVec3(C.Y, -C.X, D));
 
 		// Handle degenerate case of 180 deg swing
-		if (FMath::Abs(D0 + D1) < SMALL_NUMBER)
+		if (FMath::Abs(D0 + D1) < UE_SMALL_NUMBER)
 		{
-			const FReal Epsilon = SMALL_NUMBER;
+			const FReal Epsilon = UE_SMALL_NUMBER;
 			Axis0.X += Epsilon;
 			Axis1.Y += Epsilon;
 			Axis2.Z += Epsilon;
@@ -381,7 +381,7 @@ namespace Chaos
 
 		// Elliptical swing limit
 		// @todo(ccaulfield): do elliptical constraints properly (axis is still for circular limit)
-		if (!FMath::IsNearlyEqual(Swing1Limit, Swing2Limit, (FReal)KINDA_SMALL_NUMBER))
+		if (!FMath::IsNearlyEqual(Swing1Limit, Swing2Limit, (FReal)UE_KINDA_SMALL_NUMBER))
 		{
 			// Map swing axis to ellipse and calculate limit for this swing axis
 			const FReal DotSwing1 = FMath::Abs(FVec3::DotProduct(SwingAxisLocal, FJointConstants::Swing1Axis()));
@@ -744,7 +744,7 @@ namespace Chaos
 		{
 			return FVec3(0, 0, 0);
 		}
-		else if (CXLen > SMALL_NUMBER)
+		else if (CXLen > UE_SMALL_NUMBER)
 		{
 			FVec3 Dir = CX / CXLen;
 			return CX - Radius * Dir;
@@ -766,7 +766,7 @@ namespace Chaos
 		{
 			CXPlane = FVec3(0, 0, 0);
 		}
-		else if (CXPlaneLen > KINDA_SMALL_NUMBER)
+		else if (CXPlaneLen > UE_KINDA_SMALL_NUMBER)
 		{
 			FVec3 Dir = CXPlane / CXPlaneLen;
 			CXPlane = CXPlane - Limit * Dir;
diff --git a/Engine/Source/Runtime/Experimental/Chaos/Private/Chaos/PBDJointConstraints.cpp b/Engine/Source/Runtime/Experimental/Chaos/Private/Chaos/PBDJointConstraints.cpp
--- a/Engine/Source/Runtime/Experimental/Chaos/Private/Chaos/PBDJointConstraints.cpp
+++ b/Engine/Source/Runtime/Experimental/Chaos/Private/Chaos/PBDJointConstraints.cpp
@@ -968,7 +968,7 @@ namespace Chaos
 				{
 					FPBDJointCachedSolver& Solver = CachedConstraintSolvers[ConstraintIndex];
 					// NOTE: LinearImpulse/AngularImpulse in the solver are not really impulses - they are mass-weighted position/rotation delta, or (impulse x dt).
-					if (Dt > SMALL_NUMBER)
+					if (Dt > UE_SMALL_NUMBER)
 					{
 						if (Index0 == 0)
 						{
@@ -997,7 +997,7 @@ namespace Chaos
 				{
 					FPBDJointSolver& Solver = ConstraintSolvers[ConstraintIndex];
 					// NOTE: LinearImpulse/AngularImpulse in the solver are not really impulses - they are mass-weighted position/rotation delta, or (impulse x dt).
-					if (Dt > SMALL_NUMBER)
+					if (Dt > UE_SMALL_NUMBER)
 					{
 						if (Index0 == 0)
 						{
diff --git a/Engine/Source/Runtime/Experimental/Chaos/Private/Chaos/PBDRigidClusteringAlgo.cpp b/Engine/Source/Runtime/Experimental/Chaos/Private/Chaos/PBDRigidClusteringAlgo.cpp
--- a/Engine/Source/Runtime/Experimental/Chaos/Private/Chaos/PBDRigidClusteringAlgo.cpp
+++ b/Engine/Source/Runtime/Experimental/Chaos/Private/Chaos/PBDRigidClusteringAlgo.cpp
@@ -109,14 +109,14 @@ namespace Chaos
 		for (int32 i = 0; i < 3; i++)
 		{
 			const FVec3& InertiaTensor = Parent->I();
-			if (InertiaTensor[i] < SMALL_NUMBER)
+			if (InertiaTensor[i] < UE_SMALL_NUMBER)
 			{
 				Parent->SetI(TVec3<FRealSingle>(1.f, 1.f, 1.f));
 				break;
 			}
 		}
 
-		if (!ensure(bHasChild) || !ensure(Parent->M() > SMALL_NUMBER))
+		if (!ensure(bHasChild) || !ensure(Parent->M() > UE_SMALL_NUMBER))
 		{
 			Parent->M() = 1.0;
 			Parent->X() = FVec3(0);
@@ -134,7 +134,7 @@ namespace Chaos
 			return;
 		}
 
-		check(Parent->M() > SMALL_NUMBER);
+		check(Parent->M() > UE_SMALL_NUMBER);
 
 		Parent->X() /= Parent->M();
 		Parent->V() /= Parent->M();
@@ -179,7 +179,7 @@ namespace Chaos
 		{
 			for (int32 i = 0; i < 3; i++)
 			{
-				if (ParentInertia.GetColumn(i)[i] < SMALL_NUMBER)
+				if (ParentInertia.GetColumn(i)[i] < UE_SMALL_NUMBER)
 				{
 					ParentInertia = FMatrix33(1,1,1);
 					break;
diff --git a/Engine/Source/Runtime/Experimental/Chaos/Private/Chaos/PBDRigidsEvolutionGBF.cpp b/Engine/Source/Runtime/Experimental/Chaos/Private/Chaos/PBDRigidsEvolutionGBF.cpp
--- a/Engine/Source/Runtime/Experimental/Chaos/Private/Chaos/PBDRigidsEvolutionGBF.cpp
+++ b/Engine/Source/Runtime/Experimental/Chaos/Private/Chaos/PBDRigidsEvolutionGBF.cpp
@@ -759,7 +759,7 @@ FPBDRigidsEvolutionGBF::FPBDRigidsEvolutionGBF(FPBDRigidsSOAs& InParticles,THand
 	{
 		ParticlesInput.ParallelFor([&](auto& Particle, int32 Index)
 		{
-			if (Dt > SMALL_NUMBER)
+			if (Dt > UE_SMALL_NUMBER)
 			{
 				const FReal SmoothRate = FMath::Clamp(SmoothedPositionLerpRate, 0.0f, 1.0f);
 				const FVec3 VImp = FVec3::CalculateVelocity(Particle.X(), Particle.P(), Dt);
diff --git a/Engine/Source/Runtime/Experimental/Chaos/Private/Chaos/PBDTriangleMeshCollisions.cpp b/Engine/Source/Runtime/Experimental/Chaos/Private/Chaos/PBDTriangleMeshCollisions.cpp
--- a/Engine/Source/Runtime/Experimental/Chaos/Private/Chaos/PBDTriangleMeshCollisions.cpp
+++ b/Engine/Source/Runtime/Experimental/Chaos/Private/Chaos/PBDTriangleMeshCollisions.cpp
@@ -1028,7 +1028,7 @@ namespace ContourMinimization
 		const FSolverVec3& EdgePosition1 = Particles.X(ContourIntersection.EdgeVertices[1]);
 		const FSolverVec3 EdgeDir = EdgePosition1 - EdgePosition0; // Unnormalized
 		const FSolverReal EdgeDirDotNormal = FSolverVec3::DotProduct(EdgeDir, EdgeFaceIntersection.FaceNormal);
-		if (FMath::Abs(EdgeDirDotNormal) >= SMALL_NUMBER)
+		if (FMath::Abs(EdgeDirDotNormal) >= UE_SMALL_NUMBER)
 		{
 			const FSolverVec3 NOverEDotN = EdgeFaceIntersection.FaceNormal * ((FSolverReal)1.f / EdgeDirDotNormal);
 
diff --git a/Engine/Source/Runtime/Experimental/Chaos/Private/Chaos/PerParticleDampVelocity.cpp b/Engine/Source/Runtime/Experimental/Chaos/Private/Chaos/PerParticleDampVelocity.cpp
--- a/Engine/Source/Runtime/Experimental/Chaos/Private/Chaos/PerParticleDampVelocity.cpp
+++ b/Engine/Source/Runtime/Experimental/Chaos/Private/Chaos/PerParticleDampVelocity.cpp
@@ -78,7 +78,7 @@ void FPerParticleDampVelocity::UpdatePositionBasedState(const FSolverParticles&
 		}
 
 		const FSolverReal Det = I.Determinant();
-		Omega = Det < (FSolverReal)SMALL_NUMBER || !FMath::IsFinite(Det) ?
+		Omega = Det < (FSolverReal)UE_SMALL_NUMBER || !FMath::IsFinite(Det) ?
 			FSolverVec3(0.) :
 #if COMPILE_WITHOUT_UNREAL_SUPPORT
 			FSolverRigidTransform3(I).InverseTransformVector(L);
diff --git a/Engine/Source/Runtime/Experimental/Chaos/Private/Chaos/PerParticlePBDCollisionConstraint.cpp b/Engine/Source/Runtime/Experimental/Chaos/Private/Chaos/PerParticlePBDCollisionConstraint.cpp
--- a/Engine/Source/Runtime/Experimental/Chaos/Private/Chaos/PerParticlePBDCollisionConstraint.cpp
+++ b/Engine/Source/Runtime/Experimental/Chaos/Private/Chaos/PerParticlePBDCollisionConstraint.cpp
@@ -157,7 +157,7 @@ void FPerParticlePBDCollisionConstraint::ApplyHelperISPC(FSolverParticles& InPar
 
 	if (Chaos_PerParticleCollision_ISPC_FastFriction)
 	{
-		if (PerGroupFriction > KINDA_SMALL_NUMBER)  // Fast friction
+		if (PerGroupFriction > UE_KINDA_SMALL_NUMBER)  // Fast friction
 		{
 			PhysicsParallelFor(NumBatches, [this, &InParticles, Dt, Offset, Range, DynamicGroupId, PerGroupFriction, PerGroupThickness](int32 BatchNumber)
 			{
diff --git a/Engine/Source/Runtime/Experimental/Chaos/Private/Chaos/TriangleMesh.cpp b/Engine/Source/Runtime/Experimental/Chaos/Private/Chaos/TriangleMesh.cpp
--- a/Engine/Source/Runtime/Experimental/Chaos/Private/Chaos/TriangleMesh.cpp
+++ b/Engine/Source/Runtime/Experimental/Chaos/Private/Chaos/TriangleMesh.cpp
@@ -267,7 +267,7 @@ void FTriangleMesh::GetFaceNormals(TArray<TVec3<T>>& Normals, const TConstArrayV
 			const TVec3<T> p20 = Points[Tri[2]] - Points[Tri[0]];
 			const TVec3<T> Cross = TVec3<T>::CrossProduct(p20, p10);
 			const T Size2 = Cross.SizeSquared();
-			if (Size2 < SMALL_NUMBER)
+			if (Size2 < UE_SMALL_NUMBER)
 			{
 				//particles should not be coincident by the time they get here. Return empty to signal problem to caller
 				ensure(false);
@@ -825,7 +825,7 @@ TMap<int32, int32> FTriangleMesh::FindCoincidentVertexRemappings(
 	}
 
 	// Return early if all points are coincident
-	if (LocalBBox.Extents().Max() < KINDA_SMALL_NUMBER)
+	if (LocalBBox.Extents().Max() < UE_KINDA_SMALL_NUMBER)
 	{
 		int32 First = INDEX_NONE;
 		for (const int32 Pt : TestIndices)
@@ -1170,7 +1170,7 @@ TArray<int32> FTriangleMesh::GetVertexImportanceOrdering(
 		int32 FarthestIdx = -1; // Index of point farthest from center
 		FReal FarthestDistSq = 0; // Squared distance of farthest point from center
 		int32 MaxCurveIdx = -1; // Index of point with largest curvature
-		FReal MaxCurve = (FReal)-MAX_FLT; // Largest curvature value
+		FReal MaxCurve = (FReal)-UE_MAX_FLT; // Largest curvature value
 	};
 	// Update FBestPointData with a new point; return true if it's the first point in the cell, false otherwise
 	auto UpdateBestPtData = [&Dist, &PointCurvatures, &Offset, &PointOrder](FBestPtData& BestData, int32 OrderIdx) -> bool
diff --git a/Engine/Source/Runtime/Experimental/Chaos/Private/Chaos/TriangleMeshImplicitObject.cpp b/Engine/Source/Runtime/Experimental/Chaos/Private/Chaos/TriangleMeshImplicitObject.cpp
--- a/Engine/Source/Runtime/Experimental/Chaos/Private/Chaos/TriangleMeshImplicitObject.cpp
+++ b/Engine/Source/Runtime/Experimental/Chaos/Private/Chaos/TriangleMeshImplicitObject.cpp
@@ -89,7 +89,7 @@ const auto MakeTriangleHelper(const QueryGeomType& Geom)
 template <typename QueryGeomType>
 const auto MakeTriangleHelper(const TImplicitObjectScaled<QueryGeomType>& ScaledGeom)
 {
-	if (FVec3::IsNearlyEqual(ScaledGeom.GetScale(), FVec3(1), SMALL_NUMBER))
+	if (FVec3::IsNearlyEqual(ScaledGeom.GetScale(), FVec3(1), UE_SMALL_NUMBER))
 	{
 		return FPBDCollisionConstraint::MakeTriangle(ScaledGeom.GetUnscaledObject());
 	}
@@ -1362,7 +1362,7 @@ FVec3 FTriangleMeshImplicitObject::GetFaceNormal(const int32 FaceIdx) const
 			const ParticleVecType AC = C - A;
 			ParticleVecType Normal = ParticleVecType::CrossProduct(AB, AC);
 			
-			if(Normal.SafeNormalize() < SMALL_NUMBER)
+			if(Normal.SafeNormalize() < UE_SMALL_NUMBER)
 			{
 				UE_LOG(LogChaos, Warning, TEXT("Degenerate triangle %d: (%f %f %f) (%f %f %f) (%f %f %f)"), FaceIdx, A.X, A.Y, A.Z, B.X, B.Y, B.Z, C.X, C.Y, C.Z);
 				ensure(false);
diff --git a/Engine/Source/Runtime/Experimental/Chaos/Private/Chaos/Utilities.cpp b/Engine/Source/Runtime/Experimental/Chaos/Private/Chaos/Utilities.cpp
--- a/Engine/Source/Runtime/Experimental/Chaos/Private/Chaos/Utilities.cpp
+++ b/Engine/Source/Runtime/Experimental/Chaos/Private/Chaos/Utilities.cpp
@@ -63,8 +63,8 @@ namespace Chaos
 		const ::Chaos::TVector<FReal, 3>& FinalVector)
 	{
 		typedef Chaos::TVector<FReal, 3> TV;
-		checkSlow(FMath::Abs(InitialVector.Size() - 1.0) < KINDA_SMALL_NUMBER);
-		checkSlow(FMath::Abs(FinalVector.Size() - 1.0) < KINDA_SMALL_NUMBER);
+		checkSlow(FMath::Abs(InitialVector.Size() - 1.0) < UE_KINDA_SMALL_NUMBER);
+		checkSlow(FMath::Abs(FinalVector.Size() - 1.0) < UE_KINDA_SMALL_NUMBER);
 
 		const double CosTheta = FMath::Clamp<FReal>(TV::DotProduct(InitialVector, FinalVector), -1., 1.);
 
@@ -91,7 +91,7 @@ namespace Chaos
 	 */
 	TVector<FReal, 3> TRotation<FReal, 3>::CalculateAngularVelocity1(const TRotation<FReal, 3>& InR0, const TRotation<FReal, 3>& InR1, const FReal InDt)
 	{
-		if (!ensure(InDt > SMALL_NUMBER))
+		if (!ensure(InDt > UE_SMALL_NUMBER))
 		{
 			return TVector<FReal, 3>(0);
 		}
@@ -120,7 +120,7 @@ namespace Chaos
 	{
 		// @todo(ccaulfield): ToAxisAndAngle starts to return increasingly random, non-normalized axes for very small angles. This 
 		// underestimates the angular velocity magnitude and randomizes direction.
-		if (!ensure(InDt > SMALL_NUMBER))
+		if (!ensure(InDt > UE_SMALL_NUMBER))
 		{
 			return TVector<FReal, 3>(0);
 		}
diff --git a/Engine/Source/Runtime/Experimental/Chaos/Private/Field/FieldSystemNodes.cpp b/Engine/Source/Runtime/Experimental/Chaos/Private/Field/FieldSystemNodes.cpp
--- a/Engine/Source/Runtime/Experimental/Chaos/Private/Field/FieldSystemNodes.cpp
+++ b/Engine/Source/Runtime/Experimental/Chaos/Private/Field/FieldSystemNodes.cpp
@@ -226,7 +226,7 @@ void FWaveScalar::Evaluate(FFieldContext& Context, TFieldArrayView<float>& Resul
 {
 	const float Velocity = (Period != 0.0f ) ? Wavelength / Period : 0.0f;
 
-	const float Wavenumber = (Wavelength != 0.0f ) ? 2.0f * PI / Wavelength : 0.0f;
+	const float Wavenumber = (Wavelength != 0.0f ) ? 2.0f * UE_PI / Wavelength : 0.0f;
 	const Chaos::FReal DeltaTime = FMath::Max(Context.TimeSeconds, 0.0f);
 	const Chaos::FReal Radius = Wavelength * DeltaTime / Period;
 	const Chaos::FReal Decay = DeltaTime / Period;
@@ -456,7 +456,7 @@ void FPlaneFalloff::Evaluator(const FFieldContext& Context, const FPlane& Plane,
 				Results[Index.Result] = Default;
 				const FPlane::FReal Delta = Plane.PlaneDot(Context.SamplePositions[Index.Sample]);
 
-				if (Delta < -SMALL_NUMBER && Delta > -Distance)
+				if (Delta < -UE_SMALL_NUMBER && Delta > -Distance)
 				{
 					const float Function = float(1.0f + Delta / Distance);		// LWC_TODO: Precision loss
 					Results[Index.Result] = EvalFalloffFunction<FalloffType>(MinRange, DeltaRange, Magnitude, Function);
diff --git a/Engine/Source/Runtime/Experimental/Chaos/Private/GeometryCollection/GeometryCollectionCollisionStructureManager.cpp b/Engine/Source/Runtime/Experimental/Chaos/Private/GeometryCollection/GeometryCollectionCollisionStructureManager.cpp
--- a/Engine/Source/Runtime/Experimental/Chaos/Private/GeometryCollection/GeometryCollectionCollisionStructureManager.cpp
+++ b/Engine/Source/Runtime/Experimental/Chaos/Private/GeometryCollection/GeometryCollectionCollisionStructureManager.cpp
@@ -306,7 +306,7 @@ FCollisionStructureManager::NewImplicitCapsule(
 	const float CollisionObjectReduction,
 	const ECollisionTypeEnum CollisionType)
 {
-	if (Length < SMALL_NUMBER)
+	if (Length < UE_SMALL_NUMBER)
 	{
 		// make a more optimized shape : sphere
 		return FCollisionStructureManager::NewImplicitSphere(Radius, CollisionObjectReduction, CollisionType);
@@ -348,7 +348,7 @@ FCollisionStructureManager::NewImplicitCapsule(
 		HalfLengthVector = Chaos::FVec3(0, 0, (ZExtent - Radius));
 	}
 
-	if (HalfLengthVector.Size() < SMALL_NUMBER)
+	if (HalfLengthVector.Size() < UE_SMALL_NUMBER)
 	{
 		// make a more optimized shape : sphere
 		return FCollisionStructureManager::NewImplicitSphere(Radius, CollisionObjectReduction, CollisionType);
@@ -374,7 +374,7 @@ FCollisionStructureManager::NewImplicitLevelset(
 	const ECollisionTypeEnum CollisionType)
 {
 	FVector HalfExtents = CollisionBounds.GetExtent();
-	if (HalfExtents.GetAbsMin() < KINDA_SMALL_NUMBER)
+	if (HalfExtents.GetAbsMin() < UE_KINDA_SMALL_NUMBER)
 	{
 		return nullptr;
 	}
diff --git a/Engine/Source/Runtime/Experimental/Chaos/Private/GeometryCollection/GeometryCollectionConvexUtility.cpp b/Engine/Source/Runtime/Experimental/Chaos/Private/GeometryCollection/GeometryCollectionConvexUtility.cpp
--- a/Engine/Source/Runtime/Experimental/Chaos/Private/GeometryCollection/GeometryCollectionConvexUtility.cpp
+++ b/Engine/Source/Runtime/Experimental/Chaos/Private/GeometryCollection/GeometryCollectionConvexUtility.cpp
@@ -188,11 +188,11 @@ Chaos::FConvex MakeHull(const TArray<Chaos::FConvex::FVec3Type>& Pts, double Sim
 		TArray<Chaos::FConvex::FVec3Type> FilteredPts;
 		FilterHullPoints(Pts, FilteredPts, SimplificationDistanceThreshold);
 		
-		return Chaos::FConvex(FilteredPts, KINDA_SMALL_NUMBER);
+		return Chaos::FConvex(FilteredPts, UE_KINDA_SMALL_NUMBER);
 	}
 	else
 	{
-		return Chaos::FConvex(Pts, KINDA_SMALL_NUMBER);
+		return Chaos::FConvex(Pts, UE_KINDA_SMALL_NUMBER);
 	}
 }
 
@@ -672,7 +672,7 @@ void CreateNonoverlappingConvexHulls(
 					if (JoinedHullPts.Num() > 0)
 					{
 						FilterHullPoints(JoinedHullPts, SimplificationDistanceThreshold);
-						TUniquePtr<Chaos::FConvex> Hull = MakeUnique<Chaos::FConvex>(JoinedHullPts, KINDA_SMALL_NUMBER);
+						TUniquePtr<Chaos::FConvex> Hull = MakeUnique<Chaos::FConvex>(JoinedHullPts, UE_KINDA_SMALL_NUMBER);
 						bool bIsTooBig = false;
 						if (Volume)
 						{
@@ -961,7 +961,7 @@ void HullsFromGeometry(
 					HullPts.Add(FVector(Transform.TransformPosition((FVector)P)));
 				}
 				// Do not simplify hulls when we're just trying to transform them
-				int32 NewConvexIdx = Convexes.Add(MakeUnique<Chaos::FConvex>(HullPts, KINDA_SMALL_NUMBER));
+				int32 NewConvexIdx = Convexes.Add(MakeUnique<Chaos::FConvex>(HullPts, UE_KINDA_SMALL_NUMBER));
 				TransformToConvexIndices[Idx].Add(NewConvexIdx);
 			}
 		}
@@ -978,7 +978,7 @@ void HullsFromGeometry(
 			}
 			ensure(HullPts.Num() > 0);
 			FilterHullPoints(HullPts, SimplificationDistanceThreshold);
-			int32 ConvexIdx = Convexes.Add(MakeUnique<Chaos::FConvex>(HullPts, KINDA_SMALL_NUMBER));
+			int32 ConvexIdx = Convexes.Add(MakeUnique<Chaos::FConvex>(HullPts, UE_KINDA_SMALL_NUMBER));
 			if (Convexes[ConvexIdx]->NumVertices() == 0 && HullPts.Num() > 0)
 			{
 				// if we've failed to make a convex hull, add a tiny bounding box to try to give it a hull anyway
@@ -1008,7 +1008,7 @@ void HullsFromGeometry(
 					HullPts.Add({ Max.X, Min.Y, Min.Z });
 					HullPts.Add({ Min.X, Max.Y, Min.Z });
 					// note: Do not use SimplificationDistanceThreshold for this fixed tiny hull
-					*Convexes[ConvexIdx] = Chaos::FConvex(HullPts, KINDA_SMALL_NUMBER);
+					*Convexes[ConvexIdx] = Chaos::FConvex(HullPts, UE_KINDA_SMALL_NUMBER);
 				}
 			}
 			TransformToConvexIndices[Idx].Add(ConvexIdx);
@@ -1034,7 +1034,7 @@ void TransformHullsToLocal(
 				HullPts.Add(FVector(Transform.InverseTransformPosition((FVector)P)));
 			}
 			// Do not simplify hulls when we're just trying to transform them
-			*Convexes[ConvexIdx] = Chaos::FConvex(HullPts, KINDA_SMALL_NUMBER);
+			*Convexes[ConvexIdx] = Chaos::FConvex(HullPts, UE_KINDA_SMALL_NUMBER);
 		}
 	}
 }
@@ -1063,7 +1063,7 @@ bool CopyHulls(
 			HullPts.Add(OutTransform.InverseTransformPosition(InTransform.TransformPosition(FVector(P))));
 		}
 		// Do not simplify hulls when we're just trying to transform them
-		int32 OutIdx = OutConvexes.Add(MakeUnique<Chaos::FConvex>(HullPts, KINDA_SMALL_NUMBER));
+		int32 OutIdx = OutConvexes.Add(MakeUnique<Chaos::FConvex>(HullPts, UE_KINDA_SMALL_NUMBER));
 		OutTransformToConvexIndices[OutBone].Add(OutIdx);
 	}
 
diff --git a/Engine/Source/Runtime/Experimental/Chaos/Private/GeometryCollection/GeometryCollectionProximityUtility.cpp b/Engine/Source/Runtime/Experimental/Chaos/Private/GeometryCollection/GeometryCollectionProximityUtility.cpp
--- a/Engine/Source/Runtime/Experimental/Chaos/Private/GeometryCollection/GeometryCollectionProximityUtility.cpp
+++ b/Engine/Source/Runtime/Experimental/Chaos/Private/GeometryCollection/GeometryCollectionProximityUtility.cpp
@@ -142,7 +142,7 @@ void FGeometryCollectionProximityUtility::FindContactPairs(int32 Zenith, int32 N
 
 bool FGeometryCollectionProximityUtility::AreNormalsOpposite(const FVector3f& Normal0, const FVector3f& Normal1) const
 {
-	return FVector3f::DotProduct(Normal0, Normal1) < (-1.0f + KINDA_SMALL_NUMBER);
+	return FVector3f::DotProduct(Normal0, Normal1) < (-1.0f + UE_KINDA_SMALL_NUMBER);
 }
 
 bool FGeometryCollectionProximityUtility::AreFacesCoPlanar(int32 Idx0, int32 Idx1) const
@@ -155,7 +155,7 @@ bool FGeometryCollectionProximityUtility::AreFacesCoPlanar(int32 Idx0, int32 Idx
 	FVector3f PlaneOrigin = TransformedVertices[Indices[Idx1].X];
 	FVector3f PlaneNormal = SurfaceNormals[Idx1];
 
-	return FMath::Abs(FVector3f::DotProduct((SamplePoint - PlaneOrigin), PlaneNormal)) < KINDA_SMALL_NUMBER;
+	return FMath::Abs(FVector3f::DotProduct((SamplePoint - PlaneOrigin), PlaneNormal)) < UE_KINDA_SMALL_NUMBER;
 }
 
 bool FGeometryCollectionProximityUtility::DoFacesOverlap(int32 Idx0, int32 Idx1) const
diff --git a/Engine/Source/Runtime/Experimental/Chaos/Private/PBDRigidsSolver.cpp b/Engine/Source/Runtime/Experimental/Chaos/Private/PBDRigidsSolver.cpp
--- a/Engine/Source/Runtime/Experimental/Chaos/Private/PBDRigidsSolver.cpp
+++ b/Engine/Source/Runtime/Experimental/Chaos/Private/PBDRigidsSolver.cpp
@@ -740,7 +740,7 @@ namespace Chaos
 		MLastDt = 0.0f;
 		CurrentFrame = 0;
 		SetMaxDeltaTime_External(1.0f);
-		SetMinDeltaTime_External(SMALL_NUMBER);
+		SetMinDeltaTime_External(UE_SMALL_NUMBER);
 		SetMaxSubSteps_External(1);
 		MEvolution = TUniquePtr<FPBDRigidsEvolution>(new FPBDRigidsEvolution(Particles, SimMaterials, &ContactModifiers, BufferMode == EMultiBufferMode::Single)); 
 
diff --git a/Engine/Source/Runtime/Experimental/Chaos/Private/PhysicsProxy/GeometryCollectionPhysicsProxy.cpp b/Engine/Source/Runtime/Experimental/Chaos/Private/PhysicsProxy/GeometryCollectionPhysicsProxy.cpp
--- a/Engine/Source/Runtime/Experimental/Chaos/Private/PhysicsProxy/GeometryCollectionPhysicsProxy.cpp
+++ b/Engine/Source/Runtime/Experimental/Chaos/Private/PhysicsProxy/GeometryCollectionPhysicsProxy.cpp
@@ -1897,7 +1897,7 @@ bool FGeometryCollectionPhysicsProxy::PullFromPhysicsState(const Chaos::FDirtyGe
 					TManagedArray<FVector3f>* AngularVelocity = DynamicCollection.FindAttributeTyped<FVector3f>("AngularVelocity", FTransformCollection::TransformGroup);
 					check(AngularVelocity);
 					FVector DiffX = ParticleToWorld.GetTranslation() - GTParticle->X();
-					FVector DiffR = (ParticleToWorld.GetRotation().Euler() - GTParticle->R().Euler()) * (PI / 180.0f);
+					FVector DiffR = (ParticleToWorld.GetRotation().Euler() - GTParticle->R().Euler()) * (UE_PI / 180.0f);
 
 					(*LinearVelocity)[TransformGroupIndex] = FVector3f(DiffX / TargetResults.SolverDt);
 					(*AngularVelocity)[TransformGroupIndex] = FVector3f(DiffR / TargetResults.SolverDt);
@@ -1950,7 +1950,7 @@ void FGeometryCollectionPhysicsProxy::InitializeSharedCollisionStructures(
 	const float MaxBoundsExtents = SharedParams.MaximumBoundingExtentClamp;
 	const float MinVolume = SharedParams.MinimumVolumeClamp();
 	const float MaxVolume = SharedParams.MaximumVolumeClamp();
-	const float MinMass = FMath::Max(SMALL_NUMBER, SharedParams.MaximumMassClamp);
+	const float MinMass = FMath::Max(UE_SMALL_NUMBER, SharedParams.MaximumMassClamp);
 	const float MaxMass = SharedParams.MinimumMassClamp;
 
 
diff --git a/Engine/Source/Runtime/Experimental/Chaos/Private/PhysicsProxy/SkeletalMeshPhysicsProxy.cpp b/Engine/Source/Runtime/Experimental/Chaos/Private/PhysicsProxy/SkeletalMeshPhysicsProxy.cpp
--- a/Engine/Source/Runtime/Experimental/Chaos/Private/PhysicsProxy/SkeletalMeshPhysicsProxy.cpp
+++ b/Engine/Source/Runtime/Experimental/Chaos/Private/PhysicsProxy/SkeletalMeshPhysicsProxy.cpp
@@ -234,7 +234,7 @@ void FSkeletalMeshPhysicsProxy::CreateRigidBodyCallback(FParticlesType& Particle
 			Particles.Q(RigidBodyId) = Particles.R(RigidBodyId);
 
 			Particles.M(RigidBodyId) = FMath::Clamp(TotalMass, (Chaos::FReal)Parameters.MinMass, (Chaos::FReal)Parameters.MaxMass);
-			Particles.InvM(RigidBodyId) = Particles.M(RigidBodyId) < KINDA_SMALL_NUMBER ?  BIG_NUMBER/2 : 1.f / Particles.M(RigidBodyId);
+			Particles.InvM(RigidBodyId) = Particles.M(RigidBodyId) < UE_KINDA_SMALL_NUMBER ? UE_BIG_NUMBER/2 : 1.f / Particles.M(RigidBodyId);
 			Particles.I(RigidBodyId) = Chaos::TVec3<Chaos::FRealSingle>((Chaos::FRealSingle)MassProperties.InertiaTensor.M[0][0], (Chaos::FRealSingle)MassProperties.InertiaTensor.M[1][1], (Chaos::FRealSingle)MassProperties.InertiaTensor.M[2][2]);
 			Particles.InvI(RigidBodyId) = Chaos::TVec3<Chaos::FRealSingle>(1.f / (Chaos::FRealSingle)MassProperties.InertiaTensor.M[0][0], 1.f / (Chaos::FRealSingle)MassProperties.InertiaTensor.M[1][1], 1.f / (Chaos::FRealSingle)MassProperties.InertiaTensor.M[2][2]);
 
@@ -496,7 +496,7 @@ void FSkeletalMeshPhysicsProxy::CaptureInputs(const float Dt, const FInputFunc&
 			{
 				const FTransform& BoneWorldTransform = *BoneHierarchy.GetAnimWorldSpaceTransformsForBone(BoneIndex);
 				const FTransform* PrevBoneWorldTransform = BoneHierarchy.GetPrevAnimWorldSpaceTransformForBone(BoneIndex);
-				if ((Dt > SMALL_NUMBER) && PrevBoneWorldTransform)
+				if ((Dt > UE_SMALL_NUMBER) && PrevBoneWorldTransform)
 				{
 					PhysicsInputs.Transforms[TransformIndex] = *PrevBoneWorldTransform;
 					PhysicsInputs.LinearVelocities[TransformIndex] = (BoneWorldTransform.GetTranslation() - PrevBoneWorldTransform->GetTranslation()) / Dt;
diff --git a/Engine/Source/Runtime/Experimental/Chaos/Private/PhysicsProxy/StaticMeshPhysicsObject.cpp b/Engine/Source/Runtime/Experimental/Chaos/Private/PhysicsProxy/StaticMeshPhysicsObject.cpp
--- a/Engine/Source/Runtime/Experimental/Chaos/Private/PhysicsProxy/StaticMeshPhysicsObject.cpp
+++ b/Engine/Source/Runtime/Experimental/Chaos/Private/PhysicsProxy/StaticMeshPhysicsObject.cpp
@@ -244,8 +244,8 @@ void FStaticMeshPhysicsProxy::CreateRigidBodyCallback(FParticlesType& Particles)
 		}
 		else
 		{
-			Bounds.Min = Chaos::FVec3(-KINDA_SMALL_NUMBER);
-			Bounds.Max = Chaos::FVec3(KINDA_SMALL_NUMBER);
+			Bounds.Min = Chaos::FVec3(-UE_KINDA_SMALL_NUMBER);
+			Bounds.Max = Chaos::FVec3(UE_KINDA_SMALL_NUMBER);
 			Particles.SetSharedGeometry(RigidBodyId, TSharedPtr<Chaos::FImplicitObject, ESPMode::ThreadSafe>());
 		}
 
@@ -256,7 +256,7 @@ void FStaticMeshPhysicsProxy::CreateRigidBodyCallback(FParticlesType& Particles)
 		Bounds = Bounds.InverseTransformBy(FTransform(CenterOfMass));
 		Bounds.Min *= Scale;
 		Bounds.Max *= Scale;
-		checkSlow((Bounds.Max + Bounds.Min).Size() < KINDA_SMALL_NUMBER);
+		checkSlow((Bounds.Max + Bounds.Min).Size() < UE_KINDA_SMALL_NUMBER);
 
 		Particles.InvM(RigidBodyId) = 0.f;
 		ensure(Parameters.Mass >= 0.f);
diff --git a/Engine/Source/Runtime/Experimental/Chaos/Public/Chaos/PBDAxialSpringConstraintsBase.h b/Engine/Source/Runtime/Experimental/Chaos/Public/Chaos/PBDAxialSpringConstraintsBase.h
--- a/Engine/Source/Runtime/Experimental/Chaos/Public/Chaos/PBDAxialSpringConstraintsBase.h
+++ b/Engine/Source/Runtime/Experimental/Chaos/Public/Chaos/PBDAxialSpringConstraintsBase.h
@@ -54,7 +54,7 @@ protected:
 		const FSolverVec3 P = (P2 - P3) * Barys[ConstraintIndex] + P3;
 		const FSolverVec3 Difference = P1 - P;
 		const FSolverReal Distance = Difference.Size();
-		if (UNLIKELY(Distance <= (FSolverReal)SMALL_NUMBER))
+		if (UNLIKELY(Distance <= (FSolverReal)UE_SMALL_NUMBER))
 		{
 			return FSolverVec3((FSolverReal)0.);
 		}
diff --git a/Engine/Source/Runtime/Experimental/Chaos/Public/Chaos/PBDRigidClusteringCollisionParticleAlgo.h b/Engine/Source/Runtime/Experimental/Chaos/Public/Chaos/PBDRigidClusteringCollisionParticleAlgo.h
--- a/Engine/Source/Runtime/Experimental/Chaos/Public/Chaos/PBDRigidClusteringCollisionParticleAlgo.h
+++ b/Engine/Source/Runtime/Experimental/Chaos/Public/Chaos/PBDRigidClusteringCollisionParticleAlgo.h
@@ -19,7 +19,7 @@ inline TArray<FVec3> CleanCollisionParticles(
 	if (MaxBBoxDim < SnapDistance)
 		return TArray<FVec3>(&Vertices[0], 1);
 
-	BBox.Thicken(FMath::Max(SnapDistance/10, (FReal)(KINDA_SMALL_NUMBER*10))); // 0.001
+	BBox.Thicken(FMath::Max(SnapDistance/10, (FReal)(UE_KINDA_SMALL_NUMBER*10))); // 0.001
 	MaxBBoxDim = BBox.Extents().Max();
 
 	const FVec3 PointsCenter = BBox.Center();
@@ -35,7 +35,7 @@ inline TArray<FVec3> CleanCollisionParticles(
 	Redundant.Reserve(NumPoints); // Excessive, but ensures consistent performance.
 
 	int32 NumCoincident = 0;
-	const int64 Resolution = static_cast<int64>(floor(MaxBBoxDim / FMath::Max(SnapDistance,(FReal)KINDA_SMALL_NUMBER)));
+	const int64 Resolution = static_cast<int64>(floor(MaxBBoxDim / FMath::Max(SnapDistance,(FReal)UE_KINDA_SMALL_NUMBER)));
 	const FReal CellSize = static_cast<FReal>(static_cast<double>(MaxBBoxDim) / static_cast<double>(Resolution));
 	for (int32 i = 0; i < 2; i++)
 	{
diff --git a/Engine/Source/Runtime/Experimental/Chaos/Public/Chaos/PBDSphericalConstraint.h b/Engine/Source/Runtime/Experimental/Chaos/Public/Chaos/PBDSphericalConstraint.h
--- a/Engine/Source/Runtime/Experimental/Chaos/Public/Chaos/PBDSphericalConstraint.h
+++ b/Engine/Source/Runtime/Experimental/Chaos/Public/Chaos/PBDSphericalConstraint.h
@@ -80,7 +80,7 @@ private:
 			const FSolverVec3 CenterToParticle = Particles.P(ParticleIndex) - Center;
 			const FSolverReal DistanceSquared = CenterToParticle.SizeSquared();
 
-			static const FSolverReal DeadZoneSquareRadius = SMALL_NUMBER; // We will not push the particle away in the dead zone
+			static const FSolverReal DeadZoneSquareRadius = UE_SMALL_NUMBER; // We will not push the particle away in the dead zone
 			if (DistanceSquared > FMath::Square(Radius) + DeadZoneSquareRadius)
 			{
 				const FSolverReal Distance = sqrt(DistanceSquared);
@@ -202,7 +202,7 @@ private:
 			const FSolverVec3 CenterToParticle = Particles.P(ParticleIndex) - Center;
 			const FSolverReal DistanceSquared = CenterToParticle.SizeSquared();
 
-			static const FSolverReal DeadZoneSquareRadius = SMALL_NUMBER;
+			static const FSolverReal DeadZoneSquareRadius = UE_SMALL_NUMBER;
 			if (DistanceSquared < DeadZoneSquareRadius)
 			{
 				Particles.P(ParticleIndex) = AnimationPosition - SphereOffsetDistance * AnimationNormal;  // Non legacy version adds radius to the distance
@@ -239,7 +239,7 @@ private:
 			const FSolverVec3 CenterToParticle = Particles.P(ParticleIndex) - Center;
 			const FSolverReal DistanceSquared = CenterToParticle.SizeSquared();
 
-			static const FSolverReal DeadZoneSquareRadius = SMALL_NUMBER;
+			static const FSolverReal DeadZoneSquareRadius = UE_SMALL_NUMBER;
 			if (DistanceSquared < DeadZoneSquareRadius)
 			{
 				Particles.P(ParticleIndex) = AnimationPosition - (SphereOffsetDistance - Radius) * AnimationNormal;  // Legacy version already includes the radius to the distance
diff --git a/Engine/Source/Runtime/Experimental/Chaos/Public/Chaos/PerParticlePBDCCDCollisionConstraint.h b/Engine/Source/Runtime/Experimental/Chaos/Public/Chaos/PerParticlePBDCCDCollisionConstraint.h
--- a/Engine/Source/Runtime/Experimental/Chaos/Public/Chaos/PerParticlePBDCCDCollisionConstraint.h
+++ b/Engine/Source/Runtime/Experimental/Chaos/Public/Chaos/PerParticlePBDCCDCollisionConstraint.h
@@ -102,7 +102,7 @@ private:
 						Normals.Emplace(NormalWorld);
 					}
 					const FSolverVec3 Direction = ContactWorld - Particles.P(Index);
-					const FSolverReal Penetration = FMath::Max((FSolverReal)0., FSolverVec3::DotProduct(NormalWorld, Direction)) + (FSolverReal)THRESH_POINT_ON_PLANE;
+					const FSolverReal Penetration = FMath::Max((FSolverReal)0., FSolverVec3::DotProduct(NormalWorld, Direction)) + (FSolverReal)UE_THRESH_POINT_ON_PLANE;
 
 					Particles.P(Index) += Penetration * NormalWorld;
 
@@ -111,7 +111,7 @@ private:
 					const FSolverVec3 RelativeDisplacement = (Particles.P(Index) - Particles.X(Index)) - (CollisionParticles.V(CollisionIndex) + FSolverVec3::CrossProduct(CollisionParticles.W(CollisionIndex), VectorToPoint)) * Dt;  // This corresponds to the tangential velocity multiplied by dt (friction will drive this to zero if it is high enough)
 					const FSolverVec3 RelativeDisplacementTangent = RelativeDisplacement - FSolverVec3::DotProduct(RelativeDisplacement, NormalWorld) * NormalWorld;  // Project displacement into the tangential plane
 					const FSolverReal RelativeDisplacementTangentLength = RelativeDisplacementTangent.Size();
-					if (RelativeDisplacementTangentLength >= SMALL_NUMBER)
+					if (RelativeDisplacementTangentLength >= UE_SMALL_NUMBER)
 					{
 						const FSolverReal PositionCorrection = FMath::Min<FSolverReal>(Penetration * Friction, RelativeDisplacementTangentLength);
 						const FSolverReal CorrectionRatio = PositionCorrection / RelativeDisplacementTangentLength;
diff --git a/Engine/Source/Runtime/Experimental/Chaos/Public/Chaos/PerParticlePBDCollisionConstraint.h b/Engine/Source/Runtime/Experimental/Chaos/Public/Chaos/PerParticlePBDCollisionConstraint.h
--- a/Engine/Source/Runtime/Experimental/Chaos/Public/Chaos/PerParticlePBDCollisionConstraint.h
+++ b/Engine/Source/Runtime/Experimental/Chaos/Public/Chaos/PerParticlePBDCollisionConstraint.h
@@ -74,7 +74,7 @@ private:
 		const FSolverReal PerGroupFriction = MPerGroupFriction[DynamicGroupId];
 		const FSolverReal PerGroupThickness = MPerGroupThickness[DynamicGroupId];
 
-		if (PerGroupFriction > (FSolverReal)KINDA_SMALL_NUMBER)
+		if (PerGroupFriction > (FSolverReal)UE_KINDA_SMALL_NUMBER)
 		{
 			PhysicsParallelFor(Range - Offset, [this, &Particles, Dt, Offset, DynamicGroupId, PerGroupFriction, PerGroupThickness](int32 i)
 			{
@@ -111,7 +111,7 @@ private:
 							const FSolverVec3 RelativeDisplacement = (Particles.P(Index) - Particles.X(Index)) - (CollisionParticles.V(i) + FSolverVec3::CrossProduct(CollisionParticles.W(i), VectorToPoint)) * Dt; // This corresponds to the tangential velocity multiplied by dt (friction will drive this to zero if it is high enough)
 							const FSolverVec3 RelativeDisplacementTangent = RelativeDisplacement - FSolverVec3::DotProduct(RelativeDisplacement, NormalWorld) * NormalWorld; // Project displacement into the tangential plane
 							const FSolverReal RelativeDisplacementTangentLength = RelativeDisplacementTangent.Size();
-							if (RelativeDisplacementTangentLength >= SMALL_NUMBER)
+							if (RelativeDisplacementTangentLength >= UE_SMALL_NUMBER)
 							{
 								const FSolverReal PositionCorrection = FMath::Min<FSolverReal>(Penetration * PerGroupFriction, RelativeDisplacementTangentLength);
 								const FSolverReal CorrectionRatio = PositionCorrection / RelativeDisplacementTangentLength;
diff --git a/Engine/Source/Runtime/Experimental/Chaos/Public/Chaos/TaperedCapsule.h b/Engine/Source/Runtime/Experimental/Chaos/Public/Chaos/TaperedCapsule.h
--- a/Engine/Source/Runtime/Experimental/Chaos/Public/Chaos/TaperedCapsule.h
+++ b/Engine/Source/Runtime/Experimental/Chaos/Public/Chaos/TaperedCapsule.h
@@ -109,7 +109,7 @@ namespace Chaos
 			const FVec3 OriginToX = x - Origin;
 			const FReal DistanceAlongAxis = FMath::Clamp(FVec3::DotProduct(OriginToX, Axis), (FReal)0.0, Height);
 			const FVec3 ClosestPoint = Origin + Axis * DistanceAlongAxis;
-			const FReal Radius = (Height > SMALL_NUMBER) ? FMath::Lerp(Radius1, Radius2, DistanceAlongAxis / Height) : FMath::Max(Radius1, Radius2);
+			const FReal Radius = (Height > UE_SMALL_NUMBER) ? FMath::Lerp(Radius1, Radius2, DistanceAlongAxis / Height) : FMath::Max(Radius1, Radius2);
 			OutNormal = (x - ClosestPoint);
 			const FReal NormalSize = OutNormal.SafeNormalize();
 			return NormalSize - Radius;
@@ -151,7 +151,7 @@ namespace Chaos
 		FReal GetArea(const bool IncludeEndCaps = true) const { return GetArea(Height, Radius1, Radius2, IncludeEndCaps); }
 		static FReal GetArea(const FReal Height, const FReal Radius1, const FReal Radius2, const bool IncludeEndCaps)
 		{
-			static const FReal TwoPI = PI * 2;
+			static const FReal TwoPI = UE_PI * 2;
 			FReal AreaNoCaps = (FReal)0.0;
 			if (Radius1 == Radius2)
 			{
@@ -160,7 +160,7 @@ namespace Chaos
 			else
 			{
 				const FReal R1_R2 = Radius1 - Radius2;
-				AreaNoCaps = PI * (Radius1 + Radius2) * FMath::Sqrt((R1_R2 * R1_R2) + (Height * Height));
+				AreaNoCaps = UE_PI * (Radius1 + Radius2) * FMath::Sqrt((R1_R2 * R1_R2) + (Height * Height));
 			}
 			if (IncludeEndCaps)
 			{
@@ -194,7 +194,7 @@ namespace Chaos
 
 			const FReal Num1 = (FReal)2.0 * HH * (R1R1 + (FReal)3.0 * R1R2 + (FReal)6.0 * R2R2); // 2H^2 * (R1^2 + 3R1R2 + 6R2^2)
 			const FReal Num2 = (FReal)3.0 * (R1R1 * R1R1 + R1R1 * R1R2 + R1R2 * R1R2 + R1R2 * R2R2 + R2R2 * R2R2); // 3 * (R1^4 + R1^3R2 + R1^2R2^2 + R1R2^3 + R2^4)
-			const FReal Den1 = PI * (R1R1 + R1R2 + R2R2); // PI * (R1^2 + R1R2 + R2^2)
+			const FReal Den1 = UE_PI * (R1R1 + R1R2 + R2R2); // PI * (R1^2 + R1R2 + R2^2)
 
 			const FReal Diag12 = Mass * (Num1 + Num2) / ((FReal)20.0 * Den1);
 			const FReal Diag3 = Mass * Num2 / ((FReal)10.0 * Den1);
@@ -248,12 +248,12 @@ namespace Chaos
 
 		static FReal GetHemisphereVolume(const FReal Radius)
 		{
-			return (FReal)2.0 * PI * (Radius * Radius * Radius) / (FReal)3.0;
+			return (FReal)2.0 * UE_PI * (Radius * Radius * Radius) / (FReal)3.0;
 		}
 
 		static FReal GetTaperedSectionVolume(const FReal Height, const FReal Radius1, const FReal Radius2)
 		{
-			static const FReal PI_OVER_3 = PI / (FReal)3.0;
+			static const FReal PI_OVER_3 = UE_PI / (FReal)3.0;
 			return PI_OVER_3 * Height * (Radius1 * Radius1 + Radius1 * Radius2 + Radius2 * Radius2);
 		}
 
@@ -269,11 +269,11 @@ namespace Chaos
 		    const int32 NumPoints)
 		{
 			if (NumPoints <= 1 ||
-			    (Capsule.GetRadius1() <= KINDA_SMALL_NUMBER &&
-				 Capsule.GetRadius2() <= KINDA_SMALL_NUMBER))
+			    (Capsule.GetRadius1() <= UE_KINDA_SMALL_NUMBER &&
+				 Capsule.GetRadius2() <= UE_KINDA_SMALL_NUMBER))
 			{
 				const int32 Offset = Points.Num();
-				if (Capsule.GetHeight() <= KINDA_SMALL_NUMBER)
+				if (Capsule.GetHeight() <= UE_KINDA_SMALL_NUMBER)
 				{
 					Points.SetNumUninitialized(Offset + 1);
 					Points[Offset] = Capsule.GetCenter();
@@ -328,7 +328,7 @@ namespace Chaos
 			const int32 SpiralSeed = 0)
 		{
 			// Axis should be normalized.
-			checkSlow(FMath::Abs(Axis.Size() - 1.0) < KINDA_SMALL_NUMBER);
+			checkSlow(FMath::Abs(Axis.Size() - 1.0) < UE_KINDA_SMALL_NUMBER);
 
 			const int32 Offset = Points.Num();
 			ComputeGoldenSpiralPointsUnoriented(Points, Radius1, Radius2, Height, NumPoints, SpiralSeed);
@@ -337,7 +337,7 @@ namespace Chaos
 			// along the Z axis.  Transform them to where they should be.
 			const FReal HalfHeight = Height / 2;
 			const FRotation3 Rotation = FRotation3::FromRotatedVector(FVec3(0, 0, 1), Axis);
-			checkSlow(((Origin + Axis * Height) - (Rotation.RotateVector(FVec3(0, 0, Height)) + Origin)).Size() < KINDA_SMALL_NUMBER);
+			checkSlow(((Origin + Axis * Height) - (Rotation.RotateVector(FVec3(0, 0, Height)) + Origin)).Size() < UE_KINDA_SMALL_NUMBER);
 			for (int32 i = Offset; i < Points.Num(); i++)
 			{
 				FVec3& Point = Points[i];
@@ -391,7 +391,7 @@ namespace Chaos
 			const FReal Cap2Area = TSphere<FReal, 3>::GetArea(Radius2) / (FReal)2.;
 			const FReal TaperedSectionArea = FTaperedCapsule::GetArea(Height, Radius1, Radius2, /*IncludeEndCaps*/ false);
 			const FReal AllArea = TaperedSectionArea + Cap1Area + Cap2Area;
-			if (AllArea > KINDA_SMALL_NUMBER)
+			if (AllArea > UE_KINDA_SMALL_NUMBER)
 			{
 				NumPointsEndCap1 = static_cast<int32>(round(Cap1Area / AllArea * static_cast<FReal>(NumPoints)));
 				NumPointsEndCap2 = static_cast<int32>(round(Cap2Area / AllArea * static_cast<FReal>(NumPoints)));
diff --git a/Engine/Source/Runtime/Experimental/Chaos/Public/GeometryCollection/GeometryCollectionUtility.h b/Engine/Source/Runtime/Experimental/Chaos/Public/GeometryCollection/GeometryCollectionUtility.h
--- a/Engine/Source/Runtime/Experimental/Chaos/Public/GeometryCollection/GeometryCollectionUtility.h
+++ b/Engine/Source/Runtime/Experimental/Chaos/Public/GeometryCollection/GeometryCollectionUtility.h
@@ -184,7 +184,7 @@ void GeometryCollection::AttributeTransfer(const FGeometryCollection * FromColle
 	ParallelFor(ToCollection->NumElements(FGeometryCollection::VerticesGroup), [&](int32 ToIndex)
 	{
 		int32 ClosestFromIndex = -1;
-		Chaos::FReal ClosestDist = MAX_FLT;
+		Chaos::FReal ClosestDist = UE_MAX_FLT;
 		for (int32 FromIndex = 0, ni = FromVertex.Num(); FromIndex < ni ; ++FromIndex)
 		{
 			Chaos::FReal CurrDist = FVector3f::DistSquared(FromVertex[FromIndex], ToVertex[ToIndex]);
diff --git a/Engine/Source/Runtime/Experimental/Voronoi/Public/Voronoi/Voronoi.h b/Engine/Source/Runtime/Experimental/Voronoi/Public/Voronoi/Voronoi.h
--- a/Engine/Source/Runtime/Experimental/Voronoi/Public/Voronoi/Voronoi.h
+++ b/Engine/Source/Runtime/Experimental/Voronoi/Public/Voronoi/Voronoi.h
@@ -2,8 +2,8 @@
 #pragma once
 #include "CoreMinimal.h"
 
-bool VORONOI_API VoronoiNeighbors(const TArrayView<const FVector> &Sites, TArray<TArray<int>> &Neighbors, bool bExcludeBounds = true, float SquaredDistSkipPtThreshold = KINDA_SMALL_NUMBER);
-bool VORONOI_API GetVoronoiEdges(const TArrayView<const FVector> &Sites, const FBox& Bounds, TArray<TTuple<FVector, FVector>> &Edges, TArray<int32> &CellMember, float SquaredDistSkipPtThreshold = KINDA_SMALL_NUMBER);
+bool VORONOI_API VoronoiNeighbors(const TArrayView<const FVector> &Sites, TArray<TArray<int>> &Neighbors, bool bExcludeBounds = true, float SquaredDistSkipPtThreshold = UE_KINDA_SMALL_NUMBER);
+bool VORONOI_API GetVoronoiEdges(const TArrayView<const FVector> &Sites, const FBox& Bounds, TArray<TTuple<FVector, FVector>> &Edges, TArray<int32> &CellMember, float SquaredDistSkipPtThreshold = UE_KINDA_SMALL_NUMBER);
 
 // All the info you would typically want about a single cell in the Voronoi diagram, in the format that is easiest to compute
 struct FVoronoiCellInfo
@@ -36,7 +36,7 @@ public:
 	 * @param SquaredDistSkipPtThreshold	A safety threshold to avoid creating invalid cells: sites that are within this distance of an already-added site will not be added
 	 *										(If you know there will be no duplicate sites, can set to zero for faster perf.)
 	 */
-	FVoronoiDiagram(const TArrayView<const FVector>& Sites, float ExtraBoundingSpace, float SquaredDistSkipPtThreshold = KINDA_SMALL_NUMBER);
+	FVoronoiDiagram(const TArrayView<const FVector>& Sites, float ExtraBoundingSpace, float SquaredDistSkipPtThreshold = UE_KINDA_SMALL_NUMBER);
 	/**
 	* @param Sites							Voronoi sites for the diagram
 	* @param Bounds							Bounding box within which to compute the Voronoi diagram
@@ -44,7 +44,7 @@ public:
 	* @param SquaredDistSkipPtThreshold		A safety threshold to avoid creating invalid cells: sites that are within this distance of an already-added site will not be added.
 	*										(If you know there will be no duplicate sites, can set to zero for faster perf.)
 	*/
-	FVoronoiDiagram(const TArrayView<const FVector>& Sites, const FBox &Bounds, float ExtraBoundingSpace, float SquaredDistSkipPtThreshold = KINDA_SMALL_NUMBER);
+	FVoronoiDiagram(const TArrayView<const FVector>& Sites, const FBox &Bounds, float ExtraBoundingSpace, float SquaredDistSkipPtThreshold = UE_KINDA_SMALL_NUMBER);
 	~FVoronoiDiagram();
 
 	int32 Num() const
