Fix missing dependency

diff --git a/Engine/Source/Runtime/Experimental/FieldSystem/Source/FieldSystemEngine/FieldSystemEngine.Build.cs b/Engine/Source/Runtime/Experimental/FieldSystem/Source/FieldSystemEngine/FieldSystemEngine.Build.cs
--- a/Engine/Source/Runtime/Experimental/FieldSystem/Source/FieldSystemEngine/FieldSystemEngine.Build.cs
+++ b/Engine/Source/Runtime/Experimental/FieldSystem/Source/FieldSystemEngine/FieldSystemEngine.Build.cs
@@ -18,6 +18,7 @@ namespace UnrealBuildTool.Rules
                     "RHI",
 					"Chaos",
 					"ChaosSolverEngine",
+					"PhysicsCore",
                 }
                 );
 
