Make the async physics tick GT task high priority so that we still get parallelism for DuringPhysics tasks

diff --git a/Engine/Source/Runtime/Experimental/Chaos/Public/Chaos/Framework/PhysicsSolverBase.h b/Engine/Source/Runtime/Experimental/Chaos/Public/Chaos/Framework/PhysicsSolverBase.h
--- a/Engine/Source/Runtime/Experimental/Chaos/Public/Chaos/Framework/PhysicsSolverBase.h
+++ b/Engine/Source/Runtime/Experimental/Chaos/Public/Chaos/Framework/PhysicsSolverBase.h
@@ -102,7 +102,7 @@ namespace Chaos
 		FPhysicsSolverFrozenGTPreSimCallbacks(FPhysicsSolverBase& InSolver) : Solver(InSolver) {}
 
 		TStatId GetStatId() const { RETURN_QUICK_DECLARE_CYCLE_STAT(FPhysicsSolverFrozenGTPreSimCallbacks, STATGROUP_TaskGraphTasks); }
-		static ENamedThreads::Type GetDesiredThread() { return ENamedThreads::GameThread; }
+		static ENamedThreads::Type GetDesiredThread() { return ENamedThreads::SetTaskPriority(ENamedThreads::GameThread, ENamedThreads::HighTaskPriority); }
 		static ESubsequentsMode::Type GetSubsequentsMode() { return ESubsequentsMode::TrackSubsequents; }
 		void DoTask(ENamedThreads::Type CurrentThread, const FGraphEventRef& MyCompletionGraphEvent) { GTPreSimCallbacks(); }
 		void GTPreSimCallbacks();
