Chaos: Fixed a few control paths that could lead to the dirty grid acceleration in scene queries to get caught inside an infinite loop which was showing up as a write-deadlock on the main thread

diff --git a/Engine/Source/Programs/HeadlessChaos/Private/HeadlessChaosApp.cpp b/Engine/Source/Programs/HeadlessChaos/Private/HeadlessChaosApp.cpp
--- a/Engine/Source/Programs/HeadlessChaos/Private/HeadlessChaosApp.cpp
+++ b/Engine/Source/Programs/HeadlessChaos/Private/HeadlessChaosApp.cpp
@@ -188,6 +188,15 @@ TEST(RaycastTests, Raycast) {
 
 TEST(SweepTests, Sweep) {
 	ChaosTest::CapsuleSweepAgainstTriMeshReal();
+	
+	SUCCEED();
+}
+
+TEST(SweepTests, DISABLED_LargeSweep)
+{
+	ChaosTest::GJKLargeDistanceCapsuleSweep();
+
+	SUCCEED();
 }
 
 TEST(MostOpposingTests, MostOpposing) {
diff --git a/Engine/Source/Programs/HeadlessChaos/Private/HeadlessChaosTestGJK.cpp b/Engine/Source/Programs/HeadlessChaos/Private/HeadlessChaosTestGJK.cpp
--- a/Engine/Source/Programs/HeadlessChaos/Private/HeadlessChaosTestGJK.cpp
+++ b/Engine/Source/Programs/HeadlessChaos/Private/HeadlessChaosTestGJK.cpp
@@ -1571,6 +1571,35 @@ namespace ChaosTest
 		
 	}
 
+	// Tests a case where we have a reasonable query but the target shape is a very large distance away.
+	// This should result in a miss but currently doesn't - and gives an OutTime that is infinite.
+	// Detected when querying global payload object in the SQ system where we test each object without
+	// considering its bounds.
+	void GJKLargeDistanceCapsuleSweep()
+	{
+		// Data from repro case
+		Chaos::TBox<Chaos::FReal, 3> A({ -24.011219501495361, -7.4698066711425781, -0.83472084999084473 }, { 32.555774211883545, 10.860815048217773, 14.719563245773315 }, 0);
+		Chaos::FCapsule B({0.0000000000000000, 0.0000000000000000, -67.499992370605469}, {0.0000000000000000, 0.0000000000000000, -67.499992370605469 + 134.99998474121094 }, 67.274772644042969);
+		Chaos::TRigidTransform<Chaos::FReal, 3> BToA;
+		BToA.SetRotation({0.0000000000000000, 0.0000000000000000, -0.70710678118654757, 0.70710678118654746});
+		BToA.SetTranslation({0.0000000000000000, -3.3237259359872290e+32, 7460.1000976562500});
+		const Chaos::FVec3 LocalDir{0.93683970992769239, 0.040186153777059030, 0.34744278175123056};
+		const Chaos::FVec3 InitialDir{-3.3237259359872290e+32, 27121.400390625000, -7460.1000976562500};
+		const FReal Length = 13.27157020568847;
+		const FReal Thickness = 0;
+		const bool bComputeMtd = true;
+
+		FReal OutTime;
+		FVec3 OutLoc;
+		FVec3 OutNorm;
+
+		// Should fail and give a valid time
+		GJKRaycast2(A, B, BToA, LocalDir, Length, OutTime, OutLoc, OutNorm, Thickness, bComputeMtd, InitialDir, Thickness);
+
+		// Expect to recieve a valid time.
+		EXPECT_TRUE(FMath::IsFinite(OutTime));
+	}
+
 	// Check that GJKPenetrationCore returns the correct result when two objects are within various distances
 	// of each other. When distance is less that GJKEpsilon, GJK will abort and call into EPA.
 	void GJKBoxBoxZeroMarginSeparationTest(FReal GJKEpsilon, FReal SeparationSize, int32 SeparationAxis)
diff --git a/Engine/Source/Programs/HeadlessChaos/Public/HeadlessChaosTestGJK.h b/Engine/Source/Programs/HeadlessChaos/Public/HeadlessChaosTestGJK.h
--- a/Engine/Source/Programs/HeadlessChaos/Public/HeadlessChaosTestGJK.h
+++ b/Engine/Source/Programs/HeadlessChaos/Public/HeadlessChaosTestGJK.h
@@ -41,4 +41,6 @@ namespace ChaosTest {
 	void GJKBoxBoxSweep();
 
 	void GJKCapsuleConvexInitialOverlapSweep();
+
+	void GJKLargeDistanceCapsuleSweep();
 }
\ No newline at end of file
diff --git a/Engine/Source/Runtime/Experimental/Chaos/Public/Chaos/AABBTreeDirtyGridUtils.h b/Engine/Source/Runtime/Experimental/Chaos/Public/Chaos/AABBTreeDirtyGridUtils.h
--- a/Engine/Source/Runtime/Experimental/Chaos/Public/Chaos/AABBTreeDirtyGridUtils.h
+++ b/Engine/Source/Runtime/Experimental/Chaos/Public/Chaos/AABBTreeDirtyGridUtils.h
@@ -44,12 +44,22 @@ namespace Chaos
 	template <typename T>
 	FORCEINLINE_DEBUGGABLE bool TooManyOverlapQueryCells(const TAABB<T, 3>& AABB, FReal DirtyElementGridCellSizeInv, int32 MaximumOverlap)
 	{
-		uint64 XsampleCount = GetDirtyCellIndexFromWorldCoordinate((FReal)AABB.Max().X, DirtyElementGridCellSizeInv) - GetDirtyCellIndexFromWorldCoordinate((FReal)AABB.Min().X, DirtyElementGridCellSizeInv) + 1;
-		uint64 YsampleCount = GetDirtyCellIndexFromWorldCoordinate((FReal)AABB.Max().Y, DirtyElementGridCellSizeInv) - GetDirtyCellIndexFromWorldCoordinate((FReal)AABB.Min().Y, DirtyElementGridCellSizeInv) + 1;
+		const int64 CellStartX = GetDirtyCellIndexFromWorldCoordinate(AABB.Min().X, DirtyElementGridCellSizeInv);
+		const int64 CellStartY = GetDirtyCellIndexFromWorldCoordinate(AABB.Min().Y, DirtyElementGridCellSizeInv);
 
-		if (XsampleCount * YsampleCount <= (uint64)MaximumOverlap)
+		const int64 CellEndX = GetDirtyCellIndexFromWorldCoordinate(AABB.Max().X, DirtyElementGridCellSizeInv);
+		const int64 CellEndY = GetDirtyCellIndexFromWorldCoordinate(AABB.Max().Y, DirtyElementGridCellSizeInv);
+
+		if(ensure(CellEndX >= CellStartX && CellEndY >= CellStartY))
 		{
-			return false;
+			const uint64 XsampleCount = (CellEndX - CellStartX) + 1;
+			const uint64 YsampleCount = (CellEndY - CellStartY) + 1;
+			if((XsampleCount <= (uint64)MaximumOverlap) &&
+			   (YsampleCount <= (uint64)MaximumOverlap) &&
+			   (XsampleCount * YsampleCount <= (uint64)MaximumOverlap))
+			{
+				return false;
+			}
 		}
 		return true;
 	}
@@ -229,6 +239,11 @@ namespace Chaos
 
 	FORCEINLINE_DEBUGGABLE bool TooManySweepQueryCells(const TVec3<FReal>& QueryHalfExtents, const FVec3& StartPoint, const FVec3& Dir, FReal Length, FReal DirtyElementGridCellSizeInv, int32 DirtyElementMaxGridCellQueryCount)
 	{
+		if(Length > (FReal)TNumericLimits<uint64>::Max())
+		{
+			return true;
+		}
+
 		int EstimatedNumberOfCells = ((int32)(QueryHalfExtents.X * 2 * DirtyElementGridCellSizeInv) + 2) * ((int32)(QueryHalfExtents.Y * 2 * DirtyElementGridCellSizeInv) + 2) +
 			((int32)(FMath::Max(QueryHalfExtents.X, QueryHalfExtents.Y) * 2 * DirtyElementGridCellSizeInv) + 2) * ((int32)(Length * DirtyElementGridCellSizeInv) + 2);
 
@@ -465,8 +480,19 @@ namespace Chaos
 
 	FORCEINLINE_DEBUGGABLE bool TooManyRaycastQueryCells(const FVec3& StartPoint, const FVec3& Dir, const FReal Length, FReal DirtyElementGridCellSizeInv, int32 DirtyElementMaxGridCellQueryCount)
 	{
+		if(Length > (FReal)TNumericLimits<uint64>::Max() || DirtyElementGridCellSizeInv == 0)
+		{
+			return true;
+		}
+
 		FVec3 EndPoint = StartPoint + Length * Dir;
 
+		const FReal ElementAbsLimit = (FReal)TNumericLimits<int64>::Max() / DirtyElementGridCellSizeInv;
+		if(StartPoint.GetAbsMax() >= ElementAbsLimit || EndPoint.GetAbsMax() >= ElementAbsLimit)
+		{
+			return true;
+		}
+
 		int64 FirstCellIndexX = GetDirtyCellIndexFromWorldCoordinate(StartPoint.X, DirtyElementGridCellSizeInv);
 		int64 FirstCellIndexY = GetDirtyCellIndexFromWorldCoordinate(StartPoint.Y, DirtyElementGridCellSizeInv);
 
@@ -482,7 +508,6 @@ namespace Chaos
 		}
 
 		return false;
-
 	}
 
 	template <typename FunctionType>
@@ -588,5 +613,4 @@ namespace Chaos
 			}
 		}
 	}
-
 }
