Chaos : async support for geometry collection physics proxies
- implement  BufferPhysicsResults_External
- add support for interpolation in PullFromPhysicsState
- remove ParticleToWorldTransform from FGeometryCollectionResults in favor for separate X,R, V and W
- specialized ProxyTimeStamp so that each proxy can have its own
- refactored a bit of the single particle proxy time stamp to group timestamp and override values

diff --git a/Engine/Source/Programs/HeadlessChaos/Private/HeadlessChaosTestRewind.cpp b/Engine/Source/Programs/HeadlessChaos/Private/HeadlessChaosTestRewind.cpp
--- a/Engine/Source/Programs/HeadlessChaos/Private/HeadlessChaosTestRewind.cpp
+++ b/Engine/Source/Programs/HeadlessChaos/Private/HeadlessChaosTestRewind.cpp
@@ -2156,7 +2156,7 @@ namespace ChaosTest {
 						const FReal ExpectedValue = ZStart + ZVel * InterpolatedTime;
 						const FReal TargetValue = ZStart + ZVel * NextSimStepTime;
 
-						if(!Proxy->IsResimSmoothing())
+						if(!Proxy->GetInterpolationData().IsResimSmoothing())
 						{
 							//no resim interpolation, just simple value interpolation
 							EXPECT_NEAR(Particle.X()[2], ExpectedValue, 1e-2);
@@ -4328,7 +4328,7 @@ namespace ChaosTest {
 			Particle2.SetGeometry(Sphere2);
 			Particle2.SetV(FVec3(0, 0, 1));
 			Particle2.SetGravityEnabled(false);
-			Proxy2->SetInterpChannel_External(1);
+			Proxy2->GetInterpolationData().SetInterpChannel_External(1);
 			Solver->RegisterObject(Proxy2);
 
 			FReal Time = 0;
diff --git a/Engine/Source/Runtime/Experimental/Chaos/Private/Chaos/Framework/ChaosResultsManager.cpp b/Engine/Source/Runtime/Experimental/Chaos/Private/Chaos/Framework/ChaosResultsManager.cpp
--- a/Engine/Source/Runtime/Experimental/Chaos/Private/Chaos/Framework/ChaosResultsManager.cpp
+++ b/Engine/Source/Runtime/Experimental/Chaos/Private/Chaos/Framework/ChaosResultsManager.cpp
@@ -2,6 +2,7 @@
 
 #include "Chaos/Framework/ChaosResultsManager.h"
 #include "PhysicsProxy/SingleParticlePhysicsProxy.h"
+#include "PhysicsProxy/GeometryCollectionPhysicsProxy.h"
 #include "Chaos/ChaosMarshallingManager.h"
 
 namespace Chaos
@@ -34,6 +35,10 @@ namespace Chaos
 
 		template <ESetPrevNextDataMode Mode>
 		void SetPrevNextDataHelper(const FPullPhysicsData& PullData);
+		
+		template <ESetPrevNextDataMode Mode, typename DirtyProxyDataType, typename InterpolationsType>
+		void SetPrevNextDataHelperTyped(const TArray<DirtyProxyDataType>& DirtyProxies, TArray<InterpolationsType>& Interpolations);
+
 		void RemoveProxy_External(FSingleParticlePhysicsProxy* Proxy)
 		{
 			ParticleToResimTarget.Remove(Proxy);
@@ -46,49 +51,57 @@ namespace Chaos
 		int32 ChannelIdx;
 	};
 
-	void FChaosInterpolationResults::Reset()
+	template <typename TProxyType, typename TInterpolationType>
+	static void ResetInterpolations(TArray<TInterpolationType>& Interpolations)
 	{
-		for (FChaosRigidInterpolationData& Data : RigidInterpolations)
+		for (TInterpolationType& Data : Interpolations)
 		{
-			if (FSingleParticlePhysicsProxy* Proxy = Data.Prev.GetProxy())
+			if (TProxyType* Proxy = Data.Prev.GetProxy())
 			{
-				Proxy->SetPullDataInterpIdx_External(INDEX_NONE);
+				Proxy->GetInterpolationData().SetPullDataInterpIdx_External(INDEX_NONE);
 			}
 		}
-		RigidInterpolations.Reset();
+		Interpolations.Reset();
+	}
+	
+	void FChaosInterpolationResults::Reset()
+	{
+		ResetInterpolations<FSingleParticlePhysicsProxy>(RigidInterpolations);
+		ResetInterpolations<FGeometryCollectionPhysicsProxy>(GeometryCollectionInterpolations);
 
 		//purposely leave Prev and Next alone as we use those for rebuild
 	}
 
-	template <ESetPrevNextDataMode Mode>
-	void FChaosResultsChannel::SetPrevNextDataHelper(const FPullPhysicsData& PullData)
+	template <ESetPrevNextDataMode Mode, typename DirtyProxyDataType, typename InterpolationsType>
+	void FChaosResultsChannel::SetPrevNextDataHelperTyped(const TArray<DirtyProxyDataType>& DirtyProxies, TArray<InterpolationsType>& Interpolations)
 	{
 		//clear results
-		const int32 Timestamp = PullData.SolverTimestamp;
-		for (const FDirtyRigidParticleData& Data : PullData.DirtyRigids)
+		for (const DirtyProxyDataType& Data : DirtyProxies)
 		{
-			if (FSingleParticlePhysicsProxy* Proxy = Data.GetProxy())
+			if (auto Proxy = Data.GetProxy())
 			{
+				FProxyInterpolationData& InterpolationData = Proxy->GetInterpolationData(); 
+				
 				//If proxy is not associated with this channel, do nothing
-				if (Proxy->GetInterpChannel_External() != ChannelIdx)
+				if (InterpolationData.GetInterpChannel_External() != ChannelIdx)
 				{
 					continue;
 				}
 
-				int32 DataIdx = Proxy->GetPullDataInterpIdx_External();
+				int32 DataIdx = InterpolationData.GetPullDataInterpIdx_External();
 				if(DataIdx == INDEX_NONE)
 				{
-					DataIdx = Results.RigidInterpolations.AddDefaulted(1);
-					Proxy->SetPullDataInterpIdx_External(DataIdx);
+					DataIdx = Interpolations.AddDefaulted(1);
+					InterpolationData.SetPullDataInterpIdx_External(DataIdx);
 
 					if(Mode == ESetPrevNextDataMode::Next)
 					{
 						//no prev so use GT data
-						Proxy->BufferPhysicsResults_External(Results.RigidInterpolations[DataIdx].Prev);
+						Proxy->BufferPhysicsResults_External(Interpolations[DataIdx].Prev);
 					}
 				}
 
-				FChaosRigidInterpolationData& OutData = Results.RigidInterpolations[DataIdx];
+				InterpolationsType& OutData = Interpolations[DataIdx];
 
 				if(Mode == ESetPrevNextDataMode::Prev)
 				{
@@ -100,11 +113,33 @@ namespace Chaos
 				{
 					OutData.Next = Data;
 				}
+			}
+		}
+	}
+	
+	template <ESetPrevNextDataMode Mode>
+	void FChaosResultsChannel::SetPrevNextDataHelper(const FPullPhysicsData& PullData)
+	{
+		//clear results
+		const int32 Timestamp = PullData.SolverTimestamp;
+
+		SetPrevNextDataHelperTyped<Mode>(PullData.DirtyRigids, Results.RigidInterpolations);
+
+		SetPrevNextDataHelperTyped<Mode>(PullData.DirtyGeometryCollections, Results.GeometryCollectionInterpolations);
 
-				//update leash target
-				if(FDirtyRigidParticleData* ResimTarget = ParticleToResimTarget.Find(Proxy))
+		// update resim target for rigids
+		for (const FDirtyRigidParticleData& Data : PullData.DirtyRigids)
+		{
+			if (FSingleParticlePhysicsProxy* Proxy = Data.GetProxy())
+			{
+				// only if the proxy is associated with this channel 
+				if (Proxy->GetInterpolationData().GetInterpChannel_External() == ChannelIdx)
 				{
-					*ResimTarget = Data;
+					//update leash target
+					if(FDirtyRigidParticleData* ResimTarget = ParticleToResimTarget.Find(Proxy))
+					{
+						*ResimTarget = Data;
+					}
 				}
 			}
 		}
@@ -239,13 +274,14 @@ namespace Chaos
 		{
 			FSingleParticlePhysicsProxy* Proxy = Itr.Key;
 
-			if(Proxy->IsResimSmoothing())
+			FProxyInterpolationData& InterpolationData = Proxy->GetInterpolationData(); 
+			if(InterpolationData.IsResimSmoothing())
 			{
-				if (Proxy->GetPullDataInterpIdx_External() == INDEX_NONE)	//not in results array
+				if (InterpolationData.GetPullDataInterpIdx_External() == INDEX_NONE)	//not in results array
 				{
 					//still need to interpolate, so add to results array
 					const int32 DataIdx = Results.RigidInterpolations.AddDefaulted(1);
-					Proxy->SetPullDataInterpIdx_External(DataIdx);
+					InterpolationData.SetPullDataInterpIdx_External(DataIdx);
 					FChaosRigidInterpolationData& RigidData = Results.RigidInterpolations[DataIdx];
 
 					RigidData.Next = Itr.Value;			//not dirty from sim, so just use whatever last next was
@@ -344,11 +380,12 @@ namespace Chaos
 		{
 			if (FSingleParticlePhysicsProxy* ResimProxy = ResimDirty.GetProxy())
 			{
+				FProxyInterpolationData& InterpolationData = ResimProxy->GetInterpolationData(); 
 				//Mark as resim only if proxy is owned by this channel
-				if(ResimProxy->GetInterpChannel_External() == ChannelIdx)
+				if(InterpolationData.GetInterpChannel_External() == ChannelIdx)
 				{
 					ParticleToResimTarget.FindOrAdd(ResimProxy) = ResimDirty;
-					ResimProxy->SetResimSmoothing(true);
+					InterpolationData.SetResimSmoothing(true);
 				}
 			}
 		}
diff --git a/Engine/Source/Runtime/Experimental/Chaos/Private/PBDRigidsSolver.cpp b/Engine/Source/Runtime/Experimental/Chaos/Private/PBDRigidsSolver.cpp
--- a/Engine/Source/Runtime/Experimental/Chaos/Private/PBDRigidsSolver.cpp
+++ b/Engine/Source/Runtime/Experimental/Chaos/Private/PBDRigidsSolver.cpp
@@ -1387,7 +1387,7 @@ namespace Chaos
 			for (int32 Idx = 0; Idx < ActiveGC.Num(); ++Idx)
 			{
 				PullData->DirtyGeometryCollections.AddDefaulted();
-				ActiveGC[Idx]->BufferPhysicsResults(this, PullData->DirtyGeometryCollections.Last());
+				ActiveGC[Idx]->BufferPhysicsResults_Internal(this, PullData->DirtyGeometryCollections.Last());
 			}
 		}
 
diff --git a/Engine/Source/Runtime/Experimental/Chaos/Private/PhysicsProxy/GeometryCollectionPhysicsProxy.cpp b/Engine/Source/Runtime/Experimental/Chaos/Private/PhysicsProxy/GeometryCollectionPhysicsProxy.cpp
--- a/Engine/Source/Runtime/Experimental/Chaos/Private/PhysicsProxy/GeometryCollectionPhysicsProxy.cpp
+++ b/Engine/Source/Runtime/Experimental/Chaos/Private/PhysicsProxy/GeometryCollectionPhysicsProxy.cpp
@@ -90,7 +90,10 @@ void FGeometryCollectionResults::Reset()
 	SolverDt = 0.0f;
 	DisabledStates.SetNum(0);
 	GlobalTransforms.SetNum(0);
-	ParticleToWorldTransforms.SetNum(0);
+	ParticleXs.SetNum(0);
+	ParticleRs.SetNum(0);
+	ParticleVs.SetNum(0);
+	ParticleWs.SetNum(0);
 	IsObjectDynamic = false;
 	IsObjectLoading = false;
 }
@@ -1596,7 +1599,7 @@ void FGeometryCollectionPhysicsProxy::PushToPhysicsState()
 							FTransform NewChildWorldTransform = PhysicsThreadCollection.MassToLocal[TransformGroupIndex] * PhysicsThreadCollection.Transform[TransformGroupIndex] * ActorToWorld;
 							Chaos::FRigidTransform3 ParentToChildTransform = Handle->ChildToParent().Inverse();
 							FTransform NewParentWorldTRansform = ParentToChildTransform * NewChildWorldTransform;
-							SetClusteredParticleKinematicTarget(ParentHandle, NewParentWorldTRansform);
+							SetClusteredParticleKinematicTarget_Internal(ParentHandle, NewParentWorldTRansform);
 
 							InternalClusterParentUpdated = true;
 						}
@@ -1605,7 +1608,7 @@ void FGeometryCollectionPhysicsProxy::PushToPhysicsState()
 					if (!Handle->Disabled())
 					{
 						FTransform WorldTransform = PhysicsThreadCollection.MassToLocal[TransformGroupIndex] * PhysicsThreadCollection.Transform[TransformGroupIndex] * ActorToWorld;
-						SetClusteredParticleKinematicTarget(Handle, WorldTransform);
+						SetClusteredParticleKinematicTarget_Internal(Handle, WorldTransform);
 					}
 				}
 			}
@@ -1613,10 +1616,10 @@ void FGeometryCollectionPhysicsProxy::PushToPhysicsState()
 	}
 }
 
-void FGeometryCollectionPhysicsProxy::SetClusteredParticleKinematicTarget(Chaos::FPBDRigidClusteredParticleHandle* Handle, const FTransform& NewWorldTransform)
+void FGeometryCollectionPhysicsProxy::SetClusteredParticleKinematicTarget_Internal(Chaos::FPBDRigidClusteredParticleHandle* Handle, const FTransform& NewWorldTransform)
 {
 	// CONTEXT: PHYSICSTHREAD
-	// this should be called only on teh physics thread
+	// this should be called only on the physics thread
 	const Chaos::EObjectStateType ObjectState = Handle->ObjectState();
 	if (ensure(ObjectState == Chaos::EObjectStateType::Kinematic))
 	{
@@ -1629,7 +1632,93 @@ void FGeometryCollectionPhysicsProxy::SetClusteredParticleKinematicTarget(Chaos:
 		}
 	}
 }
-void FGeometryCollectionPhysicsProxy::BufferPhysicsResults(Chaos::FPBDRigidsSolver* CurrentSolver, Chaos::FDirtyGeometryCollectionData& BufferData)
+
+static EObjectStateTypeEnum GetObjectStateFromHandle(const Chaos::TPBDRigidClusteredParticleHandle<Chaos::FReal, 3>* Handle)
+{
+	if (!Handle->Sleeping())
+	{
+		switch (Handle->ObjectState())
+		{
+		case Chaos::EObjectStateType::Kinematic:
+			return EObjectStateTypeEnum::Chaos_Object_Kinematic;
+		case Chaos::EObjectStateType::Static:
+			return EObjectStateTypeEnum::Chaos_Object_Static;
+		case Chaos::EObjectStateType::Sleeping:
+			return EObjectStateTypeEnum::Chaos_Object_Sleeping;
+		case Chaos::EObjectStateType::Dynamic:
+		case Chaos::EObjectStateType::Uninitialized:
+		default:
+			return EObjectStateTypeEnum::Chaos_Object_Dynamic;
+		}
+	}
+	return EObjectStateTypeEnum::Chaos_Object_Sleeping;
+}
+
+void FGeometryCollectionPhysicsProxy::PrepareBufferData(Chaos::FDirtyGeometryCollectionData& BufferData, const FGeometryDynamicCollection& ThreadCollection,  Chaos::FReal SolverLastDt)
+{
+	/**
+	* CONTEXT: GAMETHREAD or PHYSICS THREAD
+	* this method needs to be careful of data access , ultimately only the this pointer can be passed to the BufferData 
+	*/
+	BufferData.SetProxy(*this);
+
+	IsObjectDynamic = false;
+	FGeometryCollectionResults& TargetResults = BufferData.Results;
+	TargetResults.SolverDt = SolverLastDt;
+
+	const int32 NumTransformGroupElements = ThreadCollection.NumElements(FGeometryCollection::TransformGroup);
+	if (TargetResults.NumTransformGroup() != NumTransformGroupElements)
+	{
+		TargetResults.InitArrays(ThreadCollection);
+	}
+}
+
+void FGeometryCollectionPhysicsProxy::BufferPhysicsResults_External(Chaos::FDirtyGeometryCollectionData& BufferData)
+{
+	/**
+	* CONTEXT: GAMETHREAD
+	* Called per-tick when async is on after the simulation has completed.
+	* goal is collect current game thread data of the proxy so it can be used if no previous physics thread data is available for interpolating  
+	*/
+	//SCOPE_CYCLE_COUNTER(STAT_GeometryCollection_BufferPhysicsResults_External);
+	if (IsObjectDeleting) return;
+	PrepareBufferData(BufferData, GameThreadCollection);
+	
+	FGeometryCollectionResults& TargetResults = BufferData.Results;
+	//TargetResults.SolverDt = CurrentSolver->GetLastDt();	//todo: should this use timestamp for async mode?
+
+	// not interpolatable, so we don't fill it 
+	//	TargetResults.DisabledStates;
+	//	TargetResults.Parent;
+	//
+	//TargetResults.DynamicState; // ObjectState
+	TManagedArray<FVector3f>* LinearVelocities = GameThreadCollection.FindAttributeTyped<FVector3f>("LinearVelocity", FTransformCollection::TransformGroup);
+	TManagedArray<FVector3f>* AngularVelocities = GameThreadCollection.FindAttributeTyped<FVector3f>("AngularVelocity", FTransformCollection::TransformGroup);
+	
+	for (int32 TransformGroupIndex = 0; TransformGroupIndex < TargetResults.Transforms.Num(); ++TransformGroupIndex)
+	{
+		const Chaos::FGeometryParticle& GTParticle = *GTParticles[TransformGroupIndex];
+		TargetResults.ParticleXs[TransformGroupIndex] = GTParticle.X();
+		TargetResults.ParticleRs[TransformGroupIndex] = GTParticle.R();
+
+		if (LinearVelocities && AngularVelocities)
+		{
+			TargetResults.ParticleVs[TransformGroupIndex] = Chaos::FVec3((*LinearVelocities)[TransformGroupIndex]);
+			TargetResults.ParticleWs[TransformGroupIndex] = Chaos::FVec3((*AngularVelocities)[TransformGroupIndex]);
+		}
+
+		TargetResults.Parent[TransformGroupIndex] = GameThreadCollection.Parent[TransformGroupIndex];
+		TargetResults.Transforms[TransformGroupIndex] = GameThreadCollection.Transform[TransformGroupIndex];
+		
+	}
+	
+	// TargetResults.GlobalTransforms; // no need, can be recomputed ?  
+	//
+	// TargetResults.IsObjectDynamic;
+	// TargetResults.IsObjectLoading;
+}
+
+void FGeometryCollectionPhysicsProxy::BufferPhysicsResults_Internal(Chaos::FPBDRigidsSolver* CurrentSolver, Chaos::FDirtyGeometryCollectionData& BufferData)
 {
 	/**
 	 * CONTEXT: PHYSICSTHREAD
@@ -1639,17 +1728,12 @@ void FGeometryCollectionPhysicsProxy::BufferPhysicsResults(Chaos::FPBDRigidsSolv
 	using namespace Chaos;
 	SCOPE_CYCLE_COUNTER(STAT_CacheResultGeomCollection);
 	if (IsObjectDeleting) return;
-	BufferData.SetProxy(*this);
+
+	//todo: should this use timestamp instead of GetLastDt for async mode?
+	PrepareBufferData(BufferData, PhysicsThreadCollection, CurrentSolver->GetLastDt());
 
 	IsObjectDynamic = false;
 	FGeometryCollectionResults& TargetResults = BufferData.Results;
-	TargetResults.SolverDt = CurrentSolver->GetLastDt();	//todo: should this use timestamp for async mode?
-
-	int32 NumTransformGroupElements = PhysicsThreadCollection.NumElements(FGeometryCollection::TransformGroup);
-	if (TargetResults.NumTransformGroup() != NumTransformGroupElements)
-	{
-		TargetResults.InitArrays(PhysicsThreadCollection);
-	}
 
 	const FTransform& ActorToWorld = Parameters.WorldTransform;
 	const TManagedArray<int32>& Parent = PhysicsThreadCollection.Parent;
@@ -1657,6 +1741,7 @@ void FGeometryCollectionPhysicsProxy::BufferPhysicsResults(Chaos::FPBDRigidsSolv
 	const bool IsActorScaled = !ActorToWorld.GetScale3D().Equals(FVector::OneVector);
 	const FTransform ActorScaleTransform(FQuat::Identity,  FVector::ZeroVector, ActorToWorld.GetScale3D());
 
+	const int32 NumTransformGroupElements = TargetResults.Transforms.Num(); 
 	if(NumTransformGroupElements > 0)
 	{ 
 		SCOPE_CYCLE_COUNTER(STAT_CalcParticleToWorld);
@@ -1675,31 +1760,7 @@ void FGeometryCollectionPhysicsProxy::BufferPhysicsResults(Chaos::FPBDRigidsSolv
 			}
 
 			// Dynamic state is also updated by the solver during field interaction.
-			if (!Handle->Sleeping())
-			{
-				const Chaos::EObjectStateType ObjectState = Handle->ObjectState();
-				switch (ObjectState)
-				{
-				case Chaos::EObjectStateType::Kinematic:
-					TargetResults.DynamicState[TransformGroupIndex] = (int)EObjectStateTypeEnum::Chaos_Object_Kinematic;
-					break;
-				case Chaos::EObjectStateType::Static:
-					TargetResults.DynamicState[TransformGroupIndex] = (int)EObjectStateTypeEnum::Chaos_Object_Static;
-					break;
-				case Chaos::EObjectStateType::Sleeping:
-					TargetResults.DynamicState[TransformGroupIndex] = (int)EObjectStateTypeEnum::Chaos_Object_Sleeping;
-					break;
-				case Chaos::EObjectStateType::Dynamic:
-				case Chaos::EObjectStateType::Uninitialized:
-				default:
-					TargetResults.DynamicState[TransformGroupIndex] = (int)EObjectStateTypeEnum::Chaos_Object_Dynamic;
-					break;
-				}
-			}
-			else
-			{
-				TargetResults.DynamicState[TransformGroupIndex] = (int)EObjectStateTypeEnum::Chaos_Object_Sleeping;
-			}
+			TargetResults.DynamicState[TransformGroupIndex] = static_cast<int>(GetObjectStateFromHandle(Handle));
 
 			// Update the transform and parent hierarchy of the active rigid bodies. Active bodies can be either
 			// rigid geometry defined from the leaf nodes of the collection, or cluster bodies that drive an entire
@@ -1715,9 +1776,12 @@ void FGeometryCollectionPhysicsProxy::BufferPhysicsResults(Chaos::FPBDRigidsSolv
 				// or a collection of rigidly attached geometries (Clustering). The cluster is represented as a
 				// single transform in the GeometryCollection, and all children are stored in the local space
 				// of the parent cluster.
-	
-				FTransform& ParticleToWorld = TargetResults.ParticleToWorldTransforms[TransformGroupIndex];
-				ParticleToWorld = Chaos::FRigidTransform3(Handle->X(), Handle->R());
+
+				TargetResults.ParticleXs[TransformGroupIndex] = Handle->X();
+				TargetResults.ParticleRs[TransformGroupIndex] = Handle->R();
+				TargetResults.ParticleVs[TransformGroupIndex] = Handle->V();
+				TargetResults.ParticleWs[TransformGroupIndex] = Handle->W();
+				FRigidTransform3 ParticleToWorld(Handle->X(), Handle->R());
 				const FTransform MassToLocal = PhysicsThreadCollection.MassToLocal[TransformGroupIndex];
 
 				TargetResults.Transforms[TransformGroupIndex] = MassToLocal.GetRelativeTransformReverse(ParticleToWorld).GetRelativeTransform(ActorToWorld);
@@ -1786,8 +1850,11 @@ void FGeometryCollectionPhysicsProxy::BufferPhysicsResults(Chaos::FPBDRigidsSolv
 						{
 							Chaos::TPBDRigidClusteredParticleHandle<Chaos::FReal, 3>* ProxyElementHandle = SolverParticleHandles[TransformGroupIndex];
 
-							FTransform& ParticleToWorld = TargetResults.ParticleToWorldTransforms[TransformGroupIndex];
-							ParticleToWorld             = ProxyElementHandle->ChildToParent() * FRigidTransform3(ClusterParent->X(), ClusterParent->R());    // aka ClusterChildToWorld
+							const FTransform ParticleToWorld = ProxyElementHandle->ChildToParent() * FRigidTransform3(ClusterParent->X(), ClusterParent->R());    // aka ClusterChildToWorld
+							TargetResults.ParticleXs[TransformGroupIndex] = ParticleToWorld.GetTranslation();
+							TargetResults.ParticleRs[TransformGroupIndex] = ParticleToWorld.GetRotation();
+							TargetResults.ParticleVs[TransformGroupIndex] = ClusterParent->V();
+							TargetResults.ParticleWs[TransformGroupIndex] = ClusterParent->W();
 
 							// GeomToActor = ActorToWorld.Inv() * ClusterChildToWorld * MassToLocal.Inv();
 							const FTransform MassToLocal                  = PhysicsThreadCollection.MassToLocal[TransformGroupIndex];
@@ -1816,6 +1883,7 @@ void FGeometryCollectionPhysicsProxy::BufferPhysicsResults(Chaos::FPBDRigidsSolv
 								FTransform ChildToWorld = Handle->ChildToParent() * FRigidTransform3(ClusterParent->X(), ClusterParent->R());
 								Handle->X() = Handle->P() = ChildToWorld.GetTranslation();
 								Handle->R() = Handle->Q() = ChildToWorld.GetRotation();
+								// TODO : shoudl we transfer velocity as well ?
 								Handle->UpdateWorldSpaceState(ChildToWorld, FVec3(0));
 								CurrentSolver->GetEvolution()->DirtyParticle(*Handle);
 							}
@@ -1855,7 +1923,7 @@ void FGeometryCollectionPhysicsProxy::FlipBuffer()
 }
 
 // Called from FPhysScene_ChaosInterface::SyncBodies(), NOT the solver.
-bool FGeometryCollectionPhysicsProxy::PullFromPhysicsState(const Chaos::FDirtyGeometryCollectionData& BufferData, const int32 SolverSyncTimestamp)
+bool FGeometryCollectionPhysicsProxy::PullFromPhysicsState(const Chaos::FDirtyGeometryCollectionData& PullData, const int32 SolverSyncTimestamp, const Chaos::FDirtyGeometryCollectionData* NextPullData, const Chaos::FRealSingle* Alpha)
 {
 	if(IsObjectDeleting) return false;
 
@@ -1867,54 +1935,81 @@ bool FGeometryCollectionPhysicsProxy::PullFromPhysicsState(const Chaos::FDirtyGe
 	 *
 	 * Note: A read lock will have been acquired for this - so the physics thread won't force a buffer flip while this
 	 * sync is ongoing
+	 *
 	 */
 
-	const FGeometryCollectionResults& TargetResults = BufferData.Results;
+	const FGeometryCollectionResults& TargetResults = PullData.Results;
 
-	FGeometryDynamicCollection& DynamicCollection = GameThreadCollection;
-
-	TManagedArray<FVector3f>* LinearVelocity = DynamicCollection.FindAttributeTyped<FVector3f>("LinearVelocity", FTransformCollection::TransformGroup);
+	TManagedArray<FVector3f>* LinearVelocities = GameThreadCollection.FindAttributeTyped<FVector3f>("LinearVelocity", FTransformCollection::TransformGroup);
+	TManagedArray<FVector3f>* AngularVelocities = GameThreadCollection.FindAttributeTyped<FVector3f>("AngularVelocity", FTransformCollection::TransformGroup);
 
 	// We should never be changing the number of entries, this would break other 
 	// attributes in the transform group.
-	const int32 NumTransforms = DynamicCollection.Transform.Num();
+	const int32 NumTransforms = GameThreadCollection.Transform.Num();
 	if (ensure(NumTransforms == TargetResults.Transforms.Num()))
 	{
+		// first : copy the non interpolate-able values
 		for (int32 TransformGroupIndex = 0; TransformGroupIndex < NumTransforms; ++TransformGroupIndex)
 		{
 			if (!TargetResults.DisabledStates[TransformGroupIndex])
 			{
-				DynamicCollection.Parent[TransformGroupIndex] = TargetResults.Parent[TransformGroupIndex];
-				const FTransform& LocalTransform = TargetResults.Transforms[TransformGroupIndex];
-				const FTransform& ParticleToWorld = TargetResults.ParticleToWorldTransforms[TransformGroupIndex];
+				GameThreadCollection.Parent[TransformGroupIndex] = TargetResults.Parent[TransformGroupIndex];
+			}
+			GameThreadCollection.DynamicState[TransformGroupIndex] = TargetResults.DynamicState[TransformGroupIndex];
+			GameThreadCollection.Active[TransformGroupIndex] = !TargetResults.DisabledStates[TransformGroupIndex];
+		}
 
-				DynamicCollection.Transform[TransformGroupIndex] = LocalTransform;
+		// second : interpolate-able ones
+		const bool bNeedInterpolation = (NextPullData!= nullptr);
+		if (bNeedInterpolation)
+		{
+			const FGeometryCollectionResults& Prev = PullData.Results;
+			const FGeometryCollectionResults& Next = NextPullData->Results;
+			for (int32 TransformGroupIndex = 0; TransformGroupIndex < NumTransforms; ++TransformGroupIndex)
+			{
+				if (!TargetResults.DisabledStates[TransformGroupIndex])
+				{
+					Chaos::FGeometryParticle& GTParticle = *GTParticles[TransformGroupIndex];
+					GTParticle.SetX(FMath::Lerp(Prev.ParticleXs[TransformGroupIndex], Next.ParticleXs[TransformGroupIndex], *Alpha), false);
+					GTParticle.SetR(FMath::Lerp(Prev.ParticleRs[TransformGroupIndex], Next.ParticleRs[TransformGroupIndex], *Alpha), false);
+					GTParticle.UpdateShapeBounds();
 
-				Chaos::FGeometryParticle* GTParticle = GTParticles[TransformGroupIndex].Get();
+					GameThreadCollection.Transform[TransformGroupIndex].Blend(Prev.Transforms[TransformGroupIndex], Next.Transforms[TransformGroupIndex], *Alpha);
 
-				if(LinearVelocity)
+					if(LinearVelocities && AngularVelocities)
+					{
+						// LWC : potential loss of precision if the velocity is insanely big ( unlikely ) 
+						(*LinearVelocities)[TransformGroupIndex] = FVector3f(FMath::Lerp(Prev.ParticleVs[TransformGroupIndex], Next.ParticleVs[TransformGroupIndex], *Alpha)); 
+						(*AngularVelocities)[TransformGroupIndex] = FVector3f(FMath::Lerp(Prev.ParticleWs[TransformGroupIndex], Next.ParticleWs[TransformGroupIndex], *Alpha));
+					}
+				}
+			}
+		}
+		else
+		{
+			for (int32 TransformGroupIndex = 0; TransformGroupIndex < NumTransforms; ++TransformGroupIndex)
+			{
+				if (!TargetResults.DisabledStates[TransformGroupIndex])
 				{
-					TManagedArray<FVector3f>* AngularVelocity = DynamicCollection.FindAttributeTyped<FVector3f>("AngularVelocity", FTransformCollection::TransformGroup);
-					check(AngularVelocity);
-					FVector DiffX = ParticleToWorld.GetTranslation() - GTParticle->X();
-					FVector DiffR = (ParticleToWorld.GetRotation().Euler() - GTParticle->R().Euler()) * (UE_PI / 180.0f);
+					Chaos::FGeometryParticle& GTParticle = *GTParticles[TransformGroupIndex];
+					GTParticle.SetX(TargetResults.ParticleXs[TransformGroupIndex], false);
+					GTParticle.SetR(TargetResults.ParticleRs[TransformGroupIndex], false);
+					GTParticle.UpdateShapeBounds();
 
-					(*LinearVelocity)[TransformGroupIndex] = FVector3f(DiffX / TargetResults.SolverDt);
-					(*AngularVelocity)[TransformGroupIndex] = FVector3f(DiffR / TargetResults.SolverDt);
-				}
+					GameThreadCollection.Transform[TransformGroupIndex] = TargetResults.Transforms[TransformGroupIndex];
 
-				GTParticles[TransformGroupIndex]->SetX(ParticleToWorld.GetTranslation());
-				GTParticles[TransformGroupIndex]->SetR(ParticleToWorld.GetRotation());
-				GTParticles[TransformGroupIndex]->UpdateShapeBounds();
+					if(LinearVelocities && AngularVelocities)
+					{
+						// LWC : potential loss of precision if the velocity is insanely big ( unlikely ) 
+						(*LinearVelocities)[TransformGroupIndex] = FVector3f(TargetResults.ParticleVs[TransformGroupIndex]); 
+						(*AngularVelocities)[TransformGroupIndex] = FVector3f(TargetResults.ParticleWs[TransformGroupIndex]);
+					}
+				}
 			}
-
-			DynamicCollection.DynamicState[TransformGroupIndex] = TargetResults.DynamicState[TransformGroupIndex];
-			DynamicCollection.Active[TransformGroupIndex] = !TargetResults.DisabledStates[TransformGroupIndex];
 		}
 
 		//question: why do we need this? Sleeping objects will always have to update GPU
-		DynamicCollection.MakeDirty();
-
+		GameThreadCollection.MakeDirty();
 	}
 
 	return true;
diff --git a/Engine/Source/Runtime/Experimental/Chaos/Private/PhysicsProxy/JointConstraintProxy.cpp b/Engine/Source/Runtime/Experimental/Chaos/Private/PhysicsProxy/JointConstraintProxy.cpp
--- a/Engine/Source/Runtime/Experimental/Chaos/Private/PhysicsProxy/JointConstraintProxy.cpp
+++ b/Engine/Source/Runtime/Experimental/Chaos/Private/PhysicsProxy/JointConstraintProxy.cpp
@@ -19,7 +19,7 @@ namespace Chaos
 {
 
 FJointConstraintPhysicsProxy::FJointConstraintPhysicsProxy(FJointConstraint* InConstraint, FPBDJointConstraintHandle* InHandle, UObject* InOwner)
-	: Base(EPhysicsProxyType::JointConstraintType, InOwner)
+	: Base(EPhysicsProxyType::JointConstraintType, InOwner, new FProxyTimestampBase)
 	, Constraint_GT(InConstraint) // This proxy assumes ownership of the Constraint, and will free it during DestroyOnPhysicsThread
 	, Constraint_PT(InHandle)
 {
diff --git a/Engine/Source/Runtime/Experimental/Chaos/Private/PhysicsProxy/SingleParticlePhysicsProxy.cpp b/Engine/Source/Runtime/Experimental/Chaos/Private/PhysicsProxy/SingleParticlePhysicsProxy.cpp
--- a/Engine/Source/Runtime/Experimental/Chaos/Private/PhysicsProxy/SingleParticlePhysicsProxy.cpp
+++ b/Engine/Source/Runtime/Experimental/Chaos/Private/PhysicsProxy/SingleParticlePhysicsProxy.cpp
@@ -25,10 +25,9 @@ int32 SyncKinematicOnGameThread = 0;
 FAutoConsoleVariableRef CVar_SyncKinematicOnGameThread(TEXT("P.Chaos.SyncKinematicOnGameThread"), SyncKinematicOnGameThread, TEXT("If set to 1, if a kinematic is flagged to send position back to game thread, move component, if 0, do not."));
 
 FSingleParticlePhysicsProxy::FSingleParticlePhysicsProxy(TUniquePtr<PARTICLE_TYPE>&& InParticle, FParticleHandle* InHandle, UObject* InOwner)
-	: IPhysicsProxyBase(EPhysicsProxyType::SingleParticleProxy, InOwner)
+	: IPhysicsProxyBase(EPhysicsProxyType::SingleParticleProxy, InOwner, new FSingleParticleProxyTimestamp)
 	, Particle(MoveTemp(InParticle))
 	, Handle(InHandle)
-	, PullDataInterpIdx_External(INDEX_NONE)
 {
 	Particle->SetProxy(this);
 }
@@ -255,29 +254,29 @@ bool FSingleParticlePhysicsProxy::PullFromPhysicsState(const Chaos::FDirtyRigidP
 	{
 		const bool bSyncXR = SyncKinematicOnGameThread || (Rigid->ObjectState() != EObjectStateType::Kinematic);
 
-		const FProxyTimestamp* ProxyTimestamp = PullData.GetTimestamp();
+		const FSingleParticleProxyTimestamp* ProxyTimestamp = PullData.GetTimestamp();
 		
 		if(NextPullData)
 		{
-			auto LerpHelper = [SolverSyncTimestamp](const int32 PropertyTimestamp, const auto& Prev, const auto& Overwrite) -> const auto*
+			auto LerpHelper = [SolverSyncTimestamp](const auto& Prev, const auto& OverwriteProperty) -> const auto*
 			{
 				//if overwrite is in the future, do nothing
 				//if overwrite is on this step, we want to interpolate from overwrite to the result of the frame that consumed the overwrite
 				//if overwrite is in the past, just do normal interpolation
 
 				//this is nested because otherwise compiler can't figure out the type of nullptr with an auto return type
-				return PropertyTimestamp <= SolverSyncTimestamp ? (PropertyTimestamp < SolverSyncTimestamp ? &Prev : &Overwrite) : nullptr;
+				return OverwriteProperty.Timestamp <= SolverSyncTimestamp ? (OverwriteProperty.Timestamp < SolverSyncTimestamp ? &Prev : &OverwriteProperty.Value) : nullptr;
 			};
 
 			if (bSyncXR)
 			{
-				const float UseResimInterpStrength = GetInterpChannel_External() == 0 ? ResimInterpStrength : ResimInterpStrength2;
+				const float UseResimInterpStrength = InterpolationData.GetInterpChannel_External() == 0 ? ResimInterpStrength : ResimInterpStrength2;
 
 				bool bKeepSmoothing = false;
-				if (const FVec3* Prev = LerpHelper(ProxyTimestamp->XTimestamp, PullData.X, ProxyTimestamp->OverWriteX))
+				if (const FVec3* Prev = LerpHelper(PullData.X, ProxyTimestamp->OverWriteX))
 				{
 					FVec3 Target = FMath::Lerp(*Prev, NextPullData->X, *Alpha);
-					if (IsResimSmoothing())
+					if (InterpolationData.IsResimSmoothing())
 					{
 						const FVec3 SmoothedTarget = FMath::Lerp(Rigid->X(), Target, UseResimInterpStrength);
 						if((SmoothedTarget - Target).SizeSquared() > MinLinError2ForResimInterp)
@@ -290,10 +289,10 @@ bool FSingleParticlePhysicsProxy::PullFromPhysicsState(const Chaos::FDirtyRigidP
 					Rigid->SetX(Target, false);
 				}
 
-				if (const FQuat* Prev = LerpHelper(ProxyTimestamp->RTimestamp, PullData.R, ProxyTimestamp->OverWriteR))
+				if (const FQuat* Prev = LerpHelper(PullData.R, ProxyTimestamp->OverWriteR))
 				{
 					FQuat Target = FMath::Lerp(*Prev, NextPullData->R, *Alpha);
-					if (IsResimSmoothing())
+					if (InterpolationData.IsResimSmoothing())
 					{
 						const FQuat SmoothedTarget = FMath::Lerp<FQuat>(Rigid->R(), Target, UseResimInterpStrength);
 						if(FQuat::ErrorAutoNormalize(SmoothedTarget, Target) > MinRotErrorForResimInterp)
@@ -305,16 +304,16 @@ bool FSingleParticlePhysicsProxy::PullFromPhysicsState(const Chaos::FDirtyRigidP
 					Rigid->SetR(Target, false);
 				}
 
-				SetResimSmoothing(bKeepSmoothing);
+				InterpolationData.SetResimSmoothing(bKeepSmoothing);
 			}
 
-			if (const FVec3* Prev = LerpHelper(ProxyTimestamp->VTimestamp, PullData.V, ProxyTimestamp->OverWriteV))
+			if (const FVec3* Prev = LerpHelper(PullData.V, ProxyTimestamp->OverWriteV))
 			{
 				FVec3 Target = FMath::Lerp(*Prev, NextPullData->V, *Alpha);
 				Rigid->SetV(Target, false);
 			}
 
-			if (const FVec3* Prev = LerpHelper(ProxyTimestamp->WTimestamp, PullData.W, ProxyTimestamp->OverWriteW))
+			if (const FVec3* Prev = LerpHelper(PullData.W, ProxyTimestamp->OverWriteW))
 			{
 				FVec3 Target = FMath::Lerp(*Prev, NextPullData->W, *Alpha);
 				Rigid->SetW(Target, false);
@@ -338,23 +337,23 @@ bool FSingleParticlePhysicsProxy::PullFromPhysicsState(const Chaos::FDirtyRigidP
 			if (bSyncXR)
 			{
 				//no interpolation, just ignore if overwrite comes after
-				if (SolverSyncTimestamp >= ProxyTimestamp->XTimestamp)
+				if (SolverSyncTimestamp >= ProxyTimestamp->OverWriteX.Timestamp)
 				{
 					Rigid->SetX(PullData.X, false);
 				}
 
-				if (SolverSyncTimestamp >= ProxyTimestamp->RTimestamp)
+				if (SolverSyncTimestamp >= ProxyTimestamp->OverWriteR.Timestamp)
 				{
 					Rigid->SetR(PullData.R, false);
 				}
 			}
 
-			if(SolverSyncTimestamp >= ProxyTimestamp->VTimestamp)
+			if(SolverSyncTimestamp >= ProxyTimestamp->OverWriteV.Timestamp)
 			{
 				Rigid->SetV(PullData.V, false);
 			}
 
-			if(SolverSyncTimestamp >= ProxyTimestamp->WTimestamp)
+			if(SolverSyncTimestamp >= ProxyTimestamp->OverWriteW.Timestamp)
 			{
 				Rigid->SetW(PullData.W, false);
 			}
diff --git a/Engine/Source/Runtime/Experimental/Chaos/Private/PhysicsProxy/SuspensionConstraintProxy.cpp b/Engine/Source/Runtime/Experimental/Chaos/Private/PhysicsProxy/SuspensionConstraintProxy.cpp
--- a/Engine/Source/Runtime/Experimental/Chaos/Private/PhysicsProxy/SuspensionConstraintProxy.cpp
+++ b/Engine/Source/Runtime/Experimental/Chaos/Private/PhysicsProxy/SuspensionConstraintProxy.cpp
@@ -18,7 +18,7 @@ namespace Chaos
 {
 
 FSuspensionConstraintPhysicsProxy::FSuspensionConstraintPhysicsProxy(FSuspensionConstraint* InConstraint, FPBDSuspensionConstraintHandle* InHandle, UObject* InOwner)
-	: Base(EPhysicsProxyType::SuspensionConstraintType, InOwner)
+	: Base(EPhysicsProxyType::SuspensionConstraintType, InOwner, new FProxyTimestampBase)
 	, Constraint_GT(InConstraint)
 	, Constraint_PT(InHandle)
 	, bInitialized(false)
diff --git a/Engine/Source/Runtime/Experimental/Chaos/Public/Chaos/Framework/ChaosResultsManager.h b/Engine/Source/Runtime/Experimental/Chaos/Public/Chaos/Framework/ChaosResultsManager.h
--- a/Engine/Source/Runtime/Experimental/Chaos/Public/Chaos/Framework/ChaosResultsManager.h
+++ b/Engine/Source/Runtime/Experimental/Chaos/Public/Chaos/Framework/ChaosResultsManager.h
@@ -16,6 +16,12 @@ namespace Chaos
 		FDirtyRigidParticleData Next;
 	};
 
+	struct CHAOS_API FChaosGeometryCollectionInterpolationData
+	{
+		FDirtyGeometryCollectionData Prev;
+		FDirtyGeometryCollectionData Next;
+	};
+	
 	struct CHAOS_API FChaosInterpolationResults
 	{
 		FChaosInterpolationResults()
@@ -27,6 +33,7 @@ namespace Chaos
 		FChaosInterpolationResults(const FChaosInterpolationResults& Other) = delete;
 		FChaosInterpolationResults(FChaosInterpolationResults&& Other)
 			: RigidInterpolations(MoveTemp(Other.RigidInterpolations))
+			, GeometryCollectionInterpolations(MoveTemp(Other.GeometryCollectionInterpolations))
 			, Prev(Other.Prev)
 			, Next(Other.Next)
 			, Alpha(Other.Alpha)
@@ -38,6 +45,7 @@ namespace Chaos
 		void Reset();
 		
 		TArray<FChaosRigidInterpolationData> RigidInterpolations;
+		TArray<FChaosGeometryCollectionInterpolationData> GeometryCollectionInterpolations;
 		FPullPhysicsData* Prev;
 		FPullPhysicsData* Next;
 		FRealSingle Alpha;
diff --git a/Engine/Source/Runtime/Experimental/Chaos/Public/Chaos/Framework/PhysicsProxy.h b/Engine/Source/Runtime/Experimental/Chaos/Public/Chaos/Framework/PhysicsProxy.h
--- a/Engine/Source/Runtime/Experimental/Chaos/Public/Chaos/Framework/PhysicsProxy.h
+++ b/Engine/Source/Runtime/Experimental/Chaos/Public/Chaos/Framework/PhysicsProxy.h
@@ -32,7 +32,7 @@ struct FBodyInstance;
  *
  * #BG TODO - rename the callbacks functions, document for the base solver object
  */
-template<class Concrete, class ConcreteData>
+template<class Concrete, class ConcreteData, typename TProxyTimeStamp>
 class TPhysicsProxy : public IPhysicsProxyBase
 {
 
@@ -44,12 +44,12 @@ public:
 	using FIntArray = Chaos::TArrayCollectionArray<int32>;
 
 	TPhysicsProxy()
-		: IPhysicsProxyBase(ConcreteType(), nullptr)
+		: IPhysicsProxyBase(ConcreteType(), nullptr, new TProxyTimeStamp)
 	{
 	}
 
 	explicit TPhysicsProxy(UObject* InOwner)
-		: IPhysicsProxyBase(ConcreteType(), InOwner)
+		: IPhysicsProxyBase(ConcreteType(), InOwner, new TProxyTimeStamp)
 	{
 	}
 
@@ -122,4 +122,9 @@ public:
 	void* GetUserData() const { return nullptr; }
 
 	Chaos::FRigidTransform3 GetTransform() const { return Chaos::FRigidTransform3(); }
+	
+	FORCEINLINE_DEBUGGABLE TProxyTimeStamp& GetSyncTimestampTyped()
+	{
+		return GetSyncTimestampAs<TProxyTimeStamp>();
+	}
 };
diff --git a/Engine/Source/Runtime/Experimental/Chaos/Public/Chaos/Framework/PhysicsProxyBase.h b/Engine/Source/Runtime/Experimental/Chaos/Public/Chaos/Framework/PhysicsProxyBase.h
--- a/Engine/Source/Runtime/Experimental/Chaos/Public/Chaos/Framework/PhysicsProxyBase.h
+++ b/Engine/Source/Runtime/Experimental/Chaos/Public/Chaos/Framework/PhysicsProxyBase.h
@@ -24,29 +24,49 @@ namespace Chaos
 	class FPhysicsSolverBase;
 }
 
-struct CHAOS_API FProxyTimestamp
+struct CHAOS_API FProxyTimestampBase
 {
-	int32 XTimestamp = INDEX_NONE;
-	int32 RTimestamp = INDEX_NONE;
-	int32 VTimestamp = INDEX_NONE;
-	int32 WTimestamp = INDEX_NONE;
-	int32 ObjectStateTimestamp = INDEX_NONE;
-	Chaos::FVec3 OverWriteX;
-	Chaos::FRotation3 OverWriteR;
-	Chaos::FVec3 OverWriteV;
-	Chaos::FVec3 OverWriteW;
 	bool bDeleted = false;
 };
 
+template <typename TPropertyType>
+struct TTimestampProperty
+{
+	FORCEINLINE_DEBUGGABLE void Set(int32 InTimestamp, const TPropertyType& InValue)
+	{
+		Value = InValue;
+		Timestamp = InTimestamp;
+	}
+	
+	TPropertyType Value;
+	int32 Timestamp = INDEX_NONE;
+};
+
+struct CHAOS_API FSingleParticleProxyTimestamp: public FProxyTimestampBase
+{
+	int32 ObjectStateTimestamp = INDEX_NONE;
+	TTimestampProperty<Chaos::FVec3> OverWriteX;
+	TTimestampProperty<Chaos::FRotation3> OverWriteR;
+	TTimestampProperty<Chaos::FVec3> OverWriteV;
+	TTimestampProperty<Chaos::FVec3> OverWriteW;
+};
+
+struct CHAOS_API FGeometryCollectionProxyTimestamp: public FProxyTimestampBase
+{
+	// nothing to add as Geometry collections are driven from the Physics thread only
+	// ( including kinematic targeting )
+};
+
+
 class CHAOS_API IPhysicsProxyBase
 {
 public:
-	IPhysicsProxyBase(EPhysicsProxyType InType, UObject* InOwner)
+	IPhysicsProxyBase(EPhysicsProxyType InType, UObject* InOwner, FProxyTimestampBase* InProxyTimeStamp)
 		: Solver(nullptr)
 		, Owner(InOwner)
 		, DirtyIdx(INDEX_NONE)
+		, SyncTimestamp(InProxyTimeStamp)
 		, Type(InType)
-		, SyncTimestamp(new FProxyTimestamp)
 	{}
 
 	UObject* GetOwner() const { return Owner; }
@@ -71,7 +91,7 @@ public:
 
 	void MarkDeleted() { SyncTimestamp->bDeleted = true; }
 
-	TSharedPtr<FProxyTimestamp,ESPMode::ThreadSafe> GetSyncTimestamp() const { return SyncTimestamp; }
+	TSharedPtr<FProxyTimestampBase,ESPMode::ThreadSafe> GetSyncTimestamp() const { return SyncTimestamp; }
 
 	bool IsInitialized() const { return InitializedOnStep != INDEX_NONE; }
 	void SetInitialized(const int32 InitializeStep)
@@ -92,6 +112,12 @@ protected:
 	// Ensures that derived classes can successfully call this destructor
 	// but no one can delete using a IPhysicsProxyBase*
 	virtual ~IPhysicsProxyBase();
+
+	template<typename TProxyTimeStamp>
+	FORCEINLINE_DEBUGGABLE TProxyTimeStamp& GetSyncTimestampAs()
+	{
+		return static_cast<TProxyTimeStamp&>(*GetSyncTimestamp());
+	}
 	
 	/** The solver that owns the solver object */
 	Chaos::FPhysicsSolverBase* Solver;
@@ -99,10 +125,10 @@ protected:
 
 private:
 	int32 DirtyIdx;
+	TSharedPtr<FProxyTimestampBase,ESPMode::ThreadSafe> SyncTimestamp;
 protected:
 	/** Proxy type */
 	EPhysicsProxyType Type;
-	TSharedPtr<FProxyTimestamp,ESPMode::ThreadSafe> SyncTimestamp;
 	int32 InitializedOnStep = INDEX_NONE;
 	int32 IgnoreDataOnStep_Internal = INDEX_NONE;
 };
@@ -112,3 +138,28 @@ struct PhysicsProxyWrapper
 	IPhysicsProxyBase* PhysicsProxy;
 	EPhysicsProxyType Type;
 };
+
+// Data ise used by interpolation code
+struct FProxyInterpolationData
+{
+public:
+	FProxyInterpolationData()
+		: PullDataInterpIdx_External(INDEX_NONE)
+		, InterpChannel_External(0)
+		, bResimSmoothing(false)
+	{}
+
+	int32 GetPullDataInterpIdx_External() const { return PullDataInterpIdx_External; }
+	void SetPullDataInterpIdx_External(const int32 Idx) { PullDataInterpIdx_External = Idx;	}
+
+	int32 GetInterpChannel_External() const { return InterpChannel_External; }
+	void SetInterpChannel_External(int32 Channel) { InterpChannel_External = Channel; }
+
+	void SetResimSmoothing(bool ResimSmoothing) { bResimSmoothing = ResimSmoothing; }
+	bool IsResimSmoothing() const { return bResimSmoothing; }
+	
+private:
+	int32 PullDataInterpIdx_External;
+	int32 InterpChannel_External;
+	bool bResimSmoothing;
+};
\ No newline at end of file
diff --git a/Engine/Source/Runtime/Experimental/Chaos/Public/Chaos/PhysicsSolverBaseImpl.h b/Engine/Source/Runtime/Experimental/Chaos/Public/Chaos/PhysicsSolverBaseImpl.h
--- a/Engine/Source/Runtime/Experimental/Chaos/Public/Chaos/PhysicsSolverBaseImpl.h
+++ b/Engine/Source/Runtime/Experimental/Chaos/Public/Chaos/PhysicsSolverBaseImpl.h
@@ -39,6 +39,8 @@ namespace Chaos
 				LatestData = Results.Next;
 				//todo: go wide
 				const int32 SolverTimestamp = Results.Next ? Results.Next->SolverTimestamp : INDEX_NONE;
+
+				// single particles
 				for (const FChaosRigidInterpolationData& RigidInterp : Results.RigidInterpolations)
 				{
 					if (FSingleParticlePhysicsProxy* Proxy = RigidInterp.Prev.GetProxy())
@@ -48,10 +50,17 @@ namespace Chaos
 							RigidFunc(Proxy);
 						}
 					}
-
+				}
+				
+				// geometry collections
+				for (const FChaosGeometryCollectionInterpolationData& GCInterp : Results.GeometryCollectionInterpolations)
+				{
+					if (FGeometryCollectionPhysicsProxy* Proxy = GCInterp.Prev.GetProxy())
+					{
+						Proxy->PullFromPhysicsState(GCInterp.Prev, SolverTimestamp, &GCInterp.Next, &Results.Alpha);
+					}
 				}
 			}
-			
 		}
 		else
 		{
@@ -63,6 +72,8 @@ namespace Chaos
 			LatestData = Results.Next;
 			//todo: go wide
 			const int32 SolverTimestamp = Results.Next ? Results.Next->SolverTimestamp : INDEX_NONE;
+
+			// Single particles
 			for (const FChaosRigidInterpolationData& RigidInterp : Results.RigidInterpolations)
 				{
 					if (FSingleParticlePhysicsProxy* Proxy = RigidInterp.Prev.GetProxy())
@@ -74,22 +85,21 @@ namespace Chaos
 					}
 				}
 
-		}
-
-		//no interpolation for GC or joints at the moment
-		if(LatestData)
-		{
-			const int32 SyncTimestamp = LatestData->SolverTimestamp;
-			for (const FDirtyGeometryCollectionData& DirtyData : LatestData->DirtyGeometryCollections)
+			// geometry collections
+			for (const FChaosGeometryCollectionInterpolationData& GCInterp : Results.GeometryCollectionInterpolations)
 			{
-				if (auto Proxy = DirtyData.GetProxy())
+				if (FGeometryCollectionPhysicsProxy* Proxy = GCInterp.Prev.GetProxy())
 				{
-					Proxy->PullFromPhysicsState(DirtyData, SyncTimestamp);
+					Proxy->PullFromPhysicsState(GCInterp.Next, SolverTimestamp);
 				}
 			}
+			
+		}
 
-			//latest data may be used multiple times during interpolation, so for non interpolated GC we clear it
-			LatestData->DirtyGeometryCollections.Reset();
+		//no interpolation for GC or joints at the moment
+		if(LatestData)
+		{
+			 const int32 SyncTimestamp = LatestData->SolverTimestamp;
 
 			//
 			// @todo(chaos) : Add Dirty Constraints Support
diff --git a/Engine/Source/Runtime/Experimental/Chaos/Public/Chaos/PullPhysicsDataImp.h b/Engine/Source/Runtime/Experimental/Chaos/Public/Chaos/PullPhysicsDataImp.h
--- a/Engine/Source/Runtime/Experimental/Chaos/Public/Chaos/PullPhysicsDataImp.h
+++ b/Engine/Source/Runtime/Experimental/Chaos/Public/Chaos/PullPhysicsDataImp.h
@@ -14,7 +14,7 @@ namespace Chaos
 {
 class FJointConstraintPhysicsProxy;
 
-template <typename TProxy>
+template <typename TProxy, typename TTimeStamp>
 struct TBasePullData
 {
 public:
@@ -33,20 +33,20 @@ public:
 		}
 		return nullptr;
 	}
-
-	const FProxyTimestamp* GetTimestamp() const { return Timestamp.Get(); }
-
+	
+	const TTimeStamp* GetTimestamp() const { return static_cast<TTimeStamp*>(Timestamp.Get()); }
+	
 protected:
 	TBasePullData() : Proxy(nullptr){}
 	~TBasePullData() = default;
 	
 private:
 	TProxy* Proxy;
-	TSharedPtr<FProxyTimestamp,ESPMode::ThreadSafe> Timestamp;	//question: is destructor expensive now? might need a better way
+	TSharedPtr<FProxyTimestampBase,ESPMode::ThreadSafe> Timestamp;	//question: is destructor expensive now? might need a better way
 };
 
 //Simple struct for when the simulation dirties a particle. Copies all properties regardless of which changed since they tend to change together
-struct FDirtyRigidParticleData : public TBasePullData<FSingleParticlePhysicsProxy>
+struct FDirtyRigidParticleData : public TBasePullData<FSingleParticlePhysicsProxy, FSingleParticleProxyTimestamp>
 {
 	FVec3 X;
 	FQuat R;
@@ -55,7 +55,7 @@ struct FDirtyRigidParticleData : public TBasePullData<FSingleParticlePhysicsProx
 	EObjectStateType ObjectState;
 };
 
-struct FDirtyGeometryCollectionData : public TBasePullData<FGeometryCollectionPhysicsProxy>
+struct FDirtyGeometryCollectionData : public TBasePullData<FGeometryCollectionPhysicsProxy, FProxyTimestampBase>
 {
 	FGeometryCollectionResults Results;
 };
@@ -70,7 +70,7 @@ struct FJointConstraintOutputData {
 
 class FJointConstraint;
 
-struct FDirtyJointConstraintData : public TBasePullData<FJointConstraintPhysicsProxy>
+struct FDirtyJointConstraintData : public TBasePullData<FJointConstraintPhysicsProxy, FProxyTimestampBase>
 {
 	FJointConstraintOutputData OutputData;
 };
diff --git a/Engine/Source/Runtime/Experimental/Chaos/Public/GeometryCollectionProxyData.h b/Engine/Source/Runtime/Experimental/Chaos/Public/GeometryCollectionProxyData.h
--- a/Engine/Source/Runtime/Experimental/Chaos/Public/GeometryCollectionProxyData.h
+++ b/Engine/Source/Runtime/Experimental/Chaos/Public/GeometryCollectionProxyData.h
@@ -132,8 +132,11 @@ public:
 		const int32 NumTransforms = Other.NumElements(FGeometryCollection::TransformGroup);
 		DisabledStates.SetNumUninitialized(NumTransforms);
 		GlobalTransforms.SetNumUninitialized(NumTransforms);
-		ParticleToWorldTransforms.SetNumUninitialized(NumTransforms);
-
+		ParticleXs.SetNumUninitialized(NumTransforms);
+		ParticleRs.SetNumUninitialized(NumTransforms);
+		ParticleVs.SetNumUninitialized(NumTransforms);
+		ParticleWs.SetNumUninitialized(NumTransforms);
+		
 		Transforms.SetNumUninitialized(NumTransforms);
 		Parent.SetNumUninitialized(NumTransforms);
 		DynamicState.SetNumUninitialized(NumTransforms);
@@ -142,7 +145,10 @@ public:
 	Chaos::FReal SolverDt;
 	TArray<bool> DisabledStates;
 	TArray<FMatrix> GlobalTransforms;
-	TArray<FTransform> ParticleToWorldTransforms;
+	TArray<Chaos::FVec3> ParticleXs;
+	TArray<Chaos::FRotation3> ParticleRs;
+	TArray<Chaos::FVec3> ParticleVs;
+	TArray<Chaos::FVec3> ParticleWs;
 
 	TArray<FTransform> Transforms;
 	TArray<int32> Parent;
diff --git a/Engine/Source/Runtime/Experimental/Chaos/Public/PhysicsProxy/GeometryCollectionPhysicsProxy.h b/Engine/Source/Runtime/Experimental/Chaos/Public/PhysicsProxy/GeometryCollectionPhysicsProxy.h
--- a/Engine/Source/Runtime/Experimental/Chaos/Public/PhysicsProxy/GeometryCollectionPhysicsProxy.h
+++ b/Engine/Source/Runtime/Experimental/Chaos/Public/PhysicsProxy/GeometryCollectionPhysicsProxy.h
@@ -54,10 +54,10 @@ private:
  * Class to manage sharing data between the game thread and the simulation thread 
  * (which may not be different than the game thread) for a \c FGeometryDynamicCollection.
  */
-class CHAOS_API FGeometryCollectionPhysicsProxy : public TPhysicsProxy<FGeometryCollectionPhysicsProxy, FStubGeometryCollectionData>
+class CHAOS_API FGeometryCollectionPhysicsProxy : public TPhysicsProxy<FGeometryCollectionPhysicsProxy, FStubGeometryCollectionData, FGeometryCollectionProxyTimestamp>
 {
 public:
-	typedef TPhysicsProxy<FGeometryCollectionPhysicsProxy, FStubGeometryCollectionData> Base;
+	typedef TPhysicsProxy<FGeometryCollectionPhysicsProxy, FStubGeometryCollectionData, FGeometryCollectionProxyTimestamp> Base;
 	typedef FCollisionStructureManager::FSimplicial FSimplicial;
 	typedef Chaos::TPBDRigidParticleHandle<Chaos::FReal, 3> FParticleHandle;
 	typedef Chaos::TPBDRigidClusteredParticleHandle<Chaos::FReal, 3> FClusterHandle;
@@ -116,8 +116,11 @@ public:
 	/** Called at the end of \c FPBDRigidsSolver::PushPhysicsStateExec(). */
 	void ClearAccumulatedData() {}
 
-	/** Push physics state into the \c PhysToGameInterchange. */
-	void BufferPhysicsResults(Chaos::FPBDRigidsSolver* CurrentSolver, Chaos::FDirtyGeometryCollectionData& BufferData);
+	/** Push PT state into the \c PhysToGameInterchange. */
+	void BufferPhysicsResults_Internal(Chaos::FPBDRigidsSolver* CurrentSolver, Chaos::FDirtyGeometryCollectionData& BufferData);
+
+	/** Push GT state into the \c PhysToGameInterchange for async physics */
+	void BufferPhysicsResults_External(Chaos::FDirtyGeometryCollectionData& BufferData);
 
 	/** Push data from the game thread to the physics thread */
 	void PushStateOnGameThread(Chaos::FPBDRigidsSolver* InSolver);
@@ -132,7 +135,7 @@ public:
 	 * Pulls data out of the PhysToGameInterchange and updates \c GTDynamicCollection. 
 	 * Called from FPhysScene_ChaosInterface::SyncBodies(), NOT the solver.
 	 */
-	bool PullFromPhysicsState(const Chaos::FDirtyGeometryCollectionData& BufferData, const int32 SolverSyncTimestamp);
+	bool PullFromPhysicsState(const Chaos::FDirtyGeometryCollectionData& BufferData, const int32 SolverSyncTimestamp, const Chaos::FDirtyGeometryCollectionData* NextPullData = nullptr, const Chaos::FRealSingle* Alpha= nullptr);
 
 	bool IsDirty() { return false; }
 
@@ -239,6 +242,9 @@ public:
 
 	void DisableParticles(TArray<int32>& TransformGroupIndices);
 
+	FProxyInterpolationData& GetInterpolationData() { return InterpolationData; }
+	const FProxyInterpolationData& GetInterpolationData() const { return InterpolationData; }
+
 protected:
 	/**
 	 * Build a physics thread cluster parent particle.
@@ -263,7 +269,9 @@ protected:
 
 	void InitializeRemoveOnFracture(FParticlesType& Particles, const TManagedArray<int32>& DynamicState);
 
-	void SetClusteredParticleKinematicTarget(Chaos::FPBDRigidClusteredParticleHandle* Handle, const FTransform& WorldTransform);
+	void SetClusteredParticleKinematicTarget_Internal(Chaos::FPBDRigidClusteredParticleHandle* Handle, const FTransform& WorldTransform);
+
+	void PrepareBufferData(Chaos::FDirtyGeometryCollectionData& BufferData, const FGeometryDynamicCollection& ThreadCollection,  Chaos::FReal SolverLastDt = 0.0);
 
 private:
 
@@ -343,6 +351,7 @@ private:
 	// paradigm, at least for this component of the handshake.
 	Chaos::FGuardedTripleBuffer<FGeometryCollectionResults> PhysToGameInterchange;
 
+	FProxyInterpolationData InterpolationData; 
 };
 
 CHAOS_API Chaos::FTriangleMesh* CreateTriangleMesh(const int32 FaceStart,const int32 FaceCount,const TManagedArray<bool>& Visible,const TManagedArray<FIntVector>& Indices, bool bRotateWinding = true);
diff --git a/Engine/Source/Runtime/Experimental/Chaos/Public/PhysicsProxy/SingleParticlePhysicsProxy.h b/Engine/Source/Runtime/Experimental/Chaos/Public/PhysicsProxy/SingleParticlePhysicsProxy.h
--- a/Engine/Source/Runtime/Experimental/Chaos/Public/PhysicsProxy/SingleParticlePhysicsProxy.h
+++ b/Engine/Source/Runtime/Experimental/Chaos/Public/PhysicsProxy/SingleParticlePhysicsProxy.h
@@ -70,18 +70,8 @@ public:
 	FSingleParticlePhysicsProxy(FSingleParticlePhysicsProxy&&) = delete;
 	virtual ~FSingleParticlePhysicsProxy();
 
-	void SetPullDataInterpIdx_External(const int32 Idx)
-	{
-		PullDataInterpIdx_External = Idx;
-	}
-
-	void SetResimSmoothing(bool ResimSmoothing) { bResimSmoothing = ResimSmoothing; }
-	bool IsResimSmoothing() const { return bResimSmoothing; }
-
-	int32 GetInterpChannel_External() const { return InterpChannel_External; }
-	void SetInterpChannel_External(int32 Channel) { InterpChannel_External = Channel; }
-
-	int32 GetPullDataInterpIdx_External() const { return PullDataInterpIdx_External; }
+	const FProxyInterpolationData& GetInterpolationData() const { return InterpolationData; }
+	FProxyInterpolationData& GetInterpolationData() { return InterpolationData; }
 
 	FORCEINLINE FRigidBodyHandle_External& GetGameThreadAPI()
 	{
@@ -180,11 +170,7 @@ protected:
 	FParticleHandle* Handle;
 
 private:
-
-	//Used by interpolation code
-	int32 PullDataInterpIdx_External;
-	int32 InterpChannel_External = 0;
-	bool bResimSmoothing = false;
+	FProxyInterpolationData InterpolationData;
 
 	//use static Create
 	FSingleParticlePhysicsProxy(TUniquePtr<PARTICLE_TYPE>&& InParticle, FParticleHandle* InHandle, UObject* InOwner = nullptr);
@@ -964,16 +950,16 @@ public:
 	{
 		VerifyContext();
 		SetXBase(InX, bInvalidate);
-		SyncTimestamp->XTimestamp = GetSolverSyncTimestamp_External();
-		SyncTimestamp->OverWriteX = InX;
+		FSingleParticleProxyTimestamp& SyncTS = GetSyncTimestampAs<FSingleParticleProxyTimestamp>(); 
+		SyncTS.OverWriteX.Set(GetSolverSyncTimestamp_External(), InX);
 	}
 
 	void SetR(const FRotation3& InR, bool bInvalidate = true)
 	{
 		VerifyContext();
 		SetRBase(InR, bInvalidate);
-		SyncTimestamp->RTimestamp = GetSolverSyncTimestamp_External();
-		SyncTimestamp->OverWriteR = InR;
+		FSingleParticleProxyTimestamp& SyncTS = GetSyncTimestampAs<FSingleParticleProxyTimestamp>();
+		SyncTS.OverWriteR.Set(GetSolverSyncTimestamp_External(), InR);
 	}
 
 	void SetV(const FVec3& InV, bool bInvalidate = true)
@@ -986,9 +972,8 @@ public:
 			//but we also want position to snap to where it currently is on external thread
 			SetX(X(), bInvalidate);
 		}
-
-		SyncTimestamp->VTimestamp = GetSolverSyncTimestamp_External();
-		SyncTimestamp->OverWriteV = InV;
+		FSingleParticleProxyTimestamp& SyncTS = GetSyncTimestampAs<FSingleParticleProxyTimestamp>();
+		SyncTS.OverWriteV.Set(GetSolverSyncTimestamp_External(), InV);
 	}
 
 	void SetW(const FVec3& InW, bool bInvalidate = true)
@@ -1002,9 +987,8 @@ public:
 			//but we also want position to snap to where it currently is on external thread
 			SetR(R(), bInvalidate);
 		}
-
-		SyncTimestamp->WTimestamp = GetSolverSyncTimestamp_External();
-		SyncTimestamp->OverWriteW = InW;
+		FSingleParticleProxyTimestamp& SyncTS = GetSyncTimestampAs<FSingleParticleProxyTimestamp>();
+		SyncTS.OverWriteW.Set(GetSolverSyncTimestamp_External(), InW);
 	}
 
 	void SetObjectState(const EObjectStateType InState, bool bAllowEvents = false, bool bInvalidate = true)
@@ -1012,7 +996,8 @@ public:
 		VerifyContext();
 		if (auto Rigid = GetParticle_LowLevel()->CastToRigidParticle())
 		{
-			SyncTimestamp->ObjectStateTimestamp = GetSolverSyncTimestamp_External();
+			FSingleParticleProxyTimestamp& SyncTS = GetSyncTimestampAs<FSingleParticleProxyTimestamp>();
+			SyncTS.ObjectStateTimestamp = GetSolverSyncTimestamp_External();
 			if (InState != EObjectStateType::Dynamic && Rigid->ObjectState() == EObjectStateType::Dynamic)
 			{
 				//we want to snap the particle to its current state on the external thread. This is because the user wants the object to fully stop right now
diff --git a/Engine/Source/Runtime/Experimental/Chaos/Public/PhysicsProxy/SkeletalMeshPhysicsProxy.h b/Engine/Source/Runtime/Experimental/Chaos/Public/PhysicsProxy/SkeletalMeshPhysicsProxy.h
--- a/Engine/Source/Runtime/Experimental/Chaos/Public/PhysicsProxy/SkeletalMeshPhysicsProxy.h
+++ b/Engine/Source/Runtime/Experimental/Chaos/Public/PhysicsProxy/SkeletalMeshPhysicsProxy.h
@@ -108,9 +108,9 @@ struct FSkeletalMeshPhysicsProxyOutputs : public Chaos::FParticleData
 };
 
 
-class CHAOS_API FSkeletalMeshPhysicsProxy : public TPhysicsProxy<FSkeletalMeshPhysicsProxy, FSkeletalMeshPhysicsProxyOutputs>
+class CHAOS_API FSkeletalMeshPhysicsProxy : public TPhysicsProxy<FSkeletalMeshPhysicsProxy, FSkeletalMeshPhysicsProxyOutputs, FProxyTimestampBase>
 {
-	typedef TPhysicsProxy<FSkeletalMeshPhysicsProxy, FSkeletalMeshPhysicsProxyOutputs> Base;
+	typedef TPhysicsProxy<FSkeletalMeshPhysicsProxy, FSkeletalMeshPhysicsProxyOutputs, FProxyTimestampBase> Base;
 public:
 
 
diff --git a/Engine/Source/Runtime/Experimental/Chaos/Public/PhysicsProxy/StaticMeshPhysicsProxy.h b/Engine/Source/Runtime/Experimental/Chaos/Public/PhysicsProxy/StaticMeshPhysicsProxy.h
--- a/Engine/Source/Runtime/Experimental/Chaos/Public/PhysicsProxy/StaticMeshPhysicsProxy.h
+++ b/Engine/Source/Runtime/Experimental/Chaos/Public/PhysicsProxy/StaticMeshPhysicsProxy.h
@@ -16,9 +16,9 @@ struct FStubSkeletalMeshData : public Chaos::FParticleData {
 	void Reset() { };
 };
 
-class CHAOS_API FStaticMeshPhysicsProxy : public TPhysicsProxy<FStaticMeshPhysicsProxy, FStubSkeletalMeshData>
+class CHAOS_API FStaticMeshPhysicsProxy : public TPhysicsProxy<FStaticMeshPhysicsProxy, FStubSkeletalMeshData, FProxyTimestampBase>
 {
-	typedef  TPhysicsProxy<FStaticMeshPhysicsProxy, FStubSkeletalMeshData> Base;
+	typedef  TPhysicsProxy<FStaticMeshPhysicsProxy, FStubSkeletalMeshData, FProxyTimestampBase> Base;
 public:
 
 	struct FShapeParams
