Chaos : make sure to account for AccelerationStructureSplitStaticAndDynamic cvar when creating GC particles

diff --git a/Engine/Source/Runtime/Experimental/Chaos/Private/PhysicsProxy/GeometryCollectionPhysicsProxy.cpp b/Engine/Source/Runtime/Experimental/Chaos/Private/PhysicsProxy/GeometryCollectionPhysicsProxy.cpp
--- a/Engine/Source/Runtime/Experimental/Chaos/Private/PhysicsProxy/GeometryCollectionPhysicsProxy.cpp
+++ b/Engine/Source/Runtime/Experimental/Chaos/Private/PhysicsProxy/GeometryCollectionPhysicsProxy.cpp
@@ -462,6 +462,15 @@ void FGeometryCollectionPhysicsProxy::Initialize(Chaos::FPBDRigidsEvolutionBase
 			// this is causing a mismatch when using the separate acceleration structures optimization which can cause crashes when destroying the particle while async tracing 
 			// todo(chaos) we should eventually refactor this code to use rigid particles on the GT side for geometry collection  
 			P->SetSpatialIdx(Chaos::FSpatialAccelerationIdx{ 0,1 });
+
+			if (Chaos::AccelerationStructureSplitStaticAndDynamic == 1)
+			{
+				P->SetSpatialIdx(Chaos::FSpatialAccelerationIdx{ 0,1 });
+			}
+			else
+			{
+				P->SetSpatialIdx(Chaos::FSpatialAccelerationIdx{ 0,0 });
+			}
 		}
 
 		if (bGeometryCollectionAlwaysGenerateGTCollisionForClusters)
