GJK reports a hit but with an invalid time when shapes are separated by a large distance

diff --git a/Engine/Source/Programs/HeadlessChaos/Private/HeadlessChaosApp.cpp b/Engine/Source/Programs/HeadlessChaos/Private/HeadlessChaosApp.cpp
--- a/Engine/Source/Programs/HeadlessChaos/Private/HeadlessChaosApp.cpp
+++ b/Engine/Source/Programs/HeadlessChaos/Private/HeadlessChaosApp.cpp
@@ -192,7 +192,7 @@ TEST(SweepTests, Sweep) {
 	SUCCEED();
 }
 
-TEST(SweepTests, DISABLED_LargeSweep)
+TEST(SweepTests, LargeSweep)
 {
 	ChaosTest::GJKLargeDistanceCapsuleSweep();
 
diff --git a/Engine/Source/Programs/HeadlessChaos/Private/HeadlessChaosTestGJK.cpp b/Engine/Source/Programs/HeadlessChaos/Private/HeadlessChaosTestGJK.cpp
--- a/Engine/Source/Programs/HeadlessChaos/Private/HeadlessChaosTestGJK.cpp
+++ b/Engine/Source/Programs/HeadlessChaos/Private/HeadlessChaosTestGJK.cpp
@@ -1594,9 +1594,11 @@ namespace ChaosTest
 		FVec3 OutNorm;
 
 		// Should fail and give a valid time
-		GJKRaycast2(A, B, BToA, LocalDir, Length, OutTime, OutLoc, OutNorm, Thickness, bComputeMtd, InitialDir, Thickness);
+		bool bHit = GJKRaycast2(A, B, BToA, LocalDir, Length, OutTime, OutLoc, OutNorm, Thickness, bComputeMtd, InitialDir, Thickness);
 
-		// Expect to recieve a valid time.
+		EXPECT_FALSE(bHit);
+
+		// Expect to receive a valid time.
 		EXPECT_TRUE(FMath::IsFinite(OutTime));
 	}
 
diff --git a/Engine/Source/Runtime/Experimental/Chaos/Public/Chaos/EPA.h b/Engine/Source/Runtime/Experimental/Chaos/Public/Chaos/EPA.h
--- a/Engine/Source/Runtime/Experimental/Chaos/Public/Chaos/EPA.h
+++ b/Engine/Source/Runtime/Experimental/Chaos/Public/Chaos/EPA.h
@@ -406,6 +406,7 @@ enum class EEPAResult
 	MaxIterations,				// We have a contact point, but did not reach the target tolerance before we hit the iteration limit - result accuracy is unknown
 	Degenerate,					// We hit a degenerate condition in EPA which prevents a solution from being generated (result is invalid but objects may be penetrating)
 	BadInitialSimplex,			// The initial setup did not provide a polytope containing the origin (objects are separated)
+	NoValidContact,             // No valid contact have been found, no contacts will be returned
 };
 
 inline const bool IsEPASuccess(EEPAResult EPAResult)
diff --git a/Engine/Source/Runtime/Experimental/Chaos/Public/Chaos/EPAVectorized.h b/Engine/Source/Runtime/Experimental/Chaos/Public/Chaos/EPAVectorized.h
--- a/Engine/Source/Runtime/Experimental/Chaos/Public/Chaos/EPAVectorized.h
+++ b/Engine/Source/Runtime/Experimental/Chaos/Public/Chaos/EPAVectorized.h
@@ -360,7 +360,7 @@ FORCEINLINE_DEBUGGABLE void VectorEPAComputeVisibilityBorder(TEPAWorkingArray<Ve
 }
 
 
-FORCEINLINE_DEBUGGABLE void VectorComputeEPAResults(const VectorRegister4Float* VertsA, const VectorRegister4Float* VertsB, const VectorTEPAEntry& Entry, VectorRegister4Float& OutPenetration, VectorRegister4Float& OutDir, VectorRegister4Float& OutA, VectorRegister4Float& OutB)
+FORCEINLINE_DEBUGGABLE void VectorComputeEPAResults(const VectorRegister4Float* VertsA, const VectorRegister4Float* VertsB, const VectorTEPAEntry& Entry, VectorRegister4Float& OutPenetration, VectorRegister4Float& OutDir, VectorRegister4Float& OutA, VectorRegister4Float& OutB, EEPAResult& ResultStatus)
 {
 	//NOTE: We use this function as fallback when robustness breaks. So - do not assume adjacency is valid as these may be new uninitialized traingles that failed
 
@@ -371,7 +371,16 @@ FORCEINLINE_DEBUGGABLE void VectorComputeEPAResults(const VectorRegister4Float*
 	constexpr VectorRegister4Int three = MakeVectorRegisterIntConstant(3, 3, 3, 3);
 	VectorRegister4Int NumVerts = three;
 	OutDir = VectorSimplexFindClosestToOrigin(Simplex, NumVerts, Barycentric, As, Bs);
-	OutPenetration = VectorSqrt(VectorDot4(OutDir, OutDir));
+
+	const VectorRegister4Float DotDir = VectorDot4(OutDir, OutDir);
+
+	if (VectorContainsNaNOrInfinite(DotDir))
+	{
+		OutPenetration = VectorZero();
+		ResultStatus = EEPAResult::NoValidContact;
+	}
+
+	OutPenetration = VectorSqrt(DotDir);
 
 	constexpr float EpsFloat = 1e-4f;
 	constexpr VectorRegister4Float Eps = MakeVectorRegisterFloatConstant(EpsFloat, EpsFloat, EpsFloat, EpsFloat);
@@ -502,16 +511,15 @@ EEPAResult VectorEPA(TArray<VectorRegister4Float>& VertsABuffer, TArray<VectorRe
 		const VectorRegister4Float W = VectorSubtract(ASupport, BSupport);
 		const VectorRegister4Float DistanceToSupportPlane = VectorDot3(Entry.PlaneNormal, W);
 
-		const VectorRegister4Float UpCoundIsGTDist = VectorCompareGT(UpperBound, DistanceToSupportPlane);
-
 		//Remember the entry that gave us the lowest upper bound and use it in case we have to terminate early
 		//This can result in very deep planes. Ideally we'd just use the plane formed at W, but it's not clear how you get back points in A, B for that
 		//BestEntry = Entry;
-		UpperBound = VectorSelect(UpCoundIsGTDist, DistanceToSupportPlane, UpperBound);
+		UpperBound = VectorMin(DistanceToSupportPlane, UpperBound);
 		LowerBound = Entry.Distance;
 
-		const VectorRegister4Float EpsIsGEAbsBounds = VectorCompareGE(Eps, VectorAbs(VectorSubtract(UpperBound, LowerBound)));
-
+		
+		const VectorRegister4Float DiffUpLow = VectorAbs(VectorSubtract(UpperBound, LowerBound));
+		const VectorRegister4Float EpsIsGEAbsBounds = VectorCompareGE(Eps, DiffUpLow);
 
 
 		//It's possible the origin is not contained by the CSO. In this case the upper bound will be negative, at which point we should just exit. Maybe return a different enum value?
@@ -605,7 +613,7 @@ EEPAResult VectorEPA(TArray<VectorRegister4Float>& VertsABuffer, TArray<VectorRe
 		}
 	}
 
-	VectorComputeEPAResults(VertsABuffer.GetData(), VertsBBuffer.GetData(), LastEntry, OutPenetration, OutDir, WitnessA, WitnessB);
+	VectorComputeEPAResults(VertsABuffer.GetData(), VertsBBuffer.GetData(), LastEntry, OutPenetration, OutDir, WitnessA, WitnessB, ResultStatus);
 	
 	return ResultStatus;
 }
diff --git a/Engine/Source/Runtime/Experimental/Chaos/Public/Chaos/GJK.h b/Engine/Source/Runtime/Experimental/Chaos/Public/Chaos/GJK.h
--- a/Engine/Source/Runtime/Experimental/Chaos/Public/Chaos/GJK.h
+++ b/Engine/Source/Runtime/Experimental/Chaos/Public/Chaos/GJK.h
@@ -2029,6 +2029,10 @@ namespace Chaos
 						OutTime = -OutTime - (MarginA + MarginB);
 						OutPosition = ClosestA;
 					}
+					else if (EPAResult == EEPAResult::NoValidContact)
+					{
+						return false;
+					}
 					else
 					{
 						//assume touching hit
