Fix LWC in BroadPhase.

diff --git a/Engine/Source/Runtime/Experimental/Chaos/Public/Chaos/AABBTree.h b/Engine/Source/Runtime/Experimental/Chaos/Public/Chaos/AABBTree.h
--- a/Engine/Source/Runtime/Experimental/Chaos/Public/Chaos/AABBTree.h
+++ b/Engine/Source/Runtime/Experimental/Chaos/Public/Chaos/AABBTree.h
@@ -152,12 +152,6 @@ struct TAABBTreeIntersectionHelper
 		check(false);
 		return true;
 	}
-	FORCEINLINE_DEBUGGABLE static bool IntersectsVectorized(const VectorRegister4Float& Start, FQueryFastData& QueryFastData, const FAABBVectorized& Bounds, const FAABBVectorized& QueryBounds,
-		const VectorRegister4Float& QueryHalfExtents, const VectorRegister4Float& InvDir, const VectorRegister4Float& Parallel)
-	{
-		check(false);
-		return true;
-	}
 };
 
 template<>
@@ -169,12 +163,6 @@ struct TAABBTreeIntersectionHelper<FQueryFastData, EAABBQueryType::Raycast>
 		return Bounds.RaycastFast(Start, Dir, InvDir, bParallel, QueryFastData.CurrentLength, QueryFastData.InvCurrentLength, TOI, OutPosition);
 	}
 
-	FORCEINLINE_DEBUGGABLE static bool IntersectsVectorized(const VectorRegister4Float& Start, FQueryFastData& QueryFastData, const FAABBVectorized& Bounds, const FAABBVectorized& QueryBounds,
-		const VectorRegister4Float& QueryHalfExtents, const VectorRegister4Float& InvDir, const VectorRegister4Float& Parallel)
-	{
-		check(false);
-		return true;
-	}
 };
 
 template <>
@@ -187,14 +175,6 @@ struct TAABBTreeIntersectionHelper<FQueryFastData, EAABBQueryType::Sweep>
 		return SweepBounds.RaycastFast(Start, Dir, InvDir, bParallel, QueryFastData.CurrentLength, QueryFastData.InvCurrentLength, TOI, OutPosition);
 	}
 
-	FORCEINLINE_DEBUGGABLE static bool IntersectsVectorized(const VectorRegister4Float& Start, FQueryFastData& QueryFastData, const FAABBVectorized& Bounds, const FAABBVectorized& QueryBounds,
-															const VectorRegister4Float& QueryHalfExtents, const VectorRegister4Float& InvDir, const VectorRegister4Float& Parallel)
-	{
-		VectorRegister4Float CurrentLength = MakeVectorRegisterFloatFromDouble(VectorLoadDouble1(&QueryFastData.CurrentLength));
-		FAABBVectorized SweepBounds(VectorSubtract(Bounds.GetMin(), QueryHalfExtents), VectorAdd(Bounds.GetMax(), QueryHalfExtents));
-		return SweepBounds.RaycastFast(Start, InvDir, Parallel, CurrentLength);
-	}
-
 };
 
 template <>
@@ -334,12 +314,8 @@ struct TAABBTreeLeafArray : public TBoundsWrapperHelper<TPayloadType, T, bComput
 	FORCEINLINE_DEBUGGABLE bool RaycastSweepImp(const TVec3<T>& Start, TQueryFastData& QueryFastData, const TVec3<T>& QueryHalfExtents, TSQVisitor& Visitor, const TVec3<T>& Dir, const TVec3<T> InvDir, const bool bParallel[3]) const
 	{
 		PHYSICS_CSV_CUSTOM_VERY_EXPENSIVE(PhysicsCounters, MaxLeafSize, Elems.Num(), ECsvCustomStatOp::Max);
-
-		const VectorRegister4Float DirSimd = MakeVectorRegisterFloatFromDouble(MakeVectorRegisterDouble(Dir.X, Dir.Y, Dir.Z, 0.0));
-		const VectorRegister4Float InvDirSimd = MakeVectorRegisterFloatFromDouble(MakeVectorRegisterDouble(InvDir.X, InvDir.Y, InvDir.Z, 0.0));
-		const VectorRegister4Float QueryHalfExtentsSimd = MakeVectorRegisterFloatFromDouble(MakeVectorRegisterDouble(QueryHalfExtents.X, QueryHalfExtents.Y, QueryHalfExtents.Z, 0.0));
-		const VectorRegister4Float Parallel = VectorCompareGT(GlobalVectorConstants::SmallNumber, VectorAbs(DirSimd));
-		
+		FVec3 TmpPosition;
+		FReal TOI;
 		for (const auto& Elem : Elems)
 		{
 			if (PrePreFilterHelper(Elem.Payload, Visitor))
@@ -348,41 +324,17 @@ struct TAABBTreeLeafArray : public TBoundsWrapperHelper<TPayloadType, T, bComput
 			}
 
 			const FAABB3 InstanceBounds(Elem.Bounds.Min(), Elem.Bounds.Max());
-			if (bSweep == true)
+			if (TAABBTreeIntersectionHelper<TQueryFastData, bSweep ? EAABBQueryType::Sweep :
+				EAABBQueryType::Raycast>::Intersects(Start, QueryFastData, TOI, TmpPosition, InstanceBounds, FAABB3(), QueryHalfExtents, Dir, InvDir, bParallel))
 			{
-				const VectorRegister4Float Min = MakeVectorRegisterFloatFromDouble(MakeVectorRegisterDouble(InstanceBounds.Min().X, InstanceBounds.Min().Y, InstanceBounds.Min().Z, 0.0));
-				const VectorRegister4Float Max = MakeVectorRegisterFloatFromDouble(MakeVectorRegisterDouble(InstanceBounds.Max().X, InstanceBounds.Max().Y, InstanceBounds.Max().Z, 0.0));
-				const FAABBVectorized InstanceBoundsSimd(Min, Max);
-				VectorRegister4Float StartSimd = MakeVectorRegisterFloatFromDouble(MakeVectorRegisterDouble(Start.X, Start.Y, Start.Z, 0.0));
-
-				if (TAABBTreeIntersectionHelper<TQueryFastData, EAABBQueryType::Sweep>::IntersectsVectorized(
-					StartSimd, QueryFastData, InstanceBoundsSimd, FAABBVectorized(), QueryHalfExtentsSimd, InvDirSimd, Parallel))
+				TSpatialVisitorData<TPayloadType> VisitData(Elem.Payload, true, InstanceBounds);
+				const bool bContinue = (bSweep && Visitor.VisitSweep(VisitData, QueryFastData)) || (!bSweep && Visitor.VisitRaycast(VisitData, QueryFastData));
+				if (!bContinue)
 				{
-					TSpatialVisitorData<TPayloadType> VisitData(Elem.Payload, true, InstanceBounds);
-					const bool bContinue = (bSweep && Visitor.VisitSweep(VisitData, QueryFastData)) || (!bSweep && Visitor.VisitRaycast(VisitData, QueryFastData));
-					if (!bContinue)
-					{
-						return false;
-					}
-				}
-			}
-			else 
-			{
-				FVec3 TmpPosition;
-				FReal TOI;
-				if (TAABBTreeIntersectionHelper<TQueryFastData,
-					EAABBQueryType::Raycast>::Intersects(Start, QueryFastData, TOI, TmpPosition, InstanceBounds, FAABB3(), QueryHalfExtents, Dir, InvDir, bParallel))
-				{
-					TSpatialVisitorData<TPayloadType> VisitData(Elem.Payload, true, InstanceBounds);
-					const bool bContinue = (bSweep && Visitor.VisitSweep(VisitData, QueryFastData)) || (!bSweep && Visitor.VisitRaycast(VisitData, QueryFastData));
-					if (!bContinue)
-					{
-						return false;
-					}
+					return false;
 				}
 			}
 		}
-
 		return true;
 	}
 
