Dependency cleanup in Chaos focused on removing includes from ChaosSolverConfiguration.h.
Takes Chaos out of the top include times in ShooterGame builds.

diff --git a/Engine/Source/Runtime/Engine/Private/Collision/SceneQuery.cpp b/Engine/Source/Runtime/Engine/Private/Collision/SceneQuery.cpp
--- a/Engine/Source/Runtime/Engine/Private/Collision/SceneQuery.cpp
+++ b/Engine/Source/Runtime/Engine/Private/Collision/SceneQuery.cpp
@@ -27,6 +27,7 @@ float DebugLineLifetime = 2.f;
 
 #include "ChaosSolversModule.h"
 #include "Physics/Experimental/ChaosInterfaceWrapper.h"
+#include "PBDRigidsSolver.h"
 
 CSV_DEFINE_CATEGORY(SceneQuery, false);
 
diff --git a/Engine/Source/Runtime/Experimental/Chaos/Private/ChaosSolverConfiguration.cpp b/Engine/Source/Runtime/Experimental/Chaos/Private/ChaosSolverConfiguration.cpp
new file mode 100644
--- /dev/null
+++ b/Engine/Source/Runtime/Experimental/Chaos/Private/ChaosSolverConfiguration.cpp
@@ -0,0 +1,26 @@
+// Copyright Epic Games, Inc. All Rights Reserved.
+
+#include "ChaosSolverConfiguration.h"
+#include "Chaos/PBDRigidsEvolutionGBF.h"
+
+using FEvolution = Chaos::FPBDRigidsEvolutionGBF;
+
+FChaosSolverConfiguration::FChaosSolverConfiguration()
+	: Iterations(FEvolution::DefaultNumIterations)
+	, CollisionPairIterations(FEvolution::DefaultNumCollisionPairIterations)
+	, PushOutIterations(FEvolution::DefaultNumPushOutIterations)
+	, CollisionPushOutPairIterations(FEvolution::DefaultNumCollisionPushOutPairIterations)
+	, CollisionMarginFraction(FEvolution::DefaultCollisionMarginFraction)
+	, CollisionMarginMax(FEvolution::DefaultCollisionMarginMax)
+	, CollisionCullDistance(FEvolution::DefaultCollisionCullDistance)
+	, CollisionMaxPushOutVelocity(FEvolution::DefaultCollisionMaxPushOutVelocity)
+	, JointPairIterations(FEvolution::DefaultNumJointPairIterations)
+	, JointPushOutPairIterations(FEvolution::DefaultNumJointPushOutPairIterations)
+	, ClusterConnectionFactor(1.0f)
+	, ClusterUnionConnectionType(EClusterUnionMethod::DelaunayTriangulation)
+	, bGenerateCollisionData(false)
+	, bGenerateBreakData(false)
+	, bGenerateTrailingData(false)
+	, bGenerateContactGraph(true)
+{
+}
diff --git a/Engine/Source/Runtime/Experimental/Chaos/Public/Chaos/KinematicTargets.h b/Engine/Source/Runtime/Experimental/Chaos/Public/Chaos/KinematicTargets.h
--- a/Engine/Source/Runtime/Experimental/Chaos/Public/Chaos/KinematicTargets.h
+++ b/Engine/Source/Runtime/Experimental/Chaos/Public/Chaos/KinematicTargets.h
@@ -4,6 +4,7 @@
 #include "Chaos/ChaosArchive.h"
 #include "Chaos/Transform.h"
 #include "UObject/FortniteReleaseBranchCustomObjectVersion.h"
+#include "UObject/PhysicsObjectVersion.h"
 
 namespace Chaos
 {
diff --git a/Engine/Source/Runtime/Experimental/Chaos/Public/ChaosSolverConfiguration.h b/Engine/Source/Runtime/Experimental/Chaos/Public/ChaosSolverConfiguration.h
--- a/Engine/Source/Runtime/Experimental/Chaos/Public/ChaosSolverConfiguration.h
+++ b/Engine/Source/Runtime/Experimental/Chaos/Public/ChaosSolverConfiguration.h
@@ -3,8 +3,10 @@
 #pragma once
 
 #include "SolverEventFilters.h"
+#if UE_ENABLE_INCLUDE_ORDER_DEPRECATED_IN_5_1
 #include "Chaos/PBDRigidsEvolutionGBF.h"
 #include "PBDRigidsSolver.h"
+#endif
 #include "ChaosSolverConfiguration.generated.h"
 
 UENUM()
@@ -22,25 +24,7 @@ struct FChaosSolverConfiguration
 {
 	GENERATED_BODY();
 
-	FChaosSolverConfiguration()
-		: Iterations(FEvolution::DefaultNumIterations)
-		, CollisionPairIterations(FEvolution::DefaultNumCollisionPairIterations)
-		, PushOutIterations(FEvolution::DefaultNumPushOutIterations)
-		, CollisionPushOutPairIterations(FEvolution::DefaultNumCollisionPushOutPairIterations)
-		, CollisionMarginFraction(FEvolution::DefaultCollisionMarginFraction)
-		, CollisionMarginMax(FEvolution::DefaultCollisionMarginMax)
-		, CollisionCullDistance(FEvolution::DefaultCollisionCullDistance)
-		, CollisionMaxPushOutVelocity(FEvolution::DefaultCollisionMaxPushOutVelocity)
-		, JointPairIterations(FEvolution::DefaultNumJointPairIterations)
-		, JointPushOutPairIterations(FEvolution::DefaultNumJointPushOutPairIterations)
-		, ClusterConnectionFactor(1.0f)
-		, ClusterUnionConnectionType(EClusterUnionMethod::DelaunayTriangulation)
-		, bGenerateCollisionData(false)
-		, bGenerateBreakData(false)
-		, bGenerateTrailingData(false)
-		, bGenerateContactGraph(true)
-	{
-	}
+	CHAOS_API FChaosSolverConfiguration();
 
 	// The number of iterations to run during the constraint solver step
 	UPROPERTY(EditAnywhere, Category = "SolverConfiguration|Iterations")
@@ -93,7 +77,6 @@ struct FChaosSolverConfiguration
 	UPROPERTY(EditAnywhere, Category = "SolverConfiguration|Iterations")
 	int32 JointPushOutPairIterations;
 
-
 	UPROPERTY(EditAnywhere, Category = "SolverConfiguration|Clustering")
 	float ClusterConnectionFactor;
 
@@ -120,7 +103,4 @@ struct FChaosSolverConfiguration
 
 	UPROPERTY(EditAnywhere, Category = "SolverConfiguration|Contacts")
 	bool bGenerateContactGraph;
-
-private:
-	using FEvolution = Chaos::FPBDRigidsEvolutionGBF;
-};
\ No newline at end of file
+};
