Chaos - one-sided triangle collision for manifolds (tri mesh and heighfield)

diff --git a/Engine/Source/Runtime/Experimental/Chaos/Private/Chaos/CollisionOneShotManifolds.cpp b/Engine/Source/Runtime/Experimental/Chaos/Private/Chaos/CollisionOneShotManifolds.cpp
--- a/Engine/Source/Runtime/Experimental/Chaos/Private/Chaos/CollisionOneShotManifolds.cpp
+++ b/Engine/Source/Runtime/Experimental/Chaos/Private/Chaos/CollisionOneShotManifolds.cpp
@@ -82,6 +82,11 @@ namespace Chaos
 	bool bChaos_Collision_UseGJK2 = false;
 	FAutoConsoleVariableRef CVarChaos_Collision_UseGJK2(TEXT("p.Chaos.Collision.UseGJK2"), bChaos_Collision_UseGJK2, TEXT(""));
 
+	bool bChaos_Collision_OneSidedTriangleMesh = false;
+	bool bChaos_Collision_OneSidedHeightfield = false;
+	FAutoConsoleVariableRef CVarChaos_Collision_OneSidedTriangleMesh(TEXT("p.Chaos.Collision.OneSidedTriangleMesh"), bChaos_Collision_OneSidedTriangleMesh, TEXT(""));
+	FAutoConsoleVariableRef CVarChaos_Collision_OneSidedHeightfield(TEXT("p.Chaos.Collision.OneSidedHeightfield"), bChaos_Collision_OneSidedHeightfield, TEXT(""));
+
 	namespace Collisions
 	{
 		// Forward delarations we need from CollisionRestitution.cpp
@@ -1299,7 +1304,7 @@ namespace Chaos
 
 		// @todo(chaos): don't use GJK/EPA for triangle-convex so we can avoid using margains and avoid errors related to misidentifying edge contacts
 		template <typename ConvexType>
-		void ConstructPlanarConvexTriangleOneShotManifold(const ConvexType& Convex, const FTriangle& Triangle, const FReal CullDistance, TCArray<FContactPoint, 4>& OutContactPoints)
+		void ConstructPlanarConvexTriangleOneShotManifold(const ConvexType& Convex, const FTriangle& Triangle, const bool bOneSided, const FReal CullDistance, TCArray<FContactPoint, 4>& OutContactPoints)
 		{
 			SCOPE_CYCLE_COUNTER_MANIFOLD();
 			check(OutContactPoints.Num() == 0);
@@ -1333,6 +1338,12 @@ namespace Chaos
 			const FVec3 TrianglePlanePosition = Triangle[0];
 			const FReal TrianglePlaneNormalDotContactNormal = FVec3::DotProduct(TrianglePlaneNormal, GJKContactPoint.ShapeContactNormal);
 
+			// One sided triangles ignore collisions with the inside face
+			if (bOneSided && (TrianglePlaneNormalDotContactNormal < FReal(0)))
+			{
+				return;
+			}
+
 			// Find the best opposing plane on the convex
 			// @todo(chaos): handle zero margins...
 			const FReal PlaneSearchDistance = ConvexMargin;
@@ -1701,7 +1712,8 @@ namespace Chaos
 		template
 		void ConstructPlanarConvexTriangleOneShotManifold(
 			const FImplicitConvex3& Convex, 
-			const FTriangle& Triangle, 
+			const FTriangle& Triangle,
+			const bool bOneSided,
 			const FReal CullDistance, 
 			TCArray<FContactPoint, 4>& OutContactPoints);
 
@@ -1709,6 +1721,7 @@ namespace Chaos
 		void ConstructPlanarConvexTriangleOneShotManifold(
 			const TImplicitObjectInstanced<FImplicitConvex3>& Convex,
 			const FTriangle& Triangle,
+			const bool bOneSided,
 			const FReal CullDistance,
 			TCArray<FContactPoint, 4>& OutContactPoints);
 
@@ -1716,6 +1729,7 @@ namespace Chaos
 		void ConstructPlanarConvexTriangleOneShotManifold(
 			const TImplicitObjectScaled<FImplicitConvex3>& Convex,
 			const FTriangle& Triangle,
+			const bool bOneSided,
 			const FReal CullDistance,
 			TCArray<FContactPoint, 4>& OutContactPoints);
 
@@ -1723,6 +1737,7 @@ namespace Chaos
 		void ConstructPlanarConvexTriangleOneShotManifold(
 			const FImplicitBox3& Convex,
 			const FTriangle& Triangle,
+			const bool bOneSided,
 			const FReal CullDistance,
 			TCArray<FContactPoint, 4>& OutContactPoints);
 
@@ -1730,6 +1745,7 @@ namespace Chaos
 		void ConstructPlanarConvexTriangleOneShotManifold(
 			const TImplicitObjectScaled<FImplicitBox3>& Convex,
 			const FTriangle& Triangle,
+			const bool bOneSided,
 			const FReal CullDistance,
 			TCArray<FContactPoint, 4>& OutContactPoints);
 
@@ -1737,6 +1753,7 @@ namespace Chaos
 		void ConstructPlanarConvexTriangleOneShotManifold(
 			const TImplicitObjectInstanced<FImplicitBox3>& Convex,
 			const FTriangle& Triangle,
+			const bool bOneSided,
 			const FReal CullDistance,
 			TCArray<FContactPoint, 4>& OutContactPoints);
 
diff --git a/Engine/Source/Runtime/Experimental/Chaos/Private/Chaos/CollisionOneShotManifolds.h b/Engine/Source/Runtime/Experimental/Chaos/Private/Chaos/CollisionOneShotManifolds.h
--- a/Engine/Source/Runtime/Experimental/Chaos/Private/Chaos/CollisionOneShotManifolds.h
+++ b/Engine/Source/Runtime/Experimental/Chaos/Private/Chaos/CollisionOneShotManifolds.h
@@ -48,6 +48,7 @@ namespace Chaos
 		void ConstructPlanarConvexTriangleOneShotManifold(
 			const ConvexType& Convex, 
 			const FTriangle& Triangle, 
+			const bool bOneSided,
 			const FReal CullDistance, 
 			TCArray<FContactPoint, 4>& OutContactPoints);
 	}
diff --git a/Engine/Source/Runtime/Experimental/Chaos/Private/Chaos/CollisionOneShotManifoldsMiscShapes.cpp b/Engine/Source/Runtime/Experimental/Chaos/Private/Chaos/CollisionOneShotManifoldsMiscShapes.cpp
--- a/Engine/Source/Runtime/Experimental/Chaos/Private/Chaos/CollisionOneShotManifoldsMiscShapes.cpp
+++ b/Engine/Source/Runtime/Experimental/Chaos/Private/Chaos/CollisionOneShotManifoldsMiscShapes.cpp
@@ -34,6 +34,8 @@ namespace Chaos
 	extern FRealSingle Chaos_Collision_GJKEpsilon;
 	extern FRealSingle Chaos_Collision_EPAEpsilon;
 
+	extern bool bChaos_Collision_OneSidedTriangleMesh;
+
 	namespace Collisions
 	{
 		void ConstructSphereSphereOneShotManifold(
@@ -735,6 +737,9 @@ namespace Chaos
 			const FTriangleMeshImplicitObject* UnscaledTriMesh = UnwrapImplicit<FTriangleMeshImplicitObject>(TriMesh, TriMeshScale, TriMeshMargin);
 			check(UnscaledTriMesh != nullptr);
 
+			// @todo(chaos): should be an option on the Tri Mesh?
+			const bool bOneSidedCollision = bChaos_Collision_OneSidedTriangleMesh;
+
 			const FRigidTransform3 TriangleMeshToConvexTransform = TriMeshTransform.GetRelativeTransformNoScale(ConvexTransform);
 
 			// Calculate the query bounds in trimesh space
@@ -757,7 +762,7 @@ namespace Chaos
 			while (TriangleProducer.NextTriangle(*UnscaledTriMesh, QueryTransform, Triangle, TriangleIndex))
 			{
 				TriangleManifoldPoints.Reset();
-				ConstructPlanarConvexTriangleOneShotManifold(Convex, Triangle, CullDistance, TriangleManifoldPoints);
+				ConstructPlanarConvexTriangleOneShotManifold(Convex, Triangle, bOneSidedCollision, CullDistance, TriangleManifoldPoints);
 					
 				for (int32 TriangleContactIndex = 0; TriangleContactIndex < TriangleManifoldPoints.Num(); ++TriangleContactIndex)
 				{
diff --git a/Engine/Source/Runtime/Experimental/Chaos/Private/Chaos/HeightField.cpp b/Engine/Source/Runtime/Experimental/Chaos/Private/Chaos/HeightField.cpp
--- a/Engine/Source/Runtime/Experimental/Chaos/Private/Chaos/HeightField.cpp
+++ b/Engine/Source/Runtime/Experimental/Chaos/Private/Chaos/HeightField.cpp
@@ -20,6 +20,8 @@ namespace Chaos
 {
 	extern FRealSingle Chaos_Collision_EdgePrunePlaneDistance;
 
+	extern bool bChaos_Collision_OneSidedHeightfield;
+
 	int32 bOneSidedHeightField = 1;
 	static FAutoConsoleVariableRef CVarOneSidedHeightField(TEXT("p.Chaos.OneSidedHeightField"), bOneSidedHeightField, TEXT("When enabled, extra steps will ensure that FHeightField::GJKContactPointImp never results in internal-facing contact data."));
 
@@ -1465,6 +1467,9 @@ namespace Chaos
 	template <typename GeomType>
 	bool FHeightField::ContactManifoldPlanarConvexImp(const GeomType& QueryGeom, const FRigidTransform3& QueryTM, const FReal Thickness, TArray<FContactPoint>& ContactPoints) const
 	{
+		// @todo(chaos): should be an option on the Heightfield?
+		const bool bOneSidedCollision = bChaos_Collision_OneSidedHeightfield;
+
 		auto OverlapTriangle = [&](const FVec3& A, const FVec3& B, const FVec3& C, TCArray<FContactPoint, 4>& TriangleContactPoints)
 		{
 			// Create triangle in query space
@@ -1473,7 +1478,7 @@ namespace Chaos
 				QueryTM.InverseTransformPositionNoScale(B), 
 				QueryTM.InverseTransformPositionNoScale(C));
 
-			Collisions::ConstructPlanarConvexTriangleOneShotManifold(QueryGeom, TriangleConvex, Thickness, TriangleContactPoints);
+			Collisions::ConstructPlanarConvexTriangleOneShotManifold(QueryGeom, TriangleConvex, bOneSidedCollision, Thickness, TriangleContactPoints);
 
 			// Convert back to shape-local space
 			for (FContactPoint& ContactPoint : TriangleContactPoints)
