Chaos : Implement AddForce, AddRadialForce and AddTorqueInRadians for GeometryCollectioComponent

diff --git a/Engine/Source/Runtime/Experimental/GeometryCollectionEngine/Private/GeometryCollection/GeometryCollectionComponent.cpp b/Engine/Source/Runtime/Experimental/GeometryCollectionEngine/Private/GeometryCollection/GeometryCollectionComponent.cpp
--- a/Engine/Source/Runtime/Experimental/GeometryCollectionEngine/Private/GeometryCollection/GeometryCollectionComponent.cpp
+++ b/Engine/Source/Runtime/Experimental/GeometryCollectionEngine/Private/GeometryCollection/GeometryCollectionComponent.cpp
@@ -689,7 +689,53 @@ void UGeometryCollectionComponent::SetSimulatePhysics(bool bEnabled)
 	{
 		RegisterAndInitializePhysicsProxy();
 	}
+}
 
+void UGeometryCollectionComponent::AddForce(FVector Force, FName BoneName, bool bAccelChange)
+{
+	ensure(bAccelChange == false); // not supported
+	
+	const FVector Direction = Force.GetSafeNormal(); 
+	const FVector::FReal Magnitude = Force.Size();
+	const FFieldSystemCommand Command = FFieldObjectCommands::CreateFieldCommand(EFieldPhysicsType::Field_LinearForce, new FUniformVector(Magnitude, Direction));
+	DispatchFieldCommand(Command);
+}
+
+void UGeometryCollectionComponent::AddRadialForce(FVector Origin, float Radius, float Strength, ERadialImpulseFalloff Falloff, bool bAccelChange)
+{
+	ensure(bAccelChange == false); // not supported
+	if(bIgnoreRadialForce)
+	{
+		return;
+	}
+
+	FFieldNodeBase* FieldNode = nullptr;
+	if (Falloff == ERadialImpulseFalloff::RIF_Constant)
+	{
+		FieldNode = new FRadialVector(Strength, Origin);
+	}
+	else
+	{
+		FRadialFalloff * FalloffField = new FRadialFalloff(Strength,0.f, 1.f, 0.f, Radius, Origin, EFieldFalloffType::Field_Falloff_Linear);
+		FRadialVector* VectorField = new FRadialVector(1.f, Origin);
+		FieldNode = new FSumVector(1.0, FalloffField, VectorField, nullptr, Field_Multiply);
+	}
+
+	if (FieldNode)
+	{
+		const FFieldSystemCommand Command = FFieldObjectCommands::CreateFieldCommand(EFieldPhysicsType::Field_LinearForce, FieldNode);
+		DispatchFieldCommand(Command);
+	}
+}
+
+void UGeometryCollectionComponent::AddTorqueInRadians(FVector Torque, FName BoneName, bool bAccelChange)
+{
+	ensure(bAccelChange == false); // not supported
+	
+	const FVector Direction = Torque.GetSafeNormal(); 
+	const FVector::FReal Magnitude = Torque.Size();
+	const FFieldSystemCommand Command = FFieldObjectCommands::CreateFieldCommand(EFieldPhysicsType::Field_AngularTorque, new FUniformVector(Magnitude, Direction));
+	DispatchFieldCommand(Command);
 }
 
 void UGeometryCollectionComponent::DispatchBreakEvent(const FChaosBreakEvent& Event)
diff --git a/Engine/Source/Runtime/Experimental/GeometryCollectionEngine/Public/GeometryCollection/GeometryCollectionComponent.h b/Engine/Source/Runtime/Experimental/GeometryCollectionEngine/Public/GeometryCollection/GeometryCollectionComponent.h
--- a/Engine/Source/Runtime/Experimental/GeometryCollectionEngine/Public/GeometryCollection/GeometryCollectionComponent.h
+++ b/Engine/Source/Runtime/Experimental/GeometryCollectionEngine/Public/GeometryCollection/GeometryCollectionComponent.h
@@ -378,6 +378,9 @@ public:
 	virtual void SetNotifyRigidBodyCollision(bool bNewNotifyRigidBodyCollision) override;
 	virtual bool CanEditSimulatePhysics() override;
 	virtual void SetSimulatePhysics(bool bEnabled) override;
+	virtual void AddForce(FVector Force, FName BoneName = NAME_None, bool bAccelChange = false) override;
+	virtual void AddRadialForce(FVector Origin, float Radius, float Strength, ERadialImpulseFalloff Falloff, bool bAccelChange = false) override;
+	virtual void AddTorqueInRadians(FVector Torque, FName BoneName = NAME_None, bool bAccelChange = false) override;
 	virtual void PostLoad() override;
 	//~ End UPrimitiveComponent Interface.
 
