ROBOMERGE

Take fix for swept constraints setting incorrect shape transforms - partial unshelf from 19636305

diff --git a/Engine/Source/Runtime/Experimental/Chaos/Private/Chaos/CCDUtilities.cpp b/Engine/Source/Runtime/Experimental/Chaos/Private/Chaos/CCDUtilities.cpp
--- a/Engine/Source/Runtime/Experimental/Chaos/Private/Chaos/CCDUtilities.cpp
+++ b/Engine/Source/Runtime/Experimental/Chaos/Private/Chaos/CCDUtilities.cpp
@@ -6,15 +6,20 @@
 #include "Chaos/PBDRigidsEvolutionGBF.h"
 #include "Chaos/CollisionResolution.h"
 
+//PRAGMA_DISABLE_OPTIMIZATION
+
 // @todo(chaos): these cvars should all be settings per particle
 
+// NOTE: With this disabled secondary CCD collisions will often be missed
 bool bChaosCollisionCCDEnableResweep = true;
 FAutoConsoleVariableRef CVarChaosCollisionCCDEnableResweep(TEXT("p.Chaos.Collision.CCD.EnableResweep"), bChaosCollisionCCDEnableResweep, TEXT("Enable resweep for CCD. Resweeping allows CCD to catch more secondary collisions but also is more costly. Default is true."));
 
+// NOTE: With this disabled, secondary collisions can be missed. When enabled, velocity will not be visually consistent after CCD collisions (if ChaosCollisionCCDConstraintMaxProcessCount is too low)
 bool bChaosCollisionCCDAllowClipping = true;
 FAutoConsoleVariableRef CVarChaosCollisionCCDAllowClipping(TEXT("p.Chaos.Collision.CCD.AllowClipping"), bChaosCollisionCCDAllowClipping, TEXT("This will clip the CCD object at colliding positions when computation budgets run out. Default is true. Turning this option off might cause tunneling."));
 
-int32 ChaosCollisionCCDConstraintMaxProcessCount = 1;
+// This needs to be at least 2 to handle sliding at CCD speeds properly
+int32 ChaosCollisionCCDConstraintMaxProcessCount = 2;
 FAutoConsoleVariableRef CVarChaosCollisionCCDConstraintMaxProcessCount(TEXT("p.Chaos.Collision.CCD.ConstraintMaxProcessCount"), ChaosCollisionCCDConstraintMaxProcessCount, TEXT("The max number of times each constraint can be resolved when applying CCD constraints. Default is 2. The larger this number is, the more fully CCD constraints are resolved."));
 
 bool bChaosCollisionCCDEnablePhiCheck = false;
@@ -284,6 +289,8 @@ namespace Chaos
 		const int32 ConstraintEnd = IslandConstraintEnd[Island];
 		check(ConstraintNum > 0);
 
+			// Debugdraw the shape at the TOI=0 (black) and TOI=1 (white)
+
 		// Sort constraints based on TOI
 		std::sort(SortedCCDConstraints.GetData() + ConstraintStart, SortedCCDConstraints.GetData() + ConstraintStart + ConstraintNum, CCDConstraintSortPredicate);
 		FReal IslandTOI = 0.f;
@@ -466,6 +473,7 @@ namespace Chaos
 									}
 								}
 							}
+
 							/** When resweeping, we need to recompute TOI for affected constraints and therefore the work (GJKRaycast) used to compute the original TOI is wasted.
 							* A potential optimization is to compute an estimate of TOI using the AABB of the particles. Sweeping AABBs to compute an estimated TOI can be very efficient, and this TOI is strictly smaller than the accurate TOI.
 							* At each for-loop iteration, we only need the constraint with the smallest TOI in the island. A potential optimized algorithm could be like:
@@ -476,23 +484,35 @@ namespace Chaos
 							*	When resweeping, compute estimated TOI instead of accurate TOI since updated TOI might need to be updated again.
 							*/
 
-							// We need to update cached position at TOI=1 in collision constraint.
-							FPBDCollisionConstraint* SweptConstraint = SortedCCDConstraints[i]->SweptConstraint;
-							FRigidTransform3 ShapeWorldTransform0 = SweptConstraint->GetShapeWorldTransform0();
-							FRigidTransform3 ShapeWorldTransform1 = SweptConstraint->GetShapeWorldTransform1();
-							ShapeWorldTransform0.SetTranslation(FConstGenericParticleHandle(SweptConstraint->GetParticle0())->P());
-							ShapeWorldTransform1.SetTranslation(FConstGenericParticleHandle(SweptConstraint->GetParticle1())->P());
+							FPBDCollisionConstraint* SweptConstraint = AttachedCCDConstraint->SweptConstraint;
+							const FConstGenericParticleHandle Particle0 = SweptConstraint->GetParticle0();
+							const FConstGenericParticleHandle Particle1 = SweptConstraint->GetParticle1();
+
+							// Initial shape sweep transforms
+							const FRigidTransform3 ShapeStartWorldTransform0 = SweptConstraint->GetShapeRelativeTransform0() * RigidTransforms[0];
+							const FRigidTransform3 ShapeStartWorldTransform1 = SweptConstraint->GetShapeRelativeTransform1() * RigidTransforms[1];
 
-							SweptConstraint->SetShapeWorldTransforms(ShapeWorldTransform0, ShapeWorldTransform1);
-							const bool bUpdated = Collisions::UpdateConstraintFromGeometrySwept<ECollisionUpdateType::Deepest>(*(AttachedCCDConstraint->SweptConstraint), RigidTransforms[0], RigidTransforms[1], RestDt);
-							if (bUpdated)
+							// End shape sweep transforms
+							const FRigidTransform3 ParticleEndWorldTransform0 = FParticleUtilities::GetActorWorldTransform(Particle0);
+							const FRigidTransform3 ParticleEndWorldTransform1 = FParticleUtilities::GetActorWorldTransform(Particle1);
+							const FRigidTransform3 ShapeEndWorldTransform0 = SweptConstraint->GetShapeRelativeTransform0() * ParticleEndWorldTransform0;
+							const FRigidTransform3 ShapeEndWorldTransform1 = SweptConstraint->GetShapeRelativeTransform1() * ParticleEndWorldTransform1;
+
+							// Update of swept constraint assumes that the constraint holds the end transforms for the sweep
+							SweptConstraint->SetShapeWorldTransforms(ShapeEndWorldTransform0, ShapeEndWorldTransform1);
+
+							if (bChaosCollisionCCDEnableResweep)
 							{
-								const FReal RestDtTOI = AttachedCCDConstraint->SweptConstraint->TimeOfImpact;
-								if (RestDtTOI >= 0 && RestDtTOI < 1.f)
-								{
-									AttachedCCDConstraint->SweptConstraint->TimeOfImpact = IslandTOI + (1.f - IslandTOI) * RestDtTOI;
-								}
+								// Resweep the shape. This is the expensive option
+								Collisions::UpdateConstraintSwept(*SweptConstraint, ShapeStartWorldTransform0, ShapeStartWorldTransform1, RestDt);
 							}
+
+							const FReal RestDtTOI = AttachedCCDConstraint->SweptConstraint->TimeOfImpact;
+							if ((RestDtTOI >= 0) && (RestDtTOI < FReal(1)))
+							{
+								AttachedCCDConstraint->SweptConstraint->TimeOfImpact = IslandTOI + (FReal(1) - IslandTOI) * RestDtTOI;
+							}
+
 							// When bUpdated==true, TOI was modified. When bUpdated==false, TOI was set to be TNumericLimits<FReal>::Max(). In either case, a re-sorting on the constraints is needed.
 							HasResweptConstraint = true;
 						}
@@ -500,7 +520,6 @@ namespace Chaos
 				}
 				if (HasResweptConstraint)
 				{
-					// This could be optimized by using bubble sort if there are only a few updated constraints.
 					std::sort(SortedCCDConstraints.GetData() + ConstraintIndex, SortedCCDConstraints.GetData() + ConstraintStart + ConstraintNum, CCDConstraintSortPredicate);
 				}
 			}
@@ -512,7 +531,6 @@ namespace Chaos
 		}
 	}
 
-
 	void FCCDManager::ResetIslandParticles(const int32 Island)
 	{
 		const int32 ParticleStart = IslandParticleStart[Island];
@@ -574,7 +592,7 @@ namespace Chaos
 		const FReal Restitution = Constraint->GetRestitution();
 		const FRigidTransform3& ShapeWorldTransform1 = Constraint->GetShapeWorldTransform1();
 		for(const FManifoldPoint &ManifoldPoint : Constraint->GetManifoldPoints())
-		{
+			{
 			if (ManifoldPoint.Flags.bDisabled)
 			{
 				continue;
