Add CCDAxisThreshold to GeometryParticle. This is a Vec3 which on its X axis containts the X-extent of the shape with the smallest X-extent (same for Y and Z).

diff --git a/Engine/Source/Runtime/Experimental/Chaos/Public/Chaos/GeometryParticles.h b/Engine/Source/Runtime/Experimental/Chaos/Public/Chaos/GeometryParticles.h
--- a/Engine/Source/Runtime/Experimental/Chaos/Public/Chaos/GeometryParticles.h
+++ b/Engine/Source/Runtime/Experimental/Chaos/Public/Chaos/GeometryParticles.h
@@ -22,6 +22,7 @@
 #define CHAOS_DETERMINISTIC 1
 #endif
 
+
 namespace Chaos
 {
 	class FConstraintHandle;
@@ -435,6 +436,7 @@ namespace Chaos
 			TArrayCollection::AddArray(&MHasCollision);
 			TArrayCollection::AddArray(&MShapesArray);
 			TArrayCollection::AddArray(&MLocalBounds);
+			TArrayCollection::AddArray(&MCCDAxisThreshold);
 			TArrayCollection::AddArray(&MWorldSpaceInflatedBounds);
 			TArrayCollection::AddArray(&MHasBounds);
 			TArrayCollection::AddArray(&MSpatialIdx);
@@ -474,6 +476,7 @@ namespace Chaos
 			, MHasCollision(MoveTemp(Other.MHasCollision))
 			, MShapesArray(MoveTemp(Other.MShapesArray))
 			, MLocalBounds(MoveTemp(Other.MLocalBounds))
+			, MCCDAxisThreshold(MoveTemp(Other.MCCDAxisThreshold))
 			, MWorldSpaceInflatedBounds(MoveTemp(Other.MWorldSpaceInflatedBounds))
 			, MHasBounds(MoveTemp(Other.MHasBounds))
 			, MSpatialIdx(MoveTemp(Other.MSpatialIdx))
@@ -505,6 +508,7 @@ namespace Chaos
 			TArrayCollection::AddArray(&MHasCollision);
 			TArrayCollection::AddArray(&MShapesArray);
 			TArrayCollection::AddArray(&MLocalBounds);
+			TArrayCollection::AddArray(&MCCDAxisThreshold);
 			TArrayCollection::AddArray(&MWorldSpaceInflatedBounds);
 			TArrayCollection::AddArray(&MHasBounds);
 			TArrayCollection::AddArray(&MSpatialIdx);
@@ -546,6 +550,7 @@ namespace Chaos
 			TArrayCollection::AddArray(&MHasCollision);
 			TArrayCollection::AddArray(&MShapesArray);
 			TArrayCollection::AddArray(&MLocalBounds);
+			TArrayCollection::AddArray(&MCCDAxisThreshold);
 			TArrayCollection::AddArray(&MWorldSpaceInflatedBounds);
 			TArrayCollection::AddArray(&MHasBounds);
 			TArrayCollection::AddArray(&MSpatialIdx);
@@ -633,6 +638,28 @@ namespace Chaos
 			{
 				MLocalBounds[Index] = TAABB<T, d>(InGeometry->BoundingBox());
 
+				// Update CCDAxisThreshold to be the extents of the bounds of the smallest shape on each axis
+				MCCDAxisThreshold[Index] = MLocalBounds[Index].Extents();
+				for (const TUniquePtr<FPerShapeData>& Shape : ShapesArray(Index))
+				{
+					// Only sim-enabled shapes should ever be swept with CCD, so make sure the
+					// sim-enabled flag is on for each shape before considering it's min bounds
+					// for CCD extents.
+					if (Shape->GetSimEnabled())
+					{
+						const TSerializablePtr<FImplicitObject> Geometry = Shape->GetGeometry();
+						if (Geometry->HasBoundingBox())
+						{
+							const TVector<T, d> ShapeExtents = Geometry->BoundingBox().Extents();
+							TVector<T, d>& CCDAxisThreshold = MCCDAxisThreshold[Index];
+							for (int32 AxisIndex = 0; AxisIndex < d; ++AxisIndex)
+							{
+								MCCDAxisThreshold[Index][AxisIndex] = FMath::Min(ShapeExtents[AxisIndex], MCCDAxisThreshold[Index][AxisIndex]);
+							}
+						}
+					}
+				}
+
 				// Update the world-space stat of all the shapes - must be called after UpdateShapesArray
 				// world space inflated bounds needs to take expansion into account - this is done in integrate for dynamics anyway, so
 				// this computation is mainly for statics
@@ -651,6 +678,11 @@ namespace Chaos
 			return MLocalBounds[Index];
 		}
 
+		CHAOS_API const TVector<T,d>& CCDAxisThreshold(const int32 Index) const
+		{
+			return MCCDAxisThreshold[Index];
+		}
+
 		CHAOS_API bool HasBounds(const int32 Index) const
 		{
 			return MHasBounds[Index];
@@ -884,6 +916,7 @@ public:
 		TArrayCollectionArray<bool> MHasCollision;
 		TArrayCollectionArray<FShapesArray> MShapesArray;
 		TArrayCollectionArray<TAABB<T,d>> MLocalBounds;
+		TArrayCollectionArray<TVector<T,d>> MCCDAxisThreshold;
 		TArrayCollectionArray<TAABB<T, d>> MWorldSpaceInflatedBounds;
 		TArrayCollectionArray<bool> MHasBounds;
 		TArrayCollectionArray<FSpatialAccelerationIdx> MSpatialIdx;
@@ -933,3 +966,4 @@ public:
 	TGeometryParticlesImp<FReal, 3, EGeometryParticlesSimType::Other>* TGeometryParticlesImp<FReal, 3, EGeometryParticlesSimType::Other>::SerializationFactory(FChaosArchive& Ar, TGeometryParticlesImp<FReal, 3, EGeometryParticlesSimType::Other>* Particles);
 
 }
+
diff --git a/Engine/Source/Runtime/Experimental/Chaos/Public/Chaos/ParticleHandle.h b/Engine/Source/Runtime/Experimental/Chaos/Public/Chaos/ParticleHandle.h
--- a/Engine/Source/Runtime/Experimental/Chaos/Public/Chaos/ParticleHandle.h
+++ b/Engine/Source/Runtime/Experimental/Chaos/Public/Chaos/ParticleHandle.h
@@ -540,6 +540,8 @@ public:
 	const TAABB<T, d>& LocalBounds() const { return GeometryParticles->LocalBounds(ParticleIdx); }
 	void SetLocalBounds(const TAABB<T, d>& NewBounds) { GeometryParticles->LocalBounds(ParticleIdx) = NewBounds; }
 
+	const TVector<T, d>& CCDAxisThreshold() const { return GeometryParticles->CCDAxisThreshold(ParticleIdx); }
+
 	const TAABB<T, d>& WorldSpaceInflatedBounds() const { return GeometryParticles->WorldSpaceInflatedBounds(ParticleIdx); }
 
 	/**
