FIX ROBOMERGE

- Per shape bounds used to determine CCD extents
- Per axis computation so that one thin axis doesn't cause CCD to trigger on motion on other axes.
- Unified method for determining if CCD should be used.

diff --git a/Engine/Source/Runtime/Experimental/Chaos/Private/Chaos/CCDUtilities.cpp b/Engine/Source/Runtime/Experimental/Chaos/Private/Chaos/CCDUtilities.cpp
--- a/Engine/Source/Runtime/Experimental/Chaos/Private/Chaos/CCDUtilities.cpp
+++ b/Engine/Source/Runtime/Experimental/Chaos/Private/Chaos/CCDUtilities.cpp
@@ -27,6 +27,12 @@ FAutoConsoleVariableRef CVarChaosCollisionCCDEnablePhiCheck(TEXT("p.Chaos.Collis
 
 namespace Chaos
 {
+	namespace CVars
+	{
+		FRealSingle CCDEnableThresholdBoundsScale = 0.4f;
+		FAutoConsoleVariableRef  CVarCCDEnableThresholdBoundsScale(TEXT("p.Chaos.CCD.EnableThresholdBoundsScale"), CCDEnableThresholdBoundsScale , TEXT("CCD is used when object position is changing > smallest bound's extent * BoundsScale. 0 will always Use CCD. Values < 0 disables CCD."));
+	}
+
 	void FCCDParticle::AddOverlappingDynamicParticle(FCCDParticle* const InParticle)
 	{
 		OverlappingDynamicParticles.Add(InParticle);
@@ -151,18 +157,15 @@ namespace Chaos
 				}
 			}
 
-			// Compute the relative displacement. If relative displacement is smaller than 0.5 * (Extents0.Min() + Extents1.Min()), it is impossible for an particle to tunnel through another particle, even though the absolute velocities of the particles might be large.
-			const FReal CCDThreshold0 = GetParticleCCDThreshold(Constraint->GetImplicit0());
-			const FReal CCDThreshold1 = GetParticleCCDThreshold(Constraint->GetImplicit1());
-			const FReal CCDConstraintThreshold = CCDThreshold0 + CCDThreshold1;
-			if ((Displacements[1] - Displacements[0]).SizeSquared() > FMath::Square(CCDConstraintThreshold))
-			{
-				bNeedCCDSolve = true;
-			}
+			// Determine if this particle pair should trigger CCD
+			const auto Particle0 = Constraint->GetParticle(0);
+			const auto Particle1 = Constraint->GetParticle(1);
+			bNeedCCDSolve = CCDHelpers::DeltaExceedsThreshold(*Particle0, *Particle1, Dt);
 
 			// make sure we ignore pairs that don't include any dynamics
 			if (CCDParticlePair[0] != nullptr || CCDParticlePair[1] != nullptr)
 			{
+				const FReal CCDConstraintThreshold = FMath::Min(Particle0->CCDAxisThreshold().GetMin(), Particle1->CCDAxisThreshold().GetMin());
 				const FReal PhiThreshold = -CCDConstraintThreshold;
 				CCDConstraints.Add(FCCDConstraint(Constraint, CCDParticlePair, Displacements, PhiThreshold));
 				for (int32 i = 0; i < 2; i++)
@@ -659,4 +662,92 @@ namespace Chaos
 			Particle->X() = Particle->P() - Particle->V() * Dt;
 		}
 	}
+
+	bool CCDHelpers::DeltaExceedsThreshold(const FVec3& AxisThreshold, const FVec3& DeltaX, const FQuat& R)
+	{
+		FVec3 AbsLocalDelta, AxisThresholdScaled, AxisThresholdDiff;
+		return DeltaExceedsThreshold(AxisThreshold, DeltaX, R, AbsLocalDelta, AxisThresholdScaled, AxisThresholdDiff);
+	}
+
+	bool CCDHelpers::DeltaExceedsThreshold(const FVec3& AxisThreshold, const FVec3& DeltaX, const FQuat& R, FVec3& OutAbsLocalDelta, FVec3& OutAxisThresholdScaled, FVec3& OutAxisThresholdDiff)
+	{
+		if (CVars::CCDEnableThresholdBoundsScale < 0.f) { return false; }
+		if (CVars::CCDEnableThresholdBoundsScale == 0.f) { return true; }
+
+		// Get per-component absolute value of position delta in local space.
+		// This is how much we've moved on each principal axis (but not which
+		// direction on that axis that we've moved in).
+		OutAbsLocalDelta = R.UnrotateVector(DeltaX).GetAbs();
+
+		// Scale the ccd extents in local space and subtract them from the 
+		// local space position deltas. This will give us a vector representing
+		// how much further we've moved on each axis than should be allowed by
+		// the CCD bounds.
+		OutAxisThresholdScaled = AxisThreshold * CVars::CCDEnableThresholdBoundsScale;
+		OutAxisThresholdDiff = OutAbsLocalDelta - OutAxisThresholdScaled;
+
+		// That is, if any element of ExtentsDiff is greater than zero, then that
+		// means DeltaX has exceeded the scaled extents
+		return OutAxisThresholdDiff.GetMax() > 0.f;
+	}
+
+	bool CCDHelpers::DeltaExceedsThreshold(
+		const FVec3& AxisThreshold0, const FVec3& DeltaX0, const FQuat& R0,
+		const FVec3& AxisThreshold1, const FVec3& DeltaX1, const FQuat& R1)
+	{
+		return CCDHelpers::DeltaExceedsThreshold(
+
+			// To combine axis thresholds:
+			// * transform particle1's threshold into particle0's local space
+			// * take the per-component minimum of each axis threshold
+			//
+			// To think about why we use component mininma to combine thresholds,
+			// imagine what happens when a large object and a small object move
+			// towards each other at the same speed. Say particle0 is the large
+			// object, and then think about particle1's motion from particle0's
+			// inertial frame of reference. In this case, clearly you should
+			// choose particle1's threshold since it is the one that is moving.
+			//
+			// Since there's no preferred inertial frame, the correct choice
+			// will always be to take the smaller object's threshold.
+			AxisThreshold0.ComponentMin(R0 * R1.UnrotateVector(AxisThreshold1)),
+
+			// Taking the difference of the deltas gives the total delta - how
+			// much the objects have moved towards each other. We choose to use
+			// particle0 as the reference.
+			DeltaX1 - DeltaX0,
+
+			// Since we're doing this in particle0's space, we choose its rotation.
+			R0);
+	}
+
+	bool CCDHelpers::DeltaExceedsThreshold(const FGeometryParticleHandle& Particle0, const FGeometryParticleHandle& Particle1)
+	{
+		// For rigids, compute DeltaX from the X - P diff and use Q for the rotation.
+		// For non-rigids, DeltaX is zero and use R for rotation.
+		const auto Rigid0 = Particle0.CastToRigidParticle();
+		const auto Rigid1 = Particle1.CastToRigidParticle();
+		const FVec3 DeltaX0 = Rigid0 ? Rigid0->P() - Rigid0->X() : FVec3::ZeroVector;
+		const FVec3 DeltaX1 = Rigid1 ? Rigid1->P() - Rigid1->X() : FVec3::ZeroVector;
+		const FQuat& R0 = Rigid0 ? Rigid0->Q() : Particle0.R();
+		const FQuat& R1 = Rigid1 ? Rigid1->Q() : Particle1.R();
+		return DeltaExceedsThreshold(
+			Particle0.CCDAxisThreshold(), DeltaX0, R0,
+			Particle1.CCDAxisThreshold(), DeltaX1, R1);
+	}
+
+	bool CCDHelpers::DeltaExceedsThreshold(const FGeometryParticleHandle& Particle0, const FGeometryParticleHandle& Particle1, const FReal Dt)
+	{
+		// For rigids, compute DeltaX from the V * Dt and use Q for the rotation.
+		// For non-rigids, DeltaX is zero and use R for rotation.
+		const auto Rigid0 = Particle0.CastToRigidParticle();
+		const auto Rigid1 = Particle1.CastToRigidParticle();
+		const FVec3 DeltaX0 = Rigid0 ? Rigid0->V() - Dt : FVec3::ZeroVector;
+		const FVec3 DeltaX1 = Rigid1 ? Rigid1->V() - Dt : FVec3::ZeroVector;
+		const FQuat& R0 = Rigid0 ? Rigid0->Q() : Particle0.R();
+		const FQuat& R1 = Rigid1 ? Rigid1->Q() : Particle1.R();
+		return DeltaExceedsThreshold(
+			Particle0.CCDAxisThreshold(), DeltaX0, R0,
+			Particle1.CCDAxisThreshold(), DeltaX1, R1);
+	}
 }
\ No newline at end of file
diff --git a/Engine/Source/Runtime/Experimental/Chaos/Private/Chaos/ChaosDebugDraw.cpp b/Engine/Source/Runtime/Experimental/Chaos/Private/Chaos/ChaosDebugDraw.cpp
--- a/Engine/Source/Runtime/Experimental/Chaos/Private/Chaos/ChaosDebugDraw.cpp
+++ b/Engine/Source/Runtime/Experimental/Chaos/Private/Chaos/ChaosDebugDraw.cpp
@@ -25,6 +25,7 @@
 #include "Chaos/PBDSuspensionConstraints.h"
 #include "Chaos/Sphere.h"
 #include "Chaos/Utilities.h"
+#include "Chaos/CCDUtilities.h"
 
 //PRAGMA_DISABLE_OPTIMIZATION
 
@@ -1677,6 +1678,71 @@ namespace Chaos
 			}
 		}
 
+		void DrawCCDAxisThreshold(const FVec3& X, const FVec3& AxisThreshold, const FVec3& DeltaX, const FQuat& R)
+		{
+			// Call the version of DeltaExceedsThreshold that provides access to its
+			// intermediate variables so we can debug draw them
+			FVec3 AbsLocalDelta, AxisThresholdScaled, AxisThresholdDiff;
+			const bool bEnableCCD = CCDHelpers::DeltaExceedsThreshold(
+				AxisThreshold, DeltaX, R, AbsLocalDelta, AxisThresholdScaled, AxisThresholdDiff);
+			const FVector AxisPercents = AbsLocalDelta / AxisThresholdScaled;
+
+			// Pick draw settings
+			const FColor BoundsColor = bEnableCCD ? FColor::Orange : FColor::Cyan;
+			const FColor DeltaColor = FColor::Red;
+			const float ArrowSize = 20.f;
+			const FString AxisStrings[] = { FString("X"), FString("Y"), FString("Z") };
+			const float AxisThickness = 4.f;
+			const float DeltaThickness = 2.f;
+			const float BoundsThickness = 2.f;
+
+			const auto DrawAxis = [&](const int32 AxisIndex)
+			{
+				const bool bAxisEnableCCD = AxisThresholdDiff[AxisIndex] > 0.f;
+				const FColor AxisColor = bEnableCCD ? FColor::Orange : FColor::Cyan;
+				FVector AxisLocalExtent = FVector::ZeroVector;
+				FVector AxisLocalDelta = FVector::ZeroVector;
+				AxisLocalExtent[AxisIndex] = AxisThresholdScaled[AxisIndex];
+				AxisLocalDelta[AxisIndex] = AbsLocalDelta[AxisIndex];
+				const FVector AxisExtent = R.RotateVector(AxisLocalExtent);
+				const FVector AxisDelta = R.RotateVector(AxisLocalDelta);
+				FDebugDrawQueue::GetInstance().DrawDebugDirectionalArrow(X, X + AxisExtent, ArrowSize, AxisColor, false, -1.f, -1, AxisThickness);
+				FDebugDrawQueue::GetInstance().DrawDebugDirectionalArrow(X, X + AxisDelta, ArrowSize, DeltaColor, false, -1.f, -1, DeltaThickness);
+				FDebugDrawQueue::GetInstance().DrawDebugString(X + AxisExtent, FString::Printf(TEXT("%s: %.2f%%"), *AxisStrings[AxisIndex], AxisPercents[AxisIndex] * 100.f), nullptr, AxisColor, -1.f, true, 1.f);
+			};
+
+			// Get the index of the largest element
+			const auto MaxIndex = [](const FVector& Vec)
+			{
+				return Vec[0] > Vec[1]
+					? (Vec[0] > Vec[2] ? (uint8)0 : (uint8)2)
+					: (Vec[1] > Vec[2] ? (uint8)1 : (uint8)2);
+			};
+
+			// Draw the axis which has the highest percentage of it's available position delta
+			const int32 AxisIndexMaxPercent = MaxIndex(AxisPercents);
+			DrawAxis(AxisIndexMaxPercent);
+
+			// Draw a box representing the size of the CCD extents on each axis at the start
+			// and end position of the CoM over this frame (gray = no ccd, orange = ccd enabled)
+			// Also draw a position delta vector (red) and the threshold delta vector (green)
+			const FVector HalfExtents = AxisThreshold * .5f;
+			FDebugDrawQueue::GetInstance().DrawDebugBox(X, HalfExtents, R, BoundsColor, false, -1.f, -1, BoundsThickness);
+			FDebugDrawQueue::GetInstance().DrawDebugBox(X + DeltaX, HalfExtents, R, BoundsColor, false, -1.f, -1, BoundsThickness);
+			const float ScaleArray[] = { -1.f, 1.f };
+			for (uint8 IX = 0; IX < 2; ++IX)
+			{
+				for (uint8 IY = 0; IY < 2; ++IY)
+				{
+					for (uint8 IZ = 0; IZ < 2; ++IZ)
+					{
+						const FVector LocalOffset = HalfExtents * FVector(ScaleArray[IX], ScaleArray[IY], ScaleArray[IZ]);
+						const FVector Offset = R.RotateVector(LocalOffset);
+						FDebugDrawQueue::GetInstance().DrawDebugLine(X, X + DeltaX, BoundsColor, false, -1.f, -1, BoundsThickness);
+					}
+				}
+			}
+		}
 #endif
 	}
 }
diff --git a/Engine/Source/Runtime/Experimental/Chaos/Private/Chaos/CollisionResolution.cpp b/Engine/Source/Runtime/Experimental/Chaos/Private/Chaos/CollisionResolution.cpp
--- a/Engine/Source/Runtime/Experimental/Chaos/Private/Chaos/CollisionResolution.cpp
+++ b/Engine/Source/Runtime/Experimental/Chaos/Private/Chaos/CollisionResolution.cpp
@@ -32,6 +32,7 @@
 #include "Chaos/GeometryQueries.h"
 #include "Chaos/CollisionOneShotManifolds.h"
 #include "Chaos/CollisionOneShotManifoldsMiscShapes.h"
+#include "Chaos/CCDUtilities.h"
 
 //PRAGMA_DISABLE_OPTIMIZATION
 
@@ -118,10 +119,6 @@ extern bool bChaosCollisionCCDEnableResweep;
 
 namespace Chaos
 {
-	namespace CVars
-	{
-		extern FRealSingle CCDEnableThresholdBoundsScale;
-	}
 	namespace Collisions
 	{
 		void ResetChaosCollisionCounters()
@@ -179,31 +176,7 @@ namespace Chaos
 				return false;
 			}
 
-			bool bUseCCD = false;
-
-			// bUseCCD is determined by the difference between P and X rather than V.
-			// This implies that the smaller Dt is, the less likely CCD will be triggered.
-			// This is because CCD is enabled to prevent tunneling. When Dt it small, tunneling is less likely to happen and therefore fewer constraints need to be processed with CCD.
-			if (bIsCCDEnabled0)
-			{
-				const FReal DSizeSquared0 = DeltaX0.SizeSquared();
-				const FReal MinExtent0 = Particle0->LocalBounds().Extents().Min();
-				const FReal CCDThreshold0 = MinExtent0 * CVars::CCDEnableThresholdBoundsScale;
-				if (DSizeSquared0 > CCDThreshold0 * CCDThreshold0)
-				{
-					bUseCCD = true;
-				}
-			}
-			if (!bUseCCD && bIsCCDEnabled1)
-			{
-				const FReal DSizeSquared1 = DeltaX1.SizeSquared();
-				const FReal MinExtent1 = Particle1->LocalBounds().Extents().Min();
-				const FReal CCDThreshold1 = MinExtent1 * CVars::CCDEnableThresholdBoundsScale;
-				if (DSizeSquared1 > CCDThreshold1 * CCDThreshold1)
-				{
-					bUseCCD = true;
-				}
-			}
+			bool bUseCCD = CCDHelpers::DeltaExceedsThreshold(*Particle0, *Particle1);
 
 			if (bUseCCD)
 			{
diff --git a/Engine/Source/Runtime/Experimental/Chaos/Private/Chaos/PBDRigidsEvolutionGBF.cpp b/Engine/Source/Runtime/Experimental/Chaos/Private/Chaos/PBDRigidsEvolutionGBF.cpp
--- a/Engine/Source/Runtime/Experimental/Chaos/Private/Chaos/PBDRigidsEvolutionGBF.cpp
+++ b/Engine/Source/Runtime/Experimental/Chaos/Private/Chaos/PBDRigidsEvolutionGBF.cpp
@@ -61,12 +61,11 @@ namespace Chaos
 		bool bChaosUseCCD = true;
 		FAutoConsoleVariableRef  CVarChaosUseCCD(TEXT("p.Chaos.CCD.UseCCD"), bChaosUseCCD , TEXT("Global flag to turn CCD on or off. Default is true"));
 
-		FRealSingle CCDEnableThresholdBoundsScale = 0.4f;
-		FAutoConsoleVariableRef  CVarCCDEnableThresholdBoundsScale(TEXT("p.Chaos.CCD.EnableThresholdBoundsScale"), CCDEnableThresholdBoundsScale , TEXT("CCD is used when object position is changing > smallest bound's extent * BoundsScale. 0 will always Use CCD. Values < 0 disables CCD."));
-
 		bool bChaosCollisionCCDUseTightBoundingBox = true; 
 		FAutoConsoleVariableRef  CVarChaosCollisionCCDUseTightBoundingBox(TEXT("p.Chaos.Collision.CCD.UseTightBoundingBox"), bChaosCollisionCCDUseTightBoundingBox , TEXT(""));
 
+		bool bChaosCollisionCCDDebugDraw = false; 
+		FAutoConsoleVariableRef  CVarChaosCollisionCCDDebugDraw(TEXT("p.Chaos.Collision.CCD.DebugDraw"), bChaosCollisionCCDDebugDraw , TEXT(""));
 
 		int32 ChaosSolverCollisionPriority = 0;
 		FAutoConsoleVariableRef CVarChaosSolverCollisionPriority(TEXT("p.Chaos.Solver.Collision.Priority"), ChaosSolverCollisionPriority, TEXT("Set constraint priority. Larger values are evaluated later [def:0]"));
diff --git a/Engine/Source/Runtime/Experimental/Chaos/Public/Chaos/CCDUtilities.h b/Engine/Source/Runtime/Experimental/Chaos/Public/Chaos/CCDUtilities.h
--- a/Engine/Source/Runtime/Experimental/Chaos/Public/Chaos/CCDUtilities.h
+++ b/Engine/Source/Runtime/Experimental/Chaos/Public/Chaos/CCDUtilities.h
@@ -119,4 +119,29 @@ namespace Chaos
 		TArray<int32> IslandParticleStart, IslandParticleNum; // IslandParticleNum is the number of dynamic particles in an island.
 		TArray<int32> IslandConstraintStart, IslandConstraintNum, IslandConstraintEnd; // The duplication of information in IslandConstraintNum and IslandConstraintEnd is needed when we group constraints using islands.
 	};
+
+	struct CCDHelpers
+	{
+		// Return true if DeltaX indicates a movement beyond a set of threshold
+		// distances for each local axis.
+		static bool DeltaExceedsThreshold(const FVec3& AxisThreshold, const FVec3& DeltaX, const FQuat& R);
+		static bool DeltaExceedsThreshold(const FVec3& AxisThreshold, const FVec3& DeltaX, const FQuat& R, FVec3& OutAbsLocalDelta, FVec3& AxisThresholdScaled, FVec3& AxisThresholdDiff);
+
+		// Return true if a pair of DeltaXs potentially indicate closing
+		// motion between two bodies which would cause the crossing of a combined
+		// threshold on any local axis.
+		static bool DeltaExceedsThreshold(
+			const FVec3& AxisThreshold0, const FVec3& DeltaX0, const FQuat& R0,
+			const FVec3& AxisThreshold1, const FVec3& DeltaX1, const FQuat& R1);
+
+		// Variant of the two-particle of DeltaExceedsThreshold which chooses
+		// DeltaX: For rigid particles, use X-P. For non-rigids, use Vec3::ZeroVector.
+		// R: For rigid particles, use Q. For non-rigids, use R.
+		static bool DeltaExceedsThreshold(const FGeometryParticleHandle& Particle0, const FGeometryParticleHandle& Particle1);
+
+		// Variant of the two-particle of DeltaExceedsThreshold which chooses
+		// DeltaX: For rigid particles, use V*Dt. For non-rigids, use Vec3::ZeroVector.
+		// R: For rigid particles, use Q. For non-rigids, use R.
+		static bool DeltaExceedsThreshold(const FGeometryParticleHandle& Particle0, const FGeometryParticleHandle& Particle1, const FReal Dt);
+	};
 }
diff --git a/Engine/Source/Runtime/Experimental/Chaos/Public/Chaos/ChaosDebugDraw.h b/Engine/Source/Runtime/Experimental/Chaos/Public/Chaos/ChaosDebugDraw.h
--- a/Engine/Source/Runtime/Experimental/Chaos/Public/Chaos/ChaosDebugDraw.h
+++ b/Engine/Source/Runtime/Experimental/Chaos/Public/Chaos/ChaosDebugDraw.h
@@ -218,6 +218,7 @@ namespace Chaos
 		CHAOS_API void DrawCollidingShapes(const FRigidTransform3& SpaceTransform, const FPBDCollisionConstraints& Collisions, FRealSingle ColorScale, const FChaosDebugDrawSettings* Settings = nullptr);
 		CHAOS_API void DrawSpatialAccelerationStructure(const ISpatialAcceleration<FAccelerationStructureHandle, FReal, 3>& SpatialAccelerationStructure, const FChaosDebugDrawSettings* Settings = nullptr);
 		CHAOS_API void DrawSuspensionConstraints(const FRigidTransform3& SpaceTransform, const FPBDSuspensionConstraints& Constraints, const FChaosDebugDrawSettings* Settings = nullptr);
+		CHAOS_API void DrawCCDAxisThreshold(const FVec3& X, const FVec3& AxisThreshold, const FVec3& DeltaX, const FQuat& R);
 #endif
 	}
 }
diff --git a/Engine/Source/Runtime/Experimental/Chaos/Public/Chaos/PBDRigidsEvolutionGBF.h b/Engine/Source/Runtime/Experimental/Chaos/Public/Chaos/PBDRigidsEvolutionGBF.h
--- a/Engine/Source/Runtime/Experimental/Chaos/Public/Chaos/PBDRigidsEvolutionGBF.h
+++ b/Engine/Source/Runtime/Experimental/Chaos/Public/Chaos/PBDRigidsEvolutionGBF.h
@@ -17,7 +17,9 @@
 #include "Chaos/PerParticleInitForce.h"
 #include "Chaos/PerParticlePBDEulerStep.h"
 #include "Chaos/CCDUtilities.h"
+#include "Chaos/Particle/ParticleUtilities.h"
 #include "Chaos/PBDSuspensionConstraints.h"
+#include "Chaos/ChaosDebugDraw.h"
 
 namespace Chaos
 {
@@ -30,9 +32,9 @@ namespace Chaos
 	{
 		CHAOS_API extern FRealSingle HackMaxAngularVelocity;
 		CHAOS_API extern FRealSingle HackMaxVelocity;
-		CHAOS_API extern FRealSingle CCDEnableThresholdBoundsScale;
 		CHAOS_API extern FRealSingle SmoothedPositionLerpRate;
 		CHAOS_API extern bool bChaosCollisionCCDUseTightBoundingBox;
+		CHAOS_API extern bool bChaosCollisionCCDDebugDraw;
 	}
 
 	using FPBDRigidsEvolutionCallback = TFunction<void()>;
@@ -298,10 +300,15 @@ namespace Chaos
 					}
 					else
 					{
-						const FReal MinBoundsAxis = Particle.LocalBounds().Extents().Min();
-						const FReal LengthCCDThreshold = MinBoundsAxis *  CVars::CCDEnableThresholdBoundsScale;
-						const FReal PXSizeSquared = (Particle.P() - Particle.X()).SizeSquared();
-						if (PXSizeSquared > LengthCCDThreshold * LengthCCDThreshold)
+
+#if CHAOS_DEBUG_DRAW
+						if (CVars::bChaosCollisionCCDDebugDraw)
+						{
+							DebugDraw::DrawCCDAxisThreshold(Particle.X(), Particle.CCDAxisThreshold(), Particle.P() - Particle.X(), Particle.Q());
+						}
+#endif
+
+						if (CCDHelpers::DeltaExceedsThreshold(Particle.CCDAxisThreshold(), Particle.P() - Particle.X(), Particle.Q()))
 						{
 							if (CVars::bChaosCollisionCCDUseTightBoundingBox)
 							{
@@ -378,6 +385,7 @@ namespace Chaos
 		void GatherSolverInput(FReal Dt, int32 GroupIndex);
 		void ScatterSolverOutput(FReal Dt, int32 GroupIndex);
 
+
 		FEvolutionResimCache* GetCurrentStepResimCache()
 		{
 			return CurrentStepResimCacheImp;
@@ -412,3 +420,4 @@ namespace Chaos
 	};
 
 }
+
