Chaos : add AddImpulse & AddRadialImpulse to GC component
- also add new field type : LinearImpulse

diff --git a/Engine/Source/Runtime/Experimental/Chaos/Private/ChaosStats.cpp b/Engine/Source/Runtime/Experimental/Chaos/Private/ChaosStats.cpp
--- a/Engine/Source/Runtime/Experimental/Chaos/Private/ChaosStats.cpp
+++ b/Engine/Source/Runtime/Experimental/Chaos/Private/ChaosStats.cpp
@@ -92,6 +92,7 @@ DEFINE_STAT(STAT_ParamUpdateField_PositionAnimated);
 DEFINE_STAT(STAT_ParamUpdateField_DynamicConstraint);
 DEFINE_STAT(STAT_ForceUpdateField_LinearForce);
 DEFINE_STAT(STAT_ForceUpdateField_AngularTorque);
+DEFINE_STAT(STAT_ForceUpdateField_LinearImpulse);
 
 DEFINE_STAT(STAT_Collisions_Detect);
 DEFINE_STAT(STAT_Collisions_ParticlePairBroadPhase);
diff --git a/Engine/Source/Runtime/Experimental/Chaos/Public/ChaosStats.h b/Engine/Source/Runtime/Experimental/Chaos/Public/ChaosStats.h
--- a/Engine/Source/Runtime/Experimental/Chaos/Public/ChaosStats.h
+++ b/Engine/Source/Runtime/Experimental/Chaos/Public/ChaosStats.h
@@ -113,6 +113,7 @@ DECLARE_CYCLE_STAT_EXTERN(TEXT("[Field Update] PositionAnimated"), STAT_ParamUpd
 DECLARE_CYCLE_STAT_EXTERN(TEXT("[Field Update] DynamicConstraint"), STAT_ParamUpdateField_DynamicConstraint, STATGROUP_Chaos, CHAOS_API);
 DECLARE_CYCLE_STAT_EXTERN(TEXT("[Field Update] LinearForce"), STAT_ForceUpdateField_LinearForce, STATGROUP_Chaos, CHAOS_API);
 DECLARE_CYCLE_STAT_EXTERN(TEXT("[Field Update] AngularTorque"), STAT_ForceUpdateField_AngularTorque, STATGROUP_Chaos, CHAOS_API);
+DECLARE_CYCLE_STAT_EXTERN(TEXT("[Field Update] Impulse"), STAT_ForceUpdateField_LinearImpulse, STATGROUP_Chaos, CHAOS_API);
 
 // Collision Detection
 DECLARE_CYCLE_STAT_EXTERN(TEXT("Collisions::Detect"), STAT_Collisions_Detect, STATGROUP_ChaosCollision, CHAOS_API);
diff --git a/Engine/Source/Runtime/Experimental/Chaos/Public/Field/FieldSystemTypes.h b/Engine/Source/Runtime/Experimental/Chaos/Public/Field/FieldSystemTypes.h
--- a/Engine/Source/Runtime/Experimental/Chaos/Public/Field/FieldSystemTypes.h
+++ b/Engine/Source/Runtime/Experimental/Chaos/Public/Field/FieldSystemTypes.h
@@ -161,6 +161,7 @@ enum EFieldPhysicsType
 	Field_ActivateDisabled			UMETA(DisplayName = "Activate Disabled", ToolTip = "Activate all the disabled particles for which the field value will be 0"),
 	Field_InitialLinearVelocity		UMETA(DisplayName = "Initial Linear Velocity", ToolTip = "Set the geometry collection initial linear velocity", Hidden),
 	Field_InitialAngularVelocity	UMETA(DisplayName = "Initial Angular Velocity", ToolTip = "Set the geometry collection initial angular velocity", Hidden),
+	Field_LinearImpulse				UMETA(DisplayName = "Linear Impulse", ToolTip = "Add a vector field affecting the particles linear impulse.", Hidden),
 	//~~~
 	//256th entry
 	Field_PhysicsType_Max           UMETA(Hidden)
@@ -177,6 +178,7 @@ enum EFieldVectorType
 	Vector_PositionTarget			UMETA(DisplayName = "Position Target", ToolTip = "Add a position constraint to the particles to follow a target position", Hidden),
 	Vector_InitialLinearVelocity	UMETA(DisplayName = "Initial Linear Velocity", ToolTip = "Set the geometry collection initial linear velocity", Hidden),
 	Vector_InitialAngularVelocity	UMETA(DisplayName = "Initial Angular Velocity", ToolTip = "Set the geometry collection initial angular velocity", Hidden),
+	Vector_LinearImpulse			UMETA(DisplayName = "Linear Impulse", ToolTip = "Add a vector field affecting the particles linear impulse.", Hidden),
 
 	//~~~
 	//256th entry
@@ -240,7 +242,8 @@ TArray<EFieldPhysicsType> CHAOS_API GetFieldTargetTypes(EFieldOutputType OutputT
 						 EFieldPhysicsType::Field_AngularTorque,
 						 EFieldPhysicsType::Field_PositionTarget,
 						 EFieldPhysicsType::Field_InitialLinearVelocity,
-						 EFieldPhysicsType::Field_InitialAngularVelocity };
+						 EFieldPhysicsType::Field_InitialAngularVelocity,
+						 EFieldPhysicsType::Field_LinearImpulse };
 		break;
 	}
 	case Field_Output_Scalar:
diff --git a/Engine/Source/Runtime/Experimental/Chaos/Public/GeometryCollection/GeometryCollectionSimulationTypes.h b/Engine/Source/Runtime/Experimental/Chaos/Public/GeometryCollection/GeometryCollectionSimulationTypes.h
--- a/Engine/Source/Runtime/Experimental/Chaos/Public/GeometryCollection/GeometryCollectionSimulationTypes.h
+++ b/Engine/Source/Runtime/Experimental/Chaos/Public/GeometryCollection/GeometryCollectionSimulationTypes.h
@@ -56,6 +56,7 @@ enum class EGeometryCollectionPhysicsTypeEnum : uint8
 	Chaos_SleepingThreshold        UMETA(DisplayName = "Sleeping Threshold", ToolTip = "Set particles in sleeping mode if their linear and angular velocity are less than the threshold."),
 	Chaos_ExternalClusterStrain    UMETA(DisplayName = "External Strain", ToolTip = "Apply an external strain over the particles. If this strain is over the internal one, the cluster will break."),
 	Chaos_InternalClusterStrain    UMETA(DisplayName = "Internal Strain", ToolTip = "Add a strain field to the particles internal one."),
+	Chaos_LinearImpulse			   UMETA(DisplayName = "Linear Impulse", ToolTip = "Add a vector field to apply an impulse to the particles."),
 	//
 	Chaos_Max						UMETA(Hidden)
 };
@@ -63,10 +64,22 @@ enum class EGeometryCollectionPhysicsTypeEnum : uint8
 
 inline CHAOS_API EFieldPhysicsType GetGeometryCollectionPhysicsType(const EGeometryCollectionPhysicsTypeEnum GeoCollectionType)
 {
-	static const TArray<EFieldPhysicsType> PhysicsTypes = { EFieldPhysicsType::Field_AngularVelociy , EFieldPhysicsType::Field_DynamicState, EFieldPhysicsType::Field_LinearVelocity, 
-		EFieldPhysicsType::Field_InitialAngularVelocity, EFieldPhysicsType::Field_InitialLinearVelocity, EFieldPhysicsType::Field_CollisionGroup, 
-		EFieldPhysicsType::Field_LinearForce, EFieldPhysicsType::Field_AngularTorque, EFieldPhysicsType::Field_DisableThreshold, EFieldPhysicsType::Field_SleepingThreshold, 
-		EFieldPhysicsType::Field_ExternalClusterStrain, EFieldPhysicsType::Field_InternalClusterStrain,EFieldPhysicsType::Field_PhysicsType_Max };
+	static const TArray<EFieldPhysicsType> PhysicsTypes = {
+		EFieldPhysicsType::Field_AngularVelociy,
+		EFieldPhysicsType::Field_DynamicState,
+		EFieldPhysicsType::Field_LinearVelocity, 
+		EFieldPhysicsType::Field_InitialAngularVelocity,
+		EFieldPhysicsType::Field_InitialLinearVelocity,
+		EFieldPhysicsType::Field_CollisionGroup, 
+		EFieldPhysicsType::Field_LinearForce,
+		EFieldPhysicsType::Field_AngularTorque,
+		EFieldPhysicsType::Field_DisableThreshold,
+		EFieldPhysicsType::Field_SleepingThreshold, 
+		EFieldPhysicsType::Field_ExternalClusterStrain,
+		EFieldPhysicsType::Field_InternalClusterStrain,
+		EFieldPhysicsType::Field_LinearImpulse,
+		EFieldPhysicsType::Field_PhysicsType_Max
+	};
 
 	return PhysicsTypes[(uint8)GeoCollectionType];
 }
diff --git a/Engine/Source/Runtime/Experimental/Chaos/Public/PhysicsProxy/FieldSystemProxyHelper.h b/Engine/Source/Runtime/Experimental/Chaos/Public/PhysicsProxy/FieldSystemProxyHelper.h
--- a/Engine/Source/Runtime/Experimental/Chaos/Public/PhysicsProxy/FieldSystemProxyHelper.h
+++ b/Engine/Source/Runtime/Experimental/Chaos/Public/PhysicsProxy/FieldSystemProxyHelper.h
@@ -594,6 +594,22 @@ namespace Chaos
 				}
 			}
 		}
+		else if (FieldCommand.PhysicsType == EFieldPhysicsType::Field_LinearImpulse)
+		{
+			SCOPE_CYCLE_COUNTER(STAT_ForceUpdateField_LinearImpulse);
+			{
+				static_cast<const FFieldNode<FVector>*>(FieldCommand.RootNode.Get())->Evaluate(FieldContext, ResultsView);
+				for (const FFieldContextIndex& Index : FieldContext.GetEvaluatedSamples())
+				{
+					Chaos::FPBDRigidParticleHandle* RigidHandle = ParticleHandles[Index.Sample]->CastToRigidParticle();
+					if (RigidHandle && RigidHandle->ObjectState() == Chaos::EObjectStateType::Dynamic)
+					{
+						const FVec3 CurrentImpulseVelocity = RigidHandle->LinearImpulseVelocity();
+						RigidHandle->SetLinearImpulseVelocity(CurrentImpulseVelocity + (ResultsView[Index.Result] * RigidHandle->InvM()));
+					}
+				}
+			}
+		}
 		else if (FieldCommand.PhysicsType == EFieldPhysicsType::Field_AngularVelociy)
 		{
 			SCOPE_CYCLE_COUNTER(STAT_ParamUpdateField_AngularVelocity);
@@ -711,7 +727,8 @@ FORCEINLINE bool IsParameterFieldValid(const FFieldSystemCommand& FieldCommand)
 	else if (FieldCommand.RootNode->Type() == FFieldNodeBase::EFieldType::EField_FVector)
 	{
 		return (FieldCommand.PhysicsType == EFieldPhysicsType::Field_LinearVelocity) ||
-			(FieldCommand.PhysicsType == EFieldPhysicsType::Field_AngularVelociy);
+				(FieldCommand.PhysicsType == EFieldPhysicsType::Field_AngularVelociy) ||
+				(FieldCommand.PhysicsType == EFieldPhysicsType::Field_LinearImpulse);
 	}
 	return false;
 }
diff --git a/Engine/Source/Runtime/Experimental/GeometryCollectionEngine/Private/GeometryCollection/GeometryCollectionComponent.cpp b/Engine/Source/Runtime/Experimental/GeometryCollectionEngine/Private/GeometryCollection/GeometryCollectionComponent.cpp
--- a/Engine/Source/Runtime/Experimental/GeometryCollectionEngine/Private/GeometryCollection/GeometryCollectionComponent.cpp
+++ b/Engine/Source/Runtime/Experimental/GeometryCollectionEngine/Private/GeometryCollection/GeometryCollectionComponent.cpp
@@ -701,6 +701,32 @@ void UGeometryCollectionComponent::AddForce(FVector Force, FName BoneName, bool
 	DispatchFieldCommand(Command);
 }
 
+void UGeometryCollectionComponent::AddImpulse(FVector Impulse, FName BoneName, bool bVelChange)
+{
+	const FVector Direction = Impulse.GetSafeNormal(); 
+	const FVector::FReal Magnitude = Impulse.Size();
+	const EFieldPhysicsType FieldType = bVelChange? EFieldPhysicsType::Field_LinearVelocity: EFieldPhysicsType::Field_LinearImpulse;
+	
+	const FFieldSystemCommand Command = FFieldObjectCommands::CreateFieldCommand(FieldType, new FUniformVector(Magnitude, Direction));
+	DispatchFieldCommand(Command);
+}
+
+TUniquePtr<FFieldNodeBase> MakeRadialField(const FVector& Origin, float Radius, float Strength, ERadialImpulseFalloff Falloff)
+{
+	TUniquePtr<FFieldNodeBase> Field;
+	if (Falloff == ERadialImpulseFalloff::RIF_Constant)
+	{
+		Field.Reset(new FRadialVector(Strength, Origin));
+	}
+	else
+	{
+		FRadialFalloff * FalloffField = new FRadialFalloff(Strength,0.f, 1.f, 0.f, Radius, Origin, EFieldFalloffType::Field_Falloff_Linear);
+		FRadialVector* VectorField = new FRadialVector(1.f, Origin);
+		Field.Reset(new FSumVector(1.0, FalloffField, VectorField, nullptr, Field_Multiply));
+	}
+	return Field;
+}
+
 void UGeometryCollectionComponent::AddRadialForce(FVector Origin, float Radius, float Strength, ERadialImpulseFalloff Falloff, bool bAccelChange)
 {
 	ensure(bAccelChange == false); // not supported
@@ -709,21 +735,24 @@ void UGeometryCollectionComponent::AddRadialForce(FVector Origin, float Radius,
 		return;
 	}
 
-	FFieldNodeBase* FieldNode = nullptr;
-	if (Falloff == ERadialImpulseFalloff::RIF_Constant)
+	if (TUniquePtr<FFieldNodeBase> Field = MakeRadialField(Origin, Radius, Strength, Falloff))
 	{
-		FieldNode = new FRadialVector(Strength, Origin);
+		const FFieldSystemCommand Command = FFieldObjectCommands::CreateFieldCommand(EFieldPhysicsType::Field_LinearForce, Field.Release());
+		DispatchFieldCommand(Command);
 	}
-	else
+}
+
+void UGeometryCollectionComponent::AddRadialImpulse(FVector Origin, float Radius, float Strength, enum ERadialImpulseFalloff Falloff, bool bVelChange)
+{
+	if(bIgnoreRadialImpulse)
 	{
-		FRadialFalloff * FalloffField = new FRadialFalloff(Strength,0.f, 1.f, 0.f, Radius, Origin, EFieldFalloffType::Field_Falloff_Linear);
-		FRadialVector* VectorField = new FRadialVector(1.f, Origin);
-		FieldNode = new FSumVector(1.0, FalloffField, VectorField, nullptr, Field_Multiply);
+		return;
 	}
 
-	if (FieldNode)
+	if (TUniquePtr<FFieldNodeBase> Field = MakeRadialField(Origin, Radius, Strength, Falloff))
 	{
-		const FFieldSystemCommand Command = FFieldObjectCommands::CreateFieldCommand(EFieldPhysicsType::Field_LinearForce, FieldNode);
+		const EFieldPhysicsType FieldType = bVelChange? EFieldPhysicsType::Field_LinearVelocity: EFieldPhysicsType::Field_LinearImpulse;
+		const FFieldSystemCommand Command = FFieldObjectCommands::CreateFieldCommand(FieldType, Field.Release());
 		DispatchFieldCommand(Command);
 	}
 }
diff --git a/Engine/Source/Runtime/Experimental/GeometryCollectionEngine/Public/GeometryCollection/GeometryCollectionComponent.h b/Engine/Source/Runtime/Experimental/GeometryCollectionEngine/Public/GeometryCollection/GeometryCollectionComponent.h
--- a/Engine/Source/Runtime/Experimental/GeometryCollectionEngine/Public/GeometryCollection/GeometryCollectionComponent.h
+++ b/Engine/Source/Runtime/Experimental/GeometryCollectionEngine/Public/GeometryCollection/GeometryCollectionComponent.h
@@ -379,7 +379,9 @@ public:
 	virtual bool CanEditSimulatePhysics() override;
 	virtual void SetSimulatePhysics(bool bEnabled) override;
 	virtual void AddForce(FVector Force, FName BoneName = NAME_None, bool bAccelChange = false) override;
+	virtual void AddImpulse(FVector Impulse, FName BoneName = NAME_None, bool bVelChange = false) override;
 	virtual void AddRadialForce(FVector Origin, float Radius, float Strength, ERadialImpulseFalloff Falloff, bool bAccelChange = false) override;
+	virtual void AddRadialImpulse(FVector Origin, float Radius, float Strength, enum ERadialImpulseFalloff Falloff, bool bVelChange = false) override;
 	virtual void AddTorqueInRadians(FVector Torque, FName BoneName = NAME_None, bool bAccelChange = false) override;
 	virtual void PostLoad() override;
 	//~ End UPrimitiveComponent Interface.
