Chaos Cloth: Update XPBD constraint stiffness calculation to produce similar results to PBD.

Also, while I fixed XPBDLongRangeConstraints to work, I removed actually using them in the ChaoCloth plugin. Unlike the other PBD constraints, this does not take mass into account, which seems to mess up the stiffness conversion. Additionally, we only do one iteration of long range constraints per subframe as an initial guess jump start, so we don't really get any XPBD benefits.

Main change is adding a method "ApplyXPBDValues" to PBDStiffness. This calculates a stiffness value that can be used by XPBD springs that has similar scaling properties as the PBD version. Namely, it scales like 1 / PBDStiffness, and it does not apply an exponent meant to capture changes in dt and num iterations as this is handled directly by the XPBD method.
Min and Max stiffness values used by this method were estimated empirically for our typical mass (density) values.

diff --git a/Engine/Plugins/Experimental/ChaosCloth/Source/Chaos/Private/ChaosCloth/ChaosClothConstraints.cpp b/Engine/Plugins/Experimental/ChaosCloth/Source/Chaos/Private/ChaosCloth/ChaosClothConstraints.cpp
--- a/Engine/Plugins/Experimental/ChaosCloth/Source/Chaos/Private/ChaosCloth/ChaosClothConstraints.cpp
+++ b/Engine/Plugins/Experimental/ChaosCloth/Source/Chaos/Private/ChaosCloth/ChaosClothConstraints.cpp
@@ -271,22 +271,14 @@ void FClothConstraints::CreateRules()
 
 	// Long range constraints modify particle P as part of Init. To avoid possible dependency order issues,
 	// add them last
-	if (XLongRangeConstraints)
-	{
-		ConstraintInits[ConstraintInitIndex++] =
-			[this](Softs::FSolverParticles& Particles, const Softs::FSolverReal Dt)
-			{
-				XLongRangeConstraints->Init();
-				XLongRangeConstraints->ApplyProperties(Dt, Evolution->GetIterations());
-				XLongRangeConstraints->Apply(Particles, Dt);  // Run the LRA constraint only once per timestep
-			};
-	}
 	if (LongRangeConstraints)
 	{
 		ConstraintInits[ConstraintInitIndex++] =
 			[this](Softs::FSolverParticles& Particles, const Softs::FSolverReal Dt)
-			{
-				LongRangeConstraints->ApplyProperties(Dt, Evolution->GetIterations());
+			{	
+				// Only doing one iteration.
+				constexpr int32 NumLRAIterations = 1;
+				LongRangeConstraints->ApplyProperties(Dt, NumLRAIterations);
 				LongRangeConstraints->Apply(Particles, Dt);  // Run the LRA constraint only once per timestep
 			};
 	}
@@ -419,34 +411,19 @@ void FClothConstraints::SetLongRangeConstraints(
 	const TArray<TConstArrayView<TTuple<int32, int32, FRealSingle>>>& Tethers,
 	const TConstArrayView<FRealSingle>& TetherStiffnessMultipliers,
 	const TConstArrayView<FRealSingle>& TetherScaleMultipliers,
-	const Softs::FSolverVec2& TetherScale, bool bUseXPBDConstraints)
+	const Softs::FSolverVec2& TetherScale)
 {
 	check(Evolution);
-
-	if (bUseXPBDConstraints)
-	{
-		XLongRangeConstraints = MakeShared<Softs::FXPBDLongRangeConstraints>(
-			Evolution->Particles(),
-			ParticleOffset,
-			NumParticles,
-			Tethers,
-			TetherStiffnessMultipliers,
-			TetherScaleMultipliers,
-			/*InStiffness =*/ Softs::FSolverVec2::UnitVector,
-			TetherScale);
-	}
-	else
-	{
-		LongRangeConstraints = MakeShared<Softs::FPBDLongRangeConstraints>(
-			Evolution->Particles(),
-			ParticleOffset,
-			NumParticles,
-			Tethers,
-			TetherStiffnessMultipliers,
-			TetherScaleMultipliers,
-			/*InStiffness =*/ Softs::FSolverVec2::UnitVector,
-			TetherScale);
-	}
+	//  Now that we're only doing a single iteration of Long range constraints, and they're more of a fake constraint to jump start our initial guess, it's not clear that using XPBD makes sense here.
+	LongRangeConstraints = MakeShared<Softs::FPBDLongRangeConstraints>(
+		Evolution->Particles(),
+		ParticleOffset,
+		NumParticles,
+		Tethers,
+		TetherStiffnessMultipliers,
+		TetherScaleMultipliers,
+		/*InStiffness =*/ Softs::FSolverVec2::UnitVector,
+		TetherScale);
 	++NumConstraintInits;  // Uses init to both update the property tables and apply the constraint
 }
 
@@ -588,11 +565,6 @@ void FClothConstraints::SetLongRangeAttachmentProperties(const Softs::FSolverVec
 		LongRangeConstraints->SetStiffness(TetherStiffness);
 		LongRangeConstraints->SetScale(TetherScale);
 	}
-	if (XLongRangeConstraints)
-	{
-		XLongRangeConstraints->SetStiffness(TetherStiffness);
-		XLongRangeConstraints->SetScale(TetherScale);
-	}
 }
 
 void FClothConstraints::SetMaximumDistanceProperties(Softs::FSolverReal MaxDistancesMultiplier)
diff --git a/Engine/Plugins/Experimental/ChaosCloth/Source/Chaos/Private/ChaosCloth/ChaosClothingSimulationCloth.cpp b/Engine/Plugins/Experimental/ChaosCloth/Source/Chaos/Private/ChaosCloth/ChaosClothingSimulationCloth.cpp
--- a/Engine/Plugins/Experimental/ChaosCloth/Source/Chaos/Private/ChaosCloth/ChaosClothingSimulationCloth.cpp
+++ b/Engine/Plugins/Experimental/ChaosCloth/Source/Chaos/Private/ChaosCloth/ChaosClothingSimulationCloth.cpp
@@ -210,8 +210,7 @@ void FClothingSimulationCloth::FLODData::Add(FClothingSimulationSolver* Solver,
 			Tethers,
 			TetherStiffnessMultipliers,
 			TetherScaleMultipliers,
-			Softs::FSolverVec2((Softs::FSolverReal)Cloth->TetherScale[0], (Softs::FSolverReal)Cloth->TetherScale[1]) * MeshScale,
-			Cloth->bUseXPBDConstraints);
+			Softs::FSolverVec2((Softs::FSolverReal)Cloth->TetherScale[0], (Softs::FSolverReal)Cloth->TetherScale[1]) * MeshScale);
 	}
 
 	// Max distances
diff --git a/Engine/Plugins/Experimental/ChaosCloth/Source/Chaos/Public/ChaosCloth/ChaosClothConstraints.h b/Engine/Plugins/Experimental/ChaosCloth/Source/Chaos/Public/ChaosCloth/ChaosClothConstraints.h
--- a/Engine/Plugins/Experimental/ChaosCloth/Source/Chaos/Public/ChaosCloth/ChaosClothConstraints.h
+++ b/Engine/Plugins/Experimental/ChaosCloth/Source/Chaos/Public/ChaosCloth/ChaosClothConstraints.h
@@ -37,7 +37,13 @@ namespace Chaos
 		void SetVolumeConstraints(TArray<TVec3<int32>>&& SurfaceElements, Softs::FSolverReal VolumeStiffness);
 		void SetLongRangeConstraints(const TArray<TConstArrayView<TTuple<int32, int32, FRealSingle>>>& Tethers,
 			const TConstArrayView<FRealSingle>& TetherStiffnessMultipliers, const TConstArrayView<FRealSingle>& TetherScaleMultipliers,
-			const Softs::FSolverVec2& TetherScale, bool bUseXPBDConstraints);
+			const Softs::FSolverVec2& TetherScale);
+		UE_DEPRECATED(5.1, "LongRange XPBDConstraints are not supported.")
+		void SetLongRangeConstraints(const TArray<TConstArrayView<TTuple<int32, int32, FRealSingle>>>& Tethers,
+			const TConstArrayView<FRealSingle>& TetherStiffnessMultipliers, const TConstArrayView<FRealSingle>& TetherScaleMultipliers,
+			const Softs::FSolverVec2& TetherScale, bool bUseXPBDConstraints) {
+			return SetLongRangeConstraints(Tethers, TetherStiffnessMultipliers, TetherScaleMultipliers, TetherScale);
+		}
 		void SetMaximumDistanceConstraints(const TConstArrayView<FRealSingle>& MaxDistances);
 		void SetBackstopConstraints(const TConstArrayView<FRealSingle>& BackstopDistances, const TConstArrayView<FRealSingle>& BackstopRadiuses, bool bUseLegacyBackstop);
 		void SetAnimDriveConstraints(const TConstArrayView<FRealSingle>& AnimDriveStiffnessMultipliers, const TConstArrayView<FRealSingle>& AnimDriveDampingMultipliers);
@@ -75,7 +81,13 @@ namespace Chaos
 		const TSharedPtr<Softs::FPBDSpringConstraints>& GetThinShellVolumeConstraints() const { return ThinShellVolumeConstraints; }
 		const TSharedPtr<Softs::FPBDVolumeConstraint>& GetVolumeConstraints() const { return VolumeConstraints; }
 		const TSharedPtr<Softs::FPBDLongRangeConstraints>& GetLongRangeConstraints() const { return LongRangeConstraints; }
-		const TSharedPtr<Softs::FXPBDLongRangeConstraints>& GetXLongRangeConstraints() const { return XLongRangeConstraints; }
+		UE_DEPRECATED(5.1,"LongRange XPBDConstraints are not supported.")
+		const TSharedPtr<Softs::FXPBDLongRangeConstraints>& GetXLongRangeConstraints() const 
+		{
+			PRAGMA_DISABLE_DEPRECATION_WARNINGS
+			return XLongRangeConstraints; 
+			PRAGMA_ENABLE_DEPRECATION_WARNINGS
+		}
 		const TSharedPtr<Softs::FPBDSphericalConstraint>& GetMaximumDistanceConstraints() const { return MaximumDistanceConstraints; }
 		const TSharedPtr<Softs::FPBDSphericalBackstopConstraint>& GetBackstopConstraints() const { return BackstopConstraints; }
 		const TSharedPtr<Softs::FPBDAnimDriveConstraint>& GetAnimDriveConstraints() const { return AnimDriveConstraints; }
@@ -95,7 +107,8 @@ namespace Chaos
 		TSharedPtr<Softs::FXPBDAxialSpringConstraints> XAreaConstraints;
 		TSharedPtr<Softs::FPBDSpringConstraints> ThinShellVolumeConstraints;
 		TSharedPtr<Softs::FPBDVolumeConstraint> VolumeConstraints;
-		TSharedPtr<Softs::FPBDLongRangeConstraints> LongRangeConstraints;
+		TSharedPtr<Softs::FPBDLongRangeConstraints> LongRangeConstraints; 
+		UE_DEPRECATED(5.1, "LongRange XPBDConstraints are not supported.")
 		TSharedPtr<Softs::FXPBDLongRangeConstraints> XLongRangeConstraints;
 		TSharedPtr<Softs::FPBDSphericalConstraint> MaximumDistanceConstraints;
 		TSharedPtr<Softs::FPBDSphericalBackstopConstraint> BackstopConstraints;
diff --git a/Engine/Source/Runtime/Experimental/Chaos/Private/Chaos/PBDLongRangeConstraintsBase.cpp b/Engine/Source/Runtime/Experimental/Chaos/Private/Chaos/PBDLongRangeConstraintsBase.cpp
--- a/Engine/Source/Runtime/Experimental/Chaos/Private/Chaos/PBDLongRangeConstraintsBase.cpp
+++ b/Engine/Source/Runtime/Experimental/Chaos/Private/Chaos/PBDLongRangeConstraintsBase.cpp
@@ -65,12 +65,8 @@ FPBDLongRangeConstraintsBase::FPBDLongRangeConstraintsBase(
 	ApplyProperties((FSolverReal)(1. / FPBDStiffness::ParameterFrequency), 1);
 }
 
-void FPBDLongRangeConstraintsBase::ApplyProperties(const FSolverReal Dt, const int32 NumIterations)
+void FPBDLongRangeConstraintsBase::ApplyScale()
 {
-	// Apply stiffness
-	Stiffness.ApplyValues(Dt, NumIterations);
-
-	// Apply scale
 	const FSolverReal Offset = Scale[0];
 	const FSolverReal Range = Scale[1] - Scale[0];
 	const int32 ScaleTableSize = ScaleTable.Num();
diff --git a/Engine/Source/Runtime/Experimental/Chaos/Public/Chaos/PBDLongRangeConstraintsBase.h b/Engine/Source/Runtime/Experimental/Chaos/Public/Chaos/PBDLongRangeConstraintsBase.h
--- a/Engine/Source/Runtime/Experimental/Chaos/Public/Chaos/PBDLongRangeConstraintsBase.h
+++ b/Engine/Source/Runtime/Experimental/Chaos/Public/Chaos/PBDLongRangeConstraintsBase.h
@@ -49,8 +49,13 @@ public:
 	// Set the scale low and high value of the scale weight map
 	void SetScale(const FSolverVec2& InScale) { Scale = FSolverVec2(FMath::Clamp((FSolverReal)InScale.X, (FSolverReal)0.01, (FSolverReal)10.), FMath::Clamp((FSolverReal)InScale.Y, (FSolverReal)0.01, (FSolverReal)10.)); }  // TODO: Fix FSolverVec2 double to float conversion error on some platforms as to remove extra cast to FSolverReal
 
+public:
 	// Set stiffness offset and range, as well as the simulation stiffness exponent
-	void ApplyProperties(const FSolverReal Dt, const int32 NumIterations);
+	void ApplyProperties(const FSolverReal Dt, const int32 NumIterations)
+	{
+		Stiffness.ApplyValues(Dt, NumIterations);
+		ApplyScale();
+	}
 
 	// Return the tethers, organized in concurent friendly batches
 	const TArray<TConstArrayView<FTether>>& GetTethers() const { return Tethers; }
@@ -83,6 +88,10 @@ protected:
 	// Return whether the constraint has been setup with a weightmap to interpolate between two low and high values of scales
 	bool HasScaleWeightMap() const { return ScaleTable.Num() > 1; }
 
+	// Part of ApplyProperties to update ScaleTable.
+	void ApplyScale();
+
+
 	// Return a vector representing the amount of segment required for the tether to shrink back to its maximum target length constraint, or zero if the constraint is already met
 	inline FSolverVec3 GetDelta(const FSolverParticles& Particles, const FTether& Tether, const FSolverReal InScale) const
 	{
@@ -120,5 +129,4 @@ protected:
 	const int32 ParticleOffset;  // Index of the first usable particle
 	const int32 ParticleCount;  // Number of particles available to this constraint
 };
-
 }  // End namespace Chaos::Softs
diff --git a/Engine/Source/Runtime/Experimental/Chaos/Public/Chaos/PBDStiffness.h b/Engine/Source/Runtime/Experimental/Chaos/Public/Chaos/PBDStiffness.h
--- a/Engine/Source/Runtime/Experimental/Chaos/Public/Chaos/PBDStiffness.h
+++ b/Engine/Source/Runtime/Experimental/Chaos/Public/Chaos/PBDStiffness.h
@@ -68,6 +68,9 @@ public:
 	/** Update the weight map table with the current simulation parameters. */
 	inline void ApplyValues(const FSolverReal Dt, const int32 NumIterations);
 
+	/** Update the weight map table with the current simulation parameters. */
+	inline void ApplyXPBDValues(const FSolverReal MinStiffness, const FSolverReal MaxStiffnesss);
+
 	/**
 	 * Lookup for the exponential weighted value at the specified weight map index.
 	 * This function will assert if it is called with a non zero index on an empty weight map.
@@ -195,14 +198,16 @@ void FPBDStiffness::ApplyValues(const FSolverReal Dt, const int32 NumIterations)
 	{
 		// If InValue is 1 then LogValue = -inf and the output becomes -inf as well
 		// In order for this function to be continuous, We want the output to be 1 when InValue = 1
-		if (InValue == (FSolverReal)1.)
+		const FSolverReal ClampedInValue = FMath::Clamp(InValue, (FSolverReal)0., (FSolverReal)1.);
+
+		if (ClampedInValue == (FSolverReal)1.)
 		{
 			return (FSolverReal)1.;
 		}
 		// Get a very steep exponential curve between the [0, 1] range to make easier to set the parameter
 		// The base has been chosen empirically
 		// ParameterValue = Pow(ParameterFitBase, ParameterValue - 1)
-		const FSolverReal ParameterFit = FMath::Exp(ParameterFitLogBase * (FMath::Clamp(InValue, (FSolverReal)0., (FSolverReal)1.) - (FSolverReal)1.));
+		const FSolverReal ParameterFit = FMath::Exp(ParameterFitLogBase * (ClampedInValue - (FSolverReal)1.));
 
 		// Use simulation dependent stiffness exponent to alleviate the variations in effect when Dt and NumIterations change
 		// This is based on the Position-Based Simulation Methods paper (page 8),
@@ -222,6 +227,36 @@ void FPBDStiffness::ApplyValues(const FSolverReal Dt, const int32 NumIterations)
 	}
 }
 
+void FPBDStiffness::ApplyXPBDValues(const FSolverReal MinStiffness, const FSolverReal MaxStiffness)
+{
+	SCOPE_CYCLE_COUNTER(STAT_PBD_StiffnessApplyValues);
+	
+	// XPBD internally handles the effects of iterations and dt. 
+	// PBD stiffness is more like a constraint compliance. This should scale like 1/ xpbd stiffness which 
+	// is an actual measure of stiffness (e.g., something with units like Pascals)	
+	const FSolverReal MinOverMaxStiffness = MinStiffness / MaxStiffness;
+	auto SimulationValue = [this, MinOverMaxStiffness, MinStiffness, MaxStiffness](const FSolverReal InValue)->FSolverReal
+	{
+		const FSolverReal ClampedInValue = FMath::Clamp(InValue, (FSolverReal)0., (FSolverReal)1.);
+		if (ClampedInValue == (FSolverReal)1.)
+		{
+			return MaxStiffness;
+		}
+		const FSolverReal ParameterFit = FMath::Exp(ParameterFitLogBase * (ClampedInValue - (FSolverReal)1.));
+		return MinStiffness / ((FSolverReal)1. - ParameterFit + MinOverMaxStiffness);
+	};
+
+	const FSolverReal Offset = WeightedValue[0];
+	const FSolverReal Range = WeightedValue[1] - WeightedValue[0];
+	const int32 TableSize = Table.Num();
+	const FSolverReal WeightIncrement = (TableSize > 1) ? (FSolverReal)1. / (FSolverReal)(TableSize - 1) : (FSolverReal)1.; // Must allow full range from 0 to 1 included
+	for (int32 Index = 0; Index < TableSize; ++Index)
+	{
+		const FSolverReal Weight = (FSolverReal)Index * WeightIncrement;
+		Table[Index] = SimulationValue(Offset + Weight * Range);
+	}
+}
+
 void FPBDStiffness::ReorderIndices(const TArray<int32>& OrigToReorderedConstraintIndices)
 {
 	if (Indices.Num() == OrigToReorderedConstraintIndices.Num())
diff --git a/Engine/Source/Runtime/Experimental/Chaos/Public/Chaos/XPBDAxialSpringConstraints.h b/Engine/Source/Runtime/Experimental/Chaos/Public/Chaos/XPBDAxialSpringConstraints.h
--- a/Engine/Source/Runtime/Experimental/Chaos/Public/Chaos/XPBDAxialSpringConstraints.h
+++ b/Engine/Source/Runtime/Experimental/Chaos/Public/Chaos/XPBDAxialSpringConstraints.h
@@ -10,7 +10,8 @@ namespace Chaos::Softs
 {
 
 // Stiffness is in N/CM^2, so it needs to be adjusted from the PBD stiffness ranging between [0,1]
-static const double XPBDAxialSpringMaxCompliance = 1e-7;  // Max stiffness: 1e+11 N/M^2 = 1e+7 N/CM^2 -> Max compliance: 1e-7 CM^2/N
+static const FSolverReal XPBDAxialSpringMinStiffness = (FSolverReal)1e-1;
+static const FSolverReal XPBDAxialSpringMaxStiffness = (FSolverReal)1e7;  // Max stiffness: 1e+11 N/M^2 = 1e+7 N/CM^2
 
 class FXPBDAxialSpringConstraints final : public FPBDAxialSpringConstraintsBase
 {
@@ -36,6 +37,8 @@ public:
 
 	virtual ~FXPBDAxialSpringConstraints() override {}
 
+	void ApplyProperties(const FSolverReal Dt, const int32 NumIterations) { Stiffness.ApplyXPBDValues(XPBDAxialSpringMinStiffness, XPBDAxialSpringMaxStiffness); }
+
 	void Init() const { for (FSolverReal& Lambda : Lambdas) { Lambda = (FSolverReal)0.; } }
 
 	void Apply(FSolverParticles& Particles, const FSolverReal Dt) const
@@ -125,7 +128,7 @@ private:
 		const FSolverReal Offset = (Distance - Dists[InConstraintIndex]);
 
 		FSolverReal& Lambda = Lambdas[InConstraintIndex];
-		const FSolverReal Alpha = (FSolverReal)XPBDAxialSpringMaxCompliance / (ExpStiffnessValue * Dt * Dt);
+		const FSolverReal Alpha = (FSolverReal)1 / (ExpStiffnessValue * Dt * Dt);
 
 		const FSolverReal DLambda = (Offset - Alpha * Lambda) / (CombinedInvMass + Alpha);
 		const FSolverVec3 Delta = DLambda * Direction;
diff --git a/Engine/Source/Runtime/Experimental/Chaos/Public/Chaos/XPBDLongRangeConstraints.h b/Engine/Source/Runtime/Experimental/Chaos/Public/Chaos/XPBDLongRangeConstraints.h
--- a/Engine/Source/Runtime/Experimental/Chaos/Public/Chaos/XPBDLongRangeConstraints.h
+++ b/Engine/Source/Runtime/Experimental/Chaos/Public/Chaos/XPBDLongRangeConstraints.h
@@ -11,7 +11,8 @@ namespace Chaos::Softs
 {
 
 // Stiffness is in N/CM^2, so it needs to be adjusted from the PBD stiffness ranging between [0,1]
-static const double XPBDLongRangeMaxCompliance = 1.e-3;
+static const FSolverReal XPBDLongRangeMinStiffness = (FSolverReal)1e-1;
+static const FSolverReal XPBDLongRangeMaxStiffness = (FSolverReal)1e7;
 
 class FXPBDLongRangeConstraints final : public FPBDLongRangeConstraintsBase
 {
@@ -30,7 +31,7 @@ public:
 		const FSolverVec2& InScale = FSolverVec2::UnitVector)
 	    : FPBDLongRangeConstraintsBase(Particles, InParticleOffset, InParticleCount, InTethers, StiffnessMultipliers, ScaleMultipliers, InStiffness, InScale)
 	{
-		int32 NumTethers = 0;
+		NumTethers = 0;
 		for (const TConstArrayView<FTether>& TetherBatch : Tethers)
 		{
 			NumTethers += TetherBatch.Num();
@@ -40,10 +41,13 @@ public:
 
 	virtual ~FXPBDLongRangeConstraints() override {}
 
+	// Set stiffness offset and range, as well as the simulation stiffness exponent
+	void ApplyProperties(const FSolverReal Dt, const int32 NumIterations) { Stiffness.ApplyXPBDValues(XPBDLongRangeMinStiffness, XPBDLongRangeMaxStiffness); ApplyScale(); }
+
 	void Init() const
 	{
 		Lambdas.Reset();
-		Lambdas.AddZeroed(Tethers.Num());
+		Lambdas.AddZeroed(NumTethers);
 	}
 
 	void Apply(FSolverParticles& Particles, const FSolverReal Dt) const 
@@ -131,7 +135,7 @@ private:
 		GetDelta(Particles, Tether, InScale, Direction, Offset);
 
 		FSolverReal& Lambda = Lambdas[ConstraintIndex];
-		const FSolverReal Alpha = (FSolverReal)XPBDLongRangeMaxCompliance / (InStiffness * Dt * Dt);
+		const FSolverReal Alpha = (FSolverReal)1 / (InStiffness * Dt * Dt);
 
 		const FSolverReal DLambda = (Offset - Alpha * Lambda) / ((FSolverReal)1. + Alpha);
 		Particles.P(GetEndParticle(Tether)) += DLambda * Direction;
@@ -146,6 +150,7 @@ private:
 	using Base::ScaleIndices;
 
 	mutable TArray<FSolverReal> Lambdas;
+	int32 NumTethers;
 };
 
 }  // End namespace Chaos::Softs
diff --git a/Engine/Source/Runtime/Experimental/Chaos/Public/Chaos/XPBDSpringConstraints.h b/Engine/Source/Runtime/Experimental/Chaos/Public/Chaos/XPBDSpringConstraints.h
--- a/Engine/Source/Runtime/Experimental/Chaos/Public/Chaos/XPBDSpringConstraints.h
+++ b/Engine/Source/Runtime/Experimental/Chaos/Public/Chaos/XPBDSpringConstraints.h
@@ -10,7 +10,8 @@ namespace Chaos::Softs
 {
 
 // Stiffness is in N/CM^2, so it needs to be adjusted from the PBD stiffness ranging between [0,1]
-static const double XPBDSpringMaxCompliance = 1e-7;  // Max stiffness: 1e+11 N/M^2 = 1e+7 N/CM^2 -> Max compliance: 1e-7 CM^2/N
+static const FSolverReal XPBDSpringMinStiffness = (FSolverReal)1e-1; // Min stiffness: 1e3 N/M^2 = 1e-1 N/CM^2. Empirically defined to have similar look to XPBDSpringConstraints.
+static const FSolverReal XPBDSpringMaxStiffness = (FSolverReal)1e7;  // Max stiffness: 1e+11 N/M^2 = 1e+7 N/CM^2
 
 class FXPBDSpringConstraints final : public FPBDSpringConstraintsBase
 {
@@ -38,6 +39,9 @@ public:
 	virtual ~FXPBDSpringConstraints() override {}
 
 	void Init() const { for (FSolverReal& Lambda : Lambdas) { Lambda = (FSolverReal)0.; } }
+	
+	// Update stiffness table, as well as the simulation stiffness exponent
+	inline void ApplyProperties(const FSolverReal Dt, const int32 NumIterations) { Stiffness.ApplyXPBDValues(XPBDSpringMinStiffness, XPBDSpringMaxStiffness); }
 
 	void Apply(FSolverParticles& Particles, const FSolverReal Dt, const int32 ConstraintIndex, const FSolverReal ExpStiffnessValue) const
 	{
@@ -98,7 +102,7 @@ private:
 		const FSolverReal Offset = Distance - Dists[ConstraintIndex];
 
 		FSolverReal& Lambda = Lambdas[ConstraintIndex];
-		const FSolverReal Alpha = (FSolverReal)XPBDSpringMaxCompliance / (ExpStiffnessValue * Dt * Dt);
+		const FSolverReal Alpha = (FSolverReal)1.f / (ExpStiffnessValue * Dt * Dt);
 
 		const FSolverReal DLambda = (Offset - Alpha * Lambda) / (CombinedInvMass + Alpha);
 		const FSolverVec3 Delta = DLambda * Direction;
