* Fix an issue with debug draw not properly drawing swept boundaries for ccd
* Fix an issue where Dt was being subtracted instead of multiplied in a particular overload of DeltaExceedsThreshold

diff --git a/Engine/Source/Runtime/Experimental/Chaos/Private/Chaos/CCDUtilities.cpp b/Engine/Source/Runtime/Experimental/Chaos/Private/Chaos/CCDUtilities.cpp
--- a/Engine/Source/Runtime/Experimental/Chaos/Private/Chaos/CCDUtilities.cpp
+++ b/Engine/Source/Runtime/Experimental/Chaos/Private/Chaos/CCDUtilities.cpp
@@ -742,8 +742,8 @@ namespace Chaos
 		// For non-rigids, DeltaX is zero and use R for rotation.
 		const auto Rigid0 = Particle0.CastToRigidParticle();
 		const auto Rigid1 = Particle1.CastToRigidParticle();
-		const FVec3 DeltaX0 = Rigid0 ? Rigid0->V() - Dt : FVec3::ZeroVector;
-		const FVec3 DeltaX1 = Rigid1 ? Rigid1->V() - Dt : FVec3::ZeroVector;
+		const FVec3 DeltaX0 = Rigid0 ? Rigid0->V() * Dt : FVec3::ZeroVector;
+		const FVec3 DeltaX1 = Rigid1 ? Rigid1->V() * Dt : FVec3::ZeroVector;
 		const FQuat& R0 = Rigid0 ? Rigid0->Q() : Particle0.R();
 		const FQuat& R1 = Rigid1 ? Rigid1->Q() : Particle1.R();
 		return DeltaExceedsThreshold(
diff --git a/Engine/Source/Runtime/Experimental/Chaos/Private/Chaos/ChaosDebugDraw.cpp b/Engine/Source/Runtime/Experimental/Chaos/Private/Chaos/ChaosDebugDraw.cpp
--- a/Engine/Source/Runtime/Experimental/Chaos/Private/Chaos/ChaosDebugDraw.cpp
+++ b/Engine/Source/Runtime/Experimental/Chaos/Private/Chaos/ChaosDebugDraw.cpp
@@ -1738,7 +1738,7 @@ namespace Chaos
 					{
 						const FVector LocalOffset = HalfExtents * FVector(ScaleArray[IX], ScaleArray[IY], ScaleArray[IZ]);
 						const FVector Offset = R.RotateVector(LocalOffset);
-						FDebugDrawQueue::GetInstance().DrawDebugLine(X, X + DeltaX, BoundsColor, false, -1.f, -1, BoundsThickness);
+						FDebugDrawQueue::GetInstance().DrawDebugLine(X + Offset, X + Offset + DeltaX, BoundsColor, false, -1.f, -1, BoundsThickness);
 					}
 				}
 			}
