don't build a reverse map when updating geometry collection bounding boxes, because we already have one on the geometry collection

diff --git a/Engine/Source/Runtime/Experimental/Chaos/Private/GeometryCollection/GeometryCollection.cpp b/Engine/Source/Runtime/Experimental/Chaos/Private/GeometryCollection/GeometryCollection.cpp
--- a/Engine/Source/Runtime/Experimental/Chaos/Private/GeometryCollection/GeometryCollection.cpp
+++ b/Engine/Source/Runtime/Experimental/Chaos/Private/GeometryCollection/GeometryCollection.cpp
@@ -869,17 +869,11 @@ void FGeometryCollection::UpdateBoundingBox()
 			BoundingBox[Idx].Init();
 		}
 
-		// Build reverse map between TransformIdx and index in the GeometryGroup
-		TMap<int32, int32> GeometryGroupIndexMap;
-		for (int32 Idx = 0; Idx < NumElements(FGeometryCollection::GeometryGroup); ++Idx)
-		{
-			GeometryGroupIndexMap.Add(TransformIndex[Idx], Idx);
-		}
 		// Compute BoundingBox
 		for (int32 Idx = 0; Idx < Vertex.Num(); ++Idx)
 		{
 			int32 TransformIndexValue = BoneMap[Idx];
-			BoundingBox[GeometryGroupIndexMap[TransformIndexValue]] += FVector(Vertex[Idx]);
+			BoundingBox[TransformToGeometryIndex[TransformIndexValue]] += FVector(Vertex[Idx]);
 		}
 	}
 }
