Choas : GeometryCollection
- Added update to the bounding box through Init. Used by the command line importer.
- Added GetBoundingBox accessor to the collection that sums the bounds.

diff --git a/Engine/Source/Runtime/Experimental/Chaos/Private/GeometryCollection/GeometryCollection.cpp b/Engine/Source/Runtime/Experimental/Chaos/Private/GeometryCollection/GeometryCollection.cpp
--- a/Engine/Source/Runtime/Experimental/Chaos/Private/GeometryCollection/GeometryCollection.cpp
+++ b/Engine/Source/Runtime/Experimental/Chaos/Private/GeometryCollection/GeometryCollection.cpp
@@ -878,6 +878,17 @@ void FGeometryCollection::UpdateBoundingBox()
 	}
 }
 
+FBoxSphereBounds FGeometryCollection::GetBoundingBox() const 
+{
+	FBoxSphereBounds ResultingBox = FBoxSphereBounds(ForceInitToZero);
+	for (int32 Idx = 0; Idx < NumElements(FGeometryCollection::GeometryGroup); ++Idx)
+	{
+		ResultingBox.ExpandBy( BoundingBox[Idx].GetExtent().Length() );
+	}
+	return ResultingBox;
+}
+
+
 void FGeometryCollection::Serialize(Chaos::FChaosArchive& Ar)
 {
 	if (Ar.IsCooking())
diff --git a/Engine/Source/Runtime/Experimental/Chaos/Private/GeometryCollection/GeometryCollectionUtility.cpp b/Engine/Source/Runtime/Experimental/Chaos/Private/GeometryCollection/GeometryCollectionUtility.cpp
--- a/Engine/Source/Runtime/Experimental/Chaos/Private/GeometryCollection/GeometryCollectionUtility.cpp
+++ b/Engine/Source/Runtime/Experimental/Chaos/Private/GeometryCollection/GeometryCollectionUtility.cpp
@@ -411,6 +411,8 @@ namespace GeometryCollection
 						}
 					}
 				}
+
+				Collection->UpdateBoundingBox();
 			}
 		}
 	}
diff --git a/Engine/Source/Runtime/Experimental/Chaos/Public/GeometryCollection/GeometryCollection.h b/Engine/Source/Runtime/Experimental/Chaos/Public/GeometryCollection/GeometryCollection.h
--- a/Engine/Source/Runtime/Experimental/Chaos/Public/GeometryCollection/GeometryCollection.h
+++ b/Engine/Source/Runtime/Experimental/Chaos/Public/GeometryCollection/GeometryCollection.h
@@ -186,6 +186,11 @@ public:
 	*/
 	void UpdateBoundingBox();
 
+	/**  
+	* GetBoundingBox 
+	*/
+	FBoxSphereBounds GetBoundingBox() const;
+
 	/**
 	 * Update the visibility of specified geometry nodes
 	 */
