Chaos : Flesh
- Simulatiion against imported asset.

diff --git a/Engine/Source/Runtime/Experimental/Chaos/Public/GeometryCollection/ManagedArrayTypeValues.inl b/Engine/Source/Runtime/Experimental/Chaos/Public/GeometryCollection/ManagedArrayTypeValues.inl
--- a/Engine/Source/Runtime/Experimental/Chaos/Public/GeometryCollection/ManagedArrayTypeValues.inl
+++ b/Engine/Source/Runtime/Experimental/Chaos/Public/GeometryCollection/ManagedArrayTypeValues.inl
@@ -43,6 +43,7 @@ MANAGED_ARRAY_TYPE(TUniquePtr<Chaos::FConvex>, FConvexUniquePtr)
 MANAGED_ARRAY_TYPE(TArray<FVector2f>, Vector2DArray)
 MANAGED_ARRAY_TYPE(double, Double)
 MANAGED_ARRAY_TYPE(FIntVector4, IntVector4)
+MANAGED_ARRAY_TYPE(FVector3d, Vector3d)
 
 // NOTE: new types must be added at the bottom to keep serialization from breaking
 
