make fracture 'hide' array only hide already-visible faces, and never show intended-invisible faces

diff --git a/Engine/Source/Runtime/Experimental/GeometryCollectionEngine/Private/GeometryCollection/GeometryCollectionComponent.cpp b/Engine/Source/Runtime/Experimental/GeometryCollectionEngine/Private/GeometryCollection/GeometryCollectionComponent.cpp
--- a/Engine/Source/Runtime/Experimental/GeometryCollectionEngine/Private/GeometryCollection/GeometryCollectionComponent.cpp
+++ b/Engine/Source/Runtime/Experimental/GeometryCollectionEngine/Private/GeometryCollection/GeometryCollectionComponent.cpp
@@ -1606,7 +1606,11 @@ void UGeometryCollectionComponent::InitConstantData(FGeometryCollectionConstantD
 #if WITH_EDITOR
 		// We will override visibility with the Hide array (if available).
 		TArray<bool> VisibleOverride;
-		VisibleOverride.Init(true, Visible.Num());
+		VisibleOverride.SetNumUninitialized(Visible.Num());
+		for (int32 FaceIdx = 0; FaceIdx < Visible.Num(); FaceIdx++)
+		{
+			VisibleOverride[FaceIdx] = Visible[FaceIdx];
+		}
 		bool bUsingHideArray = false;
 
 		if (Collection->HasAttribute("Hide", FGeometryCollection::TransformGroup))
@@ -1618,7 +1622,8 @@ void UGeometryCollectionComponent::InitConstantData(FGeometryCollectionConstantD
 			const TManagedArray<bool>& Hide = Collection->GetAttribute<bool>("Hide", FGeometryCollection::TransformGroup);
 			for (int32 GeomIdx = 0; GeomIdx < NumGeom; ++GeomIdx)
 			{
-				if (Hide[TransformIndex[GeomIdx]])
+				int32 TransformIdx = TransformIndex[GeomIdx];
+				if (Hide[TransformIdx])
 				{
 					// (Temporarily) hide faces of this hidden geometry
 					for (int32 FaceIdxOffset = 0; FaceIdxOffset < FaceCount[GeomIdx]; ++FaceIdxOffset)
@@ -1626,16 +1631,16 @@ void UGeometryCollectionComponent::InitConstantData(FGeometryCollectionConstantD
 						VisibleOverride[FaceStart[GeomIdx]+FaceIdxOffset] = false;
 					}
 				}
-				else
+				else if (bAllHidden && Collection->IsVisible(TransformIdx))
 				{
 					bAllHidden = false;
 				}
 			}
-			if (!ensure(!bAllHidden)) // if they're all hidden, rendering would crash -- unhide them
+			if (!ensure(!bAllHidden)) // if they're all hidden, rendering would crash -- reset to default visibility instead
 			{
 				for (int32 FaceIdx = 0; FaceIdx < VisibleOverride.Num(); ++FaceIdx)
 				{
-					VisibleOverride[FaceIdx] = true;
+					VisibleOverride[FaceIdx] = Visible[FaceIdx];
 				}
 			}
 		}
