Chaos - comments in contact modifier

diff --git a/Engine/Source/Runtime/Experimental/Chaos/Public/Chaos/ContactModification.h b/Engine/Source/Runtime/Experimental/Chaos/Public/Chaos/ContactModification.h
--- a/Engine/Source/Runtime/Experimental/Chaos/Public/Chaos/ContactModification.h
+++ b/Engine/Source/Runtime/Experimental/Chaos/Public/Chaos/ContactModification.h
@@ -72,12 +72,17 @@ namespace Chaos
 		void ModifyTargetSeparation(FReal TargetSeparation, int32 ContactPointIdx);
 
 		/**
+		* @brief Get the world-space contact normal.
+		* @note The contact normal always points away from the second body.
+		* E.g., a sphere lying on a flat ground could return a WorldNormal pointing up or down, depending on whether the
+		* sphere is the first or second body in the constraint.
 		* @return Normal of contact in world space.
 		*/
 		FVec3 GetWorldNormal(int32 ContactPointIdx) const;
 
 		/**
 		* Modify contact normal in world space. If modifying separation and normal, order of operations should be considered.
+		* @note The contact normal should always point away from the second body.
 		*/
 		void ModifyWorldNormal(const FVec3& Normal, int32 ContactPointIdx);
 
