incorporate transforms + center into FGeometryCollection::GetBounds

diff --git a/Engine/Source/Runtime/Experimental/Chaos/Private/GeometryCollection/GeometryCollection.cpp b/Engine/Source/Runtime/Experimental/Chaos/Private/GeometryCollection/GeometryCollection.cpp
--- a/Engine/Source/Runtime/Experimental/Chaos/Private/GeometryCollection/GeometryCollection.cpp
+++ b/Engine/Source/Runtime/Experimental/Chaos/Private/GeometryCollection/GeometryCollection.cpp
@@ -880,12 +880,16 @@ void FGeometryCollection::UpdateBoundingBox()
 
 FBoxSphereBounds FGeometryCollection::GetBoundingBox() const 
 {
-	FBoxSphereBounds ResultingBox = FBoxSphereBounds(ForceInitToZero);
-	for (int32 Idx = 0; Idx < NumElements(FGeometryCollection::GeometryGroup); ++Idx)
+	TArray<FTransform> GlobalTransformArray;
+	GeometryCollectionAlgo::GlobalMatrices(Transform, Parent, GlobalTransformArray);
+	FBox CombinedBounds(EForceInit::ForceInit);
+	for (int32 GeoIdx = 0; GeoIdx < NumElements(FGeometryCollection::GeometryGroup); ++GeoIdx)
 	{
-		ResultingBox.ExpandBy( BoundingBox[Idx].GetExtent().Length() );
+		int32 TransformIdx = TransformIndex[GeoIdx];
+		CombinedBounds += BoundingBox[GeoIdx].TransformBy(GlobalTransformArray[TransformIdx]);
 	}
-	return ResultingBox;
+	FBoxSphereBounds CombinedBoxSphereBounds(CombinedBounds);
+	return CombinedBoxSphereBounds;
 }
 
 
