Add missing GetAbs() to two particle version of CCDHelpers::DeltaExceedsThreshold

diff --git a/Engine/Source/Runtime/Experimental/Chaos/Private/Chaos/CCDUtilities.cpp b/Engine/Source/Runtime/Experimental/Chaos/Private/Chaos/CCDUtilities.cpp
--- a/Engine/Source/Runtime/Experimental/Chaos/Private/Chaos/CCDUtilities.cpp
+++ b/Engine/Source/Runtime/Experimental/Chaos/Private/Chaos/CCDUtilities.cpp
@@ -710,7 +710,7 @@ namespace Chaos
 			//
 			// Since there's no preferred inertial frame, the correct choice
 			// will always be to take the smaller object's threshold.
-			AxisThreshold0.ComponentMin(R0 * R1.UnrotateVector(AxisThreshold1)),
+			AxisThreshold0.ComponentMin((R0 * R1.UnrotateVector(AxisThreshold1)).GetAbs()),
 
 			// Taking the difference of the deltas gives the total delta - how
 			// much the objects have moved towards each other. We choose to use
