FIX ROBOMERGE

Add Chris' ccd debug draw stuff from 19636305

diff --git a/Engine/Source/Runtime/Experimental/Chaos/Private/Chaos/CCDUtilities.cpp b/Engine/Source/Runtime/Experimental/Chaos/Private/Chaos/CCDUtilities.cpp
--- a/Engine/Source/Runtime/Experimental/Chaos/Private/Chaos/CCDUtilities.cpp
+++ b/Engine/Source/Runtime/Experimental/Chaos/Private/Chaos/CCDUtilities.cpp
@@ -31,6 +31,9 @@ namespace Chaos
 	{
 		FRealSingle CCDEnableThresholdBoundsScale = 0.4f;
 		FAutoConsoleVariableRef  CVarCCDEnableThresholdBoundsScale(TEXT("p.Chaos.CCD.EnableThresholdBoundsScale"), CCDEnableThresholdBoundsScale , TEXT("CCD is used when object position is changing > smallest bound's extent * BoundsScale. 0 will always Use CCD. Values < 0 disables CCD."));
+
+		extern int32 ChaosSolverDrawCCDInteractions;
+		extern DebugDraw::FChaosDebugDrawSettings ChaosSolverDebugDebugDrawSettings;	
 	}
 
 	void FCCDParticle::AddOverlappingDynamicParticle(FCCDParticle* const InParticle)
@@ -292,7 +295,17 @@ namespace Chaos
 		const int32 ConstraintEnd = IslandConstraintEnd[Island];
 		check(ConstraintNum > 0);
 
+#if CHAOS_DEBUG_DRAW
+		if (CVars::ChaosSolverDrawCCDInteractions)
+		{
 			// Debugdraw the shape at the TOI=0 (black) and TOI=1 (white)
+			for (FCCDConstraint* CCDConstraint : SortedCCDConstraints)
+			{
+				DebugDraw::DrawCCDCollisionShape(FRigidTransform3::Identity, *CCDConstraint, true, FColor::Black, &CVars::ChaosSolverDebugDebugDrawSettings);
+				DebugDraw::DrawCCDCollisionShape(FRigidTransform3::Identity, *CCDConstraint, false, FColor::White, &CVars::ChaosSolverDebugDebugDrawSettings);
+			}
+		}
+#endif
 
 		// Sort constraints based on TOI
 		std::sort(SortedCCDConstraints.GetData() + ConstraintStart, SortedCCDConstraints.GetData() + ConstraintStart + ConstraintNum, CCDConstraintSortPredicate);
@@ -322,7 +335,6 @@ namespace Chaos
 				ConstraintIndex++;
 				continue;
 			}
-
 			// Skip CCD resolution if the depth at TOI=1 is low and let the standard contact resolution handle it.
 			// This can be important because CCD contact resolution uses point-mass physics and does not apply rotation 
 			// so we must leave some work for the regular solver to do to get decent behaviour when sliding fast along
@@ -527,11 +539,23 @@ namespace Chaos
 				}
 			}
 		}
-		
+
+		// Debugdraw the shapes at the final position
 		for (FCCDConstraint* CCDConstraint : SortedCCDConstraints)
-		{ 
+		{
 			CCDConstraint->SweptConstraint->SetCCDResults(CCDConstraint->NetImpulse);
 		}
+		
+#if CHAOS_DEBUG_DRAW
+		// Debugdraw the shapes at the final position
+		if (CVars::ChaosSolverDrawCCDInteractions)
+		{
+			for (FCCDConstraint* CCDConstraint : SortedCCDConstraints)
+			{
+				DebugDraw::DrawCCDCollisionShape(FRigidTransform3::Identity, *CCDConstraint, false, FColor::Green, &CVars::ChaosSolverDebugDebugDrawSettings);
+			}
+		}
+#endif
 	}
 
 	void FCCDManager::ResetIslandParticles(const int32 Island)
@@ -622,7 +646,6 @@ namespace Chaos
 				{
 					Rigid1->V() -= Impulse * InvM1;
 				}
-
 				CCDConstraint->NetImpulse += Impulse;
 			}
 		}
diff --git a/Engine/Source/Runtime/Experimental/Chaos/Private/Chaos/ChaosDebugDraw.cpp b/Engine/Source/Runtime/Experimental/Chaos/Private/Chaos/ChaosDebugDraw.cpp
--- a/Engine/Source/Runtime/Experimental/Chaos/Private/Chaos/ChaosDebugDraw.cpp
+++ b/Engine/Source/Runtime/Experimental/Chaos/Private/Chaos/ChaosDebugDraw.cpp
@@ -79,6 +79,9 @@ namespace Chaos
 		float ChaosDebugDrawConvexExplodeDistance = 0.0f;
 		FAutoConsoleVariableRef CVarChaosDebugDrawConvexExplodeDistance(TEXT("p.Chaos.DebugDraw.ConvexExplodeDistance"), ChaosDebugDrawConvexExplodeDistance, TEXT("Explode convex edges by this amount (useful for looking at convex integrity)"));
 
+		float ChaosDebugDrawCCDDuration = 0.0f;
+		FAutoConsoleVariableRef CVarChaosDebugDrawCCDDuration(TEXT("p.Chaos.DebugDraw.CCDDuration"), ChaosDebugDrawCCDDuration, TEXT("How long CCD debug draw should remain on screen"));
+
 
 		// NOTE: These settings should never really be used - they are the fallback defaults
 		// if the user does not specify settings in the debug draw call.
@@ -260,15 +263,15 @@ namespace Chaos
 		//
 		//
 
-		void DrawShapesImpl(const FGeometryParticleHandle* Particle, const FRigidTransform3& ShapeTransform, const FImplicitObject* Implicit, const FShapeOrShapesArray& Shapes, const FReal Margin, const FColor& Color, const FChaosDebugDrawSettings& Settings);
+		void DrawShapesImpl(const FGeometryParticleHandle* Particle, const FRigidTransform3& ShapeTransform, const FImplicitObject* Implicit, const FShapeOrShapesArray& Shapes, const FReal Margin, const FColor& Color, const FRealSingle Duration, const FChaosDebugDrawSettings& Settings);
 
 		void DrawShape(const FRigidTransform3& ShapeTransform, const FImplicitObject* Implicit, const FShapeOrShapesArray& Shapes, const FColor& Color, const FChaosDebugDrawSettings* Settings)
 		{
 			// Note: At the time of commenting this function does nothing as DrawShapesImpl does not handle null particle.
-			DrawShapesImpl(nullptr, ShapeTransform, Implicit, Shapes, 0.0f, Color, GetChaosDebugDrawSettings(Settings));
+			DrawShapesImpl(nullptr, ShapeTransform, Implicit, Shapes, 0.0f, Color, 0.0f, GetChaosDebugDrawSettings(Settings));
 		}
 
-		void DrawShapesConvexImpl(const TGeometryParticleHandle<FReal, 3>* Particle, const FRigidTransform3& ShapeTransform, const FConvex* Shape, const FReal InMargin, const FColor& Color, const FChaosDebugDrawSettings& Settings)
+		void DrawShapesConvexImpl(const TGeometryParticleHandle<FReal, 3>* Particle, const FRigidTransform3& ShapeTransform, const FConvex* Shape, const FReal InMargin, const FColor& Color, const FRealSingle Duration, const FChaosDebugDrawSettings& Settings)
 		{
 			if (Shape->HasStructureData())
 			{
@@ -289,7 +292,7 @@ namespace Chaos
 						const FVec3 ExplodeDelta = ChaosDebugDrawConvexExplodeDistance * N0;
 
 						// Outer shape
-						FDebugDrawQueue::GetInstance().DrawDebugLine(OuterP0 + ExplodeDelta, OuterP1 + ExplodeDelta, Color, false, -1.f, Settings.DrawPriority, Settings.ShapeThicknesScale * Settings.LineThickness);
+						FDebugDrawQueue::GetInstance().DrawDebugLine(OuterP0 + ExplodeDelta, OuterP1 + ExplodeDelta, Color, false, Duration, Settings.DrawPriority, Settings.ShapeThicknesScale * Settings.LineThickness);
 
 						// Core shape and lines connecting core to outer
 						if (Margin > 0.0f)
@@ -297,21 +300,21 @@ namespace Chaos
 							const FRealSingle LineThickness = 0.5f * Settings.ShapeThicknesScale * Settings.LineThickness;
 							const FVec3 InnerP0 = ShapeTransform.TransformPositionNoScale(Shape->GetMarginAdjustedVertexScaled(VertexIndex0, Margin, ShapeTransform.GetScale3D(), nullptr));
 							const FVec3 InnerP1 = ShapeTransform.TransformPositionNoScale(Shape->GetMarginAdjustedVertexScaled(VertexIndex1, Margin, ShapeTransform.GetScale3D(), nullptr));
-							FDebugDrawQueue::GetInstance().DrawDebugLine(InnerP0, InnerP1, FColor::Blue, false, -1.f, Settings.DrawPriority, LineThickness);
-							FDebugDrawQueue::GetInstance().DrawDebugLine(InnerP0, OuterP0, FColor::Black, false, -1.f, Settings.DrawPriority, LineThickness);
+							FDebugDrawQueue::GetInstance().DrawDebugLine(InnerP0, InnerP1, FColor::Blue, false, Duration, Settings.DrawPriority, LineThickness);
+							FDebugDrawQueue::GetInstance().DrawDebugLine(InnerP0, OuterP0, FColor::Black, false, Duration, Settings.DrawPriority, LineThickness);
 						}
 
 						// Vertex and face normal
 						if (bChaosDebugDebugDrawConvexVertices)
 						{
-							FDebugDrawQueue::GetInstance().DrawDebugLine(OuterP0 + ExplodeDelta, OuterP0 + ExplodeDelta + Settings.DrawScale * 20.0f * N0, FColor::Green, false, -1.f, Settings.DrawPriority, Settings.LineThickness);
+							FDebugDrawQueue::GetInstance().DrawDebugLine(OuterP0 + ExplodeDelta, OuterP0 + ExplodeDelta + Settings.DrawScale * 20.0f * N0, FColor::Green, false, Duration, Settings.DrawPriority, Settings.LineThickness);
 						}
 					}
 				}
 			}
 		}
 
-		void DrawShapesHeightFieldImpl(const TGeometryParticleHandle<FReal, 3>* Particle, const FRigidTransform3& ShapeTransform, const FHeightField* Shape, const FColor& Color, const FChaosDebugDrawSettings& Settings)
+		void DrawShapesHeightFieldImpl(const TGeometryParticleHandle<FReal, 3>* Particle, const FRigidTransform3& ShapeTransform, const FHeightField* Shape, const FColor& Color, const FRealSingle Duration, const FChaosDebugDrawSettings& Settings)
 		{
 			const FVec3& WorldQueryCenter = FDebugDrawQueue::GetInstance().GetCenterOfInterest();
 			const FReal WorldQueryRadius = FDebugDrawQueue::GetInstance().GetRadiusOfInterest();
@@ -323,13 +326,13 @@ namespace Chaos
 					FVec3 A = ShapeTransform.TransformPosition(Tri[0]);
 					FVec3 B = ShapeTransform.TransformPosition(Tri[1]);
 					FVec3 C = ShapeTransform.TransformPosition(Tri[2]);
-					FDebugDrawQueue::GetInstance().DrawDebugLine(A, B, Color, false, -1.f, 0, Settings.ShapeThicknesScale * Settings.LineThickness);
-					FDebugDrawQueue::GetInstance().DrawDebugLine(B, C, Color, false, -1.f, 0, Settings.ShapeThicknesScale * Settings.LineThickness);
-					FDebugDrawQueue::GetInstance().DrawDebugLine(C, A, Color, false, -1.f, 0, Settings.ShapeThicknesScale * Settings.LineThickness);
+					FDebugDrawQueue::GetInstance().DrawDebugLine(A, B, Color, false, Duration, 0, Settings.ShapeThicknesScale * Settings.LineThickness);
+					FDebugDrawQueue::GetInstance().DrawDebugLine(B, C, Color, false, Duration, 0, Settings.ShapeThicknesScale * Settings.LineThickness);
+					FDebugDrawQueue::GetInstance().DrawDebugLine(C, A, Color, false, Duration, 0, Settings.ShapeThicknesScale * Settings.LineThickness);
 				});
 		}
 
-		void DrawShapesTriangleMeshImpl(const TGeometryParticleHandle<FReal, 3>* Particle, const FRigidTransform3& ShapeTransform, const FTriangleMeshImplicitObject* Shape, const FColor& Color, const FChaosDebugDrawSettings& Settings)
+		void DrawShapesTriangleMeshImpl(const TGeometryParticleHandle<FReal, 3>* Particle, const FRigidTransform3& ShapeTransform, const FTriangleMeshImplicitObject* Shape, const FColor& Color, const FRealSingle Duration, const FChaosDebugDrawSettings& Settings)
 		{
 			const FVec3& WorldQueryCenter = FDebugDrawQueue::GetInstance().GetCenterOfInterest();
 			const FReal WorldQueryRadius = FDebugDrawQueue::GetInstance().GetRadiusOfInterest();
@@ -341,13 +344,13 @@ namespace Chaos
 				FVec3 A = ShapeTransform.TransformPosition(Tri[0]);
 				FVec3 B = ShapeTransform.TransformPosition(Tri[1]);
 				FVec3 C = ShapeTransform.TransformPosition(Tri[2]);
-				FDebugDrawQueue::GetInstance().DrawDebugLine(A, B, Color, false, -1.f, 0, Settings.ShapeThicknesScale * Settings.LineThickness);
-				FDebugDrawQueue::GetInstance().DrawDebugLine(B, C, Color, false, -1.f, 0, Settings.ShapeThicknesScale * Settings.LineThickness);
-				FDebugDrawQueue::GetInstance().DrawDebugLine(C, A, Color, false, -1.f, 0, Settings.ShapeThicknesScale * Settings.LineThickness);
+				FDebugDrawQueue::GetInstance().DrawDebugLine(A, B, Color, false, Duration, 0, Settings.ShapeThicknesScale * Settings.LineThickness);
+				FDebugDrawQueue::GetInstance().DrawDebugLine(B, C, Color, false, Duration, 0, Settings.ShapeThicknesScale * Settings.LineThickness);
+				FDebugDrawQueue::GetInstance().DrawDebugLine(C, A, Color, false, Duration, 0, Settings.ShapeThicknesScale * Settings.LineThickness);
 			});
 		}
 
-		void DrawShapesLevelSetImpl(const TGeometryParticleHandle<FReal, 3>* Particle, const FRigidTransform3& ShapeTransform, const FLevelSet* Shape, const FColor& Color, const FChaosDebugDrawSettings& Settings)
+		void DrawShapesLevelSetImpl(const TGeometryParticleHandle<FReal, 3>* Particle, const FRigidTransform3& ShapeTransform, const FLevelSet* Shape, const FColor& Color, const FRealSingle Duration, const FChaosDebugDrawSettings& Settings)
 		{
 			if (!Settings.bShowLevelSetCollision)
 			{
@@ -362,14 +365,14 @@ namespace Chaos
 					for (int32 ParticleIndex = 0; ParticleIndex < (int32)CollisionParticles->Size(); ++ParticleIndex)
 					{
 						const FVec3 P = ShapeTransform.TransformPosition(CollisionParticles->X(ParticleIndex));
-						FDebugDrawQueue::GetInstance().DrawDebugPoint(P, Color, false, -1.f, 0, Settings.PointSize);
+						FDebugDrawQueue::GetInstance().DrawDebugPoint(P, Color, false, Duration, 0, Settings.PointSize);
 					}
 				}
 			}
 		}
 
 		template <bool bInstanced>
-		void DrawShapesScaledImpl(const FGeometryParticleHandle* Particle, const FRigidTransform3& ShapeTransform, const FImplicitObject* Implicit, const FShapeOrShapesArray& Shapes, const FReal Margin, const FColor& Color, const FChaosDebugDrawSettings& Settings)
+		void DrawShapesScaledImpl(const FGeometryParticleHandle* Particle, const FRigidTransform3& ShapeTransform, const FImplicitObject* Implicit, const FShapeOrShapesArray& Shapes, const FReal Margin, const FColor& Color, const FRealSingle Duration, const FChaosDebugDrawSettings& Settings)
 		{
 			const EImplicitObjectType PackedType = Implicit->GetType();
 			const EImplicitObjectType InnerType = GetInnerType(PackedType);
@@ -396,7 +399,7 @@ namespace Chaos
 				const TImplicitObjectScaled<FLevelSet, bInstanced>* Scaled = Implicit->template GetObject<TImplicitObjectScaled<FLevelSet, bInstanced>>();
 				// even though thhe levelset is scaled, the debugdraw uses the collisionParticles  that are pre-scaled
 				// so no need to pass the scaled transform and just extract the wrapped LevelSet
-				DrawShapesImpl(Particle, ShapeTransform, Scaled->GetUnscaledObject(), Shapes, Scaled->GetMargin(), Color, Settings);
+				DrawShapesImpl(Particle, ShapeTransform, Scaled->GetUnscaledObject(), Shapes, Scaled->GetMargin(), Color, Duration, Settings);
 				break;
 			}
 			case ImplicitObjectType::Unknown:
@@ -405,7 +408,7 @@ namespace Chaos
 			{
 				const TImplicitObjectScaled<FConvex, bInstanced>* Scaled = Implicit->template GetObject<TImplicitObjectScaled<FConvex, bInstanced>>();
 				ScaleTM.SetScale3D(Scaled->GetScale());
-				DrawShapesImpl(Particle, ScaleTM * ShapeTransform, Scaled->GetUnscaledObject(), Shapes, Scaled->GetMargin(), Color, Settings);
+				DrawShapesImpl(Particle, ScaleTM * ShapeTransform, Scaled->GetUnscaledObject(), Shapes, Scaled->GetMargin(), Color, Duration, Settings);
 				break;
 			}
 			case ImplicitObjectType::TaperedCylinder:
@@ -416,14 +419,14 @@ namespace Chaos
 			{
 				const TImplicitObjectScaled<FTriangleMeshImplicitObject, bInstanced>* Scaled = Implicit->template GetObject<TImplicitObjectScaled<FTriangleMeshImplicitObject, bInstanced>>();
 				ScaleTM.SetScale3D(Scaled->GetScale());
-				DrawShapesImpl(Particle, ScaleTM * ShapeTransform, Scaled->GetUnscaledObject(), Shapes, 0.0f, Color, Settings);
+				DrawShapesImpl(Particle, ScaleTM * ShapeTransform, Scaled->GetUnscaledObject(), Shapes, 0.0f, Color, Duration, Settings);
 				break;
 			}
 			case ImplicitObjectType::HeightField:
 			{
 				const TImplicitObjectScaled<FHeightField, bInstanced>* Scaled = Implicit->template GetObject<TImplicitObjectScaled<FHeightField, bInstanced>>();
 				ScaleTM.SetScale3D(Scaled->GetScale());
-				DrawShapesImpl(Particle, ScaleTM * ShapeTransform, Scaled->GetUnscaledObject(), Shapes, 0.0f, Color, Settings);
+				DrawShapesImpl(Particle, ScaleTM * ShapeTransform, Scaled->GetUnscaledObject(), Shapes, 0.0f, Color, Duration, Settings);
 				break;
 			}
 			default:
@@ -431,7 +434,7 @@ namespace Chaos
 			}
 		}
 
-		void DrawShapesInstancedImpl(const FGeometryParticleHandle* Particle, const FRigidTransform3& ShapeTransform, const FImplicitObject* Implicit, const FShapeOrShapesArray& Shapes, const FReal Margin, const FColor& Color, const FChaosDebugDrawSettings& Settings)
+		void DrawShapesInstancedImpl(const FGeometryParticleHandle* Particle, const FRigidTransform3& ShapeTransform, const FImplicitObject* Implicit, const FShapeOrShapesArray& Shapes, const FReal Margin, const FColor& Color, const FRealSingle Duration, const FChaosDebugDrawSettings& Settings)
 		{
 			const EImplicitObjectType PackedType = Implicit->GetType();
 			const EImplicitObjectType InnerType = GetInnerType(PackedType);
@@ -459,7 +462,7 @@ namespace Chaos
 			case ImplicitObjectType::Convex:
 			{
 				const TImplicitObjectInstanced<FConvex>* Instanced = Implicit->template GetObject<TImplicitObjectInstanced<FConvex>>();
-				DrawShapesImpl(Particle, ShapeTransform, Instanced->GetInstancedObject(), Shapes, Instanced->GetMargin(), Color, Settings);
+				DrawShapesImpl(Particle, ShapeTransform, Instanced->GetInstancedObject(), Shapes, Instanced->GetMargin(), Color, Duration, Settings);
 				break;
 			}
 			case ImplicitObjectType::TaperedCylinder:
@@ -469,13 +472,13 @@ namespace Chaos
 			case ImplicitObjectType::TriangleMesh:
 			{
 				const TImplicitObjectInstanced<FTriangleMeshImplicitObject>* Scaled = Implicit->template GetObject<TImplicitObjectInstanced<FTriangleMeshImplicitObject>>();
-				DrawShapesImpl(Particle, ShapeTransform, Scaled->GetInstancedObject(), Shapes, 0.0f, Color, Settings);
+				DrawShapesImpl(Particle, ShapeTransform, Scaled->GetInstancedObject(), Shapes, 0.0f, Color, Duration, Settings);
 				break;
 			}
 			case ImplicitObjectType::HeightField:
 			{
 				const TImplicitObjectInstanced<FHeightField>* Scaled = Implicit->template GetObject<TImplicitObjectInstanced<FHeightField>>();
-				DrawShapesImpl(Particle, ShapeTransform, Scaled->GetInstancedObject(), Shapes, 0.0f, Color, Settings);
+				DrawShapesImpl(Particle, ShapeTransform, Scaled->GetInstancedObject(), Shapes, 0.0f, Color, Duration, Settings);
 				break;
 			}
 			default:
@@ -483,7 +486,7 @@ namespace Chaos
 			}
 		}
 
-		void DrawShapesImpl(const FGeometryParticleHandle* Particle, const FRigidTransform3& ShapeTransform, const FImplicitObject* Implicit, const FShapeOrShapesArray& Shapes, const FReal Margin, const FColor& Color, const FChaosDebugDrawSettings& Settings)
+		void DrawShapesImpl(const FGeometryParticleHandle* Particle, const FRigidTransform3& ShapeTransform, const FImplicitObject* Implicit, const FShapeOrShapesArray& Shapes, const FReal Margin, const FColor& Color, const FRealSingle Duration, const FChaosDebugDrawSettings& Settings)
 		{
 			if (!Particle || !Implicit) return;
 
@@ -501,24 +504,24 @@ namespace Chaos
 			{
 				if (IsInstanced(PackedType))
 				{
-					DrawShapesScaledImpl<true>(Particle, ShapeTransform, Implicit, Shapes, Margin, Color, Settings);
+					DrawShapesScaledImpl<true>(Particle, ShapeTransform, Implicit, Shapes, Margin, Color, Duration, Settings);
 				}
 				else
 				{
-					DrawShapesScaledImpl<false>(Particle, ShapeTransform, Implicit, Shapes, Margin, Color, Settings);
+					DrawShapesScaledImpl<false>(Particle, ShapeTransform, Implicit, Shapes, Margin, Color, Duration, Settings);
 				}
 				return;
 			}
 			else if (IsInstanced(PackedType))
 			{
-				DrawShapesInstancedImpl(Particle, ShapeTransform, Implicit, Shapes, Margin, Color, Settings);
+				DrawShapesInstancedImpl(Particle, ShapeTransform, Implicit, Shapes, Margin, Color, Duration, Settings);
 				return;
 			}
 			else if (InnerType == ImplicitObjectType::Transformed)
 			{
 				const TImplicitObjectTransformed<FReal, 3>* Transformed = Implicit->template GetObject<TImplicitObjectTransformed<FReal, 3>>();
 				FRigidTransform3 TransformedTransform = FRigidTransform3(ShapeTransform.TransformPosition(Transformed->GetTransform().GetLocation()), ShapeTransform.GetRotation() * Transformed->GetTransform().GetRotation());
-				DrawShapesImpl(Particle, TransformedTransform, Transformed->GetTransformedObject(), Shapes, Margin, Color, Settings);
+				DrawShapesImpl(Particle, TransformedTransform, Transformed->GetTransformedObject(), Shapes, Margin, Color, Duration, Settings);
 				return;
 			}
 			else if (InnerType == ImplicitObjectType::Union)
@@ -535,7 +538,7 @@ namespace Chaos
 						PerShapeData = ShapesArray[UnionIdx].Get();
 					}
 					
-					DrawShapesImpl(Particle, ShapeTransform, UnionImplicit.Get(), FShapeOrShapesArray(PerShapeData), Margin, Color, Settings);
+					DrawShapesImpl(Particle, ShapeTransform, UnionImplicit.Get(), FShapeOrShapesArray(PerShapeData), Margin, Color, Duration, Settings);
 					UnionIdx++;
 				}
 				return;
@@ -548,7 +551,7 @@ namespace Chaos
 					const TPBDRigidParticleHandle<FReal, 3>* OriginalParticle = Union->FindParticleForImplicitObject(UnionImplicit.Get());
 					if (ensure(OriginalParticle))
 					{
-						DrawShapesImpl(Particle, ShapeTransform, UnionImplicit.Get(), FShapeOrShapesArray(OriginalParticle), Margin, Color, Settings);
+						DrawShapesImpl(Particle, ShapeTransform, UnionImplicit.Get(), FShapeOrShapesArray(OriginalParticle), Margin, Color, Duration, Settings);
 					}
 				}
 				return;
@@ -601,14 +604,14 @@ namespace Chaos
 			{
 				const TSphere<FReal, 3>* Sphere = Implicit->template GetObject<TSphere<FReal, 3>>();
 				const FVec3 P = ShapeTransform.TransformPosition(Sphere->GetCenter());
-				FDebugDrawQueue::GetInstance().DrawDebugSphere(P, Sphere->GetRadius(), 8, ShapeColor, false, UE_KINDA_SMALL_NUMBER, Settings.DrawPriority, Settings.ShapeThicknesScale * Settings.LineThickness);
+				FDebugDrawQueue::GetInstance().DrawDebugSphere(P, Sphere->GetRadius(), 8, ShapeColor, false, Duration, Settings.DrawPriority, Settings.ShapeThicknesScale * Settings.LineThickness);
 				break;
 			}
 			case ImplicitObjectType::Box:
 			{
 				const TBox<FReal, 3>* Box = Implicit->template GetObject<TBox<FReal, 3>>();
 				const FVec3 P = ShapeTransform.TransformPosition(Box->GetCenter());
-				FDebugDrawQueue::GetInstance().DrawDebugBox(P, (FReal)0.5 * Box->Extents(), ShapeTransform.GetRotation(), ShapeColor, false, UE_KINDA_SMALL_NUMBER, Settings.DrawPriority, Settings.ShapeThicknesScale * Settings.LineThickness);
+				FDebugDrawQueue::GetInstance().DrawDebugBox(P, (FReal)0.5 * Box->Extents(), ShapeTransform.GetRotation(), ShapeColor, false, Duration, Settings.DrawPriority, Settings.ShapeThicknesScale * Settings.LineThickness);
 				break;
 			}
 			case ImplicitObjectType::Plane:
@@ -618,13 +621,13 @@ namespace Chaos
 				const FCapsule* Capsule = Implicit->template GetObject<FCapsule>();
 				const FVec3 P = ShapeTransform.TransformPosition(Capsule->GetCenter());
 				const FRotation3 Q = ShapeTransform.GetRotation() * FRotationMatrix::MakeFromZ(Capsule->GetAxis());
-				FDebugDrawQueue::GetInstance().DrawDebugCapsule(P, (FReal)0.5 * Capsule->GetHeight() + Capsule->GetRadius(), Capsule->GetRadius(), Q, ShapeColor, false, UE_KINDA_SMALL_NUMBER, Settings.DrawPriority, Settings.ShapeThicknesScale * Settings.LineThickness);
+				FDebugDrawQueue::GetInstance().DrawDebugCapsule(P, (FReal)0.5 * Capsule->GetHeight() + Capsule->GetRadius(), Capsule->GetRadius(), Q, ShapeColor, false, Duration, Settings.DrawPriority, Settings.ShapeThicknesScale * Settings.LineThickness);
 				break;
 			}
 			case ImplicitObjectType::LevelSet:
 			{
 				const FLevelSet* LevelSet = Implicit->template GetObject<FLevelSet>();
-				DrawShapesLevelSetImpl(Particle, ShapeTransform, LevelSet, ShapeColor, Settings);
+				DrawShapesLevelSetImpl(Particle, ShapeTransform, LevelSet, ShapeColor, Duration, Settings);
 				break;
 			}
 			break;
@@ -633,7 +636,7 @@ namespace Chaos
 				const FConvex* Convex = Implicit->template GetObject<FConvex>();
 
 				const FReal NetMargin = bChaosDebugDebugDrawCoreShapes ? Margin + Convex->GetMargin() : 0.0f;
-				DrawShapesConvexImpl(Particle, ShapeTransform, Convex, NetMargin, ShapeColor, Settings);
+				DrawShapesConvexImpl(Particle, ShapeTransform, Convex, NetMargin, ShapeColor, Duration, Settings);
 
 				// Generate the exact marging-reduced convex for comparison with the runtime approximation
 				// Warning: extremely expensive!
@@ -649,7 +652,7 @@ namespace Chaos
 					ShrunkScaledConvex.MovePlanesAndRebuild(FConvex::FRealType(-Margin)); // potential loss of precision but margin should remain within the float range
 
 					const FRigidTransform3 ShrunkScaledTransform = FRigidTransform3(ShapeTransform.GetTranslation(), ShapeTransform.GetRotation());
-					DrawShapesConvexImpl(Particle, ShrunkScaledTransform, &ShrunkScaledConvex, 0.0f, FColor::Red, Settings);
+					DrawShapesConvexImpl(Particle, ShrunkScaledTransform, &ShrunkScaledConvex, 0.0f, FColor::Red, Duration, Settings);
 				}
 				break;
 			}
@@ -660,13 +663,13 @@ namespace Chaos
 			case ImplicitObjectType::TriangleMesh:
 			{
 				const FTriangleMeshImplicitObject* TriangleMesh = Implicit->template GetObject<FTriangleMeshImplicitObject>();
-				DrawShapesTriangleMeshImpl(Particle, ShapeTransform, TriangleMesh, ShapeColor, Settings);
+				DrawShapesTriangleMeshImpl(Particle, ShapeTransform, TriangleMesh, ShapeColor, Duration, Settings);
 				break;
 			}
 			case ImplicitObjectType::HeightField:
 			{
 				const FHeightField* HeightField = Implicit->template GetObject<FHeightField>();
-				DrawShapesHeightFieldImpl(Particle, ShapeTransform, HeightField, ShapeColor, Settings);
+				DrawShapesHeightFieldImpl(Particle, ShapeTransform, HeightField, ShapeColor, Duration, Settings);
 				break;
 			}
 			default:
@@ -683,7 +686,7 @@ namespace Chaos
 						for (int32 ParticleIndex = 0; ParticleIndex < (int32)Particles->Size(); ++ParticleIndex)
 						{
 							FVec3 P = ShapeTransform.TransformPosition(Particles->X(ParticleIndex));
-							FDebugDrawQueue::GetInstance().DrawDebugPoint(P, ShapeColor, false, UE_KINDA_SMALL_NUMBER, Settings.DrawPriority, Settings.PointSize);
+							FDebugDrawQueue::GetInstance().DrawDebugPoint(P, ShapeColor, false, Duration, Settings.DrawPriority, Settings.PointSize);
 						}
 					}
 				}
@@ -694,7 +697,7 @@ namespace Chaos
 				const FColor ShapeBoundsColor = FColor::Orange;
 				const FAABB3& ShapeBounds = Implicit->BoundingBox();
 				const FVec3 ShapeBoundsPos = ShapeTransform.TransformPosition(ShapeBounds.Center());
-				FDebugDrawQueue::GetInstance().DrawDebugBox(ShapeBoundsPos, 0.5f * ShapeBounds.Extents(), ShapeTransform.GetRotation(), ShapeBoundsColor, false, UE_KINDA_SMALL_NUMBER, Settings.DrawPriority, Settings.LineThickness);
+				FDebugDrawQueue::GetInstance().DrawDebugBox(ShapeBoundsPos, 0.5f * ShapeBounds.Extents(), ShapeTransform.GetRotation(), ShapeBoundsColor, false, Duration, Settings.DrawPriority, Settings.LineThickness);
 			}
 		}
 
@@ -703,7 +706,7 @@ namespace Chaos
 			FVec3 P = SpaceTransform.TransformPosition(Particle->ObjectState() == EObjectStateType::Dynamic ? Particle->CastToRigidParticle()->P() : Particle->X());
 			FRotation3 Q = SpaceTransform.GetRotation() * (Particle->ObjectState() == EObjectStateType::Dynamic ? Particle->CastToRigidParticle()->Q() : Particle->R());
 
-			DrawShapesImpl(Particle, FRigidTransform3(P, Q), Particle->Geometry().Get(), FShapeOrShapesArray(Particle), 0.0f, InColor, Settings);
+			DrawShapesImpl(Particle, FRigidTransform3(P, Q), Particle->Geometry().Get(), FShapeOrShapesArray(Particle), 0.0f, InColor, 0.0f, Settings);
 		}
 
 		void DrawParticleShapesImpl(const FRigidTransform3& SpaceTransform, const FGeometryParticle* Particle, const FColor& InColor, const FChaosDebugDrawSettings& Settings)
@@ -711,7 +714,7 @@ namespace Chaos
 			FVec3 P = SpaceTransform.TransformPosition(Particle->X());
 			FRotation3 Q = SpaceTransform.GetRotation() * (Particle->R());
 
-			DrawShapesImpl(Particle->Handle(), FRigidTransform3(P, Q), Particle->Geometry().Get(), FShapeOrShapesArray(Particle->Handle()), 0.0f, InColor, Settings);
+			DrawShapesImpl(Particle->Handle(), FRigidTransform3(P, Q), Particle->Geometry().Get(), FShapeOrShapesArray(Particle->Handle()), 0.0f, InColor, 0.0f, Settings);
 		}
 
 		static EImplicitObjectType GetFirstConcreteShapeType(const FImplicitObject* Shape)
@@ -823,7 +826,7 @@ namespace Chaos
 			}
 		}
 
-		void DrawCollisionImpl(const FRigidTransform3& SpaceTransform, const FPBDCollisionConstraint& Contact, FRealSingle ColorScale, const FChaosDebugDrawSettings& Settings)
+		void DrawCollisionImpl(const FRigidTransform3& SpaceTransform, const FPBDCollisionConstraint& Contact, FRealSingle ColorScale, const FRealSingle Duration, const FChaosDebugDrawSettings& Settings)
 		{
 			if ((Settings.ContactWidth > 0) || (Settings.ContactLen > 0) || (Settings.ContactInfoWidth > 0) || (Settings.ImpulseScale > 0.0f))
 			{
@@ -895,43 +898,43 @@ namespace Chaos
 					// Pushout
 					if ((Settings.PushOutScale > 0) && !ManifoldPoint.NetPushOut.IsNearlyZero())
 					{
-						FColor Color = (ColorScale * PushOutColor).ToFColor(false);
-						FDebugDrawQueue::GetInstance().DrawDebugLine(WorldPointPlaneLocation, WorldPointPlaneLocation + Settings.DrawScale * Settings.PushOutScale * SpaceTransform.TransformVectorNoScale(ManifoldPoint.NetPushOut), Color, false, UE_KINDA_SMALL_NUMBER, Settings.DrawPriority, Settings.LineThickness);
+						FColor Color = (ColorScale * PushOutImpusleColor).ToFColor(false);
+						FDebugDrawQueue::GetInstance().DrawDebugLine(WorldPointPlaneLocation, WorldPointPlaneLocation + Settings.DrawScale * Settings.PushOutScale * SpaceTransform.TransformVectorNoScale(ManifoldPoint.NetPushOut), Color, false, Duration, Settings.DrawPriority, Settings.LineThickness);
 					}
 					if ((Settings.ImpulseScale > 0) && !ManifoldPoint.NetImpulse.IsNearlyZero())
 					{
 						FColor Color = (ColorScale * ImpulseColor).ToFColor(false);
-						FDebugDrawQueue::GetInstance().DrawDebugLine(WorldPointPlaneLocation, WorldPointPlaneLocation + Settings.DrawScale * Settings.ImpulseScale * SpaceTransform.TransformVectorNoScale(ManifoldPoint.NetImpulse), Color, false, UE_KINDA_SMALL_NUMBER, Settings.DrawPriority, Settings.LineThickness);
+						FDebugDrawQueue::GetInstance().DrawDebugLine(WorldPointPlaneLocation, WorldPointPlaneLocation + Settings.DrawScale * Settings.ImpulseScale * SpaceTransform.TransformVectorNoScale(ManifoldPoint.NetImpulse), Color, false, Duration, Settings.DrawPriority, Settings.LineThickness);
 					}
 
 					// Manifold plane and normal
 					if (Settings.ContactWidth > 0)
 					{
 						FColor C0 = (ColorScale * DiscColor).ToFColor(false);
-						FDebugDrawQueue::GetInstance().DrawDebugCircle(WorldPlaneLocation, Settings.DrawScale * Settings.ContactWidth, 12, C0, false, UE_KINDA_SMALL_NUMBER, Settings.DrawPriority, Settings.LineThickness, Axes.GetUnitAxis(EAxis::Y), Axes.GetUnitAxis(EAxis::Z), false);
+						FDebugDrawQueue::GetInstance().DrawDebugCircle(WorldPlaneLocation, Settings.DrawScale * Settings.ContactWidth, 12, C0, false, Duration, Settings.DrawPriority, Settings.LineThickness, Axes.GetUnitAxis(EAxis::Y), Axes.GetUnitAxis(EAxis::Z), false);
 					}
 					if (Settings.ContactLen > 0)
 					{
 						FColor NormalColor = ((ManifoldPoint.ContactPoint.ContactType != EContactPointType::EdgeEdge) ? PlaneNormalColor : EdgeNormalColor);
 						FColor C1 = (ColorScale * NormalColor).ToFColor(false);
-						FDebugDrawQueue::GetInstance().DrawDebugLine(WorldPlaneLocation, WorldPlaneLocation + Settings.DrawScale * Settings.ContactLen * WorldPlaneNormal, C1, false, UE_KINDA_SMALL_NUMBER, Settings.DrawPriority, Settings.LineThickness);
+						FDebugDrawQueue::GetInstance().DrawDebugLine(WorldPlaneLocation, WorldPlaneLocation + Settings.DrawScale * Settings.ContactLen * WorldPlaneNormal, C1, false, Duration, Settings.DrawPriority, Settings.LineThickness);
 					}
 					if (Settings.ContactPhiWidth > 0 && (ManifoldPoint.ContactPoint.Phi < FLT_MAX))
 					{
 						FColor C2 = (ColorScale * FColor(128, 128, 0)).ToFColor(false);
-						FDebugDrawQueue::GetInstance().DrawDebugCircle(WorldPlaneLocation - ManifoldPoint.ContactPoint.Phi * WorldPlaneNormal, Settings.DrawScale * Settings.ContactPhiWidth, 12, C2, false, UE_KINDA_SMALL_NUMBER, Settings.DrawPriority, Settings.LineThickness, Axes.GetUnitAxis(EAxis::Y), Axes.GetUnitAxis(EAxis::Z), false);
+						FDebugDrawQueue::GetInstance().DrawDebugCircle(WorldPlaneLocation - ManifoldPoint.ContactPoint.Phi * WorldPlaneNormal, Settings.DrawScale * Settings.ContactPhiWidth, 12, C2, false, Duration, Settings.DrawPriority, Settings.LineThickness, Axes.GetUnitAxis(EAxis::Y), Axes.GetUnitAxis(EAxis::Z), false);
 					}
 
 					// Manifold point
-					FDebugDrawQueue::GetInstance().DrawDebugCircle(WorldPointLocation, 0.5f * Settings.DrawScale * Settings.ContactWidth, 12, DiscColor, false, UE_KINDA_SMALL_NUMBER, Settings.DrawPriority, Settings.LineThickness, Axes.GetUnitAxis(EAxis::Y), Axes.GetUnitAxis(EAxis::Z), false);
+					FDebugDrawQueue::GetInstance().DrawDebugCircle(WorldPointLocation, 0.5f * Settings.DrawScale * Settings.ContactWidth, 12, DiscColor, false, Duration, Settings.DrawPriority, Settings.LineThickness, Axes.GetUnitAxis(EAxis::Y), Axes.GetUnitAxis(EAxis::Z), false);
 
 					// Previous points
 					if (bIsActive && ManifoldPoint.Flags.bWasFrictionRestored)
 					{
 						const FVec3 WorldPrevPointLocation = SpaceTransform.TransformPosition(PointTransform.TransformPosition(ManifoldPoint.ShapeAnchorPoints[ContactPointOwner]));
 						const FVec3 WorldPrevPlaneLocation = SpaceTransform.TransformPosition(PlaneTransform.TransformPosition(ManifoldPoint.ShapeAnchorPoints[ContactPlaneOwner]));
-						FDebugDrawQueue::GetInstance().DrawDebugLine(WorldPrevPointLocation, WorldPointLocation, FColor::White, false, UE_KINDA_SMALL_NUMBER, Settings.DrawPriority, Settings.LineThickness);
-						FDebugDrawQueue::GetInstance().DrawDebugLine(WorldPrevPlaneLocation, WorldPlaneLocation, FColor::White, false, UE_KINDA_SMALL_NUMBER, Settings.DrawPriority, Settings.LineThickness);
+						FDebugDrawQueue::GetInstance().DrawDebugLine(WorldPrevPointLocation, WorldPointLocation, FColor::White, false, Duration, Settings.DrawPriority, Settings.LineThickness);
+						FDebugDrawQueue::GetInstance().DrawDebugLine(WorldPrevPlaneLocation, WorldPlaneLocation, FColor::White, false, Duration, Settings.DrawPriority, Settings.LineThickness);
 					}
 
 					// Whether restored
@@ -940,17 +943,17 @@ namespace Chaos
 						if (Contact.WasManifoldRestored())
 						{
 							const FReal BoxScale = Settings.DrawScale * Settings.ContactInfoWidth;
-							FDebugDrawQueue::GetInstance().DrawDebugBox(WorldPlaneLocation, FVec3(BoxScale, BoxScale, FReal(0.01)), FRotation3(FRotationMatrix::MakeFromZ(WorldPlaneNormal)), FColor::Blue, false, UE_KINDA_SMALL_NUMBER, Settings.DrawPriority, 0.5f * Settings.LineThickness);
+							FDebugDrawQueue::GetInstance().DrawDebugBox(WorldPlaneLocation, FVec3(BoxScale, BoxScale, FReal(0.01)), FRotation3(FRotationMatrix::MakeFromZ(WorldPlaneNormal)), FColor::Blue, false, Duration, Settings.DrawPriority, 0.5f * Settings.LineThickness);
 						}
 						else if (ManifoldPoint.Flags.bWasRestored)
 						{
 							const FReal BoxScale = Settings.DrawScale * Settings.ContactInfoWidth;
-							FDebugDrawQueue::GetInstance().DrawDebugBox(WorldPlaneLocation, FVec3(BoxScale, BoxScale, FReal(0.01)), FRotation3(FRotationMatrix::MakeFromZ(WorldPlaneNormal)), FColor::Purple, false, UE_KINDA_SMALL_NUMBER, Settings.DrawPriority, 0.5f * Settings.LineThickness);
+							FDebugDrawQueue::GetInstance().DrawDebugBox(WorldPlaneLocation, FVec3(BoxScale, BoxScale, FReal(0.01)), FRotation3(FRotationMatrix::MakeFromZ(WorldPlaneNormal)), FColor::Purple, false, Duration, Settings.DrawPriority, 0.5f * Settings.LineThickness);
 						}
 						else if (ManifoldPoint.Flags.bWasReplaced)
 						{
 							const FReal BoxScale = Settings.DrawScale * Settings.ContactInfoWidth;
-							FDebugDrawQueue::GetInstance().DrawDebugBox(WorldPlaneLocation, FVec3(BoxScale, BoxScale, FReal(0.01)), FRotation3(FRotationMatrix::MakeFromZ(WorldPlaneNormal)), FColor::Orange, false, UE_KINDA_SMALL_NUMBER, Settings.DrawPriority, 0.5f * Settings.LineThickness);
+							FDebugDrawQueue::GetInstance().DrawDebugBox(WorldPlaneLocation, FVec3(BoxScale, BoxScale, FReal(0.01)), FRotation3(FRotationMatrix::MakeFromZ(WorldPlaneNormal)), FColor::Orange, false, Duration, Settings.DrawPriority, 0.5f * Settings.LineThickness);
 						}
 					}
 				}
@@ -960,7 +963,7 @@ namespace Chaos
 				{
 					FColor Color = (ColorScale * FColor::White).ToFColor(false);
 					const FVec3 ImpulsePos = SpaceTransform.TransformPosition(WorldActorTransform0.GetTranslation());
-					FDebugDrawQueue::GetInstance().DrawDebugLine(ImpulsePos, ImpulsePos + Settings.DrawScale * Settings.ImpulseScale * SpaceTransform.TransformVectorNoScale(Contact.GetAccumulatedImpulse()), Color, false, UE_KINDA_SMALL_NUMBER, Settings.DrawPriority, Settings.LineThickness);
+					FDebugDrawQueue::GetInstance().DrawDebugLine(ImpulsePos, ImpulsePos + Settings.DrawScale * Settings.ImpulseScale * SpaceTransform.TransformVectorNoScale(Contact.GetAccumulatedImpulse()), Color, false, Duration, Settings.DrawPriority, Settings.LineThickness);
 				}
 
 			}
@@ -978,9 +981,57 @@ namespace Chaos
 			}
 		}
 		
-		void DrawCollisionImpl(const FRigidTransform3& SpaceTransform, const FPBDCollisionConstraintHandle* ConstraintHandle, FRealSingle ColorScale, const FChaosDebugDrawSettings& Settings)
+		void DrawCollisionImpl(const FRigidTransform3& SpaceTransform, const FPBDCollisionConstraintHandle* ConstraintHandle, FRealSingle ColorScale, const FRealSingle Duration, const FChaosDebugDrawSettings& Settings)
+		{
+			DrawCollisionImpl(SpaceTransform, ConstraintHandle->GetContact(), ColorScale, Duration, Settings);
+		}
+
+		void DrawParticleCCDCollisionShapeImpl(const FRigidTransform3& SpaceTransform, const FCCDParticle* CCDParticle, const bool bShowStartPos,  const FColor& ShapeColor, const FRealSingle Duration, const FChaosDebugDrawSettings& Settings)
+		{
+			if (CCDParticle != nullptr)
+			{
+				FConstGenericParticleHandle P0 = CCDParticle->Particle;
+				if (P0->IsDynamic() && P0->CCDEnabled())
+				{
+					const FRigidTransform3 ActorTransform0 = bShowStartPos ? FParticleUtilitiesXR::GetActorWorldTransform(P0): FParticleUtilitiesPQ::GetActorWorldTransform(P0);
+					DrawShapesImpl(P0->Handle(), ActorTransform0 * SpaceTransform, P0->Geometry().Get(), FShapeOrShapesArray(P0->Handle()), 0.0f, ShapeColor, Duration, Settings);
+				}
+			}
+		}
+
+		void DrawParticleCCDCollisionImpulseImpl(const FRigidTransform3& SpaceTransform, const FCCDParticle* CCDParticle, const FCCDConstraint& CCDConstraint, const int32 ManifoldPointIndex, const FVec3& Impulse, const FRealSingle Duration, const FChaosDebugDrawSettings& Settings)
+		{
+			if (CCDParticle != nullptr)
+			{
+				FConstGenericParticleHandle P0 = CCDParticle->Particle;
+				if (P0->IsDynamic() && P0->CCDEnabled())
+				{
+					if (!Impulse.IsNearlyZero() && (Settings.ImpulseScale > 0))
+					{
+						const FRigidTransform3 CoMTransformX0 = FParticleUtilitiesXR::GetCoMWorldTransform(P0);
+						const FVec3 Pos0 = SpaceTransform.TransformPosition(CoMTransformX0.GetTranslation());
+						FDebugDrawQueue::GetInstance().DrawDebugLine(Pos0, Pos0 + Settings.ImpulseScale * Impulse, FColor::Red, false, Duration, Settings.DrawPriority, Settings.LineThickness * 0.5f);
+					}
+				}
+			}
+		}
+
+		void DrawCCDCollisionShapeImpl(const FRigidTransform3& SpaceTransform, const FCCDConstraint& CCDConstraint, const bool bShowStartPos, const FColor& ShapeColor, const FChaosDebugDrawSettings& Settings)
+		{
+			const FRealSingle Duration = ChaosDebugDrawCCDDuration;
+
+			DrawParticleCCDCollisionShapeImpl(SpaceTransform, CCDConstraint.Particle[0], bShowStartPos, ShapeColor, Duration, Settings);
+			DrawParticleCCDCollisionShapeImpl(SpaceTransform, CCDConstraint.Particle[1], bShowStartPos, ShapeColor, Duration, Settings);
+		}
+
+		void DrawCCDCollisionImpulseImpl(const FRigidTransform3& SpaceTransform, const FCCDConstraint& CCDConstraint, const int32 ManifoldPointIndex, const FVec3& Impulse, const FChaosDebugDrawSettings& Settings)
 		{
-			DrawCollisionImpl(SpaceTransform, ConstraintHandle->GetContact(), ColorScale, Settings);
+			const FRealSingle Duration = ChaosDebugDrawCCDDuration;
+
+			DrawCollisionImpl(SpaceTransform, CCDConstraint.SweptConstraint, 1.0f, Duration, Settings);
+
+			DrawParticleCCDCollisionImpulseImpl(SpaceTransform, CCDConstraint.Particle[0], CCDConstraint, ManifoldPointIndex, Impulse, Duration, Settings);
+			DrawParticleCCDCollisionImpulseImpl(SpaceTransform, CCDConstraint.Particle[1], CCDConstraint, ManifoldPointIndex, Impulse, Duration, Settings);
 		}
 
 		void DrawCollidingShapesImpl(const FRigidTransform3& SpaceTransform, const FPBDCollisionConstraints& Collisions, FRealSingle ColorScale, const FChaosDebugDrawSettings& Settings)
@@ -1036,6 +1087,7 @@ namespace Chaos
 					Shapes[ShapeIndex],
 					0.0f, 
 					ShapeParticles[ShapeIndex]->IsDynamic() ? FColor::Yellow : FColor::Red, 
+					0.0f,
 					Settings);
 			}
 		}
@@ -1553,7 +1605,7 @@ namespace Chaos
 					TVec2<const FGeometryParticleHandle*> ConstrainedParticles = ConstraintHandle->GetConstrainedParticles();
 					if ((ConstrainedParticles[0] == Particle) || (ConstrainedParticles[1] == Particle))
 					{
-						DrawCollisionImpl(SpaceTransform, ConstraintHandle, 1.0f, GetChaosDebugDrawSettings(Settings));
+						DrawCollisionImpl(SpaceTransform, ConstraintHandle, 1.0f, 0.0f, GetChaosDebugDrawSettings(Settings));
 					}
 				}
 			}
@@ -1565,7 +1617,7 @@ namespace Chaos
 			{
 				for (int32 ConstraintIndex = 0; ConstraintIndex < Collisions.NumConstraints(); ++ConstraintIndex)
 				{
-					DrawCollisionImpl(SpaceTransform, Collisions.GetConstraint(ConstraintIndex), ColorScale, GetChaosDebugDrawSettings(Settings));
+					DrawCollisionImpl(SpaceTransform, Collisions.GetConstraint(ConstraintIndex), ColorScale, 0.0f, GetChaosDebugDrawSettings(Settings));
 				}
 			}
 		}
@@ -1577,7 +1629,7 @@ namespace Chaos
 				CollisionAllocator.VisitConstCollisions(
 					[&](const FPBDCollisionConstraint& Collision)
 					{
-						DrawCollisionImpl(SpaceTransform, &Collision, ColorScale, GetChaosDebugDrawSettings(Settings));
+						DrawCollisionImpl(SpaceTransform, &Collision, ColorScale, 0.0f, GetChaosDebugDrawSettings(Settings));
 						return ECollisionVisitorResult::Continue;
 					});
 			}
@@ -1678,6 +1730,22 @@ namespace Chaos
 			}
 		}
 
+		void DrawCCDCollisionShape(const FRigidTransform3& SpaceTransform, const FCCDConstraint& CCDConstraint, const bool bShowStartPos, const FColor& ShapeColor, const FChaosDebugDrawSettings* Settings)
+		{
+			if (FDebugDrawQueue::IsDebugDrawingEnabled())
+			{
+				DrawCCDCollisionShapeImpl(SpaceTransform, CCDConstraint, bShowStartPos, ShapeColor, GetChaosDebugDrawSettings(Settings));
+			}
+		}
+
+		void DrawCCDCollisionImpulse(const FRigidTransform3& SpaceTransform, const FCCDConstraint& CCDConstraint, const int32 ManifoldPointIndex, const FVec3& Impulse, const FChaosDebugDrawSettings* Settings)
+		{
+			if (FDebugDrawQueue::IsDebugDrawingEnabled())
+			{
+				DrawCCDCollisionImpulseImpl(SpaceTransform, CCDConstraint, ManifoldPointIndex, Impulse, GetChaosDebugDrawSettings(Settings));
+			}
+		}
+
 		void DrawCCDAxisThreshold(const FVec3& X, const FVec3& AxisThreshold, const FVec3& DeltaX, const FQuat& R)
 		{
 			// Call the version of DeltaExceedsThreshold that provides access to its
diff --git a/Engine/Source/Runtime/Experimental/Chaos/Private/Chaos/PBDRigidsEvolutionGBF.cpp b/Engine/Source/Runtime/Experimental/Chaos/Private/Chaos/PBDRigidsEvolutionGBF.cpp
--- a/Engine/Source/Runtime/Experimental/Chaos/Private/Chaos/PBDRigidsEvolutionGBF.cpp
+++ b/Engine/Source/Runtime/Experimental/Chaos/Private/Chaos/PBDRigidsEvolutionGBF.cpp
@@ -64,9 +64,6 @@ namespace Chaos
 		bool bChaosCollisionCCDUseTightBoundingBox = true; 
 		FAutoConsoleVariableRef  CVarChaosCollisionCCDUseTightBoundingBox(TEXT("p.Chaos.Collision.CCD.UseTightBoundingBox"), bChaosCollisionCCDUseTightBoundingBox , TEXT(""));
 
-		bool bChaosCollisionCCDDebugDraw = false; 
-		FAutoConsoleVariableRef  CVarChaosCollisionCCDDebugDraw(TEXT("p.Chaos.Collision.CCD.DebugDraw"), bChaosCollisionCCDDebugDraw , TEXT(""));
-
 		int32 ChaosSolverCollisionPriority = 0;
 		FAutoConsoleVariableRef CVarChaosSolverCollisionPriority(TEXT("p.Chaos.Solver.Collision.Priority"), ChaosSolverCollisionPriority, TEXT("Set constraint priority. Larger values are evaluated later [def:0]"));
 
diff --git a/Engine/Source/Runtime/Experimental/Chaos/Private/PBDRigidsSolver.cpp b/Engine/Source/Runtime/Experimental/Chaos/Private/PBDRigidsSolver.cpp
--- a/Engine/Source/Runtime/Experimental/Chaos/Private/PBDRigidsSolver.cpp
+++ b/Engine/Source/Runtime/Experimental/Chaos/Private/PBDRigidsSolver.cpp
@@ -75,6 +75,8 @@ namespace Chaos
 		int32 ChaosSolverDebugDrawBounds = 0;
 		int32 ChaosSolverDrawTransforms = 0;
 		int32 ChaosSolverDrawIslands = 0;
+		int32 ChaosSolverDrawCCDInteractions = 0;
+		int32 ChaosSolverDrawCCDThresholds = 0;
 		int32 ChaosSolverDrawShapesShowStatic = 1;
 		int32 ChaosSolverDrawShapesShowKinematic = 1;
 		int32 ChaosSolverDrawShapesShowDynamic = 1;
@@ -91,6 +93,8 @@ namespace Chaos
 		FAutoConsoleVariableRef CVarChaosSolverDrawBounds(TEXT("p.Chaos.Solver.DebugDrawBounds"), ChaosSolverDebugDrawBounds, TEXT("Draw bounding volumes inside the broadphase (0 = never; 1 = end of frame)."));
 		FAutoConsoleVariableRef CVarChaosSolverDrawTransforms(TEXT("p.Chaos.Solver.DebugDrawTransforms"), ChaosSolverDrawTransforms, TEXT("Draw particle transforms (0 = never; 1 = end of frame)."));
 		FAutoConsoleVariableRef CVarChaosSolverDrawIslands(TEXT("p.Chaos.Solver.DebugDrawIslands"), ChaosSolverDrawIslands, TEXT("Draw solver islands (0 = never; 1 = end of frame)."));
+		FAutoConsoleVariableRef CVarChaosSolverDrawCCD(TEXT("p.Chaos.Solver.DebugDrawCCDInteractions"), ChaosSolverDrawCCDInteractions, TEXT("Draw CCD interactions."));
+		FAutoConsoleVariableRef CVarChaosSolverDrawCCDThresholds(TEXT("p.Chaos.Solver.DebugDrawCCDThresholds"), ChaosSolverDrawCCDThresholds, TEXT("Draw CCD swept thresholds."));
 		FAutoConsoleVariableRef CVarChaosSolverDrawShapesShapesStatic(TEXT("p.Chaos.Solver.DebugDraw.ShowStatics"), ChaosSolverDrawShapesShowStatic, TEXT("If DebugDrawShapes is enabled, whether to show static objects"));
 		FAutoConsoleVariableRef CVarChaosSolverDrawShapesShapesKinematic(TEXT("p.Chaos.Solver.DebugDraw.ShowKinematics"), ChaosSolverDrawShapesShowKinematic, TEXT("If DebugDrawShapes is enabled, whether to show kinematic objects"));
 		FAutoConsoleVariableRef CVarChaosSolverDrawShapesShapesDynamic(TEXT("p.Chaos.Solver.DebugDraw.ShowDynamics"), ChaosSolverDrawShapesShowDynamic, TEXT("If DebugDrawShapes is enabled, whether to show dynamic objects"));
diff --git a/Engine/Source/Runtime/Experimental/Chaos/Public/Chaos/ChaosDebugDraw.h b/Engine/Source/Runtime/Experimental/Chaos/Public/Chaos/ChaosDebugDraw.h
--- a/Engine/Source/Runtime/Experimental/Chaos/Public/Chaos/ChaosDebugDraw.h
+++ b/Engine/Source/Runtime/Experimental/Chaos/Public/Chaos/ChaosDebugDraw.h
@@ -10,6 +10,7 @@
 namespace Chaos
 {
 	class FAccelerationStructureHandle;
+	struct FCCDConstraint;
 	class FCollisionConstraintAllocator;
 	template <typename PayloadType, typename T, int d> class ISpatialAcceleration;
 	class FPBDSuspensionConstraints;
@@ -218,6 +219,8 @@ namespace Chaos
 		CHAOS_API void DrawCollidingShapes(const FRigidTransform3& SpaceTransform, const FPBDCollisionConstraints& Collisions, FRealSingle ColorScale, const FChaosDebugDrawSettings* Settings = nullptr);
 		CHAOS_API void DrawSpatialAccelerationStructure(const ISpatialAcceleration<FAccelerationStructureHandle, FReal, 3>& SpatialAccelerationStructure, const FChaosDebugDrawSettings* Settings = nullptr);
 		CHAOS_API void DrawSuspensionConstraints(const FRigidTransform3& SpaceTransform, const FPBDSuspensionConstraints& Constraints, const FChaosDebugDrawSettings* Settings = nullptr);
+		CHAOS_API void DrawCCDCollisionShape(const FRigidTransform3& SpaceTransform, const FCCDConstraint& CCDConstraint, const bool bShowStartPos, const FColor& ShapeColor, const FChaosDebugDrawSettings* Settings = nullptr);
+		CHAOS_API void DrawCCDCollisionImpulse(const FRigidTransform3& SpaceTransform, const FCCDConstraint& CCDConstraint, const int32 ManifoldPointIndex, const FVec3& Impulse, const FChaosDebugDrawSettings* Settings = nullptr);
 		CHAOS_API void DrawCCDAxisThreshold(const FVec3& X, const FVec3& AxisThreshold, const FVec3& DeltaX, const FQuat& R);
 #endif
 	}
diff --git a/Engine/Source/Runtime/Experimental/Chaos/Public/Chaos/CollisionResolution.h b/Engine/Source/Runtime/Experimental/Chaos/Public/Chaos/CollisionResolution.h
--- a/Engine/Source/Runtime/Experimental/Chaos/Public/Chaos/CollisionResolution.h
+++ b/Engine/Source/Runtime/Experimental/Chaos/Public/Chaos/CollisionResolution.h
@@ -45,7 +45,7 @@ namespace Chaos
 
 		/**
 		 * @brief Update the contact manifold on the constraint
-		 * @note Transforms are shape world-space transforms (not particle transforms)
+		 * @note Transforms are shape world-space transforms (not particle transforms) at the start of the sweep. The end of the sweep are the transforms stored in the constraint as ShapeWorldTransform0/1
 		*/
 		void CHAOS_API UpdateConstraintSwept(FPBDCollisionConstraint& Constraint, const FRigidTransform3& ShapeWorldTransform0, const FRigidTransform3& ShapeWorldTransform1, const FReal Dt);
 
@@ -61,11 +61,9 @@ namespace Chaos
 		bool CHAOS_API ShouldUseCCD(const FGeometryParticleHandle* Particle0, const FVec3& DeltaX0, const FGeometryParticleHandle* Particle1, const FVec3& DeltaX1, FVec3& Dir, FReal& Length, const bool bForceDisableCCD);
 
 		// Update the constraint by re-running collision detection on the shape pair.
+		// @todo(chaos): remove this and use UpdateConstraint instead
 		template<ECollisionUpdateType UpdateType>
 		void CHAOS_API UpdateConstraintFromGeometry(FPBDCollisionConstraint& Constraint, const FRigidTransform3& ParticleTransform0, const FRigidTransform3& ParticleTransform1, const FReal Dt);
-
-		// Update the constraint by re-running collision detection on the shape pair.
-		// Return whether CCD is needed for this constraint.
 		template<ECollisionUpdateType UpdateType>
 		bool CHAOS_API UpdateConstraintFromGeometrySwept(FPBDCollisionConstraint& Constraint, const FRigidTransform3& ParticleTransform0, const FRigidTransform3& ParticleTransform1, const FReal Dt);
 
diff --git a/Engine/Source/Runtime/Experimental/Chaos/Public/Chaos/PBDRigidsEvolutionGBF.h b/Engine/Source/Runtime/Experimental/Chaos/Public/Chaos/PBDRigidsEvolutionGBF.h
--- a/Engine/Source/Runtime/Experimental/Chaos/Public/Chaos/PBDRigidsEvolutionGBF.h
+++ b/Engine/Source/Runtime/Experimental/Chaos/Public/Chaos/PBDRigidsEvolutionGBF.h
@@ -34,7 +34,7 @@ namespace Chaos
 		CHAOS_API extern FRealSingle HackMaxVelocity;
 		CHAOS_API extern FRealSingle SmoothedPositionLerpRate;
 		CHAOS_API extern bool bChaosCollisionCCDUseTightBoundingBox;
-		CHAOS_API extern bool bChaosCollisionCCDDebugDraw;
+		CHAOS_API extern int32 ChaosSolverDrawCCDThresholds;
 	}
 
 	using FPBDRigidsEvolutionCallback = TFunction<void()>;
@@ -302,7 +302,7 @@ namespace Chaos
 					{
 
 #if CHAOS_DEBUG_DRAW
-						if (CVars::bChaosCollisionCCDDebugDraw)
+						if (CVars::ChaosSolverDrawCCDThresholds)
 						{
 							DebugDraw::DrawCCDAxisThreshold(Particle.X(), Particle.CCDAxisThreshold(), Particle.P() - Particle.X(), Particle.Q());
 						}
