FIX CIS - Merge garbage

diff --git a/Engine/Source/Runtime/Experimental/Chaos/Private/Chaos/ChaosDebugDraw.cpp b/Engine/Source/Runtime/Experimental/Chaos/Private/Chaos/ChaosDebugDraw.cpp
--- a/Engine/Source/Runtime/Experimental/Chaos/Private/Chaos/ChaosDebugDraw.cpp
+++ b/Engine/Source/Runtime/Experimental/Chaos/Private/Chaos/ChaosDebugDraw.cpp
@@ -872,7 +872,8 @@ namespace Chaos
 					FColor PlaneNormalColor = FColor(200, 0, 0);
 					FColor EdgeNormalColor = FColor(200, 150, 0);
 					FColor ImpulseColor = FColor(0, 0, 200);
-					FColor PushOutColor = FColor(0, 200, 200);
+					FColor PushOutColor = FColor(0, 200, 0);
+					FColor PushOutImpusleColor = FColor(0, 200, 200);
 					if (ManifoldPoint.Flags.bInsideStaticFrictionCone)
 					{
 						DiscColor = FColor(150, 200, 0);
