Fix CIS - add missing #if CHAOS_DEBUG_DRAW

diff --git a/Engine/Source/Runtime/Experimental/Chaos/Private/Chaos/CCDUtilities.cpp b/Engine/Source/Runtime/Experimental/Chaos/Private/Chaos/CCDUtilities.cpp
--- a/Engine/Source/Runtime/Experimental/Chaos/Private/Chaos/CCDUtilities.cpp
+++ b/Engine/Source/Runtime/Experimental/Chaos/Private/Chaos/CCDUtilities.cpp
@@ -33,7 +33,10 @@ namespace Chaos
 		FAutoConsoleVariableRef  CVarCCDEnableThresholdBoundsScale(TEXT("p.Chaos.CCD.EnableThresholdBoundsScale"), CCDEnableThresholdBoundsScale , TEXT("CCD is used when object position is changing > smallest bound's extent * BoundsScale. 0 will always Use CCD. Values < 0 disables CCD."));
 
 		extern int32 ChaosSolverDrawCCDInteractions;
+
+#if CHAOS_DEBUG_DRAW
 		extern DebugDraw::FChaosDebugDrawSettings ChaosSolverDebugDebugDrawSettings;	
+#endif
 	}
 
 	void FCCDParticle::AddOverlappingDynamicParticle(FCCDParticle* const InParticle)
