Edigrate 19255298 and 19255737

19255298: Store local inertia tensor/inverse tensor diagonals in vec3 instead of matrix, reducing memory.
19255737: Fix FVector conversion errors after FVector constructor was made explicit.

diff --git a/Engine/Source/Runtime/Core/Public/UObject/FortniteMainBranchObjectVersion.h b/Engine/Source/Runtime/Core/Public/UObject/FortniteMainBranchObjectVersion.h
--- a/Engine/Source/Runtime/Core/Public/UObject/FortniteMainBranchObjectVersion.h
+++ b/Engine/Source/Runtime/Core/Public/UObject/FortniteMainBranchObjectVersion.h
@@ -202,6 +202,9 @@ struct CORE_API FFortniteMainBranchObjectVersion
 		// Migrated function handlers to the CDO/archetype data
 		MigratedFunctionHandlersToDefaults,
 
+		// Storing inertia tensor as vec3 instead of matrix.
+		ChaosInertiaConvertedToVec3,
+
 		// -----<new versions can be added above this line>-------------------------------------------------
 		VersionPlusOne,
 		LatestVersion = VersionPlusOne - 1
diff --git a/Engine/Source/Runtime/Experimental/Chaos/Public/Chaos/RigidParticles.h b/Engine/Source/Runtime/Experimental/Chaos/Public/Chaos/RigidParticles.h
--- a/Engine/Source/Runtime/Experimental/Chaos/Public/Chaos/RigidParticles.h
+++ b/Engine/Source/Runtime/Experimental/Chaos/Public/Chaos/RigidParticles.h
@@ -12,6 +12,7 @@
 #include "Chaos/Particles.h"
 #include "Chaos/Rotation.h"
 #include "HAL/LowLevelMemTracker.h"
+#include "UObject/FortniteMainBranchObjectVersion.h"
 
 namespace Chaos
 {
@@ -313,8 +314,10 @@ public:
 
 		Ar << MAcceleration << MAngularAcceleration << MLinearImpulseVelocity << MAngularImpulseVelocity;
 
+		Ar.UsingCustomVersion(FFortniteMainBranchObjectVersion::GUID);
 		Ar.UsingCustomVersion(FUE5ReleaseStreamObjectVersion::GUID);
-		if (Ar.IsLoading() && Ar.CustomVer(FUE5ReleaseStreamObjectVersion::GUID) < FUE5ReleaseStreamObjectVersion::ChaosInertiaConvertedToVec3)
+		if (Ar.IsLoading() && Ar.CustomVer(FFortniteMainBranchObjectVersion::GUID) < FFortniteMainBranchObjectVersion::ChaosInertiaConvertedToVec3
+			&& Ar.CustomVer(FUE5ReleaseStreamObjectVersion::GUID) < FUE5ReleaseStreamObjectVersion::ChaosInertiaConvertedToVec3)
 		{
 			TArray<PMatrix<T, d, d>> IArray;
 			TArray<PMatrix<T, d, d>> InvIArray;
diff --git a/Engine/Source/Runtime/Experimental/ChaosCore/Public/Chaos/Core.h b/Engine/Source/Runtime/Experimental/ChaosCore/Public/Chaos/Core.h
--- a/Engine/Source/Runtime/Experimental/ChaosCore/Public/Chaos/Core.h
+++ b/Engine/Source/Runtime/Experimental/ChaosCore/Public/Chaos/Core.h
@@ -42,6 +42,6 @@ namespace Chaos
 	// NOTE: if you get a merge conflict on the GUID, you must replace it with a new GUID - do not accept the source or target
 	// or you will likely get DDC version conflicts resulting in crashes during load.
 	// Core version string for Chaos data. Any DDC builder dependent on Chaos for serialization should depend on this version
-	static const FString ChaosVersionString = TEXT("D89F8E66-8E2B-4C04-83B8-FB8761F99CF4");
+	static const FString ChaosVersionString = TEXT("F3C4A6F4-BF50-488B-9FE2-7968054E9C6B");
 
 }
