Chaos
- Fix CCD enable/disable decision for shallow contacts when TOI=0
- Dedupe TOI calculation for consistency on primary and secondary CCD interations
- Fix CCD with multiple CCD iterations (if bounds expansion is turned on in cvars)
- Propagate CCD bounds expansion to narrow phase cull distance (when enabled)
- Fix CCD shape transforms for secondary sweeps
- CCD debug draw

diff --git a/Engine/Source/Runtime/Experimental/Chaos/Private/Chaos/CCDUtilities.cpp b/Engine/Source/Runtime/Experimental/Chaos/Private/Chaos/CCDUtilities.cpp
--- a/Engine/Source/Runtime/Experimental/Chaos/Private/Chaos/CCDUtilities.cpp
+++ b/Engine/Source/Runtime/Experimental/Chaos/Private/Chaos/CCDUtilities.cpp
@@ -8,27 +8,26 @@
 
 //PRAGMA_DISABLE_OPTIMIZATION
 
-// @todo(chaos): these cvars should all be settings per particle
-
-// NOTE: With this disabled secondary CCD collisions will often be missed
-bool bChaosCollisionCCDEnableResweep = true;
-FAutoConsoleVariableRef CVarChaosCollisionCCDEnableResweep(TEXT("p.Chaos.Collision.CCD.EnableResweep"), bChaosCollisionCCDEnableResweep, TEXT("Enable resweep for CCD. Resweeping allows CCD to catch more secondary collisions but also is more costly. Default is true."));
-
-// NOTE: With this disabled, secondary collisions can be missed. When enabled, velocity will not be visually consistent after CCD collisions (if ChaosCollisionCCDConstraintMaxProcessCount is too low)
-bool bChaosCollisionCCDAllowClipping = true;
-FAutoConsoleVariableRef CVarChaosCollisionCCDAllowClipping(TEXT("p.Chaos.Collision.CCD.AllowClipping"), bChaosCollisionCCDAllowClipping, TEXT("This will clip the CCD object at colliding positions when computation budgets run out. Default is true. Turning this option off might cause tunneling."));
-
-// This needs to be at least 2 to handle sliding at CCD speeds properly
-int32 ChaosCollisionCCDConstraintMaxProcessCount = 2;
-FAutoConsoleVariableRef CVarChaosCollisionCCDConstraintMaxProcessCount(TEXT("p.Chaos.Collision.CCD.ConstraintMaxProcessCount"), ChaosCollisionCCDConstraintMaxProcessCount, TEXT("The max number of times each constraint can be resolved when applying CCD constraints. Default is 2. The larger this number is, the more fully CCD constraints are resolved."));
-
-bool bChaosCollisionCCDEnablePhiCheck = false;
-FAutoConsoleVariableRef CVarChaosCollisionCCDEnablePhiCheck(TEXT("p.Chaos.Collision.CCD.EnablePhiCheck"), bChaosCollisionCCDEnablePhiCheck, TEXT("Enable phi check for CCD. If disabled we skip CCD processing of contacts if the penetration depth is below the CCD threshold at TOI=1."));
-
 namespace Chaos
 {
 	namespace CVars
 	{
+		// @todo(chaos): most of these cvars should be settings per particle
+
+		// NOTE: With this disabled secondary CCD collisions will often be missed
+		// @todo(chaos): resweeping also change contacts so it raises questions about collision modifier callbacks and CCD
+		bool bChaosCollisionCCDEnableResweep = true;
+		FAutoConsoleVariableRef CVarChaosCollisionCCDEnableResweep(TEXT("p.Chaos.Collision.CCD.EnableResweep"), bChaosCollisionCCDEnableResweep, TEXT("Enable resweep for CCD. Resweeping allows CCD to catch more secondary collisions but also is more costly. Default is true."));
+
+		// NOTE: With this disabled, secondary collisions can be missed. When enabled, velocity will not be visually consistent after CCD collisions (if ChaosCollisionCCDConstraintMaxProcessCount is too low)
+		bool bChaosCollisionCCDAllowClipping = true;
+		FAutoConsoleVariableRef CVarChaosCollisionCCDAllowClipping(TEXT("p.Chaos.Collision.CCD.AllowClipping"), bChaosCollisionCCDAllowClipping, TEXT("This will clip the CCD object at colliding positions when computation budgets run out. Default is true. Turning this option off might cause tunneling."));
+
+		// By default, we stop processing CCD contacts after a single CCD interaction
+		// This will result in a visual velocity glitch when it happens, but usually this doesn't matter since the impact is very high energy anyway
+		int32 ChaosCollisionCCDConstraintMaxProcessCount = 1;
+		FAutoConsoleVariableRef CVarChaosCollisionCCDConstraintMaxProcessCount(TEXT("p.Chaos.Collision.CCD.ConstraintMaxProcessCount"), ChaosCollisionCCDConstraintMaxProcessCount, TEXT("The max number of times each constraint can be resolved when applying CCD constraints. Default is 2. The larger this number is, the more fully CCD constraints are resolved."));
+
 		FRealSingle CCDEnableThresholdBoundsScale = 0.4f;
 		FAutoConsoleVariableRef  CVarCCDEnableThresholdBoundsScale(TEXT("p.Chaos.CCD.EnableThresholdBoundsScale"), CCDEnableThresholdBoundsScale , TEXT("CCD is used when object position is changing > smallest bound's extent * BoundsScale. 0 will always Use CCD. Values < 0 disables CCD."));
 
@@ -124,6 +123,9 @@ namespace Chaos
 		bool bNeedCCDSolve = false;
 		for (FPBDCollisionConstraint* Constraint : SweptConstraints)
 		{   
+			// A contact can be disabled by a user callback or contact pruning so we need to ignore these.
+			// NOTE: It important that we explicitly check for disabled here, rather than for 0 manifold points since we
+			// may want to use the contact later if resweeping is enabled. 
 			if (!Constraint->IsEnabled())
 			{
 				continue;
@@ -288,7 +290,7 @@ namespace Chaos
 
 	bool CCDConstraintSortPredicate(const FCCDConstraint* Constraint0, const FCCDConstraint* Constraint1)
 	{
-		return Constraint0->SweptConstraint->TimeOfImpact < Constraint1->SweptConstraint->TimeOfImpact;
+		return Constraint0->SweptConstraint->CCDTimeOfImpact < Constraint1->SweptConstraint->CCDTimeOfImpact;
 	}
 
 	void FCCDManager::ApplyIslandSweptConstraints(const int32 Island, const FReal Dt)
@@ -322,7 +324,7 @@ namespace Chaos
 			FCCDParticle* CCDParticle0 = CCDConstraint->Particle[0];
 			FCCDParticle* CCDParticle1 = CCDConstraint->Particle[1];
 
-			IslandTOI = CCDConstraint->SweptConstraint->TimeOfImpact;
+			IslandTOI = CCDConstraint->SweptConstraint->CCDTimeOfImpact;
 
 			// Constraints whose TOIs are in the range of [0, 1) are resolved for this frame. TOI = 1 means that the two 
 			// particles just start touching at the end of the frame and therefore cannot have tunneling this frame. 
@@ -333,28 +335,16 @@ namespace Chaos
 			}
 
 			// If both particles are marked Done (due to clipping), continue
-			if (bChaosCollisionCCDAllowClipping && (!CCDParticle0 || CCDParticle0->Done) && (!CCDParticle1 || CCDParticle1->Done))
-			{
-				ConstraintIndex++;
-				continue;
-			}
-			// Skip CCD resolution if the depth at TOI=1 is low and let the standard contact resolution handle it.
-			// This can be important because CCD contact resolution uses point-mass physics and does not apply rotation 
-			// so we must leave some work for the regular solver to do to get decent behaviour when sliding fast along
-			// a surface for example.
-			// @todo(chaos): enabling this can lead to tunnelling problems when you have a jumble of CCD objects.
-			// Possibly because if we need to re-examin this CCD contact because of an impulse from another CCD contact,
-			// we may be initially penetrating. Investigate and aee if we can come up with a more robust solution. 
-			if (bChaosCollisionCCDEnablePhiCheck && (CCDConstraint->SweptConstraint->GetPhi() > CCDConstraint->PhiThreshold))
+			if (CVars::bChaosCollisionCCDAllowClipping && (!CCDParticle0 || CCDParticle0->Done) && (!CCDParticle1 || CCDParticle1->Done))
 			{
 				ConstraintIndex++;
 				continue;
 			}
 			
-			ensure(CCDConstraint->ProcessedCount < ChaosCollisionCCDConstraintMaxProcessCount);
+			ensure(CCDConstraint->ProcessedCount < CVars::ChaosCollisionCCDConstraintMaxProcessCount);
 
-			// In UpdateConstraintFromGeometrySwept, InitManifoldPoint requires P, Q to be at TOI=1., but the input of 
-			// UpdateConstraintFromGeometrySwept requires transforms at current TOI. So instead of rewinding P, Q, we 
+			// In UpdateConstraintSwept, InitManifoldPoint requires P, Q to be at TOI=1., but the input of 
+			// UpdateConstraintSwept requires transforms at current TOI. So instead of rewinding P, Q, we 
 			// advance X, R to current TOI and keep P, Q at TOI=1.
 			if (CCDParticle0 && !CCDParticle0->Done)
 			{
@@ -365,15 +355,25 @@ namespace Chaos
 				AdvanceParticleXToTOI(CCDParticle1, IslandTOI, Dt);
 			}
 
+#if CHAOS_DEBUG_DRAW
+			// Debugdraw the shape at the current TOI
+			if (CVars::ChaosSolverDrawCCDInteractions)
+			{
+				DebugDraw::DrawCCDCollisionShape(FRigidTransform3::Identity, *CCDConstraint, true, FColor::Magenta, &CVars::ChaosSolverDebugDebugDrawSettings);
+			}
+#endif
+
 			ApplyImpulse(CCDConstraint);
 			CCDConstraint->ProcessedCount++;
 
 			// After applying impulse, constraint TOI need be updated to reflect the new velocities. 
 			// Usually the new velocities are separating, and therefore TOI should be infinity.
 			// See resweep below which (optionally) updates TOI for all other contacts as a result of handling this one
-			CCDConstraint->SweptConstraint->TimeOfImpact = TNumericLimits<FReal>::Max();
+			CCDConstraint->SweptConstraint->CCDTimeOfImpact = TNumericLimits<FReal>::Max();
 
-			if (CCDConstraint->ProcessedCount >= ChaosCollisionCCDConstraintMaxProcessCount)
+			bool bMovedParticle0 = false;
+			bool bMovedParticle1 = false;
+			if (CCDConstraint->ProcessedCount >= CVars::ChaosCollisionCCDConstraintMaxProcessCount)
 			{
 				/* Here is how clipping works:
 				* Assuming collision detection gives us all the possible collision pairs in the current frame.
@@ -384,7 +384,7 @@ namespace Chaos
 				* For kinematic particles, we cannot freeze them in place. In this case, we simply offset the particle with the kinematic motion from [IslandTOI, 1] along the collision normal and freeze it there.
 				* If collision detection is not perfect and does not give us all the secondary collision pairs, setting ChaosCollisionCCDConstraintMaxProcessCount to 1 will always prevent tunneling.
 				*/ 
-				if (bChaosCollisionCCDAllowClipping)
+				if (CVars::bChaosCollisionCCDAllowClipping)
 				{
 					if (CCDParticle0)
 					{
@@ -401,6 +401,7 @@ namespace Chaos
 							ClipParticleP(CCDParticle0);
 						}
 						CCDParticle0->Done = true;
+						bMovedParticle0 = true;
 					}
 					if (CCDParticle1)
 					{
@@ -417,6 +418,7 @@ namespace Chaos
 							ClipParticleP(CCDParticle1);
 						}
 						CCDParticle1->Done = true;
+						bMovedParticle1 = true;
 					}
 				}
 				// If clipping is not allowed, we update particle P (at TOI=1) based on new velocities. 
@@ -425,10 +427,12 @@ namespace Chaos
 					if (CCDParticle0)
 					{
 						UpdateParticleP(CCDParticle0, Dt);
+						bMovedParticle0 = true;
 					}
 					if (CCDParticle1)
 					{
 						UpdateParticleP(CCDParticle1, Dt);
+						bMovedParticle1 = true;
 					}
 				}
 				// Increment ConstraintIndex if we run out of computational budget for this constraint.
@@ -440,115 +444,39 @@ namespace Chaos
 				if (CCDParticle0 && !CCDParticle0->Done)
 				{
 					UpdateParticleP(CCDParticle0, Dt);
+					bMovedParticle0 = true;
 				}
 				if (CCDParticle1 && !CCDParticle1->Done)
 				{
 					UpdateParticleP(CCDParticle1, Dt);
+					bMovedParticle1 = true;
 				}
 			}
 
-			if (bChaosCollisionCCDEnableResweep)
+			// We applied a CCD impulse and updated the particle positions, so we need to update all the constraints involving these particles
+			bool bHasResweptConstraint = false;
+			if (bMovedParticle0)
 			{
-				// For each constraint that contains the two particles on which we applied impulses, we need to update its TOI.
-				const FReal RestDt = (1.f - IslandTOI) * Dt;
-				bool HasResweptConstraint = false;
-				for (int32 i = 0; i < 2; i++)
-				{
-					FCCDParticle *CCDParticle = CCDConstraint->Particle[i];
-					if (CCDParticle != nullptr)
-					{
-						for (int32 AttachedCCDConstraintIndex = 0; AttachedCCDConstraintIndex < CCDParticle->AttachedCCDConstraints.Num(); AttachedCCDConstraintIndex++)
-						{
-							FCCDConstraint *AttachedCCDConstraint = CCDParticle->AttachedCCDConstraints[AttachedCCDConstraintIndex];
-							if (AttachedCCDConstraint == CCDConstraint || AttachedCCDConstraint->ProcessedCount >= ChaosCollisionCCDConstraintMaxProcessCount)
-							{
-								continue;
-							}
-							FRigidTransform3 RigidTransforms[2]; // Rigid transforms at TOI
-							for(int32 j = 0; j < 2; j++)
-							{
-								FCCDParticle *AffectedCCDParticle = AttachedCCDConstraint->Particle[j];
-								if (AffectedCCDParticle != nullptr)
-								{
-									TPBDRigidParticleHandle<FReal, 3>* AffectedParticle = AffectedCCDParticle->Particle;
-									if (!AffectedCCDParticle->Done)
-									{
-										AdvanceParticleXToTOI(AffectedCCDParticle, IslandTOI, Dt);
-									}
-									RigidTransforms[j] = FRigidTransform3(AffectedParticle->X(), AffectedParticle->R());
-								}
-								else
-								{
-									FGenericParticleHandle AffectedParticle = FGenericParticleHandle(AttachedCCDConstraint->SweptConstraint->GetParticle(j));
-									const bool IsKinematic = AffectedParticle->ObjectState() == EObjectStateType::Kinematic;
-									if (IsKinematic)
-									{
-										RigidTransforms[j] = FRigidTransform3(AffectedParticle->P() - AffectedParticle->V() * ((1.f - IslandTOI) * Dt), AffectedParticle->Q());
-									}
-									else // Static case
-									{
-										RigidTransforms[j] = FRigidTransform3(AffectedParticle->X(), AffectedParticle->R());
-									}
-								}
-							}
-
-							/** When resweeping, we need to recompute TOI for affected constraints and therefore the work (GJKRaycast) used to compute the original TOI is wasted.
-							* A potential optimization is to compute an estimate of TOI using the AABB of the particles. Sweeping AABBs to compute an estimated TOI can be very efficient, and this TOI is strictly smaller than the accurate TOI.
-							* At each for-loop iteration, we only need the constraint with the smallest TOI in the island. A potential optimized algorithm could be like:
-							* 	First, sort constraints based on estimated TOI.
-							*	Find the constraint with the smallest accurate TOI:
-							*		Walk through the constraint list, change estimated TOI to accurate TOI
-							*		If accurate TOI is smaller than estimated TOI of the next constraint, we know we found the constraint.
-							*	When resweeping, compute estimated TOI instead of accurate TOI since updated TOI might need to be updated again.
-							*/
-
-							FPBDCollisionConstraint* SweptConstraint = AttachedCCDConstraint->SweptConstraint;
-							const FConstGenericParticleHandle Particle0 = SweptConstraint->GetParticle0();
-							const FConstGenericParticleHandle Particle1 = SweptConstraint->GetParticle1();
-
-							// Initial shape sweep transforms
-							const FRigidTransform3 ShapeStartWorldTransform0 = SweptConstraint->GetShapeRelativeTransform0() * RigidTransforms[0];
-							const FRigidTransform3 ShapeStartWorldTransform1 = SweptConstraint->GetShapeRelativeTransform1() * RigidTransforms[1];
-
-							// End shape sweep transforms
-							const FRigidTransform3 ParticleEndWorldTransform0 = FParticleUtilities::GetActorWorldTransform(Particle0);
-							const FRigidTransform3 ParticleEndWorldTransform1 = FParticleUtilities::GetActorWorldTransform(Particle1);
-							const FRigidTransform3 ShapeEndWorldTransform0 = SweptConstraint->GetShapeRelativeTransform0() * ParticleEndWorldTransform0;
-							const FRigidTransform3 ShapeEndWorldTransform1 = SweptConstraint->GetShapeRelativeTransform1() * ParticleEndWorldTransform1;
-
-							// Update of swept constraint assumes that the constraint holds the end transforms for the sweep
-							SweptConstraint->SetShapeWorldTransforms(ShapeEndWorldTransform0, ShapeEndWorldTransform1);
-
-							if (bChaosCollisionCCDEnableResweep)
-							{
-								// Resweep the shape. This is the expensive option
-								Collisions::UpdateConstraintSwept(*SweptConstraint, ShapeStartWorldTransform0, ShapeStartWorldTransform1, RestDt);
-							}
-
-							const FReal RestDtTOI = AttachedCCDConstraint->SweptConstraint->TimeOfImpact;
-							if ((RestDtTOI >= 0) && (RestDtTOI < FReal(1)))
-							{
-								AttachedCCDConstraint->SweptConstraint->TimeOfImpact = IslandTOI + (FReal(1) - IslandTOI) * RestDtTOI;
-							}
-
-							// When bUpdated==true, TOI was modified. When bUpdated==false, TOI was set to be TNumericLimits<FReal>::Max(). In either case, a re-sorting on the constraints is needed.
-							HasResweptConstraint = true;
-						}
-					}
-				}
-				if (HasResweptConstraint)
-				{
-					std::sort(SortedCCDConstraints.GetData() + ConstraintIndex, SortedCCDConstraints.GetData() + ConstraintStart + ConstraintNum, CCDConstraintSortPredicate);
-				}
+				bHasResweptConstraint |= UpdateParticleSweptConstraints(CCDParticle0, IslandTOI, Dt);
+			}
+			if (bMovedParticle1)
+			{
+				bHasResweptConstraint |= UpdateParticleSweptConstraints(CCDParticle1, IslandTOI, Dt);
 			}
-		}
 
-		// Debugdraw the shapes at the final position
+			// If we updated some constraints, we need to sort so that we handle the next TOI event
+			if (bHasResweptConstraint)
+			{
+				std::sort(SortedCCDConstraints.GetData() + ConstraintIndex, SortedCCDConstraints.GetData() + ConstraintStart + ConstraintNum, CCDConstraintSortPredicate);
+			}
+		}
+		
+		// Update the constraint with the CCD results
 		for (FCCDConstraint* CCDConstraint : SortedCCDConstraints)
-		{
+		{ 
 			CCDConstraint->SweptConstraint->SetCCDResults(CCDConstraint->NetImpulse);
 		}
-		
+
 #if CHAOS_DEBUG_DRAW
 		// Debugdraw the shapes at the final position
 		if (CVars::ChaosSolverDrawCCDInteractions)
@@ -561,6 +489,100 @@ namespace Chaos
 #endif
 	}
 
+	bool FCCDManager::UpdateParticleSweptConstraints(FCCDParticle* CCDParticle, const FReal IslandTOI, const FReal Dt)
+	{
+		const FReal RestDt = (1.f - IslandTOI) * Dt;
+		bool HasResweptConstraint = false;
+		if (CCDParticle != nullptr)
+		{
+			for (int32 AttachedCCDConstraintIndex = 0; AttachedCCDConstraintIndex < CCDParticle->AttachedCCDConstraints.Num(); AttachedCCDConstraintIndex++)
+			{
+				FCCDConstraint* AttachedCCDConstraint = CCDParticle->AttachedCCDConstraints[AttachedCCDConstraintIndex];
+				if (AttachedCCDConstraint->ProcessedCount >= CVars::ChaosCollisionCCDConstraintMaxProcessCount)
+				{
+					continue;
+				}
+
+				// Particle transforms at TOI
+				FRigidTransform3 ParticleStartWorldTransforms[2];
+				for (int32 j = 0; j < 2; j++)
+				{
+					FCCDParticle* AffectedCCDParticle = AttachedCCDConstraint->Particle[j];
+					if (AffectedCCDParticle != nullptr)
+					{
+						TPBDRigidParticleHandle<FReal, 3>* AffectedParticle = AffectedCCDParticle->Particle;
+						if (!AffectedCCDParticle->Done)
+						{
+							AdvanceParticleXToTOI(AffectedCCDParticle, IslandTOI, Dt);
+						}
+						ParticleStartWorldTransforms[j] = FRigidTransform3(AffectedParticle->X(), AffectedParticle->R());
+					}
+					else
+					{
+						FGenericParticleHandle AffectedParticle = FGenericParticleHandle(AttachedCCDConstraint->SweptConstraint->GetParticle(j));
+						const bool IsKinematic = AffectedParticle->ObjectState() == EObjectStateType::Kinematic;
+						if (IsKinematic)
+						{
+							ParticleStartWorldTransforms[j] = FRigidTransform3(AffectedParticle->P() - AffectedParticle->V() * RestDt, AffectedParticle->Q());
+						}
+						else // Static case
+						{
+							ParticleStartWorldTransforms[j] = FRigidTransform3(AffectedParticle->X(), AffectedParticle->R());
+						}
+					}
+				}
+				/** When resweeping, we need to recompute TOI for affected constraints and therefore the work (GJKRaycast) used to compute the original TOI is wasted.
+				* A potential optimization is to compute an estimate of TOI using the AABB of the particles. Sweeping AABBs to compute an estimated TOI can be very efficient, and this TOI is strictly smaller than the accurate TOI.
+				* At each for-loop iteration, we only need the constraint with the smallest TOI in the island. A potential optimized algorithm could be like:
+				* 	First, sort constraints based on estimated TOI.
+				*	Find the constraint with the smallest accurate TOI:
+				*		Walk through the constraint list, change estimated TOI to accurate TOI
+				*		If accurate TOI is smaller than estimated TOI of the next constraint, we know we found the constraint.
+				*	When resweeping, compute estimated TOI instead of accurate TOI since updated TOI might need to be updated again.
+				*/
+
+				FPBDCollisionConstraint* SweptConstraint = AttachedCCDConstraint->SweptConstraint;
+				const FConstGenericParticleHandle Particle0 = SweptConstraint->GetParticle0();
+				const FConstGenericParticleHandle Particle1 = SweptConstraint->GetParticle1();
+
+				// Initial shape sweep transforms
+				const FRigidTransform3 ShapeStartWorldTransform0 = SweptConstraint->GetShapeRelativeTransform0() * ParticleStartWorldTransforms[0];
+				const FRigidTransform3 ShapeStartWorldTransform1 = SweptConstraint->GetShapeRelativeTransform1() * ParticleStartWorldTransforms[1];
+
+				// End shape sweep transforms
+				const FRigidTransform3 ParticleEndWorldTransform0 = FParticleUtilities::GetActorWorldTransform(Particle0);
+				const FRigidTransform3 ParticleEndWorldTransform1 = FParticleUtilities::GetActorWorldTransform(Particle1);
+				const FRigidTransform3 ShapeEndWorldTransform0 = SweptConstraint->GetShapeRelativeTransform0() * ParticleEndWorldTransform0;
+				const FRigidTransform3 ShapeEndWorldTransform1 = SweptConstraint->GetShapeRelativeTransform1() * ParticleEndWorldTransform1;
+
+				// Update of swept constraint assumes that the constraint holds the end transforms for the sweep
+				SweptConstraint->SetShapeWorldTransforms(ShapeEndWorldTransform0, ShapeEndWorldTransform1);
+
+				if (CVars::bChaosCollisionCCDEnableResweep)
+				{
+					// Resweep the shape. This is the expensive option
+					Collisions::UpdateConstraintSwept(*SweptConstraint, ShapeStartWorldTransform0, ShapeStartWorldTransform1, RestDt);
+				}
+				else
+				{
+					// Keep the contact as-is but update the depth and TOI based on current transforms
+					SweptConstraint->UpdateSweptManifoldPoints(ShapeStartWorldTransform0.GetTranslation(), ShapeStartWorldTransform1.GetTranslation(), Dt);
+				}
+
+				const FReal RestDtTOI = AttachedCCDConstraint->SweptConstraint->CCDTimeOfImpact;
+				if ((RestDtTOI >= 0) && (RestDtTOI < FReal(1)))
+				{
+					AttachedCCDConstraint->SweptConstraint->CCDTimeOfImpact = IslandTOI + (FReal(1) - IslandTOI) * RestDtTOI;
+				}
+
+				// When bUpdated==true, TOI was modified. When bUpdated==false, TOI was set to be TNumericLimits<FReal>::Max(). In either case, a re-sorting on the constraints is needed.
+				HasResweptConstraint = true;
+			}
+		}
+
+		return HasResweptConstraint;
+	}
+
 	void FCCDManager::ResetIslandParticles(const int32 Island)
 	{
 		const int32 ParticleStart = IslandParticleStart[Island];
@@ -621,8 +643,9 @@ namespace Chaos
 		check(Rigid0 != nullptr || Rigid1 != nullptr);
 		const FReal Restitution = Constraint->GetRestitution();
 		const FRigidTransform3& ShapeWorldTransform1 = Constraint->GetShapeWorldTransform1();
-		for(const FManifoldPoint &ManifoldPoint : Constraint->GetManifoldPoints())
-			{
+		for (int32 ManifoldPointIndex = 0; ManifoldPointIndex < Constraint->NumManifoldPoints(); ++ManifoldPointIndex)
+		{
+			const FManifoldPoint& ManifoldPoint = Constraint->GetManifoldPoint(ManifoldPointIndex);
 			if (ManifoldPoint.Flags.bDisabled)
 			{
 				continue;
@@ -636,8 +659,8 @@ namespace Chaos
 			{
 				const FReal TargetNormalV = -Restitution * NormalV;
 				// If a particle is marked done, we treat it as static by setting InvM to 0. 
-				const bool bInfMass0 = Rigid0 == nullptr || (bChaosCollisionCCDAllowClipping && CCDConstraint->Particle[0] && CCDConstraint->Particle[0]->Done);
-				const bool bInfMass1 = Rigid1 == nullptr || (bChaosCollisionCCDAllowClipping && CCDConstraint->Particle[1] && CCDConstraint->Particle[1]->Done);
+				const bool bInfMass0 = Rigid0 == nullptr || (CVars::bChaosCollisionCCDAllowClipping && CCDConstraint->Particle[0] && CCDConstraint->Particle[0]->Done);
+				const bool bInfMass1 = Rigid1 == nullptr || (CVars::bChaosCollisionCCDAllowClipping && CCDConstraint->Particle[1] && CCDConstraint->Particle[1]->Done);
 				const FReal InvM0 = bInfMass0 ? 0.f : Rigid0->InvM();
 				const FReal InvM1 = bInfMass1 ? 0.f : Rigid1->InvM();
 				const FVec3 Impulse = (TargetNormalV - NormalV) * Normal / (InvM0 + InvM1);
@@ -649,7 +672,15 @@ namespace Chaos
 				{
 					Rigid1->V() -= Impulse * InvM1;
 				}
+
 				CCDConstraint->NetImpulse += Impulse;
+
+#if CHAOS_DEBUG_DRAW
+				if (CVars::ChaosSolverDrawCCDInteractions)
+				{
+					DebugDraw::DrawCCDCollisionImpulse(FRigidTransform3::Identity, *CCDConstraint, ManifoldPointIndex, Impulse, &CVars::ChaosSolverDebugDebugDrawSettings);
+				}
+#endif
 			}
 		}
 	}
diff --git a/Engine/Source/Runtime/Experimental/Chaos/Private/Chaos/ChaosDebugDraw.cpp b/Engine/Source/Runtime/Experimental/Chaos/Private/Chaos/ChaosDebugDraw.cpp
--- a/Engine/Source/Runtime/Experimental/Chaos/Private/Chaos/ChaosDebugDraw.cpp
+++ b/Engine/Source/Runtime/Experimental/Chaos/Private/Chaos/ChaosDebugDraw.cpp
@@ -2,6 +2,7 @@
 #include "Chaos/ChaosDebugDraw.h"
 #include "Chaos/Box.h"
 #include "Chaos/Capsule.h"
+#include "Chaos/CCDUtilities.h"
 #include "Chaos/Collision/CollisionConstraintAllocator.h"
 #include "Chaos/Collision/ParticlePairMidPhase.h"
 #include "Chaos/Convex.h"
@@ -29,13 +30,12 @@
 
 //PRAGMA_DISABLE_OPTIMIZATION
 
+#if CHAOS_DEBUG_DRAW
+
 namespace Chaos
 {
-
 	namespace DebugDraw
 	{
-#if CHAOS_DEBUG_DRAW
-
 		bool bChaosDebugDebugDrawShapeBounds = false;
 		FAutoConsoleVariableRef CVarChaosDebugDrawShapeBounds(TEXT("p.Chaos.DebugDraw.ShowShapeBounds"), bChaosDebugDebugDrawShapeBounds, TEXT("Whether to show the bounds of each shape in DrawShapes"));
 
@@ -79,9 +79,12 @@ namespace Chaos
 		float ChaosDebugDrawConvexExplodeDistance = 0.0f;
 		FAutoConsoleVariableRef CVarChaosDebugDrawConvexExplodeDistance(TEXT("p.Chaos.DebugDraw.ConvexExplodeDistance"), ChaosDebugDrawConvexExplodeDistance, TEXT("Explode convex edges by this amount (useful for looking at convex integrity)"));
 
-		float ChaosDebugDrawCCDDuration = 0.0f;
+		float ChaosDebugDrawCCDDuration = 2.0f;
 		FAutoConsoleVariableRef CVarChaosDebugDrawCCDDuration(TEXT("p.Chaos.DebugDraw.CCDDuration"), ChaosDebugDrawCCDDuration, TEXT("How long CCD debug draw should remain on screen"));
 
+		float ChaosDebugDrawCollisionDuration = 0.0f;
+		FAutoConsoleVariableRef CVarChaosDebugDrawCollisionDuration(TEXT("p.Chaos.DebugDraw.CollisionDuration"), ChaosDebugDrawCollisionDuration, TEXT("How long Collision debug draw should remain on screen"));
+
 
 		// NOTE: These settings should never really be used - they are the fallback defaults
 		// if the user does not specify settings in the debug draw call.
@@ -1618,7 +1621,7 @@ namespace Chaos
 			{
 				for (int32 ConstraintIndex = 0; ConstraintIndex < Collisions.NumConstraints(); ++ConstraintIndex)
 				{
-					DrawCollisionImpl(SpaceTransform, Collisions.GetConstraint(ConstraintIndex), ColorScale, 0.0f, GetChaosDebugDrawSettings(Settings));
+					DrawCollisionImpl(SpaceTransform, Collisions.GetConstraint(ConstraintIndex), ColorScale, ChaosDebugDrawCollisionDuration, GetChaosDebugDrawSettings(Settings));
 				}
 			}
 		}
@@ -1630,7 +1633,7 @@ namespace Chaos
 				CollisionAllocator.VisitConstCollisions(
 					[&](const FPBDCollisionConstraint& Collision)
 					{
-						DrawCollisionImpl(SpaceTransform, &Collision, ColorScale, 0.0f, GetChaosDebugDrawSettings(Settings));
+						DrawCollisionImpl(SpaceTransform, &Collision, ColorScale, ChaosDebugDrawCollisionDuration, GetChaosDebugDrawSettings(Settings));
 						return ECollisionVisitorResult::Continue;
 					});
 			}
@@ -1812,6 +1815,8 @@ namespace Chaos
 				}
 			}
 		}
+
+	}	// namespace DebugDraw
+}	// namespace Chaos
+
 #endif
-	}
-}
diff --git a/Engine/Source/Runtime/Experimental/Chaos/Private/Chaos/Collision/CollisionPruning.cpp b/Engine/Source/Runtime/Experimental/Chaos/Private/Chaos/Collision/CollisionPruning.cpp
--- a/Engine/Source/Runtime/Experimental/Chaos/Private/Chaos/Collision/CollisionPruning.cpp
+++ b/Engine/Source/Runtime/Experimental/Chaos/Private/Chaos/Collision/CollisionPruning.cpp
@@ -242,12 +242,6 @@ namespace Chaos
 		{
 			const FContactGroup& ContactGroup = ContactGroupContainer.ContactGroups[ContactGroupIndex];
 
-			if (ContactGroup.PlaneNormal.Z < FReal(0.5))
-			{
-				static int32 sHere = 0;
-				++sHere;
-			}
-
 			// @todo(chaos): assuming here that the contacts are not all in a line...do we need to check for that?
 			if (!ContactGroup.bActive)
 			{
@@ -280,7 +274,7 @@ namespace Chaos
 				{
 					for (FContactGroupManifoldPoint& RejectCandidateManifoldPoint : RejectCandidateContactGroup.ManifoldPoints)
 					{
-						RejectCandidateManifoldPoint.Collision->GetManifoldPoint(RejectCandidateManifoldPoint.ManifoldPointIndex).Flags.bDisabled = true;
+						RejectCandidateManifoldPoint.Collision->DisableManifoldPoint(RejectCandidateManifoldPoint.ManifoldPointIndex);
 					}
 					RejectCandidateContactGroup.bActive = false;
 				}
diff --git a/Engine/Source/Runtime/Experimental/Chaos/Private/Chaos/Collision/PBDCollisionConstraint.cpp b/Engine/Source/Runtime/Experimental/Chaos/Private/Chaos/Collision/PBDCollisionConstraint.cpp
--- a/Engine/Source/Runtime/Experimental/Chaos/Private/Chaos/Collision/PBDCollisionConstraint.cpp
+++ b/Engine/Source/Runtime/Experimental/Chaos/Private/Chaos/Collision/PBDCollisionConstraint.cpp
@@ -16,6 +16,8 @@
 
 //PRAGMA_DISABLE_OPTIMIZATION
 
+extern Chaos::FRealSingle CCDAllowedDepthBoundsScale;
+
 namespace Chaos
 {
 	extern bool bChaos_Collision_EnableManifoldGJKInject;
@@ -179,7 +181,8 @@ namespace Chaos
 		, Material()
 		, Stiffness(1)
 		, AccumulatedImpulse(0)
-		, TimeOfImpact(0)
+		, CCDTimeOfImpact(0)
+		, CCDPenetrationThreshold(0)
 		, ContainerCookie()
 		, ShapesType(EContactShapesType::Unknown)
 		, CCDType(ECollisionCCDType::Disabled)
@@ -218,7 +221,8 @@ namespace Chaos
 		, Material()
 		, Stiffness(1)
 		, AccumulatedImpulse(0)
-		, TimeOfImpact(0)
+		, CCDTimeOfImpact(0)
+		, CCDPenetrationThreshold(0)
 		, ContainerCookie()
 		, ShapesType(EContactShapesType::Unknown)
 		, CCDType(ECollisionCCDType::Disabled)
@@ -345,6 +349,14 @@ namespace Chaos
 		Flags.bIsQuadratic1 = bIsQuadratic1;
 	}
 
+	void FPBDCollisionConstraint::InitCCDThreshold()
+	{
+		// Calculate the max penetration that we ignore with CCD contacts
+		const FReal Threshold0 = FConstGenericParticleHandle(Particle[0])->CCDEnabled() ? Implicit[0]->BoundingBox().Extents().GetAbsMin() : FReal(0);
+		const FReal Threshold1 = FConstGenericParticleHandle(Particle[1])->CCDEnabled() ? Implicit[1]->BoundingBox().Extents().GetAbsMin() : FReal(0);
+		CCDPenetrationThreshold = FMath::Max(Threshold0, Threshold1) * CCDAllowedDepthBoundsScale;
+	}
+
 	void FPBDCollisionConstraint::SetIsSleeping(const bool bInIsSleeping)
 	{
 		// This actually sets the sleeping state on all constraints between the same particle pair so calling this with multiple
@@ -1025,4 +1037,91 @@ namespace Chaos
 			return NoRestingDependency;
 		}
 	}
+
+	// NOTE: This only works well if we do not update the rotation during CCD interations (which we currently do not)
+	// and we start off with a good selection for the contact plane. The latter is true if we have large objects but not so true
+	// when colliding against non-smooth triangle meshes
+	void FPBDCollisionConstraint::UpdateSweptManifoldPoints(const FVec3& ShapeStartWorldPosition0, const FVec3& ShapeStartWorldPosition1, const FReal Dt)
+	{
+		FReal MinTOI = TNumericLimits<FReal>::Max();
+		int32 MinTOIManifoldPointIndex = INDEX_NONE;
+
+		for (int32 ManifoldPointIndex = 0; ManifoldPointIndex < ManifoldPoints.Num(); ++ManifoldPointIndex)
+		{
+			FManifoldPoint& ManifoldPoint = ManifoldPoints[ManifoldPointIndex];
+			if (ManifoldPoint.Flags.bDisabled)
+			{
+				continue;
+			}
+
+			// Start and End contact positions in world space
+			const FVec3 ContactOffset0 = ShapeWorldTransform0.TransformVectorNoScale(ManifoldPoint.ContactPoint.ShapeContactPoints[0]);
+			const FVec3 ContactStartPos0 = ShapeStartWorldPosition0 + ContactOffset0;
+			const FVec3 ContactEndPos0 = ShapeWorldTransform0.GetTranslation() + ContactOffset0;
+			const FVec3 ContactOffset1 = ShapeWorldTransform1.TransformVectorNoScale(ManifoldPoint.ContactPoint.ShapeContactPoints[1]);
+			const FVec3 ContactStartPos1 = ShapeStartWorldPosition1 + ContactOffset1;
+			const FVec3 ContactEndPos1 = ShapeWorldTransform1.GetTranslation() + ContactOffset1;
+
+			// Contact normal in world space
+			const FVec3 ContactNormal = ShapeWorldTransform1.TransformVectorNoScale(ManifoldPoint.ContactPoint.ShapeContactNormal);
+
+			// Start and end separation
+			const FReal ContactStartPhi = FVec3::DotProduct(ContactStartPos0 - ContactStartPos1, ContactNormal);
+			const FReal ContactEndPhi = FVec3::DotProduct(ContactEndPos0 - ContactEndPos1, ContactNormal);
+
+			// Update the TOI and track closest manifold point
+			const FReal ContactTOI = CalculateModifiedSweptTOI(ContactStartPhi, ContactEndPhi);
+			if (ContactTOI < MinTOI)
+			{
+				MinTOI = ContactTOI;
+				MinTOIManifoldPointIndex = ManifoldPointIndex;
+			}
+
+			// Update the contact separation
+			ManifoldPoint.ContactPoint.Phi = ContactEndPhi;
+		}
+
+		CCDTimeOfImpact = MinTOI;
+		ClosestManifoldPointIndex = MinTOIManifoldPointIndex;
+	}
+
+	// Calculate a Time of Impact (TOI) to use with CCD. This is based on the initial and final contact separation from the sweep test but modified so that
+	// - we ignore separating contacts (increasing Phi)
+	// - we ignore contacts that are separated at TOI=1 (EndPhi > 0)
+	// - we ignore contacts if the penetration is less than the CCD tolerance
+	FReal FPBDCollisionConstraint::CalculateModifiedSweptTOI(const FReal StartPhi, const FReal EndPhi)
+	{
+		const FReal InfiniteTOI = TNumericLimits<FReal>::Max();
+		const FReal MovementTolerance = KINDA_SMALL_NUMBER;
+
+		// If we end up separated at TOI=1 ignore the contact
+		if (EndPhi > 0)
+		{
+			return InfiniteTOI;
+		}
+
+		// If contact is moving in the right direction or not moving ignore the contact
+		if (EndPhi > StartPhi - MovementTolerance)
+		{
+			return InfiniteTOI;
+		}
+
+		// If we penetrate by less than the CCD tolerance, treat it as TOI=1. This mean no CCD impulse and the non-CCD 
+		// solve is expected to handle it. E.g., this improves the behaviour when we are sliding along a surface at
+		// above CCD speeds - we don't want to handle TOI events with the floor
+		const FReal PhiThreshold = -CCDPenetrationThreshold;
+		if (EndPhi > PhiThreshold)
+		{
+			return FReal(1);
+		}
+
+		// If we penetrate by more than the CCD threshold we roll back all the way to the TOI leaving no penetration
+		// It would be nice to leave to penetration for the non-CCD solve to handle in this case, but then we will have
+		// initiial-overlap problems when we have secondary CCD collisions that can result in missed collisions
+		const FReal DesiredPhi = FReal(0);
+		const FReal TOI = (DesiredPhi - StartPhi) / (EndPhi - StartPhi);
+
+		return FMath::Clamp(TOI, FReal(0), FReal(1));
+	}
+
 }
\ No newline at end of file
diff --git a/Engine/Source/Runtime/Experimental/Chaos/Private/Chaos/Collision/ParticlePairMidPhase.cpp b/Engine/Source/Runtime/Experimental/Chaos/Private/Chaos/Collision/ParticlePairMidPhase.cpp
--- a/Engine/Source/Runtime/Experimental/Chaos/Private/Chaos/Collision/ParticlePairMidPhase.cpp
+++ b/Engine/Source/Runtime/Experimental/Chaos/Private/Chaos/Collision/ParticlePairMidPhase.cpp
@@ -358,6 +358,7 @@ namespace Chaos
 			Constraint->SetShapeWorldTransforms(ShapeWorldTransform0, ShapeWorldTransform1);
 
 			Constraint->SetCCDEnabled(true);
+			Constraint->SetCullDistance(CullDistance);
 			Constraint->ResetManifold();
 			Constraint->ResetActiveManifoldContacts();
 
@@ -865,10 +866,15 @@ namespace Chaos
 		}
 
 		// CullDistance is scaled by the size of the dynamic objects.
-		const FReal CullDistance = InCullDistance * CullDistanceScale;
+		FReal CullDistance = InCullDistance * CullDistanceScale;
 
 		// Enable CCD?
 		const bool bUseCCD = Flags.bIsCCD && ShouldEnableCCD(Dt);
+		if (bUseCCD)
+		{
+			const FReal VMax = (FConstGenericParticleHandle(GetParticle0())->V() - FConstGenericParticleHandle(GetParticle1())->V()).GetAbsMax();
+			CullDistance += VMax * Dt;
+		}
 		
 		// Run collision detection on all potentially colliding shape pairs
 		NumActiveConstraints = 0;
diff --git a/Engine/Source/Runtime/Experimental/Chaos/Private/Chaos/CollisionResolution.cpp b/Engine/Source/Runtime/Experimental/Chaos/Private/Chaos/CollisionResolution.cpp
--- a/Engine/Source/Runtime/Experimental/Chaos/Private/Chaos/CollisionResolution.cpp
+++ b/Engine/Source/Runtime/Experimental/Chaos/Private/Chaos/CollisionResolution.cpp
@@ -115,8 +115,6 @@ bool bChaos_Collision_ShapesArrayMode = true;
 FAutoConsoleVariableRef CVarChaos_Collision_ShapesArrayMode(TEXT("p.Chaos.Collision.ShapesArrayMode"), bChaos_Collision_ShapesArrayMode, TEXT(""));
 
 
-extern bool bChaosCollisionCCDEnableResweep;
-
 namespace Chaos
 {
 	namespace Collisions
@@ -199,56 +197,33 @@ namespace Chaos
 		}
 
 		// Initializes TimeOfImpact and adjusts Phi for TOI.
+		// @todo(chaos): this repeats normal calculation and dot product - move to ComputeSweptContactPhiAndTOIHelper?
 		void SetSweptConstraintTOI(TGeometryParticleHandle<FReal, 3>* Particle, const FReal TOI, const FReal Length, const FVec3& Dir, FPBDCollisionConstraint& SweptConstraint)
 		{
 			CONDITIONAL_SCOPE_CYCLE_COUNTER(STAT_Collisions_SetSweptConstraintTOI, ConstraintsDetailedStats);
 
-			if (SweptConstraint.GetPhi() > 0.0f)
+			if (TOI <= FReal(1))
 			{
-				SweptConstraint.TimeOfImpact = TNumericLimits<FReal>::Max();
-				return;
-			}
+				const FVec3 ContactNormal = SweptConstraint.CalculateWorldContactNormal();
+				const FReal DirDotNormal = FVec3::DotProduct(Dir, ContactNormal);
 
-			const FReal MinBounds = Particle->Geometry()->BoundingBox().Extents().Min();
-			const FReal Depth = -SweptConstraint.GetPhi();
-			const FReal MaxDepth = MinBounds * CCDAllowedDepthBoundsScale;
-
-			if (TOI <= 0.0f)
-			{
-				// Initial overlap. Phi is correct already.
-				TPBDRigidParticleHandle<FReal, 3> *Rigid0 = SweptConstraint.GetParticle0()->CastToRigidParticle();
-				TPBDRigidParticleHandle<FReal, 3> *Rigid1 = SweptConstraint.GetParticle1()->CastToRigidParticle();
-				const FRigidTransform3& ShapeWorldTransform1 = SweptConstraint.GetShapeWorldTransform1();
-				const FVec3 Normal = ShapeWorldTransform1.TransformVectorNoScale(SweptConstraint.GetClosestManifoldPoint()->ContactPoint.ShapeContactNormal);
-				const FVec3 V0 = Rigid0 != nullptr ? Rigid0->V() : FVec3(0.f);
-				const FVec3 V1 = Rigid1 != nullptr ? Rigid1->V() : FVec3(0.f);
-				const FReal NormalV = FVec3::DotProduct(V0 - V1, Normal);
-				// The two objects are already separating. This constraint needs not to be considered for CCD if the normal velocity is positive.
-				// If TOI < 0 (initial penetration case), CCD will not depenetrate the objects since CCD is only a velocity solve. Depenetration will be left for the main solve.
-				if (NormalV >= 0) 
+				// Modify TOI so we ignore separating and shallow CCD contacts
+				// NOTE (See GJKRaycast2): If initially penetrating, Phi is the separation at the start of the sweep.
+				// But if not initially penetrating Phi is the separation at the end of the sweep
+				// @todo(chaos): The way Phi is set for sweeps is quite confusing - can we make this better?
+				FReal StartPhi, EndPhi;
+				if (TOI <= FReal(0))
 				{
-					SweptConstraint.TimeOfImpact = std::numeric_limits<FReal>::max();
+					StartPhi = SweptConstraint.GetPhi();
+					EndPhi = StartPhi + DirDotNormal * Length;
 				}
 				else
 				{
-					SweptConstraint.TimeOfImpact = 0.0f;
+					EndPhi = SweptConstraint.GetPhi();
+					StartPhi = EndPhi - DirDotNormal * Length;
 				}
-			}
-			else if (TOI > 0.0f && TOI <= 1.0f)
-			{
-				if (Depth <= MaxDepth)
-				{
-					// Although we may have an earlier TOI, depth at TOI = 1 is within tolerance, so just use TOI = 1.
-					SweptConstraint.TimeOfImpact = 1.0f;
-				}
-				else if (Depth > MaxDepth)
-				{
-					// Compute ExtraT that will take us from TOI and 0 phi to slightly past TOI w/ phi at maximum allowed depth.
-					const FReal DirDotNormal = FMath::Abs(FVec3::DotProduct(Dir, SweptConstraint.CalculateWorldContactNormal()));
-					const FReal ExtraT = MaxDepth / (DirDotNormal * Length);
 
-					SweptConstraint.TimeOfImpact = TOI + ExtraT;
-				}
+				SweptConstraint.CCDTimeOfImpact = SweptConstraint.CalculateModifiedSweptTOI(StartPhi, EndPhi);
 			}
 		}
 
@@ -1842,7 +1817,7 @@ namespace Chaos
 				}
 			}
 
-			Constraint.TimeOfImpact = TNumericLimits<FReal>::Max(); // CCD will not be used for this constraint
+			Constraint.CCDTimeOfImpact = TNumericLimits<FReal>::Max(); // CCD will not be used for this constraint
 			// Do a normal non-swept update if we reach this point - Not required for now, since this will be done in solver
 			// UpdateConstraintFromGeometryImpl<UpdateType>(*(Constraint.As<FPBDCollisionConstraint>()), WorldTransform0, WorldTransform1, Dt);
 			return false;
@@ -2164,9 +2139,9 @@ namespace Chaos
 			UpdateConstraintFromGeometryImpl<ECollisionUpdateType::Deepest>(Constraint, ShapeWorldTransform0, ShapeWorldTransform1, Dt);
 		}
 
-		void UpdateConstraintSwept(FPBDCollisionConstraint& Constraint, const FRigidTransform3& ShapeWorldTransform0, const FRigidTransform3& ShapeWorldTransform1, const FReal Dt)
+		bool UpdateConstraintSwept(FPBDCollisionConstraint& Constraint, const FRigidTransform3& ShapeStartWorldTransform0, const FRigidTransform3& ShapeStartWorldTransform1, const FReal Dt)
 		{
-			UpdateConstraintFromGeometrySweptImpl<ECollisionUpdateType::Deepest>(Constraint, ShapeWorldTransform0, ShapeWorldTransform1, Dt);
+			return UpdateConstraintFromGeometrySweptImpl<ECollisionUpdateType::Deepest>(Constraint, ShapeStartWorldTransform0, ShapeStartWorldTransform1, Dt);
 		}
 
 		// Run collision detection for the specified constraint to update the nearest contact point.
@@ -2181,16 +2156,6 @@ namespace Chaos
 			UpdateConstraintFromGeometryImpl<UpdateType>(Constraint, WorldTransform0, WorldTransform1, Dt);
 		}
 
-		template<ECollisionUpdateType UpdateType>
-		bool UpdateConstraintFromGeometrySwept(FPBDCollisionConstraint& Constraint, const FRigidTransform3& ParticleWorldTransform0, const FRigidTransform3& ParticleWorldTransform1, const FReal Dt)
-		{
-			const FRigidTransform3 WorldTransform0 = Constraint.GetShapeRelativeTransform0() * ParticleWorldTransform0;
-			const FRigidTransform3 WorldTransform1 = Constraint.GetShapeRelativeTransform1() * ParticleWorldTransform1;
-			Constraint.SetShapeWorldTransforms(WorldTransform0, WorldTransform1);
-
-			return UpdateConstraintFromGeometrySweptImpl<UpdateType>(Constraint, WorldTransform0, WorldTransform1, Dt);
-		}
-
 		template<typename T_TRAITS>
 		void LegacyConstructConstraints(TGeometryParticleHandle<FReal, 3>* Particle0, TGeometryParticleHandle<FReal, 3>* Particle1, const FImplicitObject* Implicit0, const FShapeOrShapesArray& Shapes0, const FBVHParticles* Simplicial0, const FImplicitObject* Implicit1, const FShapeOrShapesArray& Shapes1, const FBVHParticles* Simplicial1, const FRigidTransform3& ParticleWorldTransform0, const FRigidTransform3& LocalTransform0, const FRigidTransform3& ParticleWorldTransform1, const FRigidTransform3& LocalTransform1, const FReal CullDistance, FReal Dt, const FCollisionContext& Context)
 		{
@@ -2619,9 +2584,6 @@ namespace Chaos
 		template void UpdateConstraintFromGeometry<ECollisionUpdateType::Any>(FPBDCollisionConstraint& ConstraintBase, const FRigidTransform3& Transform0, const FRigidTransform3& Transform1, const FReal Dt);
 		template void UpdateConstraintFromGeometry<ECollisionUpdateType::Deepest>(FPBDCollisionConstraint& ConstraintBase, const FRigidTransform3& Transform0, const FRigidTransform3& Transform1, const FReal Dt);
 
-		template bool UpdateConstraintFromGeometrySwept<ECollisionUpdateType::Any>(FPBDCollisionConstraint& ConstraintBase, const FRigidTransform3& Transform0, const FRigidTransform3& Transform1, const FReal Dt);
-		template bool UpdateConstraintFromGeometrySwept<ECollisionUpdateType::Deepest>(FPBDCollisionConstraint& ConstraintBase, const FRigidTransform3& Transform0, const FRigidTransform3& Transform1, const FReal Dt);
-
 	} // Collisions
 
 
diff --git a/Engine/Source/Runtime/Experimental/Chaos/Private/Chaos/PBDRigidsEvolutionGBF.cpp b/Engine/Source/Runtime/Experimental/Chaos/Private/Chaos/PBDRigidsEvolutionGBF.cpp
--- a/Engine/Source/Runtime/Experimental/Chaos/Private/Chaos/PBDRigidsEvolutionGBF.cpp
+++ b/Engine/Source/Runtime/Experimental/Chaos/Private/Chaos/PBDRigidsEvolutionGBF.cpp
@@ -48,9 +48,8 @@ namespace Chaos
 		int CollisionDisableCulledContacts = 0;
 		FAutoConsoleVariableRef CVarDisableCulledContacts(TEXT("p.CollisionDisableCulledContacts"), CollisionDisableCulledContacts, TEXT("Allow the PBDRigidsEvolutionGBF collision constraints to throw out contacts mid solve if they are culled."));
 
-		// @todo(chaos): this should be 0 but we need it for CCD atm
 		FRealSingle BoundsThicknessVelocityMultiplier = 0.0f;
-		FAutoConsoleVariableRef CVarBoundsThicknessVelocityMultiplier(TEXT("p.CollisionBoundsVelocityInflation"), BoundsThicknessVelocityMultiplier, TEXT("Collision velocity inflation for speculatibe contact generation.[def:2.0]"));
+		FAutoConsoleVariableRef CVarBoundsThicknessVelocityMultiplier(TEXT("p.CollisionBoundsVelocityInflation"), BoundsThicknessVelocityMultiplier, TEXT("Collision velocity inflation for speculative contact generation.[def:0.0]"));
 
 		FRealSingle SmoothedPositionLerpRate = 0.1f;
 		FAutoConsoleVariableRef CVarSmoothedPositionLerpRate(TEXT("p.Chaos.SmoothedPositionLerpRate"), SmoothedPositionLerpRate, TEXT("The interpolation rate for the smoothed position calculation. Used for sleeping."));
@@ -61,7 +60,8 @@ namespace Chaos
 		bool bChaosUseCCD = true;
 		FAutoConsoleVariableRef  CVarChaosUseCCD(TEXT("p.Chaos.CCD.UseCCD"), bChaosUseCCD , TEXT("Global flag to turn CCD on or off. Default is true"));
 
-		bool bChaosCollisionCCDUseTightBoundingBox = true; 
+		// NOTE: If we have mrope than 1 CCD iteration (ChaosCollisionCCDConstraintMaxProcessCount), the tight bounding box will cause us to miss secondary CCD collisions if the first one(s) result in a change in direction
+		bool bChaosCollisionCCDUseTightBoundingBox = true;
 		FAutoConsoleVariableRef  CVarChaosCollisionCCDUseTightBoundingBox(TEXT("p.Chaos.Collision.CCD.UseTightBoundingBox"), bChaosCollisionCCDUseTightBoundingBox , TEXT(""));
 
 		int32 ChaosSolverCollisionPriority = 0;
diff --git a/Engine/Source/Runtime/Experimental/Chaos/Public/Chaos/CCDUtilities.h b/Engine/Source/Runtime/Experimental/Chaos/Public/Chaos/CCDUtilities.h
--- a/Engine/Source/Runtime/Experimental/Chaos/Public/Chaos/CCDUtilities.h
+++ b/Engine/Source/Runtime/Experimental/Chaos/Public/Chaos/CCDUtilities.h
@@ -91,6 +91,7 @@ namespace Chaos
 
 	private:
 		void ApplySweptConstraints(const FReal Dt, TArrayView<FPBDCollisionConstraint* const> InSweptConstraints, const int32 NumDynamicParticles = TNumericLimits<int32>::Max());
+		bool UpdateParticleSweptConstraints(FCCDParticle* CCDParticle, const FReal IslandTOI, const FReal Dt);
 		// This is called after ApplySweptConstraints. This function updates manifold data which will be used in normal solve.
 		void UpdateSweptConstraints(const FReal Dt, FCollisionConstraintAllocator *CollisionAllocator);
 		void OverwriteXUsingV(const FReal Dt);
diff --git a/Engine/Source/Runtime/Experimental/Chaos/Public/Chaos/Collision/PBDCollisionConstraint.h b/Engine/Source/Runtime/Experimental/Chaos/Public/Chaos/Collision/PBDCollisionConstraint.h
--- a/Engine/Source/Runtime/Experimental/Chaos/Public/Chaos/Collision/PBDCollisionConstraint.h
+++ b/Engine/Source/Runtime/Experimental/Chaos/Public/Chaos/Collision/PBDCollisionConstraint.h
@@ -281,7 +281,14 @@ namespace Chaos
 		/**
 		 * @brief Enable or disable CCD for this constraint
 		*/
-		void SetCCDEnabled(const bool bCCDEnabled) { CCDType = bCCDEnabled ? ECollisionCCDType::Enabled : ECollisionCCDType::Disabled; }
+		void SetCCDEnabled(const bool bCCDEnabled)
+		{ 
+			CCDType = bCCDEnabled ? ECollisionCCDType::Enabled : ECollisionCCDType::Disabled;
+			if (bCCDEnabled && (CCDPenetrationThreshold == FReal(0)))
+			{
+				InitCCDThreshold();
+			}
+		}
 
 		//
 		// API
@@ -402,6 +409,31 @@ namespace Chaos
 		const FSavedManifoldPoint* FindSavedManifoldPoint(const FManifoldPoint& ManifoldPoint) const;
 		void ResetSavedManifoldPoints() { SavedManifoldPoints.Reset(); }
 
+		inline int32 NumEnabledManifoldPoints() const
+		{
+			int32 NumEnabled = 0;
+			if (IsEnabled())
+			{
+				for (int32 ManifoldPointIndex = 0; ManifoldPointIndex < NumManifoldPoints(); ++ManifoldPointIndex)
+				{
+					if (!ManifoldPoints[ManifoldPointIndex].Flags.bDisabled)
+					{
+						++NumEnabled;
+					}
+				}
+			}
+			return NumEnabled;
+		}
+
+		inline void DisableManifoldPoint(const int32 DisabledManifoldPointIndex)
+		{
+			ManifoldPoints[DisabledManifoldPointIndex].Flags.bDisabled = true;
+			if (NumEnabledManifoldPoints() == 0)
+			{
+				SetDisabled(true);
+			}
+		}
+
 		void AddIncrementalManifoldContact(const FContactPoint& ContactPoint);
 
 		inline void AddOneshotManifoldContact(const FContactPoint& ContactPoint)
@@ -458,6 +490,17 @@ namespace Chaos
 			}
 		}
 
+		/**
+		 * @brief Update the contact depth and TOI based on the current shape world transforms and the given start transforms
+		*/
+		void UpdateSweptManifoldPoints(const FVec3& ShapeStartWorldPosition0, const FVec3& ShapeStartWorldPosition1, const FReal Dt);
+
+		/**
+		 * @brief return the TOI to use for CCD given the start and end separations
+		 * We ignore separated and separating contacts, and modify TOI to leave some minimim overlap.
+		*/
+		FReal CalculateModifiedSweptTOI(const FReal StartPhi, const FReal EndPhi);
+
 		// The GJK warm-start data. This is updated directly in the narrow phase
 		FGJKSimplexData& GetGJKWarmStartData() { return GJKWarmStartData; }
 
@@ -611,6 +654,8 @@ namespace Chaos
 
 		void InitMarginsAndTolerances(const EImplicitObjectType ImplicitType0, const EImplicitObjectType ImplicitType1, const FReal Margin0, const FReal Margin1);
 
+		void InitCCDThreshold();
+
 	private:
 		FReal CalculateSavedManifoldPointScore(const FSavedManifoldPoint& SavedManifoldPoint, const FManifoldPoint& ManifoldPoint, const FReal DistanceToleranceSq) const;
 
@@ -629,7 +674,8 @@ namespace Chaos
 		FVec3 AccumulatedImpulse;					// @todo(chaos): we need to accumulate angular impulse separately
 
 		// Value in range [0,1] used to interpolate P between [X,P] that we will rollback to when solving at time of impact.
-		FReal TimeOfImpact;
+		FReal CCDTimeOfImpact;
+		FReal CCDPenetrationThreshold;
 
 	private:
 		FPBDCollisionConstraintContainerCookie ContainerCookie;
diff --git a/Engine/Source/Runtime/Experimental/Chaos/Public/Chaos/CollisionResolution.h b/Engine/Source/Runtime/Experimental/Chaos/Public/Chaos/CollisionResolution.h
--- a/Engine/Source/Runtime/Experimental/Chaos/Public/Chaos/CollisionResolution.h
+++ b/Engine/Source/Runtime/Experimental/Chaos/Public/Chaos/CollisionResolution.h
@@ -47,7 +47,7 @@ namespace Chaos
 		 * @brief Update the contact manifold on the constraint
 		 * @note Transforms are shape world-space transforms (not particle transforms) at the start of the sweep. The end of the sweep are the transforms stored in the constraint as ShapeWorldTransform0/1
 		*/
-		void CHAOS_API UpdateConstraintSwept(FPBDCollisionConstraint& Constraint, const FRigidTransform3& ShapeWorldTransform0, const FRigidTransform3& ShapeWorldTransform1, const FReal Dt);
+		bool CHAOS_API UpdateConstraintSwept(FPBDCollisionConstraint& Constraint, const FRigidTransform3& ShapeStartWorldTransform0, const FRigidTransform3& ShapeStartWorldTransform1, const FReal Dt);
 
 		/**
 		 * @brief Determine the shape pair type for use in UpdateConstraints
@@ -64,9 +64,6 @@ namespace Chaos
 		// @todo(chaos): remove this and use UpdateConstraint instead
 		template<ECollisionUpdateType UpdateType>
 		void CHAOS_API UpdateConstraintFromGeometry(FPBDCollisionConstraint& Constraint, const FRigidTransform3& ParticleTransform0, const FRigidTransform3& ParticleTransform1, const FReal Dt);
-		template<ECollisionUpdateType UpdateType>
-		bool CHAOS_API UpdateConstraintFromGeometrySwept(FPBDCollisionConstraint& Constraint, const FRigidTransform3& ParticleTransform0, const FRigidTransform3& ParticleTransform1, const FReal Dt);
-
 		// Create constraints for the particle pair. This could create multiple constraints: one for each potentially colliding shape pair in multi-shape particles.
 		void CHAOS_API ConstructConstraints(TGeometryParticleHandle<FReal, 3>* Particle0, TGeometryParticleHandle<FReal, 3>* Particle1, const FImplicitObject* Implicit0, const FPerShapeData* Shape0, const FBVHParticles* Simplicial0, const FImplicitObject* Implicit1, const FPerShapeData* Shape1, const FBVHParticles* Simplicial1, const FRigidTransform3& ParticleWorldTransform0, const FRigidTransform3& Transform0, const FRigidTransform3& ParticleWorldTransform1, const FRigidTransform3& Transform1, const FReal CullDistance, const FReal Dt,const FCollisionContext& Context);
 
diff --git a/Engine/Source/Runtime/Experimental/Chaos/Public/Chaos/PBDRigidsEvolutionGBF.h b/Engine/Source/Runtime/Experimental/Chaos/Public/Chaos/PBDRigidsEvolutionGBF.h
--- a/Engine/Source/Runtime/Experimental/Chaos/Public/Chaos/PBDRigidsEvolutionGBF.h
+++ b/Engine/Source/Runtime/Experimental/Chaos/Public/Chaos/PBDRigidsEvolutionGBF.h
@@ -34,6 +34,7 @@ namespace Chaos
 		CHAOS_API extern FRealSingle HackMaxVelocity;
 		CHAOS_API extern FRealSingle SmoothedPositionLerpRate;
 		CHAOS_API extern bool bChaosCollisionCCDUseTightBoundingBox;
+		CHAOS_API extern int32 ChaosCollisionCCDConstraintMaxProcessCount;
 		CHAOS_API extern int32 ChaosSolverDrawCCDThresholds;
 	}
 
@@ -296,6 +297,8 @@ namespace Chaos
 
 					if (!Particle.CCDEnabled())
 					{
+						// Expand bounds about P/Q by a small amount. This can still result in missed collisions, especially
+						// when we have joints that pull the body back to X/R, if P-X is greater than the BoundsThickness
 						Particle.UpdateWorldSpaceState(FRigidTransform3(Particle.P(), Particle.Q()), FVec3(BoundsThickness));
 					}
 					else
@@ -310,14 +313,16 @@ namespace Chaos
 
 						if (CCDHelpers::DeltaExceedsThreshold(Particle.CCDAxisThreshold(), Particle.P() - Particle.X(), Particle.Q()))
 						{
-							if (CVars::bChaosCollisionCCDUseTightBoundingBox)
+							// We sweep the bounds from P back along the velocity and expand by a small amount.
+							// If not using tight bounds we also expand the bounds in all directions by Velocity. This is necessary only for secondary CCD collisions
+							// @todo(chaos): expanding the bounds by velocity is very expensive - revisit this
+							const FVec3 VDt = Particle.V() * Dt;
+							FReal CCDBoundsExpansion = BoundsThickness;
+							if (!CVars::bChaosCollisionCCDUseTightBoundingBox && (CVars::ChaosCollisionCCDConstraintMaxProcessCount > 1))
 							{
-								Particle.UpdateWorldSpaceStateSwept(FRigidTransform3(Particle.P(), Particle.Q()), FVec3(BoundsThickness), Particle.X() - Particle.P());
-							}
-							else
-							{
-								Particle.UpdateWorldSpaceState(FRigidTransform3(Particle.P(), Particle.Q()), FVec3(BoundsThickness) + Particle.V() * Dt);
+								CCDBoundsExpansion += VDt.GetAbsMax();
 							}
+							Particle.UpdateWorldSpaceStateSwept(FRigidTransform3(Particle.P(), Particle.Q()), FVec3(CCDBoundsExpansion), -VDt);
 						}
 						else
 						{
