Chaos : Geometry Collection : share implicit creation code between clusters and other particles
- also share some size specific query code

diff --git a/Engine/Source/Runtime/Experimental/Chaos/Private/PhysicsProxy/GeometryCollectionPhysicsProxy.cpp b/Engine/Source/Runtime/Experimental/Chaos/Private/PhysicsProxy/GeometryCollectionPhysicsProxy.cpp
--- a/Engine/Source/Runtime/Experimental/Chaos/Private/PhysicsProxy/GeometryCollectionPhysicsProxy.cpp
+++ b/Engine/Source/Runtime/Experimental/Chaos/Private/PhysicsProxy/GeometryCollectionPhysicsProxy.cpp
@@ -2027,8 +2027,150 @@ bool FGeometryCollectionPhysicsProxy::PullFromPhysicsState(const Chaos::FDirtyGe
 //==============================================================================
 // STATIC SETUP FUNCTIONS
 //==============================================================================
+static const FSharedSimulationSizeSpecificData& GetSizeSpecificData(const TArray<FSharedSimulationSizeSpecificData>& SizeSpecificData, const FGeometryCollection& RestCollection, const int32 TransformIndex, const FBox& BoundingBox)
+{
+	// If we have a normalized Size available, use that to determine SizeSpecific index, otherwise fall back on Bounds volume.
+	int32 SizeSpecificIdx;
+	const TManagedArray<float>* RelativeSizes = RestCollection.FindAttribute<float>("Size", FTransformCollection::TransformGroup);
+	if (RelativeSizes)
+	{
+		SizeSpecificIdx = GeometryCollection::SizeSpecific::FindIndexForVolume(SizeSpecificData, (*RelativeSizes)[TransformIndex]);
+	}
+	else
+	{
+		SizeSpecificIdx = GeometryCollection::SizeSpecific::FindIndexForVolume(SizeSpecificData, BoundingBox);
+	}
+	return SizeSpecificData[SizeSpecificIdx];
+}
 
+FGeometryDynamicCollection::FSharedImplicit CreateImplicitGeometry(
+	const FSharedSimulationSizeSpecificData& SizeSpecificData,
+	const int32 TransformGroupIndex,
+	const FGeometryCollection& RestCollection,
+	const Chaos::FParticles& MassSpaceParticles,
+	const Chaos::FTriangleMesh& TriMesh,
+	const FBox& InstanceBoundingBox,
+	const Chaos::FReal InnerRadius,
+	const Chaos::FErrorReporter& ErrorReporter,
+	const Chaos::FVec3* ClusterMaxChildBounds = nullptr
+	)
+{
+	FGeometryDynamicCollection::FSharedImplicit NewImplicit;
+	if (SizeSpecificData.CollisionShapesData.Num())
+	{
+		const TManagedArray<FTransform>& CollectionMassToLocal = RestCollection.GetAttribute<FTransform>("MassToLocal", FTransformCollection::TransformGroup);
+		const TManagedArray<TSet<int32>>* TransformToConvexIndices = RestCollection.FindAttribute<TSet<int32>>("TransformToConvexIndices", FTransformCollection::TransformGroup);
+		const TManagedArray<TUniquePtr<Chaos::FConvex>>* ConvexGeometry = RestCollection.FindAttribute<TUniquePtr<Chaos::FConvex>>("ConvexHull", "Convex");
 
+		const FCollectionCollisionTypeData& CollisionTypeData = SizeSpecificData.CollisionShapesData[0]; 
+		switch(CollisionTypeData.ImplicitType)
+		{
+		case EImplicitTypeEnum::Chaos_Implicit_LevelSet:
+			{
+				const FCollectionLevelSetData& LevelSetData = CollisionTypeData.LevelSetData; 
+				int32 MinResolution = LevelSetData.MinLevelSetResolution;
+				int32 MaxResolution = LevelSetData.MaxLevelSetResolution;
+				if (ClusterMaxChildBounds)
+				{
+					const Chaos::FVec3 Scale = 2. * InstanceBoundingBox.GetExtent() / (*ClusterMaxChildBounds); // FBox's extents are 1/2 (Max - Min)
+					const Chaos::FReal ScaleMax = Scale.GetAbsMax();
+					const Chaos::FReal ScaleMin = Scale.GetAbsMin();
+
+					Chaos::FReal MinResolutionReal = ScaleMin * (Chaos::FReal)LevelSetData.MinLevelSetResolution;
+					MinResolutionReal = FMath::Clamp(MinResolutionReal,
+						(Chaos::FReal)LevelSetData.MinLevelSetResolution,
+						(Chaos::FReal)LevelSetData.MinClusterLevelSetResolution);
+
+					Chaos::FReal MaxResolutionReal = ScaleMax * (Chaos::FReal)LevelSetData.MaxLevelSetResolution;
+					MaxResolutionReal = FMath::Clamp(MaxResolutionReal,
+						(Chaos::FReal)LevelSetData.MaxLevelSetResolution,
+						(Chaos::FReal)LevelSetData.MaxClusterLevelSetResolution);
+
+					MinResolution = FMath::FloorToInt32(MinResolutionReal);
+					MaxResolution = FMath::FloorToInt32(MaxResolutionReal);
+				}
+				
+				NewImplicit = FGeometryDynamicCollection::FSharedImplicit(
+					FCollisionStructureManager::NewImplicitLevelset(
+						ErrorReporter,
+						MassSpaceParticles,
+						TriMesh,
+						InstanceBoundingBox,
+						MinResolution,
+						MaxResolution,
+						CollisionTypeData.CollisionObjectReductionPercentage,
+						CollisionTypeData.CollisionType));
+				
+				// Fall back on sphere if level set rasterization failed.
+				if (!NewImplicit)
+				{
+					NewImplicit = FGeometryDynamicCollection::FSharedImplicit(
+						FCollisionStructureManager::NewImplicitSphere(
+						InnerRadius,
+						CollisionTypeData.CollisionObjectReductionPercentage,
+						CollisionTypeData.CollisionType));
+				}
+			}
+			break;
+				
+		case EImplicitTypeEnum::Chaos_Implicit_Box:
+			{
+				NewImplicit = FGeometryDynamicCollection::FSharedImplicit(
+					FCollisionStructureManager::NewImplicitBox(
+						InstanceBoundingBox,
+						CollisionTypeData.CollisionObjectReductionPercentage,
+						CollisionTypeData.CollisionType));
+			}
+			break;
+			
+		case EImplicitTypeEnum::Chaos_Implicit_Sphere:
+			{
+				NewImplicit = FGeometryDynamicCollection::FSharedImplicit(
+					FCollisionStructureManager::NewImplicitSphere(
+						InnerRadius,
+						CollisionTypeData.CollisionObjectReductionPercentage,
+						CollisionTypeData.CollisionType));
+			}
+			break;
+			
+		case EImplicitTypeEnum::Chaos_Implicit_Convex:
+			{
+				if (ConvexGeometry && TransformToConvexIndices)
+				{
+					NewImplicit = FGeometryDynamicCollection::FSharedImplicit(
+						FCollisionStructureManager::NewImplicitConvex(
+							(*TransformToConvexIndices)[TransformGroupIndex].Array(),
+							ConvexGeometry,
+							CollisionTypeData.CollisionType,
+							CollectionMassToLocal[TransformGroupIndex],
+							CollisionTypeData.CollisionMarginFraction
+						)
+					);
+				}
+			}
+			break;
+			
+		case EImplicitTypeEnum::Chaos_Implicit_Capsule:
+			{
+				NewImplicit = FGeometryDynamicCollection::FSharedImplicit(
+					FCollisionStructureManager::NewImplicitCapsule(
+						InstanceBoundingBox,
+						CollisionTypeData.CollisionObjectReductionPercentage,
+						CollisionTypeData.CollisionType));
+			}
+			break;
+		
+		case EImplicitTypeEnum::Chaos_Implicit_None:
+			// nothing to do here 
+			break;
+			
+		default:
+			ensure(false); // unsupported implicit type!
+			break;
+		}
+	}
+	return NewImplicit;
+}
 
 /** 
 	NOTE - Making any changes to data stored on the rest collection below MUST be accompanied
@@ -2369,20 +2511,7 @@ void FGeometryCollectionPhysicsProxy::InitializeSharedCollisionStructures(
 				InstanceBoundingBox = FBox(MassProperties.CenterOfMass, MassProperties.CenterOfMass);
 			}
 
-			// If we have a normalized Size available, use that to determine SizeSpecific index, otherwise fall back on Bounds volume.
-			int32 SizeSpecificIdx;
-			if (bUseRelativeSize)
-			{
-				const TManagedArray<float>& RelativeSize = RestCollection.GetAttribute<float>(TEXT("Size"), FTransformCollection::TransformGroup);
-				SizeSpecificIdx = GeometryCollection::SizeSpecific::FindIndexForVolume(SharedParams.SizeSpecificData, RelativeSize[TransformGroupIndex]);
-			}
-			else
-			{
-				SizeSpecificIdx = GeometryCollection::SizeSpecific::FindIndexForVolume(SharedParams.SizeSpecificData, InstanceBoundingBox);
-			}
-			
-			const FSharedSimulationSizeSpecificData& SizeSpecificData = SharedParams.SizeSpecificData[SizeSpecificIdx];
-
+			const FSharedSimulationSizeSpecificData& SizeSpecificData = GetSizeSpecificData(SharedParams.SizeSpecificData, RestCollection, TransformGroupIndex, InstanceBoundingBox);
 			if (SizeSpecificData.CollisionShapesData.Num())
 			{
 				//
@@ -2405,77 +2534,17 @@ void FGeometryCollectionPhysicsProxy::InitializeSharedCollisionStructures(
 						ensureMsgf(false, TEXT("Simplicial is empty."));
 					}
 
-					if (SizeSpecificData.CollisionShapesData[0].ImplicitType == EImplicitTypeEnum::Chaos_Implicit_LevelSet)
-					{
-						LocalErrorReporter.SetPrefix(BaseErrorPrefix + " | Transform Index: " + FString::FromInt(TransformGroupIndex) + " of " + FString::FromInt(TransformIndex.Num()));
-						CollectionImplicits[TransformGroupIndex] = FGeometryDynamicCollection::FSharedImplicit(
-							FCollisionStructureManager::NewImplicitLevelset(
-								LocalErrorReporter,
-								MassSpaceParticles,
-								*TriMesh,
-								InstanceBoundingBox,
-								SizeSpecificData.CollisionShapesData[0].LevelSetData.MinLevelSetResolution,
-								SizeSpecificData.CollisionShapesData[0].LevelSetData.MaxLevelSetResolution,
-								SizeSpecificData.CollisionShapesData[0].CollisionObjectReductionPercentage,
-								SizeSpecificData.CollisionShapesData[0].CollisionType));
-						// Fall back on sphere if level set rasterization failed.
-						if (!CollectionImplicits[TransformGroupIndex])
-						{
-							CollectionImplicits[TransformGroupIndex] = FGeometryDynamicCollection::FSharedImplicit(
-								FCollisionStructureManager::NewImplicitSphere(
-									InnerRadius[GeometryIndex],
-									SizeSpecificData.CollisionShapesData[0].CollisionObjectReductionPercentage,
-									SizeSpecificData.CollisionShapesData[0].CollisionType));
-						}
-					}
-					else if (SizeSpecificData.CollisionShapesData[0].ImplicitType == EImplicitTypeEnum::Chaos_Implicit_Box)
-					{
-						CollectionImplicits[TransformGroupIndex] = FGeometryDynamicCollection::FSharedImplicit(
-							FCollisionStructureManager::NewImplicitBox(
-								InstanceBoundingBox,
-								SizeSpecificData.CollisionShapesData[0].CollisionObjectReductionPercentage,
-								SizeSpecificData.CollisionShapesData[0].CollisionType));
-					}
-					else if (SizeSpecificData.CollisionShapesData[0].ImplicitType == EImplicitTypeEnum::Chaos_Implicit_Sphere)
-					{
-						CollectionImplicits[TransformGroupIndex] = FGeometryDynamicCollection::FSharedImplicit(
-							FCollisionStructureManager::NewImplicitSphere(
-								InnerRadius[GeometryIndex],
-								SizeSpecificData.CollisionShapesData[0].CollisionObjectReductionPercentage,
-								SizeSpecificData.CollisionShapesData[0].CollisionType));
-					}
-					else if (SizeSpecificData.CollisionShapesData[0].ImplicitType == EImplicitTypeEnum::Chaos_Implicit_Convex)
-					{
-						if (ConvexGeometry && TransformToConvexIndices)
-						{
-							CollectionImplicits[TransformGroupIndex] = FGeometryDynamicCollection::FSharedImplicit(
-									FCollisionStructureManager::NewImplicitConvex(
-										(*TransformToConvexIndices)[TransformGroupIndex].Array(),
-										ConvexGeometry,
-										SizeSpecificData.CollisionShapesData[0].CollisionType,
-										CollectionMassToLocal[TransformGroupIndex],
-										(Chaos::FReal)SizeSpecificData.CollisionShapesData[0].CollisionMarginFraction
-									)
-							);
-						}
-					}
-					else if (SizeSpecificData.CollisionShapesData[0].ImplicitType == EImplicitTypeEnum::Chaos_Implicit_Capsule)
-					{
-						CollectionImplicits[TransformGroupIndex] = FGeometryDynamicCollection::FSharedImplicit(
-							FCollisionStructureManager::NewImplicitCapsule(
-								InstanceBoundingBox,
-								SizeSpecificData.CollisionShapesData[0].CollisionObjectReductionPercentage,
-								SizeSpecificData.CollisionShapesData[0].CollisionType));
-					}
-					else if (SizeSpecificData.CollisionShapesData[0].ImplicitType == EImplicitTypeEnum::Chaos_Implicit_None)
-					{
-						CollectionImplicits[TransformGroupIndex] = nullptr;
-					}
-					else
-					{
-						ensure(false); // unsupported implicit type!
-					}
-
+					LocalErrorReporter.SetPrefix(BaseErrorPrefix + " | Transform Index: " + FString::FromInt(TransformGroupIndex) + " of " + FString::FromInt(TransformIndex.Num()));
+					CollectionImplicits[TransformGroupIndex] = CreateImplicitGeometry(
+						SizeSpecificData,
+						TransformGroupIndex,
+						RestCollection,
+						MassSpaceParticles,
+						*TriMesh,
+						InstanceBoundingBox,
+						InnerRadius[GeometryIndex],
+						LocalErrorReporter);
+					
 					if (CollectionImplicits[TransformGroupIndex] && CollectionImplicits[TransformGroupIndex]->HasBoundingBox())
 					{
 						const auto Implicit = CollectionImplicits[TransformGroupIndex];
@@ -2644,119 +2713,26 @@ void FGeometryCollectionPhysicsProxy::InitializeSharedCollisionStructures(
 				CollectionInertiaTensor[ClusterTransformIdx] = FVector3f(InertiaDiagonal);	// LWC_TODO: Precision loss
 				CollectionMass[ClusterTransformIdx] = (FRealSingle)CollectionSpaceParticles->M(ClusterTransformIdx);
 
+				const FSharedSimulationSizeSpecificData& SizeSpecificData = GetSizeSpecificData(SharedParams.SizeSpecificData, RestCollection, TransformGroupIndex, InstanceBoundingBox);
 
-				int32 SizeSpecificIdx;
-				if (bUseRelativeSize)
-				{
-					const TManagedArray<float>& RelativeSize = RestCollection.GetAttribute<float>(TEXT("Size"), FTransformCollection::TransformGroup);
-					SizeSpecificIdx = GeometryCollection::SizeSpecific::FindIndexForVolume(SharedParams.SizeSpecificData, RelativeSize[TransformGroupIndex]);
-				}
-				else
-				{
-					SizeSpecificIdx = GeometryCollection::SizeSpecific::FindIndexForVolume(SharedParams.SizeSpecificData, InstanceBoundingBox);	
-				}
-				const FSharedSimulationSizeSpecificData& SizeSpecificData = SharedParams.SizeSpecificData[SizeSpecificIdx];
-
-				if (SizeSpecificData.CollisionShapesData.Num())
-				{
-					if (SizeSpecificData.CollisionShapesData[0].ImplicitType == EImplicitTypeEnum::Chaos_Implicit_LevelSet)
-					{
-						const FVec3 Scale = 2 * InstanceBoundingBox.GetExtent() / MaxChildBounds; // FBox's extents are 1/2 (Max - Min)
-						const FReal ScaleMax = Scale.GetAbsMax();
-						const FReal ScaleMin = Scale.GetAbsMin();
-
-						FReal MinResolution = ScaleMin * (FReal)SizeSpecificData.CollisionShapesData[0].LevelSetData.MinLevelSetResolution;
-						MinResolution = FMath::Clamp(MinResolution,
-							(FReal)SizeSpecificData.CollisionShapesData[0].LevelSetData.MinLevelSetResolution,
-							(FReal)SizeSpecificData.CollisionShapesData[0].LevelSetData.MinClusterLevelSetResolution);
-
-						FReal MaxResolution = ScaleMax * (FReal)SizeSpecificData.CollisionShapesData[0].LevelSetData.MaxLevelSetResolution;
-						MaxResolution = FMath::Clamp(MaxResolution,
-							(FReal)SizeSpecificData.CollisionShapesData[0].LevelSetData.MaxLevelSetResolution,
-							(FReal)SizeSpecificData.CollisionShapesData[0].LevelSetData.MaxClusterLevelSetResolution);
-
-						//don't support non level-set serialization
-						ErrorReporter.SetPrefix(BaseErrorPrefix + " | Cluster Transform Index: " + FString::FromInt(ClusterTransformIdx));
-						CollectionImplicits[ClusterTransformIdx] = FGeometryDynamicCollection::FSharedImplicit(
-							FCollisionStructureManager::NewImplicitLevelset(
-								ErrorReporter,
-								MassSpaceParticles,
-								*UnionMesh,
-								InstanceBoundingBox,
-								FMath::FloorToInt32(MinResolution),
-								FMath::FloorToInt32(MaxResolution),
-								SizeSpecificData.CollisionShapesData[0].CollisionObjectReductionPercentage,
-								SizeSpecificData.CollisionShapesData[0].CollisionType));
-						// Fall back on sphere if level set rasterization failed.
-						if (!CollectionImplicits[ClusterTransformIdx])
-						{
-							CollectionImplicits[ClusterTransformIdx] = FGeometryDynamicCollection::FSharedImplicit(
-								FCollisionStructureManager::NewImplicitSphere(
-									InstanceBoundingBox.GetExtent().GetAbsMin(), // FBox's extents are 1/2 (Max - Min)
-									SizeSpecificData.CollisionShapesData[0].CollisionObjectReductionPercentage,
-									SizeSpecificData.CollisionShapesData[0].CollisionType));
-						}
-
-						CollectionSimplicials[ClusterTransformIdx] = TUniquePtr<FSimplicial>(
-							FCollisionStructureManager::NewSimplicial(MassSpaceParticles, *UnionMesh, CollectionImplicits[ClusterTransformIdx].Get(),
-							SharedParams.MaximumCollisionParticleCount));
-					}
-					else if (SizeSpecificData.CollisionShapesData[0].ImplicitType == EImplicitTypeEnum::Chaos_Implicit_Box)
-					{
-						CollectionImplicits[ClusterTransformIdx] = FGeometryDynamicCollection::FSharedImplicit(
-							FCollisionStructureManager::NewImplicitBox(
-								InstanceBoundingBox,
-								SizeSpecificData.CollisionShapesData[0].CollisionObjectReductionPercentage,
-								SizeSpecificData.CollisionShapesData[0].CollisionType));
+				ErrorReporter.SetPrefix(BaseErrorPrefix + " | Cluster Transform Index: " + FString::FromInt(ClusterTransformIdx));
 
-						CollectionSimplicials[ClusterTransformIdx] = TUniquePtr<FSimplicial>(
-							FCollisionStructureManager::NewSimplicial(MassSpaceParticles, *UnionMesh, CollectionImplicits[ClusterTransformIdx].Get(),
-							SharedParams.MaximumCollisionParticleCount));
-					}
-					else if (SizeSpecificData.CollisionShapesData[0].ImplicitType == EImplicitTypeEnum::Chaos_Implicit_Sphere)
-					{
-						CollectionImplicits[ClusterTransformIdx] = FGeometryDynamicCollection::FSharedImplicit(
-							FCollisionStructureManager::NewImplicitSphere(
-								InstanceBoundingBox.GetExtent().GetAbsMin(), // FBox's extents are 1/2 (Max - Min)
-								SizeSpecificData.CollisionShapesData[0].CollisionObjectReductionPercentage,
-								SizeSpecificData.CollisionShapesData[0].CollisionType));
+				CollectionImplicits[ClusterTransformIdx] = CreateImplicitGeometry(
+					SizeSpecificData,
+					ClusterTransformIdx,
+					RestCollection,
+					MassSpaceParticles,
+					*UnionMesh,
+					InstanceBoundingBox,
+					InstanceBoundingBox.GetExtent().GetAbsMin(), // InnerRadius
+					ErrorReporter,
+					&MaxChildBounds
+					);
 
-						CollectionSimplicials[ClusterTransformIdx] = TUniquePtr<FSimplicial>(
+				// create simplicial from the implicit geometry
+				CollectionSimplicials[ClusterTransformIdx] = TUniquePtr<FSimplicial>(
 							FCollisionStructureManager::NewSimplicial(MassSpaceParticles, *UnionMesh, CollectionImplicits[ClusterTransformIdx].Get(),
 							SharedParams.MaximumCollisionParticleCount));
-					}
-					else if (SizeSpecificData.CollisionShapesData[0].ImplicitType == EImplicitTypeEnum::Chaos_Implicit_Convex)
-					{
-						if (ConvexGeometry && TransformToConvexIndices)
-						{
-							CollectionImplicits[TransformGroupIndex] = FGeometryDynamicCollection::FSharedImplicit(
-								FCollisionStructureManager::NewImplicitConvex(
-									(*TransformToConvexIndices)[TransformGroupIndex].Array(),
-									ConvexGeometry,
-									SizeSpecificData.CollisionShapesData[0].CollisionType,
-									CollectionMassToLocal[TransformGroupIndex],
-									(Chaos::FReal)SizeSpecificData.CollisionShapesData[0].CollisionMarginFraction)
-							);
-						}
-					}
-					else if (SizeSpecificData.CollisionShapesData[0].ImplicitType == EImplicitTypeEnum::Chaos_Implicit_Capsule)
-					{
-						CollectionImplicits[ClusterTransformIdx] = FGeometryDynamicCollection::FSharedImplicit(
-							FCollisionStructureManager::NewImplicitCapsule(
-								InstanceBoundingBox,
-								SizeSpecificData.CollisionShapesData[0].CollisionObjectReductionPercentage,
-								SizeSpecificData.CollisionShapesData[0].CollisionType));
-
-						CollectionSimplicials[ClusterTransformIdx] = TUniquePtr<FSimplicial>(
-							FCollisionStructureManager::NewSimplicial(MassSpaceParticles, *UnionMesh, CollectionImplicits[ClusterTransformIdx].Get(),
-								SharedParams.MaximumCollisionParticleCount));
-					}
-					else // Assume it's a union???
-					{
-						CollectionImplicits[ClusterTransformIdx].Reset();	//union so just set as null
-						CollectionSimplicials[ClusterTransformIdx].Reset();
-					}
-				}
 
 				TriangleMeshesArray[ClusterTransformIdx] = MoveTemp(UnionMesh);
 			}
