Interchange - Fix NoPCH compile + Adding label option when settings asset source file.

diff --git a/Engine/Plugins/Experimental/Interchange/Runtime/Source/Import/Private/Mesh/InterchangeSkeletalMeshFactory.cpp b/Engine/Plugins/Experimental/Interchange/Runtime/Source/Import/Private/Mesh/InterchangeSkeletalMeshFactory.cpp
--- a/Engine/Plugins/Experimental/Interchange/Runtime/Source/Import/Private/Mesh/InterchangeSkeletalMeshFactory.cpp
+++ b/Engine/Plugins/Experimental/Interchange/Runtime/Source/Import/Private/Mesh/InterchangeSkeletalMeshFactory.cpp
@@ -1521,7 +1521,8 @@ bool UInterchangeSkeletalMeshFactory::SetSourceFilename(const UObject* Object, c
 #if WITH_EDITORONLY_DATA
 	if (const USkeletalMesh* SkeletalMesh = Cast<USkeletalMesh>(Object))
 	{
-		return UE::Interchange::FFactoryCommon::SetSourceFilename(SkeletalMesh->GetAssetImportData(), SourceFilename, SourceIndex);
+		const FString SourceLabel = USkeletalMesh::GetSourceFileLabelFromIndex(SourceIndex).ToString();
+		return UE::Interchange::FFactoryCommon::SetSourceFilename(SkeletalMesh->GetAssetImportData(), SourceFilename, SourceIndex, SourceLabel);
 	}
 #endif
 
diff --git a/Engine/Source/Runtime/Experimental/Interchange/Engine/Private/InterchangeImportCommon.cpp b/Engine/Source/Runtime/Experimental/Interchange/Engine/Private/InterchangeImportCommon.cpp
--- a/Engine/Source/Runtime/Experimental/Interchange/Engine/Private/InterchangeImportCommon.cpp
+++ b/Engine/Source/Runtime/Experimental/Interchange/Engine/Private/InterchangeImportCommon.cpp
@@ -192,20 +192,19 @@ namespace UE
 			return false;
 		}
 
-		bool FFactoryCommon::SetSourceFilename(UAssetImportData* AssetImportData, const FString& SourceFilename, int32 SourceIndex)
+		bool FFactoryCommon::SetSourceFilename(UAssetImportData* AssetImportData, const FString& SourceFilename, int32 SourceIndex, const FString& SourceLabel)
 		{
 			if (AssetImportData)
 			{
-				int32 RealSourceFileIndex = SourceIndex == INDEX_NONE ? 0 : SourceIndex;
-				if (RealSourceFileIndex < AssetImportData->GetSourceFileCount())
+				const int32 SafeSourceIndex = SourceIndex == INDEX_NONE ? 0 : SourceIndex;				
+				if (SafeSourceIndex < AssetImportData->GetSourceFileCount())
 				{
-					AssetImportData->UpdateFilenameOnly(SourceFilename, SourceIndex);
+					AssetImportData->UpdateFilenameOnly(SourceFilename, SafeSourceIndex);
 				}
 				else
 				{
 					//Create a source file entry, this case happen when user import a specific content for the first time
-					FString SourceIndexLabel = USkeletalMesh::GetSourceFileLabelFromIndex(RealSourceFileIndex).ToString();
-					AssetImportData->AddFileName(SourceFilename, RealSourceFileIndex, SourceIndexLabel);
+					AssetImportData->AddFileName(SourceFilename, SafeSourceIndex, SourceLabel);
 				}
 				return true;
 			}
diff --git a/Engine/Source/Runtime/Experimental/Interchange/Engine/Public/InterchangeImportCommon.h b/Engine/Source/Runtime/Experimental/Interchange/Engine/Public/InterchangeImportCommon.h
--- a/Engine/Source/Runtime/Experimental/Interchange/Engine/Public/InterchangeImportCommon.h
+++ b/Engine/Source/Runtime/Experimental/Interchange/Engine/Public/InterchangeImportCommon.h
@@ -87,7 +87,7 @@ namespace UE
 			/**
 			 * Sets the SourceFileName value at the specified index.
 			 */
-			static bool SetSourceFilename(UAssetImportData* AssetImportData, const FString& SourceFilename, int32 SourceIndex);
+			static bool SetSourceFilename(UAssetImportData* AssetImportData, const FString& SourceFilename, int32 SourceIndex, const FString& SourceLabel = FString());
 
 			/**
 			 * Set the object's reimport source at the specified index value.
