Chaos Cloth Bending Element constraint fixes.

There were a few bugs with the bending element constraints -- firstly since it was using acos to calculate the bend angle, it could only calculate an angle between 0 and pi rather than -pi to pi. Switched to use atan2 instead.
Secondly, there were bugs in the gradient calculation. I'm not quite sure where it came from, but numerical differentiation tests showed it was not correct, and the constraints were not stable in practice.

Updated with new gradient calculations (based on atan2). I did numerical differentiation tests, and they now pass. I also added a limit to how large of a step you can take in a single iteration (pi/4) since the linearization (gradients) of the constraints will not produce a good direction for taking larger steps.

In testing, the constraints are now stable, assuming you have stretch and or area constraints to resist large stretching deformations. These constraints are also much better than the spring-based  bending constraints at holding nearly flat configurations.

It should be noted that this code has not been optimized to use parallelism or ispc, so they are very slow. Optimization will be done in a later pass once additional improvements to the behavior of this constraint are made (namely biphasic stiffness as well as warp vs weft directional stiffnesses).

diff --git a/Engine/Source/Runtime/Experimental/Chaos/Public/Chaos/PBDBendingConstraints.h b/Engine/Source/Runtime/Experimental/Chaos/Public/Chaos/PBDBendingConstraints.h
--- a/Engine/Source/Runtime/Experimental/Chaos/Public/Chaos/PBDBendingConstraints.h
+++ b/Engine/Source/Runtime/Experimental/Chaos/Public/Chaos/PBDBendingConstraints.h
@@ -25,7 +25,7 @@ public:
 			const int32 i2 = Constraint[1];
 			const int32 i3 = Constraint[2];
 			const int32 i4 = Constraint[3];
-			const TArray<FSolverVec3> Grads = Base::GetGradients(InParticles, i);
+			const TStaticArray<FSolverVec3, 4> Grads = Base::GetGradients(InParticles, i);
 			const FSolverReal S = Base::GetScalingFactor(InParticles, i, Grads);
 			InParticles.P(i1) -= S * InParticles.InvM(i1) * Grads[0];
 			InParticles.P(i2) -= S * InParticles.InvM(i2) * Grads[1];
diff --git a/Engine/Source/Runtime/Experimental/Chaos/Public/Chaos/PBDBendingConstraintsBase.h b/Engine/Source/Runtime/Experimental/Chaos/Public/Chaos/PBDBendingConstraintsBase.h
--- a/Engine/Source/Runtime/Experimental/Chaos/Public/Chaos/PBDBendingConstraintsBase.h
+++ b/Engine/Source/Runtime/Experimental/Chaos/Public/Chaos/PBDBendingConstraintsBase.h
@@ -2,10 +2,9 @@
 #pragma once
 
 #include "Chaos/PBDSoftsEvolutionFwd.h"
-#include "Chaos/Array.h"
 #include "Chaos/PBDSoftsSolverParticles.h"
 #include "Chaos/ParticleRule.h"
-#include "GenericPlatform/GenericPlatformMath.h"
+#include "Containers/StaticArray.h"
 
 namespace Chaos::Softs
 {
@@ -22,36 +21,25 @@ class FPBDBendingConstraintsBase
 			const FSolverVec3& P2 = InParticles.X(Constraint[1]);
 			const FSolverVec3& P3 = InParticles.X(Constraint[2]);
 			const FSolverVec3& P4 = InParticles.X(Constraint[3]);
-			MAngles.Add(GetAngle(P1, P2, P3, P4));
+			MAngles.Add(CalcAngle(P1, P2, P3, P4));
 		}
 	}
 	virtual ~FPBDBendingConstraintsBase() {}
 
-	TArray<FSolverVec3> GetGradients(const FSolverParticles& InParticles, const int32 i) const
+	TStaticArray<FSolverVec3, 4> GetGradients(const FSolverParticles& InParticles, const int32 i) const
 	{
-		TArray<FSolverVec3> Grads;
-		Grads.SetNum(4);
-		const auto& Constraint = Constraints[i];
+		const TVec4<int32>& Constraint = Constraints[i];
 		const FSolverVec3& P1 = InParticles.P(Constraint[0]);
 		const FSolverVec3& P2 = InParticles.P(Constraint[1]);
 		const FSolverVec3& P3 = InParticles.P(Constraint[2]);
 		const FSolverVec3& P4 = InParticles.P(Constraint[3]);
-		const FSolverVec3 Edge = P2 - P1;
-		auto Normal1 = FSolverVec3::CrossProduct(P3 - P1, P3 - P2);
-		SafeDivide(Normal1, Normal1.SizeSquared());
-		auto Normal2 = FSolverVec3::CrossProduct(P4 - P2, P4 - P1);
-		SafeDivide(Normal2, Normal2.SizeSquared());
-		FSolverReal EdgeSize = Edge.Size();
-		Grads[0] = SafeDivide(FSolverVec3::DotProduct(Edge, P3 - P2), EdgeSize) * Normal1 + SafeDivide(FSolverVec3::DotProduct(Edge, P4 - P2), EdgeSize) * Normal2;
-		Grads[1] = SafeDivide(FSolverVec3::DotProduct(Edge, P1 - P3), EdgeSize) * Normal1 + SafeDivide(FSolverVec3::DotProduct(Edge, P1 - P4), EdgeSize) * Normal2;
-		Grads[2] = EdgeSize * Normal1;
-		Grads[3] = EdgeSize * Normal2;
-		return Grads;
+
+		return CalcGradients(P1, P2, P3, P4);
 	}
 
-	FSolverReal GetScalingFactor(const FSolverParticles& InParticles, const int32 i, const TArray<FSolverVec3>& Grads) const
+	FSolverReal GetScalingFactor(const FSolverParticles& InParticles, const int32 i, const TStaticArray<FSolverVec3, 4>& Grads) const
 	{
-		const auto& Constraint = Constraints[i];
+		const TVec4<int32>& Constraint = Constraints[i];
 		const int32 i1 = Constraint[0];
 		const int32 i2 = Constraint[1];
 		const int32 i3 = Constraint[2];
@@ -60,40 +48,69 @@ class FPBDBendingConstraintsBase
 		const FSolverVec3& P2 = InParticles.P(i2);
 		const FSolverVec3& P3 = InParticles.P(i3);
 		const FSolverVec3& P4 = InParticles.P(i4);
-		FSolverReal Angle = GetAngle(P1, P2, P3, P4);
-		FSolverReal Denom = (InParticles.InvM(i1) * Grads[0].SizeSquared() + InParticles.InvM(i2) * Grads[1].SizeSquared() + InParticles.InvM(i3) * Grads[2].SizeSquared() + InParticles.InvM(i4) * Grads[3].SizeSquared());
-		{
-			auto Edge = P2 - P1;
-			auto Normal1 = FSolverVec3::CrossProduct(P3 - P1, P3 - P2).GetSafeNormal();
-			auto Normal2 = FSolverVec3::CrossProduct(P4 - P2, P4 - P1).GetSafeNormal();
-			Denom = FSolverVec3::DotProduct(Edge, FSolverVec3::CrossProduct(Normal1, Normal2)) > (FSolverReal)0. ? -Denom : Denom;
-		}
-		FSolverReal Delta = Angle - MAngles[i];
-		return Stiffness * SafeDivide(Delta, Denom);
+		const FSolverReal Angle = CalcAngle(P1, P2, P3, P4);
+		const FSolverReal Denom = (InParticles.InvM(i1) * Grads[0].SizeSquared() + InParticles.InvM(i2) * Grads[1].SizeSquared() + InParticles.InvM(i3) * Grads[2].SizeSquared() + InParticles.InvM(i4) * Grads[3].SizeSquared());
+		constexpr FSolverReal SingleStepAngleLimit = (FSolverReal)(UE_PI * .25f); // this constraint is very non-linear. taking large steps is not accurate
+		const FSolverReal Delta = FMath::Clamp(Stiffness*(Angle - MAngles[i]), -SingleStepAngleLimit, SingleStepAngleLimit);
+		return SafeDivide(Delta, Denom);
 	}
 
 	void SetStiffness(FSolverReal InStiffness) { Stiffness = FMath::Clamp(InStiffness, (FSolverReal)0., (FSolverReal)1.); }
 
   private:
 	template<class TNum>
-	inline TNum SafeDivide(const TNum& Numerator, const FSolverReal& Denominator) const
+	static inline TNum SafeDivide(const TNum& Numerator, const FSolverReal& Denominator)
 	{
-		if (Denominator > (FSolverReal)1e-7)
+		if (Denominator > SMALL_NUMBER)
 			return Numerator / Denominator;
 		return TNum(0);
 	}
 
-	inline FSolverReal Clamp(const FSolverReal& Value, const FSolverReal& Low, const FSolverReal& High) const
+	static FSolverReal CalcAngle(const FSolverVec3& P1, const FSolverVec3& P2, const FSolverVec3& P3, const FSolverVec3& P4)
 	{
-		return Value < Low ? Low : (Value > High ? High : Value);
-	}
+		const FSolverVec3 Normal1 = FSolverVec3::CrossProduct(P1 - P3, P2 - P3).GetSafeNormal();
+		const FSolverVec3 Normal2 = FSolverVec3::CrossProduct(P2 - P4, P1 - P4).GetSafeNormal();
+
+		const FSolverVec3 SharedEdge = (P2 - P1).GetSafeNormal();
 
-	FSolverReal GetAngle(const FSolverVec3& P1, const FSolverVec3& P2, const FSolverVec3& P3, const FSolverVec3& P4) const
+		const FSolverReal CosPhi = FMath::Clamp(FSolverVec3::DotProduct(Normal1, Normal2), (FSolverReal)-1, (FSolverReal)1);
+		const FSolverReal SinPhi = FMath::Clamp(FSolverVec3::DotProduct(FSolverVec3::CrossProduct(Normal2, Normal1), SharedEdge), (FSolverReal)-1, (FSolverReal)1);
+		return FMath::Atan2(SinPhi, CosPhi);
+	}
+	
+	static TStaticArray<FSolverVec3, 4> CalcGradients(const FSolverVec3& P1, const FSolverVec3& P2, const FSolverVec3& P3, const FSolverVec3& P4)
 	{
-		auto Normal1 = FSolverVec3::CrossProduct(P3 - P1, P3 - P2).GetSafeNormal();
-		auto Normal2 = FSolverVec3::CrossProduct(P4 - P2, P4 - P1).GetSafeNormal();
-		auto Dot = FSolverVec3::DotProduct(Normal1, Normal2);
-		return FGenericPlatformMath::Acos(Clamp(Dot, 1e-4, 1 - 1e-4));
+		TStaticArray<FSolverVec3, 4> Grads;
+		// Calculated using Phi = atan2(SinPhi, CosPhi)
+		// where SinPhi = (Normal1 ^ Normal2)*SharedEdgeNormalized, CosPhi = Normal1 * Normal2
+		// Full gradients are calculated here, i.e., no simplifying assumptions around things like edge lengths being constant.
+		const FSolverVec3 SharedEdgeNormalized = (P2 - P1).GetSafeNormal();
+		const FSolverVec3 P13CrossP23 = FSolverVec3::CrossProduct(P1 - P3, P2 - P3);
+		const FSolverReal Normal1Len = P13CrossP23.Size();
+		const FSolverVec3 Normal1 = SafeDivide(P13CrossP23, Normal1Len);
+		const FSolverVec3 P24CrossP14 = FSolverVec3::CrossProduct(P2 - P4, P1 - P4);
+		const FSolverReal Normal2Len = P24CrossP14.Size();
+		const FSolverVec3 Normal2 = SafeDivide(P24CrossP14, Normal2Len);
+
+		const FSolverVec3 N2CrossN1 = FSolverVec3::CrossProduct(Normal2, Normal1);
+
+		const FSolverReal CosPhi = FMath::Clamp(FSolverVec3::DotProduct(Normal1, Normal2), (FSolverReal)-1, (FSolverReal)1);
+		const FSolverReal SinPhi = FMath::Clamp(FSolverVec3::DotProduct(N2CrossN1, SharedEdgeNormalized), (FSolverReal)-1, (FSolverReal)1);
+
+		const FSolverVec3 DPhiDN1_OverNormal1Len = SafeDivide(CosPhi * FSolverVec3::CrossProduct(SharedEdgeNormalized, Normal2) - SinPhi * Normal2, Normal1Len);
+		const FSolverVec3 DPhiDN2_OverNormal2Len = SafeDivide(CosPhi * FSolverVec3::CrossProduct(Normal1, SharedEdgeNormalized) - SinPhi * Normal1, Normal2Len);
+
+		const FSolverVec3 DPhiDP13 = FSolverVec3::CrossProduct(P2 - P3, DPhiDN1_OverNormal1Len);
+		const FSolverVec3 DPhiDP23 = FSolverVec3::CrossProduct(DPhiDN1_OverNormal1Len, P1 - P3);
+		const FSolverVec3 DPhiDP24 = FSolverVec3::CrossProduct(P1 - P4, DPhiDN2_OverNormal2Len);
+		const FSolverVec3 DPhiDP14 = FSolverVec3::CrossProduct(DPhiDN2_OverNormal2Len, P2 - P4);
+
+		Grads[0] = DPhiDP13 + DPhiDP14;
+		Grads[1] = DPhiDP23 + DPhiDP24;
+		Grads[2] = -DPhiDP13 - DPhiDP23;
+		Grads[3] = -DPhiDP14 - DPhiDP24;
+
+		return Grads;
 	}
 
   protected:
