Chaos - reset material each frame to undo potential modifications

diff --git a/Engine/Source/Runtime/Experimental/Chaos/Private/Chaos/Collision/ParticlePairMidPhase.cpp b/Engine/Source/Runtime/Experimental/Chaos/Private/Chaos/Collision/ParticlePairMidPhase.cpp
--- a/Engine/Source/Runtime/Experimental/Chaos/Private/Chaos/Collision/ParticlePairMidPhase.cpp
+++ b/Engine/Source/Runtime/Experimental/Chaos/Private/Chaos/Collision/ParticlePairMidPhase.cpp
@@ -281,6 +281,7 @@ namespace Chaos
 			// NOTE: these are not used by CCD which continuously moves the particles
 			Constraint->SetShapeWorldTransforms(ShapeWorldTransform0, ShapeWorldTransform1);
 
+			Constraint->ResetMaterial();
 			Constraint->SetCCDEnabled(false);
 
 			// If the constraint was not used last frame, it needs to be reset. 
@@ -357,6 +358,7 @@ namespace Chaos
 			// NOTE: these are not used by CCD which continuously moves the particles
 			Constraint->SetShapeWorldTransforms(ShapeWorldTransform0, ShapeWorldTransform1);
 
+			Constraint->ResetMaterial();
 			Constraint->SetCCDEnabled(true);
 			Constraint->SetCullDistance(CullDistance);
 			Constraint->ResetManifold();
@@ -562,6 +564,8 @@ namespace Chaos
 		const FRigidTransform3 ShapeWorldTransform1 = ShapeRelativeTransform1 * ParticleTransform1;
 		Constraint->SetShapeWorldTransforms(ShapeWorldTransform0, ShapeWorldTransform1);
 
+		Constraint->ResetMaterial();
+
 		NewConstraints.Add(Constraint);
 		return Constraint;
 	}
diff --git a/Engine/Source/Runtime/Experimental/Chaos/Private/Chaos/PBDCollisionConstraints.cpp b/Engine/Source/Runtime/Experimental/Chaos/Private/Chaos/PBDCollisionConstraints.cpp
--- a/Engine/Source/Runtime/Experimental/Chaos/Private/Chaos/PBDCollisionConstraints.cpp
+++ b/Engine/Source/Runtime/Experimental/Chaos/Private/Chaos/PBDCollisionConstraints.cpp
@@ -189,34 +189,34 @@ namespace Chaos
 		FPBDCollisionConstraintMaterial& CollisionMaterial = Constraint.Material;
 		if (PhysicsMaterial0 && PhysicsMaterial1)
 		{
-			const FChaosPhysicsMaterial::ECombineMode RestitutionCombineMode = FChaosPhysicsMaterial::ChooseCombineMode(PhysicsMaterial0->RestitutionCombineMode, PhysicsMaterial1->RestitutionCombineMode);
-			CollisionMaterial.Restitution = FChaosPhysicsMaterial::CombineHelper(PhysicsMaterial0->Restitution, PhysicsMaterial1->Restitution, RestitutionCombineMode);
+			const FChaosPhysicsMaterial::ECombineMode RestitutionCombineMode = FChaosPhysicsMaterial::ChooseCombineMode(PhysicsMaterial0->RestitutionCombineMode,PhysicsMaterial1->RestitutionCombineMode);
+			CollisionMaterial.MaterialRestitution = FChaosPhysicsMaterial::CombineHelper(PhysicsMaterial0->Restitution, PhysicsMaterial1->Restitution, RestitutionCombineMode);
 
-			const FChaosPhysicsMaterial::ECombineMode FrictionCombineMode = FChaosPhysicsMaterial::ChooseCombineMode(PhysicsMaterial0->FrictionCombineMode, PhysicsMaterial1->FrictionCombineMode);
-			CollisionMaterial.Friction = FChaosPhysicsMaterial::CombineHelper(PhysicsMaterial0->Friction, PhysicsMaterial1->Friction, FrictionCombineMode);
+			const FChaosPhysicsMaterial::ECombineMode FrictionCombineMode = FChaosPhysicsMaterial::ChooseCombineMode(PhysicsMaterial0->FrictionCombineMode,PhysicsMaterial1->FrictionCombineMode);
+			CollisionMaterial.MaterialDynamicFriction = FChaosPhysicsMaterial::CombineHelper(PhysicsMaterial0->Friction,PhysicsMaterial1->Friction, FrictionCombineMode);
 			const FReal StaticFriction0 = FMath::Max(PhysicsMaterial0->Friction, PhysicsMaterial0->StaticFriction);
 			const FReal StaticFriction1 = FMath::Max(PhysicsMaterial1->Friction, PhysicsMaterial1->StaticFriction);
-			CollisionMaterial.AngularFriction = FChaosPhysicsMaterial::CombineHelper(StaticFriction0, StaticFriction1, FrictionCombineMode);
+			CollisionMaterial.MaterialStaticFriction = FChaosPhysicsMaterial::CombineHelper(StaticFriction0, StaticFriction1, FrictionCombineMode);
 		}
 		else if (PhysicsMaterial0)
 		{
 			const FReal StaticFriction0 = FMath::Max(PhysicsMaterial0->Friction, PhysicsMaterial0->StaticFriction);
-			CollisionMaterial.Restitution = PhysicsMaterial0->Restitution;
-			CollisionMaterial.Friction = PhysicsMaterial0->Friction;
-			CollisionMaterial.AngularFriction = StaticFriction0;
+			CollisionMaterial.MaterialRestitution = PhysicsMaterial0->Restitution;
+			CollisionMaterial.MaterialDynamicFriction = PhysicsMaterial0->Friction;
+			CollisionMaterial.MaterialStaticFriction = StaticFriction0;
 		}
 		else if (PhysicsMaterial1)
 		{
 			const FReal StaticFriction1 = FMath::Max(PhysicsMaterial1->Friction, PhysicsMaterial1->StaticFriction);
-			CollisionMaterial.Restitution = PhysicsMaterial1->Restitution;
-			CollisionMaterial.Friction = PhysicsMaterial1->Friction;
-			CollisionMaterial.AngularFriction = StaticFriction1;
+			CollisionMaterial.MaterialRestitution = PhysicsMaterial1->Restitution;
+			CollisionMaterial.MaterialDynamicFriction = PhysicsMaterial1->Friction;
+			CollisionMaterial.MaterialStaticFriction = StaticFriction1;
 		}
 		else
 		{
-			CollisionMaterial.Friction = DefaultCollisionFriction;
-			CollisionMaterial.AngularFriction = DefaultCollisionFriction;
-			CollisionMaterial.Restitution = DefaultCollisionRestitution;
+			CollisionMaterial.MaterialDynamicFriction = DefaultCollisionFriction;
+			CollisionMaterial.MaterialStaticFriction = DefaultCollisionFriction;
+			CollisionMaterial.MaterialRestitution = DefaultCollisionRestitution;
 		}
 
 		CollisionMaterial.RestitutionThreshold = (CollisionRestitutionThresholdOverride >= 0.0f) ? CollisionRestitutionThresholdOverride : RestitutionThreshold;
@@ -224,21 +224,23 @@ namespace Chaos
 		// Overrides for testing
 		if (CollisionFrictionOverride >= 0)
 		{
-			CollisionMaterial.Friction = CollisionFrictionOverride;
-			CollisionMaterial.AngularFriction = CollisionFrictionOverride;
+			CollisionMaterial.MaterialDynamicFriction = CollisionFrictionOverride;
+			CollisionMaterial.MaterialStaticFriction = CollisionFrictionOverride;
 		}
 		if (CollisionRestitutionOverride >= 0)
 		{
-			CollisionMaterial.Restitution = CollisionRestitutionOverride;
+			CollisionMaterial.MaterialRestitution = CollisionRestitutionOverride;
 		}
 		if (CollisionAngularFrictionOverride >= 0)
 		{
-			CollisionMaterial.AngularFriction = CollisionAngularFrictionOverride;
+			CollisionMaterial.MaterialStaticFriction = CollisionAngularFrictionOverride;
 		}
 		if (!bEnableRestitution)
 		{
-			CollisionMaterial.Restitution = 0.0f;
+			CollisionMaterial.MaterialRestitution = 0.0f;
 		}
+		
+		CollisionMaterial.ResetMaterialModifications();
 	}
 
 	void FPBDCollisionConstraints::UpdatePositionBasedState(const FReal Dt)
diff --git a/Engine/Source/Runtime/Experimental/Chaos/Public/Chaos/Collision/PBDCollisionConstraint.h b/Engine/Source/Runtime/Experimental/Chaos/Public/Chaos/Collision/PBDCollisionConstraint.h
--- a/Engine/Source/Runtime/Experimental/Chaos/Public/Chaos/Collision/PBDCollisionConstraint.h
+++ b/Engine/Source/Runtime/Experimental/Chaos/Public/Chaos/Collision/PBDCollisionConstraint.h
@@ -104,7 +104,10 @@ namespace Chaos
 	{
 	public:
 		FPBDCollisionConstraintMaterial()
-			: Friction(0)
+			: MaterialStaticFriction(0)
+			, MaterialDynamicFriction(0)
+			, MaterialRestitution(0)
+			, Friction(0)
 			, AngularFriction(0)
 			, Restitution(0)
 			, RestitutionPadding(0)
@@ -116,9 +119,16 @@ namespace Chaos
 		{
 		}
 
+		// Material properties pulled from the materials of the two shapes involved in the contact
+		FReal MaterialStaticFriction;
+		FReal MaterialDynamicFriction;
+		FReal MaterialRestitution;
+
+		// Final material properties (post modifier) used by the solver. These get reset every frame to the material values above
 		FReal Friction;			// @todo(chaos): rename DynamicFriction
 		FReal AngularFriction;	// @todo(chaos): rename StaticFriction
 		FReal Restitution;
+
 		FReal RestitutionPadding; // For StandardPBD implementation of resitution, we pad constraints on initial contact to enforce outward velocity
 		FReal RestitutionThreshold;
 		FReal InvMassScale0;
@@ -126,6 +136,15 @@ namespace Chaos
 		FReal InvInertiaScale0;
 		FReal InvInertiaScale1;
 
+		// Reset the material properties to those pulled from the shape's materials (i.e., back to the state before any contact modification)
+		void ResetMaterialModifications()
+		{
+			Friction = MaterialDynamicFriction;
+			AngularFriction = MaterialStaticFriction;
+			Restitution = MaterialRestitution;
+		}
+
+		// @todo(chaos): remove this
 		void Reset()
 		{
 			RestitutionPadding = 0;
@@ -393,6 +412,14 @@ namespace Chaos
 		// NOTE: This is initially set based on whether are allowing incremental manifolds
 		bool GetUseIncrementalCollisionDetection() const { return !Flags.bUseManifold || Flags.bUseIncrementalManifold; }
 
+		/**
+		* Reset the material properties to those from the shape materials. Called each frame to reset contact modifications to the material.
+		*/
+		inline void ResetMaterial()
+		{
+			Manifold.ResetMaterialModifications();
+		}
+
 		/**
 		 * @brief Clear the current and previous manifolds
 		*/
