Fix CIS post merge

diff --git a/Engine/Source/Runtime/Experimental/Chaos/Public/Chaos/Collision/PBDCollisionConstraint.h b/Engine/Source/Runtime/Experimental/Chaos/Public/Chaos/Collision/PBDCollisionConstraint.h
--- a/Engine/Source/Runtime/Experimental/Chaos/Public/Chaos/Collision/PBDCollisionConstraint.h
+++ b/Engine/Source/Runtime/Experimental/Chaos/Public/Chaos/Collision/PBDCollisionConstraint.h
@@ -417,7 +417,7 @@ namespace Chaos
 		*/
 		inline void ResetMaterial()
 		{
-			Manifold.ResetMaterialModifications();
+			Material.ResetMaterialModifications();
 		}
 
 		/**
