Chaos : Destruction : cleanup mass property and radii computation for GCs
- move Inner/outer radii computation outsie of the init function
- mass sproperties are now computed into 2 steps, first Volume/CoM then INertia without having each stepping on each other
- extracted mass property computation in external function for clarity
- moved some of the bounding box mass property computattion into Chaos MassProperties files

diff --git a/Engine/Source/Runtime/Experimental/Chaos/Private/Chaos/MassProperties.cpp b/Engine/Source/Runtime/Experimental/Chaos/Private/Chaos/MassProperties.cpp
--- a/Engine/Source/Runtime/Experimental/Chaos/Private/Chaos/MassProperties.cpp
+++ b/Engine/Source/Runtime/Experimental/Chaos/Private/Chaos/MassProperties.cpp
@@ -90,6 +90,35 @@ namespace Chaos
 		return FinalRotation;
 	}
 
+	void CalculateVolumeAndCenterOfMass(const FBox& BoundingBox, FVector::FReal& OutVolume, FVector& OutCenterOfMass)
+	{
+		const FVector Extents = static_cast<FVector::FReal>(2) * BoundingBox.GetExtent(); // FBox::GetExtent() returns half the size, but FAABB::Extents() returns total size
+		OutVolume = Extents.X * Extents.Y * Extents.Z;
+		OutCenterOfMass = BoundingBox.GetCenter();
+	}
+
+	void CalculateInertiaAndRotationOfMass(const FBox& BoundingBox, const FVector::FReal Density, FMatrix33& OutInertiaTensor, FRotation3& OutRotationOfMass)
+	{
+		const FVector Extents = static_cast<FVector::FReal>(2) * BoundingBox.GetExtent(); // FBox::GetExtent() returns half the size, but FAABB::Extents() returns total size
+		const FVector::FReal Volume = Extents.X * Extents.Y * Extents.Z;
+		const FVector::FReal Mass = Volume * Density;
+		const FVector::FReal ExtentsYZ = Extents.Y * Extents.Y + Extents.Z * Extents.Z;
+		const FVector::FReal ExtentsXZ = Extents.X * Extents.X + Extents.Z * Extents.Z;
+		const FVector::FReal ExtentsXY = Extents.X * Extents.X + Extents.Y * Extents.Y;
+		OutInertiaTensor = FMatrix33((Mass * ExtentsYZ) / 12., (Mass * ExtentsXZ) / 12., (Mass * ExtentsXY) / 12.);
+		OutRotationOfMass = FRotation3::Identity;
+	}
+	
+   	FMassProperties CalculateMassProperties(const FBox& BoundingBox, const FVector::FReal Density)
+	{
+    	FMassProperties MassProperties;
+    	CalculateVolumeAndCenterOfMass(BoundingBox, MassProperties.Volume, MassProperties.CenterOfMass);
+    	check(Density > 0);
+		MassProperties.Mass = MassProperties.Volume * Density;
+		CalculateInertiaAndRotationOfMass(BoundingBox, Density, MassProperties.InertiaTensor, MassProperties.RotationOfMass);
+   		return MassProperties;
+   	}
+	
 	template<typename T, typename TSurfaces>
 	void CHAOS_API CalculateVolumeAndCenterOfMass(const TParticles<T, 3>& Vertices, const TSurfaces& Surfaces, T& OutVolume, TVec3<T>& OutCenterOfMass)
 	{
diff --git a/Engine/Source/Runtime/Experimental/Chaos/Private/GeometryCollection/GeometryCollectionUtility.cpp b/Engine/Source/Runtime/Experimental/Chaos/Private/GeometryCollection/GeometryCollectionUtility.cpp
--- a/Engine/Source/Runtime/Experimental/Chaos/Private/GeometryCollection/GeometryCollectionUtility.cpp
+++ b/Engine/Source/Runtime/Experimental/Chaos/Private/GeometryCollection/GeometryCollectionUtility.cpp
@@ -247,6 +247,34 @@ namespace GeometryCollection
 		(Collection->Parent)[TransformIndex4] = TransformIndex3;
 	}
 
+	void ComputeInnerAndOuterRadiiFromGeometryVertices(const TManagedArray<FVector3f>& Vertices, const int32 VertexStart, const int32 VertexCount, float& OutInnerRadius, float& OutOuterRadius)
+	{
+		// first compute the geometry center
+		FVector VertexSum{0.0 };
+		for (int32 VertexIndex = 0; VertexIndex < VertexCount; ++VertexIndex)
+		{
+			VertexSum += FVector(Vertices[VertexStart + VertexIndex]); 
+		}
+		FVector3f Center{0.};
+		if (VertexCount > 0)
+		{
+			// controlled precision loss, we are now back into the range of vertices
+			Center = FVector3f{VertexSum / static_cast<FVector::FReal>(VertexCount)};
+		}
+		
+		OutInnerRadius = VertexCount ? TNumericLimits<float>::Max() : 0.0f;
+		OutOuterRadius = 0.0f;
+		for (int32 VertexIndex = 0; VertexIndex < VertexCount; ++VertexIndex)
+		{
+			const FVector3f& Pt = Vertices[VertexStart + VertexIndex];
+			const float DistSq = FVector3f::DistSquared(Pt, Center);
+			OutInnerRadius = FMath::Min(OutInnerRadius, DistSq);
+			OutOuterRadius = FMath::Max(OutOuterRadius, DistSq);
+		}
+		OutInnerRadius = FMath::Sqrt(OutInnerRadius);
+		OutOuterRadius = FMath::Sqrt(OutOuterRadius);
+	}
+	
 	void AddGeometryProperties(FGeometryCollection * Collection)
 	{
 		if (Collection)
diff --git a/Engine/Source/Runtime/Experimental/Chaos/Private/PhysicsProxy/GeometryCollectionPhysicsProxy.cpp b/Engine/Source/Runtime/Experimental/Chaos/Private/PhysicsProxy/GeometryCollectionPhysicsProxy.cpp
--- a/Engine/Source/Runtime/Experimental/Chaos/Private/PhysicsProxy/GeometryCollectionPhysicsProxy.cpp
+++ b/Engine/Source/Runtime/Experimental/Chaos/Private/PhysicsProxy/GeometryCollectionPhysicsProxy.cpp
@@ -1027,7 +1027,7 @@ void FGeometryCollectionPhysicsProxy::InitializeBodiesPT(Chaos::FPBDRigidsSolver
 				//   leaf nodes like kinematic or dynamic clusters. 
 				if (!Handle->Disabled() && Parameters.ObjectType == EObjectStateTypeEnum::Chaos_Object_Sleeping)
 				{
-					RigidsSolver->GetEvolution()->SetParticleObjectState(Handle, Chaos::EObjectStateType::Sleeping);
+				 	RigidsSolver->GetEvolution()->SetParticleObjectState(Handle, Chaos::EObjectStateType::Sleeping);
 				}
 
 				RigidsSolver->GetEvolution()->DirtyParticle(*Handle);
@@ -2043,7 +2043,20 @@ static const FSharedSimulationSizeSpecificData& GetSizeSpecificData(const TArray
 	return SizeSpecificData[SizeSpecificIdx];
 }
 
-FGeometryDynamicCollection::FSharedImplicit CreateImplicitGeometry(
+static void GenerateInnerAndOuterRadiiIfNeeded(FGeometryCollection& RestCollection, const int32 GeometryIndex)
+{
+	if (RestCollection.InnerRadius[GeometryIndex] == 0.0f || RestCollection.OuterRadius[GeometryIndex] == 0.0f)
+	{
+		GeometryCollection::ComputeInnerAndOuterRadiiFromGeometryVertices(
+			RestCollection.Vertex,
+			RestCollection.VertexStart[GeometryIndex],
+			RestCollection.VertexCount[GeometryIndex],
+			RestCollection.InnerRadius[GeometryIndex],
+			RestCollection.OuterRadius[GeometryIndex]);
+	}	
+}
+
+static FGeometryDynamicCollection::FSharedImplicit CreateImplicitGeometry(
 	const FSharedSimulationSizeSpecificData& SizeSpecificData,
 	const int32 TransformGroupIndex,
 	const FGeometryCollection& RestCollection,
@@ -2172,6 +2185,49 @@ FGeometryDynamicCollection::FSharedImplicit CreateImplicitGeometry(
 	return NewImplicit;
 }
 
+static void ComputeGeometryVolumeAndCenterOfMass(
+	const Chaos::FParticles& MassSpaceParticles,
+	const Chaos::FTriangleMesh& TriMesh,
+	const FBox& BoundingBox,
+	const FSharedSimulationParameters& SharedParams, 
+	Chaos::FReal& OutUncorrectedVolume,
+	Chaos::FReal& OutVolume,
+	Chaos::FVec3& OutCenterOfMass,
+	bool& OutIsTooSmallGeometry
+	)
+{
+	const float MinVolume = FMath::Max(UE_SMALL_NUMBER, SharedParams.MinimumVolumeClamp());
+	const float MaxVolume = FMath::Max(MinVolume, SharedParams.MaximumVolumeClamp());
+
+	OutUncorrectedVolume = 0.f;
+	OutVolume = 0.f;
+    OutCenterOfMass = FVector::ZeroVector;
+	OutIsTooSmallGeometry = (BoundingBox.GetExtent().GetAbsMin() < MinVolume);
+	
+	if (!OutIsTooSmallGeometry)
+	{
+		Chaos::CalculateVolumeAndCenterOfMass(MassSpaceParticles, TriMesh.GetElements(), OutUncorrectedVolume, OutCenterOfMass);
+		OutVolume = OutUncorrectedVolume;
+		
+		if (OutUncorrectedVolume == 0)
+		{
+			Chaos::CalculateVolumeAndCenterOfMass(BoundingBox, OutVolume, OutCenterOfMass);
+		}
+		if (OutUncorrectedVolume < MinVolume)
+		{
+			// For rigid bodies outside of range just default to a clamped bounding box, and warn the user.
+			OutVolume = MinVolume;
+			OutCenterOfMass = BoundingBox.GetCenter();
+		}
+		else if (MaxVolume < OutUncorrectedVolume)
+		{
+			// For rigid bodies outside of range just default to a clamped bounding box, and warn the user
+			OutVolume = MaxVolume;
+			OutCenterOfMass = BoundingBox.GetCenter();
+		}
+	}
+}
+
 /** 
 	NOTE - Making any changes to data stored on the rest collection below MUST be accompanied
 	by a rotation of the DDC key in FDerivedDataGeometryCollectionCooker::GetVersionString
@@ -2191,15 +2247,6 @@ void FGeometryCollectionPhysicsProxy::InitializeSharedCollisionStructures(
 		return;
 	}
 
-	// clamps
-	const float MinBoundsExtents = SharedParams.MinimumBoundingExtentClamp;
-	const float MaxBoundsExtents = SharedParams.MaximumBoundingExtentClamp;
-	const float MinVolume = SharedParams.MinimumVolumeClamp();
-	const float MaxVolume = SharedParams.MaximumVolumeClamp();
-	const float MinMass = FMath::Max(UE_SMALL_NUMBER, SharedParams.MaximumMassClamp);
-	const float MaxMass = SharedParams.MinimumMassClamp;
-
-
 	//TArray<TArray<TArray<int32>>> BoundaryVertexIndices;
 	//GeometryCollectionAlgo::FindOpenBoundaries(&RestCollection, 1e-2, BoundaryVertexIndices);
 	//GeometryCollectionAlgo::TriangulateBoundaries(&RestCollection, BoundaryVertexIndices);
@@ -2289,8 +2336,8 @@ void FGeometryCollectionPhysicsProxy::InitializeSharedCollisionStructures(
 	TArray<FMassProperties> MassPropertiesArray;
 	MassPropertiesArray.AddUninitialized(NumGeometries);
 
-	TArray<bool> InertiaComputationNeeded;
-	InertiaComputationNeeded.Init(false, NumGeometries);
+	TArray<FVector::FReal> UncorrectedVolume;
+	UncorrectedVolume.Init(false, NumGeometries);
 
 	// We skip very small geometry and log as a warning. To avoid log spamming, we wait
 	// until we complete the loop before reporting the skips.
@@ -2304,9 +2351,15 @@ void FGeometryCollectionPhysicsProxy::InitializeSharedCollisionStructures(
 	// Otherwise, iterating over the geometry group is a convenient way to iterate
 	// over all the leaves of the hierarchy.
 	check(!TransformIndex.Contains(INDEX_NONE)); // TODO: implement support for instanced bodies
+	// @todo(chaos) could this be in a parallel for and having the total volume computed in a tight loop right after? 
 	for (int32 GeometryIndex = 0; GeometryIndex < NumGeometries; GeometryIndex++)
 	{
 		const int32 TransformGroupIndex = TransformIndex[GeometryIndex];
+		
+		// @chaos(todo): is that really needed? 
+		GenerateInnerAndOuterRadiiIfNeeded(RestCollection, GeometryIndex);
+
+		CollectionSimulatableParticles[TransformGroupIndex] = false;
 		if (SimulationType[TransformGroupIndex] > FGeometryCollection::ESimulationTypes::FST_None)
 		{
 			TUniquePtr<FTriangleMesh> TriMesh(
@@ -2316,59 +2369,31 @@ void FGeometryCollectionPhysicsProxy::InitializeSharedCollisionStructures(
 					Visible,
 					Indices));
 
+			// Compute volume and center of mass 
 			FMassProperties& MassProperties = MassPropertiesArray[GeometryIndex];
-
 			{
-				MassProperties.CenterOfMass = FVector3f::ZeroVector;
-				MassProperties.RotationOfMass = FRotation3(FQuat::Identity).GetNormalized();
-				MassProperties.Volume = 0.f;
-				MassProperties.InertiaTensor = FMatrix33(1,1,1);
-				MassProperties.Mass = 1.0f; // start with unit mass, scaled later by density
-
-				if (BoundingBox[GeometryIndex].GetExtent().GetAbsMin() < MinVolume)
+				bool bIsTooSmallGeometry = false;
+				ComputeGeometryVolumeAndCenterOfMass(
+					MassSpaceParticles,
+					*TriMesh,
+					BoundingBox[GeometryIndex],
+					SharedParams,
+					UncorrectedVolume[GeometryIndex],
+					MassProperties.Volume,
+					MassProperties.CenterOfMass,
+					bIsTooSmallGeometry
+					);
+				
+				if (bIsTooSmallGeometry)
 				{
 					bSkippedSmallGeometry = true;
-					CollectionSimulatableParticles[TransformGroupIndex] = false;	//do not simulate tiny particles
-					MassProperties.Mass = 0.f;
-					MassProperties.InertiaTensor = FMatrix33(0,0,0);
 				}
 				else
 				{
-					CalculateVolumeAndCenterOfMass(MassSpaceParticles, TriMesh->GetElements(), MassProperties.Volume, MassProperties.CenterOfMass);
-					InertiaComputationNeeded[GeometryIndex] = true;
-					if(MassProperties.Volume == 0)
-					{
-						FVector Extents = FReal(2) * BoundingBox[GeometryIndex].GetExtent(); // FBox::GetExtent() returns half the size, but FAABB::Extents() returns total size
-						MassProperties.Volume = Extents.X * Extents.Y * Extents.Z;
-						FReal ExtentsYZ = Extents.Y * Extents.Y + Extents.Z * Extents.Z;
-						FReal ExtentsXZ = Extents.X * Extents.X + Extents.Z * Extents.Z;
-						FReal ExtentsXY = Extents.X * Extents.X + Extents.Y * Extents.Y;
-						MassProperties.InertiaTensor = PMatrix<FReal, 3, 3>(ExtentsYZ / 12.f, ExtentsXZ / 12.f, ExtentsXY / 12.f);
-						MassProperties.CenterOfMass = BoundingBox[GeometryIndex].GetCenter();
-						CollectionMassToLocal[TransformGroupIndex] = FTransform(FQuat::Identity, MassProperties.CenterOfMass);
-						InertiaComputationNeeded[GeometryIndex] = false;
-					}
-
-					if (MassProperties.Volume < MinVolume)
-					{
-						// For rigid bodies outside of range just defaut to a clamped bounding box, and warn the user.
-						MassProperties.Volume = MinVolume;
-						CollectionMassToLocal[TransformGroupIndex] = FTransform(FQuat::Identity, BoundingBox[GeometryIndex].GetCenter());
-						InertiaComputationNeeded[GeometryIndex] = false;
-					}
-					else if (MaxVolume < MassProperties.Volume)
-					{
-						// For rigid bodies outside of range just defaut to a clamped bounding box, and warn the user
-						MassProperties.Volume = MaxVolume;
-						CollectionMassToLocal[TransformGroupIndex] = FTransform(FQuat::Identity, BoundingBox[GeometryIndex].GetCenter());
-						InertiaComputationNeeded[GeometryIndex] = false;
-					}
-					else
-					{
-						CollectionMassToLocal[TransformGroupIndex] = FTransform(FQuat::Identity, MassProperties.CenterOfMass);
-					}
-
-					FVector MassTranslation = CollectionMassToLocal[TransformGroupIndex].GetTranslation();
+					CollectionSimulatableParticles[TransformGroupIndex] = true;
+				
+					// shift the particles to the mass local space
+					const FVector MassTranslation = MassProperties.CenterOfMass;
 					if (!FMath::IsNearlyZero(MassTranslation.SizeSquared()))
 					{
 						const int32 IdxStart = VertexStart[GeometryIndex];
@@ -2379,55 +2404,26 @@ void FGeometryCollectionPhysicsProxy::InitializeSharedCollisionStructures(
 						}
 					}
 				}
-			}
-
-			if (InnerRadius[GeometryIndex] == 0.0f || OuterRadius[GeometryIndex] == 0.0f)
-			{
-				const int32 VCount = VertexCount[GeometryIndex];
-				if (VCount != 0)
-				{
-					const FVector3f Center = (FVector3f)BoundingBox[GeometryIndex].GetCenter();
-					const int32 VStart = VertexStart[GeometryIndex];
 
-					InnerRadius[GeometryIndex] = VCount ? TNumericLimits<FRealSingle>::Max() : 0.0f;
-					OuterRadius[GeometryIndex] = 0.0f;
-					for (int32 VIdx = 0; VIdx < VCount; ++VIdx)
-					{
-						const int32 PtIdx = VStart + VIdx;
-						const FVector3f& Pt = Vertex[PtIdx];
-						const float DistSq = FVector3f::DistSquared(Pt, Center);
-						if (InnerRadius[GeometryIndex] > DistSq)
-						{
-							InnerRadius[GeometryIndex] = DistSq;
-						}
-						if (OuterRadius[GeometryIndex] < DistSq)
-						{
-							OuterRadius[GeometryIndex] = DistSq;
-						}
-					}
-					InnerRadius[GeometryIndex] = FMath::Sqrt(InnerRadius[GeometryIndex]);
-					OuterRadius[GeometryIndex] = FMath::Sqrt(OuterRadius[GeometryIndex]);
-				}
+				// update Mass to local matrix
+				CollectionMassToLocal[TransformGroupIndex] = FTransform(FQuat::Identity, MassProperties.CenterOfMass);
 			}
 
 			TotalVolume += MassProperties.Volume;
 			TriangleMeshesArray[TransformGroupIndex] = MoveTemp(TriMesh);
 		}
-		else
-		{
-			CollectionSimulatableParticles[TransformGroupIndex] = false;
-		}
 	}
 
+	// log any warning necessary
 	if (bSkippedSmallGeometry)
 	{
 		UE_LOG(LogChaos, Warning, TEXT("Some geometry is too small to be simulated and has been skipped."));
 	}
 
-	//User provides us with total mass or density.
-	//Density must be the same for individual parts and the total. Density_i = Density = Mass_i / Volume_i
-	//Total mass must equal sum of individual parts. Mass_i = TotalMass * Volume_i / TotalVolume => Density_i = TotalMass / TotalVolume
-	TotalVolume = FMath::Max(TotalVolume, MinBoundsExtents * MinBoundsExtents * MinBoundsExtents);
+	// User provides us with total mass or density.
+	// Density must be the same for individual parts and the total. Density_i = Density = Mass_i / Volume_i
+	// Total mass must equal sum of individual parts. Mass_i = TotalMass * Volume_i / TotalVolume => Density_i = TotalMass / TotalVolume
+	TotalVolume = FMath::Max(TotalVolume, SharedParams.MinimumVolumeClamp());
 	const FReal DesiredTotalMass = SharedParams.bMassAsDensity ? SharedParams.Mass * TotalVolume : SharedParams.Mass;
 	const FReal ClampedTotalMass = FMath::Clamp(DesiredTotalMass, SharedParams.MinimumMassClamp, SharedParams.MaximumMassClamp);
 	const FReal DesiredDensity = ClampedTotalMass / TotalVolume;
@@ -2443,7 +2439,7 @@ void FGeometryCollectionPhysicsProxy::InitializeSharedCollisionStructures(
 		const FReal Volume_i = MassPropertiesArray[GeometryIndex].Volume;
 		if (CollectionSimulatableParticles[TransformGroupIndex])
 		{
-			//Must clamp each individual mass regardless of desired density
+			// Must clamp each individual mass regardless of desired density
 			if (DesiredDensity * Volume_i > SharedParams.MaximumMassClamp)
 			{
 				// For rigid bodies outside of range just defaut to a clamped bounding box, and warn the user.
@@ -2454,6 +2450,7 @@ void FGeometryCollectionPhysicsProxy::InitializeSharedCollisionStructures(
 			}
 		}
 
+		// compute the rest of the mass property ( Mass, InertiaTensor and RotationOfMass )
 		if (CollectionSimulatableParticles[TransformGroupIndex])
 		{
 			TUniquePtr<FTriangleMesh>& TriMesh = TriangleMeshesArray[TransformGroupIndex];
@@ -2463,31 +2460,40 @@ void FGeometryCollectionPhysicsProxy::InitializeSharedCollisionStructures(
 			const FReal Density_i = Mass_i / Volume_i;
 			CollectionMass[TransformGroupIndex] = (FRealSingle)Mass_i;
 
-			if (InertiaComputationNeeded[GeometryIndex])
+			const bool bUseBoundingBoxInertia = (UncorrectedVolume[GeometryIndex] == 0);  
+			if (bUseBoundingBoxInertia)
+			{
+				CalculateInertiaAndRotationOfMass(BoundingBox[GeometryIndex], Density_i, MassProperties.InertiaTensor, MassProperties.RotationOfMass);
+			}
+			else
 			{
 				// Note: particles already in CoM space, so passing in zero as CoM
 				CalculateInertiaAndRotationOfMass(MassSpaceParticles, TriMesh->GetSurfaceElements(), Density_i, FVec3(0), MassProperties.InertiaTensor, MassProperties.RotationOfMass);
-				CollectionInertiaTensor[TransformGroupIndex] = FVector3f((float)MassProperties.InertiaTensor.M[0][0], (float)MassProperties.InertiaTensor.M[1][1], (float)MassProperties.InertiaTensor.M[2][2]);
-				CollectionMassToLocal[TransformGroupIndex] = FTransform(MassProperties.RotationOfMass, MassProperties.CenterOfMass);
-
-
-				if (!MassProperties.RotationOfMass.Equals(FQuat::Identity))
+				// scale the inertia tensor accordingly
+				if (MassProperties.Volume != UncorrectedVolume[GeometryIndex])
 				{
-					FTransform InverseMassRotation = FTransform(MassProperties.RotationOfMass.Inverse());
-					const int32 IdxStart = VertexStart[GeometryIndex];
-					const int32 IdxEnd = IdxStart + VertexCount[GeometryIndex];
-					for (int32 Idx = IdxStart; Idx < IdxEnd; ++Idx)
-					{
-						MassSpaceParticles.X(Idx) = InverseMassRotation.TransformPosition(MassSpaceParticles.X(Idx));
-					}
+					const FVector::FReal InertiaScale = MassProperties.Volume / UncorrectedVolume[GeometryIndex];
+					MassProperties.InertiaTensor = MassProperties.InertiaTensor.ApplyScale(InertiaScale);
 				}
 			}
-			else
+			
+			// Orient particle in the local mass space 
+			if (!MassProperties.RotationOfMass.Equals(FQuat::Identity))
 			{
-				const FVec3 DiagonalInertia(MassProperties.InertiaTensor.M[0][0], MassProperties.InertiaTensor.M[1][1], MassProperties.InertiaTensor.M[2][2]);
-				CollectionInertiaTensor[TransformGroupIndex] = FVector3f(DiagonalInertia * Mass_i);
+				FTransform InverseMassRotation = FTransform(MassProperties.RotationOfMass.Inverse());
+				const int32 IdxStart = VertexStart[GeometryIndex];
+				const int32 IdxEnd = IdxStart + VertexCount[GeometryIndex];
+				for (int32 Idx = IdxStart; Idx < IdxEnd; ++Idx)
+				{
+					MassSpaceParticles.X(Idx) = InverseMassRotation.TransformPosition(MassSpaceParticles.X(Idx));
+				}
 			}
 
+			// update Rest collection properties
+			CollectionInertiaTensor[TransformGroupIndex] = FVector3f((float)MassProperties.InertiaTensor.M[0][0], (float)MassProperties.InertiaTensor.M[1][1], (float)MassProperties.InertiaTensor.M[2][2]);
+			CollectionMassToLocal[TransformGroupIndex] = FTransform(MassProperties.RotationOfMass, MassProperties.CenterOfMass);
+
+
 			FBox InstanceBoundingBox(EForceInit::ForceInitToZero);
 			if (TriMesh->GetElements().Num())
 			{
diff --git a/Engine/Source/Runtime/Experimental/Chaos/Public/Chaos/MassProperties.h b/Engine/Source/Runtime/Experimental/Chaos/Public/Chaos/MassProperties.h
--- a/Engine/Source/Runtime/Experimental/Chaos/Public/Chaos/MassProperties.h
+++ b/Engine/Source/Runtime/Experimental/Chaos/Public/Chaos/MassProperties.h
@@ -48,6 +48,8 @@ namespace Chaos
 	void CHAOS_API CalculateInertiaAndRotationOfMass(const FParticles& Vertices, const TSurfaces& Surfaces, const FReal Density, const FVec3& CenterOfMass,
 		FMatrix33& OutInertiaTensor, FRotation3& OutRotationOfMass);
 
+	void CHAOS_API CalculateVolumeAndCenterOfMass(const FBox& BoundingBox, FVector::FReal& OutVolume, FVector& OutCenterOfMass);
+	void CHAOS_API CalculateInertiaAndRotationOfMass(const FBox& BoundingBox, const FVector::FReal Density, FMatrix33& OutInertiaTensor, FRotation3& OutRotationOfMass);
 
 	// Combine a list of transformed inertia tensors into a single inertia. Also diagonalize the inertia
 	// and set the rotation of mass accordingly.
diff --git a/Engine/Source/Runtime/Experimental/Chaos/Public/GeometryCollection/GeometryCollectionUtility.h b/Engine/Source/Runtime/Experimental/Chaos/Public/GeometryCollection/GeometryCollectionUtility.h
--- a/Engine/Source/Runtime/Experimental/Chaos/Public/GeometryCollection/GeometryCollectionUtility.h
+++ b/Engine/Source/Runtime/Experimental/Chaos/Public/GeometryCollection/GeometryCollectionUtility.h
@@ -165,6 +165,19 @@ namespace GeometryCollection
 	CHAOS_API
 	GenerateTemporaryGuids(FTransformCollection* Collection, int32 StartIdx = 0, bool bForceInit=false);
 
+	/***
+	* Compute inner and outer radius from a set of vertices 
+	* @param Vertices		array containing the vertices used to computer the radii
+	* @param VertexStart	start index of the range of vertex to use within the array
+	* @param VertexCount	Number of vertices to use from VertexStart
+	* @param VertexCount	Total number of vertices to use (from VertexStart)
+	* @param OutInnerRadius	Computed InnerRadius ( existing value will be overriden )
+	* @param OutOuterRadius	Computed outerRadius ( existing value will be overriden )
+	*/
+	void 
+	CHAOS_API
+	ComputeInnerAndOuterRadiiFromGeometryVertices(const TManagedArray<FVector3f>& Vertices, const int32 VertexStart, const int32 VertexCount, float& OutInnerRadius, float& OutOuterRadius);
+
 };
 
 // AttributeTransfer implementation
