Fix geometry collection regression

diff --git a/Engine/Source/Runtime/Experimental/GeometryCollectionEngine/Private/GeometryCollection/GeometryCollectionSceneProxy.cpp b/Engine/Source/Runtime/Experimental/GeometryCollectionEngine/Private/GeometryCollection/GeometryCollectionSceneProxy.cpp
--- a/Engine/Source/Runtime/Experimental/GeometryCollectionEngine/Private/GeometryCollection/GeometryCollectionSceneProxy.cpp
+++ b/Engine/Source/Runtime/Experimental/GeometryCollectionEngine/Private/GeometryCollection/GeometryCollectionSceneProxy.cpp
@@ -1031,7 +1031,11 @@ void FGeometryCollectionSceneProxy::GetDynamicRayTracingInstances(FRayTracingMat
 
 			FRayTracingInstance RayTracingInstance;
 			RayTracingInstance.Geometry = &RayTracingGeometry;
-			RayTracingInstance.InstanceTransforms.Append(DynamicData->Transforms);
+			RayTracingInstance.InstanceTransforms.Reserve(DynamicData->Transforms.Num());
+			for (int32 TransformIndex = 0; TransformIndex < DynamicData->Transforms.Num(); ++TransformIndex)
+			{
+				RayTracingInstance.InstanceTransforms.Emplace(GetLocalToWorld());
+			}
 
 			// Grab the material proxies we'll be using for each section
 			TArray<FMaterialRenderProxy*, TInlineAllocator<32>> MaterialProxies;
@@ -1039,7 +1043,15 @@ void FGeometryCollectionSceneProxy::GetDynamicRayTracingInstances(FRayTracingMat
 			for (int32 SectionIndex = 0; SectionIndex < SectionArray.Num(); ++SectionIndex)
 			{
 				const FGeometryCollectionSection& Section = SectionArray[SectionIndex];
-				FMaterialRenderProxy* MaterialProxy = GetMaterial(Context.RayTracingMeshResourceCollector, Section.MaterialID);
+				
+				//TODO: Add BoneColor support in Path/Ray tracing?
+				FMaterialRenderProxy* MaterialProxy= Materials[Section.MaterialID]->GetRenderProxy();
+
+				if (MaterialProxy == nullptr)
+				{
+					MaterialProxy = UMaterial::GetDefaultMaterial(MD_Surface)->GetRenderProxy();
+				}
+
 				MaterialProxies.Add(MaterialProxy);
 			}
 
