Raycast Height Field optimization

diff --git a/Engine/Source/Runtime/Experimental/Chaos/Private/Chaos/HeightField.cpp b/Engine/Source/Runtime/Experimental/Chaos/Private/Chaos/HeightField.cpp
--- a/Engine/Source/Runtime/Experimental/Chaos/Private/Chaos/HeightField.cpp
+++ b/Engine/Source/Runtime/Experimental/Chaos/Private/Chaos/HeightField.cpp
@@ -45,6 +45,12 @@ namespace Chaos
 				bParallel[Axis] = FMath::IsNearlyZero(Dir[Axis], (FReal)1.e-8);
 				InvDir[Axis] = bParallel[Axis] ? 0 : 1 / Dir[Axis];
 			}
+
+			StartPointSimd = MakeVectorRegisterFloatFromDouble(MakeVectorRegister(Start.X, Start.Y, Start.Z, 0.0));
+			VectorRegister4Float DirSimd = MakeVectorRegisterFloatFromDouble(MakeVectorRegister(Dir.X, Dir.Y, Dir.Z, 0.0));
+			Parallel = VectorCompareLT(VectorAbs(DirSimd), GlobalVectorConstants::SmallNumber);
+			InvDirSimd = VectorBitwiseNotAnd(Parallel, VectorDivide(VectorOne(), DirSimd));
+			ThicknessSimd = VectorSet1(static_cast<FRealSingle>(InThickness));
 		}
 
 		enum class ERaycastType
@@ -99,14 +105,11 @@ namespace Chaos
 			FVec3 Points[4];
 			FAABB3 CellBounds;
 			GeomData->GetPointsAndBoundsScaled(FullIndex, Points, CellBounds);
-			CellBounds.Thicken(Thickness);
+			FAABBVectorized Bounds(CellBounds);
+			Bounds.Thicken(ThicknessSimd);
 
-			// Check cell bounds
-			//todo: can do it without raycast
-			FReal TOI;
-			FVec3 HitPoint;
 			bool bHit = false;
-			if (CellBounds.RaycastFast(Start, Dir, InvDir, bParallel, CurrentLength, 1.0f / CurrentLength, TOI, HitPoint))
+			if (Bounds.RaycastFast(StartPointSimd, InvDirSimd, Parallel, VectorSet1(static_cast<FRealSingle>(CurrentLength))))
 			{
 				// Test both triangles that are in this cell, as we could hit both in any order
 				bHit |= TestTriangle(Payload * 2, Points[0], Points[1], Points[3]);
@@ -275,9 +278,13 @@ namespace Chaos
 		FVec3 OutNormal;
 		int32 OutFaceIndex;
 
-	private:
-
+		VectorRegister4Float StartPointSimd;
+		VectorRegister4Float InvDirSimd;
+		VectorRegister4Float Parallel;
+		VectorRegister4Float ThicknessSimd;
+		
 		const typename FHeightField::FDataType* GeomData;
+	private:
 
 		FVec3 Start;
 		FVec3 Dir;
@@ -897,6 +904,60 @@ namespace Chaos
 			const FVec2 ScaledMin = FlatGrid.MinCorner() * Scale2D;
 			const FVec3 ScaleSign = GeomData.Scale.GetSignVector();
 
+			const FVec3 DirScaled = Dir / ScaledDx;
+			const FReal SumPlaneAxis = DirScaled[0] + DirScaled[1];
+			const bool bCanWalk = SumPlaneAxis > UE_SMALL_NUMBER;
+			if (bCanWalk)
+			{
+				const FVec3 NextStartOri = NextStart;
+				const VectorRegister4Float CurrentLengthSimd = VectorSet1(static_cast<FRealSingle>(CurrentLength));
+				const FReal BoundsMinZ = CachedBounds.Min().Z;
+				const FReal BoundsMaxZ = CachedBounds.Max().Z;
+				const FReal Length2D = FMath::Sqrt(DirScaled[0] * DirScaled[0] + DirScaled[1] * DirScaled[1]);
+				const FReal FastInc = 5.0;
+				const FReal FastIncScaled = FastInc / Length2D;
+				
+				const FVec3 CellsToInspectScaled = DirScaled * FastIncScaled;
+				FReal NextStartZ = NextStart.Z;
+				const FReal Length3D = DirScaled.Length();
+				const FReal FastIncScaled3D = FastInc / Length3D;
+				const FReal CellsToInspectZ = DirScaled.Z * FastIncScaled3D;
+				FReal DistanceProcessed = 0.0;		
+				FReal FastIncScaledAccum = 0.0;
+				while (true)
+				{
+					if (!FlatGrid.IsValid(CellIdx))
+					{
+						return false;
+					}
+					FVec3 NewNextStart = CellsToInspectScaled + NextStart + UE_SMALL_NUMBER;
+					TVec2<int32> NewNextStartInt = TVec2<int32>(static_cast<int32>(NewNextStart[0]), static_cast<int32>(NewNextStart[1]));
+					TVec2<int32> NextStartInt = TVec2<int32>(static_cast<int32>(NextStart[0]), static_cast<int32>(NextStart[1]));
+
+					TVec2<int32> DiffInt = NewNextStartInt - NextStartInt;
+					TVec2<int32> AddedCellIdx = TVec2<int32>(static_cast<int32>(DiffInt[0]), static_cast<int32>(DiffInt[1]));
+
+					FAABBVectorized Bounds;
+					Visitor.GeomData->GetBoundsScaled(CellIdx, AddedCellIdx, Bounds);
+					Bounds.Thicken(Visitor.ThicknessSimd);
+					if (Bounds.RaycastFast(Visitor.StartPointSimd, Visitor.InvDirSimd, Visitor.Parallel, CurrentLengthSimd))
+					{
+						NextStart = NextStartOri + Dir * FastIncScaledAccum;
+						break;
+					}
+
+					NextStart = NewNextStart;
+					CellIdx += AddedCellIdx;
+					if (DistanceProcessed > CurrentLength || NextStartZ < BoundsMinZ || NextStartZ > BoundsMaxZ)
+					{
+						return false;
+					}
+					DistanceProcessed += FastIncScaled3D;
+					NextStartZ += CellsToInspectZ;
+					FastIncScaledAccum += FastIncScaled;
+				}
+			}
+
 			//START
 			do
 			{
@@ -913,7 +974,6 @@ namespace Chaos
 				}
 
 				//find next cell
-
 				//We want to know which plane we used to cross into next cell
 				const FVec2 ScaledCellCenter2D = ScaledMin + FVec2(static_cast<FReal>(CellIdx[0]) + 0.5f, static_cast<FReal>(CellIdx[1]) + 0.5f) * ScaledDx2D;
 				const FVec3 ScaledCellCenter(ScaledCellCenter2D[0], ScaledCellCenter2D[1], ZMidPoint);
diff --git a/Engine/Source/Runtime/Experimental/Chaos/Public/Chaos/AABBVectorized.h b/Engine/Source/Runtime/Experimental/Chaos/Public/Chaos/AABBVectorized.h
--- a/Engine/Source/Runtime/Experimental/Chaos/Public/Chaos/AABBVectorized.h
+++ b/Engine/Source/Runtime/Experimental/Chaos/Public/Chaos/AABBVectorized.h
@@ -90,6 +90,13 @@ namespace Chaos
 			return !static_cast<bool>(VectorMaskBits(IsFalse));
 		}
 
+		FORCEINLINE_DEBUGGABLE FAABBVectorized& Thicken(const VectorRegister4Float& Thickness)
+		{
+			Min = VectorSubtract(Min, Thickness);
+			Max = VectorAdd(Max, Thickness);
+			return *this;
+		}
+
 	private:
 		VectorRegister4Float Min, Max;
 	};
diff --git a/Engine/Source/Runtime/Experimental/Chaos/Public/Chaos/HeightField.h b/Engine/Source/Runtime/Experimental/Chaos/Public/Chaos/HeightField.h
--- a/Engine/Source/Runtime/Experimental/Chaos/Public/Chaos/HeightField.h
+++ b/Engine/Source/Runtime/Experimental/Chaos/Public/Chaos/HeightField.h
@@ -286,6 +286,48 @@ namespace Chaos
 				OutBounds = FAABB3(FVec3((FReal)X, (FReal)Y, MinZ), FVec3((FReal)X + 1, (FReal)Y + 1, MaxZ));
 			}
 
+			FORCEINLINE void GetBounds(TVec2<int32> CellIdx, TVec2<int32> Area, FAABBVectorized& OutBounds) const
+			{
+				FRealSingle MinHeight = UE_BIG_NUMBER;
+				FRealSingle MaxHeight = -UE_BIG_NUMBER;
+
+				const int32 EndIndexX = CellIdx[0] + Area[0];
+				int32 FirstIndexX = FMath::Min<int32>(CellIdx[0], EndIndexX);
+				FirstIndexX = FMath::Max<int32>(FirstIndexX, 0);
+				int32 LastIndexX = FMath::Max<int32>(CellIdx[0], EndIndexX);
+				// Increment LastIndex to look into the four corners of a cell 
+				LastIndexX++;
+				LastIndexX = FMath::Min<int32>(LastIndexX, NumCols);
+				
+				const int32 EndIndexY = CellIdx[1] + Area[1];
+				int32 FirstIndexY = FMath::Min<int32>(CellIdx[1], EndIndexY);
+				FirstIndexY = FMath::Max<int32>(FirstIndexY, 0);
+				int32 LastIndexY = FMath::Max<int32>(CellIdx[1], EndIndexY);
+				LastIndexY++;
+				LastIndexY = FMath::Min<int32>(LastIndexY, NumRows);
+
+				for (int IndexY = FirstIndexY; IndexY <= LastIndexY; IndexY++)
+				{
+					for (int IndexX = FirstIndexX; IndexX <= LastIndexX; IndexX++)
+					{
+						const int32 Index = IndexY * (NumCols - 1) + IndexX + IndexY;
+						if (Index < Heights.Num())
+						{
+							MinHeight = FMath::Min<FRealSingle>(Heights[Index], MinHeight);
+							MaxHeight = FMath::Max<FRealSingle>(Heights[Index], MaxHeight);
+						}
+					}
+				}
+
+				const FRealSingle MinValueSingle = static_cast<FRealSingle>(MinValue);
+				const FRealSingle HeightPerUnitSingle = static_cast<FRealSingle>(HeightPerUnit);
+				MinHeight = MinValueSingle + MinHeight * HeightPerUnitSingle;
+				MaxHeight = MinValueSingle + MaxHeight * HeightPerUnitSingle;
+
+				OutBounds = FAABBVectorized(MakeVectorRegisterFloat(static_cast<FRealSingle>(FirstIndexX), static_cast<FRealSingle>(FirstIndexY), MinHeight, 0.0f),
+											MakeVectorRegisterFloat(static_cast<FRealSingle>(LastIndexX), static_cast<FRealSingle>(LastIndexY), MaxHeight, 0.0f));
+			}
+
 			FORCEINLINE void GetPointsScaled(int32 Index, FVec3 OutPts[4]) const
 			{
 				GetPoints(Index, OutPts);
@@ -305,8 +347,21 @@ namespace Chaos
 				OutPts[2] *= Scale;
 				OutPts[3] *= Scale;
 
-				const FVec3 Scale3(Scale);
-				OutBounds = FAABB3::FromPoints(OutBounds.Min() * Scale3, OutBounds.Max() * Scale3);
+				OutBounds = FAABB3::FromPoints(OutBounds.Min() * Scale, OutBounds.Max() * Scale);
+			}
+
+			FORCEINLINE void GetBoundsScaled(TVec2<int32> CellIdx, TVec2<int32> Area, FAABBVectorized& OutBounds) const
+			{
+				GetBounds(CellIdx, Area, OutBounds);
+
+				VectorRegister4Float ScaleSimd = MakeVectorRegisterFloatFromDouble(MakeVectorRegisterDouble(Scale.X, Scale.Y, Scale.Z, 0.0));
+
+				VectorRegister4Float P0 = VectorMultiply(OutBounds.GetMin(), ScaleSimd);
+				VectorRegister4Float P1 = VectorMultiply(OutBounds.GetMax(), ScaleSimd);
+				VectorRegister4Float Min = VectorMin(P0, P1);
+				VectorRegister4Float Max = VectorMax(P0, P1);
+
+				OutBounds = FAABBVectorized(Min, Max);
 			}
 
 			FORCEINLINE FReal GetMinHeight() const
diff --git a/Engine/Source/Runtime/Experimental/Chaos/Public/Chaos/TriangleMeshImplicitObject.h b/Engine/Source/Runtime/Experimental/Chaos/Public/Chaos/TriangleMeshImplicitObject.h
--- a/Engine/Source/Runtime/Experimental/Chaos/Public/Chaos/TriangleMeshImplicitObject.h
+++ b/Engine/Source/Runtime/Experimental/Chaos/Public/Chaos/TriangleMeshImplicitObject.h
@@ -166,7 +166,6 @@ namespace Chaos
 		template <typename SQVisitor>
 		FORCEINLINE_DEBUGGABLE void Overlap(const FAABBVectorized& AABB, SQVisitor& Visitor) const
 		{
-			CSV_SCOPED_TIMING_STAT(PhysicsVerbose, BVHOverlap);
 			const auto BoundsFilter = [&AABB](const FAABBVectorized& Bounds) -> EFilterResult
 			{
 				FAABBVectorized CleanBounds(Bounds.GetMin(), Bounds.GetMax());
