Chaos : Make sure the simulation sizespecific array has the default MaxSize to be always higher than anything the user will set

diff --git a/Engine/Source/Runtime/Experimental/GeometryCollectionEngine/Private/GeometryCollection/GeometryCollectionObject.cpp b/Engine/Source/Runtime/Experimental/GeometryCollectionEngine/Private/GeometryCollection/GeometryCollectionObject.cpp
--- a/Engine/Source/Runtime/Experimental/GeometryCollectionEngine/Private/GeometryCollection/GeometryCollectionObject.cpp
+++ b/Engine/Source/Runtime/Experimental/GeometryCollectionEngine/Private/GeometryCollection/GeometryCollectionObject.cpp
@@ -331,7 +331,7 @@ void UGeometryCollection::GetSharedSimulationParams(FSharedSimulationParameters&
 		InfSize.CollisionShapes[0].CollisionParticles.CollisionParticlesFraction = SizeSpecificDefault.CollisionShapes[0].CollisionParticles.CollisionParticlesFraction;
 		InfSize.CollisionShapes[0].CollisionParticles.MaximumCollisionParticles = SizeSpecificDefault.CollisionShapes[0].CollisionParticles.MaximumCollisionParticles;
 	}
-	InfSize.MaxSize = DefaultMaxSizeValue;
+	InfSize.MaxSize = TNumericLimits<float>::Max();
 	OutParams.SizeSpecificData.SetNum(SizeSpecificData.Num() + 1);
 	FillSharedSimulationSizeSpecificData(OutParams.SizeSpecificData[0], InfSize);
 
