Chaos : Flesh Fix copyright boiler plate issue

diff --git a/Engine/Source/Runtime/Experimental/ChaosCore/Public/Chaos/ImplicitQRSVD.h b/Engine/Source/Runtime/Experimental/ChaosCore/Public/Chaos/ImplicitQRSVD.h
--- a/Engine/Source/Runtime/Experimental/ChaosCore/Public/Chaos/ImplicitQRSVD.h
+++ b/Engine/Source/Runtime/Experimental/ChaosCore/Public/Chaos/ImplicitQRSVD.h
@@ -1,9 +1,4 @@
-
 // Copyright Epic Games, Inc. All Rights Reserved.
-
-/**
-SVD based on implicit QR with Wilkinson Shift
-*/
 #pragma once
 
 #include "Chaos/Matrix.h"
