Chaos : destruction : reorganize damage threshold computation

diff --git a/Engine/Source/Runtime/Experimental/Chaos/Private/PhysicsProxy/GeometryCollectionPhysicsProxy.cpp b/Engine/Source/Runtime/Experimental/Chaos/Private/PhysicsProxy/GeometryCollectionPhysicsProxy.cpp
--- a/Engine/Source/Runtime/Experimental/Chaos/Private/PhysicsProxy/GeometryCollectionPhysicsProxy.cpp
+++ b/Engine/Source/Runtime/Experimental/Chaos/Private/PhysicsProxy/GeometryCollectionPhysicsProxy.cpp
@@ -76,6 +76,8 @@ FAutoConsoleVariableRef CVarGeometryCollectionAlwaysGenerateGTCollisionForCluste
 
 DEFINE_LOG_CATEGORY_STATIC(UGCC_LOG, Error, All);
 
+static const FSharedSimulationSizeSpecificData& GetSizeSpecificData(const TArray<FSharedSimulationSizeSpecificData>& SizeSpecificData, const FGeometryCollection& RestCollection, const int32 TransformIndex, const FBox& BoundingBox);
+
 //==============================================================================
 // FGeometryCollectionResults
 //==============================================================================
@@ -931,7 +933,7 @@ void FGeometryCollectionPhysicsProxy::InitializeBodiesPT(Chaos::FPBDRigidsSolver
 					Chaos::FGeometryParticle* GTParticle = GTParticles[TransformGroupIndex].Get();
 
 					Chaos::FUniqueIdx ExistingIndex = GTParticle->UniqueIdx();
-					Chaos::FPBDRigidClusteredParticleHandle* Handle = BuildClusters(TransformGroupIndex, RigidChildren, RigidChildrenTransformGroupIndex, CreationParameters, &ExistingIndex);
+					Chaos::FPBDRigidClusteredParticleHandle* Handle = BuildClusters_Internal(TransformGroupIndex, RigidChildren, RigidChildrenTransformGroupIndex, CreationParameters, &ExistingIndex);
 					Handle->GTGeometryParticle() = GTParticle;
 
 					int32 RigidChildrenIdx = 0;
@@ -966,7 +968,7 @@ void FGeometryCollectionPhysicsProxy::InitializeBodiesPT(Chaos::FPBDRigidsSolver
 			Particles.UpdateGeometryCollectionViews(true); 
 
 			// Set cluster connectivity.  TPBDRigidClustering::CreateClusterParticle() 
-			// will optionally do this, but we switch that functionality off in BuildClusters().
+			// will optionally do this, but we switch that functionality off in BuildClusters_Internal().
 			for(int32 TransformGroupIndex = 0; TransformGroupIndex < NumTransforms; ++TransformGroupIndex)
 			{
 				if (RestCollection->IsClustered(TransformGroupIndex))
@@ -1044,12 +1046,43 @@ FAutoConsoleVariableRef CVarReportNoLevelsetCluster(TEXT("p.gc.ReportNoLevelsetC
 DECLARE_CYCLE_STAT(TEXT("FGeometryCollectionPhysicsProxy::BuildClusters"), STAT_BuildClusters, STATGROUP_Chaos);
 DECLARE_CYCLE_STAT(TEXT("FGeometryCollectionPhysicsProxy::BuildClusters:GlobalMatrices"), STAT_BuildClustersGlobalMatrices, STATGROUP_Chaos);
 
+float FGeometryCollectionPhysicsProxy::ComputeDamageThreshold(const FGeometryDynamicCollection& DynamicCollection, int32 TransformIndex) const 
+{
+	float DamageThreshold = TNumericLimits<float>::Max();
+	if (Parameters.bUseSizeSpecificDamageThresholds)
+	{
+		// bounding box volume is used as a fallback to find specific size if the relative size if not available
+		// ( May happen with older GC )
+		FBox LocalBoundingBox;
+		const FGeometryDynamicCollection::FSharedImplicit& Implicit = DynamicCollection.Implicits[TransformIndex];
+		if (Implicit && Implicit->HasBoundingBox())
+		{
+			const Chaos::FAABB3& ImplicitBoundingBox = Implicit->BoundingBox();
+			LocalBoundingBox = FBox(ImplicitBoundingBox.Min(), ImplicitBoundingBox.Max());
+		}
+		
+		const FSharedSimulationSizeSpecificData& SizeSpecificData = GetSizeSpecificData(Parameters.Shared.SizeSpecificData, *Parameters.RestCollection, TransformIndex, LocalBoundingBox);
+		DamageThreshold = SizeSpecificData.DamageThreshold;
+	}
+	else
+	{
+		const int32 NumThresholds = Parameters.DamageThreshold.Num();
+		const int32 Level = FMath::Clamp(CalculateHierarchyLevel(DynamicCollection, TransformIndex), 0, INT_MAX);
+		const float DefaultDamage = NumThresholds > 0 ? Parameters.DamageThreshold[NumThresholds - 1] : 0.f;
+		DamageThreshold = Level < NumThresholds ? Parameters.DamageThreshold[Level] : DefaultDamage;
 
+		if (Level >= Parameters.MaxClusterLevel)
+		{
+			DamageThreshold = TNumericLimits<float>::Max();
+		}
+	}
+	return DamageThreshold;
+}
 
-Chaos::TPBDRigidClusteredParticleHandle<Chaos::FReal, 3>*
-FGeometryCollectionPhysicsProxy::BuildClusters(
+Chaos::FPBDRigidClusteredParticleHandle*
+FGeometryCollectionPhysicsProxy::BuildClusters_Internal(
 	const uint32 CollectionClusterIndex, // TransformGroupIndex
-	TArray<Chaos::TPBDRigidParticleHandle<Chaos::FReal,3>*>& ChildHandles,
+	TArray<Chaos::FPBDRigidParticleHandle*>& ChildHandles,
 	const TArray<int32>& ChildTransformGroupIndices,
 	const Chaos::FClusterCreationParameters & ClusterParameters,
 	const Chaos::FUniqueIdx* ExistingIndex)
@@ -1147,45 +1180,9 @@ FGeometryCollectionPhysicsProxy::BuildClusters(
 	// two-way mapping
 	SolverClusterHandles[CollectionClusterIndex] = Parent;
 
-	const int32 NumThresholds = Parameters.DamageThreshold.Num();
-	const int32 Level = FMath::Clamp(CalculateHierarchyLevel(DynamicCollection, CollectionClusterIndex), 0, INT_MAX);
-	const float DefaultDamage = NumThresholds > 0 ? Parameters.DamageThreshold[NumThresholds - 1] : 0.f;
-	float Damage = Level < NumThresholds ? Parameters.DamageThreshold[Level] : DefaultDamage;
-
-	if(Level >= Parameters.MaxClusterLevel)
-	{
-		Damage = FLT_MAX;
-	}
-
-	if (Parameters.bUseSizeSpecificDamageThresholds)
-	{
-		// If RelativeSize is available, use that to determine SizeSpecific index, otherwise, fall back to bounds volume.
-		int32 SizeSpecificIdx = 0;
-		if (Parameters.RestCollection->HasAttribute("Size", FTransformCollection::TransformGroup))
-		{
-			const TManagedArray<float>& RelativeSize = Parameters.RestCollection->GetAttribute<float>("Size", FTransformCollection::TransformGroup);
-			SizeSpecificIdx = GeometryCollection::SizeSpecific::FindIndexForVolume(Parameters.Shared.SizeSpecificData, RelativeSize[CollectionClusterIndex]);
-		}
-		else
-		{
-			const TManagedArray<FGeometryDynamicCollection::FSharedImplicit>& Implicit = DynamicCollection.Implicits;
-			if (Implicit[CollectionClusterIndex] && Implicit[CollectionClusterIndex]->HasBoundingBox())
-			{
-				FBox LocalBoundingBox(Implicit[CollectionClusterIndex]->BoundingBox().Min(), Implicit[CollectionClusterIndex]->BoundingBox().Max());
-				SizeSpecificIdx = GeometryCollection::SizeSpecific::FindIndexForVolume(Parameters.Shared.SizeSpecificData, LocalBoundingBox);
-			}
-		}
-
-		if (0 <= SizeSpecificIdx && SizeSpecificIdx < Parameters.Shared.SizeSpecificData.Num())
-		{
-			const FSharedSimulationSizeSpecificData& SizeSpecificData = Parameters.Shared.SizeSpecificData[SizeSpecificIdx];
-			Damage = SizeSpecificData.DamageThreshold;
-		}
-	}
-
-	Parent->SetStrains(Damage);
-
-
+	const float DamageThreshold = ComputeDamageThreshold(DynamicCollection, CollectionClusterIndex);
+	Parent->SetStrains(DamageThreshold);
+	
 	// #BGTODO This will not automatically update - material properties should only ever exist in the material, not in other arrays
 	const Chaos::FChaosPhysicsMaterial* CurMaterial = static_cast<Chaos::FPBDRigidsSolver*>(Solver)->GetSimMaterials().Get(Parameters.PhysicalMaterialHandle.InnerHandle);
 	if(CurMaterial)
@@ -1205,10 +1202,11 @@ FGeometryCollectionPhysicsProxy::BuildClusters(
 	int32 MinCollisionGroup = INT_MAX;
 	for(int32 Idx=0; Idx < ChildHandles.Num(); Idx++)
 	{
-		Chaos::TPBDRigidParticleHandle<Chaos::FReal, 3>* Child = ChildHandles[Idx];
-		if (Chaos::TPBDRigidClusteredParticleHandle<Chaos::FReal, 3>* ClusteredChild = Child->CastToClustered())
+		// set the damage threshold on children as they are the one where the strain is tested when breaking 
+		Chaos::FPBDRigidParticleHandle* Child = ChildHandles[Idx];
+		if (Chaos::FPBDRigidClusteredParticleHandle* ClusteredChild = Child->CastToClustered())
 		{
-			ClusteredChild->SetStrains(Damage);
+			ClusteredChild->SetStrains(DamageThreshold);
 		}
 
 		const int32 ChildTransformGroupIndex = ChildTransformGroupIndices[Idx];
@@ -1428,12 +1426,12 @@ void FGeometryCollectionPhysicsProxy::DisableParticles(TArray<int32>& TransformG
 	}
 }
 
-int32 FGeometryCollectionPhysicsProxy::CalculateHierarchyLevel(const FGeometryDynamicCollection& GeometryCollection, int32 TransformIndex) const
+int32 FGeometryCollectionPhysicsProxy::CalculateHierarchyLevel(const FGeometryDynamicCollection& DynamicCollection, int32 TransformIndex)
 {
 	int32 Level = 0;
-	while (GeometryCollection.Parent[TransformIndex] != -1)
+	while (DynamicCollection.Parent[TransformIndex] != INDEX_NONE)
 	{
-		TransformIndex = GeometryCollection.Parent[TransformIndex];
+		TransformIndex = DynamicCollection.Parent[TransformIndex];
 		Level++;
 	}
 	return Level;
diff --git a/Engine/Source/Runtime/Experimental/Chaos/Public/PhysicsProxy/GeometryCollectionPhysicsProxy.h b/Engine/Source/Runtime/Experimental/Chaos/Public/PhysicsProxy/GeometryCollectionPhysicsProxy.h
--- a/Engine/Source/Runtime/Experimental/Chaos/Public/PhysicsProxy/GeometryCollectionPhysicsProxy.h
+++ b/Engine/Source/Runtime/Experimental/Chaos/Public/PhysicsProxy/GeometryCollectionPhysicsProxy.h
@@ -246,6 +246,15 @@ public:
 	const FProxyInterpolationData& GetInterpolationData() const { return InterpolationData; }
 
 protected:
+	/**
+	* Compute damage threshold for a specific transform
+	* this account for component level damage threshold as well as size specific ones
+	* @param DynamicCollection dynamic collection to use
+	* @param TransformIndex index of the transform to compute the threshold for
+	* #return damage threshold value
+	*/
+	float ComputeDamageThreshold(const FGeometryDynamicCollection& DynamicCollection, int32 TransformIndex) const;
+		
 	/**
 	 * Build a physics thread cluster parent particle.
 	 *	\p CollectionClusterIndex - the source geometry collection transform index.
@@ -254,9 +263,9 @@ protected:
 	 *  \P Parameters - uh, yeah...  Other parameters.
 	 */
 
-	Chaos::TPBDRigidClusteredParticleHandle<Chaos::FReal, 3>* BuildClusters(
+	Chaos::FPBDRigidClusteredParticleHandle* BuildClusters_Internal(
 		const uint32 CollectionClusterIndex, 
-		TArray<Chaos::TPBDRigidParticleHandle<Chaos::FReal, 3>*>& ChildHandles,
+		TArray<Chaos::FPBDRigidParticleHandle*>& ChildHandles,
 		const TArray<int32>& ChildTransformGroupIndices,
 		const Chaos::FClusterCreationParameters & Parameters,
 		const Chaos::FUniqueIdx* ExistingIndex);
@@ -265,7 +274,7 @@ protected:
 	 * Traverses the parents of \p TransformIndex in \p GeometryCollection, counting
 	 * the number of levels until the next parent is \c INDEX_NONE.
 	 */
-	int32 CalculateHierarchyLevel(const FGeometryDynamicCollection& GeometryCollection, int32 TransformIndex) const;
+	static int32 CalculateHierarchyLevel(const FGeometryDynamicCollection& DynamicCollection, int32 TransformIndex);
 
 	void InitializeRemoveOnFracture(FParticlesType& Particles, const TManagedArray<int32>& DynamicState);
 
