- Fixing crash in IslandManager.cpp where the ConstraintHandleHolder's Handle is null

diff --git a/Engine/Source/Runtime/Experimental/Chaos/Private/Chaos/Island/IslandManager.cpp b/Engine/Source/Runtime/Experimental/Chaos/Private/Chaos/Island/IslandManager.cpp
--- a/Engine/Source/Runtime/Experimental/Chaos/Private/Chaos/Island/IslandManager.cpp
+++ b/Engine/Source/Runtime/Experimental/Chaos/Private/Chaos/Island/IslandManager.cpp
@@ -152,7 +152,10 @@ inline void UpdateSleepState(FPBDIslandSolver* IslandSolver, FPBDRigidsSOAs& Par
 		// Island constraints are updating their sleeping flag + awaken one 
 		for (auto& IslandConstraint : IslandSolver->GetConstraints())
 		{
-			IslandConstraint->SetIsSleeping(bIsSleeping);
+			if (FConstraintHandle* ConstraintHandle = IslandConstraint.Get())
+			{
+				ConstraintHandle->SetIsSleeping(bIsSleeping);
+			}
 		}
 	}
 }
