Chaos : fix shape generation not properly working using relative size for size specific data when clustering is on

diff --git a/Engine/Source/Runtime/Experimental/Chaos/Private/PhysicsProxy/GeometryCollectionPhysicsProxy.cpp b/Engine/Source/Runtime/Experimental/Chaos/Private/PhysicsProxy/GeometryCollectionPhysicsProxy.cpp
--- a/Engine/Source/Runtime/Experimental/Chaos/Private/PhysicsProxy/GeometryCollectionPhysicsProxy.cpp
+++ b/Engine/Source/Runtime/Experimental/Chaos/Private/PhysicsProxy/GeometryCollectionPhysicsProxy.cpp
@@ -2719,6 +2719,17 @@ void FGeometryCollectionPhysicsProxy::InitializeSharedCollisionStructures(
 
 				const FSharedSimulationSizeSpecificData& SizeSpecificData = GetSizeSpecificData(SharedParams.SizeSpecificData, RestCollection, TransformGroupIndex, InstanceBoundingBox);
 
+				int32 SizeSpecificIdx;
+				if (bUseRelativeSize)
+				{
+					const TManagedArray<float>& RelativeSize = RestCollection.GetAttribute<float>(TEXT("Size"), FTransformCollection::TransformGroup);
+					SizeSpecificIdx = GeometryCollection::SizeSpecific::FindIndexForVolume(SharedParams.SizeSpecificData, RelativeSize[TransformGroupIndex]);
+				}
+				else
+				{
+					SizeSpecificIdx = GeometryCollection::SizeSpecific::FindIndexForVolume(SharedParams.SizeSpecificData, InstanceBoundingBox);	
+				}
+
 				ErrorReporter.SetPrefix(BaseErrorPrefix + " | Cluster Transform Index: " + FString::FromInt(ClusterTransformIdx));
 
 				CollectionImplicits[ClusterTransformIdx] = CreateImplicitGeometry(
