Chaos
- Inertia CombineWorldSpace now always returned a world-space inertia with an identity rotation of mass
- fix inerta orientation when a particle has a single shape and it has a relative rotation
- fix inertia scaling for multi-shape and rotated shapes

diff --git a/Engine/Source/Runtime/Engine/Private/PhysicsEngine/BodyInstance.cpp b/Engine/Source/Runtime/Engine/Private/PhysicsEngine/BodyInstance.cpp
--- a/Engine/Source/Runtime/Engine/Private/PhysicsEngine/BodyInstance.cpp
+++ b/Engine/Source/Runtime/Engine/Private/PhysicsEngine/BodyInstance.cpp
@@ -3327,14 +3327,9 @@ void FBodyInstance::UpdateMassProperties()
 					{
 						// The component scale is already built into the geometry, but if the user has set up a CoM
 						// modifier, it will need to be transformed by the component scale.
+						const bool bInertiaScaleIncludeMass = false;
 						FTransform MassModifierTransform(FQuat::Identity, FVector(0.f, 0.f, 0.f), Scale3D);
-						TotalMassProperties = BodyUtils::ComputeMassProperties(this, Shapes, MassModifierTransform, false/*bInertiaScaleIncludeMass*/);
-
-#if WITH_CHAOS
-						// @todo(chaos): this is starting to get messy - we should probably just split off the PhysX implementation now
-						// Make the inertia diagonal and calculate the rotation of mass
-						TotalMassProperties.RotationOfMass = Chaos::TransformToLocalSpace(TotalMassProperties.InertiaTensor);
-#endif
+						TotalMassProperties = BodyUtils::ComputeMassProperties(this, Shapes, MassModifierTransform, bInertiaScaleIncludeMass);
 					}
 				}
 
diff --git a/Engine/Source/Runtime/Engine/Private/PhysicsEngine/BodyUtils.cpp b/Engine/Source/Runtime/Engine/Private/PhysicsEngine/BodyUtils.cpp
--- a/Engine/Source/Runtime/Engine/Private/PhysicsEngine/BodyUtils.cpp
+++ b/Engine/Source/Runtime/Engine/Private/PhysicsEngine/BodyUtils.cpp
@@ -102,6 +102,9 @@ namespace BodyUtils
 		Chaos::FMassProperties MassProps;
 		ChaosInterface::CalculateMassPropertiesFromShapeCollection(MassProps, Shapes, GetBodyInstanceDensity(OwningBodyInstance));
 
+		// Diagonalize the inertia
+		Chaos::TransformToLocalSpace(MassProps);
+
 		// Apply the BodyInstance's mass and inertia modifiers
 		return ApplyMassPropertiesModifiers(OwningBodyInstance, MassProps, MassModifierTransform, bInertaScaleIncludeMass);
 	}
@@ -112,6 +115,9 @@ namespace BodyUtils
 		Chaos::FMassProperties MassProps;
 		ChaosInterface::CalculateMassPropertiesFromShapeCollection(MassProps, Shapes, bContributesToMass, GetBodyInstanceDensity(OwningBodyInstance));
 
+		// Diagonalize the inertia
+		Chaos::TransformToLocalSpace(MassProps);
+
 		// Apply the BodyInstance's mass and inertia modifiers
 		return ApplyMassPropertiesModifiers(OwningBodyInstance, MassProps, MassModifierTransform, bInertaScaleIncludeMass);
 	}
diff --git a/Engine/Source/Runtime/Engine/Private/PhysicsEngine/BodyUtils.h b/Engine/Source/Runtime/Engine/Private/PhysicsEngine/BodyUtils.h
--- a/Engine/Source/Runtime/Engine/Private/PhysicsEngine/BodyUtils.h
+++ b/Engine/Source/Runtime/Engine/Private/PhysicsEngine/BodyUtils.h
@@ -31,6 +31,7 @@ namespace BodyUtils
 
 	/** 
 	 * Computes and adds the mass properties (inertia, com, etc...) based on the mass settings of the body instance. 
+	 * The inertia returned will be diagonal, and there may be a non-identity rotation of mass.
 	 * Note: this includes a call to ModifyMassProperties, so the BodyInstance modifiers will be included in the calculation.
 	 */
 	Chaos::FMassProperties ComputeMassProperties(const FBodyInstance* OwningBodyInstance, const TArray<FPhysicsShapeHandle>& Shapes, const FTransform& MassModifierTransform, const bool bInertaScaleIncludeMass = false);
diff --git a/Engine/Source/Runtime/Engine/Private/PhysicsEngine/Experimental/ChaosInterfaceUtils.cpp b/Engine/Source/Runtime/Engine/Private/PhysicsEngine/Experimental/ChaosInterfaceUtils.cpp
--- a/Engine/Source/Runtime/Engine/Private/PhysicsEngine/Experimental/ChaosInterfaceUtils.cpp
+++ b/Engine/Source/Runtime/Engine/Private/PhysicsEngine/Experimental/ChaosInterfaceUtils.cpp
@@ -497,24 +497,21 @@ namespace ChaosInterface
 		}
 
 		Chaos::FMatrix33 Tensor;
-		Chaos::FRotation3 RotationOfMass;
 
 		// If no shapes contribute to mass, or they are scaled to zero, we may end up with zero mass here
 		if ((TotalMass > 0.f) && (MassPropertiesList.Num() > 0))
 		{
 			TotalCenterOfMass /= TotalMass;
 
-			// NOTE: If multiple items in the list, rotation of mass will be zero, but if only 1 item is the list the item is returned directly and we may have a rotation of mass
 			Chaos::FMassProperties CombinedMassProperties = Chaos::CombineWorldSpace(MassPropertiesList);
+			ensure(CombinedMassProperties.RotationOfMass.IsIdentity());
 			Tensor = CombinedMassProperties.InertiaTensor;
-			RotationOfMass = CombinedMassProperties.RotationOfMass;
 		}
 		else
 		{
 			// @todo(chaos): We should support shape-less particles as long as their mass an inertia are set directly
 			// For now hard-code a 50cm sphere with density 1g/cc
 			Tensor = Chaos::FMatrix33(5.24e5f, 5.24e5f, 5.24e5f);
-			RotationOfMass = Chaos::FRotation3::Identity;
 			TotalMass = 523.0f;
 			TotalVolume = 523000;
 		}
@@ -523,7 +520,7 @@ namespace ChaosInterface
 		OutProperties.Mass = TotalMass;
 		OutProperties.Volume = TotalVolume;
 		OutProperties.CenterOfMass = TotalCenterOfMass;
-		OutProperties.RotationOfMass = RotationOfMass;
+		OutProperties.RotationOfMass = Chaos::FRotation3::Identity;
 	}
 
 
diff --git a/Engine/Source/Runtime/Engine/Private/PhysicsEngine/ImmediatePhysics/ImmediatePhysicsChaos/ImmediatePhysicsActorHandle_Chaos.cpp b/Engine/Source/Runtime/Engine/Private/PhysicsEngine/ImmediatePhysics/ImmediatePhysicsChaos/ImmediatePhysicsActorHandle_Chaos.cpp
--- a/Engine/Source/Runtime/Engine/Private/PhysicsEngine/ImmediatePhysics/ImmediatePhysicsChaos/ImmediatePhysicsActorHandle_Chaos.cpp
+++ b/Engine/Source/Runtime/Engine/Private/PhysicsEngine/ImmediatePhysics/ImmediatePhysicsChaos/ImmediatePhysicsActorHandle_Chaos.cpp
@@ -187,7 +187,6 @@ namespace ImmediatePhysics_Chaos
 			// as well as the dimension change even though we don't actually change the mass.
 			const bool bInertaScaleIncludeMass = true;
 			FMassProperties MassProperties = BodyUtils::ComputeMassProperties(BodyInstance, Shapes, bContributesToMass, FTransform::Identity, bInertaScaleIncludeMass);
-			MassProperties.RotationOfMass = Chaos::TransformToLocalSpace(MassProperties.InertiaTensor);
 			OutMass = MassProperties.Mass;
 			OutInertia = MassProperties.InertiaTensor.GetDiagonal();
 			OutCoMTransform = FTransform(MassProperties.RotationOfMass, MassProperties.CenterOfMass);
diff --git a/Engine/Source/Runtime/Engine/Public/Physics/Experimental/ChaosInterfaceUtils.h b/Engine/Source/Runtime/Engine/Public/Physics/Experimental/ChaosInterfaceUtils.h
--- a/Engine/Source/Runtime/Engine/Public/Physics/Experimental/ChaosInterfaceUtils.h
+++ b/Engine/Source/Runtime/Engine/Public/Physics/Experimental/ChaosInterfaceUtils.h
@@ -21,6 +21,10 @@ namespace ChaosInterface
 	void CreateGeometry(const FGeometryAddParams& InParams, TArray<TUniquePtr<Chaos::FImplicitObject>>& OutGeoms, Chaos::FShapesArray& OutShapes);
 
 #if WITH_CHAOS
+	/**
+	 * Generate the mass properties for a set of shapes in the space of the shapes' owner. 
+	 * Rotation will be built into the inertia matrix (it may not be diagonal) and RotationOfMass will be identity.
+	*/
 	void CalculateMassPropertiesFromShapeCollection(Chaos::FMassProperties& OutProperties, const TArray<FPhysicsShapeHandle>& InShapes, float InDensityKGPerCM);
 	void CalculateMassPropertiesFromShapeCollection(Chaos::FMassProperties& OutProperties, const Chaos::FShapesArray& InShapes, const TArray<bool>& bContributesToMass, float InDensityKGPerCM);
 #endif
diff --git a/Engine/Source/Runtime/Experimental/Chaos/Private/Chaos/MassProperties.cpp b/Engine/Source/Runtime/Experimental/Chaos/Private/Chaos/MassProperties.cpp
--- a/Engine/Source/Runtime/Experimental/Chaos/Private/Chaos/MassProperties.cpp
+++ b/Engine/Source/Runtime/Experimental/Chaos/Private/Chaos/MassProperties.cpp
@@ -5,6 +5,7 @@
 #include "Chaos/Matrix.h"
 #include "Chaos/Particles.h"
 #include "Chaos/TriangleMesh.h"
+#include "Chaos/Utilities.h"
 #include "ChaosCheck.h"
 
 
@@ -23,13 +24,13 @@ namespace Chaos
 		if (Trace <= UE_SMALL_NUMBER)
 		{
 			// Tiny inertia - numerical instability would follow. We should not get this unless we have bad input.
-			return FRotation3::FromElements(FVec3(0), 1);
+			return FRotation3::FromIdentity();
 		}
 
 		if ((OffDiagSize / Trace) < UE_SMALL_NUMBER)
 		{
 			// Almost diagonal matrix - we are already in local space.
-			return FRotation3::FromElements(FVec3(0), 1);
+			return FRotation3::FromIdentity();
 		}
 
 		FReal Size = static_cast<FReal>(FMath::Sqrt((FMath::Square(Inertia.M[0][0] - Trace) + FMath::Square(Inertia.M[1][1] - Trace) + FMath::Square(Inertia.M[2][2] - Trace) + 2. * OffDiagSize) / 6.));
@@ -84,12 +85,22 @@ namespace Chaos
 		FinalRotation = FRotation3(RotationMatrix.GetTransposed());
 		if (!ensure(FMath::IsNearlyEqual((float)FinalRotation.Size(), 1.0f, UE_KINDA_SMALL_NUMBER)))
 		{
-			return FRotation3::FromElements(FVec3(0), 1);
+			return FRotation3::FromIdentity();
 		}
 		
 		return FinalRotation;
 	}
 
+	void TransformToLocalSpace(FMassProperties& MassProperties)
+	{
+		// Diagonalize inertia
+		const FRotation3 InertiaRotation = TransformToLocalSpace(MassProperties.InertiaTensor);
+
+		// Calculate net rotation
+		MassProperties.RotationOfMass = MassProperties.RotationOfMass * InertiaRotation;
+	}
+
+
 	void CalculateVolumeAndCenterOfMass(const FBox& BoundingBox, FVector::FReal& OutVolume, FVector& OutCenterOfMass)
 	{
 		const FVector Extents = static_cast<FVector::FReal>(2) * BoundingBox.GetExtent(); // FBox::GetExtent() returns half the size, but FAABB::Extents() returns total size
@@ -219,39 +230,47 @@ namespace Chaos
 	FMassProperties Combine(const TArray<FMassProperties>& MPArray)
 	{
 		FMassProperties NewMP = CombineWorldSpace(MPArray);
-		NewMP.RotationOfMass = TransformToLocalSpace(NewMP.InertiaTensor);
+		TransformToLocalSpace(NewMP);
 		return NewMP;
 	}
 
 	FMassProperties CombineWorldSpace(const TArray<FMassProperties>& MPArray)
 	{
 		check(MPArray.Num() > 0);
-		if (MPArray.Num() == 1)
+
+		if ((MPArray.Num() == 1) && MPArray[0].RotationOfMass.IsIdentity())
+		{
 			return MPArray[0];
+		}
+
 		FMassProperties NewMP;
+
 		for (const FMassProperties& Child : MPArray)
 		{
 			NewMP.Volume += Child.Volume;
-			const FMatrix33 ChildRI = Child.RotationOfMass.ToMatrix();
-			const FMatrix33 ChildWorldSpaceI = ChildRI.GetTransposed() * Child.InertiaTensor * ChildRI;
-			NewMP.InertiaTensor += ChildWorldSpaceI;
+			NewMP.InertiaTensor += Utilities::ComputeWorldSpaceInertia(Child.RotationOfMass, Child.InertiaTensor);
 			NewMP.CenterOfMass += Child.CenterOfMass * Child.Mass;
 			NewMP.Mass += Child.Mass;
 		}
 		check(NewMP.Mass > UE_SMALL_NUMBER);
 		NewMP.CenterOfMass /= NewMP.Mass;
-		for (const FMassProperties& Child : MPArray)
+
+		if (MPArray.Num() > 1)
 		{
-			const FReal M = Child.Mass;
-			const FVec3 ParentToChild = Child.CenterOfMass - NewMP.CenterOfMass;
-			const FReal P0 = ParentToChild[0];
-			const FReal P1 = ParentToChild[1];
-			const FReal P2 = ParentToChild[2];
-			const FReal MP0P0 = M * P0 * P0;
-			const FReal MP1P1 = M * P1 * P1;
-			const FReal MP2P2 = M * P2 * P2;
-			NewMP.InertiaTensor += FMatrix33(MP1P1 + MP2P2, -M * P1 * P0, -M * P2 * P0, MP2P2 + MP0P0, -M * P2 * P1, MP1P1 + MP0P0);
+			for (const FMassProperties& Child : MPArray)
+			{
+				const FReal M = Child.Mass;
+				const FVec3 ParentToChild = Child.CenterOfMass - NewMP.CenterOfMass;
+				const FReal P0 = ParentToChild[0];
+				const FReal P1 = ParentToChild[1];
+				const FReal P2 = ParentToChild[2];
+				const FReal MP0P0 = M * P0 * P0;
+				const FReal MP1P1 = M * P1 * P1;
+				const FReal MP2P2 = M * P2 * P2;
+				NewMP.InertiaTensor += FMatrix33(MP1P1 + MP2P2, -M * P1 * P0, -M * P2 * P0, MP2P2 + MP0P0, -M * P2 * P1, MP1P1 + MP0P0);
+			}
 		}
+
 		return NewMP;
 	}
 
diff --git a/Engine/Source/Runtime/Experimental/Chaos/Public/Chaos/MassProperties.h b/Engine/Source/Runtime/Experimental/Chaos/Public/Chaos/MassProperties.h
--- a/Engine/Source/Runtime/Experimental/Chaos/Public/Chaos/MassProperties.h
+++ b/Engine/Source/Runtime/Experimental/Chaos/Public/Chaos/MassProperties.h
@@ -34,6 +34,7 @@ namespace Chaos
 	};
 
 	FRotation3 CHAOS_API TransformToLocalSpace(FMatrix33& Inertia);
+	void CHAOS_API TransformToLocalSpace(FMassProperties& MassProperties);
 
 	template<typename T, typename TSurfaces>
 	void CHAOS_API CalculateVolumeAndCenterOfMass(const TParticles<T,3>& Vertices, const TSurfaces& Surfaces, T& OutVolume, TVec3<T>& OutCenterOfMass);
@@ -51,13 +52,14 @@ namespace Chaos
 	void CHAOS_API CalculateVolumeAndCenterOfMass(const FBox& BoundingBox, FVector::FReal& OutVolume, FVector& OutCenterOfMass);
 	void CHAOS_API CalculateInertiaAndRotationOfMass(const FBox& BoundingBox, const FVector::FReal Density, FMatrix33& OutInertiaTensor, FRotation3& OutRotationOfMass);
 
-	// Combine a list of transformed inertia tensors into a single inertia. Also diagonalize the inertia
-	// and set the rotation of mass accordingly.
+	// Combine a list of transformed inertia tensors into a single inertia. Also diagonalize the inertia and set the rotation of mass accordingly. 
+	// This is equivalent to a call to CombineWorldSpace followed by TransformToLocalSpace.
+	// @see CombineWorldSpace()
 	FMassProperties CHAOS_API Combine(const TArray<FMassProperties>& MPArray);
 
-	// Combine a list of transformed inertia tensors into a single inertia. 
-	// NOTE: If there is more than one item in the list, the output may be non-diagonaly and will have a zero rotation.
-	// If there is only 1 item in the list it will return it directly, so the rotation of mass may be non-zero.
+	// Combine a list of transformed inertia tensors into a single inertia tensor.
+	// @note The inertia matrix is not diagonalized, and any rotation will be built into the matrix (RotationOfMass will always be Identity)
+	// @see Combine()
 	FMassProperties CHAOS_API CombineWorldSpace(const TArray<FMassProperties>& MPArray);
 
 
