Chaos - Add Remove at and Insert at Managed Array functions.

diff --git a/Engine/Source/Runtime/Experimental/Chaos/Private/GeometryCollection/ManagedArrayCollection.cpp b/Engine/Source/Runtime/Experimental/Chaos/Private/GeometryCollection/ManagedArrayCollection.cpp
--- a/Engine/Source/Runtime/Experimental/Chaos/Private/GeometryCollection/ManagedArrayCollection.cpp
+++ b/Engine/Source/Runtime/Experimental/Chaos/Private/GeometryCollection/ManagedArrayCollection.cpp
@@ -59,6 +59,16 @@ void FManagedArrayCollection::RemoveElements(const FName& Group, const TArray<in
 	}
 }
 
+void FManagedArrayCollection::RemoveElements(const FName& Group, int32 NumberElements, int32 Position)
+{
+	TArray<int32> SortedDeletionList;
+	SortedDeletionList.SetNumUninitialized(NumberElements);
+	for (int32 Idx = 0; Idx < NumberElements; ++Idx)
+	{
+		SortedDeletionList[Idx] = Position + Idx;
+	}
+	RemoveElements(Group, SortedDeletionList);
+}
 
 TArray<FName> FManagedArrayCollection::GroupNames() const
 {
@@ -122,6 +132,35 @@ int32 FManagedArrayCollection::AddElements(int32 NumberElements, FName Group)
 	return StartSize;
 }
 
+int32 FManagedArrayCollection::InsertElements(int32 NumberElements, int32 Position, FName Group)
+{
+	const int32 OldGroupSize = AddElements(NumberElements, Group);
+	const int32 NewGroupSize = OldGroupSize + NumberElements;
+	check(Position <= OldGroupSize);
+	const int32 NumberElementsToMove = OldGroupSize - Position;
+	const int32 MoveToPosition = Position + NumberElements;
+
+	TArray<int32> NewOrder;
+	NewOrder.SetNumUninitialized(NewGroupSize);
+
+	for (int32 Idx = 0; Idx < Position; ++Idx)
+	{
+		NewOrder[Idx] = Idx;
+	}
+	for (int32 Idx = Position; Idx < MoveToPosition; ++Idx)
+	{
+		NewOrder[Idx] = Idx + NumberElementsToMove;
+	}
+	for (int32 Idx = MoveToPosition; Idx < NewGroupSize; ++Idx)
+	{
+		NewOrder[Idx] = Idx - NumberElements;
+	}
+
+	ReorderElements(Group, NewOrder);
+
+	return Position;
+}
+
 void FManagedArrayCollection::RemoveAttribute(FName Name, FName Group)
 {
 	FKeyType Key = FManagedArrayCollection::MakeMapKey(Name, Group);
diff --git a/Engine/Source/Runtime/Experimental/Chaos/Public/GeometryCollection/ManagedArrayCollection.h b/Engine/Source/Runtime/Experimental/Chaos/Public/GeometryCollection/ManagedArrayCollection.h
--- a/Engine/Source/Runtime/Experimental/Chaos/Public/GeometryCollection/ManagedArrayCollection.h
+++ b/Engine/Source/Runtime/Experimental/Chaos/Public/GeometryCollection/ManagedArrayCollection.h
@@ -161,6 +161,15 @@ public:
 	*/
 	int32 AddElements(int32 NumberElements, FName Group);
 
+	/**
+	* Insert elements to a group
+	* @param NumberElements - The number of array entries to add
+	* @param Position - The position in the managed array where to insert entries.
+	* @param Group - The group to append entries to.
+	* @return starting index of the new ManagedArray<T> entries (same as Position).
+	*/
+	int32 InsertElements(int32 NumberElements, int32 Position, FName Group);
+
 	/**
 	* Returns attribute(Name) of Type(T) from the group
 	* @param Name - The name of the attribute
@@ -254,6 +263,14 @@ public:
 	*/
 	virtual void RemoveElements(const FName & Group, const TArray<int32> & SortedDeletionList, FProcessingParameters Params = FProcessingParameters());
 
+	/**
+	* Remove the elements at Position and reindex the dependent arrays 
+	* @param Group - The group that manages the attributes
+	* @param NumberElements - The number of array entries to remove
+	* @param Position - The position from which to remove entries
+	*/
+	virtual void RemoveElements(const FName& Group, int32 NumberElements, int32 Position);
+
 
 	/**
 	* Remove the attribute from the collection.
