Fix sleeping clusters not being in the SOA active array when waken up by Island manager

diff --git a/Engine/Source/Runtime/Experimental/Chaos/Public/Chaos/PBDRigidsSOAs.h b/Engine/Source/Runtime/Experimental/Chaos/Public/Chaos/PBDRigidsSOAs.h
--- a/Engine/Source/Runtime/Experimental/Chaos/Public/Chaos/PBDRigidsSOAs.h
+++ b/Engine/Source/Runtime/Experimental/Chaos/Public/Chaos/PBDRigidsSOAs.h
@@ -469,6 +469,10 @@ public:
 							bGeometryCollectionDirty = true;
 							return true;
 						}
+						else
+						{
+							AddToActiveArray(PBDRigid);
+						}
 					}
 					else
 					{
@@ -533,6 +537,10 @@ public:
 							FellAsleepGeometryCollectionArray.Insert(PBDRigid);
 							bGeometryCollectionDirty = true;
 						}
+						else
+						{
+							RemoveFromActiveArray(PBDRigid, /*bStillDirty=*/true);
+						}
 					}
 					else
 					{
