Updated GeometryCollection replication so that going in and out of relevancy works and does a hard snap.

diff --git a/Engine/Source/Runtime/Experimental/GeometryCollectionEngine/Private/GeometryCollection/GeometryCollectionComponent.cpp b/Engine/Source/Runtime/Experimental/GeometryCollectionEngine/Private/GeometryCollection/GeometryCollectionComponent.cpp
--- a/Engine/Source/Runtime/Experimental/GeometryCollectionEngine/Private/GeometryCollection/GeometryCollectionComponent.cpp
+++ b/Engine/Source/Runtime/Experimental/GeometryCollectionEngine/Private/GeometryCollection/GeometryCollectionComponent.cpp
@@ -154,21 +154,24 @@ bool FGeometryCollectionRepData::NetSerialize(FArchive& Ar, class UPackageMap* M
 
 	Ar << Version;
 
-	int32 NumPoses = Poses.Num();
-	Ar << NumPoses;
+	Ar << OneOffActivated;
+
+	int32 NumClusters = Clusters.Num();
+	Ar << NumClusters;
 
 	if(Ar.IsLoading())
 	{
-		Poses.SetNum(NumPoses);
+		Clusters.SetNum(NumClusters);
 	}
 
-	for(FGeometryCollectionRepPose& Pose : Poses)
+	for(FGeometryCollectionClusterRep& Cluster : Clusters)
 	{
-		SerializePackedVector<100, 30>(Pose.Position, Ar);
-		SerializePackedVector<100, 30>(Pose.LinearVelocity, Ar);
-		SerializePackedVector<100, 30>(Pose.AngularVelocity, Ar);
-		Pose.Rotation.NetSerialize(Ar, Map, bOutSuccess);
-		Ar << Pose.ParticleIndex;
+		Ar << Cluster.Position;
+		Ar << Cluster.LinearVelocity;
+		Ar << Cluster.AngularVelocity;
+		Ar << Cluster.Rotation;
+		Ar << Cluster.ClusterIdx;
+		Ar << Cluster.ObjectState;
 	}
 
 	return true;
@@ -247,7 +250,7 @@ UGeometryCollectionComponent::UGeometryCollectionComponent(const FObjectInitiali
 	, bStoreVelocities(false)
 	, bShowBoneColors(false)
 	, bEnableReplication(false)
-	, bEnableAbandonAfterLevel(false)
+	, bEnableAbandonAfterLevel(true)
 	, ReplicationAbandonClusterLevel(0)
 	, bRenderStateDirty(true)
 	, bEnableBoneSelection(false)
@@ -268,6 +271,8 @@ UGeometryCollectionComponent::UGeometryCollectionComponent(const FObjectInitiali
 	bTickInEditor = true;
 	bAutoActivate = true;
 
+	bAsyncPhysicsTickEnabled = true;
+
 	static uint32 GlobalNavMeshInvalidationCounter = 0;
 	//space these out over several frames (3 is arbitrary)
 	GlobalNavMeshInvalidationCounter += 3;
@@ -1123,41 +1128,17 @@ void UGeometryCollectionComponent::InitializeComponent()
 
 	// If we're replicating we need some extra setup - check netmode as we don't need this for
 	// standalone runtimes where we aren't going to network the component
-	if(GetIsReplicated() && NetMode != NM_Standalone)
+	if(GetIsReplicated())
 	{
-		if(LocalRole == ENetRole::ROLE_Authority)
+		if(LocalRole != ENetRole::ROLE_Authority)
 		{
-			// As we're the authority we need to track velocities in the dynamic collection so we
-			// can send them over to the other clients to correctly set their state. Attach this now.
-			// The physics proxy will pick them up and populate them as needed
-			if (!DynamicCollection->FindAttributeTyped<FVector3f>("LinearVelocity", FTransformCollection::TransformGroup))
-			{
-				DynamicCollection->AddAttribute<FVector3f>("LinearVelocity", FTransformCollection::TransformGroup);
-			}
 
-			if (!DynamicCollection->FindAttributeTyped<FVector3f>("AngularVelocity", FTransformCollection::TransformGroup))
-			{
-				DynamicCollection->AddAttribute<FVector3f>("AngularVelocity", FTransformCollection::TransformGroup);
-			}
-
-			// We also need to track our control of particles if that control can be shared between server and client
-			if(bEnableAbandonAfterLevel)
-			{
-				TManagedArray<bool>& ControlFlags = DynamicCollection->AddAttribute<bool>("AuthControl", FTransformCollection::TransformGroup);
-				for(bool& Flag : ControlFlags)
-				{
-					Flag = true;
-				}
-			}
-		}
-		else
-		{
 			// We're a replicated component and we're not in control.
 			Chaos::FPhysicsSolver* CurrSolver = GetSolver(*this);
 
 			if(CurrSolver)
 			{
-				CurrSolver->RegisterSimOneShotCallback([Prox = PhysicsProxy]()
+				CurrSolver->RegisterSimOneShotCallback([Prox = PhysicsProxy, ReplicationLevel = ReplicationAbandonClusterLevel, AbandonAfterLevel = bEnableAbandonAfterLevel]()
 				{
 					// As we're not in control we make it so our simulated proxy cannot break clusters
 					// We have to set the strain to a high value but be below the max for the data type
@@ -1173,7 +1154,21 @@ void UGeometryCollectionComponent::InitializeComponent()
 							continue;
 						}
 
-						P->SetStrain(MaxStrain);
+						int32 Level = AbandonAfterLevel ? 0 : -1;
+						if(AbandonAfterLevel)
+						{
+							Chaos::TPBDRigidClusteredParticleHandle<Chaos::FReal, 3>* Current = P;
+							while ( Current->Parent())
+							{
+								Current = Current->Parent();
+								++Level;
+							}
+						}
+
+						if (Level <= ReplicationLevel + 1)	//we only replicate up until level X, but it means we should replicate the breaking event of level X+1 (but not X+1's positions)
+						{
+							P->SetStrain(MaxStrain);
+						}
 					}
 				});
 			}
@@ -1311,140 +1306,206 @@ void ActivateClusters(Chaos::FRigidClustering& Clustering, Chaos::TPBDRigidClust
 	Clustering.DeactivateClusterParticle(Cluster);
 }
 
-void UGeometryCollectionComponent::OnRep_RepData(const FGeometryCollectionRepData& OldData)
+void UGeometryCollectionComponent::UpdateRepData()
 {
-	if(!DynamicCollection)
+	using namespace Chaos;
+	if(!bEnableReplication)
 	{
 		return;
 	}
 
-	if(AActor* Owner = GetOwner())
+	AActor* Owner = GetOwner();
+	
+	// If we have no owner or our netmode means we never require replication then early out
+	if(!Owner || Owner->GetNetMode() == ENetMode::NM_Standalone)
+	{
+		return;
+	}
+	
+	if (Owner && GetIsReplicated() && Owner->GetLocalRole() == ROLE_Authority)
 	{
-		const int32 NumTransforms = DynamicCollection->Transform.Num();
-		const int32 NumNewPoses = RepData.Poses.Num();
-		if(NumTransforms < NumNewPoses)
+		if(ClustersToRep == nullptr)
 		{
-			return;
+			//we only allocate set if needed because it's pretty big to have per components that don't replicate
+			ClustersToRep = MakeUnique<TSet<Chaos::FPBDRigidClusteredParticleHandle*>>();
 		}
 
-		Chaos::FPhysicsSolver* Solver = GetSolver(*this);
+		//We need to build a snapshot of the GC
+		//We rely on the fact that clusters always fracture with one off pieces being removed.
+		//This means we only need to record the one offs that broke and we get the connected components for free
+		//The cluster properties are replicated with the first child of each connected component. These are always children that are known at author time and have a unique id per component
+		//If the first child is disabled it means the properties apply to the parent (i.e. the cluster)
+		//If the first child is enabled it means it's a one off and the cluster IS the first child
+		
+		//TODO: for now we have to iterate over all particles to find the clusters, would be better if we had the clusters and children already available
+		//Large refactor happening to this stuff so for now we just iterate
+		//We are relying on the fact that we fracture one level per step. This means we will see all one offs here
 
-		for(int32 Index = 0; Index < NumNewPoses; ++Index)
-		{
-			const FGeometryCollectionRepPose& SourcePose = RepData.Poses[Index];
-			const int32 ParticleIndex = SourcePose.ParticleIndex;
+		bool bClustersChanged = false;
 
-			if(ParticleIndex >= NumTransforms)
+		const FPBDRigidsSolver* Solver = PhysicsProxy->GetSolver<Chaos::FPBDRigidsSolver>();
+		const FRigidClustering& RigidClustering = Solver->GetEvolution()->GetRigidClustering();
+
+		//see if we have any new clusters that are enabled
+		TSet<FPBDRigidClusteredParticleHandle*> Processed;
+		for (FPBDRigidClusteredParticleHandle* Particle : PhysicsProxy->GetParticles())
+		{
+			bool bProcess = true;
+			Processed.Add(Particle);
+			FPBDRigidClusteredParticleHandle* Root = Particle;
+			while (Root->Parent())
 			{
-				// Out of range
-				continue;
+				Root = Root->Parent();
+
+				//TODO: set avoids n^2, would be nice if clustered particle cached its root
+				if(Processed.Contains(Root))
+				{
+					bProcess = false;
+					break;
+				}
+				else
+				{
+					Processed.Add(Root);
+				}
 			}
 
-			Solver->RegisterSimOneShotCallback([SourcePose, Prox = PhysicsProxy]()
+			if (bProcess && Root->Disabled() == false && ClustersToRep->Find(Root) == nullptr)
 			{
-				Chaos::TPBDRigidClusteredParticleHandle<Chaos::FReal, 3>* Particle = Prox->GetParticles()[SourcePose.ParticleIndex];
+				//first time in here so needs a new count
+				//TODO: check root needs to replicate if abandon by level is enabled
+				ClustersToRep->Add(Root);
+				if(Root->InternalCluster() == false)
+				{
+					//a one off so record it
+					const int32 TransformGroupIdx = PhysicsProxy->GetTransformGroupIndexFromHandle(Root);
+					ensureMsgf(TransformGroupIdx >= 0, TEXT("Non-internal cluster should always have a group index"));
+					ensureMsgf(TransformGroupIdx < TNumericLimits<uint16>::Max(), TEXT("Trying to replicate GC with more than 65k pieces. We assumed uint16 would suffice"));
+					RepData.OneOffActivated.Add(TransformGroupIdx);
+					bClustersChanged = true;
+				}
+			}
+		}
 
-				Chaos::FPhysicsSolver* Solver = Prox->GetSolver<Chaos::FPhysicsSolver>();
-				Chaos::FPBDRigidsEvolution* Evo = Solver->GetEvolution();
-				check(Evo);
-				Chaos::FRigidClustering& Clustering = Evo->GetRigidClustering();
-				
-				// Set X/R/V/W for next sim step from the replicated state
-				Particle->SetX(SourcePose.Position);
-				Particle->SetR(SourcePose.Rotation);
-				Particle->SetV(SourcePose.LinearVelocity);
-				Particle->SetW(SourcePose.AngularVelocity);
+		//build up clusters to replicate and compare with previous frame
+		TArray<FGeometryCollectionClusterRep> Clusters;
+
+		//remove disabled clusters and update rep data if needed
+		for (auto Itr = ClustersToRep->CreateIterator(); Itr; ++Itr)
+		{
+			FPBDRigidClusteredParticleHandle* Cluster = *Itr;
+			if (Cluster->Disabled())
+			{
+				Itr.RemoveCurrent();
+			}
+			else
+			{
+				Clusters.AddDefaulted();
+				FGeometryCollectionClusterRep& ClusterRep = Clusters.Last();
 
-				if(Particle->ClusterIds().Id)
+				ClusterRep.Position = Cluster->X();
+				ClusterRep.Rotation = Cluster->R();
+				ClusterRep.LinearVelocity = Cluster->V();
+				ClusterRep.AngularVelocity = Cluster->W();
+				ClusterRep.ObjectState = static_cast<int8>(Cluster->ObjectState());
+				int32 TransformGroupIdx;
+				if(Cluster->InternalCluster())
 				{
-					// This particle is clustered but the remote authority has it activated. Fracture the parent cluster
-					ActivateClusters(Clustering, Particle->Parent());
+					const TArray<FPBDRigidParticleHandle*>& Children = RigidClustering.GetChildrenMap()[Cluster];
+					ensureMsgf(Children.Num(), TEXT("Internal cluster yet we have no children?"));
+					TransformGroupIdx = PhysicsProxy->GetTransformGroupIndexFromHandle(Children[0]);
 				}
-				else if(Particle->Disabled())
+				else
 				{
-					// We might have disabled the particle - need to reactivate if it's active on the remote.
-					Particle->SetDisabled(false);
+					//not internal so we can just use the cluster's ID. On client we'll know based on the parent whether to use this index or the parent
+					TransformGroupIdx = PhysicsProxy->GetTransformGroupIndexFromHandle(Cluster);
 				}
 
-				// Make sure to wake corrected particles
-				Particle->SetSleeping(false);
-			});
+				ensureMsgf(TransformGroupIdx < TNumericLimits<uint16>::Max(), TEXT("Trying to replicate GC with more than 65k pieces. We assumed uint16 would suffice"));
+				ClusterRep.ClusterIdx = TransformGroupIdx;
+
+				if(!bClustersChanged)
+				{
+					//compare to previous frame data
+					if(RepData.Clusters.Num() >= Clusters.Num())
+					{
+						const FGeometryCollectionClusterRep& PrevCluster = RepData.Clusters[Clusters.Num() - 1];
+						if(ClusterRep.ClusterChanged(PrevCluster))
+						{
+							bClustersChanged = true;
+						}
+					}
+					else
+					{
+						//must be some new clusters so definitely changed
+						bClustersChanged = true;
+					}
+				}
+			}
+		}
+
+		if (bClustersChanged)
+		{
+			RepData.Clusters = MoveTemp(Clusters);
+
+			MARK_PROPERTY_DIRTY_FROM_NAME(UGeometryCollectionComponent, RepData, this);
+			++RepData.Version;
 		}
 	}
 }
 
-void UGeometryCollectionComponent::UpdateRepData()
+int32 GeometryCollectionHardMissingUpdatesSnapThreshold = 20;
+FAutoConsoleVariableRef CVarGeometryCollectionHardMissingUpdatesSnapThreshold(TEXT("p.GeometryCollectionHardMissingUpdatesSnapThreshold"), GeometryCollectionHardMissingUpdatesSnapThreshold,
+	TEXT("Determines how many missing updates before we trigger a hard snap"));
+
+void UGeometryCollectionComponent::ProcessRepData()
 {
-	if(!bEnableReplication)
+	using namespace Chaos;
+	if(VersionProcessed == RepData.Version)
 	{
 		return;
 	}
 
-	AActor* Owner = GetOwner();
-	
-	// If we have no owner or our netmode means we never require replication then early out
-	if(!Owner || Owner->GetNetMode() == ENetMode::NM_Standalone)
+	bool bHardSnap = false;
+	if(VersionProcessed < RepData.Version)
 	{
-		return;
+		//TODO: this will not really work if a fracture happens and then immediately goes to sleep without updating client enough times
+		//A time method would work better here, but is limited to async mode. Maybe we can support both
+		bHardSnap = (RepData.Version - VersionProcessed) > GeometryCollectionHardMissingUpdatesSnapThreshold;
 	}
-	
-	if(Owner && GetIsReplicated() && Owner->GetLocalRole() == ROLE_Authority)
+	else
 	{
-		// We're inside a replicating actor and we're the authority - update the rep data
-		const int32 NumTransforms = DynamicCollection->Transform.Num();
-		RepData.Poses.Reset(NumTransforms);
-
-		TManagedArray<FVector3f>* LinearVelocity = DynamicCollection->FindAttributeTyped<FVector3f>("LinearVelocity", FTransformCollection::TransformGroup);
-		TManagedArray<FVector3f>* AngularVelocity = DynamicCollection->FindAttributeTyped<FVector3f>("AngularVelocity", FTransformCollection::TransformGroup);
-
-		for(int32 Index = 0; Index < NumTransforms; ++Index)
-		{
-			TManagedArray<TUniquePtr<Chaos::FGeometryParticle>>& GTParticles = PhysicsProxy->GetExternalParticles();
-			Chaos::FGeometryParticle* Particle = GTParticles[Index].Get();
-			if(!DynamicCollection->Active[Index] || DynamicCollection->DynamicState[Index] != static_cast<uint8>(Chaos::EObjectStateType::Dynamic))
-			{
-				continue;
-			}
-
-			const int32 ClusterLevel = GetClusterLevel(RestCollection->GetGeometryCollection().Get(), Index);
-			const bool bLevelValid = !EnableClustering || !bEnableAbandonAfterLevel || ClusterLevel <= ReplicationAbandonClusterLevel;
-			if(!bLevelValid)
-			{
-				const int32 ParentTransformIndex = RestCollection->GetGeometryCollection()->Parent[Index];
-				TManagedArray<bool>* ControlFlags = DynamicCollection->FindAttributeTyped<bool>("AuthControl", FTransformCollection::TransformGroup);
-
-				if(ControlFlags && (*ControlFlags)[ParentTransformIndex])
-				{
-					(*ControlFlags)[ParentTransformIndex] = false;
-					NetAbandonCluster(ParentTransformIndex);
-				}
+		//rollover so just treat as hard snap - this case is extremely rare and a one off
+		bHardSnap = true;
+	}
 
-				continue;
-			}
+	FPBDRigidsSolver* Solver = PhysicsProxy->GetSolver<Chaos::FPBDRigidsSolver>();
+	FRigidClustering& RigidClustering = Solver->GetEvolution()->GetRigidClustering();
 
-			RepData.Poses.AddDefaulted();
-			FGeometryCollectionRepPose& Pose = RepData.Poses.Last();
+	//First make sure all one off activations have been applied. This ensures our connectivity graph is the same and we have the same clusters as the server
+	for (; OneOffActivatedProcessed < RepData.OneOffActivated.Num(); ++OneOffActivatedProcessed)
+	{
+		FPBDRigidParticleHandle* OneOff = PhysicsProxy->GetParticles()[RepData.OneOffActivated[OneOffActivatedProcessed]];
+		RigidClustering.ReleaseClusterParticles(TArray<FPBDRigidParticleHandle*>{ OneOff });
+	}
 
-			// No scale transfered - shouldn't be a simulated property
-			Pose.ParticleIndex = Index;
-			Pose.Position = Particle->X();
-			Pose.Rotation = Particle->R();
-			if(LinearVelocity)
-			{
-				check(AngularVelocity);
-				Pose.LinearVelocity = (FVector)(*LinearVelocity)[Index];
-				Pose.AngularVelocity = (FVector)(*AngularVelocity)[Index];
-			}
-			else
+	if(bHardSnap)
+	{
+		for (const FGeometryCollectionClusterRep& RepCluster : RepData.Clusters)
+		{
+			FPBDRigidParticleHandle* Cluster = PhysicsProxy->GetParticles()[RepCluster.ClusterIdx];
+			if(Cluster->Disabled() == false)
 			{
-				Pose.LinearVelocity = FVector::ZeroVector;
-				Pose.AngularVelocity = FVector::ZeroVector;
+				Cluster->SetX(RepCluster.Position);
+				Cluster->SetR(RepCluster.Rotation);
+				Cluster->SetV(RepCluster.LinearVelocity);
+				Cluster->SetW(RepCluster.AngularVelocity);
+				//TODO: snap object state too once we fix interpolation
+				//Solver->GetEvolution()->SetParticleObjectState(Cluster, static_cast<Chaos::EObjectStateType>(RepCluster.ObjectState));
 			}
 		}
-
-		RepData.Version++;
-		MARK_PROPERTY_DIRTY_FROM_NAME(UGeometryCollectionComponent, RepData, this);
 	}
+
+	VersionProcessed = RepData.Version;
 }
 
 void UGeometryCollectionComponent::SetDynamicState(const Chaos::EObjectStateType& NewDynamicState)
@@ -1512,49 +1573,6 @@ void SetHierarchyStrain(Chaos::TPBDRigidClusteredParticleHandle<Chaos::FReal, 3>
 	}
 }
 
-void UGeometryCollectionComponent::NetAbandonCluster_Implementation(int32 TransformIndex)
-{
-	// Called on clients when the server abandons a particle. TransformIndex is the index of the parent
-	// of that particle, should only get called once per cluster but survives multiple calls
-	
-	if(GetOwnerRole() == ENetRole::ROLE_Authority)
-	{
-		// Owner called abandon - takes no action
-		return;
-	}
-
-	if(!EnableClustering)
-	{
-		// No clustering information to update
-		return;
-	}
-
-	if(TransformIndex >= 0 && TransformIndex < DynamicCollection->NumElements(FTransformCollection::TransformGroup))
-	{
-		int32 ClusterLevel = GetClusterLevel(RestCollection->GetGeometryCollection().Get(), TransformIndex);
-		float Strain = DamageThreshold.IsValidIndex(ClusterLevel) ? DamageThreshold[ClusterLevel] : DamageThreshold.Num() > 0 ? DamageThreshold[0] : 0.0f;
-
-		if(Strain >= 0)
-		{
-			Chaos::FPhysicsSolver* Solver = GetSolver(*this);
-
-			Solver->RegisterSimOneShotCallback([Prox = PhysicsProxy, Strain, TransformIndex, Solver]()
-			{
-				Chaos::FRigidClustering& Clustering = Solver->GetEvolution()->GetRigidClustering();
-				Chaos::FPBDRigidClusteredParticleHandle* Parent = Prox->GetParticles()[TransformIndex];
-
-				if(!Parent->Disabled())
-				{
-					SetHierarchyStrain(Parent, Clustering.GetChildrenMap(), Strain);
-
-					// We know the server must have fractured this cluster, so repeat here
-					Clustering.DeactivateClusterParticle(Parent);
-				}
-			});
-		}
-	}
-}
-
 void UGeometryCollectionComponent::InitConstantData(FGeometryCollectionConstantData* ConstantData) const
 {
 	// Constant data should all be moved to the DDC as time permits.
@@ -1928,7 +1946,14 @@ void UGeometryCollectionComponent::TickComponent(float DeltaTime, enum ELevelTic
 		}
 	}
 #endif
+}
+
+void UGeometryCollectionComponent::AsyncPhysicsTickComponent(float DeltaTime, float SimTime)
+{
+	Super::AsyncPhysicsTickComponent(DeltaTime, SimTime);
 
+	UpdateRepData();
+	ProcessRepData();
 }
 
 void UGeometryCollectionComponent::OnRegister()
diff --git a/Engine/Source/Runtime/Experimental/GeometryCollectionEngine/Public/GeometryCollection/GeometryCollectionComponent.h b/Engine/Source/Runtime/Experimental/GeometryCollectionEngine/Public/GeometryCollection/GeometryCollectionComponent.h
--- a/Engine/Source/Runtime/Experimental/GeometryCollectionEngine/Public/GeometryCollection/GeometryCollectionComponent.h
+++ b/Engine/Source/Runtime/Experimental/GeometryCollectionEngine/Public/GeometryCollection/GeometryCollectionComponent.h
@@ -284,15 +284,27 @@ public:
  * Raw struct to serialize for network. We need to custom netserialize to optimize
  * the vector serialize as much as possible and rather than have the property system
  * iterate an array of reflected structs we handle everything in the NetSerialize for
- * the container (FGeometryCollectionRepData)
+ * the container (FGeometryCollectionClusterRep)
  */
-struct FGeometryCollectionRepPose
+struct FGeometryCollectionClusterRep
 {
-	FVector Position;
-	FVector LinearVelocity;
-	FVector AngularVelocity;
+	FVector_NetQuantize100 Position;
+	FVector_NetQuantize100 LinearVelocity;
+	FVector_NetQuantize100 AngularVelocity;
 	FQuat Rotation;
-	uint16 ParticleIndex;
+	uint16 ClusterIdx;
+	int8 ObjectState;
+
+	bool ClusterChanged(const FGeometryCollectionClusterRep& Other) const
+	{
+		return Other.ObjectState != ObjectState
+			|| Other.ClusterIdx != ClusterIdx
+			|| Other.ObjectState != ObjectState
+			|| Other.Position != Position
+			|| Other.LinearVelocity != LinearVelocity
+			|| Other.AngularVelocity != AngularVelocity
+			|| Other.Rotation != Rotation;
+	}
 };
 
 /**
@@ -310,8 +322,11 @@ struct FGeometryCollectionRepData
 
 	}
 
-	// Array of per-particle data required to synchronize clients
-	TArray<FGeometryCollectionRepPose> Poses;
+	//Array of one off pieces that became activated
+	TArray<uint16> OneOffActivated;
+
+	// Array of cluster data requires to synchronize clients
+	TArray<FGeometryCollectionClusterRep> Clusters;
 
 	// Version counter, every write to the rep data is a new state so Identical only references this version
 	// as there's no reason to compare the Poses array.
@@ -368,6 +383,7 @@ public:
 
 	virtual bool HasAnySockets() const override { return false; }
 	virtual void TickComponent(float DeltaTime, enum ELevelTick TickType, FActorComponentTickFunction *ThisTickFunction) override;
+	virtual void AsyncPhysicsTickComponent(float DeltaTime, float SimTime) override;
 	//~ Begin USceneComponent Interface.
 
 
@@ -786,25 +802,14 @@ protected:
 	UPROPERTY(EditInstanceOnly, BlueprintReadOnly, Category = Network)
 	int32 ReplicationAbandonClusterLevel;
 
-	UPROPERTY(ReplicatedUsing=OnRep_RepData)
+	UPROPERTY(Replicated)
 	FGeometryCollectionRepData RepData;
 
-	/** Called on non-authoritative clients when receiving new repdata from the server */
-	UFUNCTION()
-	void OnRep_RepData(const FGeometryCollectionRepData& OldData);
-
 	/** Called post solve to allow authoritative components to update their replication data */
 	void UpdateRepData();
 
 private:
 
-	/** 
-	 * Notifies all clients that a server has abandoned control of a particle, clients should restore the strain
-	 * values on abandoned particles and their children then fracture them before continuing
-	 */
-	UFUNCTION(NetMulticast, Reliable)
-	void NetAbandonCluster(int32 TransformIndex);
-
 	bool bRenderStateDirty;
 	bool bEnableBoneSelection;
 	int ViewLevel;
@@ -875,6 +880,14 @@ private:
 
 	void IncrementSleepTimer(float DeltaTime);
 	bool CalculateInnerSphere(int32 TransformIndex, FSphere& SphereOut) const;
+	void ProcessRepData();
+
+	/** The clusters we need to replicate */
+	TUniquePtr<TSet<Chaos::FPBDRigidClusteredParticleHandle*>> ClustersToRep;
+
+	/** One off activation is processed in the same order as server so remember the last one we processed */
+	int32 OneOffActivatedProcessed = 0;
+	int32 VersionProcessed = INDEX_NONE;
 
 	/** True if GeometryCollection transforms have changed from previous tick. */
 	bool bIsMoving;
