Chaos : fix moving kinematic GCs
- make sure the proxy is marked dirty when updating the world transform
- make sure to update the cluster particle in oreder to update the physics thread acceleration structure

diff --git a/Engine/Source/Runtime/Experimental/Chaos/Private/PhysicsProxy/GeometryCollectionPhysicsProxy.cpp b/Engine/Source/Runtime/Experimental/Chaos/Private/PhysicsProxy/GeometryCollectionPhysicsProxy.cpp
--- a/Engine/Source/Runtime/Experimental/Chaos/Private/PhysicsProxy/GeometryCollectionPhysicsProxy.cpp
+++ b/Engine/Source/Runtime/Experimental/Chaos/Private/PhysicsProxy/GeometryCollectionPhysicsProxy.cpp
@@ -1564,6 +1564,14 @@ void FGeometryCollectionPhysicsProxy::SetWorldTransform(const FTransform& WorldT
 {
 	check(IsInGameThread());
 	GameThreadPerFrameData.SetWorldTransform(WorldTransform);
+
+	if (Chaos::FPhysicsSolver* RBDSolver = GetSolver<Chaos::FPhysicsSolver>())
+	{
+		RBDSolver->EnqueueCommandImmediate([this, RBDSolver]()
+		{
+			RBDSolver->AddDirtyProxy(this);
+		});
+	}
 }
 
 void FGeometryCollectionPhysicsProxy::PushStateOnGameThread(Chaos::FPBDRigidsSolver* InSolver)
@@ -1636,6 +1644,7 @@ void FGeometryCollectionPhysicsProxy::SetClusteredParticleKinematicTarget_Intern
 		if (Chaos::FPhysicsSolver* RBDSolver = GetSolver<Chaos::FPhysicsSolver>())
 		{
 			RBDSolver->GetEvolution()->SetParticleKinematicTarget(Handle, NewKinematicTarget);
+			RBDSolver->GetEvolution()->DirtyParticle(*Handle);
 		}
 	}
 }
