Chaos Cloth: add "buckling" stiffness parameters for bending elements.

This is just biphasic stiffnesses for the bending elements. After the edge has bend beyond a certain amount (controlled  by [REDACTED_PERSON]), it's considered to be "Buckled" and will using the "BucklingStiffness" instead of the "BendingStiffness"

Currently the BucklingStiffness value is also paintable, but the Buckling Ratio is not.

Updated the Bending Constraint debug draw to draw the edges colored based on whether or not they are buckled.

diff --git a/Engine/Plugins/Experimental/ChaosCloth/Source/Chaos/Private/ChaosCloth/ChaosClothConstraints.cpp b/Engine/Plugins/Experimental/ChaosCloth/Source/Chaos/Private/ChaosCloth/ChaosClothConstraints.cpp
--- a/Engine/Plugins/Experimental/ChaosCloth/Source/Chaos/Private/ChaosCloth/ChaosClothConstraints.cpp
+++ b/Engine/Plugins/Experimental/ChaosCloth/Source/Chaos/Private/ChaosCloth/ChaosClothConstraints.cpp
@@ -156,6 +156,11 @@ void FClothConstraints::CreateRules()
 	}
 	if (BendingElementConstraints)
 	{
+		ConstraintInits[ConstraintInitIndex++] =
+			[this](Softs::FSolverParticles& /*Particles*/, const Softs::FSolverReal Dt)
+			{
+				BendingElementConstraints->ApplyProperties(Dt, Evolution->GetIterations());
+			};
 		ConstraintRules[ConstraintRuleIndex++] =
 			[this](Softs::FSolverParticles& Particles, const Softs::FSolverReal Dt)
 			{
@@ -347,11 +352,31 @@ void FClothConstraints::SetBendingConstraints(const TArray<TVec2<int32>>& Edges,
 	++NumConstraintRules;
 }
 
+void FClothConstraints::SetBendingConstraints(TArray<TVec4<int32>>&& BendingElements, const TConstArrayView<FRealSingle>& StiffnessMultipliers, const TConstArrayView<FRealSingle>& BucklingStiffnessMultipliers, bool bUseXPBDConstraints)
+{
+	check(Evolution);
+
+	// TODO Add XPBD
+	BendingElementConstraints = MakeShared<Softs::FPBDBendingConstraints>(
+		Evolution->Particles(),
+		ParticleOffset, NumParticles,
+		MoveTemp(BendingElements),
+		StiffnessMultipliers,
+		BucklingStiffnessMultipliers,
+		/*InStiffness =*/ Softs::FSolverVec2::UnitVector,
+		/*InBucklingRatio=*/ (Softs::FSolverReal)0.f,
+		/*InBucklingStiffnessMultiplier =*/ Softs::FSolverVec2::UnitVector,
+		/*bTrimKinematicConstraints =*/ true);
+	++NumConstraintInits;  // Uses init to update the property tables
+	++NumConstraintRules;
+}
+
 void FClothConstraints::SetBendingConstraints(TArray<TVec4<int32>>&& BendingElements, Softs::FSolverReal BendingStiffness)
 {
 	check(Evolution);
 
-	BendingElementConstraints = MakeShared<Softs::FPBDBendingConstraints>(Evolution->Particles(), MoveTemp(BendingElements), BendingStiffness);  // TODO: Strip kinematic constraints
+	BendingElementConstraints = MakeShared<Softs::FPBDBendingConstraints>(Evolution->Particles(), MoveTemp(BendingElements), BendingStiffness);
+	++NumConstraintInits;  // Uses init to update the property tables
 	++NumConstraintRules;
 }
 
@@ -520,7 +545,7 @@ void FClothConstraints::SetEdgeProperties(const Softs::FSolverVec2& EdgeStiffnes
 	}
 }
 
-void FClothConstraints::SetBendingProperties(const Softs::FSolverVec2& BendingStiffness)
+void FClothConstraints::SetBendingProperties(const Softs::FSolverVec2& BendingStiffness, Softs::FSolverReal BucklingRatio, const Softs::FSolverVec2& BucklingStiffness)
 {
 	if (BendingConstraints)
 	{
@@ -530,6 +555,10 @@ void FClothConstraints::SetBendingProperties(const Softs::FSolverVec2& BendingSt
 	{
 		XBendingConstraints->SetProperties(BendingStiffness);
 	}
+	if (BendingElementConstraints)
+	{
+		BendingElementConstraints->SetProperties(BendingStiffness, BucklingRatio, BucklingStiffness);
+	}
 }
 
 void FClothConstraints::SetAreaProperties(const Softs::FSolverVec2& AreaStiffness)
diff --git a/Engine/Plugins/Experimental/ChaosCloth/Source/Chaos/Private/ChaosCloth/ChaosClothingSimulation.cpp b/Engine/Plugins/Experimental/ChaosCloth/Source/Chaos/Private/ChaosCloth/ChaosClothingSimulation.cpp
--- a/Engine/Plugins/Experimental/ChaosCloth/Source/Chaos/Private/ChaosCloth/ChaosClothingSimulation.cpp
+++ b/Engine/Plugins/Experimental/ChaosCloth/Source/Chaos/Private/ChaosCloth/ChaosClothingSimulation.cpp
@@ -31,6 +31,7 @@
 #include "Chaos/PBDSpringConstraints.h"
 #include "Chaos/PBDCollisionSpringConstraints.h"
 #include "Chaos/PBDTriangleMeshCollisions.h"
+#include "Chaos/PBDBendingConstraints.h"
 #include "Chaos/VelocityField.h"
 #endif  // #if WITH_EDITOR || CHAOS_DEBUG_DRAW
 
@@ -395,6 +396,8 @@ void FClothingSimulation::CreateActor(USkeletalMeshComponent* InOwnerComponent,
 		ClothConfig->MinPerParticleMass,
 		TVec2<FRealSingle>(ClothConfig->EdgeStiffnessWeighted.Low, ClothConfig->EdgeStiffnessWeighted.High),
 		TVec2<FRealSingle>(ClothConfig->BendingStiffnessWeighted.Low, ClothConfig->BendingStiffnessWeighted.High),
+		ClothConfig->BucklingRatio,
+		TVec2<FRealSingle>(ClothConfig->BucklingStiffnessWeighted.Low, ClothConfig->BucklingStiffnessWeighted.High),
 		ClothConfig->bUseBendingElements,
 		TVec2<FRealSingle>(ClothConfig->AreaStiffnessWeighted.Low, ClothConfig->AreaStiffnessWeighted.High),
 		ClothConfig->VolumeStiffness,
@@ -819,6 +822,8 @@ void FClothingSimulation::RefreshClothConfig(const IClothingSimulationContext* I
 			ClothConfig->MinPerParticleMass,
 			TVec2<FRealSingle>(ClothConfig->EdgeStiffnessWeighted.Low, ClothConfig->EdgeStiffnessWeighted.High),
 			TVec2<FRealSingle>(ClothConfig->BendingStiffnessWeighted.Low, ClothConfig->BendingStiffnessWeighted.High),
+			ClothConfig->BucklingRatio,
+			TVec2<FRealSingle>(ClothConfig->BucklingStiffnessWeighted.Low, ClothConfig->BucklingStiffnessWeighted.High),
 			ClothConfig->bUseBendingElements,
 			TVec2<FRealSingle>(ClothConfig->AreaStiffnessWeighted.Low, ClothConfig->AreaStiffnessWeighted.High),
 			ClothConfig->VolumeStiffness,
@@ -1840,6 +1845,29 @@ void FClothingSimulation::DebugDrawBendingConstraint(FPrimitiveDrawInterface* PD
 				DrawLine(PDI, Pos0, Pos1, FLinearColor::Black);
 			}
 		}
+
+		if (const Softs::FPBDBendingConstraints* const BendingConstraints = ClothConstraints.GetBendingElementConstraints().Get())
+		{
+			const TArray<TVec4<int32>>& Constraints = BendingConstraints->GetConstraints();
+			const TArray<FSolverReal>& RestAngles = BendingConstraints->GetRestAngles();
+
+			// Color constraint edge with red or blue: Red = Buckled, Blue = Not Buckled. This is based on its current configuration, not the values used
+			// in the actual simulation.
+			for (int32 ConstraintIndex = 0; ConstraintIndex < Constraints.Num(); ++ConstraintIndex)
+			{
+				const Softs::FSolverVec3& P1 = Positions[Constraints[ConstraintIndex][0]];
+				const Softs::FSolverVec3& P2 = Positions[Constraints[ConstraintIndex][1]];
+				const Softs::FSolverVec3& P3 = Positions[Constraints[ConstraintIndex][2]];
+				const Softs::FSolverVec3& P4 = Positions[Constraints[ConstraintIndex][3]];
+				const Softs::FSolverReal Angle = Softs::FPBDBendingConstraintsBase::CalcAngle(P1, P2, P3, P4);
+				
+				const bool bIsBuckled = BendingConstraints->AngleIsBuckled(Angle, RestAngles[ConstraintIndex]);
+
+				const FVec3 Pos0 = FVec3(P1) + LocalSpaceLocation;
+				const FVec3 Pos1 = FVec3(P2) + LocalSpaceLocation;
+				DrawLine(PDI, Pos0, Pos1, bIsBuckled ? FLinearColor::Red : FLinearColor::Blue);
+			}
+		}
 	}
 }
 
diff --git a/Engine/Plugins/Experimental/ChaosCloth/Source/Chaos/Private/ChaosCloth/ChaosClothingSimulationCloth.cpp b/Engine/Plugins/Experimental/ChaosCloth/Source/Chaos/Private/ChaosCloth/ChaosClothingSimulationCloth.cpp
--- a/Engine/Plugins/Experimental/ChaosCloth/Source/Chaos/Private/ChaosCloth/ChaosClothingSimulationCloth.cpp
+++ b/Engine/Plugins/Experimental/ChaosCloth/Source/Chaos/Private/ChaosCloth/ChaosClothingSimulationCloth.cpp
@@ -133,12 +133,14 @@ void FClothingSimulationCloth::FLODData::Add(FClothingSimulationSolver* Solver,
 
 	// Bending constraints
 	const TConstArrayView<FRealSingle>& BendingStiffnessMultipliers = WeightMaps[(int32)EChaosWeightMapTarget::BendingStiffness];
-	if (Cloth->BendingStiffness[0] > (FRealSingle)0. || (Cloth->BendingStiffness[1] > (FRealSingle)0. && BendingStiffnessMultipliers.Num() == NumParticles))
+	const TConstArrayView<FRealSingle>& BucklingStiffnessMultipliers = WeightMaps[(int32)EChaosWeightMapTarget::BucklingStiffness];
+	if (Cloth->BendingStiffness[0] > (FRealSingle)0. || (Cloth->BendingStiffness[1] > (FRealSingle)0. && BendingStiffnessMultipliers.Num() == NumParticles) ||
+		Cloth->BucklingStiffness[0] > (FRealSingle)0. || (Cloth->BucklingStiffness[1] > (FRealSingle)0. && BucklingStiffnessMultipliers.Num() == NumParticles))
 	{
 		if (Cloth->bUseBendingElements)
 		{
 			TArray<Chaos::TVec4<int32>> BendingElements = TriangleMesh.GetUniqueAdjacentElements();
-			ClothConstraints.SetBendingConstraints(MoveTemp(BendingElements), (Softs::FSolverReal)Cloth->BendingStiffness[0]);  // TODO: Add support for weight maps
+			ClothConstraints.SetBendingConstraints(MoveTemp(BendingElements), BendingStiffnessMultipliers, BucklingStiffnessMultipliers, Cloth->bUseXPBDConstraints);
 		}
 		else
 		{
@@ -280,7 +282,7 @@ void FClothingSimulationCloth::FLODData::Update(FClothingSimulationSolver* Solve
 	FClothConstraints& ClothConstraints = Solver->GetClothConstraints(Offset);
 	ClothConstraints.SetMaximumDistanceProperties((Softs::FSolverReal)Cloth->MaxDistancesMultiplier * MeshScale);
 	ClothConstraints.SetEdgeProperties(Softs::FSolverVec2((Softs::FSolverReal)Cloth->EdgeStiffness[0], (Softs::FSolverReal)Cloth->EdgeStiffness[1]));
-	ClothConstraints.SetBendingProperties(Softs::FSolverVec2((Softs::FSolverReal)Cloth->BendingStiffness[0], (Softs::FSolverReal)Cloth->BendingStiffness[1]));
+	ClothConstraints.SetBendingProperties(Softs::FSolverVec2((Softs::FSolverReal)Cloth->BendingStiffness[0], (Softs::FSolverReal)Cloth->BendingStiffness[1]), Cloth->BucklingRatio, Softs::FSolverVec2((Softs::FSolverReal)Cloth->BucklingStiffness[0], (Softs::FSolverReal)Cloth->BucklingStiffness[1]));
 	ClothConstraints.SetAreaProperties(Softs::FSolverVec2((FSolverReal)Cloth->AreaStiffness[0], (Softs::FSolverReal)Cloth->AreaStiffness[1]));
 	ClothConstraints.SetLongRangeAttachmentProperties(
 		Softs::FSolverVec2((Softs::FSolverReal)Cloth->TetherStiffness[0], (Softs::FSolverReal)Cloth->TetherStiffness[1]),
@@ -340,6 +342,8 @@ FClothingSimulationCloth::FClothingSimulationCloth(
 	FRealSingle InMinPerParticleMass,
 	const TVec2<FRealSingle>& InEdgeStiffness,
 	const TVec2<FRealSingle>& InBendingStiffness,
+	FRealSingle InBucklingRatio,
+	const TVec2<FRealSingle>& InBucklingStiffness,
 	bool bInUseBendingElements,
 	const TVec2<FRealSingle>& InAreaStiffness,
 	FRealSingle InVolumeStiffness,
@@ -381,6 +385,8 @@ FClothingSimulationCloth::FClothingSimulationCloth(
 	, MinPerParticleMass(InMinPerParticleMass)
 	, EdgeStiffness(InEdgeStiffness)
 	, BendingStiffness(InBendingStiffness)
+	, BucklingRatio(InBucklingRatio)
+	, BucklingStiffness(InBucklingStiffness)
 	, bUseBendingElements(bInUseBendingElements)
 	, AreaStiffness(InAreaStiffness)
 	, VolumeStiffness(InVolumeStiffness)
diff --git a/Engine/Plugins/Experimental/ChaosCloth/Source/Chaos/Private/ChaosCloth/ChaosWeightMapTarget.h b/Engine/Plugins/Experimental/ChaosCloth/Source/Chaos/Private/ChaosCloth/ChaosWeightMapTarget.h
--- a/Engine/Plugins/Experimental/ChaosCloth/Source/Chaos/Private/ChaosCloth/ChaosWeightMapTarget.h
+++ b/Engine/Plugins/Experimental/ChaosCloth/Source/Chaos/Private/ChaosCloth/ChaosWeightMapTarget.h
@@ -22,5 +22,6 @@ enum class EChaosWeightMapTarget : uint8
 	Lift,
 	EdgeStiffness,
 	BendingStiffness,
-	AreaStiffness
+	AreaStiffness,
+	BucklingStiffness,
 };
diff --git a/Engine/Plugins/Experimental/ChaosCloth/Source/Chaos/Public/ChaosCloth/ChaosClothConfig.h b/Engine/Plugins/Experimental/ChaosCloth/Source/Chaos/Public/ChaosCloth/ChaosClothConfig.h
--- a/Engine/Plugins/Experimental/ChaosCloth/Source/Chaos/Public/ChaosCloth/ChaosClothConfig.h
+++ b/Engine/Plugins/Experimental/ChaosCloth/Source/Chaos/Public/ChaosCloth/ChaosClothConfig.h
@@ -134,6 +134,23 @@ public:
 	UPROPERTY(EditAnywhere, Category = "Material Properties", DisplayName = "Bending Stiffness", meta = (UIMin = "0", UIMax = "1", ClampMin = "0", ClampMax = "1"))
 	FChaosClothWeightedValue BendingStiffnessWeighted = { 1.f, 1.f };
 
+	/**
+	* Once the element has bent such that it's folded more than this ratio from its rest angle ("buckled"), switch to using Buckling Stiffness instead of Bending Stiffness.
+	* When Buckling Ratio = 0, the Buckling Stiffness will never be used. When BucklingRatio = 1, the Buckling Stiffness will be used as soon as its bent past its rest configuration.
+	*/
+	UPROPERTY(EditAnywhere, Category = "Material Properties", DisplayName = "Buckling Ratio", meta = (UIMin = "0", UIMax = "1", ClampMin = "0", ClampMax = "1", EditCondition = "bUseBendingElements"))
+	float BucklingRatio = 0.f;
+
+	/**
+	* Bending will use this stiffness instead of Bending Stiffness once the cloth has buckled, i.e., bent beyond a certain angle.
+	* Typically, Buckling Stiffness is set to be less than Bending Stiffness. Buckling Ratio determines the switch point between using Bending Stiffness and Buckling Stiffness.
+	* If an enabled Weight Map (Mask with values in the range [0;1]) targeting the "Buckling Stiffness" is added to the cloth, 
+	* then both the Low and High values will be used in conjunction with the per particle Weight stored in the Weight Map to interpolate the final value from them.
+	* Otherwise only the Low value is meaningful and sufficient to enable this constraint.
+	*/
+	UPROPERTY(EditAnywhere, Category = "Material Properties", DisplayName = "Buckling Stiffness", meta = (UIMin = "0", UIMax = "1", ClampMin = "0", ClampMax = "1", EditCondition = "bUseBendingElements"))
+	FChaosClothWeightedValue BucklingStiffnessWeighted = { 1.f, 1.f };
+
 	/** Enable the more accurate bending element constraints instead of the faster cross-edge spring constraints used for controlling bending stiffness. */
 	UPROPERTY(EditAnywhere, Category = "Material Properties")
 	bool bUseBendingElements = false;
diff --git a/Engine/Plugins/Experimental/ChaosCloth/Source/Chaos/Public/ChaosCloth/ChaosClothConstraints.h b/Engine/Plugins/Experimental/ChaosCloth/Source/Chaos/Public/ChaosCloth/ChaosClothConstraints.h
--- a/Engine/Plugins/Experimental/ChaosCloth/Source/Chaos/Public/ChaosCloth/ChaosClothConstraints.h
+++ b/Engine/Plugins/Experimental/ChaosCloth/Source/Chaos/Public/ChaosCloth/ChaosClothConstraints.h
@@ -31,6 +31,8 @@ namespace Chaos
 		// ---- Cloth interface ----
 		void SetEdgeConstraints(const TArray<TVec3<int32>>& SurfaceElements, const TConstArrayView<FRealSingle>& StiffnessMultipliers, bool bUseXPBDConstraints);
 		void SetBendingConstraints(const TArray<TVec2<int32>>& Edges, const TConstArrayView<FRealSingle>& StiffnessMultipliers, bool bUseXPBDConstraints);
+		void SetBendingConstraints(TArray<TVec4<int32>>&& BendingElements, const TConstArrayView<FRealSingle>& StiffnessMultipliers, const TConstArrayView<FRealSingle>& BucklingStiffnessMultipliers, bool bUseXPBDConstraints);
+		UE_DEPRECATED(5.1, "Use SetBendingConstraints(TArray<TVec4<int32>>&& BendingElements, const TConstArrayView<FRealSingle>& StiffnessMultipliers, const TConstArrayView<FRealSingle>& BucklingStiffnessMultipliers, bool bUseXPBDConstraints) instead.")
 		void SetBendingConstraints(TArray<TVec4<int32>>&& BendingElements, Softs::FSolverReal BendingStiffness);
 		void SetAreaConstraints(const TArray<TVec3<int32>>& SurfaceElements, const TConstArrayView<FRealSingle>& StiffnessMultipliers, bool bUseXPBDConstraints);
 		void SetVolumeConstraints(const TArray<TVec2<int32>>& DoubleBendingEdges, Softs::FSolverReal VolumeStiffness);
@@ -56,7 +58,7 @@ namespace Chaos
 		void Enable(bool bEnable);
 
 		void SetEdgeProperties(const Softs::FSolverVec2& EdgeStiffness);
-		void SetBendingProperties(const Softs::FSolverVec2& BendingStiffness);
+		void SetBendingProperties(const Softs::FSolverVec2& BendingStiffness, Softs::FSolverReal BucklingRatio = 0.f, const Softs::FSolverVec2& BucklingStiffnessMultiplier = Softs::FSolverVec2::UnitVector);
 		void SetAreaProperties(const Softs::FSolverVec2& AreaStiffness);
 		void SetThinShellVolumeProperties(Softs::FSolverReal VolumeStiffness);
 		void SetVolumeProperties(Softs::FSolverReal VolumeStiffness);
diff --git a/Engine/Plugins/Experimental/ChaosCloth/Source/Chaos/Public/ChaosCloth/ChaosClothingSimulationCloth.h b/Engine/Plugins/Experimental/ChaosCloth/Source/Chaos/Public/ChaosCloth/ChaosClothingSimulationCloth.h
--- a/Engine/Plugins/Experimental/ChaosCloth/Source/Chaos/Public/ChaosCloth/ChaosClothingSimulationCloth.h
+++ b/Engine/Plugins/Experimental/ChaosCloth/Source/Chaos/Public/ChaosCloth/ChaosClothingSimulationCloth.h
@@ -35,6 +35,8 @@ namespace Chaos
 			FRealSingle InMinPerParticleMass,
 			const TVec2<FRealSingle>& InEdgeStiffness,
 			const TVec2<FRealSingle>& InBendingStiffness,
+			FRealSingle InBucklingRatio,
+			const TVec2<FRealSingle>& InBucklingStiffness,
 			bool bInUseBendingElements,
 			const TVec2<FRealSingle>& InAreaStiffness,
 			FRealSingle InVolumeStiffness,
@@ -195,6 +197,8 @@ namespace Chaos
 		FRealSingle MinPerParticleMass;
 		TVec2<FRealSingle> EdgeStiffness;
 		TVec2<FRealSingle> BendingStiffness;
+		FRealSingle BucklingRatio;
+		TVec2<FRealSingle> BucklingStiffness;
 		bool bUseBendingElements;
 		TVec2<FRealSingle> AreaStiffness;
 		FRealSingle VolumeStiffness;
diff --git a/Engine/Source/Runtime/Experimental/Chaos/Public/Chaos/PBDBendingConstraints.h b/Engine/Source/Runtime/Experimental/Chaos/Public/Chaos/PBDBendingConstraints.h
--- a/Engine/Source/Runtime/Experimental/Chaos/Public/Chaos/PBDBendingConstraints.h
+++ b/Engine/Source/Runtime/Experimental/Chaos/Public/Chaos/PBDBendingConstraints.h
@@ -12,25 +12,68 @@ class FPBDBendingConstraints : public FPBDBendingConstraintsBase
 	using Base::Constraints;
 
 public:
+	FPBDBendingConstraints(const FSolverParticles& InParticles,
+		int32 ParticleOffset,
+		int32 ParticleCount,
+		TArray<TVec4<int32>>&& InConstraints,
+		const TConstArrayView<FRealSingle>& StiffnessMultipliers,
+		const TConstArrayView<FRealSingle>& BucklingStiffnessMultipliers,
+		const FSolverVec2& InStiffness,
+		const FSolverReal InBucklingRatio,
+		const FSolverVec2& InBucklingStiffness,
+		bool bTrimKinematicConstraints = false)
+		:Base(InParticles, ParticleOffset, ParticleCount, MoveTemp(InConstraints), StiffnessMultipliers, BucklingStiffnessMultipliers, InStiffness, InBucklingRatio, InBucklingStiffness, bTrimKinematicConstraints) {}
+
 	FPBDBendingConstraints(const FSolverParticles& InParticles, TArray<TVec4<int32>>&& InConstraints, const FSolverReal InStiffness = (FSolverReal)1.)
 	    : Base(InParticles, MoveTemp(InConstraints), InStiffness) {}
+
+
 	virtual ~FPBDBendingConstraints() override {}
 
 	void Apply(FSolverParticles& InParticles, const FSolverReal Dt) const
 	{
-		for (int i = 0; i < Constraints.Num(); ++i)
+		const bool StiffnessHasWeightMap = Stiffness.HasWeightMap();
+		const bool BucklingStiffnessHasWeightMap = BucklingStiffness.HasWeightMap();
+		if (!StiffnessHasWeightMap && !BucklingStiffnessHasWeightMap)
+		{
+			const FSolverReal ExpStiffnessValue = (FSolverReal)Stiffness;
+			const FSolverReal ExpBucklingValue = (FSolverReal)BucklingStiffness;
+
+			for (int i = 0; i < Constraints.Num(); ++i)
+			{
+				const TVec4<int32>& Constraint = Constraints[i];
+				const int32 i1 = Constraint[0];
+				const int32 i2 = Constraint[1];
+				const int32 i3 = Constraint[2];
+				const int32 i4 = Constraint[3];
+				const TStaticArray<FSolverVec3, 4> Grads = Base::GetGradients(InParticles, i);
+				const FSolverReal S = Base::GetScalingFactor(InParticles, i, Grads, ExpStiffnessValue, ExpBucklingValue);
+				InParticles.P(i1) -= S * InParticles.InvM(i1) * Grads[0];
+				InParticles.P(i2) -= S * InParticles.InvM(i2) * Grads[1];
+				InParticles.P(i3) -= S * InParticles.InvM(i3) * Grads[2];
+				InParticles.P(i4) -= S * InParticles.InvM(i4) * Grads[3];
+			}
+		}
+		else
 		{
-			const TVec4<int32>& Constraint = Constraints[i];
-			const int32 i1 = Constraint[0];
-			const int32 i2 = Constraint[1];
-			const int32 i3 = Constraint[2];
-			const int32 i4 = Constraint[3];
-			const TStaticArray<FSolverVec3, 4> Grads = Base::GetGradients(InParticles, i);
-			const FSolverReal S = Base::GetScalingFactor(InParticles, i, Grads);
-			InParticles.P(i1) -= S * InParticles.InvM(i1) * Grads[0];
-			InParticles.P(i2) -= S * InParticles.InvM(i2) * Grads[1];
-			InParticles.P(i3) -= S * InParticles.InvM(i3) * Grads[2];
-			InParticles.P(i4) -= S * InParticles.InvM(i4) * Grads[3];
+			const FSolverReal StiffnessNoMap = (FSolverReal)Stiffness;
+			const FSolverReal BucklingStiffnessNoMap = (FSolverReal)BucklingStiffness;
+			for (int i = 0; i < Constraints.Num(); ++i)
+			{
+				const FSolverReal ExpStiffnessValue = StiffnessHasWeightMap ? Stiffness[i] : StiffnessNoMap;
+				const FSolverReal ExpBucklingValue = BucklingStiffnessHasWeightMap ? BucklingStiffness[i] : BucklingStiffnessNoMap;
+				const TVec4<int32>& Constraint = Constraints[i];
+				const int32 i1 = Constraint[0];
+				const int32 i2 = Constraint[1];
+				const int32 i3 = Constraint[2];
+				const int32 i4 = Constraint[3];
+				const TStaticArray<FSolverVec3, 4> Grads = Base::GetGradients(InParticles, i);
+				const FSolverReal S = Base::GetScalingFactor(InParticles, i, Grads, ExpStiffnessValue, ExpBucklingValue);
+				InParticles.P(i1) -= S * InParticles.InvM(i1) * Grads[0];
+				InParticles.P(i2) -= S * InParticles.InvM(i2) * Grads[1];
+				InParticles.P(i3) -= S * InParticles.InvM(i3) * Grads[2];
+				InParticles.P(i4) -= S * InParticles.InvM(i4) * Grads[3];
+			}
 		}
 	}
 };
diff --git a/Engine/Source/Runtime/Experimental/Chaos/Public/Chaos/PBDBendingConstraintsBase.h b/Engine/Source/Runtime/Experimental/Chaos/Public/Chaos/PBDBendingConstraintsBase.h
--- a/Engine/Source/Runtime/Experimental/Chaos/Public/Chaos/PBDBendingConstraintsBase.h
+++ b/Engine/Source/Runtime/Experimental/Chaos/Public/Chaos/PBDBendingConstraintsBase.h
@@ -3,6 +3,7 @@
 
 #include "Chaos/PBDSoftsEvolutionFwd.h"
 #include "Chaos/PBDSoftsSolverParticles.h"
+#include "Chaos/PBDStiffness.h"
 #include "Chaos/ParticleRule.h"
 #include "Containers/StaticArray.h"
 
@@ -11,9 +12,23 @@ namespace Chaos::Softs
 
 class FPBDBendingConstraintsBase
 {
-  public:
-	  FPBDBendingConstraintsBase(const FSolverParticles& InParticles, TArray<TVec4<int32>>&& InConstraints, const FSolverReal InStiffness = (FSolverReal)1.)
-	    : Constraints(MoveTemp(InConstraints)), Stiffness(InStiffness)
+public:
+
+	FPBDBendingConstraintsBase(const FSolverParticles& InParticles,
+		int32 ParticleOffset,
+		int32 ParticleCount, 
+		TArray<TVec4<int32>>&& InConstraints,
+		const TConstArrayView<FRealSingle>& StiffnessMultipliers,
+		const TConstArrayView<FRealSingle>& BucklingStiffnessMultipliers,
+		const FSolverVec2& InStiffness,
+		const FSolverReal InBucklingRatio,
+		const FSolverVec2& InBucklingStiffness,
+		bool bTrimKinematicConstraints = false)
+		: Constraints(bTrimKinematicConstraints ? TrimKinematicConstraints(InConstraints, InParticles): MoveTemp(InConstraints))
+		, ConstraintSharedEdges(ExtractConstraintSharedEdges(Constraints))
+		, Stiffness(InStiffness, StiffnessMultipliers, TConstArrayView<TVec2<int32>>(ConstraintSharedEdges), ParticleOffset, ParticleCount)
+		, BucklingRatio(InBucklingRatio)
+		, BucklingStiffness(InBucklingStiffness, BucklingStiffnessMultipliers, TConstArrayView<TVec2<int32>>(ConstraintSharedEdges), ParticleOffset, ParticleCount)
 	{
 		for (const TVec4<int32>& Constraint : Constraints)
 		{
@@ -21,11 +36,43 @@ class FPBDBendingConstraintsBase
 			const FSolverVec3& P2 = InParticles.X(Constraint[1]);
 			const FSolverVec3& P3 = InParticles.X(Constraint[2]);
 			const FSolverVec3& P4 = InParticles.X(Constraint[3]);
-			MAngles.Add(CalcAngle(P1, P2, P3, P4));
+			RestAngles.Add(CalcAngle(P1, P2, P3, P4));
+		}	
+	}
+
+	FPBDBendingConstraintsBase(const FSolverParticles& InParticles, TArray<TVec4<int32>>&& InConstraints, const FSolverReal InStiffness = (FSolverReal)1.)
+	    : Constraints(MoveTemp(InConstraints))
+		, ConstraintSharedEdges(ExtractConstraintSharedEdges(Constraints))
+		, Stiffness(FSolverVec2(InStiffness))
+		, BucklingRatio(0.f)
+		, BucklingStiffness(FSolverVec2(InStiffness))
+	{
+		for (const TVec4<int32>& Constraint : Constraints)
+		{
+			const FSolverVec3& P1 = InParticles.X(Constraint[0]);
+			const FSolverVec3& P2 = InParticles.X(Constraint[1]);
+			const FSolverVec3& P3 = InParticles.X(Constraint[2]);
+			const FSolverVec3& P4 = InParticles.X(Constraint[3]);
+			RestAngles.Add(CalcAngle(P1, P2, P3, P4));
 		}
 	}
+
 	virtual ~FPBDBendingConstraintsBase() {}
 
+	// Update stiffness values
+	void SetProperties(const FSolverVec2& InStiffness, const FSolverReal InBucklingRatio, const FSolverVec2& InBucklingStiffness)
+	{ 
+		Stiffness.SetWeightedValue(InStiffness.ClampAxes((FSolverReal)0., (FSolverReal)1.)); 
+		BucklingRatio = InBucklingRatio;
+		BucklingStiffness.SetWeightedValue(InBucklingStiffness.ClampAxes((FSolverReal)0., (FSolverReal)1.));
+	}
+
+	// Update stiffness table, as well as the simulation stiffness exponent
+	void ApplyProperties(const FSolverReal Dt, const int32 NumIterations) { Stiffness.ApplyValues(Dt, NumIterations); BucklingStiffness.ApplyValues(Dt, NumIterations); }
+
+	UE_DEPRECATED(5.1, "Use SetProperties instead.")
+	void SetStiffness(FSolverReal InStiffness) { SetProperties(FSolverVec2(InStiffness), 0.f, 1.f); }
+
 	TStaticArray<FSolverVec3, 4> GetGradients(const FSolverParticles& InParticles, const int32 i) const
 	{
 		const TVec4<int32>& Constraint = Constraints[i];
@@ -37,7 +84,7 @@ class FPBDBendingConstraintsBase
 		return CalcGradients(P1, P2, P3, P4);
 	}
 
-	FSolverReal GetScalingFactor(const FSolverParticles& InParticles, const int32 i, const TStaticArray<FSolverVec3, 4>& Grads) const
+	FSolverReal GetScalingFactor(const FSolverParticles& InParticles, const int32 i, const TStaticArray<FSolverVec3, 4>& Grads, const FSolverReal ExpStiffnessValue, const FSolverReal ExpBucklingValue) const
 	{
 		const TVec4<int32>& Constraint = Constraints[i];
 		const int32 i1 = Constraint[0];
@@ -50,20 +97,27 @@ class FPBDBendingConstraintsBase
 		const FSolverVec3& P4 = InParticles.P(i4);
 		const FSolverReal Angle = CalcAngle(P1, P2, P3, P4);
 		const FSolverReal Denom = (InParticles.InvM(i1) * Grads[0].SizeSquared() + InParticles.InvM(i2) * Grads[1].SizeSquared() + InParticles.InvM(i3) * Grads[2].SizeSquared() + InParticles.InvM(i4) * Grads[3].SizeSquared());
+
+		const FSolverReal StiffnessValue = AngleIsBuckled(Angle, RestAngles[i]) ? ExpBucklingValue : ExpStiffnessValue;
+
 		constexpr FSolverReal SingleStepAngleLimit = (FSolverReal)(UE_PI * .25f); // this constraint is very non-linear. taking large steps is not accurate
-		const FSolverReal Delta = FMath::Clamp(Stiffness*(Angle - MAngles[i]), -SingleStepAngleLimit, SingleStepAngleLimit);
+		const FSolverReal Delta = FMath::Clamp(StiffnessValue * (Angle - RestAngles[i]), -SingleStepAngleLimit, SingleStepAngleLimit);
 		return SafeDivide(Delta, Denom);
 	}
 
-	void SetStiffness(FSolverReal InStiffness) { Stiffness = FMath::Clamp(InStiffness, (FSolverReal)0., (FSolverReal)1.); }
-
-  private:
-	template<class TNum>
-	static inline TNum SafeDivide(const TNum& Numerator, const FSolverReal& Denominator)
+	UE_DEPRECATED(5.1, "Use GetScalingFactor(const FSolverParticles& InParticles, const int32 i, const TStaticArray<FSolverVec3, 4>& Grads, const FSolverReal ExpStiffnessValue, const FSolverReal BucklingRatio, const FSolverReal ExpBucklingValue) instead.")
+	FSolverReal GetScalingFactor(const FSolverParticles& InParticles, const int32 i, const TArray<FSolverVec3>& Grads) const
 	{
-		if (Denominator > SMALL_NUMBER)
-			return Numerator / Denominator;
-		return TNum(0);
+		TStaticArray<FSolverVec3, 4> GradsStaticArray;
+		GradsStaticArray[0] = Grads[0];
+		GradsStaticArray[1] = Grads[1];
+		GradsStaticArray[2] = Grads[2];
+		GradsStaticArray[3] = Grads[3];
+		if (!Stiffness.HasWeightMap())
+		{
+			return GetScalingFactor(InParticles, i, GradsStaticArray, (FSolverReal)Stiffness, (FSolverReal)BucklingStiffness);
+		}
+		return GetScalingFactor(InParticles, i, GradsStaticArray, Stiffness[i], BucklingStiffness[i]);
 	}
 
 	static FSolverReal CalcAngle(const FSolverVec3& P1, const FSolverVec3& P2, const FSolverVec3& P3, const FSolverVec3& P4)
@@ -77,6 +131,25 @@ class FPBDBendingConstraintsBase
 		const FSolverReal SinPhi = FMath::Clamp(FSolverVec3::DotProduct(FSolverVec3::CrossProduct(Normal2, Normal1), SharedEdge), (FSolverReal)-1, (FSolverReal)1);
 		return FMath::Atan2(SinPhi, CosPhi);
 	}
+
+	bool AngleIsBuckled(const FSolverReal Angle, const FSolverReal RestAngle) const
+	{
+		// Angle is 0 when completely flat. This is easier to think of in terms of Angle' = (PI - |Angle|), which is 0 when completely folded.
+		// Consider buckled when Angle' <= BucklingRatio * RestAngle', and use buckling stiffness instead of stiffness.
+		return UE_PI - FMath::Abs(Angle) < BucklingRatio * (UE_PI - FMath::Abs(RestAngle));
+	}
+
+	const TArray<FSolverReal>& GetRestAngles() const { return RestAngles; }
+	const TArray<TVec4<int32>>& GetConstraints() const { return Constraints; }
+
+private:
+	template<class TNum>
+	static TNum SafeDivide(const TNum& Numerator, const FSolverReal& Denominator)
+	{
+		if (Denominator > SMALL_NUMBER)
+			return Numerator / Denominator;
+		return TNum(0);
+	}
 	
 	static TStaticArray<FSolverVec3, 4> CalcGradients(const FSolverVec3& P1, const FSolverVec3& P2, const FSolverVec3& P3, const FSolverVec3& P4)
 	{
@@ -113,12 +186,42 @@ class FPBDBendingConstraintsBase
 		return Grads;
 	}
 
-  protected:
+	static TArray<TVec4<int32>> TrimKinematicConstraints(const TArray<TVec4<int32>>& InConstraints, const FSolverParticles& InParticles)
+	{
+		TArray<TVec4<int32>> TrimmedConstraints;
+		TrimmedConstraints.Reserve(InConstraints.Num());
+		for (const TVec4<int32>& Constraint : InConstraints)
+		{
+			if (InParticles.InvM(Constraint[0]) != (FSolverReal)0. || InParticles.InvM(Constraint[1]) != (FSolverReal)0. || InParticles.InvM(Constraint[2]) != (FSolverReal)0. || InParticles.InvM(Constraint[3]) != (FSolverReal)0.)
+			{
+				TrimmedConstraints.Add(Constraint);
+			}
+		}
+		TrimmedConstraints.Shrink();
+		return TrimmedConstraints;
+	}
+
+	static TArray<TVec2<int32>> ExtractConstraintSharedEdges(const TArray<TVec4<int32>>& Constraints)
+	{
+		TArray<TVec2<int32>> ExtractedEdges;
+		ExtractedEdges.Reserve(Constraints.Num());
+		for (const TVec4<int32>& Constraint : Constraints)
+		{
+			ExtractedEdges.Emplace(Constraint[0], Constraint[1]);
+		}
+		return ExtractedEdges;
+	}
+
+protected:
 	TArray<TVec4<int32>> Constraints;
+	TArray<TVec2<int32>> ConstraintSharedEdges; // Only shared edges are used for calculating weighted stiffnesses.
+
+	FPBDStiffness Stiffness;
+	FSolverReal BucklingRatio;
+	FPBDStiffness BucklingStiffness;
 
-  private:
-	TArray<FSolverReal> MAngles;
-	FSolverReal Stiffness;
+private:
+	TArray<FSolverReal> RestAngles;
 };
 
 }  // End namespace Chaos::Softs
