Fix static analysis

diff --git a/Engine/Source/Runtime/Experimental/GeometryCollectionEngine/Public/GeometryCollection/GeometryCollectionComponent.h b/Engine/Source/Runtime/Experimental/GeometryCollectionEngine/Public/GeometryCollection/GeometryCollectionComponent.h
--- a/Engine/Source/Runtime/Experimental/GeometryCollectionEngine/Public/GeometryCollection/GeometryCollectionComponent.h
+++ b/Engine/Source/Runtime/Experimental/GeometryCollectionEngine/Public/GeometryCollection/GeometryCollectionComponent.h
@@ -299,7 +299,6 @@ struct FGeometryCollectionClusterRep
 	{
 		return Other.ObjectState != ObjectState
 			|| Other.ClusterIdx != ClusterIdx
-			|| Other.ObjectState != ObjectState
 			|| Other.Position != Position
 			|| Other.LinearVelocity != LinearVelocity
 			|| Other.AngularVelocity != AngularVelocity
