Fix crash in Chaos debug draw when particle does not have handle.

diff --git a/Engine/Source/Runtime/Experimental/Chaos/Private/Chaos/GeometryParticles.cpp b/Engine/Source/Runtime/Experimental/Chaos/Private/Chaos/GeometryParticles.cpp
--- a/Engine/Source/Runtime/Experimental/Chaos/Private/Chaos/GeometryParticles.cpp
+++ b/Engine/Source/Runtime/Experimental/Chaos/Private/Chaos/GeometryParticles.cpp
@@ -308,21 +308,24 @@ namespace Chaos
 
 	FShapeOrShapesArray::FShapeOrShapesArray(const FGeometryParticleHandle* Particle)
 	{
-		const FImplicitObject* Geometry = Particle->Geometry().Get();
-		if (Geometry)
+		if (Particle)
 		{
-			if (Geometry->IsUnderlyingUnion())
+			const FImplicitObject* Geometry = Particle->Geometry().Get();
+			if (Geometry)
 			{
-				ShapeArray = &Particle->ShapesArray();
-				bIsSingleShape = false;
-			}
-			else
-			{
-				Shape = Particle->ShapesArray()[0].Get();
-				bIsSingleShape = true;
-			}
+				if (Geometry->IsUnderlyingUnion())
+				{
+					ShapeArray = &Particle->ShapesArray();
+					bIsSingleShape = false;
+				}
+				else
+				{
+					Shape = Particle->ShapesArray()[0].Get();
+					bIsSingleShape = true;
+				}
 
-			return;
+				return;
+			}
 		}
 
 		Shape = nullptr;
