Chaos : fix sleeping GC being woken up by distant fields
- changed the logic in order to avoid going through all the clusters of the solver to update their state after a dynamic state field is processed
- this also optimize the updatethis state update by only updating the top parent clusters

diff --git a/Engine/Source/Runtime/Experimental/Chaos/Private/Chaos/PBDRigidClusteringAlgo.cpp b/Engine/Source/Runtime/Experimental/Chaos/Private/Chaos/PBDRigidClusteringAlgo.cpp
--- a/Engine/Source/Runtime/Experimental/Chaos/Private/Chaos/PBDRigidClusteringAlgo.cpp
+++ b/Engine/Source/Runtime/Experimental/Chaos/Private/Chaos/PBDRigidClusteringAlgo.cpp
@@ -213,7 +213,7 @@ namespace Chaos
 	void 
 		UpdateKinematicProperties(
 			Chaos::FPBDRigidParticleHandle* InParent,
-			FRigidClustering::FClusterMap& MChildren,
+			const FRigidClustering::FClusterMap& MChildren,
 			FRigidClustering::FRigidEvolution& MEvolution)
 	{
 		SCOPE_CYCLE_COUNTER(STAT_UpdateKinematicProperties);
diff --git a/Engine/Source/Runtime/Experimental/Chaos/Private/PhysicsProxy/GeometryCollectionPhysicsProxy.cpp b/Engine/Source/Runtime/Experimental/Chaos/Private/PhysicsProxy/GeometryCollectionPhysicsProxy.cpp
--- a/Engine/Source/Runtime/Experimental/Chaos/Private/PhysicsProxy/GeometryCollectionPhysicsProxy.cpp
+++ b/Engine/Source/Runtime/Experimental/Chaos/Private/PhysicsProxy/GeometryCollectionPhysicsProxy.cpp
@@ -2933,11 +2933,14 @@ void FGeometryCollectionPhysicsProxy::FieldParameterUpdateCallback(Chaos::FPBDRi
 						{
 							SCOPE_CYCLE_COUNTER(STAT_ParamUpdateField_DynamicState);
 							{
-								bool bHasStateChanged = false;
 								InitDynamicStateResults(ParticleHandles, FieldContext, FinalResults);
 
 								static_cast<const FFieldNode<int32>*>(FieldCommand.RootNode.Get())->Evaluate(FieldContext, ResultsView);
-								for (const FFieldContextIndex& Index : FieldContext.GetEvaluatedSamples())
+
+								bool bHasStateChanged = false;
+								
+								const TFieldArrayView<FFieldContextIndex>& EvaluatedSamples = FieldContext.GetEvaluatedSamples();
+								for (const FFieldContextIndex& Index : EvaluatedSamples)
 								{
 									Chaos::TPBDRigidParticleHandle<Chaos::FReal, 3>* RigidHandle = ParticleHandles[Index.Sample]->CastToRigidParticle();
 									if (RigidHandle)
@@ -2965,9 +2968,9 @@ void FGeometryCollectionPhysicsProxy::FieldParameterUpdateCallback(Chaos::FPBDRi
 										Collection.DynamicState[TransformIndex] = ResultState;
 									}
 								}
-								if (bUpdateViews)
+								if (bUpdateViews && bHasStateChanged)
 								{
-									UpdateSolverParticlesState(RigidSolver, bHasStateChanged);
+									UpdateSolverParticlesState(RigidSolver, EvaluatedSamples, ParticleHandles);
 								}
 							}
 						}
diff --git a/Engine/Source/Runtime/Experimental/Chaos/Public/Chaos/PBDRigidClusteringAlgo.h b/Engine/Source/Runtime/Experimental/Chaos/Public/Chaos/PBDRigidClusteringAlgo.h
--- a/Engine/Source/Runtime/Experimental/Chaos/Public/Chaos/PBDRigidClusteringAlgo.h
+++ b/Engine/Source/Runtime/Experimental/Chaos/Public/Chaos/PBDRigidClusteringAlgo.h
@@ -13,7 +13,7 @@ namespace Chaos
 
 	void CHAOS_API UpdateKinematicProperties(
 		FPBDRigidParticleHandle* Parent,
-		FRigidClustering::FClusterMap&,
+		const FRigidClustering::FClusterMap&,
 		FRigidClustering::FRigidEvolution&);
 
 	void CHAOS_API UpdateGeometry(
diff --git a/Engine/Source/Runtime/Experimental/Chaos/Public/PhysicsProxy/FieldSystemProxyHelper.h b/Engine/Source/Runtime/Experimental/Chaos/Public/PhysicsProxy/FieldSystemProxyHelper.h
--- a/Engine/Source/Runtime/Experimental/Chaos/Public/PhysicsProxy/FieldSystemProxyHelper.h
+++ b/Engine/Source/Runtime/Experimental/Chaos/Public/PhysicsProxy/FieldSystemProxyHelper.h
@@ -243,27 +243,35 @@ namespace Chaos
 	/**
 	 * Update all the clustered particles object state to static/kinematic if one of its children state has been changed to static/kinematic
 	 * @param    Rigidsolver Rigid solver owning the particle handle
-	 * @param    bHasStateChanged Boolean to check before updating the handle state
+	 * @param    UpdatedParticles List of particles that had their state updated
 	 */
-	static void UpdateSolverParticlesState(Chaos::FPBDRigidsSolver* RigidSolver, const bool bHasStateChanged)
+	static void UpdateSolverParticlesState(Chaos::FPBDRigidsSolver* RigidSolver, const TFieldArrayView<FFieldContextIndex>& UpdatedParticleIndices, const TArray<Chaos::FGeometryParticleHandle*>& Particles)
 	{
-		if (bHasStateChanged)
+		if (UpdatedParticleIndices.Num() > 0)
 		{
+			// chaos(todo) : do we really need this ? 
 			RigidSolver->GetParticles().UpdateGeometryCollectionViews(true);
 
-			Chaos::FPBDRigidsSOAs& SolverParticles = RigidSolver->GetParticles();
-			auto& Clustering = RigidSolver->GetEvolution()->GetRigidClustering();
-
-			FPBDRigidClusteredParticles& ClusteredParticles = SolverParticles.GetClusteredParticles();
+			TSet<FPBDRigidClusteredParticleHandle*> TopParentClusters;
+			for (const FFieldContextIndex& ParticleIndex: UpdatedParticleIndices)
+			{
+				if (FPBDRigidClusteredParticleHandle* ClusteredHandle = Particles[ParticleIndex.Sample]->CastToClustered())
+				{
+					FPBDRigidClusteredParticleHandle* TopParent = ClusteredHandle;
+					while (FPBDRigidClusteredParticleHandle* DirectParent = TopParent->Parent())
+					{
+						TopParent = DirectParent;
+					}
+					TopParentClusters.Add(TopParent);
+				}
+			}
 
-			for (uint32 ParticleIndex = 0, NumParticles = ClusteredParticles.Size(); ParticleIndex < NumParticles; ++ParticleIndex)
+			const FRigidClustering& Clustering = RigidSolver->GetEvolution()->GetRigidClustering();
+			for (FPBDRigidClusteredParticleHandle* TopParentClusteredHandle: TopParentClusters)
 			{
-				FPBDRigidClusteredParticleHandle* ClusteredHandle = ClusteredParticles.Handle(ParticleIndex);
-				if (ClusteredHandle && ClusteredHandle->ClusterIds().NumChildren && !ClusteredHandle->Disabled())
+				if (TopParentClusteredHandle && TopParentClusteredHandle->ClusterIds().NumChildren && !TopParentClusteredHandle->Disabled())
 				{
-					Chaos::FPBDRigidParticleHandle* RigidHandle = ClusteredHandle->CastToRigidParticle();
-					check(RigidHandle);
-					UpdateKinematicProperties(RigidHandle, Clustering.GetChildrenMap(), *RigidSolver->GetEvolution());
+					UpdateKinematicProperties(TopParentClusteredHandle, Clustering.GetChildrenMap(), *RigidSolver->GetEvolution());
 				}
 			}
 		}
@@ -378,11 +386,14 @@ namespace Chaos
 		{
 			SCOPE_CYCLE_COUNTER(STAT_ParamUpdateField_DynamicState);
 			{
-				bool bHasStateChanged = false;
 				InitDynamicStateResults(ParticleHandles, FieldContext, FinalResults);
 
 				static_cast<const FFieldNode<int32>*>(FieldCommand.RootNode.Get())->Evaluate(FieldContext, ResultsView);
-				for (const FFieldContextIndex& Index : FieldContext.GetEvaluatedSamples())
+
+				bool bHasStateChanged = false;
+				
+				const TFieldArrayView<FFieldContextIndex>& EvaluatedSamples = FieldContext.GetEvaluatedSamples();
+				for (const FFieldContextIndex& Index : EvaluatedSamples)
 				{
 					Chaos::FPBDRigidParticleHandle* RigidHandle = ParticleHandles[Index.Sample]->CastToRigidParticle();
 					if (RigidHandle)
@@ -395,7 +406,10 @@ namespace Chaos
 							false, Chaos::FVec3(0), false, Chaos::FVec3(0));
 					}
 				}
-				UpdateSolverParticlesState(RigidSolver, bHasStateChanged);
+				if (bHasStateChanged)
+				{
+					UpdateSolverParticlesState(RigidSolver, EvaluatedSamples, ParticleHandles);
+				}
 			}
 		}
 		else if (FieldCommand.PhysicsType == EFieldPhysicsType::Field_ActivateDisabled)
