Chaos : Flesh
Add graph coloring for tet constraints, the associated graph check and unit test

diff --git a/Engine/Source/Runtime/Experimental/Chaos/Private/Chaos/GraphColoring.cpp b/Engine/Source/Runtime/Experimental/Chaos/Private/Chaos/GraphColoring.cpp
--- a/Engine/Source/Runtime/Experimental/Chaos/Private/Chaos/GraphColoring.cpp
+++ b/Engine/Source/Runtime/Experimental/Chaos/Private/Chaos/GraphColoring.cpp
@@ -79,6 +79,59 @@ static bool VerifyGraph(TArray<TArray<int32>> ColorGraph, const TArray<Chaos::TV
 	return true;
 }
 
+template<typename T>
+static bool VerifyGraph(TArray<TArray<int32>> ColorGraph, const TArray<Chaos::TVec4<int32>>& Graph, const Chaos::TDynamicParticles<T, 3>& InParticles)
+{	
+	for (int32 i = 0; i < ColorGraph.Num(); ++i)
+	{
+		TMap<int32, int32> NodeToColorMap;
+		for (const auto& Edge : ColorGraph[i])
+		{
+			int32 Node1 = Graph[Edge][0];
+			int32 Node2 = Graph[Edge][1];
+			int32 Node3 = Graph[Edge][2];
+			int32 Node4 = Graph[Edge][3];
+			if (NodeToColorMap.Contains(Node1))
+			{
+				UE_LOG(LogChaos, Error, TEXT("Color %d has duplicate Node %d"), i, Node1);
+				return false;
+			}
+			if (NodeToColorMap.Contains(Node2))
+			{
+				UE_LOG(LogChaos, Error, TEXT("Color %d has duplicate Node %d"), i, Node2);
+				return false;
+			}
+			if (NodeToColorMap.Contains(Node3))
+			{
+				UE_LOG(LogChaos, Error, TEXT("Color %d has duplicate Node %d"), i, Node3);
+				return false;
+			}
+			if (NodeToColorMap.Contains(Node4))
+			{
+				UE_LOG(LogChaos, Error, TEXT("Color %d has duplicate Node %d"), i, Node4);
+				return false;
+			}
+			if (InParticles.InvM(Node1) != 0)
+			{
+				NodeToColorMap.Add(Node1, i);
+			}
+			if (InParticles.InvM(Node2) != 0)
+			{
+				NodeToColorMap.Add(Node2, i);
+			}
+			if (InParticles.InvM(Node3) != 0)
+			{
+				NodeToColorMap.Add(Node3, i);
+			}
+			if (InParticles.InvM(Node4) != 0)
+			{
+				NodeToColorMap.Add(Node4, i);
+			}
+		}
+	}
+	return true;
+}
+
 template<typename T>
 TArray<TArray<int32>> Chaos::FGraphColoring::ComputeGraphColoring(const TArray<Chaos::TVec2<int32>>& Graph, const Chaos::TDynamicParticles<T, 3>& InParticles)
 {
@@ -388,7 +441,265 @@ TArray<TArray<int32>> Chaos::FGraphColoring::ComputeGraphColoring(const TArray<T
 	return ColorGraph;
 }
 
+template<typename T>
+TArray<TArray<int32>> Chaos::FGraphColoring::ComputeGraphColoring(const TArray<TVec4<int32>>& Graph, const Chaos::TDynamicParticles<T, 3>& InParticles)
+{
+	using namespace Chaos;
+
+	TArray<TArray<int32>> ColorGraph;
+	TArray<FGraphNode> Nodes;
+	TArray<FGraphTetEdge> Edges;
+	Nodes.SetNum(InParticles.Size());
+	Edges.SetNum(Graph.Num());
+	int32 MaxColor = -1;
+
+	for (int32 i = 0; i < Graph.Num(); ++i)
+	{
+		const TVec4<int32>& Constraint = Graph[i];
+		Edges[i].FirstNode = Constraint[0];
+		Edges[i].SecondNode = Constraint[1];
+		Edges[i].ThirdNode = Constraint[2];
+		Edges[i].FourthNode = Constraint[3];
+		Nodes[Constraint[0]].Edges.Add(i);
+		Nodes[Constraint[1]].Edges.Add(i);
+		Nodes[Constraint[2]].Edges.Add(i);
+		Nodes[Constraint[3]].Edges.Add(i);
+	}
+
+	TSet<int32> ProcessedNodes;
+	TArray<int32> NodesToProcess;
+
+	for (uint32 ParticleNodeIndex = 0; ParticleNodeIndex < InParticles.Size(); ++ParticleNodeIndex)
+	{
+		const bool bIsParticleDynamic = InParticles.InvM(ParticleNodeIndex) != (T)0.;
+		if (ProcessedNodes.Contains(ParticleNodeIndex) || !bIsParticleDynamic)
+		{
+			continue;
+		}
+
+		NodesToProcess.Add(ParticleNodeIndex);
+
+		while (NodesToProcess.Num())
+		{
+			const int32 NodeIndex = NodesToProcess.Last();
+			FGraphNode& GraphNode = Nodes[NodeIndex];
+
+			NodesToProcess.SetNum(NodesToProcess.Num() - 1, /*bAllowShrinking=*/false);
+			ProcessedNodes.Add(NodeIndex);
+
+			for (const int32 EdgeIndex : GraphNode.Edges)
+			{
+				FGraphTetEdge& GraphEdge = Edges[EdgeIndex];
+
+				// If edge has been colored skip it
+				if (GraphEdge.Color >= 0)
+				{
+					continue;
+				}
+
+				// Get index to the other node on the edge
+				int32 OtherNodeIndex = INDEX_NONE;
+				int32 OtherNodeIndex2 = INDEX_NONE;
+				int32 OtherNodeIndex3 = INDEX_NONE;
+				if (GraphEdge.FirstNode == NodeIndex)
+				{
+					OtherNodeIndex = GraphEdge.SecondNode;
+					OtherNodeIndex2 = GraphEdge.ThirdNode;
+					OtherNodeIndex3 = GraphEdge.FourthNode;
+				}
+				if (GraphEdge.SecondNode == NodeIndex)
+				{
+					OtherNodeIndex = GraphEdge.FirstNode;
+					OtherNodeIndex2 = GraphEdge.ThirdNode;
+					OtherNodeIndex3 = GraphEdge.FourthNode;
+				}
+				if (GraphEdge.ThirdNode == NodeIndex)
+				{
+					OtherNodeIndex = GraphEdge.FirstNode;
+					OtherNodeIndex2 = GraphEdge.SecondNode;
+					OtherNodeIndex3 = GraphEdge.FourthNode;
+				}
+				if (GraphEdge.FourthNode == NodeIndex)
+				{
+					OtherNodeIndex = GraphEdge.FirstNode;
+					OtherNodeIndex2 = GraphEdge.SecondNode;
+					OtherNodeIndex3 = GraphEdge.ThirdNode;
+				}
+
+				// Find next color that is not used already at this node
+				while (GraphNode.UsedColors.Contains(GraphNode.NextColor))
+				{
+					GraphNode.NextColor++;
+				}
+				int32 ColorToUse = GraphNode.NextColor;
+
+				// Exclude colors used by the other node (but still allow this node to use them for other edges)
+				if (OtherNodeIndex != INDEX_NONE)
+				{
+					FGraphNode& OtherNode = Nodes[OtherNodeIndex];
+
+					const bool bIsOtherGraphNodeDynamic = InParticles.InvM(OtherNodeIndex) != (T)0.;
+					if (bIsOtherGraphNodeDynamic)
+					{
+						while (OtherNode.UsedColors.Contains(ColorToUse) || GraphNode.UsedColors.Contains(ColorToUse))
+						{
+							ColorToUse++;
+						}
+					}
+				}
+				if (OtherNodeIndex2 != INDEX_NONE)
+				{
+					FGraphNode& OtherNode = Nodes[OtherNodeIndex2];
+
+					const bool bIsOtherGraphNodeDynamic = InParticles.InvM(OtherNodeIndex2) != (T)0.;
+					if (bIsOtherGraphNodeDynamic)
+					{
+						if (OtherNodeIndex == INDEX_NONE)
+						{
+							while (OtherNode.UsedColors.Contains(ColorToUse) || GraphNode.UsedColors.Contains(ColorToUse))
+							{
+								ColorToUse++;
+							}
+						}
+						else
+						{
+							FGraphNode& PrevOtherNode = Nodes[OtherNodeIndex];
+							while (OtherNode.UsedColors.Contains(ColorToUse) || PrevOtherNode.UsedColors.Contains(ColorToUse) || GraphNode.UsedColors.Contains(ColorToUse))
+							{
+								ColorToUse++;
+							}
+						}
+					}
+				}
+				if (OtherNodeIndex3 != INDEX_NONE)
+				{
+					FGraphNode& OtherNode = Nodes[OtherNodeIndex3];
+
+					const bool bIsOtherGraphNodeDynamic = InParticles.InvM(OtherNodeIndex3) != (T)0.;
+					if (bIsOtherGraphNodeDynamic)
+					{
+						if (OtherNodeIndex == INDEX_NONE && OtherNodeIndex2 == INDEX_NONE)
+						{
+							while (OtherNode.UsedColors.Contains(ColorToUse) || GraphNode.UsedColors.Contains(ColorToUse))
+							{
+								ColorToUse++;
+							}
+						}
+						else if (OtherNodeIndex == INDEX_NONE)
+						{
+							//OtherNodeIndex2 != INDEX_NONE
+							FGraphNode& PrevOtherNode = Nodes[OtherNodeIndex2];
+							while (OtherNode.UsedColors.Contains(ColorToUse) || PrevOtherNode.UsedColors.Contains(ColorToUse) || GraphNode.UsedColors.Contains(ColorToUse))
+							{
+								ColorToUse++;
+							}
+						}
+						else if (OtherNodeIndex2 == INDEX_NONE)
+						{
+							FGraphNode& PrevOtherNode = Nodes[OtherNodeIndex];
+							while (OtherNode.UsedColors.Contains(ColorToUse) || PrevOtherNode.UsedColors.Contains(ColorToUse) || GraphNode.UsedColors.Contains(ColorToUse))
+							{
+								ColorToUse++;
+							}
+						}
+						else
+						{
+							FGraphNode& PrevOtherNode = Nodes[OtherNodeIndex];
+							FGraphNode& PrevPrevOtherNode = Nodes[OtherNodeIndex2];
+							while (OtherNode.UsedColors.Contains(ColorToUse) || PrevOtherNode.UsedColors.Contains(ColorToUse) || PrevPrevOtherNode.UsedColors.Contains(ColorToUse) || GraphNode.UsedColors.Contains(ColorToUse))
+							{
+								ColorToUse++;
+							}
+						}
+					}
+				}
+
+				// Assign color and set as used at this node
+				MaxColor = FMath::Max(ColorToUse, MaxColor);
+				GraphNode.UsedColors.Add(ColorToUse);
+				GraphEdge.Color = ColorToUse;
+
+				// Bump color to use next time, but only if we weren't forced to use a different color by the other node
+				if ((ColorToUse == GraphNode.NextColor) && bIsParticleDynamic == true)
+				{
+					GraphNode.NextColor++;
+				}
+
+				if (ColorGraph.Num() <= MaxColor)
+				{
+					ColorGraph.SetNum(MaxColor + 1);
+				}
+				ColorGraph[GraphEdge.Color].Add(EdgeIndex);
+
+				if (OtherNodeIndex != INDEX_NONE)
+				{
+					FGraphNode& OtherGraphNode = Nodes[OtherNodeIndex];
+					const bool bIsOtherGraphNodeDynamic = InParticles.InvM(OtherNodeIndex) != (T)0.;
+					if (bIsOtherGraphNodeDynamic)
+					{
+						// Mark other node as not allowing use of this color
+						if (bIsParticleDynamic)
+						{
+							OtherGraphNode.UsedColors.Add(GraphEdge.Color);
+						}
+
+						// Queue other node for processing
+						if (!ProcessedNodes.Contains(OtherNodeIndex))
+						{
+							NodesToProcess.Add(OtherNodeIndex);
+						}
+					}
+				}
+				if (OtherNodeIndex2 != INDEX_NONE)
+				{
+					FGraphNode& OtherGraphNode = Nodes[OtherNodeIndex2];
+					const bool bIsOtherGraphNodeDynamic = InParticles.InvM(OtherNodeIndex2) != (T)0.;
+					if (bIsOtherGraphNodeDynamic)
+					{
+						// Mark other node as not allowing use of this color
+						if (bIsParticleDynamic)
+						{
+							OtherGraphNode.UsedColors.Add(GraphEdge.Color);
+						}
+
+						// Queue other node for processing
+						if (!ProcessedNodes.Contains(OtherNodeIndex2))
+						{
+							NodesToProcess.Add(OtherNodeIndex2);
+						}
+					}
+				}
+				if (OtherNodeIndex3 != INDEX_NONE)
+				{
+					FGraphNode& OtherGraphNode = Nodes[OtherNodeIndex3];
+					const bool bIsOtherGraphNodeDynamic = InParticles.InvM(OtherNodeIndex3) != (T)0.;
+					if (bIsOtherGraphNodeDynamic)
+					{
+						// Mark other node as not allowing use of this color
+						if (bIsParticleDynamic)
+						{
+							OtherGraphNode.UsedColors.Add(GraphEdge.Color);
+						}
+
+						// Queue other node for processing
+						if (!ProcessedNodes.Contains(OtherNodeIndex3))
+						{
+							NodesToProcess.Add(OtherNodeIndex3);
+						}
+					}
+				}
+			}
+		}
+	}
+
+	checkSlow(VerifyGraph(ColorGraph, Graph, InParticles)); 
+	return ColorGraph;
+}
+
+
 template CHAOS_API TArray<TArray<int32>> Chaos::FGraphColoring::ComputeGraphColoring<Chaos::FRealSingle>(const TArray<Chaos::TVector<int32, 2>>&, const Chaos::TDynamicParticles<Chaos::FRealSingle, 3>&);
 template CHAOS_API TArray<TArray<int32>> Chaos::FGraphColoring::ComputeGraphColoring<Chaos::FRealDouble>(const TArray<Chaos::TVector<int32, 2>>&, const Chaos::TDynamicParticles<Chaos::FRealDouble, 3>&);
 template CHAOS_API TArray<TArray<int32>> Chaos::FGraphColoring::ComputeGraphColoring<Chaos::FRealSingle>(const TArray<Chaos::TVector<int32, 3>>&, const Chaos::TDynamicParticles<Chaos::FRealSingle, 3>&);
 template CHAOS_API TArray<TArray<int32>> Chaos::FGraphColoring::ComputeGraphColoring<Chaos::FRealDouble>(const TArray<Chaos::TVector<int32, 3>>&, const Chaos::TDynamicParticles<Chaos::FRealDouble, 3>&);
+template CHAOS_API TArray<TArray<int32>> Chaos::FGraphColoring::ComputeGraphColoring<Chaos::FRealSingle>(const TArray<Chaos::TVector<int32, 4>>&, const Chaos::TDynamicParticles<Chaos::FRealSingle, 3>&);
+template CHAOS_API TArray<TArray<int32>> Chaos::FGraphColoring::ComputeGraphColoring<Chaos::FRealDouble>(const TArray<Chaos::TVector<int32, 4>>&, const Chaos::TDynamicParticles<Chaos::FRealDouble, 3>&);
diff --git a/Engine/Source/Runtime/Experimental/Chaos/Public/Chaos/GraphColoring.h b/Engine/Source/Runtime/Experimental/Chaos/Public/Chaos/GraphColoring.h
--- a/Engine/Source/Runtime/Experimental/Chaos/Public/Chaos/GraphColoring.h
+++ b/Engine/Source/Runtime/Experimental/Chaos/Public/Chaos/GraphColoring.h
@@ -49,11 +49,23 @@ class FGraphColoring
 		int32 ThirdNode;
 	};
 
+	struct FGraphTetEdge : FGraph3dEdge
+	{
+		FGraphTetEdge()
+			: FourthNode(INDEX_NONE)
+		{
+		}
+
+		int32 FourthNode;
+	};
+
   public:
 	template<typename T>
 	CHAOS_API static TArray<TArray<int32>> ComputeGraphColoring(const TArray<TVector<int32, 2>>& Graph, const TDynamicParticles<T, 3>& InParticles);
 	template<typename T>
 	CHAOS_API static TArray<TArray<int32>> ComputeGraphColoring(const TArray<TVector<int32, 3>>& Graph, const TDynamicParticles<T, 3>& InParticles);
+	template<typename T>
+	CHAOS_API static TArray<TArray<int32>> ComputeGraphColoring(const TArray<TVector<int32, 4>>& Graph, const TDynamicParticles<T, 3>& InParticles);
 };
 
 }
