Chaos: Fix geometry collection/provider asserting that its geometry is of inner types without unwrapping its implicit. Caused crash inside collision analyzer when resolving which shape type was hit.
Chaos: Fix the cast helpers for Chaos geometies only accepting callable objects that return value types

diff --git a/Engine/Source/Runtime/Experimental/Chaos/Public/Chaos/CastingUtilities.h b/Engine/Source/Runtime/Experimental/Chaos/Public/Chaos/CastingUtilities.h
--- a/Engine/Source/Runtime/Experimental/Chaos/Public/Chaos/CastingUtilities.h
+++ b/Engine/Source/Runtime/Experimental/Chaos/Public/Chaos/CastingUtilities.h
@@ -17,7 +17,7 @@ namespace Chaos
 	{
 		// Call the lambda with concrete shape type. Unwraps shapes contained in Instanced (e.g., Instanced-Sphere will be called with Sphere. Note that this will effectively discard any instance properties like the margin)
 		template <typename Lambda>
-		FORCEINLINE_DEBUGGABLE auto CastHelper(const FImplicitObject& Geom, const Lambda& Func)
+		FORCEINLINE_DEBUGGABLE decltype(auto) CastHelper(const FImplicitObject& Geom, const Lambda& Func)
 		{
 			const EImplicitObjectType Type = Geom.GetType();
 			switch (Type)
@@ -83,7 +83,7 @@ namespace Chaos
 
 		// Call the lambda with concrete shape type. Unwraps shapes contained in Instanced (e.g., Instanced-Sphere will be called with Sphere. Note that this will effectively discard any instance properties like the margin)
 		template <typename Lambda>
-		FORCEINLINE_DEBUGGABLE auto CastHelper(const FImplicitObject& Geom, const FRigidTransform3& TM, const Lambda& Func)
+		FORCEINLINE_DEBUGGABLE decltype(auto) CastHelper(const FImplicitObject& Geom, const FRigidTransform3& TM, const Lambda& Func)
 		{
 			const EImplicitObjectType Type = Geom.GetType();
 			switch (Type)
@@ -125,7 +125,7 @@ namespace Chaos
 
 		// Call the lambda with concrete shape type. This version does NOT unwrap shapes contained in Instanced or Scaled.
 		template <typename Lambda>
-		FORCEINLINE_DEBUGGABLE auto CastHelperNoUnwrap(const FImplicitObject& Geom, const FRigidTransform3& TM, const Lambda& Func)
+		FORCEINLINE_DEBUGGABLE decltype(auto) CastHelperNoUnwrap(const FImplicitObject& Geom, const FRigidTransform3& TM, const Lambda& Func)
 		{
 			const EImplicitObjectType Type = Geom.GetType();
 			switch (Type)
@@ -154,7 +154,7 @@ namespace Chaos
 		}
 
 		template <typename Lambda>
-		FORCEINLINE_DEBUGGABLE auto CastHelperNoUnwrap(const FImplicitObject& Geom, const FRigidTransform3& TM0, const FRigidTransform3& TM1, const Lambda& Func)
+		FORCEINLINE_DEBUGGABLE decltype(auto) CastHelperNoUnwrap(const FImplicitObject& Geom, const FRigidTransform3& TM0, const FRigidTransform3& TM1, const Lambda& Func)
 		{
 			const EImplicitObjectType Type = Geom.GetType();
 			switch (Type)
@@ -240,7 +240,7 @@ namespace Chaos
 		 * which will be const T*, const TImplicitObjectInstanced<T>*, or const TImplicitObjectScaled<T>*
 		*/
 		template <typename T, typename Lambda>
-		FORCEINLINE_DEBUGGABLE auto CastWrapped(const FImplicitObject& A, const Lambda& Func)
+		FORCEINLINE_DEBUGGABLE decltype(auto) CastWrapped(const FImplicitObject& A, const Lambda& Func)
 		{
 			if (const TImplicitObjectScaled<T>* ScaledImplicit = A.template GetObject<TImplicitObjectScaled<T>>())
 			{
diff --git a/Engine/Source/Runtime/PhysicsCore/Private/ChaosEngineInterface.cpp b/Engine/Source/Runtime/PhysicsCore/Private/ChaosEngineInterface.cpp
--- a/Engine/Source/Runtime/PhysicsCore/Private/ChaosEngineInterface.cpp
+++ b/Engine/Source/Runtime/PhysicsCore/Private/ChaosEngineInterface.cpp
@@ -23,8 +23,10 @@ FPhysicsDelegatesCore::FOnUpdatePhysXMaterial FPhysicsDelegatesCore::OnUpdatePhy
 #include "Chaos/PBDSuspensionConstraintData.h"
 #include "Chaos/Collision/CollisionConstraintFlags.h"
 #include "PhysicsProxy/SingleParticlePhysicsProxy.h"
+#include "Chaos/ImplicitObject.h"
 #include "PBDRigidsSolver.h"
 #include "PhysicsProxy/SingleParticlePhysicsProxy.h"
+#include "Chaos/CastingUtilities.h"
 
 bool bEnableChaosJointConstraints = true;
 FAutoConsoleVariableRef CVarEnableChaosJointConstraints(TEXT("p.ChaosSolverEnableJointConstraints"), bEnableChaosJointConstraints, TEXT("Enable Joint Constraints defined within the Physics Asset Editor"));
@@ -54,28 +56,40 @@ const Chaos::FImplicitObject& FPhysicsGeometryCollection_Chaos::GetGeometry() co
 	return Geom;
 }
 
+template<typename InnerType>
+const InnerType& GetInnerGeometryChecked(const Chaos::FImplicitObject& InGeometry)
+{
+	const Chaos::FImplicitObject& InnerObject = Chaos::Utilities::CastHelper(InGeometry,
+		[](const Chaos::FImplicitObject& CastGeom) -> const Chaos::FImplicitObject&
+		{
+			return CastGeom;
+		});
+
+	return InnerObject.GetObjectChecked<InnerType>();
+}
+
 const Chaos::TBox<Chaos::FReal,3>& FPhysicsGeometryCollection_Chaos::GetBoxGeometry() const
 {
-	return Geom.GetObjectChecked<Chaos::TBox<Chaos::FReal,3>>();
+	return GetInnerGeometryChecked<Chaos::TBox<Chaos::FReal, 3>>(Geom);
 }
 
 const Chaos::TSphere<Chaos::FReal,3>&  FPhysicsGeometryCollection_Chaos::GetSphereGeometry() const
 {
-	return Geom.GetObjectChecked<Chaos::TSphere<Chaos::FReal,3>>();
+	return GetInnerGeometryChecked<Chaos::TSphere<Chaos::FReal, 3>>(Geom);
 }
 const Chaos::FCapsule&  FPhysicsGeometryCollection_Chaos::GetCapsuleGeometry() const
 {
-	return Geom.GetObjectChecked<Chaos::FCapsule>();
+	return GetInnerGeometryChecked<Chaos::FCapsule>(Geom);
 }
 
 const Chaos::FConvex& FPhysicsGeometryCollection_Chaos::GetConvexGeometry() const
 {
-	return Geom.GetObjectChecked<Chaos::FConvex>();
+	return GetInnerGeometryChecked<Chaos::FConvex>(Geom);
 }
 
 const Chaos::FTriangleMeshImplicitObject& FPhysicsGeometryCollection_Chaos::GetTriMeshGeometry() const
 {
-	return Geom.GetObjectChecked<Chaos::FTriangleMeshImplicitObject>();
+	return GetInnerGeometryChecked<Chaos::FTriangleMeshImplicitObject>(Geom);
 }
 
 FPhysicsGeometryCollection_Chaos::FPhysicsGeometryCollection_Chaos(const FPhysicsShapeReference_Chaos& InShape)
