Chaos : fix inertia for cluster when world transform has a rotation

diff --git a/Engine/Source/Runtime/Experimental/Chaos/Private/Chaos/PBDRigidClusteringAlgo.cpp b/Engine/Source/Runtime/Experimental/Chaos/Private/Chaos/PBDRigidClusteringAlgo.cpp
--- a/Engine/Source/Runtime/Experimental/Chaos/Private/Chaos/PBDRigidClusteringAlgo.cpp
+++ b/Engine/Source/Runtime/Experimental/Chaos/Private/Chaos/PBDRigidClusteringAlgo.cpp
@@ -4,6 +4,7 @@
 
 #include "Chaos/ErrorReporter.h"
 #include "Chaos/Levelset.h"
+#include "Chaos/Utilities.h"
 #include "Chaos/PBDRigidClusteringCollisionParticleAlgo.h"
 #include "Chaos/PBDRigidsEvolutionGBF.h"
 
@@ -65,9 +66,7 @@ namespace Chaos
 			ChildRotation = Child->R();
 
 			const FReal ChildMass = Child->M();
-			const FMatrix33 IMatrix(Child->I());
-			const FMatrix33 ChildWorldSpaceI =
-				(ChildRotation * FMatrix::Identity) * IMatrix * (ChildRotation * FMatrix::Identity).GetTransposed();
+			const FMatrix33 ChildWorldSpaceI = Chaos::Utilities::ComputeWorldSpaceInertia(ChildRotation, Child->I());
 			if (ChildWorldSpaceI.ContainsNaN())
 			{
 				continue;
@@ -91,9 +90,7 @@ namespace Chaos
 				const FRotation3& ChildRotation = Child->R();
 				const FReal ChildMass = Child->M();
 
-				const FMatrix33 IMatrix(Child->I());
-				const FMatrix33 ChildWorldSpaceI =
-					(ChildRotation * FMatrix::Identity) * IMatrix * (ChildRotation * FMatrix::Identity).GetTransposed();
+				const FMatrix33 ChildWorldSpaceI = Chaos::Utilities::ComputeWorldSpaceInertia(ChildRotation, Child->I());
 				if (ChildWorldSpaceI.ContainsNaN())
 				{
 					continue;
@@ -171,6 +168,12 @@ namespace Chaos
 						m * (p2 * p2 + p0 * p0), -m * p2 * p1, m * (p1 * p1 + p0 * p0));
 			}
 		}
+
+		if (ForceMassOrientation)
+		{
+			ParentInertia = ForceMassOrientation->GetRotation().ToMatrix() * ParentInertia * ForceMassOrientation->GetRotation().ToMatrix().GetTransposed();
+		}
+
 		if (ParentInertia.ContainsNaN())
 		{
 			ParentInertia = FMatrix33(1,1,1);
