Chaos Cloth Air Pressure Field (Normal displacement)

This adds a new "Pressure" field and map to cloth which is used to model a difference in air pressure on two sides of cloth (e.g., inflatable objects, puffy jackets, pillows).

For the sake of optimization, this has been folded into the existing FVelocityField which has been renamed FVelocityAndPressureField, as it is easy to compute the Pressure force using the intermediates already computed to apply Air Lift and Drag. Everything has been renamed to VelocityAndPressure instead of Velocity (old API exists but has been deprecated). The lion's share of these changes are just updating all of the files to the new API.

diff --git a/Engine/Plugins/Experimental/ChaosCloth/Source/Chaos/Private/ChaosCloth/ChaosClothingSimulation.cpp b/Engine/Plugins/Experimental/ChaosCloth/Source/Chaos/Private/ChaosCloth/ChaosClothingSimulation.cpp
--- a/Engine/Plugins/Experimental/ChaosCloth/Source/Chaos/Private/ChaosCloth/ChaosClothingSimulation.cpp
+++ b/Engine/Plugins/Experimental/ChaosCloth/Source/Chaos/Private/ChaosCloth/ChaosClothingSimulation.cpp
@@ -419,6 +419,7 @@ void FClothingSimulation::CreateActor(USkeletalMeshComponent* InOwnerComponent,
 		TVec2<FRealSingle>(ClothConfig->Drag.Low, ClothConfig->Drag.High),  // Animatable
 		TVec2<FRealSingle>(ClothConfig->Lift.Low, ClothConfig->Lift.High),  // Animatable
 		ClothConfig->bUsePointBasedWindModel,
+		TVec2<FRealSingle>(ClothConfig->Pressure.Low, ClothConfig->Pressure.High),  // Animatable
 		ClothConfig->DampingCoefficient,
 		ClothConfig->LocalDampingCoefficient,
 		ClothConfig->CollisionThickness,
@@ -600,7 +601,7 @@ void FClothingSimulation::Simulate(IClothingSimulationContext* InContext)
 	if (ClothingSimulationCVar::DebugDrawAnimDrive           .GetValueOnAnyThread()) { DebugDrawAnimDrive           (); }
 	if (ClothingSimulationCVar::DebugDrawBendingConstraint   .GetValueOnAnyThread()) { DebugDrawBendingConstraint   (); }
 	if (ClothingSimulationCVar::DebugDrawLongRangeConstraint .GetValueOnAnyThread()) { DebugDrawLongRangeConstraint (); }
-	if (ClothingSimulationCVar::DebugDrawWindForces          .GetValueOnAnyThread()) { DebugDrawWindForces          (); }
+	if (ClothingSimulationCVar::DebugDrawWindForces          .GetValueOnAnyThread()) { DebugDrawWindAndPressureForces(); }
 	if (ClothingSimulationCVar::DebugDrawSelfCollision       .GetValueOnAnyThread()) { DebugDrawSelfCollision       (); }
 	if (ClothingSimulationCVar::DebugDrawSelfIntersection    .GetValueOnAnyThread()) { DebugDrawSelfIntersection    (); }
 #endif  // #if CHAOS_DEBUG_DRAW
@@ -845,6 +846,7 @@ void FClothingSimulation::RefreshClothConfig(const IClothingSimulationContext* I
 			TVec2<FRealSingle>(ClothConfig->Drag.Low, ClothConfig->Drag.High),  // Animatable
 			TVec2<FRealSingle>(ClothConfig->Lift.Low, ClothConfig->Lift.High),  // Animatable
 			ClothConfig->bUsePointBasedWindModel,
+			TVec2<FRealSingle>(ClothConfig->Pressure.Low, ClothConfig->Pressure.High),  // Animatable
 			ClothConfig->DampingCoefficient,
 			ClothConfig->LocalDampingCoefficient,
 			ClothConfig->CollisionThickness,
@@ -1952,7 +1954,7 @@ void FClothingSimulation::DebugDrawLongRangeConstraint(FPrimitiveDrawInterface*
 	}
 }
 
-void FClothingSimulation::DebugDrawWindForces(FPrimitiveDrawInterface* PDI) const
+void FClothingSimulation::DebugDrawWindAndPressureForces(FPrimitiveDrawInterface* PDI) const
 {
 	check(Solver);
 
@@ -1967,7 +1969,7 @@ void FClothingSimulation::DebugDrawWindForces(FPrimitiveDrawInterface* PDI) cons
 			continue;
 		}
 
-		const Softs::FVelocityField& VelocityField = Solver->GetWindVelocityField(Cloth->GetGroupId());
+		const Softs::FVelocityAndPressureField& VelocityField = Solver->GetWindVelocityAndPressureField(Cloth->GetGroupId());
 
 		const TConstArrayView<TVec3<int32>>& Elements = VelocityField.GetElements();
 		const TConstArrayView<Softs::FSolverVec3> Forces = VelocityField.GetForces();
diff --git a/Engine/Plugins/Experimental/ChaosCloth/Source/Chaos/Private/ChaosCloth/ChaosClothingSimulationCloth.cpp b/Engine/Plugins/Experimental/ChaosCloth/Source/Chaos/Private/ChaosCloth/ChaosClothingSimulationCloth.cpp
--- a/Engine/Plugins/Experimental/ChaosCloth/Source/Chaos/Private/ChaosCloth/ChaosClothingSimulationCloth.cpp
+++ b/Engine/Plugins/Experimental/ChaosCloth/Source/Chaos/Private/ChaosCloth/ChaosClothingSimulationCloth.cpp
@@ -365,6 +365,7 @@ FClothingSimulationCloth::FClothingSimulationCloth(
 	const TVec2<FRealSingle>& InDrag,
 	const TVec2<FRealSingle>& InLift,
 	bool bInUseLegacyWind,
+	const TVec2<FRealSingle>& InPressure,
 	FRealSingle InDampingCoefficient,
 	FRealSingle InLocalDampingCoefficient,
 	FRealSingle InCollisionThickness,
@@ -410,6 +411,7 @@ FClothingSimulationCloth::FClothingSimulationCloth(
 	, WindVelocity((FReal)0., (FReal)0., (FReal)0.)  // Set by clothing interactor
 	, AirDensity(1.225e-6f)  // Set by clothing interactor
 	, bUseLegacyWind(bInUseLegacyWind)
+	, Pressure(InPressure)
 	, DampingCoefficient(InDampingCoefficient)
 	, LocalDampingCoefficient(InLocalDampingCoefficient)
 	, CollisionThickness(InCollisionThickness)
@@ -746,8 +748,9 @@ void FClothingSimulationCloth::Update(FClothingSimulationSolver* Solver)
 				// Update the wind velocity field for the new LOD mesh
 				const TConstArrayView<FRealSingle>& DragMultipliers = LODData[LODIndex].WeightMaps[(int32)EChaosWeightMapTarget::Drag];
 				const TConstArrayView<FRealSingle>& LiftMultipliers = LODData[LODIndex].WeightMaps[(int32)EChaosWeightMapTarget::Lift];
+				const TConstArrayView<FRealSingle>& PressureMultipliers = LODData[LODIndex].WeightMaps[(int32)EChaosWeightMapTarget::Pressure];
 
-				Solver->SetWindGeometry(GroupId, GetTriangleMesh(Solver), DragMultipliers, LiftMultipliers);
+				Solver->SetWindAndPressureGeometry(GroupId, GetTriangleMesh(Solver), DragMultipliers, LiftMultipliers, PressureMultipliers);
 		}
 		else
 		{
@@ -810,11 +813,11 @@ void FClothingSimulationCloth::Update(FClothingSimulationSolver* Solver)
 
 		if (bUseLegacyWind && ClothingSimulationClothConsoleVariables::CVarLegacyDisablesAccurateWind.GetValueOnAnyThread())
 		{
-			Solver->SetWindProperties(GroupId, TVec2<FRealSingle>(0.), TVec2<FRealSingle>(0.), (FRealSingle)0.);  // Disable the wind velocity field
+			Solver->SetWindAndPressureProperties(GroupId, TVec2<FRealSingle>(0.), TVec2<FRealSingle>(0.), (FRealSingle)0., Pressure);  // Disable the wind velocity field
 		}
 		else
 		{
-			Solver->SetWindProperties(GroupId, Drag, Lift, AirDensity);
+			Solver->SetWindAndPressureProperties(GroupId, Drag, Lift, AirDensity, Pressure);
 		}
 		Solver->SetWindVelocity(GroupId, WindVelocity + Solver->GetWindVelocity());
 
diff --git a/Engine/Plugins/Experimental/ChaosCloth/Source/Chaos/Private/ChaosCloth/ChaosClothingSimulationInteractor.cpp b/Engine/Plugins/Experimental/ChaosCloth/Source/Chaos/Private/ChaosCloth/ChaosClothingSimulationInteractor.cpp
--- a/Engine/Plugins/Experimental/ChaosCloth/Source/Chaos/Private/ChaosCloth/ChaosClothingSimulationInteractor.cpp
+++ b/Engine/Plugins/Experimental/ChaosCloth/Source/Chaos/Private/ChaosCloth/ChaosClothingSimulationInteractor.cpp
@@ -106,6 +106,14 @@ void UChaosClothingInteractor::SetWind(FVector2D Drag, FVector2D Lift, float Air
 	}));
 }
 
+void UChaosClothingInteractor::SetPressure(FVector2D Pressure)
+{
+	Commands.Add(FChaosClothingInteractorCommand::CreateLambda([Pressure](Chaos::FClothingSimulationCloth* Cloth)
+	{
+		Cloth->SetPressureProperties(Chaos::TVec2<Chaos::FRealSingle>(Pressure[0], Pressure[1]));
+	}));
+}
+
 void UChaosClothingInteractor::SetGravity(float GravityScale, bool bIsGravityOverridden, FVector GravityOverride)
 {
 	Commands.Add(FChaosClothingInteractorCommand::CreateLambda([GravityScale, bIsGravityOverridden, GravityOverride](Chaos::FClothingSimulationCloth* Cloth)
diff --git a/Engine/Plugins/Experimental/ChaosCloth/Source/Chaos/Private/ChaosCloth/ChaosClothingSimulationSolver.cpp b/Engine/Plugins/Experimental/ChaosCloth/Source/Chaos/Private/ChaosCloth/ChaosClothingSimulationSolver.cpp
--- a/Engine/Plugins/Experimental/ChaosCloth/Source/Chaos/Private/ChaosCloth/ChaosClothingSimulationSolver.cpp
+++ b/Engine/Plugins/Experimental/ChaosCloth/Source/Chaos/Private/ChaosCloth/ChaosClothingSimulationSolver.cpp
@@ -685,25 +685,28 @@ void FClothingSimulationSolver::SetWindVelocity(const TVec3<FRealSingle>& InWind
 
 void FClothingSimulationSolver::SetWindVelocity(uint32 GroupId, const TVec3<FRealSingle>& InWindVelocity)
 {
-	Softs::FVelocityField& VelocityField = Evolution->GetVelocityField(GroupId);
+	Softs::FVelocityAndPressureField& VelocityField = Evolution->GetVelocityAndPressureField(GroupId);
 	VelocityField.SetVelocity(Softs::FSolverVec3(InWindVelocity));
 }
 
-void FClothingSimulationSolver::SetWindGeometry(uint32 GroupId, const FTriangleMesh& TriangleMesh, const TConstArrayView<FRealSingle>& DragMultipliers, const TConstArrayView<FRealSingle>& LiftMultipliers)
+void FClothingSimulationSolver::SetWindAndPressureGeometry(uint32 GroupId, const FTriangleMesh& TriangleMesh, const TConstArrayView<FRealSingle>& DragMultipliers, const TConstArrayView<FRealSingle>& LiftMultipliers,
+	const TConstArrayView<FRealSingle>& PressureMultipliers)
 {
-	Softs::FVelocityField& VelocityField = Evolution->GetVelocityField(GroupId);
-	VelocityField.SetGeometry(&TriangleMesh, DragMultipliers, LiftMultipliers);
+	Softs::FVelocityAndPressureField& VelocityField = Evolution->GetVelocityAndPressureField(GroupId);
+	VelocityField.SetGeometry(&TriangleMesh, DragMultipliers, LiftMultipliers, PressureMultipliers);
 }
 
-void FClothingSimulationSolver::SetWindProperties(uint32 GroupId, const TVec2<FRealSingle>& Drag, const TVec2<FRealSingle>& Lift, FRealSingle AirDensity)
+void FClothingSimulationSolver::SetWindAndPressureProperties(uint32 GroupId, const TVec2<FRealSingle>& Drag, const TVec2<FRealSingle>& Lift, FRealSingle AirDensity, const TVec2<FRealSingle>& Pressure)
 {
-	Softs::FVelocityField& VelocityField = Evolution->GetVelocityField(GroupId);
-	VelocityField.SetProperties(Drag, Lift, AirDensity);
+	Softs::FVelocityAndPressureField& VelocityField = Evolution->GetVelocityAndPressureField(GroupId);
+
+	// UI Pressure is in kg/m s^2. Need to convert to kg/cm s^2 for solver.
+	VelocityField.SetProperties(Drag, Lift, AirDensity, Pressure / ClothingSimulationSolverConstant::WorldScale);
 }
 
-const Softs::FVelocityField& FClothingSimulationSolver::GetWindVelocityField(uint32 GroupId)
+const Softs::FVelocityAndPressureField& FClothingSimulationSolver::GetWindVelocityAndPressureField(uint32 GroupId)
 {
-	return Evolution->GetVelocityField(GroupId);
+	return Evolution->GetVelocityAndPressureField(GroupId);
 }
 
 void FClothingSimulationSolver::AddExternalForces(uint32 GroupId, bool bUseLegacyWind)
@@ -932,7 +935,7 @@ void FClothingSimulationSolver::Update(Softs::FSolverReal InDeltaTime)
 
 			// Pre-update overridable solver properties first
 			Evolution->SetGravity(Gravity, GroupId);
-			Evolution->GetVelocityField(GroupId).SetVelocity(WindVelocity);
+			Evolution->GetVelocityAndPressureField(GroupId).SetVelocity(WindVelocity);
 
 			Cloth->Update(this);
 		}, /*bForceSingleThreaded =*/ !bClothSolverParallelClothUpdate);
diff --git a/Engine/Plugins/Experimental/ChaosCloth/Source/Chaos/Private/ChaosCloth/ChaosWeightMapTarget.h b/Engine/Plugins/Experimental/ChaosCloth/Source/Chaos/Private/ChaosCloth/ChaosWeightMapTarget.h
--- a/Engine/Plugins/Experimental/ChaosCloth/Source/Chaos/Private/ChaosCloth/ChaosWeightMapTarget.h
+++ b/Engine/Plugins/Experimental/ChaosCloth/Source/Chaos/Private/ChaosCloth/ChaosWeightMapTarget.h
@@ -24,4 +24,5 @@ enum class EChaosWeightMapTarget : uint8
 	BendingStiffness,
 	AreaStiffness,
 	BucklingStiffness,
+	Pressure,
 };
diff --git a/Engine/Plugins/Experimental/ChaosCloth/Source/Chaos/Public/ChaosCloth/ChaosClothConfig.h b/Engine/Plugins/Experimental/ChaosCloth/Source/Chaos/Public/ChaosCloth/ChaosClothConfig.h
--- a/Engine/Plugins/Experimental/ChaosCloth/Source/Chaos/Public/ChaosCloth/ChaosClothConfig.h
+++ b/Engine/Plugins/Experimental/ChaosCloth/Source/Chaos/Public/ChaosCloth/ChaosClothConfig.h
@@ -300,6 +300,16 @@ public:
 	UPROPERTY(EditAnywhere, Category = "Environmental Properties", meta = (EditCondition = "bUseGravityOverride"))
 	FVector Gravity = { 0.f, 0.f, -980.665f };
 
+	/** 
+	 * Pressure force strength applied in the normal direction(use negative value to push toward backface)
+	 * If an enabled Weight Map (Mask with values in the range [0;1]) targeting the "Pressure" is added to the cloth, 
+	 * then both the Low and High values will be used in conjunction with the per particle Weight stored
+	 * in the Weight Map to interpolate the final value from them.
+	 * Otherwise only the Low value is meaningful and sufficient to set the pressure.	 
+	 */
+	UPROPERTY(EditAnywhere, Category = "Environmental Properties", meta = (UIMin = "-10", UIMax = "10", ClampMin = "-100", ClampMax = "100"))
+	FChaosClothWeightedValue Pressure = { 0.0f, 1.f };
+
 	/**
 	 * The strength of the constraint driving the cloth towards the animated goal mesh.
 	 * If an enabled Weight Map (Mask with values in the range [0;1]) targeting the "Anim Drive Stiffness" is added to the cloth, 
diff --git a/Engine/Plugins/Experimental/ChaosCloth/Source/Chaos/Public/ChaosCloth/ChaosClothingSimulation.h b/Engine/Plugins/Experimental/ChaosCloth/Source/Chaos/Public/ChaosCloth/ChaosClothingSimulation.h
--- a/Engine/Plugins/Experimental/ChaosCloth/Source/Chaos/Public/ChaosCloth/ChaosClothingSimulation.h
+++ b/Engine/Plugins/Experimental/ChaosCloth/Source/Chaos/Public/ChaosCloth/ChaosClothingSimulation.h
@@ -115,7 +115,9 @@ namespace Chaos
 		CHAOSCLOTH_API void DebugDrawAnimDrive(FPrimitiveDrawInterface* PDI = nullptr) const;
 		CHAOSCLOTH_API void DebugDrawBendingConstraint(FPrimitiveDrawInterface* PDI = nullptr) const;
 		CHAOSCLOTH_API void DebugDrawLongRangeConstraint(FPrimitiveDrawInterface* PDI = nullptr) const;
-		CHAOSCLOTH_API void DebugDrawWindForces(FPrimitiveDrawInterface* PDI = nullptr) const;
+		CHAOSCLOTH_API void DebugDrawWindAndPressureForces(FPrimitiveDrawInterface* PDI = nullptr) const;
+		UE_DEPRECATED(5.1, "Chaos::Softs::FVelocityField has been renamed FVelocityAndPressureField to match its new behavior.")
+		CHAOSCLOTH_API void DebugDrawWindForces(FPrimitiveDrawInterface* PDI = nullptr) const { return DebugDrawWindAndPressureForces(PDI); }
 		CHAOSCLOTH_API void DebugDrawLocalSpace(FPrimitiveDrawInterface* PDI = nullptr) const;
 		CHAOSCLOTH_API void DebugDrawSelfCollision(FPrimitiveDrawInterface* PDI = nullptr) const;
 		CHAOSCLOTH_API void DebugDrawSelfIntersection(FPrimitiveDrawInterface* PDI = nullptr) const;
diff --git a/Engine/Plugins/Experimental/ChaosCloth/Source/Chaos/Public/ChaosCloth/ChaosClothingSimulationCloth.h b/Engine/Plugins/Experimental/ChaosCloth/Source/Chaos/Public/ChaosCloth/ChaosClothingSimulationCloth.h
--- a/Engine/Plugins/Experimental/ChaosCloth/Source/Chaos/Public/ChaosCloth/ChaosClothingSimulationCloth.h
+++ b/Engine/Plugins/Experimental/ChaosCloth/Source/Chaos/Public/ChaosCloth/ChaosClothingSimulationCloth.h
@@ -58,6 +58,7 @@ namespace Chaos
 			const TVec2<FRealSingle>& InDrag,
 			const TVec2<FRealSingle>& InLift,
 			bool bInUseLegacyWind,
+			const TVec2<FRealSingle>& InPressure,
 			FRealSingle InDampingCoefficient,
 			FRealSingle InLocalDampingCoefficient,
 			FRealSingle InCollisionThickness,
@@ -87,6 +88,7 @@ namespace Chaos
 		void SetBackstopProperties(bool bInEnableBackstop) { bEnableBackstop = bInEnableBackstop; }
 		void SetDampingProperties(FRealSingle InDampingCoefficient) { DampingCoefficient = InDampingCoefficient; }
 		void SetAerodynamicsProperties(const TVec2<FRealSingle>& InDrag, const TVec2<FRealSingle>& InLift, FRealSingle InAirDensity, const FVec3& InWindVelocity) { Drag = InDrag; Lift = InLift; InAirDensity = AirDensity; WindVelocity = InWindVelocity; }
+		void SetPressureProperties(const TVec2<FRealSingle>& InPressure) { Pressure = InPressure; }
 		void SetGravityProperties(FRealSingle InGravityScale, bool bInIsGravityOverridden, const FVec3& InGravityOverride) { GravityScale = InGravityScale; bIsGravityOverridden = bInIsGravityOverridden; GravityOverride = InGravityOverride; }
 		void SetAnimDriveProperties(const TVec2<FRealSingle>& InAnimDriveStiffness, const TVec2<FRealSingle>& InAnimDriveDamping) { AnimDriveStiffness = InAnimDriveStiffness; AnimDriveDamping = InAnimDriveDamping; }
 		void GetAnimDriveProperties(TVec2<FRealSingle>& OutAnimDriveStiffness, TVec2<FRealSingle>& OutAnimDriveDamping) { OutAnimDriveStiffness = AnimDriveStiffness; OutAnimDriveDamping = AnimDriveDamping; }
@@ -222,6 +224,7 @@ namespace Chaos
 		FVec3 WindVelocity;
 		FRealSingle AirDensity;
 		bool bUseLegacyWind;
+		TVec2<FRealSingle> Pressure;
 		FRealSingle DampingCoefficient;
 		FRealSingle LocalDampingCoefficient;
 		FRealSingle CollisionThickness;
diff --git a/Engine/Plugins/Experimental/ChaosCloth/Source/Chaos/Public/ChaosCloth/ChaosClothingSimulationInteractor.h b/Engine/Plugins/Experimental/ChaosCloth/Source/Chaos/Public/ChaosCloth/ChaosClothingSimulationInteractor.h
--- a/Engine/Plugins/Experimental/ChaosCloth/Source/Chaos/Public/ChaosCloth/ChaosClothingSimulationInteractor.h
+++ b/Engine/Plugins/Experimental/ChaosCloth/Source/Chaos/Public/ChaosCloth/ChaosClothingSimulationInteractor.h
@@ -52,6 +52,9 @@ public:
 	UFUNCTION(BlueprintCallable, Category = ClothingSimulation, Meta = (Keywords = "Chaos Wind Drag Lift Coefficient Air Density Velocity"))
 	void SetWind(FVector2D Drag = FVector2D(0.07f, 0.5f), FVector2D Lift = FVector2D(0.07f, 0.5f), float AirDensity = 1.225e-6f, FVector WindVelocity = FVector(0.f, 0.f, 0.f));
 
+	UFUNCTION(BlueprintCallable, Category = ClothingSimulation, Meta = (Keywords = "Chaos Pressure"))
+	void SetPressure(FVector2D Pressure = FVector2D(0.f, 1.f));
+
 	UFUNCTION(BlueprintCallable, Category = ClothingSimulation, Meta = (Keywords = "Chaos Scale Override"))
 	void SetGravity(float GravityScale = 1.f, bool bIsGravityOverridden = false, FVector GravityOverride = FVector(0.f, 0.f, -981.f));
 
diff --git a/Engine/Plugins/Experimental/ChaosCloth/Source/Chaos/Public/ChaosCloth/ChaosClothingSimulationSolver.h b/Engine/Plugins/Experimental/ChaosCloth/Source/Chaos/Public/ChaosCloth/ChaosClothingSimulationSolver.h
--- a/Engine/Plugins/Experimental/ChaosCloth/Source/Chaos/Public/ChaosCloth/ChaosClothingSimulationSolver.h
+++ b/Engine/Plugins/Experimental/ChaosCloth/Source/Chaos/Public/ChaosCloth/ChaosClothingSimulationSolver.h
@@ -109,14 +109,28 @@ namespace Chaos
 		// Set per group wind velocity, used to override solver's wind velocity. Must be called during cloth update.
 		void SetWindVelocity(uint32 GroupId, const TVec3<FRealSingle>& InWindVelocity);
 
-		// Set the geometry affected by the wind.
-		void SetWindGeometry(uint32 GroupId, const FTriangleMesh& TriangleMesh, const TConstArrayView<FRealSingle>& DragMultipliers, const TConstArrayView<FRealSingle>& LiftMultipliers);
-
-		// Set the wind properties.
-		void SetWindProperties(uint32 GroupId, const TVec2<FRealSingle>& Drag, const TVec2<FRealSingle>& Lift, FRealSingle AirDensity = 1.225e-6f);
-
-		// Return the wind velocity field associated with a given group id.
-		const Softs::FVelocityField& GetWindVelocityField(uint32 GroupId);
+		// Set the geometry affected by the wind and pressure
+		void SetWindAndPressureGeometry(uint32 GroupId, const FTriangleMesh& TriangleMesh, const TConstArrayView<FRealSingle>& DragMultipliers, const TConstArrayView<FRealSingle>& LiftMultipliers, const TConstArrayView<FRealSingle>& PressureMultipliers);
+		UE_DEPRECATED(5.1, "Chaos::Softs::FVelocityField has been renamed FVelocityAndPressureField to match its new behavior.")
+		void SetWindGeometry(uint32 GroupId, const FTriangleMesh& TriangleMesh, const TConstArrayView<FRealSingle>& DragMultipliers, const TConstArrayView<FRealSingle>& LiftMultipliers)
+		{
+			SetWindAndPressureGeometry(GroupId, TriangleMesh, DragMultipliers, LiftMultipliers, TConstArrayView<FRealSingle>());
+		}
+
+		// Set the wind and pressure properties.
+		void SetWindAndPressureProperties(uint32 GroupId, const TVec2<FRealSingle>& Drag, const TVec2<FRealSingle>& Lift, FRealSingle AirDensity = 1.225e-6f, const TVec2<FRealSingle>& Pressure = TVec2<FRealSingle>::ZeroVector);
+		UE_DEPRECATED(5.1, "Chaos::Softs::FVelocityField has been renamed FVelocityAndPressureField to match its new behavior.")
+		void SetWindProperties(uint32 GroupId, const TVec2<FRealSingle>& Drag, const TVec2<FRealSingle>& Lift, FRealSingle AirDensity = 1.225e-6f)
+		{
+			SetWindAndPressureProperties(GroupId, Drag, Lift, AirDensity);
+		}
+
+		// Return the wind velocity and pressure field associated with a given group id.
+		const Softs::FVelocityAndPressureField& GetWindVelocityAndPressureField(uint32 GroupId);
+		UE_DEPRECATED(5.1, "Chaos::Softs::FVelocityField has been renamed FVelocityAndPressureField to match its new behavior.")
+		PRAGMA_DISABLE_DEPRECATION_WARNINGS
+		const Softs::FVelocityField& GetWindVelocityField(uint32 GroupId) { return GetWindVelocityAndPressureField(GroupId); }
+		PRAGMA_ENABLE_DEPRECATION_WARNINGS
 
 		// Add external forces to the particles
 		void AddExternalForces(uint32 GroupId, bool bUseLegacyWind);
diff --git a/Engine/Plugins/Experimental/ChaosClothEditor/Source/Chaos/Private/ChaosClothEditor/ChaosSimulationEditorExtender.cpp b/Engine/Plugins/Experimental/ChaosClothEditor/Source/Chaos/Private/ChaosClothEditor/ChaosSimulationEditorExtender.cpp
--- a/Engine/Plugins/Experimental/ChaosClothEditor/Source/Chaos/Private/ChaosClothEditor/ChaosSimulationEditorExtender.cpp
+++ b/Engine/Plugins/Experimental/ChaosClothEditor/Source/Chaos/Private/ChaosClothEditor/ChaosSimulationEditorExtender.cpp
@@ -80,7 +80,7 @@ const FVisualizationOption FVisualizationOption::OptionData[] =
 	FVisualizationOption(&Chaos::FClothingSimulation::DebugDrawAnimDrive           , TEXT("p.ChaosClothEditor.DebugDrawAnimDrive"           ), LOCTEXT("ChaosVisName_AnimDrive"           , "Anim Drive"                 ), LOCTEXT("ChaosVisName_AnimDrive_Tooltip"           , "Draws the current skinned reference mesh for the simulation which anim drive will attempt to reach if enabled")),
 	FVisualizationOption(&Chaos::FClothingSimulation::DebugDrawBendingConstraint   , TEXT("p.ChaosClothEditor.DebugDrawBendingConstraint"   ), LOCTEXT("ChaosVisName_BendingConstraint"   , "Bending Constraint"         ), LOCTEXT("ChaosVisName_BendingConstraint_Tooltip"   , "Draws the bending spring constraints")),
 	FVisualizationOption(&Chaos::FClothingSimulation::DebugDrawLongRangeConstraint , TEXT("p.ChaosClothEditor.DebugDrawLongRangeConstraint" ), LOCTEXT("ChaosVisName_LongRangeConstraint" , "Long Range Constraint"      ), LOCTEXT("ChaosVisName_LongRangeConstraint_Tooltip" , "Draws the long range attachment constraint distances")),
-	FVisualizationOption(&Chaos::FClothingSimulation::DebugDrawWindForces          , TEXT("p.ChaosClothEditor.DebugDrawWindForces"          ), LOCTEXT("ChaosVisName_WindForces"          , "Wind Aerodynamic Forces"    ), LOCTEXT("ChaosVisName_Wind_Tooltip"                , "Draws the Wind drag and lift forces")),
+	FVisualizationOption(&Chaos::FClothingSimulation::DebugDrawWindAndPressureForces, TEXT("p.ChaosClothEditor.DebugDrawWindAndPressureForces"), LOCTEXT("ChaosVisName_WindAndPressureForces", "Wind Aerodynamic And Pressure Forces"), LOCTEXT("ChaosVisName_WindAndPressure_Tooltip", "Draws the Wind drag and lift and pressure forces")),
 	FVisualizationOption(&Chaos::FClothingSimulation::DebugDrawLocalSpace          , TEXT("p.ChaosClothEditor.DebugDrawLocalSpace"          ), LOCTEXT("ChaosVisName_LocalSpace"          , "Local Space Reference Bone" ), LOCTEXT("ChaosVisName_LocalSpace_Tooltip"          , "Draws the local space reference bone")),
 	FVisualizationOption(&Chaos::FClothingSimulation::DebugDrawSelfCollision       , TEXT("p.ChaosClothEditor.DebugDrawSelfCollision"       ), LOCTEXT("ChaosVisName_SelfCollision"       , "Self Collision"             ), LOCTEXT("ChaosVisName_SelfCollision_Tooltip"       , "Draws the self collision thickness/debugging information")),
 	FVisualizationOption(&Chaos::FClothingSimulation::DebugDrawSelfIntersection    , TEXT("p.ChaosClothEditor.DebugDrawSelfIntersection"    ), LOCTEXT("ChaosVisName_SelfIntersection"    , "Self Intersection"          ), LOCTEXT("ChaosVisName_SelfIntersection_Tooltip"    , "Draws the self intersection contour/region information")),
diff --git a/Engine/Source/Runtime/Experimental/Chaos/Private/Chaos/PBDEvolution.cpp b/Engine/Source/Runtime/Experimental/Chaos/Private/Chaos/PBDEvolution.cpp
--- a/Engine/Source/Runtime/Experimental/Chaos/Private/Chaos/PBDEvolution.cpp
+++ b/Engine/Source/Runtime/Experimental/Chaos/Private/Chaos/PBDEvolution.cpp
@@ -98,7 +98,7 @@ FPBDEvolution::FPBDEvolution(
 {
 	// Add group arrays
 	TArrayCollection::AddArray(&MGroupGravityAccelerations);
-	TArrayCollection::AddArray(&MGroupVelocityFields);
+	TArrayCollection::AddArray(&MGroupVelocityAndPressureFields);
 	TArrayCollection::AddArray(&MGroupForceRules);
 	TArrayCollection::AddArray(&MGroupCollisionThicknesses);
 	TArrayCollection::AddArray(&MGroupCoefficientOfFrictions);
@@ -218,13 +218,13 @@ void FPBDEvolution::PreIterationUpdate(
 	const uint32 ParticleGroupId = MParticleGroupIds[Offset];
 	const TFunction<void(FSolverParticles&, const FSolverReal, const int32)>& ForceRule = MGroupForceRules[ParticleGroupId];
 	const FSolverVec3& Gravity = MGroupGravityAccelerations[ParticleGroupId];
-	FVelocityField& VelocityField = MGroupVelocityFields[ParticleGroupId];
+	FVelocityAndPressureField& VelocityAndPressureField = MGroupVelocityAndPressureFields[ParticleGroupId];
 
 	if (bVelocityField)
 	{
 		TRACE_CPUPROFILER_EVENT_SCOPE(ChaosPBDVelocityFieldUpdateForces);
 		SCOPE_CYCLE_COUNTER(STAT_ChaosPBDVelocityFieldUpdateForces);
-		VelocityField.UpdateForces(MParticles, Dt);  // Update force per surface element
+		VelocityAndPressureField.UpdateForces(MParticles, Dt);  // Update force per surface element
 	}
 
 	FPerParticleDampVelocity DampVelocityRule(MGroupLocalDampings[ParticleGroupId]);
@@ -262,7 +262,7 @@ void FPBDEvolution::PreIterationUpdate(
 
 		const int32 RangeSize = Range - Offset;
 		PhysicsParallelFor(RangeSize,
-			[this, &Offset, &ForceRule, &Gravity, &VelocityField, &DampVelocityRule, DampingPowDt, DampingIntegrated, Dt](int32 i)
+			[this, &Offset, &ForceRule, &Gravity, &VelocityAndPressureField, &DampVelocityRule, DampingPowDt, DampingIntegrated, Dt](int32 i)
 			{
 				const int32 Index = Offset + i;
 				if (MParticles.InvM(Index) != (FSolverReal)0.)  // Process dynamic particles
@@ -279,7 +279,7 @@ void FPBDEvolution::PreIterationUpdate(
 					// Velocity Field
 					if (bVelocityField)
 					{
-						VelocityField.Apply(MParticles, Dt, Index);
+						VelocityAndPressureField.Apply(MParticles, Dt, Index);
 					}
 
 					// Euler Step Velocity
@@ -339,7 +339,7 @@ void FPBDEvolution::AdvanceOneTimeStep(const FSolverReal Dt, const bool bSmoothD
 			{
 				const uint32 ParticleGroupId = MParticleGroupIds[Offset];
 
-				if (MGroupVelocityFields[ParticleGroupId].IsActive())
+				if (MGroupVelocityAndPressureFields[ParticleGroupId].IsActive())
 				{
 					if (MGroupLocalDampings[ParticleGroupId] > (FSolverReal)0.)
 					{
diff --git a/Engine/Source/Runtime/Experimental/Chaos/Private/Chaos/VelocityField.cpp b/Engine/Source/Runtime/Experimental/Chaos/Private/Chaos/VelocityField.cpp
--- a/Engine/Source/Runtime/Experimental/Chaos/Private/Chaos/VelocityField.cpp
+++ b/Engine/Source/Runtime/Experimental/Chaos/Private/Chaos/VelocityField.cpp
@@ -7,7 +7,6 @@
 
 static_assert(sizeof(ispc::FVector3f) == sizeof(Chaos::Softs::FSolverVec3), "sizeof(ispc::FVector3f) != sizeof(Chaos::Softs::FSolverVec3)");
 static_assert(sizeof(ispc::FIntVector) == sizeof(Chaos::TVec3<int32>), "sizeof(ispc::FIntVector) != sizeof(Chaos::TVec3<int32>)");
-static_assert(sizeof(ispc::FVector2f) == sizeof(Chaos::Softs::FSolverVec2), "sizeof(ispc::FVector2f) != sizeof(Chaos::Softs::FSolverVec2)");
 #endif
 
 #if INTEL_ISPC && !UE_BUILD_SHIPPING
@@ -17,7 +16,7 @@ FAutoConsoleVariableRef CVarChaosVelocityFieldISPCEnabled(TEXT("p.Chaos.Velocity
 
 namespace Chaos::Softs {
 
-void FVelocityField::SetGeometry(const FTriangleMesh* TriangleMesh, const TConstArrayView<FRealSingle>& DragMultipliers, const TConstArrayView<FRealSingle>& LiftMultipliers)
+void FVelocityAndPressureField::SetGeometry(const FTriangleMesh* TriangleMesh, const TConstArrayView<FRealSingle>& DragMultipliers, const TConstArrayView<FRealSingle>& LiftMultipliers, const TConstArrayView<FRealSingle>& PressureMultipliers)
 {
 	if (TriangleMesh)
 	{
@@ -30,18 +29,11 @@ void FVelocityField::SetGeometry(const FTriangleMesh* TriangleMesh, const TConst
 
 		Forces.SetNumUninitialized(Elements.Num());
 
-		enum class EWeighMaps : uint8
-		{
-			None = 0,
-			DragMultipliers = 1 << 0,
-			LiftMultipliers = 1 << 1
-		};
-
-		const EWeighMaps WeighMaps = (EWeighMaps)(
-			(uint8)(DragMultipliers.Num() == NumParticles ? EWeighMaps::DragMultipliers : EWeighMaps::None) |
-			(uint8)(LiftMultipliers.Num() == NumParticles ? EWeighMaps::LiftMultipliers : EWeighMaps::None));
+		const bool bHasDragMultipliers = DragMultipliers.Num() == NumParticles;
+		const bool bHasLiftMultipliers = LiftMultipliers.Num() == NumParticles;
+		const bool bHasPressureMultipliers = PressureMultipliers.Num() == NumParticles;
 
-		if (WeighMaps == EWeighMaps::None)
+		if (!bHasDragMultipliers && !bHasLiftMultipliers && !bHasPressureMultipliers)
 		{
 			Multipliers.Reset();
 		}
@@ -58,20 +50,11 @@ void FVelocityField::SetGeometry(const FTriangleMesh* TriangleMesh, const TConst
 				const int32 I1 = Element[1] - Offset;
 				const int32 I2 = Element[2] - Offset;
 
-				switch (WeighMaps)
-				{
-				case EWeighMaps::DragMultipliers:
-					Multipliers[ElementIndex] = FSolverVec2((FSolverReal)(DragMultipliers[I0] + DragMultipliers[I1] + DragMultipliers[I2]) * OneThird, (FSolverReal)0.); break;
-					break;
-				case EWeighMaps::LiftMultipliers:
-					Multipliers[ElementIndex] = FSolverVec2((FSolverReal)0., (FSolverReal)(LiftMultipliers[I0] + LiftMultipliers[I1] + LiftMultipliers[I2]) * OneThird); break;
-					break;
-				default:
-					Multipliers[ElementIndex] = FSolverVec2(
-						(FSolverReal)(DragMultipliers[I0] + DragMultipliers[I1] + DragMultipliers[I2]) * OneThird,
-						(FSolverReal)(LiftMultipliers[I0] + LiftMultipliers[I1] + LiftMultipliers[I2]) * OneThird);
-					break;
-				}
+				const FSolverReal DragMultiplier = bHasDragMultipliers ? (FSolverReal)(DragMultipliers[I0] + DragMultipliers[I1] + DragMultipliers[I2]) * OneThird : (FSolverReal)0.;
+				const FSolverReal LiftMultiplier = bHasLiftMultipliers ? (FSolverReal)(LiftMultipliers[I0] + LiftMultipliers[I1] + LiftMultipliers[I2]) * OneThird : (FSolverReal)0.;
+				const FSolverReal PressureMultiplier = bHasPressureMultipliers ? (FSolverReal)(PressureMultipliers[I0] + PressureMultipliers[I1] + PressureMultipliers[I2]) * OneThird : (FSolverReal)0.;
+
+				Multipliers[ElementIndex] = FSolverVec3(DragMultiplier, LiftMultiplier, PressureMultiplier);
 			}
 		}
 	}
@@ -83,13 +66,13 @@ void FVelocityField::SetGeometry(const FTriangleMesh* TriangleMesh, const TConst
 		NumParticles = 0;
 		Forces.Reset();
 		Multipliers.Reset();
-		SetProperties(FSolverVec2(0.), FSolverVec2(0.), (FSolverReal)0.);
+		SetProperties(FSolverVec2(0.), FSolverVec2(0.), (FSolverReal)0., FSolverVec2(0.));
 	}
 }
 
 
 
-void FVelocityField::UpdateForces(const FSolverParticles& InParticles, const FSolverReal /*Dt*/)
+void FVelocityAndPressureField::UpdateForces(const FSolverParticles& InParticles, const FSolverReal /*Dt*/)
 {
 	if (!Multipliers.Num())
 	{
@@ -105,6 +88,7 @@ void FVelocityField::UpdateForces(const FSolverParticles& InParticles, const FSo
 				QuarterRho,
 				DragBase,
 				LiftBase,
+				PressureBase,
 				Elements.Num());
 		}
 		else
@@ -112,7 +96,7 @@ void FVelocityField::UpdateForces(const FSolverParticles& InParticles, const FSo
 		{
 			for (int32 ElementIndex = 0; ElementIndex < Elements.Num(); ++ElementIndex)
 			{
-				UpdateField(InParticles, ElementIndex, Velocity, DragBase, LiftBase);
+				UpdateField(InParticles, ElementIndex, Velocity, DragBase, LiftBase, PressureBase);
 			}
 		}
 	}
@@ -126,13 +110,15 @@ void FVelocityField::UpdateForces(const FSolverParticles& InParticles, const FSo
 				(const ispc::FIntVector*)Elements.GetData(),
 				(const ispc::FVector3f*)InParticles.GetV().GetData(),
 				(const ispc::FVector3f*)InParticles.XArray().GetData(),
-				(const ispc::FVector2f*)Multipliers.GetData(),
+				(const ispc::FVector3f*)Multipliers.GetData(),
 				(const ispc::FVector3f&)Velocity,
 				QuarterRho,
 				DragBase,
 				DragRange,
 				LiftBase,
 				LiftRange,
+				PressureBase,
+				PressureRange,
 				Elements.Num());
 		}
 		else
@@ -140,11 +126,12 @@ void FVelocityField::UpdateForces(const FSolverParticles& InParticles, const FSo
 		{
 			for (int32 ElementIndex = 0; ElementIndex < Elements.Num(); ++ElementIndex)
 			{
-				const FSolverVec2& Multiplier = Multipliers[ElementIndex];
+				const FSolverVec3& Multiplier = Multipliers[ElementIndex];
 				const FSolverReal Cd = DragBase + DragRange * Multiplier[0];
 				const FSolverReal Cl = LiftBase + LiftRange * Multiplier[1];
+				const FSolverReal Cp = PressureBase + PressureRange * Multiplier[2];
 
-				UpdateField(InParticles, ElementIndex, Velocity, Cd, Cl);
+				UpdateField(InParticles, ElementIndex, Velocity, Cd, Cl, Cp);
 			}
 		}
 	}
diff --git a/Engine/Source/Runtime/Experimental/Chaos/Private/Chaos/VelocityField.ispc b/Engine/Source/Runtime/Experimental/Chaos/Private/Chaos/VelocityField.ispc
--- a/Engine/Source/Runtime/Experimental/Chaos/Private/Chaos/VelocityField.ispc
+++ b/Engine/Source/Runtime/Experimental/Chaos/Private/Chaos/VelocityField.ispc
@@ -21,6 +21,7 @@ export void UpdateField(uniform FVector3f Forces[],
 								const uniform float QuarterRho,
 								const uniform float Cd,
 								const uniform float Cl,
+								const uniform float Cp,
 								const uniform int32 NumElements)
 {
 	varying FIntVector Element;
@@ -69,7 +70,7 @@ export void UpdateField(uniform FVector3f Forces[],
 		const FVector3f EvenForce = (Cd - Cl) * VDotN * RelVelocity + Cl * VSquare * N;
 		const FVector3f OddForce = (Cl - Cd) * VDotN * RelVelocity - Cl * VSquare * N;
 
-		const FVector3f Force = QuarterRho * DoubleArea * VectorSelect(VDotN >= 0, EvenForce, OddForce);
+		const FVector3f Force = QuarterRho * DoubleArea * VectorSelect(VDotN >= 0, EvenForce, OddForce) + DoubleArea * 0.5 * Cp * N;
 
 		VectorStore(&Forces[Offset], Force);
 
@@ -81,13 +82,15 @@ export void UpdateFieldWithWeightMaps(uniform FVector3f Forces[],
 								const uniform FIntVector Elements[],
 								const uniform FVector3f V[],
 								const uniform FVector3f X[],
-								const uniform FVector2f Multipliers[],
+								const uniform FVector3f Multipliers[],
 								const uniform FVector3f& Velocity,
 								const uniform float QuarterRho,
 								const uniform float DragBase,
 								const uniform float DragRange,
 								const uniform float LiftBase,
 								const uniform float LiftRange,
+								const uniform float PressureBase,
+								const uniform float PressureRange,
 								const uniform int32 NumElements)
 {
 	varying FIntVector Element;
@@ -109,7 +112,7 @@ export void UpdateFieldWithWeightMaps(uniform FVector3f Forces[],
 		const FVector3f X1 = VectorGather(&X[Element.V[1]]);
 		const FVector3f X2 = VectorGather(&X[Element.V[2]]);
 
-		const FVector2f Multiplier = VectorGather(&Multipliers[ElementIndex]);
+		const FVector3f Multiplier = VectorGather(&Multipliers[ElementIndex]);
 
 		// Calculate the normal and the area of the surface exposed to the flow
 		FVector3f N = VectorCross(X1 - X0, X2 - X0);
@@ -129,6 +132,7 @@ export void UpdateFieldWithWeightMaps(uniform FVector3f Forces[],
 
 		const float Cd = DragBase + DragRange * Multiplier.V[0];
 		const float Cl = LiftBase + LiftRange * Multiplier.V[1];
+		const float Cp = PressureBase + PressureRange * Multiplier.V[2];
 
 #if HW_GATHER_SUPPORTED == 1
 		if(ElementIndex + programCount < NumElements)
@@ -140,7 +144,7 @@ export void UpdateFieldWithWeightMaps(uniform FVector3f Forces[],
 		const FVector3f EvenForce = (Cd - Cl) * VDotN * RelVelocity + Cl * VSquare * N;
 		const FVector3f OddForce = (Cl - Cd) * VDotN * RelVelocity - Cl * VSquare * N;
 
-		const FVector3f Force = QuarterRho * DoubleArea * VectorSelect(VDotN >= 0, EvenForce, OddForce);
+		const FVector3f Force = QuarterRho * DoubleArea * VectorSelect(VDotN >= 0, EvenForce, OddForce) + DoubleArea * 0.5 * Cp * N;
 
 		VectorStore(&Forces[Offset], Force);
 
diff --git a/Engine/Source/Runtime/Experimental/Chaos/Public/Chaos/PBDEvolution.h b/Engine/Source/Runtime/Experimental/Chaos/Public/Chaos/PBDEvolution.h
--- a/Engine/Source/Runtime/Experimental/Chaos/Public/Chaos/PBDEvolution.h
+++ b/Engine/Source/Runtime/Experimental/Chaos/Public/Chaos/PBDEvolution.h
@@ -112,8 +112,17 @@ class CHAOS_API FPBDEvolution : public TArrayCollection
 	const FSolverVec3& GetGravity(const uint32 GroupId = 0) const { check(GroupId < TArrayCollection::Size()); return MGroupGravityAccelerations[GroupId]; }
 	void SetGravity(const FSolverVec3& Acceleration, const uint32 GroupId = 0) { check(GroupId < TArrayCollection::Size()); MGroupGravityAccelerations[GroupId] = Acceleration; }
 
-	FVelocityField& GetVelocityField(const uint32 GroupId = 0) { check(GroupId < TArrayCollection::Size()); return MGroupVelocityFields[GroupId]; }
-	const FVelocityField& GetVelocityField(const uint32 GroupId = 0) const { check(GroupId < TArrayCollection::Size()); return MGroupVelocityFields[GroupId]; }
+	FVelocityAndPressureField& GetVelocityAndPressureField(const uint32 GroupId = 0) { check(GroupId < TArrayCollection::Size()); return MGroupVelocityAndPressureFields[GroupId]; }
+	const FVelocityAndPressureField& GetVelocityAndPressureField(const uint32 GroupId = 0) const { check(GroupId < TArrayCollection::Size()); return MGroupVelocityAndPressureFields[GroupId]; }
+
+	UE_DEPRECATED(5.1, "Chaos::Softs::FVelocityField has been renamed FVelocityAndPressureField to match its new behavior.")
+	PRAGMA_DISABLE_DEPRECATION_WARNINGS
+	FVelocityField& GetVelocityField(const uint32 GroupId = 0) { return GetVelocityAndPressureField(GroupId); }
+	PRAGMA_ENABLE_DEPRECATION_WARNINGS
+	UE_DEPRECATED(5.1, "Chaos::Softs::FVelocityField has been renamed FVelocityAndPressureField to match its new behavior.")
+	PRAGMA_DISABLE_DEPRECATION_WARNINGS
+	const FVelocityField& GetVelocityField(const uint32 GroupId = 0) const { return GetVelocityAndPressureField(GroupId); }
+	PRAGMA_ENABLE_DEPRECATION_WARNINGS
 
 	int32 GetIterations() const { return MNumIterations; }
 	void SetIterations(const int32 Iterations) { MNumIterations = Iterations; }
@@ -165,7 +174,7 @@ private:
 	TArray<FSolverVec3> MCollisionNormals;
 
 	TArrayCollectionArray<FSolverVec3> MGroupGravityAccelerations;
-	TArrayCollectionArray<FVelocityField> MGroupVelocityFields;
+	TArrayCollectionArray<FVelocityAndPressureField> MGroupVelocityAndPressureFields;
 	TArrayCollectionArray<TFunction<void(FSolverParticles&, const FSolverReal, const int32)>> MGroupForceRules;
 	TArrayCollectionArray<FSolverReal> MGroupCollisionThicknesses;
 	TArrayCollectionArray<FSolverReal> MGroupCoefficientOfFrictions;
diff --git a/Engine/Source/Runtime/Experimental/Chaos/Public/Chaos/PBDSoftsEvolutionFwd.h b/Engine/Source/Runtime/Experimental/Chaos/Public/Chaos/PBDSoftsEvolutionFwd.h
--- a/Engine/Source/Runtime/Experimental/Chaos/Public/Chaos/PBDSoftsEvolutionFwd.h
+++ b/Engine/Source/Runtime/Experimental/Chaos/Public/Chaos/PBDSoftsEvolutionFwd.h
@@ -61,6 +61,7 @@ namespace Chaos
 		class FPBDTriangleMeshCollisions;
 
 		// Softs solver forces
-		class FVelocityField;
+		class FVelocityAndPressureField;
+		using FVelocityField UE_DEPRECATED(5.1, "Chaos::Softs::FVelocityField has been renamed FVelocityAndPressureField to match its new behavior.") = FVelocityAndPressureField;
 	}  // End namespace Softs
 }  // End namespace Chaos
diff --git a/Engine/Source/Runtime/Experimental/Chaos/Public/Chaos/VelocityField.h b/Engine/Source/Runtime/Experimental/Chaos/Public/Chaos/VelocityField.h
--- a/Engine/Source/Runtime/Experimental/Chaos/Public/Chaos/VelocityField.h
+++ b/Engine/Source/Runtime/Experimental/Chaos/Public/Chaos/VelocityField.h
@@ -9,7 +9,7 @@ namespace Chaos::Softs
 {
 
 // Velocity field used solely for aerodynamics effects, use Chaos Fields for other types of fields.
-class CHAOS_API FVelocityField final
+class CHAOS_API FVelocityAndPressureField final
 {
 public:
 	static constexpr FSolverReal DefaultDragCoefficient = (FSolverReal)0.5;
@@ -17,14 +17,14 @@ public:
 	static constexpr FSolverReal DefaultFluidDensity = (FSolverReal)1.225e-6;
 
 	// Construct an uninitialized field. Mesh, properties, and velocity will have to be set for this field to be valid.
-	FVelocityField()
+	FVelocityAndPressureField()
 		: Offset(INDEX_NONE)
 		, NumParticles(0)
 	{
 		SetProperties(FSolverVec2(0.), FSolverVec2(0.), (FSolverReal)0.);
 	}
 
-	~FVelocityField() {}
+	~FVelocityAndPressureField() {}
 
 	void UpdateForces(const FSolverParticles& InParticles, const FSolverReal /*Dt*/);
 
@@ -39,7 +39,7 @@ public:
 		}
 	}
 
-	void SetProperties(const FSolverVec2& Drag, const FSolverVec2& Lift, const FSolverReal FluidDensity)
+	void SetProperties(const FSolverVec2& Drag, const FSolverVec2& Lift, const FSolverReal FluidDensity, const FSolverVec2& Pressure = FSolverVec2::ZeroVector)
 	{
 		constexpr FSolverReal OneQuarter = (FSolverReal)0.25;
 		QuarterRho = FluidDensity * OneQuarter;
@@ -50,14 +50,17 @@ public:
 		DragRange = FMath::Clamp(Drag[1], MinCoefficient, MaxCoefficient) - DragBase;
 		LiftBase = FMath::Clamp(Lift[0], MinCoefficient, MaxCoefficient);
 		LiftRange = FMath::Clamp(Lift[1], MinCoefficient, MaxCoefficient) - LiftBase;
+		PressureBase = Pressure[0];
+		PressureRange = Pressure[1] - PressureBase;
 	}
 
 	bool IsActive() const 
-	{
-		return (DragBase > (FSolverReal)0. || DragRange != (FSolverReal)0.) || (LiftBase > (FSolverReal)0. || LiftRange != (FSolverReal)0.);  // Note: range can be a negative value (although not when base is zero)
+	{ 
+		// Note: range can be a negative value (although not when Lift or Drag base is zero)
+		return (DragBase > (FSolverReal)0. || DragRange != (FSolverReal)0.) || (LiftBase > (FSolverReal)0. || LiftRange != (FSolverReal)0.) || PressureBase != (FSolverReal)0. || PressureRange != (FSolverReal)0.; 
 	}
 
-	void SetGeometry(const FTriangleMesh* TriangleMesh, const TConstArrayView<FRealSingle>& DragMultipliers, const TConstArrayView<FRealSingle>& LiftMultipliers);
+	void SetGeometry(const FTriangleMesh* TriangleMesh, const TConstArrayView<FRealSingle>& DragMultipliers, const TConstArrayView<FRealSingle>& LiftMultipliers, const TConstArrayView<FRealSingle>& PressureMultipliers = TConstArrayView<FRealSingle>());
 
 	void SetVelocity(const FSolverVec3& InVelocity) { Velocity = InVelocity; }
 
@@ -65,7 +68,7 @@ public:
 	TConstArrayView<FSolverVec3> GetForces() const { return TConstArrayView<FSolverVec3>(Forces); }
 
 private:
-	void UpdateField(const FSolverParticles& InParticles, int32 ElementIndex, const FSolverVec3& InVelocity, const FSolverReal Cd, const FSolverReal Cl)
+	void UpdateField(const FSolverParticles& InParticles, int32 ElementIndex, const FSolverVec3& InVelocity, const FSolverReal Cd, const FSolverReal Cl, const FSolverReal Cp)
 	{
 		const TVec3<int32>& Element = Elements[ElementIndex];
 
@@ -88,24 +91,28 @@ private:
 
 		Forces[ElementIndex] = QuarterRho * DoubleArea * (VDotN >= (FSolverReal)0. ?  // The flow can hit either side of the triangle, so the normal might need to be reversed
 			(Cd - Cl) * VDotN * V + Cl * VSquare * N :
-			(Cl - Cd) * VDotN * V - Cl * VSquare * N);
+			(Cl - Cd) * VDotN * V - Cl * VSquare * N) + DoubleArea * (FSolverReal)0.5 * Cp * N;
 	}
 
 private:
 	TConstArrayView<TArray<int32>> PointToTriangleMap;
 	TConstArrayView<TVec3<int32>> Elements;
 	TArray<FSolverVec3> Forces;
-	TArray<FSolverVec2> Multipliers;
+	TArray<FSolverVec3> Multipliers;
 	FSolverVec3 Velocity;
 	FSolverReal DragBase;
 	FSolverReal DragRange;
 	FSolverReal LiftBase;
 	FSolverReal LiftRange;
+	FSolverReal PressureBase; 
+	FSolverReal PressureRange;
 	FSolverReal QuarterRho;
 	int32 Offset;
 	int32 NumParticles;
 };
 
+using FVelocityField UE_DEPRECATED(5.1, "Chaos::Softs::FVelocityField has been renamed FVelocityAndPressureField to match its new behavior.") = FVelocityAndPressureField;
+
 }  // End namespace Chaos::Softs
 
 // Support ISPC enable/disable in non-shipping builds
