Chaos - add cvar to override collision stiffness

diff --git a/Engine/Source/Runtime/Experimental/Chaos/Private/Chaos/Collision/PBDCollisionConstraint.cpp b/Engine/Source/Runtime/Experimental/Chaos/Private/Chaos/Collision/PBDCollisionConstraint.cpp
--- a/Engine/Source/Runtime/Experimental/Chaos/Private/Chaos/Collision/PBDCollisionConstraint.cpp
+++ b/Engine/Source/Runtime/Experimental/Chaos/Private/Chaos/Collision/PBDCollisionConstraint.cpp
@@ -46,6 +46,9 @@ namespace Chaos
 	FRealSingle Chaos_Collision_ConvexZeroMargin = 0.0f;
 	FAutoConsoleVariableRef CVarChaos_Collision_ConvexZeroMargin(TEXT("p.Chaos.Collision.ConvexZeroMargin"), Chaos_Collision_ConvexZeroMargin, TEXT(""));
 
+	FRealSingle Chaos_Collision_Stiffness = -1.0f;
+	FAutoConsoleVariableRef CVarChaos_Collision_Stiffness(TEXT("p.Chaos.Collision.Stiffness"), Chaos_Collision_Stiffness, TEXT("Override the collision solver stiffness (if >= 0)"));
+
 	struct FCollisionTolerances
 	{
 		// Multiplied by the contact margin to produce a distance within which contacts are considered to be the same point
@@ -273,6 +276,12 @@ namespace Chaos
 		const EImplicitObjectType ImplicitType0 = GetInnerType(GetImplicit0()->GetCollisionType());
 		const EImplicitObjectType ImplicitType1 = GetInnerType(GetImplicit1()->GetCollisionType());
 		InitMarginsAndTolerances(ImplicitType0, ImplicitType1, Margin0, Margin1);
+
+		// Debug testing for solver stiffness
+		if (Chaos_Collision_Stiffness >= 0)
+		{
+			SetStiffness(Chaos_Collision_Stiffness);
+		}
 	}
 
 	void FPBDCollisionConstraint::InitMarginsAndTolerances(const EImplicitObjectType ImplicitType0, const EImplicitObjectType ImplicitType1, const FReal Margin0, const FReal Margin1)
