Chaos - remove dead code

diff --git a/Engine/Source/Runtime/Experimental/Chaos/Private/Chaos/Evolution/PBDMinEvolution.cpp b/Engine/Source/Runtime/Experimental/Chaos/Private/Chaos/Evolution/PBDMinEvolution.cpp
--- a/Engine/Source/Runtime/Experimental/Chaos/Private/Chaos/Evolution/PBDMinEvolution.cpp
+++ b/Engine/Source/Runtime/Experimental/Chaos/Private/Chaos/Evolution/PBDMinEvolution.cpp
@@ -59,74 +59,6 @@ namespace Chaos
 	bool Chaos_MinEvolution_ForceMaxConstraintIterations = false;
 	FAutoConsoleVariableRef CVarChaosMinEvolutionForceMaxConstraintIterations(TEXT("p.Chaos.MinEvolution.ForceMaxConstraintIterations"), Chaos_MinEvolution_ForceMaxConstraintIterations, TEXT("Whether to force constraints to always use the worst-case maximum number of iterations"));
 
-	//
-	//
-	//
-
-	struct FPBDRigidArrays
-	{
-		FPBDRigidArrays()
-			: NumParticles(0)
-		{
-		}
-
-		FPBDRigidArrays(TPBDRigidParticles<FReal, 3>& Dynamics)
-		{
-			NumParticles = Dynamics.Size();
-			ObjectState = Dynamics.AllObjectState().GetData();
-			X = Dynamics.AllX().GetData();
-			P = Dynamics.AllP().GetData();
-			R = Dynamics.AllR().GetData();
-			Q = Dynamics.AllQ().GetData();
-			V = Dynamics.AllV().GetData();
-			PreV = Dynamics.AllPreV().GetData();
-			W = Dynamics.AllW().GetData();
-			PreW = Dynamics.AllPreW().GetData();
-			CenterOfMass = Dynamics.AllCenterOfMass().GetData();
-			RotationOfMass = Dynamics.AllRotationOfMass().GetData();
-			InvM = Dynamics.AllInvM().GetData();
-			InvI = Dynamics.AllInvI().GetData();
-			Acceleration = Dynamics.AllAcceleration().GetData();
-			AngularAcceleration = Dynamics.AllAngularAcceleration().GetData();
-			LinearImpulseVelocity = Dynamics.AllLinearImpulseVelocity().GetData();
-			AngularImpulseVelocity = Dynamics.AllAngularImpulseVelocity().GetData();
-			Disabled = Dynamics.AllDisabled().GetData();
-			GravityEnabled = Dynamics.AllGravityEnabled().GetData();
-			LinearEtherDrag = Dynamics.AllLinearEtherDrag().GetData();
-			AngularEtherDrag = Dynamics.AllAngularEtherDrag().GetData();
-			HasBounds = Dynamics.AllHasBounds().GetData();
-			LocalBounds = Dynamics.AllLocalBounds().GetData();
-			WorldBounds = Dynamics.AllWorldSpaceInflatedBounds().GetData();
-		}
-
-		int32 NumParticles;
-		EObjectStateType* ObjectState;
-		FVec3* X;
-		FVec3* P;
-		FRotation3* R;
-		FRotation3* Q;
-		FVec3* V;
-		FVec3* PreV;
-		FVec3* W;
-		FVec3* PreW;
-		FVec3* CenterOfMass;
-		FRotation3* RotationOfMass;
-		FReal* InvM;
-		TVec3<FRealSingle>* InvI;
-		FVec3* Acceleration;
-		FVec3* AngularAcceleration;
-		FVec3* LinearImpulseVelocity;
-		FVec3* AngularImpulseVelocity;
-		bool* Disabled;
-		bool* GravityEnabled;
-		FReal* LinearEtherDrag;
-		FReal* AngularEtherDrag;
-		bool* HasBounds;
-		FAABB3* LocalBounds;
-		FAABB3* WorldBounds;
-	};
-
-
 	//
 	//
 	//
@@ -303,7 +235,6 @@ namespace Chaos
 		}
 	}
 
-	// @todo(ccaulfield): dedupe (PBDRigidsEvolutionGBF)
 	void FPBDMinEvolution::Integrate(FReal Dt)
 	{
 		SCOPE_CYCLE_COUNTER(STAT_MinEvolution_Integrate);
diff --git a/Engine/Source/Runtime/Experimental/Chaos/Private/Chaos/Evolution/SolverBody.cpp b/Engine/Source/Runtime/Experimental/Chaos/Private/Chaos/Evolution/SolverBody.cpp
--- a/Engine/Source/Runtime/Experimental/Chaos/Private/Chaos/Evolution/SolverBody.cpp
+++ b/Engine/Source/Runtime/Experimental/Chaos/Private/Chaos/Evolution/SolverBody.cpp
@@ -8,73 +8,9 @@ namespace Chaos
 {
 
 	FSolverBody::FSolverBody()
-		//: Particle(nullptr)
 	{
 	}
 
-	//FSolverBody::FSolverBody(FGenericParticleHandle InParticle)
-	//	//: Particle(InParticle)
-	//{
-	//	//GatherInput();
-
-	//	// This will get set to something lower if there is a contact, based on the contact graph
-	//	State.Level = TNumericLimits<int32>::Max();
-	//}
-
-	//void FSolverBody::GatherInput()
-	//{
-	//	if (Particle.IsValid())
-	//	{
-	//		FRigidTransform3 CoMTransform = FParticleUtilitiesPQ::GetCoMWorldTransform(Particle);
-	//		State.P = CoMTransform.GetLocation();
-	//		State.Q = CoMTransform.GetRotation();
-	//		State.V = Particle->V();
-	//		State.W = Particle->W();
-
-	//		if (Particle->IsDynamic())
-	//		{
-	//			FRigidTransform3 PrevCoMTransform = FParticleUtilitiesXR::GetCoMWorldTransform(Particle);
-	//			State.X = PrevCoMTransform.GetLocation();
-	//			State.R = PrevCoMTransform.GetRotation();
-
-	//			State.InvM = Particle->InvM();
-	//			State.InvILocal = Particle->InvI().GetDiagonal();
-
-	//			// This will get set to something lower if there is a contact, based on the contact graph
-	//			State.Level = TNumericLimits<int32>::Max();
-	//		}
-	//		else
-	//		{
-	//			State.X = State.P;
-	//			State.R = State.Q;
-	//		}
-
-	//		UpdateRotationDependentState();
-	//	}
-	//}
-
-	//void FSolverBody::ScatterOutput()
-	//{
-	//	if (Particle.IsValid())
-	//	{
-	//		if (IsDynamic())
-	//		{
-	//			// Set the particle state
-	//			FParticleUtilities::SetCoMWorldTransform(Particle, State.P, State.Q);
-	//			Particle->SetV(State.V);
-	//			Particle->SetW(State.W);
-
-	//			if (State.bHasActiveCollision)
-	//			{
-	//				//Particle->AuxilaryValue(*ParticleParameters.Collided) = true;
-	//			}
-	//		}
-
-	//		// Reset SolverBodyIndex cookie every step - it will be reassigned next step
-	//		Particle->SetSolverBodyIndex(INDEX_NONE);
-	//	}
-	//}
-
 	void FSolverBody::UpdateRotationDependentState()
 	{
 		if (IsDynamic())
