Optimize Height Field Overlap
- Vectorization of GJK Intersection Same Space
- Remove dead code

diff --git a/Engine/Source/Programs/HeadlessChaos/Private/HeadlessChaosTestHeightfield.cpp b/Engine/Source/Programs/HeadlessChaos/Private/HeadlessChaosTestHeightfield.cpp
--- a/Engine/Source/Programs/HeadlessChaos/Private/HeadlessChaosTestHeightfield.cpp
+++ b/Engine/Source/Programs/HeadlessChaos/Private/HeadlessChaosTestHeightfield.cpp
@@ -455,7 +455,7 @@ namespace ChaosTest {
 	}
 
 
-	void SweepSphereTest()
+	void SweepSmallSphereTest()
 	{
 		const int32 Columns = 10;
 		const int32 Rows = 10;
@@ -495,16 +495,11 @@ namespace ChaosTest {
 					for (int32 Col = 0; Col < Columns-1; ++Col)
 					{
 						const FVec3 Start(Col * Scale[0] + 0.5, Row * Scale[1]+0.5, 1000 * Scale[2]);
-						FRigidTransform3 StartTM(Start, TRotation<FReal, 3>());
+						FRigidTransform3 StartTM(Start, TRotation<FReal, 3>::Identity);
 						FVec3 Dir(0, 0, -1);
 
 						bool Result = Heightfield.SweepGeom(Sphere, StartTM, Dir, 2000 * Scale[2], TOI, Position, Normal, FaceIdx, FaceNormal, 0.0, true);
 						EXPECT_TRUE(Result);
-						FVec3 AddGeomSize(1.0, 0.0, 1.0);
-						//FVec3 ExpectedPos(Col ,Row , Heights[Row * Columns + Col]);
-						//ExpectedPos += AddGeomSize;
-						//ExpectedPos *= Scale;
-						//EXPECT_VECTOR_NEAR(Position, ExpectedPos, 1e-2);
 					}
 				}
 			};
@@ -533,14 +528,11 @@ namespace ChaosTest {
 					{
 
 						const FVec3 Start(-Scale[0], Row * Scale[1], Heights[Row * Columns + Col] * Scale[2] + 0.01 * Scale[2]);
-						FRigidTransform3 StartTM(Start, TRotation<FReal, 3>());
+						FRigidTransform3 StartTM(Start, TRotation<FReal, 3>::Identity);
 						FVec3 Dir(1, 0, 0);
 
 						bool Result = Heightfield.SweepGeom(Sphere, StartTM, Dir, 2000 * Scale[0], TOI, Position, Normal, FaceIdx, FaceNormal, 0.0, true);
-
 						EXPECT_TRUE(Result);
-						//FVec3 ExpectedPos(1.0, Row, Scale[2] * (1.0 + Row * Columns) );
-						//EXPECT_VECTOR_NEAR(Position, ExpectedPos, 1e-2);
 					}
 				}
 			};
@@ -551,13 +543,12 @@ namespace ChaosTest {
 		}
 	}
 
-	void SweepBoxTest()
+	void SweepSmallBoxTest()
 	{
 		const int32 Columns = 10;
 		const int32 Rows = 10;
 		const FReal CountToWorldScale = 1;
 
-
 		TBox<FReal, 3> Box(FVec3(-0.05, -0.05, -0.05), FVec3(0.05, 0.05, 0.05));
 		{
 
@@ -591,16 +582,11 @@ namespace ChaosTest {
 					for (int32 Col = 0; Col < Columns - 1; ++Col)
 					{
 						const FVec3 Start(Col * Scale[0] + 0.5, Row * Scale[1] + 0.5, 1000 * Scale[2]);
-						FRigidTransform3 StartTM(Start, TRotation<FReal, 3>());
+						FRigidTransform3 StartTM(Start, TRotation<FReal, 3>::Identity);
 						FVec3 Dir(0, 0, -1);
 
 						bool Result = Heightfield.SweepGeom(Box, StartTM, Dir, 2000 * Scale[2], TOI, Position, Normal, FaceIdx, FaceNormal, 0.0, true);
 						EXPECT_TRUE(Result);
-						FVec3 AddGeomSize(1.0, 0.0, 1.0);
-						//FVec3 ExpectedPos(Col, Row, Heights[Row * Columns + Col]);
-						//ExpectedPos += AddGeomSize;
-						//ExpectedPos *= Scale;
-						//EXPECT_VECTOR_NEAR(Position, ExpectedPos, 1e-2);
 					}
 				}
 			};
@@ -629,15 +615,11 @@ namespace ChaosTest {
 					{
 
 						const FVec3 Start(-Scale[0], Row * Scale[1], Heights[Row * Columns + Col] * Scale[2] + 0.01 * Scale[2]);
-						FRigidTransform3 StartTM(Start, TRotation<FReal, 3>());
+						FRigidTransform3 StartTM(Start, TRotation<FReal, 3>::Identity);
 						FVec3 Dir(1, 0, 0);
 
 						bool Result = Heightfield.SweepGeom(Box, StartTM, Dir, 2000 * Scale[0], TOI, Position, Normal, FaceIdx, FaceNormal, 0.0, true);
-
 						EXPECT_TRUE(Result);
-
-						// FVec3 ExpectedPos(1.0, Row, Scale[2] * (1.0 + Row * Columns));
-						// EXPECT_VECTOR_NEAR(Position, ExpectedPos, 1e-2);
 					}
 				}
 			};
@@ -648,17 +630,262 @@ namespace ChaosTest {
 		}
 	}
 
+	TArray<FReal> CreateMountain(int32 Columns, int32 Rows)
+	{
+		TArray<FReal> Heights;
+		Heights.AddZeroed(Rows * Columns);
+		// Plane
+		for (int32 Row = 0; Row < Rows; ++Row)
+		{
+			for (int32 Col = 0; Col < Columns; ++Col)
+			{
+				Heights[Row * Columns + Col] = 0.0;
+			}
+		}
+
+		// Mountain
+		for (int32 Row = Rows / 2 - 1; Row <= Rows / 2 + 1; ++Row)
+		{
+			for (int32 Col = Columns / 2 - 1; Col <= Columns / 2 + 1; ++Col)
+			{
+				Heights[Row * Columns + Col] = 5.0;
+			}
+		}
+		// Submit
+		Heights[(Rows / 2) * Columns + Columns / 2] = 10.0;
+
+		return Heights;
+	}
+
+	void SweepBoxTest()
+	{
+		const FReal CountToWorldScale = 1;
+		const int32 Columns = 10;
+		const int32 Rows = 10;
+
+		TBox<FReal, 3> Box(FVec3(-1.0, -1.0, -1.0), FVec3(1.0, 1.0, 1.0));
+
+		TArray<FReal> Heights = CreateMountain(Columns, Rows);
+
+		TArray<FReal> HeightsCopy = Heights;
+		FVec3 Scale(1.0, 1.0, 1.0);
+		FHeightField Heightfield(MoveTemp(HeightsCopy), TArray<uint8>(), Rows, Columns, Scale);
+		const auto& Bounds = Heightfield.BoundingBox();	//Current API forces us to do this to cache the bounds
+
+		FReal TOI;
+		FVec3 Position, Normal, FaceNormal;
+		int32 FaceIdx;
+		int32 ExpectedFaceIdx = 0;
+		{
+			// Miss on one side of the mountain
+			const FVec3 Start(0.0, 3.5, 10.0);
+			FRigidTransform3 StartTM(Start, TRotation<FReal, 3>::Identity);
+			FVec3 Dir(1, 0, 0);
+			bool Result = Heightfield.SweepGeom(Box, StartTM, Dir, 100, TOI, Position, Normal, FaceIdx, FaceNormal, 0.0, true);
+			EXPECT_FALSE(Result);
+		}
+		{
+			// Miss on the other side of the mountain
+			const FVec3 Start(0.0, 7.0, 10.0);
+			FRigidTransform3 StartTM(Start, TRotation<FReal, 3>::Identity);
+			FVec3 Dir(1, 0, 0);
+			bool Result = Heightfield.SweepGeom(Box, StartTM, Dir, 100, TOI, Position, Normal, FaceIdx, FaceNormal, 0.0, true);
+			EXPECT_FALSE(Result);
+		}
+		// TODO This unit test is not working, to investigate
+		{
+			// Same missing in Y 
+			const FVec3 Start(3.0, 0.0, 10.0);
+			FRigidTransform3 StartTM(Start, TRotation<FReal, 3>::Identity);
+			FVec3 Dir(0, 1, 0);
+			bool Result = Heightfield.SweepGeom(Box, StartTM, Dir, 100, TOI, Position, Normal, FaceIdx, FaceNormal, 0.0, true);
+			EXPECT_FALSE(Result);
+		}
+		{
+			const FVec3 Start(7.0, 0.0, 10.0);
+			FRigidTransform3 StartTM(Start, TRotation<FReal, 3>::Identity);
+			FVec3 Dir(0, 1, 0);
+			bool Result = Heightfield.SweepGeom(Box, StartTM, Dir, 100, TOI, Position, Normal, FaceIdx, FaceNormal, 0.0, true);
+			EXPECT_FALSE(Result);
+		}
+		// Hit on the side of the box
+		{
+			const FVec3 Start(0.0, 4.0, 10.0);
+			FRigidTransform3 StartTM(Start, TRotation<FReal, 3>::Identity);
+			FVec3 Dir(1, 0, 0);
+			bool Result = Heightfield.SweepGeom(Box, StartTM, Dir, 100, TOI, Position, Normal, FaceIdx, FaceNormal, 0.0, true);
+			EXPECT_TRUE(Result);
+		}
+		{
+			const FVec3 Start(0.0, 6.0, 10.0);
+			FRigidTransform3 StartTM(Start, TRotation<FReal, 3>::Identity);
+			FVec3 Dir(1, 0, 0);
+			bool Result = Heightfield.SweepGeom(Box, StartTM, Dir, 100, TOI, Position, Normal, FaceIdx, FaceNormal, 0.0, true);
+			EXPECT_TRUE(Result);
+		}
+		{
+			const FVec3 Start(4.0, 0.0, 10.0);
+			FRigidTransform3 StartTM(Start, TRotation<FReal, 3>());
+			FVec3 Dir(0, 1, 0);
+			bool Result = Heightfield.SweepGeom(Box, StartTM, Dir, 100, TOI, Position, Normal, FaceIdx, FaceNormal, 0.0, true);
+			EXPECT_TRUE(Result);
+		}
+		{
+			const FVec3 Start(6.0, 0.0, 10.0);
+			FRigidTransform3 StartTM(Start, TRotation<FReal, 3>::Identity);
+			FVec3 Dir(0, 1, 0);
+			bool Result = Heightfield.SweepGeom(Box, StartTM, Dir, 100, TOI, Position, Normal, FaceIdx, FaceNormal, 0.0, true);
+			EXPECT_TRUE(Result);
+		}
+		// Hit in Diagonal
+		{
+
+			const FVec3 Start(0.5, 0.5, 10.0);
+			FRigidTransform3 StartTM(Start, TRotation<FReal, 3>::Identity);
+			FVec3 Dir(1, 1, 0);
+			Dir.SafeNormalize();
+			bool Result = Heightfield.SweepGeom(Box, StartTM, Dir, 100, TOI, Position, Normal, FaceIdx, FaceNormal, 0.0, true);
+			EXPECT_TRUE(Result);
+		}
+		{
+			const FVec3 Start(10.5, 10.5, 10.0);
+			FRigidTransform3 StartTM(Start, TRotation<FReal, 3>::Identity);
+			FVec3 Dir(-1, -1, 0);
+			Dir.SafeNormalize();
+			bool Result = Heightfield.SweepGeom(Box, StartTM, Dir, 100, TOI, Position, Normal, FaceIdx, FaceNormal, 0.0, true);
+			EXPECT_TRUE(Result);
+		}
+	}
+
 	void SweepTest()
 	{
-		SweepSphereTest();
+		SweepSmallSphereTest();
+		SweepSmallBoxTest();
 		SweepBoxTest();
 	}
+
+	void OverlapTest()
+	{
+		const FReal CountToWorldScale = 1;
+		const int32 Columns = 10;
+		const int32 Rows = 10;
+
+		
+
+		TArray<FReal> Heights = CreateMountain(Columns, Rows);
+
+		TArray<FReal> HeightsCopy = Heights;
+		FVec3 Scale(1.0, 1.0, 1.0);
+		FHeightField Heightfield(MoveTemp(HeightsCopy), TArray<uint8>(), Rows, Columns, Scale);
+		const auto& Bounds = Heightfield.BoundingBox();	//Current API forces us to do this to cache the bounds
+		// Long box
+		{
+			TBox<FReal, 3> Box(FVec3(-1.0, -1.0, -1.0), FVec3(1.0, 1.0, 10.0));
+			for (int32 Row = 0; Row < Rows; ++Row)
+			{
+				for (int32 Col = 0; Col < Columns; ++Col)
+				{
+					const FVec3 Translation(Col, Row, 1.5);
+					FRigidTransform3 QueryTM(Translation, TRotation<FReal, 3>::Identity);
+					FVec3 Dir(1, 0, 0);
+					bool Result = Heightfield.OverlapGeom(Box, QueryTM, 0.0, nullptr);
+					// No collision on the side of the mountain
+					if ((Col < 3 || Col > 7) && (Row < 3 || Row > 7))
+					{
+						EXPECT_FALSE(Result);
+					}
+					// Collision with the mountain
+					else if ((Col >= 3 && Col <= 7) && (Row >= 3 && Row <= 7))
+					{
+						EXPECT_TRUE(Result);
+					}
+					else
+					{
+						EXPECT_FALSE(Result);
+					}
+				}
+			}
+		}
+		// Box rotated in X
+		{
+			TBox<FReal, 3> Box(FVec3(-1.0, -1.0, -20.0), FVec3(1.0, 1.0, 20.0));
+			for (int32 Row = 0; Row < Rows; ++Row)
+			{
+				for (int32 Col = 0; Col < Columns; ++Col)
+				{
+					const FVec3 Translation(Col, Row, 2.0);
+					FRigidTransform3 QueryTM(Translation, TRotation<FReal, 3>(UE::Math::TQuat<FReal>(TVector<FReal, 3>(1.0, 0.0, 0.0), 3.14159/2.0)));
+					FVec3 Dir(1, 0, 0);
+					bool Result = Heightfield.OverlapGeom(Box, QueryTM, 0.0, nullptr);
+					// Collision with the mountain
+					if (Col >= 3 && Col <= 7)
+					{
+						EXPECT_TRUE(Result);
+					}
+					else
+					{
+						EXPECT_FALSE(Result);
+					}
+				}
+			}
+		}
+		// Box rotated in Y
+		{
+			TBox<FReal, 3> Box(FVec3(-1.0, -1.0, -20.0), FVec3(1.0, 1.0, 20.0));
+			for (int32 Row = 0; Row < Rows; ++Row)
+			{
+				for (int32 Col = 0; Col < Columns; ++Col)
+				{
+					const FVec3 Translation(Col, Row, 2.0);
+					FRigidTransform3 QueryTM(Translation, TRotation<FReal, 3>(UE::Math::TQuat<FReal>(TVector<FReal, 3>(0.0, 1.0, 0.0), 3.14159 / 2.0)));
+					FVec3 Dir(1, 0, 0);
+					bool Result = Heightfield.OverlapGeom(Box, QueryTM, 0.0, nullptr);
+					// Collision with the mountain
+					if (Row >= 3 && Row <= 7)
+					{
+						EXPECT_TRUE(Result);
+					}
+					else
+					{
+						EXPECT_FALSE(Result);
+					}
+				}
+			}
+		}
+		// Thin Box
+		{
+			TBox<FReal, 3> Box(FVec3(-10.0, -10.0, -0.001), FVec3(10.0, 10.0, 0.001));
+			for (int32 Row = 0; Row < Rows; ++Row)
+			{
+				for (int32 Col = 0; Col < Columns; ++Col)
+				{
+					const FVec3 Translation(Col, Row, 2.0);
+					FRigidTransform3 QueryTM(Translation, TRotation<FReal, 3>(UE::Math::TQuat<FReal>(TVector<FReal, 3>(1.0, 0.0, 0.0), 0.0)));
+					FVec3 Dir(1, 0, 0);
+					bool Result = Heightfield.OverlapGeom(Box, QueryTM, 0.0, nullptr);
+					// Collision with the mountain
+					EXPECT_TRUE(Result);
+				}
+			}
+		}
+		// Inclined Box
+		{
+			TBox<FReal, 3> Box(FVec3(-1.0, -1.0, -0.0001), FVec3(1.0, 1.0, 10.0));
+			const FVec3 Translation(5.0, 0.0, 15.0);
+			FRigidTransform3 QueryTM(Translation, TRotation<FReal, 3>(UE::Math::TQuat<FReal>(TVector<FReal, 3>(1.0, 0.0, 0.0), -3*3.1415926 / 4.0)));
+			FVec3 Dir(1, 0, 0);
+			bool Result = Heightfield.OverlapGeom(Box, QueryTM, 0.0, nullptr);
+			// Collision with the mountain
+			EXPECT_TRUE(Result);
+		}
+	}
 	
 
 	TEST(ChaosTests, Heightfield)
 	{
 		ChaosTest::Raycast();
 		ChaosTest::SweepTest();
+		ChaosTest::OverlapTest();
 		EditHeights();
 		SUCCEED();
 	}
diff --git a/Engine/Source/Runtime/Experimental/Chaos/Private/Chaos/HeightField.cpp b/Engine/Source/Runtime/Experimental/Chaos/Private/Chaos/HeightField.cpp
--- a/Engine/Source/Runtime/Experimental/Chaos/Private/Chaos/HeightField.cpp
+++ b/Engine/Source/Runtime/Experimental/Chaos/Private/Chaos/HeightField.cpp
@@ -1912,12 +1912,6 @@ namespace Chaos
 	template <typename QueryGeomType>
 	bool FHeightField::OverlapGeomImp(const QueryGeomType& QueryGeom, const FRigidTransform3& QueryTM, const FReal Thickness, FMTDInfo* OutMTD) const
 	{
-		if (OutMTD)
-		{
-			OutMTD->Normal = FVec3(0);
-			OutMTD->Penetration = TNumericLimits<FReal>::Lowest();
-		}
-
 		auto OverlapTriangleMTD = [&](const FVec3& A, const FVec3& B, const FVec3& C, FMTDInfo* InnerMTD) -> bool
 		{
 			const FVec3 AB = B - A;
@@ -1947,18 +1941,18 @@ namespace Chaos
 			return false;
 		};
 
-		auto OverlapTriangleNoMTD = [&](const FVec3& A, const FVec3& B, const FVec3& C) -> bool
+		auto OverlapTriangleNoMTD = [&](const VectorRegister4Float& A, const VectorRegister4Float& B, const VectorRegister4Float& C) -> bool
 		{
 			// points are assumed to be in the same space as the overlap geometry
-			const FVec3 AB = B - A;
-			const FVec3 AC = C - A;
+			const VectorRegister4Float AB = VectorSubtract(B, A);
+			const VectorRegister4Float AC = VectorSubtract(C, A);
 
 			//It's most likely that the query object is in front of the triangle since queries tend to be on the outside.
 			//However, maybe we should check if it's behind the triangle plane. Also, we should enforce this winding in some way
-			const FVec3 Offset = FVec3::CrossProduct(AB, AC);
+			const VectorRegister4Float Offset = VectorCross(AB, AC);
 
-			FTriangle TriangleConvex(A, B, C);
-			return GJKIntersectionSameSpace(TriangleConvex, QueryGeom, Thickness, Offset);
+			FTriangleRegister TriangleConvex(A, B, C);
+			return GJKIntersectionSameSpaceSimd(TriangleConvex, QueryGeom, Thickness, Offset);
 		};
 
 		bool bResult = false;
@@ -1971,13 +1965,16 @@ namespace Chaos
 		FlatQueryBounds.Max = FVec2(QueryBounds.Max()[0], QueryBounds.Max()[1]);
 
 		TArray<TVec2<int32>> Intersections;
-		FVec3 Points[4];
+		
 
 		GetGridIntersections(FlatQueryBounds, Intersections);
 
 		bool bOverlaps = false;
 		if (OutMTD)
 		{
+			OutMTD->Normal = FVec3(0);
+			OutMTD->Penetration = TNumericLimits<FReal>::Lowest();
+			FVec3 Points[4];
 			for (const TVec2<int32>& Cell : Intersections)
 			{
 				const int32 SingleIndex = Cell[1] * (GeomData.NumCols) + Cell[0];
@@ -1990,19 +1987,31 @@ namespace Chaos
 		}
 		else
 		{
-			FAABB3 CellBounds;
+			VectorRegister4Float Points[4];
+			const UE::Math::TQuat<FReal>& RotationDouble = QueryTM.GetRotation();
+			VectorRegister4Float Rotation = MakeVectorRegisterFloatFromDouble(MakeVectorRegister(RotationDouble.X, RotationDouble.Y, RotationDouble.Z, RotationDouble.W));
+			// Normalize rotation
+			Rotation = VectorNormalizeSafe(Rotation, GlobalVectorConstants::Float0001);
+			VectorRegister4Float InvRotation = VectorQuaternionInverse(Rotation);
+
+			const UE::Math::TVector<FReal>& TranslationDouble = QueryTM.GetTranslation();
+			const VectorRegister4Float Translation = MakeVectorRegisterFloatFromDouble(MakeVectorRegister(TranslationDouble.X, TranslationDouble.Y, TranslationDouble.Z, 0.0));
+			
+			FAABBVectorized CellBounds;
+			const FAABBVectorized QueryBoundsSimd(QueryBounds);
 			for (const TVec2<int32>& Cell : Intersections)
 			{
 				const int32 SingleIndex = Cell[1] * (GeomData.NumCols) + Cell[0];
-				GeomData.GetPointsAndBoundsScaled(SingleIndex, Points, CellBounds);
-
-				if (CellBounds.Intersects(QueryBounds))
+				GeomData.GetPointsAndBoundsScaledSimd(SingleIndex, Points, CellBounds);
+				if (CellBounds.Intersects(QueryBoundsSimd))
 				{
 					// pre-transform the triangle in overlap geometry space
-					Points[0] = QueryTM.InverseTransformPositionNoScale(Points[0]);
-					Points[1] = QueryTM.InverseTransformPositionNoScale(Points[1]);
-					Points[2] = QueryTM.InverseTransformPositionNoScale(Points[2]);
-					Points[3] = QueryTM.InverseTransformPositionNoScale(Points[3]);
+					for (int32 Index = 0; Index < 4; ++Index)
+					{
+						VectorRegister4Float& P = Points[Index];
+						P = VectorSubtract(P, Translation);
+						P = VectorQuaternionRotateVector(InvRotation, P);
+					}
 
 					if (OverlapTriangleNoMTD(Points[0], Points[1], Points[3]))
 					{
diff --git a/Engine/Source/Runtime/Experimental/Chaos/Public/Chaos/GJK.h b/Engine/Source/Runtime/Experimental/Chaos/Public/Chaos/GJK.h
--- a/Engine/Source/Runtime/Experimental/Chaos/Public/Chaos/GJK.h
+++ b/Engine/Source/Runtime/Experimental/Chaos/Public/Chaos/GJK.h
@@ -8,21 +8,15 @@
 #include "Chaos/ImplicitObjectScaled.h"
 #include "Chaos/Simplex.h"
 #include "Chaos/Sphere.h"
-
-#include "ChaosCheck.h"
-#include "ChaosLog.h"
-
-
-// Platform specific vector intrinsics include.
-#if PLATFORM_ENABLE_VECTORINTRINSICS_NEON || PLATFORM_ENABLE_VECTORINTRINSICS
 #include "Math/VectorRegister.h"
 #include "Chaos/VectorUtility.h"
 #include "Chaos/SimplexVectorized.h"
 #include "Chaos/EPAVectorized.h"
-#define GJK_VECTORIZED 1
-#else
-#define GJK_VECTORIZED 0
-#endif
+
+
+#include "ChaosCheck.h"
+#include "ChaosLog.h"
+
 
 namespace Chaos
 {
@@ -63,7 +57,6 @@ namespace Chaos
 		@param ThicknessB The amount of geometry inflation for Geometry B(for example if the surface distance of two geometries with thickness 0 would be 2, a thickness of 0.5 would give a distance of 1.5)
 		@return True if the geometries overlap, False otherwise 
 	 */
-#if GJK_VECTORIZED
 	template <typename T, typename TGeometryA, typename TGeometryB>
 	bool GJKIntersection(const TGeometryA& RESTRICT A, const TGeometryB& RESTRICT B, const TRigidTransform<T, 3>& BToATM, const T InThicknessA = 0, const TVector<T, 3>& InitialDir = TVector<T, 3>(1, 0, 0), const T InThicknessB = 0)
 	{
@@ -88,8 +81,6 @@ namespace Chaos
 		VectorRegister4Float PrevDist2Simd = MakeVectorRegisterFloatConstant(FLT_MAX, FLT_MAX, FLT_MAX, FLT_MAX);
 
 		VectorRegister4Float SimplexSimd[4] = { VectorZeroFloat(), VectorZeroFloat(), VectorZeroFloat(), VectorZeroFloat() };
-		//VectorRegister4Float As[4] = { VectorZeroFloat(), VectorZeroFloat(), VectorZeroFloat(), VectorZeroFloat() };
-		//VectorRegister4Float Bs[4] = { VectorZeroFloat(), VectorZeroFloat(), VectorZeroFloat(), VectorZeroFloat() };
 		VectorRegister4Float BarycentricSimd;
 		VectorRegister4Int NumVerts = GlobalVectorConstants::IntZero;
 
@@ -114,7 +105,6 @@ namespace Chaos
 			const VectorRegister4Float SupportASimd = A.SupportCoreSimd(NegVSimd, ThicknessA);
 			const VectorRegister4Float VInBSimd = VectorQuaternionRotateVector(AToBRotationSimd, VSimd);
 			const VectorRegister4Float SupportBLocalSimd = B.SupportCoreSimd(VInBSimd, ThicknessB);
-			//const TVector<T, 3> SupportB = BToATM.TransformPositionNoScale(SupportBLocal);  // Original code
 			const VectorRegister4Float SupportBSimd = VectorAdd(VectorQuaternionRotateVector(RotationSimd, SupportBLocalSimd), TranslationSimd);
 			const VectorRegister4Float WSimd = VectorSubtract(SupportASimd, SupportBSimd);
 
@@ -153,76 +143,6 @@ namespace Chaos
 
 		return bNearZero;
 	}
-#else
-	template <typename T, typename TGeometryA, typename TGeometryB>
-	bool GJKIntersection(const TGeometryA& A, const TGeometryB& B, const TRigidTransform<T, 3>& BToATM, const T InThicknessA = 0, const TVector<T, 3>& InitialDir = TVector<T, 3>(1, 0, 0), const T InThicknessB = 0)
-	{
-		TVector<T, 3> V = -InitialDir;
-		if (V.SafeNormalize() == 0)
-		{
-			V = TVec3<T>(-1, 0, 0);
-		}
-
-		FSimplex SimplexIDs;
-		TVector<T, 3> Simplex[4];
-		T Barycentric[4] = { -1,-1,-1,-1 };	//not needed, but compiler warns
-		const TRotation<T, 3> AToBRotation = BToATM.GetRotation().Inverse();
-		bool bTerminate;
-		bool bNearZero = false;
-		int NumIterations = 0;
-		T PrevDist2 = FLT_MAX;
-		T ThicknessA;
-		T ThicknessB;
-		CalculateQueryMargins(A, B, ThicknessA, ThicknessB);
-		ThicknessA += InThicknessA;
-		ThicknessB += InThicknessB;
-		const T Inflation = ThicknessA + ThicknessB + static_cast<T>(1e-3);
-		const T Inflation2 = Inflation * Inflation;
-		int32 VertexIndexA = INDEX_NONE, VertexIndexB = INDEX_NONE;
-		do
-		{
-			if (!ensure(NumIterations++ < 32))	//todo: take this out
-			{
-				break;	//if taking too long just stop. This should never happen
-			}
-			const TVector<T, 3> NegV = -V;
-			const TVector<T, 3> SupportA = A.SupportCore(NegV, ThicknessA, nullptr, VertexIndexA);
-			const TVector<T, 3> VInB = AToBRotation * V;
-			const TVector<T, 3> SupportBLocal = B.SupportCore(VInB, ThicknessB, nullptr, VertexIndexB);
-			const TVector<T, 3> SupportB = BToATM.TransformPositionNoScale(SupportBLocal);
-			const TVector<T, 3> W = SupportA - SupportB;
-
-			if (TVector<T, 3>::DotProduct(V, W) > Inflation)
-			{
-				return false;
-			}
-
-			SimplexIDs[SimplexIDs.NumVerts] = SimplexIDs.NumVerts;
-			Simplex[SimplexIDs.NumVerts++] = W;
-
-			V = SimplexFindClosestToOrigin(Simplex, SimplexIDs, Barycentric);
-
-			T NewDist2 = V.SizeSquared();
-			bNearZero = NewDist2 < Inflation2;
-
-			//as simplices become degenerate we will stop making progress. This is a side-effect of precision, in that case take V as the current best approximation
-			//question: should we take previous v in case it's better?
-			const bool bMadeProgress = NewDist2 < PrevDist2;
-			bTerminate = bNearZero || !bMadeProgress;
-
-			PrevDist2 = NewDist2;
-
-			if (!bTerminate)
-			{
-				V /= FMath::Sqrt(NewDist2);
-			}
-
-		} while (!bTerminate);
-
-		return bNearZero;
-	}
-
-#endif
 
 	/** 
 		Determines if two convex geometries in the same space overlap
@@ -235,65 +155,71 @@ namespace Chaos
 		@param ThicknessB The amount of geometry inflation for Geometry B(for example if the surface distance of two geometries with thickness 0 would be 2, a thickness of 0.5 would give a distance of 1.5)
 		@return True if the geometries overlap, False otherwise 
 	*/
-	template <typename T, typename TGeometryA, typename TGeometryB>
-	bool GJKIntersectionSameSpace(const TGeometryA& A, const TGeometryB& B, const T InThicknessA = 0, const TVector<T, 3>& InitialDir = TVector<T, 3>(1, 0, 0), const T InThicknessB = 0)
+
+	template <typename TGeometryA, typename TGeometryB>
+	bool GJKIntersectionSameSpaceSimd(const TGeometryA& A, const TGeometryB& B, FReal InThicknessA, const VectorRegister4Float& InitialDir)
 	{
-		TVector<T, 3> V = -InitialDir;
-		if (V.SafeNormalize() == 0)
-		{
-			V = TVec3<T>(-1, 0, 0);
-		}
+		VectorRegister4Float V = VectorNegate(InitialDir);
 
 		FSimplex SimplexIDs;
-		TVector<T, 3> Simplex[4];
-		T Barycentric[4] = { -1,-1,-1,-1 };	//not needed, but compiler warns
+		VectorRegister4Float Simplex[4] = { VectorZero(), VectorZero(), VectorZero(), VectorZero() };
+		VectorRegister4Int NumVerts = GlobalVectorConstants::IntZero;
+		alignas(16) int32 NumVertsInts[4];
+		VectorRegister4Float Barycentric = VectorZero();
+
 		bool bTerminate;
 		bool bNearZero = false;
 		int NumIterations = 0;
-		T PrevDist2 = FLT_MAX;
-		T ThicknessA;
-		T ThicknessB;
+		VectorRegister4Float PrevDist2 = GlobalVectorConstants::BigNumber;
+		FReal ThicknessA;
+		FReal ThicknessB;
 		CalculateQueryMargins(A, B, ThicknessA, ThicknessB);
 		ThicknessA += InThicknessA;
-		ThicknessB += InThicknessB;
-		const T Inflation = ThicknessA + ThicknessB + static_cast<T>(1e-3);
-		const T Inflation2 = Inflation * Inflation;
-		int32 VertexIndexA = INDEX_NONE, VertexIndexB = INDEX_NONE;
+
+		const FReal Inflation = ThicknessA + ThicknessB + 1e-3;
+		const FReal Inflation2 = Inflation * Inflation;
+
+		VectorRegister4Float InflationSimd = VectorSet1(static_cast<FRealSingle>(Inflation));
+		VectorRegister4Float Inflation2Simd = VectorMultiply(InflationSimd, InflationSimd);
 		do
 		{
 			if (!ensure(NumIterations++ < 32))	//todo: take this out
 			{
 				break;	//if taking too long just stop. This should never happen
 			}
-			const TVector<T, 3> NegV = -V;
-			const TVector<T, 3> SupportA = A.SupportCore(NegV, ThicknessA, nullptr, VertexIndexA);
-			const TVector<T, 3> VInB = V; // same space
-			const TVector<T, 3> SupportB = B.SupportCore(VInB, ThicknessB, nullptr, VertexIndexB);
-			const TVector<T, 3> W = SupportA - SupportB;
+			const VectorRegister4Float NegV = VectorNegate(V);
+			const VectorRegister4Float SupportA = A.SupportCoreSimd(NegV, ThicknessA);
+			const VectorRegister4Float VInB = V; // same space
+			const VectorRegister4Float SupportB = B.SupportCoreSimd(VInB, ThicknessB);
+			const VectorRegister4Float W = VectorSubtract(SupportA, SupportB);
+
+			VectorRegister4Float VDotW = VectorDot3(V, W);
 
-			if (TVector<T, 3>::DotProduct(V, W) > Inflation)
+			if (VectorMaskBits(VectorCompareGT(VDotW, InflationSimd)))
 			{
 				return false;
 			}
 
-			SimplexIDs[SimplexIDs.NumVerts] = SimplexIDs.NumVerts;
-			Simplex[SimplexIDs.NumVerts++] = W;
-
-			V = SimplexFindClosestToOrigin(Simplex, SimplexIDs, Barycentric);
+			VectorIntStoreAligned(NumVerts, NumVertsInts);
+			const int32 NumVertsInt = NumVertsInts[0];
+			Simplex[NumVertsInt] = W;
+			NumVerts = VectorIntAdd(NumVerts, GlobalVectorConstants::IntOne);
 
-			T NewDist2 = V.SizeSquared();
-			bNearZero = NewDist2 < Inflation2;
+			V = VectorSimplexFindClosestToOrigin<false>(Simplex, NumVerts, Barycentric, nullptr, nullptr);
+			const VectorRegister4Float NewDist2 = VectorDot3(V, V);
+			
+			bNearZero = static_cast<bool>(VectorMaskBits(VectorCompareGT(Inflation2Simd, NewDist2)));
 
 			//as simplices become degenerate we will stop making progress. This is a side-effect of precision, in that case take V as the current best approximation
 			//question: should we take previous v in case it's better?
-			const bool bMadeProgress = NewDist2 < PrevDist2;
+			const bool bMadeProgress = static_cast<bool>(VectorMaskBits(VectorCompareGT(PrevDist2, NewDist2)));
 			bTerminate = bNearZero || !bMadeProgress;
 
 			PrevDist2 = NewDist2;
 
 			if (!bTerminate)
 			{
-				V /= FMath::Sqrt(NewDist2);
+				V = VectorDivide(V, VectorSqrt(NewDist2));
 			}
 
 		} while (!bTerminate);
@@ -301,6 +227,7 @@ namespace Chaos
 		return bNearZero;
 	}
 
+
 	/**
 	 * @brief Internal simplex data for GJK that can also be stored for warm-starting subsequent calls
 	 * @tparam T the number type (float or double)
@@ -1432,271 +1359,21 @@ namespace Chaos
 		return true;
 	}
 	
-	template <typename TGeometryA, typename TGeometryB, typename T>
-	bool GJKRaycast2Impl(const TGeometryA& A, const TGeometryB& B, const TRigidTransform<T, 3>& StartTM, const TVector<T, 3>& RayDir, const T RayLength,
-		T& OutTime, TVector<T, 3>& OutPosition, TVector<T, 3>& OutNormal, const T GivenThicknessA = 0, bool bComputeMTD = false, const TVector<T, 3>& InitialDir = TVector<T, 3>(1, 0, 0), const T GivenThicknessB = 0)
-	{
-		ensure(FMath::IsNearlyEqual(RayDir.SizeSquared(), (T)1, (T)UE_KINDA_SMALL_NUMBER));
-		ensure(RayLength > 0);
-
-		T MarginA;
-		T MarginB;
-		CalculateQueryMargins(A,B, MarginA, MarginB);
-
-		const TVector<T, 3> StartPoint = StartTM.GetLocation();
-
-		TVector<T, 3> Simplex[4] = { TVector<T,3>(0), TVector<T,3>(0), TVector<T,3>(0), TVector<T,3>(0) };
-		TVector<T, 3> As[4] = { TVector<T,3>(0), TVector<T,3>(0), TVector<T,3>(0), TVector<T,3>(0) };
-		TVector<T, 3> Bs[4] = { TVector<T,3>(0), TVector<T,3>(0), TVector<T,3>(0), TVector<T,3>(0) };
-
-		T Barycentric[4] = { -1,-1,-1,-1 };	//not needed, but compiler warns
-		const T Inflation = MarginA + MarginB;
-		const T Inflation2 = Inflation * Inflation + static_cast<T>(1e-6);
-
-		FSimplex SimplexIDs;
-		const TRotation<T, 3> BToARotation = StartTM.GetRotation();
-		const TRotation<T, 3> AToBRotation = BToARotation.Inverse();
-
-		auto SupportAFunc = [&A, MarginA](const TVec3<T>& V)
-		{
-			int32 VertexIndex = INDEX_NONE;
-			return A.SupportCore(V, MarginA, nullptr, VertexIndex);
-		};
-
-		auto SupportBFunc = [&B, MarginB, &AToBRotation, &BToARotation](const TVec3<T>& V)
-		{
-			int32 VertexIndex = INDEX_NONE;
-			const TVector<T, 3> VInB = AToBRotation * V;
-			const TVector<T, 3> SupportBLocal = B.SupportCore(VInB, MarginB, nullptr, VertexIndex);
-			return BToARotation * SupportBLocal;
-		};
-		
-		auto SupportBAtOriginFunc = [&B, MarginB, &StartTM, &AToBRotation](const TVec3<T>& Dir)
-		{
-			int32 VertexIndex = INDEX_NONE;
-			const TVector<T, 3> DirInB = AToBRotation * Dir;
-			const TVector<T, 3> SupportBLocal = B.SupportCore(DirInB, MarginB, nullptr, VertexIndex);
-			return StartTM.TransformPositionNoScale(SupportBLocal);
-		};
-
-		TVector<T, 3> SupportA = SupportAFunc(InitialDir);
-		As[0] = SupportA;
-
-		TVector<T, 3> SupportB = SupportBFunc(-InitialDir);
-		Bs[0] = SupportB;
-
-		T Lambda = 0;
-		TVector<T, 3> X = StartPoint;
-		TVector<T, 3> V = X - (SupportA - SupportB);
-		TVector<T, 3> Normal(0,0,1);
-
-		const T InitialPreDist2 = V.SizeSquared();
-		constexpr T Eps2 = 1e-6f;
-		//mtd needs to find closest point even in inflation region, so can only skip if we found the closest points
-		bool bCloseEnough = InitialPreDist2 < Inflation2 && (!bComputeMTD || InitialPreDist2 < Eps2);
-		bool bDegenerate = false;
-		bool bTerminate = bCloseEnough;
-		bool bInflatedCloseEnough = bCloseEnough;
-		int NumIterations = 0;
-		T GJKPreDist2 = TNumericLimits<T>::Max();
-		while (!bTerminate)
-		{
-			//if (!ensure(NumIterations++ < 32))	//todo: take this out
-			if (!(NumIterations++ < 32))	//todo: take this out
-			{
-				break;	//if taking too long just stop. This should never happen
-			}
-
-			V = V.GetUnsafeNormal();
-
-			SupportA = SupportAFunc(V);
-			SupportB = SupportBFunc(-V);
-			const TVector<T, 3> P = SupportA - SupportB;
-			const TVector<T, 3> W = X - P;
-			SimplexIDs[SimplexIDs.NumVerts] = SimplexIDs.NumVerts;	//is this needed?
-			As[SimplexIDs.NumVerts] = SupportA;
-			Bs[SimplexIDs.NumVerts] = SupportB;
-
-			const T VDotW = TVector<T, 3>::DotProduct(V, W);
-
-			if (VDotW > Inflation)
-			{
-				const T VDotRayDir = TVector<T, 3>::DotProduct(V, RayDir);
-				if (VDotRayDir >= 0)
-				{
-					return false;
-				}
-
-				const T PreLambda = Lambda;	//use to check for no progress
-				// @todo(ccaulfield): this can still overflow - the comparisons against zero above should be changed (though not sure to what yet)
-				Lambda = Lambda - (VDotW - Inflation) / VDotRayDir;
-				if (Lambda > PreLambda)
-				{
-					if (Lambda > RayLength)
-					{
-						return false;
-					}
-
-					const TVector<T, 3> OldX = X;
-					X = StartPoint + Lambda * RayDir;
-					Normal = V;
-
-					//Update simplex from (OldX - P) to (X - P)
-					const TVector<T, 3> XMinusOldX = X - OldX;
-					Simplex[0] += XMinusOldX;
-					Simplex[1] += XMinusOldX;
-					Simplex[2] += XMinusOldX;
-					Simplex[SimplexIDs.NumVerts++] = X - P;
-
-					GJKPreDist2 = TNumericLimits<T>::Max();	//translated origin so restart gjk search
-					bInflatedCloseEnough = false;
-				}
-			}
-			else
-			{
-				Simplex[SimplexIDs.NumVerts++] = W;	//this is really X - P which is what we need for simplex computation
-			}
-
-			if (bInflatedCloseEnough && VDotW >= 0)
-			{
-				//Inflated shapes are close enough, but we want MTD so we need to find closest point on core shape
-				const T VDotW2 = VDotW * VDotW;
-				bCloseEnough = GJKPreDist2 <= Eps2 + VDotW2;	//todo: relative error
-			}
-
-			if (!bCloseEnough)
-			{
-				V = SimplexFindClosestToOrigin(Simplex, SimplexIDs, Barycentric, As, Bs);
-				T NewDist2 = V.SizeSquared();	//todo: relative error
-				bCloseEnough = NewDist2 < Inflation2;
-				bDegenerate = NewDist2 >= GJKPreDist2;
-				GJKPreDist2 = NewDist2;
-
-
-				if (bComputeMTD && bCloseEnough && Lambda == 0 && GJKPreDist2 > 1e-6 && Inflation2 > 1e-6 && SimplexIDs.NumVerts < 4)
-				{
-					//For mtd of inflated shapes we have to find the closest point, so we have to keep going
-					bCloseEnough = false;
-					bInflatedCloseEnough = true;
-				}
-			}
-			else
-			{
-				//It must be that we want MTD and we can terminate. However, we must make one final call to fixup the simplex
-				V = SimplexFindClosestToOrigin(Simplex, SimplexIDs, Barycentric, As, Bs);
-			}
-			bTerminate = bCloseEnough || bDegenerate;
-		}
-
-		OutTime = Lambda;
-
-		if (Lambda > 0)
-		{
-			OutNormal = Normal;
-			TVector<T, 3> ClosestB(0);
-
-			for (int i = 0; i < SimplexIDs.NumVerts; ++i)
-			{
-				ClosestB += Bs[i] * Barycentric[i];
-			}
-			const TVector<T, 3> ClosestLocal = ClosestB - OutNormal * MarginB;
-
-			OutPosition = StartPoint + RayDir * Lambda + ClosestLocal;
-		}
-		else if (bComputeMTD)
-		{
-			// If Inflation == 0 we would expect GJKPreDist2 to be 0
-			// However, due to precision we can still end up with GJK failing.
-			// When that happens fall back on EPA
-			if (Inflation > 0 && GJKPreDist2 > 1e-6 && GJKPreDist2 < TNumericLimits<T>::Max())
-			{
-				OutNormal = Normal;
-				TVector<T, 3> ClosestA(0);
-				TVector<T, 3> ClosestB(0);
-
-				if (NumIterations)
-				{
-					for (int i = 0; i < SimplexIDs.NumVerts; ++i)
-					{
-						ClosestA += As[i] * Barycentric[i];
-						ClosestB += Bs[i] * Barycentric[i];
-					}
-				}
-				else
-				{
-					//didn't even go into gjk loop
-					ClosestA = As[0];
-					ClosestB = Bs[0];
-				}
-				
-				const TVec3<T> ClosestBInA = StartPoint + ClosestB;
-				const T InGJKPreDist = FMath::Sqrt(GJKPreDist2);
-				OutNormal = V.GetUnsafeNormal();
-
-				const T Penetration = FMath::Clamp<T>(MarginA + MarginB - InGJKPreDist, 0, TNumericLimits<T>::Max());
-				const TVector<T, 3> ClosestLocal = ClosestB - OutNormal * MarginB;
-
-				OutPosition = StartPoint + ClosestLocal + OutNormal * Penetration;
-				OutTime = -Penetration;
-			}
-			else
-			{
-				//use EPA
-				TArray<TVec3<T>> VertsA;
-				TArray<TVec3<T>> VertsB;
-
-				VertsA.Reserve(8);
-				VertsB.Reserve(8);
-
-				if (NumIterations)
-				{
-					for (int i = 0; i < SimplexIDs.NumVerts; ++i)
-					{
-						VertsA.Add(As[i]);
-						const TVec3<T> BAtOrigin = Bs[i] + X;
-						VertsB.Add(BAtOrigin);
-					}
-
-					T Penetration;
-					TVec3<T> MTD, ClosestA, ClosestBInA;
-					const EEPAResult EPAResult = EPA(VertsA, VertsB, SupportAFunc, SupportBAtOriginFunc, Penetration, MTD, ClosestA, ClosestBInA);
-					if (IsEPASuccess(EPAResult))
-					{
-						OutNormal = MTD;
-						OutTime = -Penetration - Inflation;
-						OutPosition = ClosestA;
-					}
-					//else if (IsEPADegenerate(EPAResult))
-					//{
-					//	// @todo(chaos): handle degenerate EPA condition
-					//}
-					else
-					{
-						//assume touching hit
-						OutTime = -Inflation;
-						OutNormal = MTD;
-						OutPosition = As[0] + OutNormal * MarginA;
-					}
-				}
-				else
-				{
-					//didn't even go into gjk loop, touching hit
-					OutTime = -Inflation;
-					OutNormal = { 0,0,1 };
-					OutPosition = As[0] + OutNormal * MarginA;
-				}
-			}
-		}
-		else
-		{
-			// Initial overlap without MTD. These properties are not valid, but assigning them anyway so they don't contain NaNs and cause issues in invoking code.
-			OutNormal = { 0,0,1 };
-			OutPosition = { 0,0,0 };
-		}
-
-		return true;
-	}
 
+	/** Sweeps one geometry against the other
+	 @A The first geometry
+	 @B The second geometry
+	 @BToARotation B's starting rotation in A's local space
+         @StartPoint B's starting position in A's local space
+	 @RayDir The ray's direction (normalized)
+	 @RayLength The ray's length
+	 @OutTime The time along the ray when the objects first overlap
+	 @OutPosition The first point of impact (in A's local space) when the objects first overlap. Invalid if time of impact is 0
+	 @OutNormal The impact normal (in A's local space) when the objects first overlap. Invalid if time of impact is 0
+	 @InitialDir The first direction we use to search the CSO
+	 @return True if the geometries overlap during the sweep, False otherwise
+	 @note If A overlaps B at the start of the ray ("initial overlap" condition) then this function returns true, and sets OutTime = 0, but does not set any other output variables.
+	 */
 	template <typename TGeometryA, typename TGeometryB, typename T>
 	bool GJKRaycast2ImplSimd(const TGeometryA& A, const TGeometryB& B, const VectorRegister4Float& BToARotation, const VectorRegister4Float& StartPoint, const VectorRegister4Float& RayDir, const T RayLength,
 		T& OutTime, VectorRegister4Float& OutPosition, VectorRegister4Float& OutNormal, bool bComputeMTD, const VectorRegister4Float& InitialDir)
@@ -2061,7 +1738,6 @@ namespace Chaos
 		return true;
 	}
 
-#if GJK_VECTORIZED
 	template <typename T, typename TGeometryA, typename TGeometryB>
 	bool GJKRaycast2(const TGeometryA& A, const TGeometryB& B, const TRigidTransform<T, 3>& StartTM, const TVector<T, 3>& RayDir, const T RayLength,
 		T& OutTime, TVector<T, 3>& OutPosition, TVector<T, 3>& OutNormal, const T GivenThicknessA = 0, bool bComputeMTD = false, const TVector<T, 3>& InitialDir = TVector<T, 3>(1, 0, 0), const T GivenThicknessB = 0)
@@ -2099,31 +1775,6 @@ namespace Chaos
 		return result;
 	}
 
-#else // GJK_VECTORIZED
-
-	/** Sweeps one geometry against the other
-	 @A The first geometry
-	 @B The second geometry
-	 @StartTM B's starting configuration in A's local space
-	 @RayDir The ray's direction (normalized)
-	 @RayLength The ray's length
-	 @OutTime The time along the ray when the objects first overlap
-	 @OutPosition The first point of impact (in A's local space) when the objects first overlap. Invalid if time of impact is 0
-	 @OutNormal The impact normal (in A's local space) when the objects first overlap. Invalid if time of impact is 0
-	 @ThicknessA The amount of geometry inflation for Geometry A (for example a capsule with radius 5 could pass in its core segnment and a thickness of 5)
-	 @InitialDir The first direction we use to search the CSO
-	 @ThicknessB The amount of geometry inflation for Geometry B (for example a sphere with radius 5 could pass in its center point and a thickness of 5)
-	 @return True if the geometries overlap during the sweep, False otherwise
-	 @note If A overlaps B at the start of the ray ("initial overlap" condition) then this function returns true, and sets OutTime = 0, but does not set any other output variables.
-	 */
-
-	template <typename T, typename TGeometryA, typename TGeometryB>
-	bool GJKRaycast2(const TGeometryA& A, const TGeometryB& B, const TRigidTransform<T, 3>& StartTM, const TVector<T, 3>& RayDir, const T RayLength,
-		T& OutTime, TVector<T, 3>& OutPosition, TVector<T, 3>& OutNormal, const T GivenThicknessA = 0, bool bComputeMTD = false, const TVector<T, 3>& InitialDir = TVector<T, 3>(1, 0, 0), const T GivenThicknessB = 0)
-	{
-		return GJKRaycast2Impl(A, B, StartTM, RayDir, RayLength, OutTime, OutPosition, OutNormal, GivenThicknessA, bComputeMTD, InitialDir, GivenThicknessB);
-	}
-#endif // GJK_VECTORIZED
 
 	/**
 	 * Used by GJKDistance. It must return a vector in the Minkowski sum A - B. In principle this can be the vector of any point
@@ -2270,9 +1921,6 @@ namespace Chaos
 	}
 
 
-
-
-
 	// Assumes objects are already intersecting, computes a minimum translation
 	// distance, deepest penetration positions on each body, and approximates
 	// a penetration normal and minimum translation distance.
diff --git a/Engine/Source/Runtime/Experimental/Chaos/Public/Chaos/SimplexVectorized.h b/Engine/Source/Runtime/Experimental/Chaos/Public/Chaos/SimplexVectorized.h
--- a/Engine/Source/Runtime/Experimental/Chaos/Public/Chaos/SimplexVectorized.h
+++ b/Engine/Source/Runtime/Experimental/Chaos/Public/Chaos/SimplexVectorized.h
@@ -1,7 +1,6 @@
 // Copyright Epic Games, Inc. All Rights Reserved.
 #pragma once
 
-#if PLATFORM_ENABLE_VECTORINTRINSICS_NEON || PLATFORM_ENABLE_VECTORINTRINSICS
 #include "Math/VectorRegister.h"
 #include "Chaos/VectorUtility.h"
 
@@ -453,5 +452,3 @@ namespace Chaos
 		return ClosestPoint;
 	}
 }
-
-#endif
\ No newline at end of file
