Chaos: clustering code cleanup
- consolidate many duplicated code sections in methods

diff --git a/Engine/Source/Runtime/Experimental/Chaos/Private/Chaos/PBDRigidClustering.cpp b/Engine/Source/Runtime/Experimental/Chaos/Private/Chaos/PBDRigidClustering.cpp
--- a/Engine/Source/Runtime/Experimental/Chaos/Private/Chaos/PBDRigidClustering.cpp
+++ b/Engine/Source/Runtime/Experimental/Chaos/Private/Chaos/PBDRigidClustering.cpp
@@ -362,273 +362,188 @@ namespace Chaos
 		return ActivatedChildren;
 	}
 
-	DECLARE_CYCLE_STAT(TEXT("TPBDRigidClustering<>::ReleaseClusterParticles(STRAIN)"), STAT_ReleaseClusterParticles_STRAIN, STATGROUP_Chaos);
-	TSet<FPBDRigidParticleHandle*> 
-	FRigidClustering::ReleaseClusterParticles(
-		FPBDRigidClusteredParticleHandle* ClusteredParticle,
-		const TMap<FGeometryParticleHandle*, Chaos::FReal>* ExternalStrainMap,
-		bool bForceRelease)
+	void FRigidClustering::SendBreakingEvent(FPBDRigidClusteredParticleHandle* ClusteredParticle)
 	{
-		SCOPE_CYCLE_COUNTER(STAT_ReleaseClusterParticles_STRAIN);
-
-		TSet<FPBDRigidParticleHandle*> ActivatedChildren;
-		if (!ensureMsgf(MChildren.Contains(ClusteredParticle), TEXT("Removing Cluster that does not exist!")))
-		{
-			return ActivatedChildren;
-		}
-		TArray<FPBDRigidParticleHandle*>& Children = MChildren[ClusteredParticle];
-
-		bool bChildrenChanged = false;
-		const FRigidTransform3 PreSolveTM = FRigidTransform3(ClusteredParticle->P(), ClusteredParticle->Q());
-
-		//@todo(ocohen): iterate with all the potential parents at once?
-		//find all children within some distance of contact point
-
-		auto RemoveChildLambda = [&](FPBDRigidParticleHandle* Child/*, const int32 Idx*/)
-		{
-			FPBDRigidClusteredParticleHandle* ClusteredChild = Child->CastToClustered();
-
-			MEvolution.EnableParticle(Child, ClusteredParticle);
-			TopLevelClusterParents.Add(ClusteredChild);
-
-			ClusteredChild->SetClusterId(ClusterId(nullptr, ClusteredChild->ClusterIds().NumChildren)); // clear Id but retain number of children
-
-			const FRigidTransform3 ChildFrame = ClusteredChild->ChildToParent() * PreSolveTM;
-			Child->SetX(ChildFrame.GetTranslation());
-			Child->SetR(ChildFrame.GetRotation());
-
-			Child->SetP(Child->X());
-			Child->SetQ(Child->R());
-
-			//todo(ocohen): for now just inherit velocity at new COM. This isn't quite right for rotation
-			//todo(ocohen): in the presence of collisions, this will leave all children with the post-collision
-			// velocity. This should be controlled by material properties so we can allow the broken pieces to
-			// maintain the clusters pre-collision velocity.
-			Child->SetV(ClusteredParticle->V());
-			Child->SetW(ClusteredParticle->W());
-			Child->SetPreV(ClusteredParticle->PreV());
-			Child->SetPreW(ClusteredParticle->PreW());
-
-			ActivatedChildren.Add(Child);
-			//if (ChildIdx != INDEX_NONE)
-			//{
-			//	Children.RemoveAtSwap(ChildIdx, 1, /*bAllowShrinking=*/false); //@todo(ocohen): maybe avoid this until we know all children are not going away?
-			//}
-
-			bChildrenChanged = true;
-		};
-
-		for (int32 ChildIdx = Children.Num() - 1; ChildIdx >= 0; --ChildIdx)
+		check(ClusteredParticle);
+		if (DoGenerateBreakingData)
 		{
-			FPBDRigidClusteredParticleHandle* Child = Children[ChildIdx]->CastToClustered();
-			
-			if (!Child)
+			FBreakingData& ClusterBreak = MAllClusterBreakings.AddDefaulted_GetRef();
+			ClusterBreak.Proxy = ClusteredParticle->PhysicsProxy();
+			ClusterBreak.Location = ClusteredParticle->X();
+			ClusterBreak.Velocity = ClusteredParticle->V();
+			ClusterBreak.AngularVelocity = ClusteredParticle->W();
+			ClusterBreak.Mass = ClusteredParticle->M();
+			if (ClusteredParticle->Geometry() && ClusteredParticle->Geometry()->HasBoundingBox())
 			{
-				continue;
+				ClusterBreak.BoundingBox = ClusteredParticle->Geometry()->BoundingBox();
 			}
-
-			Chaos::FReal ChildStrain = 0.0;
-
-			if(ExternalStrainMap)
+			if (ClusterBreak.Proxy->GetType() == EPhysicsProxyType::GeometryCollectionType)
 			{
-				const Chaos::FReal* MapStrain = ExternalStrainMap->Find(Child);
-				ChildStrain = MapStrain ? *MapStrain : Child->CollisionImpulses();
+				const FGeometryCollectionPhysicsProxy* ConcreteProxy = static_cast<FGeometryCollectionPhysicsProxy*>(ClusterBreak.Proxy);
+				ClusterBreak.TransformGroupIndex = ConcreteProxy->GetTransformGroupIndexFromHandle(ClusteredParticle);
 			}
 			else
 			{
-				ChildStrain = Child->CollisionImpulses();
+				ClusterBreak.TransformGroupIndex = INDEX_NONE;
 			}
+		}
+	}
 
-			if (ChildStrain >= Child->Strain() || bForceRelease)
-			{
-				//UE_LOG(LogTemp, Warning, TEXT("Releasing child %d from parent %p due to strain %.5f Exceeding internal strain %.5f (Source: %s)"), ChildIdx, ClusteredParticle, ChildStrain, Child->Strain(), bForceRelease ? TEXT("Forced by caller") : ExternalStrainMap ? TEXT("External") : TEXT("Collision"));
+	TArray<FRigidClustering::FParticleIsland> FRigidClustering::FindIslandsInChildren(const FPBDRigidClusteredParticleHandle* ClusteredParticle)
+	{
+		const TArray<FPBDRigidParticleHandle*>& Children = MChildren[ClusteredParticle];
+		
+		TArray<FParticleIsland> Islands;
 
-				// The piece that hits just breaks off - we may want more control 
-				// by looking at the edges of this piece which would give us cleaner 
-				// breaks (this approach produces more rubble)
-				RemoveChildLambda(Child);
+		// traverse connectivity and see how many connected pieces we have
+		TSet<FPBDRigidParticleHandle*> ProcessedChildren;
+		ProcessedChildren.Reserve(Children.Num());
 
-				// Remove from the children array without freeing memory yet. 
-				// We're looping over Children and it'd be silly to free the array
-				// 1 entry at a time.
-				Children.RemoveAtSwap(ChildIdx, 1, false);
+		for (FPBDRigidParticleHandle* Child : Children)
+		{
+			if (ProcessedChildren.Contains(Child))
+			{
+				continue;
+			}
+			TArray<FPBDRigidParticleHandle*>& Island = Islands.AddDefaulted_GetRef();
 
-				if (Child->ToBeRemovedOnFracture())
-				{
-					MActiveRemovalIndices.Add(Child);
-				}
-				else
+			TArray<FPBDRigidParticleHandle*> ProcessingQueue;
+			ProcessingQueue.Add(Child);
+			while (ProcessingQueue.Num())
+			{
+				FPBDRigidParticleHandle* ChildToProcess = ProcessingQueue.Pop();
+				if (!ProcessedChildren.Contains(ChildToProcess))
 				{
-					if (DoGenerateBreakingData)
+					ProcessedChildren.Add(ChildToProcess);
+					Island.Add(ChildToProcess);
+					for (const TConnectivityEdge<FReal>& Edge : ChildToProcess->CastToClustered()->ConnectivityEdges())
 					{
-						const int32 NewIdx = MAllClusterBreakings.Add(FBreakingData());
-						FBreakingData& ClusterBreak = MAllClusterBreakings[NewIdx];
-						ClusterBreak.Proxy = Child->PhysicsProxy();
-						ClusterBreak.Location = Child->X();
-						ClusterBreak.Velocity = Child->V();
-						ClusterBreak.AngularVelocity = Child->W();
-						ClusterBreak.Mass = Child->M();
-						if (Child->Geometry() && Child->Geometry()->HasBoundingBox())
-						{
-							ClusterBreak.BoundingBox = Child->Geometry()->BoundingBox();
-						}
-						if (ClusterBreak.Proxy->GetType() == EPhysicsProxyType::GeometryCollectionType)
-						{
-							FGeometryCollectionPhysicsProxy* ConcreteProxy = static_cast<FGeometryCollectionPhysicsProxy*>(ClusterBreak.Proxy);
-							ClusterBreak.TransformGroupIndex = ConcreteProxy->GetTransformGroupIndexFromHandle(Child);
-						}
-						else
+						if (!ProcessedChildren.Contains(Edge.Sibling))
 						{
-							ClusterBreak.TransformGroupIndex = INDEX_NONE;
+							ProcessingQueue.Add(Edge.Sibling);
 						}
 					}
 				}
 			}
 		}
 
-		if (bChildrenChanged)
+		return Islands;
+	}
+
+	void FRigidClustering::RemoveChildFromParent(FPBDRigidParticleHandle* Child, const FPBDRigidClusteredParticleHandle* ClusteredParent)
+	{
+		ensure(Child);
+		ensure(ClusteredParent);
+		FPBDRigidClusteredParticleHandle* ClusteredChild = Child->CastToClustered();
+
+		MEvolution.EnableParticle(Child, ClusteredParent);
+		TopLevelClusterParents.Add(ClusteredChild);
+
+		ClusteredChild->SetClusterId(ClusterId(nullptr, ClusteredChild->ClusterIds().NumChildren)); // clear Id but retain number of children
+
+		const FRigidTransform3 PreSolveTM(ClusteredParent->P(), ClusteredParent->Q());
+		const FRigidTransform3 ChildFrame = ClusteredChild->ChildToParent() * PreSolveTM;
+		Child->SetX(ChildFrame.GetTranslation());
+		Child->SetR(ChildFrame.GetRotation());
+
+		Child->SetP(Child->X());
+		Child->SetQ(Child->R());
+
+		//todo(ocohen): for now just inherit velocity at new COM. This isn't quite right for rotation
+		//todo(ocohen): in the presence of collisions, this will leave all children with the post-collision
+		// velocity. This should be controlled by material properties so we can allow the broken pieces to
+		// maintain the clusters pre-collision velocity.
+		Child->SetV(ClusteredParent->V());
+		Child->SetW(ClusteredParent->W());
+		Child->SetPreV(ClusteredParent->PreV());
+		Child->SetPreW(ClusteredParent->PreW());
+	};
+
+	TArray<FPBDRigidParticleHandle*> FRigidClustering::CreateClustersFromNewIslands(
+		TArray<FParticleIsland>& Islands,
+		FPBDRigidClusteredParticleHandle* ClusteredParent
+		)
+	{
+		TArray<FPBDRigidParticleHandle*> NewClusters;
+		
+		// only for island with more than one particle
+		int32 NumNewClusters = 0;
+		for (const TArray<FPBDRigidParticleHandle*>& Island : Islands)
 		{
-			if (Children.Num() == 0)
+			if (Island.Num() > 1)
 			{
-				// Free the memory if we can do so cheaply (no data copies).
-				Children.Empty(); 
+				NumNewClusters++;
 			}
+		}
+		NewClusters.Reserve(NumNewClusters);
 
-			if (UseConnectivity)
+		const FRigidTransform3 PreSolveTM = FRigidTransform3(ClusteredParent->P(), ClusteredParent->Q());
+		
+		TArray<Chaos::FPBDRigidClusteredParticleHandle*> NewClusterHandles = MEvolution.CreateClusteredParticles(NumNewClusters);
+		int32 ClusterHandlesIdx = 0;
+		for (TArray<FPBDRigidParticleHandle*>& Island : Islands)
+		{
+			if (Island.Num() > 1) //now build the remaining pieces
 			{
-				// The cluster may have contained forests, so find the connected pieces and cluster them together.
-
-				//first update the connected graph of the children we already removed
-				for (FPBDRigidParticleHandle* Child : ActivatedChildren)
-				{
-					RemoveNodeConnections(Child);
-				}
-
-				if (Children.Num())
+				FClusterCreationParameters CreationParameters;
+				CreationParameters.ClusterParticleHandle = NewClusterHandles[ClusterHandlesIdx++];
+				Chaos::FPBDRigidClusteredParticleHandle* NewCluster = 
+					CreateClusterParticleFromClusterChildren(
+						MoveTemp(Island), 
+						ClusteredParent, 
+						PreSolveTM, 
+						CreationParameters);
+
+				MEvolution.SetPhysicsMaterial(NewCluster, MEvolution.GetPhysicsMaterial(ClusteredParent));
+
+				NewCluster->SetStrain(ClusteredParent->Strain());
+				NewCluster->SetV(ClusteredParent->V());
+				NewCluster->SetW(ClusteredParent->W());
+				NewCluster->SetPreV(ClusteredParent->PreV());
+				NewCluster->SetPreW(ClusteredParent->PreW());
+				NewCluster->SetP(NewCluster->X());
+				NewCluster->SetQ(NewCluster->R());
+
+				// Need to get the material from the previous particle and apply it to the new one
+				const FShapesArray& ChildShapes = ClusteredParent->ShapesArray();
+				const FShapesArray& NewShapes = NewCluster->ShapesArray();
+				const int32 NumChildShapes = ClusteredParent->ShapesArray().Num();
+
+				if(NumChildShapes > 0)
 				{
-					TArray<TArray<FPBDRigidParticleHandle*>> ConnectedPiecesArray;
-
-					{ // tmp scope
-
-						//traverse connectivity and see how many connected pieces we have
-						TSet<FPBDRigidParticleHandle*> ProcessedChildren;
-						ProcessedChildren.Reserve(Children.Num());
-
-						for (FPBDRigidParticleHandle* PotentialActivatedChild : Children)
-						{
-							if (ProcessedChildren.Contains(PotentialActivatedChild))
-							{
-								continue;
-							}
-							ConnectedPiecesArray.AddDefaulted();
-							TArray<FPBDRigidParticleHandle*>& ConnectedPieces = ConnectedPiecesArray.Last();
-
-							TArray<FPBDRigidParticleHandle*> ProcessingQueue;
-							ProcessingQueue.Add(PotentialActivatedChild);
-							while (ProcessingQueue.Num())
-							{
-								FPBDRigidParticleHandle* Child = ProcessingQueue.Pop();
-								if (!ProcessedChildren.Contains(Child))
-								{
-									ProcessedChildren.Add(Child);
-									ConnectedPieces.Add(Child);
-									for (const TConnectivityEdge<FReal>& Edge : Child->CastToClustered()->ConnectivityEdges())
-									{
-										if (!ProcessedChildren.Contains(Edge.Sibling))
-										{
-											ProcessingQueue.Add(Edge.Sibling);
-										}
-									}
-								}
-							}
-						}
-					} // tmp scope
-
-					int32 NumNewClusters = 0;
-					for (TArray<FPBDRigidParticleHandle*>& ConnectedPieces : ConnectedPiecesArray)
-					{
-						if (ConnectedPieces.Num() == 1) //need to break single pieces first
-						{
-							FPBDRigidParticleHandle* Child = ConnectedPieces[0];
-							RemoveChildLambda(Child);
-						}
-						else if (ConnectedPieces.Num() > 1)
-						{
-							NumNewClusters++;
-						}
-					}
-					TArray<Chaos::FPBDRigidClusteredParticleHandle*> NewClusterHandles = 
-						MEvolution.CreateClusteredParticles(NumNewClusters);
-					int32 ClusterHandlesIdx = 0;
-					for (TArray<FPBDRigidParticleHandle*>& ConnectedPieces : ConnectedPiecesArray)
+					// Can only take materials if the child has any - otherwise we fall back on defaults.
+					// Due to GC initialisation however, we should always have a valid material as even
+					// when one cannot be found we fall back on the default on GEngine
+					const int32 NumChildMaterials = ChildShapes[0]->GetMaterials().Num();
+					if(NumChildMaterials > 0)
 					{
-						if (ConnectedPieces.Num() > 1) //now build the remaining pieces
-						{
-							Chaos::FClusterCreationParameters CreationParameters;
-							CreationParameters.ClusterParticleHandle = NewClusterHandles[ClusterHandlesIdx++];
-							Chaos::FPBDRigidClusteredParticleHandle* NewCluster = 
-								CreateClusterParticleFromClusterChildren(
-									MoveTemp(ConnectedPieces), 
-									ClusteredParticle, 
-									PreSolveTM, 
-									CreationParameters);
-
-							MEvolution.SetPhysicsMaterial(
-								NewCluster, MEvolution.GetPhysicsMaterial(ClusteredParticle));
-
-							NewCluster->SetStrain(ClusteredParticle->Strain());
-							NewCluster->SetV(ClusteredParticle->V());
-							NewCluster->SetW(ClusteredParticle->W());
-							NewCluster->SetPreV(ClusteredParticle->PreV());
-							NewCluster->SetPreW(ClusteredParticle->PreW());
-							NewCluster->SetP(NewCluster->X());
-							NewCluster->SetQ(NewCluster->R());
-
-							// Need to get the material from the previous particle and apply it to the new one
-							const FShapesArray& ChildShapes = ClusteredParticle->ShapesArray();
-							const FShapesArray& NewShapes = NewCluster->ShapesArray();
-							const int32 NumChildShapes = ClusteredParticle->ShapesArray().Num();
-
-							if(NumChildShapes > 0)
-							{
-								// Can only take materials if the child has any - otherwise we fall back on defaults.
-								// Due to GC initialisation however, we should always have a valid material as even
-								// when one cannot be found we fall back on the default on GEngine
-								const int32 NumChildMaterials = ChildShapes[0]->GetMaterials().Num();
-								if(NumChildMaterials > 0)
-								{
-									Chaos::FMaterialHandle ChildMat = ChildShapes[0]->GetMaterials()[0];
-
-									for(const TUniquePtr<FPerShapeData>& PerShape : NewShapes)
-									{
-										PerShape->SetMaterial(ChildMat);
-									}
-								}
-							}
+						Chaos::FMaterialHandle ChildMat = ChildShapes[0]->GetMaterials()[0];
 
-							ActivatedChildren.Add(NewCluster);
+						for(const TUniquePtr<FPerShapeData>& PerShape : NewShapes)
+						{
+							PerShape->SetMaterial(ChildMat);
 						}
 					}
 				}
+				NewClusters.Add(NewCluster);
 			}
-
-			for (FPBDRigidParticleHandle* Child : ActivatedChildren)
-			{
-				UpdateKinematicProperties(Child, MChildren, MEvolution);
-			}
-
-			//disable cluster
-			DisableCluster(ClusteredParticle);
-		} // bChildrenChanged
-
-		return ActivatedChildren;
+		}
+		return NewClusters;
 	}
+	
+	DECLARE_CYCLE_STAT(TEXT("TPBDRigidClustering<>::ReleaseClusterParticles(STRAIN)"), STAT_ReleaseClusterParticles_STRAIN, STATGROUP_Chaos);
+	TSet<FPBDRigidParticleHandle*> FRigidClustering::ReleaseClusterParticles(
+		FPBDRigidClusteredParticleHandle* ClusteredParticle,
+		const TMap<FGeometryParticleHandle*, Chaos::FReal>* ExternalStrainMap,
+		bool bForceRelease)
+	{
+		SCOPE_CYCLE_COUNTER(STAT_ReleaseClusterParticles_STRAIN);
 
-
-
+		return ReleaseClusterParticlesImpl(ClusteredParticle, ExternalStrainMap, bForceRelease, true /*bCreateNewClusters*/);
+	}
+	
 	DECLARE_CYCLE_STAT(TEXT("TPBDRigidClustering<>::ReleaseClusterParticlesNoInternalCluster"), STAT_ReleaseClusterParticlesNoInternalCluster, STATGROUP_Chaos);
-	TSet<FPBDRigidParticleHandle*>
-	FRigidClustering::ReleaseClusterParticlesNoInternalCluster(
+	TSet<FPBDRigidParticleHandle*> FRigidClustering::ReleaseClusterParticlesNoInternalCluster(
 		FPBDRigidClusteredParticleHandle* ClusteredParticle,
 		const TMap<FGeometryParticleHandle*, Chaos::FReal>* ExternalStrainMap,
 		bool bForceRelease)
@@ -640,67 +555,36 @@ namespace Chaos
 		
 		SCOPE_CYCLE_COUNTER(STAT_ReleaseClusterParticlesNoInternalCluster);
 
+		return ReleaseClusterParticlesImpl(ClusteredParticle, ExternalStrainMap, bForceRelease, false /*bCreateNewClusters*/);
+	}
+	
+	TSet<FPBDRigidParticleHandle*> FRigidClustering::ReleaseClusterParticlesImpl(
+		FPBDRigidClusteredParticleHandle* ClusteredParticle,
+		const TMap<FGeometryParticleHandle*, Chaos::FReal>* ExternalStrainMap,
+		bool bForceRelease,
+		bool bCreateNewClusters)
+	{	
 		TSet<FPBDRigidParticleHandle*> ActivatedChildren;
 		if (!ensureMsgf(MChildren.Contains(ClusteredParticle), TEXT("Removing Cluster that does not exist!")))
 		{
 			return ActivatedChildren;
 		}
+		
 		TArray<FPBDRigidParticleHandle*>& Children = MChildren[ClusteredParticle];
-
-		bool bChildrenChanged = false;
-		const FRigidTransform3 PreSolveTM = FRigidTransform3(ClusteredParticle->P(), ClusteredParticle->Q());
-
-		//@todo(ocohen): iterate with all the potential parents at once?
-		//find all children within some distance of contact point
-
-		auto RemoveChildLambda = [&](FPBDRigidParticleHandle* Child/*, const int32 Idx*/)
-		{
-			FPBDRigidClusteredParticleHandle* ClusteredChild = Child->CastToClustered();
-
-			MEvolution.EnableParticle(Child, ClusteredParticle);
-			TopLevelClusterParents.Add(ClusteredChild);
-
-			ClusteredChild->SetClusterId(ClusterId(nullptr, ClusteredChild->ClusterIds().NumChildren)); // clear Id but retain number of children
-
-			const FRigidTransform3 ChildFrame = ClusteredChild->ChildToParent() * PreSolveTM;
-			Child->SetX(ChildFrame.GetTranslation());
-			Child->SetR(ChildFrame.GetRotation());
-
-			Child->SetP(Child->X());
-			Child->SetQ(Child->R());
-
-			//todo(ocohen): for now just inherit velocity at new COM. This isn't quite right for rotation
-			//todo(ocohen): in the presence of collisions, this will leave all children with the post-collision
-			// velocity. This should be controlled by material properties so we can allow the broken pieces to
-			// maintain the clusters pre-collision velocity.
-			Child->SetV(ClusteredParticle->V());
-			Child->SetW(ClusteredParticle->W());
-			Child->SetPreV(ClusteredParticle->PreV());
-			Child->SetPreW(ClusteredParticle->PreW());
-
-			ActivatedChildren.Add(Child);
-			//if (ChildIdx != INDEX_NONE)
-			//{
-			//	Children.RemoveAtSwap(ChildIdx, 1, /*bAllowShrinking=*/false); //@todo(ocohen): maybe avoid this until we know all children are not going away?
-			//}
-
-			bChildrenChanged = true;
-		};
-
 		for (int32 ChildIdx = Children.Num() - 1; ChildIdx >= 0; --ChildIdx)
 		{
 			FPBDRigidClusteredParticleHandle* Child = Children[ChildIdx]->CastToClustered();
-
+			
 			if (!Child)
 			{
 				continue;
 			}
 
-			Chaos::FReal ChildStrain = 0.0;
+			FReal ChildStrain = 0.0;
 
-			if (ExternalStrainMap)
+			if(ExternalStrainMap)
 			{
-				const Chaos::FReal* MapStrain = ExternalStrainMap->Find(Child);
+				const FReal* MapStrain = ExternalStrainMap->Find(Child);
 				ChildStrain = MapStrain ? *MapStrain : Child->CollisionImpulses();
 			}
 			else
@@ -708,7 +592,6 @@ namespace Chaos
 				ChildStrain = Child->CollisionImpulses();
 			}
 
-
 			if (ChildStrain >= Child->Strain() || bForceRelease)
 			{
 				//UE_LOG(LogTemp, Warning, TEXT("Releasing child %d from parent %p due to strain %.5f Exceeding internal strain %.5f (Source: %s)"), ChildIdx, ClusteredParticle, ChildStrain, Child->Strain(), bForceRelease ? TEXT("Forced by caller") : ExternalStrainMap ? TEXT("External") : TEXT("Collision"));
@@ -716,7 +599,8 @@ namespace Chaos
 				// The piece that hits just breaks off - we may want more control 
 				// by looking at the edges of this piece which would give us cleaner 
 				// breaks (this approach produces more rubble)
-				RemoveChildLambda(Child);
+				RemoveChildFromParent(Child, ClusteredParticle);
+				ActivatedChildren.Add(Child);
 
 				// Remove from the children array without freeing memory yet. 
 				// We're looping over Children and it'd be silly to free the array
@@ -729,39 +613,17 @@ namespace Chaos
 				}
 				else
 				{
-					if (DoGenerateBreakingData)
-					{
-						const int32 NewIdx = MAllClusterBreakings.Add(FBreakingData());
-						FBreakingData& ClusterBreak = MAllClusterBreakings[NewIdx];
-						ClusterBreak.Proxy = Child->PhysicsProxy();
-						ClusterBreak.Location = Child->X();
-						ClusterBreak.Velocity = Child->V();
-						ClusterBreak.AngularVelocity = Child->W();
-						ClusterBreak.Mass = Child->M();
-						if (Child->Geometry() && Child->Geometry()->HasBoundingBox())
-						{
-							ClusterBreak.BoundingBox = Child->Geometry()->BoundingBox();
-						}
-						if (ClusterBreak.Proxy && ClusterBreak.Proxy->GetType() == EPhysicsProxyType::GeometryCollectionType)
-						{
-							FGeometryCollectionPhysicsProxy* ConcreteProxy = static_cast<FGeometryCollectionPhysicsProxy*>(ClusterBreak.Proxy);
-							ClusterBreak.TransformGroupIndex = ConcreteProxy->GetTransformGroupIndexFromHandle(Child);
-						}
-						else
-						{
-							ClusterBreak.TransformGroupIndex = INDEX_NONE;
-						}
-					}
+					SendBreakingEvent(Child);
 				}
 			}
 		}
 
-		if (bChildrenChanged)
+		if (ActivatedChildren.Num() > 0)
 		{
 			if (Children.Num() == 0)
 			{
 				// Free the memory if we can do so cheaply (no data copies).
-				Children.Empty();
+				Children.Empty(); 
 			}
 
 			if (UseConnectivity)
@@ -776,57 +638,22 @@ namespace Chaos
 
 				if (Children.Num())
 				{
-					TArray<TArray<FPBDRigidParticleHandle*>> ConnectedPiecesArray;
-
-					{ // tmp scope
-
-						//traverse connectivity and see how many connected pieces we have
-						TSet<FPBDRigidParticleHandle*> ProcessedChildren;
-						ProcessedChildren.Reserve(Children.Num());
-
-						for (FPBDRigidParticleHandle* PotentialActivatedChild : Children)
+					TArray<FParticleIsland> Islands = FindIslandsInChildren(ClusteredParticle);
+					for (const FParticleIsland& Island: Islands)
+					{
+						if (Island.Num() == 1) //need to break single pieces first
 						{
-							if (ProcessedChildren.Contains(PotentialActivatedChild))
-							{
-								continue;
-							}
-							ConnectedPiecesArray.AddDefaulted();
-							TArray<FPBDRigidParticleHandle*>& ConnectedPieces = ConnectedPiecesArray.Last();
-
-							TArray<FPBDRigidParticleHandle*> ProcessingQueue;
-							ProcessingQueue.Add(PotentialActivatedChild);
-							while (ProcessingQueue.Num())
-							{
-								FPBDRigidParticleHandle* Child = ProcessingQueue.Pop();
-								if (!ProcessedChildren.Contains(Child))
-								{
-									ProcessedChildren.Add(Child);
-									ConnectedPieces.Add(Child);
-									for (const TConnectivityEdge<FReal>& Edge : Child->CastToClustered()->ConnectivityEdges())
-									{
-										if (!ProcessedChildren.Contains(Edge.Sibling))
-										{
-											ProcessingQueue.Add(Edge.Sibling);
-										}
-									}
-								}
-							}
+							FPBDRigidParticleHandle* Child = Island[0];
+							RemoveChildFromParent(Child, ClusteredParticle);
+							ActivatedChildren.Add(Child);
 						}
-					} // tmp scope
+					}
 
-					int32 NumNewClusters = 0;
-					for (TArray<FPBDRigidParticleHandle*>& ConnectedPieces : ConnectedPiecesArray)
+					if (bCreateNewClusters)
 					{
-						if (ConnectedPieces.Num() == 1) //need to break single pieces first
-						{
-							FPBDRigidParticleHandle* Child = ConnectedPieces[0];
-							RemoveChildLambda(Child);
-						}
-						else if (ConnectedPieces.Num() > 1)
-						{
-							NumNewClusters++;
-						}
-					}			
+						TArray<FPBDRigidParticleHandle*> NewClusters = CreateClustersFromNewIslands(Islands, ClusteredParticle); 
+						ActivatedChildren.Append(MoveTemp(NewClusters));
+					}
 				}
 			}
 
@@ -837,14 +664,11 @@ namespace Chaos
 
 			//disable cluster
 			DisableCluster(ClusteredParticle);
-		} // bChildrenChanged
+		}
 
 		return ActivatedChildren;
 	}
 
-
-
-
 	DECLARE_CYCLE_STAT(TEXT("TPBDRigidClustering<>::ReleaseClusterParticles(LIST)"), STAT_ReleaseClusterParticles_LIST, STATGROUP_Chaos);
 	TSet<FPBDRigidParticleHandle*> 
 	FRigidClustering::ReleaseClusterParticles(
@@ -915,7 +739,7 @@ namespace Chaos
 					Particle->V().SizeSquared() > Threshold && 
 					Particle->PreV().SizeSquared() > Threshold)
 				{
-					DisableParticleWithBreakEvent(Particle);
+					DisableParticleWithBreakEvent(Particle->CastToClustered());
 				}
 			}
 		}
@@ -1006,30 +830,7 @@ namespace Chaos
 				{
 					if (ClusteredParticle->CollisionImpulses() >= ClusteredParticle->Strains())
 					{
-						DisableCluster(ClusteredParticle);
-						if (DoGenerateBreakingData)
-						{
-							int32 NewIdx = MAllClusterBreakings.Add(FBreakingData());
-							FBreakingData& ClusterBreak = MAllClusterBreakings[NewIdx];
-							ClusterBreak.Proxy = ClusteredParticle->PhysicsProxy();
-							ClusterBreak.Location = ClusteredParticle->X();
-							ClusterBreak.Velocity = ClusteredParticle->V();
-							ClusterBreak.AngularVelocity = ClusteredParticle->W();
-							ClusterBreak.Mass = ClusteredParticle->M();
-							if (ClusteredParticle->Geometry()->HasBoundingBox())
-							{
-								ClusterBreak.BoundingBox = ClusteredParticle->Geometry()->BoundingBox();
-							}
-							if (ClusterBreak.Proxy->GetType() == EPhysicsProxyType::GeometryCollectionType)
-							{
-								FGeometryCollectionPhysicsProxy* ConcreteProxy = static_cast<FGeometryCollectionPhysicsProxy*>(ClusterBreak.Proxy);
-								ClusterBreak.TransformGroupIndex = ConcreteProxy->GetTransformGroupIndexFromHandle(ClusteredParticle);
-							}
-							else
-							{
-								ClusterBreak.TransformGroupIndex = INDEX_NONE;
-							}
-						}
+						DisableParticleWithBreakEvent(ClusteredParticle);
 					}
 				}
 			}
@@ -1242,9 +1043,7 @@ namespace Chaos
 		}
 	}
 
-	void 
-	FRigidClustering::DisableCluster(
-		FPBDRigidClusteredParticleHandle* ClusteredParticle)
+	void FRigidClustering::DisableCluster(FPBDRigidClusteredParticleHandle* ClusteredParticle)
 	{
 		// #note: we don't recursively descend to the children
 		MEvolution.DisableParticle(ClusteredParticle);
@@ -1255,35 +1054,12 @@ namespace Chaos
 		MActiveRemovalIndices.Remove(ClusteredParticle);
 	}
 
-	void 
-	FRigidClustering::DisableParticleWithBreakEvent(
-		Chaos::FPBDRigidParticleHandle* Particle)
+	void FRigidClustering::DisableParticleWithBreakEvent(FPBDRigidClusteredParticleHandle* ClusteredParticle)
 	{
-		DisableCluster(Particle->CastToClustered());
-
-		if (DoGenerateBreakingData)
+		if (ClusteredParticle)
 		{
-			const int32 NewIdx = MAllClusterBreakings.Add(FBreakingData());
-			FBreakingData& ClusterBreak = MAllClusterBreakings[NewIdx];
-			ClusterBreak.Proxy = Particle->PhysicsProxy();
-			ClusterBreak.Location = Particle->X();
-			ClusterBreak.Velocity = Particle->V();
-			ClusterBreak.AngularVelocity = Particle->W();
-			ClusterBreak.Mass = Particle->M();
-			if (Particle->Geometry()->HasBoundingBox())
-			{
-				ClusterBreak.BoundingBox = Particle->Geometry()->BoundingBox();
-			}
-			if (ClusterBreak.Proxy->GetType() == EPhysicsProxyType::GeometryCollectionType)
-			{
-				FGeometryCollectionPhysicsProxy* ConcreteProxy = static_cast<FGeometryCollectionPhysicsProxy*>(ClusterBreak.Proxy);
-				ClusterBreak.TransformGroupIndex = ConcreteProxy->GetTransformGroupIndexFromHandle(Particle);
-			}
-			else
-			{
-				ClusterBreak.TransformGroupIndex = INDEX_NONE;
-			}
-			
+			DisableCluster(ClusteredParticle);
+			SendBreakingEvent(ClusteredParticle);
 		}
 	}
 
@@ -1316,7 +1092,7 @@ namespace Chaos
 			ensure(ClusteredParticle->ClusterIds().Id == nullptr);
 		}
 
-		// reset the structures
+	// reset the structures
 		TopLevelClusterParents.Remove(ClusteredParticle);
 		MActiveRemovalIndices.Remove(ClusteredParticle);
 
diff --git a/Engine/Source/Runtime/Experimental/Chaos/Public/Chaos/PBDRigidClustering.h b/Engine/Source/Runtime/Experimental/Chaos/Public/Chaos/PBDRigidClustering.h
--- a/Engine/Source/Runtime/Experimental/Chaos/Public/Chaos/PBDRigidClustering.h
+++ b/Engine/Source/Runtime/Experimental/Chaos/Public/Chaos/PBDRigidClustering.h
@@ -278,7 +278,7 @@ public:
 	void ComputeStrainFromCollision(const FPBDCollisionConstraints& CollisionRule);
 	void ResetCollisionImpulseArray();
 	void DisableCluster(FPBDRigidClusteredParticleHandle* ClusteredParticle);
-	void DisableParticleWithBreakEvent(Chaos::FPBDRigidParticleHandle* Particle);
+	void DisableParticleWithBreakEvent(FPBDRigidClusteredParticleHandle* ClusteredParticle);
 
 	/*
 	* Connectivity
@@ -303,6 +303,18 @@ public:
 	void RemoveNodeConnections(FPBDRigidParticleHandle* Child);
 	void RemoveNodeConnections(FPBDRigidClusteredParticleHandle* Child);
 
+	void RemoveChildFromParent(FPBDRigidParticleHandle* Child, const FPBDRigidClusteredParticleHandle* ClusteredParent);
+	void SendBreakingEvent(FPBDRigidClusteredParticleHandle* ClusteredParticle);
+
+	TSet<FPBDRigidParticleHandle*> ReleaseClusterParticlesImpl(
+		FPBDRigidClusteredParticleHandle* ClusteredParticle, 
+		const TMap<FGeometryParticleHandle*, FReal>* ExternalStrainMap,
+		bool bForceRelease,
+		bool bCreateNewClusters);
+	
+	using FParticleIsland = TArray<FPBDRigidParticleHandle*>;
+	TArray<FParticleIsland> FindIslandsInChildren(const FPBDRigidClusteredParticleHandle* ClusteredParticle);
+	TArray<FPBDRigidParticleHandle*> CreateClustersFromNewIslands(TArray<FParticleIsland>& Islands, FPBDRigidClusteredParticleHandle* ClusteredParent);
 private:
 
 	FRigidEvolution& MEvolution;
