Chaos - change zero inertia assert to a warning and set default

diff --git a/Engine/Source/Runtime/Experimental/Chaos/Private/Chaos/MassProperties.cpp b/Engine/Source/Runtime/Experimental/Chaos/Private/Chaos/MassProperties.cpp
--- a/Engine/Source/Runtime/Experimental/Chaos/Private/Chaos/MassProperties.cpp
+++ b/Engine/Source/Runtime/Experimental/Chaos/Private/Chaos/MassProperties.cpp
@@ -252,7 +252,19 @@ namespace Chaos
 			NewMP.CenterOfMass += Child.CenterOfMass * Child.Mass;
 			NewMP.Mass += Child.Mass;
 		}
-		check(NewMP.Mass > UE_SMALL_NUMBER);
+
+		// Default to 100cm cube of water for zero mass and volume objects
+		if (!ensureMsgf(NewMP.Mass > UE_SMALL_NUMBER, TEXT("CombineWorldSpace: zero total mass detected")))
+		{
+			const FReal Dim = 100;	// cm
+			const FReal Density = FReal(0.001); // kg/cm3
+			NewMP.Volume = Dim * Dim * Dim;
+			NewMP.Mass = NewMP.Volume * Density;
+			NewMP.InertiaTensor = (NewMP.Mass * Dim * Dim / FReal(6)) * FMatrix33::Identity;
+			NewMP.CenterOfMass = FVec3(0);
+			return NewMP;
+		}
+
 		NewMP.CenterOfMass /= NewMP.Mass;
 
 		if (MPArray.Num() > 1)
