Chaos destruction: fix kinematic geometry collection having their velocity set to 0 when activated by a field
- reset kinematic target when setting the state to dynamic to avoid it to be evaluated after the field has been

diff --git a/Engine/Source/Runtime/Experimental/Chaos/Public/PhysicsProxy/FieldSystemProxyHelper.h b/Engine/Source/Runtime/Experimental/Chaos/Public/PhysicsProxy/FieldSystemProxyHelper.h
--- a/Engine/Source/Runtime/Experimental/Chaos/Public/PhysicsProxy/FieldSystemProxyHelper.h
+++ b/Engine/Source/Runtime/Experimental/Chaos/Public/PhysicsProxy/FieldSystemProxyHelper.h
@@ -192,6 +192,10 @@ namespace Chaos
 		else
 		{
 			RigidHandle->SetObjectStateLowLevel(FieldState);
+			if (FieldState == Chaos::EObjectStateType::Dynamic)
+			{
+				RigidSolver->GetEvolution()->SetParticleKinematicTarget(RigidHandle, FKinematicTarget());
+			}
 		}
 	}
 
