Backout changes to VectorCast4FloatTo4Int from 19967382

diff --git a/Engine/Source/Runtime/Experimental/Chaos/Public/Chaos/VectorUtility.h b/Engine/Source/Runtime/Experimental/Chaos/Public/Chaos/VectorUtility.h
--- a/Engine/Source/Runtime/Experimental/Chaos/Public/Chaos/VectorUtility.h
+++ b/Engine/Source/Runtime/Experimental/Chaos/Public/Chaos/VectorUtility.h
@@ -8,7 +8,18 @@
  * @param V	vector
  * @return		VectorRegister4Float( B.x, A.y, A.z, A.w)
  */
-#define VectorCast4IntTo4Float(Vec)		VectorCastIntToFloat(Vec)
+#if !defined(_MSC_VER)|| PLATFORM_ENABLE_VECTORINTRINSICS_NEON
+FORCEINLINE VectorRegister4Float VectorCast4IntTo4Float(const VectorRegister4Int& V)
+{
+	return VectorRegister4Float(V);
+}
+#else
+FORCEINLINE VectorRegister4Float VectorCast4IntTo4Float(const VectorRegister4Int& V)
+{
+	return _mm_castsi128_ps(V);
+}
+
+#endif
 
 /**
  * Cast VectorRegister4Float in VectorRegister4Int
@@ -74,3 +85,4 @@ FORCEINLINE VectorRegister4Float VectorBitwiseNotAnd(const VectorRegister4Float&
 		uint32(~((uint32*)(A.V))[3] & ((uint32*)(B.V))[3]));
 #endif
 }
+
