Chaos Cloth: Fix bug where (newly added) pressure is pushing in the opposite direction it's supposed to.

It turns out that our meshes have a left-handed winding, but most of our internal cloth code calculates right-handed winding normals. Flip sign of N when applying pressure force internally.

diff --git a/Engine/Source/Runtime/Experimental/Chaos/Private/Chaos/VelocityField.ispc b/Engine/Source/Runtime/Experimental/Chaos/Private/Chaos/VelocityField.ispc
--- a/Engine/Source/Runtime/Experimental/Chaos/Private/Chaos/VelocityField.ispc
+++ b/Engine/Source/Runtime/Experimental/Chaos/Private/Chaos/VelocityField.ispc
@@ -70,7 +70,7 @@ export void UpdateField(uniform FVector3f Forces[],
 		const FVector3f EvenForce = (Cd - Cl) * VDotN * RelVelocity + Cl * VSquare * N;
 		const FVector3f OddForce = (Cl - Cd) * VDotN * RelVelocity - Cl * VSquare * N;
 
-		const FVector3f Force = QuarterRho * DoubleArea * VectorSelect(VDotN >= 0, EvenForce, OddForce) + DoubleArea * 0.5 * Cp * N;
+		const FVector3f Force = QuarterRho * DoubleArea * VectorSelect(VDotN >= 0, EvenForce, OddForce) - DoubleArea * 0.5 * Cp * N; // N points in the opposite direction of the actual mesh normals
 
 		VectorStore(&Forces[Offset], Force);
 
@@ -144,7 +144,7 @@ export void UpdateFieldWithWeightMaps(uniform FVector3f Forces[],
 		const FVector3f EvenForce = (Cd - Cl) * VDotN * RelVelocity + Cl * VSquare * N;
 		const FVector3f OddForce = (Cl - Cd) * VDotN * RelVelocity - Cl * VSquare * N;
 
-		const FVector3f Force = QuarterRho * DoubleArea * VectorSelect(VDotN >= 0, EvenForce, OddForce) + DoubleArea * 0.5 * Cp * N;
+		const FVector3f Force = QuarterRho * DoubleArea * VectorSelect(VDotN >= 0, EvenForce, OddForce) - DoubleArea * 0.5 * Cp * N; // N points in the opposite direction of the actual mesh normals
 
 		VectorStore(&Forces[Offset], Force);
 
diff --git a/Engine/Source/Runtime/Experimental/Chaos/Public/Chaos/VelocityField.h b/Engine/Source/Runtime/Experimental/Chaos/Public/Chaos/VelocityField.h
--- a/Engine/Source/Runtime/Experimental/Chaos/Public/Chaos/VelocityField.h
+++ b/Engine/Source/Runtime/Experimental/Chaos/Public/Chaos/VelocityField.h
@@ -91,7 +91,7 @@ private:
 
 		Forces[ElementIndex] = QuarterRho * DoubleArea * (VDotN >= (FSolverReal)0. ?  // The flow can hit either side of the triangle, so the normal might need to be reversed
 			(Cd - Cl) * VDotN * V + Cl * VSquare * N :
-			(Cl - Cd) * VDotN * V - Cl * VSquare * N) + DoubleArea * (FSolverReal)0.5 * Cp * N;
+			(Cl - Cd) * VDotN * V - Cl * VSquare * N) - DoubleArea * (FSolverReal)0.5 * Cp * N; // N points in the opposite direction of the actual mesh normals
 	}
 
 private:
