FIX ROBOMERGE

Add the option to make collision constraints into "probe" constraints, which report contacts when they occur but are not considered for resolution inside the solver

diff --git a/Engine/Source/Runtime/Experimental/Chaos/Private/Chaos/CCDUtilities.cpp b/Engine/Source/Runtime/Experimental/Chaos/Private/Chaos/CCDUtilities.cpp
--- a/Engine/Source/Runtime/Experimental/Chaos/Private/Chaos/CCDUtilities.cpp
+++ b/Engine/Source/Runtime/Experimental/Chaos/Private/Chaos/CCDUtilities.cpp
@@ -122,7 +122,7 @@ namespace Chaos
 		CCDConstraints.Reserve(SweptConstraints.Num());
 		bool bNeedCCDSolve = false;
 		for (FPBDCollisionConstraint* Constraint : SweptConstraints)
-		{   
+		{
 			// A contact can be disabled by a user callback or contact pruning so we need to ignore these.
 			// NOTE: It important that we explicitly check for disabled here, rather than for 0 manifold points since we
 			// may want to use the contact later if resweeping is enabled. 
@@ -131,6 +131,12 @@ namespace Chaos
 				continue;
 			}
 
+			// For now, probe constraints do not work with CCD
+			if (Constraint->IsProbe())
+			{
+				continue;
+			}
+
 			// Create CCDParticle for all dynamic particles affected by swept constraints (UseCCD() could be either true or false). For static or kinematic particles, this pointer remains to be nullptr.
 			FCCDParticle* CCDParticlePair[2] = {nullptr, nullptr};
 			bool IsDynamic[2] = {false, false};
diff --git a/Engine/Source/Runtime/Experimental/Chaos/Private/Chaos/ChaosDebugDraw.cpp b/Engine/Source/Runtime/Experimental/Chaos/Private/Chaos/ChaosDebugDraw.cpp
--- a/Engine/Source/Runtime/Experimental/Chaos/Private/Chaos/ChaosDebugDraw.cpp
+++ b/Engine/Source/Runtime/Experimental/Chaos/Private/Chaos/ChaosDebugDraw.cpp
@@ -1038,7 +1038,37 @@ namespace Chaos
 			DrawParticleCCDCollisionImpulseImpl(SpaceTransform, CCDConstraint.Particle[1], CCDConstraint, ManifoldPointIndex, Impulse, Duration, Settings);
 		}
 
-		void DrawCollidingShapesImpl(const FRigidTransform3& SpaceTransform, const FPBDCollisionConstraints& Collisions, FRealSingle ColorScale, const FChaosDebugDrawSettings& Settings)
+		void DrawCollidingShapesImpl(const FRigidTransform3& SpaceTransform, const FPBDCollisionConstraint& Collision, FRealSingle ColorScale, const FRealSingle Duration, const FChaosDebugDrawSettings& Settings)
+		{
+			if (!Collision.GetDisabled() && (Collision.GetPhi() < TNumericLimits<FReal>::Max()))
+			{
+				const FImplicitObject* Implicit0 = Collision.GetImplicit0();
+				const FImplicitObject* Implicit1 = Collision.GetImplicit1();
+				if ((Implicit0 != nullptr) && (Implicit1 != nullptr))
+				{
+					FConstGenericParticleHandle Particle0 = Collision.GetParticle0();
+					FConstGenericParticleHandle Particle1 = Collision.GetParticle1();
+					const FPerShapeData* Shape0 = Collision.GetShape0();
+					const FPerShapeData* Shape1 = Collision.GetShape1();
+					const FRigidTransform3 WorldActorTransform0 = FParticleUtilities::GetActorWorldTransform(Particle0);
+					const FRigidTransform3 WorldActorTransform1 = FParticleUtilities::GetActorWorldTransform(Particle1);
+					const FRigidTransform3 ShapeWorldTransform0 = Collision.GetShapeRelativeTransform0() * WorldActorTransform0;
+					const FRigidTransform3 ShapeWorldTransform1 = Collision.GetShapeRelativeTransform1() * WorldActorTransform1;
+					DrawShapesImpl(
+						Particle0->Handle(), ShapeWorldTransform0,
+						Implicit0, FShapeOrShapesArray(Shape0),
+						0.0f, Particle0->IsDynamic() ? FColor::Yellow : FColor::Red,
+						Duration, Settings);
+					DrawShapesImpl(
+						Particle1->Handle(), ShapeWorldTransform1,
+						Implicit1, FShapeOrShapesArray(Shape1),
+						0.0f, Particle1->IsDynamic() ? FColor::Yellow : FColor::Red,
+						Duration, Settings);
+				}
+			}
+		}
+		
+		void DrawCollidingShapesImpl(const FRigidTransform3& SpaceTransform, const FPBDCollisionConstraints& Collisions, FRealSingle ColorScale, const FRealSingle Duration, const FChaosDebugDrawSettings& Settings)
 		{
 			TArray<const FImplicitObject*> Implicits;
 			TArray<const FPerShapeData*> Shapes;
@@ -1091,7 +1121,7 @@ namespace Chaos
 					Shapes[ShapeIndex],
 					0.0f, 
 					ShapeParticles[ShapeIndex]->IsDynamic() ? FColor::Yellow : FColor::Red, 
-					0.0f,
+					Duration,
 					Settings);
 			}
 		}
@@ -1685,11 +1715,19 @@ namespace Chaos
 			}
 		}
 
-		void DrawCollidingShapes(const FRigidTransform3& SpaceTransform, const FPBDCollisionConstraints& Collisions, FRealSingle ColorScale, const FChaosDebugDrawSettings* Settings)
+		void DrawCollidingShapes(const FRigidTransform3& SpaceTransform, const FPBDCollisionConstraints& Collisions, FRealSingle ColorScale, const FRealSingle Duration, const FChaosDebugDrawSettings* Settings)
+		{
+			if (FDebugDrawQueue::IsDebugDrawingEnabled())
+			{
+				DrawCollidingShapesImpl(SpaceTransform, Collisions, ColorScale, Duration, GetChaosDebugDrawSettings(Settings));
+			}
+		}
+
+		void DrawCollidingShapes(const FRigidTransform3& SpaceTransform, const FPBDCollisionConstraint& Collision, FRealSingle ColorScale, const FRealSingle Duration, const FChaosDebugDrawSettings* Settings)
 		{
 			if (FDebugDrawQueue::IsDebugDrawingEnabled())
 			{
-				DrawCollidingShapesImpl(SpaceTransform, Collisions, ColorScale, GetChaosDebugDrawSettings(Settings));
+				DrawCollidingShapesImpl(SpaceTransform, Collision, ColorScale, Duration, GetChaosDebugDrawSettings(Settings));
 			}
 		}
 
diff --git a/Engine/Source/Runtime/Experimental/Chaos/Private/Chaos/Collision/PBDCollisionConstraint.cpp b/Engine/Source/Runtime/Experimental/Chaos/Private/Chaos/Collision/PBDCollisionConstraint.cpp
--- a/Engine/Source/Runtime/Experimental/Chaos/Private/Chaos/Collision/PBDCollisionConstraint.cpp
+++ b/Engine/Source/Runtime/Experimental/Chaos/Private/Chaos/Collision/PBDCollisionConstraint.cpp
@@ -989,7 +989,7 @@ namespace Chaos
 
 	ECollisionConstraintDirection FPBDCollisionConstraint::GetConstraintDirection(const FReal Dt) const
 	{
-		if (GetDisabled())
+		if (GetDisabled() || GetIsProbe())
 		{
 			return NoRestingDependency;
 		}
diff --git a/Engine/Source/Runtime/Experimental/Chaos/Private/Chaos/Collision/PBDCollisionConstraintHandle.cpp b/Engine/Source/Runtime/Experimental/Chaos/Private/Chaos/Collision/PBDCollisionConstraintHandle.cpp
--- a/Engine/Source/Runtime/Experimental/Chaos/Private/Chaos/Collision/PBDCollisionConstraintHandle.cpp
+++ b/Engine/Source/Runtime/Experimental/Chaos/Private/Chaos/Collision/PBDCollisionConstraintHandle.cpp
@@ -42,6 +42,11 @@ namespace Chaos
 		return !GetContact().GetDisabled();
 	}
 
+	bool FPBDCollisionConstraintHandle::IsProbe() const
+	{
+		return GetContact().GetIsProbe();
+	}
+
 	FVec3 FPBDCollisionConstraintHandle::GetAccumulatedImpulse() const
 	{
 		return GetContact().AccumulatedImpulse;
diff --git a/Engine/Source/Runtime/Experimental/Chaos/Private/Chaos/ContactModification.cpp b/Engine/Source/Runtime/Experimental/Chaos/Private/Chaos/ContactModification.cpp
--- a/Engine/Source/Runtime/Experimental/Chaos/Private/Chaos/ContactModification.cpp
+++ b/Engine/Source/Runtime/Experimental/Chaos/Private/Chaos/ContactModification.cpp
@@ -16,6 +16,11 @@ namespace Chaos
 		Modifier->EnableConstraint(*Constraint);
 	}
 
+	void FContactPairModifier::ConvertToProbe()
+	{
+		Modifier->ConvertToProbeConstraint(*Constraint);
+	}
+
 	int32 FContactPairModifier::GetNumContacts() const
 	{
 		return Constraint->GetManifoldPoints().Num();
@@ -476,6 +481,11 @@ namespace Chaos
 		Constraint.SetDisabled(false);
 	}
 
+	void FCollisionContactModifier::ConvertToProbeConstraint(FPBDCollisionConstraint& Constraint)
+	{
+		Constraint.SetIsProbe(true);
+	}
+
 	void FCollisionContactModifier::MarkConstraintForManifoldUpdate(FPBDCollisionConstraint& Constraint)
 	{
 		NeedsManifoldUpdate.Add(&Constraint);
diff --git a/Engine/Source/Runtime/Experimental/Chaos/Private/Chaos/GeometryParticles.cpp b/Engine/Source/Runtime/Experimental/Chaos/Private/Chaos/GeometryParticles.cpp
--- a/Engine/Source/Runtime/Experimental/Chaos/Private/Chaos/GeometryParticles.cpp
+++ b/Engine/Source/Runtime/Experimental/Chaos/Private/Chaos/GeometryParticles.cpp
@@ -9,6 +9,8 @@
 #include "Chaos/ParticleHandle.h"
 #include "Chaos/Framework/PhysicsSolverBase.h"
 
+#pragma optimize("", off)
+
 namespace Chaos
 {
 	void UpdateShapesArrayFromGeometry(FShapesArray& ShapesArray, TSerializablePtr<FImplicitObject> Geometry, const FRigidTransform3& ActorTM, IPhysicsProxyBase* Proxy)
@@ -249,7 +251,7 @@ namespace Chaos
 		return FRigidTransformRealSingle3::Identity;
 	}
 
-	FRigidTransform3 FPerShapeData::GetLeafWorldTransform(FGeometryParticleHandle* Particle) const
+	FRigidTransform3 FPerShapeData::GetLeafWorldTransform(const FGeometryParticleHandle* Particle) const
 	{
 		if (bHasCachedLeafInfo)
 		{
diff --git a/Engine/Source/Runtime/Experimental/Chaos/Private/Chaos/PBDCollisionConstraints.cpp b/Engine/Source/Runtime/Experimental/Chaos/Private/Chaos/PBDCollisionConstraints.cpp
--- a/Engine/Source/Runtime/Experimental/Chaos/Private/Chaos/PBDCollisionConstraints.cpp
+++ b/Engine/Source/Runtime/Experimental/Chaos/Private/Chaos/PBDCollisionConstraints.cpp
@@ -28,7 +28,7 @@
 #include "PBDCollisionConstraints.ispc.generated.h"
 #endif
 
-//PRAGMA_DISABLE_OPTIMIZATION
+PRAGMA_DISABLE_OPTIMIZATION
 
 namespace Chaos
 {
@@ -70,14 +70,25 @@ namespace Chaos
 
 	bool CollisionsAllowParticleTracking = true;
 	FAutoConsoleVariableRef CVarCollisionsAllowParticleTracking(TEXT("p.Chaos.Collision.AllowParticleTracking"), CollisionsAllowParticleTracking, TEXT("Allow particles to track their collisions constraints when their DoBufferCollisions flag is enable [def:true]"));
-
+	
 	bool bCollisionsEnableSubSurfaceCollisionPruning = false;
 	FAutoConsoleVariableRef CVarCollisionsEnableSubSurfaceCollisionPruning(TEXT("p.Chaos.Collision.EnableSubSurfaceCollisionPruning"), bCollisionsEnableSubSurfaceCollisionPruning, TEXT(""));
+
+	bool DebugDrawProbeDetection = false;
+	FAutoConsoleVariableRef CVarDebugDrawProbeDetection(TEXT("p.Chaos.Collision.DebugDrawProbeDetection"), DebugDrawProbeDetection, TEXT("Draw probe constraint detection."));
+	
+#if CHAOS_DEBUG_DRAW
+	namespace CVars
+	{
+		extern DebugDraw::FChaosDebugDrawSettings ChaosSolverDebugDebugDrawSettings;
+	}
+#endif
 	
 	DECLARE_CYCLE_STAT(TEXT("Collisions::Reset"), STAT_Collisions_Reset, STATGROUP_ChaosCollision);
 	DECLARE_CYCLE_STAT(TEXT("Collisions::UpdatePointConstraints"), STAT_Collisions_UpdatePointConstraints, STATGROUP_ChaosCollision);
 	DECLARE_CYCLE_STAT(TEXT("Collisions::BeginDetect"), STAT_Collisions_BeginDetect, STATGROUP_ChaosCollision);
 	DECLARE_CYCLE_STAT(TEXT("Collisions::EndDetect"), STAT_Collisions_EndDetect, STATGROUP_ChaosCollision);
+	DECLARE_CYCLE_STAT(TEXT("Collisions::DetectProbeCollisions"), STAT_Collisions_DetectProbeCollisions, STATGROUP_ChaosCollision);
 
 	//
 	// Collision Constraint Container
@@ -293,6 +304,37 @@ namespace Chaos
 		}
 	}
 
+	void FPBDCollisionConstraints::DetectProbeCollisions(FReal Dt)
+	{
+		SCOPE_CYCLE_COUNTER(STAT_Collisions_DetectProbeCollisions);
+
+		// Do a final narrow-phase update on probe constraints. This is done to avoid
+		// false positive probe hit results, which may occur if the constraint was created
+		// for a contact which no longer is occurring due to resolution of another constraint.
+		for (FPBDCollisionConstraint* Contact : GetConstraints())
+		{
+			if (Contact->IsProbe())
+			{
+				const FGeometryParticleHandle* Particle0 = Contact->GetParticle0();
+				const FGeometryParticleHandle* Particle1 = Contact->GetParticle1();
+				const FRigidTransform3& ShapeWorldTransform0 = Contact->GetShapeRelativeTransform0() * FParticleUtilities::GetActorWorldTransform(FConstGenericParticleHandle(Particle0));
+				const FRigidTransform3& ShapeWorldTransform1 = Contact->GetShapeRelativeTransform1() * FParticleUtilities::GetActorWorldTransform(FConstGenericParticleHandle(Particle1));
+				Contact->SetCullDistance(0.f);
+				Contact->SetShapeWorldTransforms(ShapeWorldTransform0, ShapeWorldTransform1);
+				Collisions::UpdateConstraint(*Contact, ShapeWorldTransform0, ShapeWorldTransform1, Dt);
+
+#if CHAOS_DEBUG_DRAW
+				if (DebugDrawProbeDetection)
+				{
+					DebugDraw::FChaosDebugDrawSettings Settings = CVars::ChaosSolverDebugDebugDrawSettings;
+					Settings.DrawDuration = 1.f;
+					DebugDraw::DrawCollidingShapes(FRigidTransform3(), *Contact, 1.f, 1.f, &Settings);
+				}
+#endif
+			}
+		}
+	}
+
 	void FPBDCollisionConstraints::ApplyCollisionModifier(const TArray<ISimCallbackObject*>& CollisionModifiers, FReal Dt)
 	{
 		if (GetConstraints().Num() > 0)
@@ -377,7 +419,7 @@ namespace Chaos
 		if (SolverType == EConstraintSolverType::QuasiPbd)
 		{
 			// We shouldn't be adding disabled constraints to the solver list. The check needs to be at caller site or we should return success/fail - see TPBDConstraintColorRule::GatherSolverInput
-			check(Constraint.IsEnabled());
+			check(Constraint.IsEnabled() && !Constraint.IsProbe());
 
 			FPBDCollisionSolverContainer& SolverContainer = GetConstraintSolverContainer(SolverData);
 			SolverContainer.AddConstraintSolver(Dt, Constraint, Particle0Level, Particle1Level, SolverData.GetBodyContainer(), SolverSettings);
@@ -408,7 +450,7 @@ namespace Chaos
 		{
 			for (FPBDCollisionConstraint* Constraint : GetConstraints())
 			{
-				if (Constraint->IsEnabled())
+				if (Constraint->IsEnabled() && !Constraint->IsProbe())
 				{
 					GatherInput(Dt, *Constraint, INDEX_NONE, INDEX_NONE, SolverData);
 				}
@@ -418,7 +460,7 @@ namespace Chaos
 		{
 			for (FPBDCollisionConstraint* Constraint : GetConstraints())
 			{
-				if (Constraint->IsEnabled())
+				if (Constraint->IsEnabled() && !Constraint->IsProbe())
 				{
 					LegacyGatherInput(Dt, *Constraint, INDEX_NONE, INDEX_NONE, SolverData);
 				}
@@ -607,7 +649,8 @@ namespace Chaos
 			for (int32 Index = BeginIndex; Index < EndIndex; ++Index)
 			{
 				FPBDCollisionConstraint* Constraint = SolverData.GetConstraintHandle<FPBDCollisionConstraint>(ContainerId,Index);
-				if (!Constraint->GetDisabled())
+
+				if ((!Constraint->GetDisabled()) && (!Constraint->GetIsProbe()))
 				{
 					Collisions::Apply(*Constraint, IterationParameters, ParticleParameters);
 					++NumActivePointConstraints;
diff --git a/Engine/Source/Runtime/Experimental/Chaos/Private/Chaos/PBDConstraintRule.cpp b/Engine/Source/Runtime/Experimental/Chaos/Private/Chaos/PBDConstraintRule.cpp
--- a/Engine/Source/Runtime/Experimental/Chaos/Private/Chaos/PBDConstraintRule.cpp
+++ b/Engine/Source/Runtime/Experimental/Chaos/Private/Chaos/PBDConstraintRule.cpp
@@ -139,6 +139,8 @@ namespace Chaos
 
 		for (typename FConstraints::FConstraintContainerHandle * ConstraintHandle : Constraints.GetConstraintHandles())
 		{
+			// NOTE: For now, we include probe constraints as well so that
+			// they can trigger collision events
 			if (ConstraintHandle->IsEnabled())
 			{
 				bool DisabledParticle0 = (ConstraintHandle->GetConstrainedParticles()[0] && ConstraintHandle->GetConstrainedParticles()[0]->CastToRigidParticle()) ? ConstraintHandle->GetConstrainedParticles()[0]->CastToRigidParticle()->Disabled() : false;
@@ -193,7 +195,7 @@ namespace Chaos
 
 								// Note we are building the SolverBodies as we go, in the order that we visit them. Each constraint
 								// references two bodies, so we won't strictly be accessing only in cache order, but it's about as good as it can be.
-								if (Constraint->IsEnabled())
+								if (Constraint->IsEnabled() && !Constraint->IsProbe())
 								{
 									// @todo(chaos): we should provide Particle Levels in the island rule as well (see TPBDConstraintColorRule)
 									Constraint->PreGatherInput(Dt, *IslandGroup);
@@ -209,7 +211,7 @@ namespace Chaos
 
 								// Note we are building the SolverBodies as we go, in the order that we visit them. Each constraint
 								// references two bodies, so we won't strictly be accessing only in cache order, but it's about as good as it can be.
-								if (Constraint->IsEnabled())
+								if (Constraint->IsEnabled() && !Constraint->IsProbe())
 								{
 									// @todo(chaos): we should provide Particle Levels in the island rule as well (see TPBDConstraintColorRule)
 									Constraint->GatherInput(Dt, INDEX_NONE, INDEX_NONE, *IslandGroup);
@@ -338,7 +340,7 @@ namespace Chaos
 										for (int32 ConstraintIndex = OffsetBegin; ConstraintIndex < OffsetEnd; ++ConstraintIndex)
 										{
 											FConstraintContainerHandle* Constraint = SortedConstraints[ConstraintIndex];
-											if (Constraint->IsEnabled())
+											if (Constraint->IsEnabled() && !Constraint->IsProbe())
 											{
 												Constraint->PreGatherInput(Dt, *IslandGroup);
 											}
@@ -351,7 +353,7 @@ namespace Chaos
 											{
 												int32 ConstraintIndex = Index + OffsetBegin;
 												FConstraintContainerHandle* Constraint = SortedConstraints[ConstraintIndex];
-												if (Constraint->IsEnabled())
+												if (Constraint->IsEnabled() && !Constraint->IsProbe())
 												{
 													// Levels that should be assigned to the bodies for shock propagation
 													// @todo(chaos): optimize the lookup
@@ -368,7 +370,7 @@ namespace Chaos
 										for (int32 ConstraintIndex = OffsetBegin; ConstraintIndex < OffsetEnd; ++ConstraintIndex)
 										{
 											FConstraintContainerHandle* Constraint = SortedConstraints[ConstraintIndex];
-											if (Constraint->IsEnabled())
+											if (Constraint->IsEnabled() && !Constraint->IsProbe())
 											{
 												// Update the current constraint set of this color
 												ColorConstrainSet.Value = ++ConstraintSetEnd;
@@ -651,4 +653,4 @@ namespace Chaos
 	template class TPBDConstraintIslandRule<FPBDSuspensionConstraints>;
 	template class TPBDConstraintIslandRule<FPBDRigidDynamicSpringConstraints>;
 	template class TPBDConstraintIslandRule<FPBDRigidSpringConstraints>;
-}
+}
\ No newline at end of file
diff --git a/Engine/Source/Runtime/Experimental/Chaos/Private/Chaos/PBDRigidsEvolutionGBF.cpp b/Engine/Source/Runtime/Experimental/Chaos/Private/Chaos/PBDRigidsEvolutionGBF.cpp
--- a/Engine/Source/Runtime/Experimental/Chaos/Private/Chaos/PBDRigidsEvolutionGBF.cpp
+++ b/Engine/Source/Runtime/Experimental/Chaos/Private/Chaos/PBDRigidsEvolutionGBF.cpp
@@ -88,6 +88,9 @@ namespace Chaos
 		bool TransferCollisionsDebugTestAgainstMaxClamp = false;
 		FAutoConsoleVariableRef CVarTransferCollisionsDebugTestAgainstMaxClamp(TEXT("p.Chaos.Solver.Joint.TransferCollisionsDebugTestAgainstMaxClamp"), TransferCollisionsDebugTestAgainstMaxClamp, TEXT("Force all joint collision constraint settings to max clamp value to validate stability [def:false]"));
 
+		bool DoFinalProbeNarrowPhase = true;
+		FAutoConsoleVariableRef CVarDoFinalProbeNarrowPhase(TEXT("p.Chaos.Solver.DoFinalProbeNarrowPhase"), DoFinalProbeNarrowPhase, TEXT(""));
+
 		DECLARE_CYCLE_STAT(TEXT("FPBDRigidsEvolutionGBF::AdvanceOneTimeStep"), STAT_Evolution_AdvanceOneTimeStep, STATGROUP_Chaos);
 		DECLARE_CYCLE_STAT(TEXT("FPBDRigidsEvolutionGBF::UnclusterUnions"), STAT_Evolution_UnclusterUnions, STATGROUP_Chaos);
 		DECLARE_CYCLE_STAT(TEXT("FPBDRigidsEvolutionGBF::Integrate"), STAT_Evolution_Integrate, STATGROUP_Chaos);
@@ -578,6 +581,12 @@ void FPBDRigidsEvolutionGBF::AdvanceOneTimeStepImpl(const FReal Dt, const FSubSt
 
 	ParticleUpdatePosition(Particles.GetDirtyParticlesView(), Dt);
 
+	if (CVars::DoFinalProbeNarrowPhase)
+	{
+		// Run contact updates on probe constraints
+		GetCollisionConstraints().DetectProbeCollisions(Dt);
+	}
+
 #if !UE_BUILD_SHIPPING
 	if(SerializeEvolution)
 	{
diff --git a/Engine/Source/Runtime/Experimental/Chaos/Private/EventDefaults.cpp b/Engine/Source/Runtime/Experimental/Chaos/Private/EventDefaults.cpp
--- a/Engine/Source/Runtime/Experimental/Chaos/Private/EventDefaults.cpp
+++ b/Engine/Source/Runtime/Experimental/Chaos/Private/EventDefaults.cpp
@@ -15,6 +15,8 @@
 #include "ChaosSolversModule.h"
 #include "Chaos/CollisionFilterData.h"
 
+#pragma optimize("", off)
+
 namespace Chaos
 {
 	DECLARE_CYCLE_STAT(TEXT("EventDefaults::RegisterCollisionEvent_Filter"), STAT_Events_RegisterCollisionEvent_Filter, STATGROUP_Chaos);
@@ -134,6 +136,7 @@ namespace Chaos
 
 							if (ensure(!Constraint.AccumulatedImpulse.ContainsNaN() && FMath::IsFinite(Constraint.GetPhi())))
 							{
+	
 								const FPerShapeData* Shape0 = Constraint.GetShape0();
 								const FPerShapeData* Shape1 = Constraint.GetShape1();
 
@@ -141,7 +144,7 @@ namespace Chaos
 								const bool bFilter0Notify = Shape0 ? Shape0->GetSimData().HasFlag(EFilterFlags::ContactNotify) : true;
 								const bool bFilter1Notify = Shape1 ? Shape1->GetSimData().HasFlag(EFilterFlags::ContactNotify) : true;
 
-								if (!bFilter0Notify && !bFilter1Notify)
+								if(!bFilter0Notify && !bFilter1Notify)
 								{
 									// No need to notify - engine didn't request notifications for either shape.
 									continue;
@@ -157,7 +160,12 @@ namespace Chaos
 								const FKinematicGeometryParticleHandle* Primary = Body0 ? Body0 : Body1;
 								const FKinematicGeometryParticleHandle* Secondary = Body0 ? Body1 : Body0;
 
-								if (!Constraint.AccumulatedImpulse.IsZero() && Primary)
+								//
+								const bool bIsProbe = Constraint.IsProbe();
+								const bool bUseManifold = Constraint.GetUseManifold();
+								const int32 NumManifoldPoints = Constraint.GetManifoldPoints().Num();
+
+								if (((Constraint.IsProbe() && NumManifoldPoints > 0) || !Constraint.AccumulatedImpulse.IsZero()) && Primary)
 								{
 									if (ensure(!Constraint.CalculateWorldContactLocation().ContainsNaN() &&
 										!Constraint.CalculateWorldContactNormal().ContainsNaN()) &&
diff --git a/Engine/Source/Runtime/Experimental/Chaos/Private/PBDRigidsSolver.cpp b/Engine/Source/Runtime/Experimental/Chaos/Private/PBDRigidsSolver.cpp
--- a/Engine/Source/Runtime/Experimental/Chaos/Private/PBDRigidsSolver.cpp
+++ b/Engine/Source/Runtime/Experimental/Chaos/Private/PBDRigidsSolver.cpp
@@ -1574,7 +1574,7 @@ CSV_CUSTOM_STAT(PhysicsCounters, Name, Value, ECsvCustomStatOp::Set);
 		}
 		if (ChaosSolverDebugDrawCollidingShapes == 1)
 		{
-			DebugDraw::DrawCollidingShapes(FRigidTransform3(), GetEvolution()->GetCollisionConstraints(), 1.f, &ChaosSolverDebugDebugDrawSettings);
+			DebugDraw::DrawCollidingShapes(FRigidTransform3(), GetEvolution()->GetCollisionConstraints(), 1.f, 0.f, &ChaosSolverDebugDebugDrawSettings);
 		}
 		if (ChaosSolverDrawJoints == 1)
 		{
@@ -1589,7 +1589,7 @@ CSV_CUSTOM_STAT(PhysicsCounters, Name, Value, ECsvCustomStatOp::Set);
 		}
 		if (ChaosSolverDebugDrawCollidingShapes == 1)
 		{
-			DebugDraw::DrawCollidingShapes(FRigidTransform3(), GetEvolution()->GetCollisionConstraints(), 1.f, &ChaosSolverDebugDebugDrawSettings);
+			DebugDraw::DrawCollidingShapes(FRigidTransform3(), GetEvolution()->GetCollisionConstraints(), 1.f, 0.f, &ChaosSolverDebugDebugDrawSettings);
 		}
 		if (ChaosSolverDebugDrawSuspensionConstraints == 1)
 		{
diff --git a/Engine/Source/Runtime/Experimental/Chaos/Public/Chaos/ChaosDebugDraw.h b/Engine/Source/Runtime/Experimental/Chaos/Public/Chaos/ChaosDebugDraw.h
--- a/Engine/Source/Runtime/Experimental/Chaos/Public/Chaos/ChaosDebugDraw.h
+++ b/Engine/Source/Runtime/Experimental/Chaos/Public/Chaos/ChaosDebugDraw.h
@@ -142,6 +142,7 @@ namespace Chaos
 			FRealSingle ImpulseScale;
 			FRealSingle PushOutScale;
 			FRealSingle InertiaScale;
+			FRealSingle DrawDuration;
 			uint8 DrawPriority;
 			bool bShowSimpleCollision;
 			bool bShowComplexCollision;
@@ -216,7 +217,8 @@ namespace Chaos
 		CHAOS_API void DrawShape(const FRigidTransform3& ShapeTransform, const FImplicitObject* Implicit, const FShapeOrShapesArray& Shapes, const FColor& Color, const FChaosDebugDrawSettings* Settings = nullptr);
 		CHAOS_API void DrawConstraintGraph(const FRigidTransform3& ShapeTransform, const FPBDConstraintGraph& Graph, const FChaosDebugDrawSettings* Settings = nullptr);
 		CHAOS_API void DrawConnectionGraph(const FRigidClustering& Clustering, const FChaosDebugDrawSettings* Settings = nullptr);
-		CHAOS_API void DrawCollidingShapes(const FRigidTransform3& SpaceTransform, const FPBDCollisionConstraints& Collisions, FRealSingle ColorScale, const FChaosDebugDrawSettings* Settings = nullptr);
+		CHAOS_API void DrawCollidingShapes(const FRigidTransform3& SpaceTransform, const FPBDCollisionConstraints& Collisions, FRealSingle ColorScale, const FRealSingle Duration = 0.f, const FChaosDebugDrawSettings* Settings = nullptr);
+		CHAOS_API void DrawCollidingShapes(const FRigidTransform3& SpaceTransform, const FPBDCollisionConstraint& Collision, FRealSingle ColorScale, const FRealSingle Duration = 0.f, const FChaosDebugDrawSettings* Settings = nullptr);
 		CHAOS_API void DrawSpatialAccelerationStructure(const ISpatialAcceleration<FAccelerationStructureHandle, FReal, 3>& SpatialAccelerationStructure, const FChaosDebugDrawSettings* Settings = nullptr);
 		CHAOS_API void DrawSuspensionConstraints(const FRigidTransform3& SpaceTransform, const FPBDSuspensionConstraints& Constraints, const FChaosDebugDrawSettings* Settings = nullptr);
 		CHAOS_API void DrawCCDAxisThreshold(const FVec3& X, const FVec3& AxisThreshold, const FVec3& DeltaX, const FQuat& R);
diff --git a/Engine/Source/Runtime/Experimental/Chaos/Public/Chaos/Collision/PBDCollisionConstraint.h b/Engine/Source/Runtime/Experimental/Chaos/Public/Chaos/Collision/PBDCollisionConstraint.h
--- a/Engine/Source/Runtime/Experimental/Chaos/Public/Chaos/Collision/PBDCollisionConstraint.h
+++ b/Engine/Source/Runtime/Experimental/Chaos/Public/Chaos/Collision/PBDCollisionConstraint.h
@@ -355,6 +355,9 @@ namespace Chaos
 		void SetDisabled(bool bInDisabled) { Flags.bDisabled = bInDisabled; }
 		bool GetDisabled() const { return Flags.bDisabled; }
 
+		void SetIsProbe(bool bInProbe) { Flags.bProbe = bInProbe; }
+		bool GetIsProbe() const { return Flags.bProbe; }
+
 		virtual void SetIsSleeping(const bool bInIsSleeping) override;
 
 		// Get the world-space normal of the closest manifold point
@@ -736,6 +739,7 @@ namespace Chaos
 			struct
 			{
 				uint32 bDisabled : 1;
+				uint32 bProbe : 1;
 				uint32 bUseManifold : 1;
 				uint32 bUseIncrementalManifold : 1;
 				uint32 bWasManifoldRestored : 1;
diff --git a/Engine/Source/Runtime/Experimental/Chaos/Public/Chaos/Collision/PBDCollisionConstraintHandle.h b/Engine/Source/Runtime/Experimental/Chaos/Public/Chaos/Collision/PBDCollisionConstraintHandle.h
--- a/Engine/Source/Runtime/Experimental/Chaos/Public/Chaos/Collision/PBDCollisionConstraintHandle.h
+++ b/Engine/Source/Runtime/Experimental/Chaos/Public/Chaos/Collision/PBDCollisionConstraintHandle.h
@@ -79,6 +79,8 @@ namespace Chaos
 
 		virtual bool IsEnabled() const override;
 
+		virtual bool IsProbe() const override;
+
 		//FVec3 GetContactLocation() const;
 
 		FVec3 GetAccumulatedImpulse() const;
diff --git a/Engine/Source/Runtime/Experimental/Chaos/Public/Chaos/ConstraintHandle.h b/Engine/Source/Runtime/Experimental/Chaos/Public/Chaos/ConstraintHandle.h
--- a/Engine/Source/Runtime/Experimental/Chaos/Public/Chaos/ConstraintHandle.h
+++ b/Engine/Source/Runtime/Experimental/Chaos/Public/Chaos/ConstraintHandle.h
@@ -129,6 +129,8 @@ namespace Chaos
 
 		virtual bool IsEnabled() const = 0;
 
+		virtual bool IsProbe() const { return false; }
+
 		virtual bool IsSleeping() const { return false; }
 		virtual void SetIsSleeping(const bool bInIsSleeping) {}
 		
diff --git a/Engine/Source/Runtime/Experimental/Chaos/Public/Chaos/ContactModification.h b/Engine/Source/Runtime/Experimental/Chaos/Public/Chaos/ContactModification.h
--- a/Engine/Source/Runtime/Experimental/Chaos/Public/Chaos/ContactModification.h
+++ b/Engine/Source/Runtime/Experimental/Chaos/Public/Chaos/ContactModification.h
@@ -36,6 +36,13 @@ namespace Chaos
 		*/
 		void Enable();
 
+		/**
+		* Convert the constraint between this pair of bodies into a probe. Collision callbacks will still occur if the
+		* bodies collide (a contact occurring in contact modification does not mean a contact actually has occurred
+		* yet, but that it may)
+		*/
+		void ConvertToProbe();
+
 		/**
 		* @return Number of contact points in constraint pair. ContactPointIdx must be below number of contacts.
 		*/
@@ -329,6 +336,10 @@ namespace Chaos
 		void DisableConstraint(FPBDCollisionConstraint& Constraint);
 		void EnableConstraint(FPBDCollisionConstraint& Constraint);
 
+		// Turn this constraint into a probe. It will still generate hit events,
+		// but will not produce impulses
+		void ConvertToProbeConstraint(FPBDCollisionConstraint& Constraint);
+
 		void MarkConstraintForManifoldUpdate(FPBDCollisionConstraint& Constraint);
 
 		// Update manifolds of modified constraints.
diff --git a/Engine/Source/Runtime/Experimental/Chaos/Public/Chaos/GeometryParticles.h b/Engine/Source/Runtime/Experimental/Chaos/Public/Chaos/GeometryParticles.h
--- a/Engine/Source/Runtime/Experimental/Chaos/Public/Chaos/GeometryParticles.h
+++ b/Engine/Source/Runtime/Experimental/Chaos/Public/Chaos/GeometryParticles.h
@@ -87,7 +87,7 @@ namespace Chaos
 		// The world-space transform of the leaf geometry.
 		// If we have non-identity leaf relative transform, is cached from the last call to UpdateWorldSpaceState.
 		// If not cahced, is constructed from arguments.
-		FRigidTransform3 GetLeafWorldTransform(FGeometryParticleHandle* Particle) const;
+		FRigidTransform3 GetLeafWorldTransform(const FGeometryParticleHandle* Particle) const;
 		void UpdateLeafWorldTransform(FGeometryParticleHandle* Particle);
 
 		// Are we caching leaf transform info?
diff --git a/Engine/Source/Runtime/Experimental/Chaos/Public/Chaos/PBDCollisionConstraints.h b/Engine/Source/Runtime/Experimental/Chaos/Public/Chaos/PBDCollisionConstraints.h
--- a/Engine/Source/Runtime/Experimental/Chaos/Public/Chaos/PBDCollisionConstraints.h
+++ b/Engine/Source/Runtime/Experimental/Chaos/Public/Chaos/PBDCollisionConstraints.h
@@ -114,6 +114,11 @@ public:
 	*/
 	void EndDetectCollisions();
 
+	/**
+	 * @brief Called after collision resolution in order to detect probes
+	 */
+	void DetectProbeCollisions(FReal Dt);
+
 	/**
 	 * Apply modifiers to the constraints and specify which constraints should be disabled.
 	 * You would probably call this in the PostComputeCallback. Prefer this to calling RemoveConstraints in a loop,
