Fix CIS - remove pragma optimize off's

diff --git a/Engine/Source/Runtime/Experimental/Chaos/Private/Chaos/GeometryParticles.cpp b/Engine/Source/Runtime/Experimental/Chaos/Private/Chaos/GeometryParticles.cpp
--- a/Engine/Source/Runtime/Experimental/Chaos/Private/Chaos/GeometryParticles.cpp
+++ b/Engine/Source/Runtime/Experimental/Chaos/Private/Chaos/GeometryParticles.cpp
@@ -9,8 +9,6 @@
 #include "Chaos/ParticleHandle.h"
 #include "Chaos/Framework/PhysicsSolverBase.h"
 
-#pragma optimize("", off)
-
 namespace Chaos
 {
 	void UpdateShapesArrayFromGeometry(FShapesArray& ShapesArray, TSerializablePtr<FImplicitObject> Geometry, const FRigidTransform3& ActorTM, IPhysicsProxyBase* Proxy)
diff --git a/Engine/Source/Runtime/Experimental/Chaos/Private/Chaos/PBDCollisionConstraints.cpp b/Engine/Source/Runtime/Experimental/Chaos/Private/Chaos/PBDCollisionConstraints.cpp
--- a/Engine/Source/Runtime/Experimental/Chaos/Private/Chaos/PBDCollisionConstraints.cpp
+++ b/Engine/Source/Runtime/Experimental/Chaos/Private/Chaos/PBDCollisionConstraints.cpp
@@ -28,7 +28,7 @@
 #include "PBDCollisionConstraints.ispc.generated.h"
 #endif
 
-PRAGMA_DISABLE_OPTIMIZATION
+//PRAGMA_DISABLE_OPTIMIZATION
 
 namespace Chaos
 {
diff --git a/Engine/Source/Runtime/Experimental/Chaos/Private/EventDefaults.cpp b/Engine/Source/Runtime/Experimental/Chaos/Private/EventDefaults.cpp
--- a/Engine/Source/Runtime/Experimental/Chaos/Private/EventDefaults.cpp
+++ b/Engine/Source/Runtime/Experimental/Chaos/Private/EventDefaults.cpp
@@ -15,8 +15,6 @@
 #include "ChaosSolversModule.h"
 #include "Chaos/CollisionFilterData.h"
 
-#pragma optimize("", off)
-
 namespace Chaos
 {
 	DECLARE_CYCLE_STAT(TEXT("EventDefaults::RegisterCollisionEvent_Filter"), STAT_Events_RegisterCollisionEvent_Filter, STATGROUP_Chaos);
