Chaos : detsruction : better damage threshold generation
- damage threshold are now by default generated per GC piece and no more at the cluster level propagating to children
- added a checkbox to keep the compatibility with older assets
- this allow for finer artist control and make the size specific parameters to be more effective

diff --git a/Engine/Source/Runtime/Core/Public/UObject/UE5MainStreamObjectVersion.h b/Engine/Source/Runtime/Core/Public/UObject/UE5MainStreamObjectVersion.h
--- a/Engine/Source/Runtime/Core/Public/UObject/UE5MainStreamObjectVersion.h
+++ b/Engine/Source/Runtime/Core/Public/UObject/UE5MainStreamObjectVersion.h
@@ -215,6 +215,9 @@ struct CORE_API FUE5MainStreamObjectVersion
 		// Fix material feature level nodes to account for new SM6 input pin.
 		MaterialFeatureLevelNodeFixForSM6,
 
+		// Fix material feature level nodes to account for new SM6 input pin.
+		GeometryCollectionPerChildDamageThreshold,
+
 		// -----<new versions can be added above this line>-------------------------------------------------
 		VersionPlusOne,
 		LatestVersion = VersionPlusOne - 1
diff --git a/Engine/Source/Runtime/Experimental/Chaos/Private/PhysicsProxy/GeometryCollectionPhysicsProxy.cpp b/Engine/Source/Runtime/Experimental/Chaos/Private/PhysicsProxy/GeometryCollectionPhysicsProxy.cpp
--- a/Engine/Source/Runtime/Experimental/Chaos/Private/PhysicsProxy/GeometryCollectionPhysicsProxy.cpp
+++ b/Engine/Source/Runtime/Experimental/Chaos/Private/PhysicsProxy/GeometryCollectionPhysicsProxy.cpp
@@ -792,7 +792,13 @@ void FGeometryCollectionPhysicsProxy::InitializeBodiesPT(Chaos::FPBDRigidsSolver
 				if (Parameters.EnableClustering)
 				{
 					Handle->SetClusterGroupIndex(Parameters.ClusterGroupIndex);
-					Handle->SetStrain(StrainDefault);
+
+					float DamageThreshold = StrainDefault; 
+					if (!Parameters.bUsePerClusterOnlyDamageThreshold)
+					{
+						DamageThreshold = ComputeDamageThreshold(DynamicCollection, TransformGroupIndex);
+					}
+					Handle->SetStrain(DamageThreshold);
 				}
 
 				// #BGTODO - non-updating parameters - remove lin/ang drag arrays and always query material if this stays a material parameter
@@ -1066,8 +1072,9 @@ float FGeometryCollectionPhysicsProxy::ComputeDamageThreshold(const FGeometryDyn
 	}
 	else
 	{
+		const int32 LevelOffset =  Parameters.bUsePerClusterOnlyDamageThreshold? 0: -1;
 		const int32 NumThresholds = Parameters.DamageThreshold.Num();
-		const int32 Level = FMath::Clamp(CalculateHierarchyLevel(DynamicCollection, TransformIndex), 0, INT_MAX);
+		const int32 Level = FMath::Clamp(CalculateHierarchyLevel(DynamicCollection, TransformIndex) + LevelOffset, 0, INT_MAX);
 		const float DefaultDamage = NumThresholds > 0 ? Parameters.DamageThreshold[NumThresholds - 1] : 0.f;
 		DamageThreshold = Level < NumThresholds ? Parameters.DamageThreshold[Level] : DefaultDamage;
 
@@ -1204,9 +1211,12 @@ FGeometryCollectionPhysicsProxy::BuildClusters_Internal(
 	{
 		// set the damage threshold on children as they are the one where the strain is tested when breaking 
 		Chaos::FPBDRigidParticleHandle* Child = ChildHandles[Idx];
-		if (Chaos::FPBDRigidClusteredParticleHandle* ClusteredChild = Child->CastToClustered())
+		if (Parameters.bUsePerClusterOnlyDamageThreshold)
 		{
-			ClusteredChild->SetStrains(DamageThreshold);
+			if (Chaos::FPBDRigidClusteredParticleHandle* ClusteredChild = Child->CastToClustered())
+			{
+				ClusteredChild->SetStrains(DamageThreshold);
+			}
 		}
 
 		const int32 ChildTransformGroupIndex = ChildTransformGroupIndices[Idx];
diff --git a/Engine/Source/Runtime/Experimental/Chaos/Public/GeometryCollection/GeometryCollectionSimulationCoreTypes.h b/Engine/Source/Runtime/Experimental/Chaos/Public/GeometryCollection/GeometryCollectionSimulationCoreTypes.h
--- a/Engine/Source/Runtime/Experimental/Chaos/Public/GeometryCollection/GeometryCollectionSimulationCoreTypes.h
+++ b/Engine/Source/Runtime/Experimental/Chaos/Public/GeometryCollection/GeometryCollectionSimulationCoreTypes.h
@@ -176,6 +176,7 @@ struct FSimulationParameters
 		, MaxClusterLevel(100)
 		, bUseSizeSpecificDamageThresholds(false)
 		, DamageThreshold({500000.f, 50000.f, 5000.f})
+		, bUsePerClusterOnlyDamageThreshold(false)
 		, ClusterConnectionMethod(Chaos::FClusterCreationParameters::EConnectionMethod::PointImplicit)
 		, CollisionGroup(0)
 		, CollisionSampleFraction(1.0)
@@ -210,6 +211,7 @@ struct FSimulationParameters
 		, MaxClusterLevel(Other.MaxClusterLevel)
 		, bUseSizeSpecificDamageThresholds(Other.bUseSizeSpecificDamageThresholds)
 		, DamageThreshold(Other.DamageThreshold)
+		, bUsePerClusterOnlyDamageThreshold(Other.bUsePerClusterOnlyDamageThreshold)
 		, ClusterConnectionMethod(Other.ClusterConnectionMethod)
 		, CollisionGroup(Other.CollisionGroup)
 		, CollisionSampleFraction(Other.CollisionSampleFraction)
@@ -260,6 +262,7 @@ struct FSimulationParameters
 	int32 MaxClusterLevel;
 	bool bUseSizeSpecificDamageThresholds;
 	TArray<float> DamageThreshold;
+	bool bUsePerClusterOnlyDamageThreshold;
 	Chaos::FClusterCreationParameters::EConnectionMethod ClusterConnectionMethod;
 
 	int32 CollisionGroup;
diff --git a/Engine/Source/Runtime/Experimental/GeometryCollectionEngine/Private/GeometryCollection/GeometryCollectionComponent.cpp b/Engine/Source/Runtime/Experimental/GeometryCollectionEngine/Private/GeometryCollection/GeometryCollectionComponent.cpp
--- a/Engine/Source/Runtime/Experimental/GeometryCollectionEngine/Private/GeometryCollection/GeometryCollectionComponent.cpp
+++ b/Engine/Source/Runtime/Experimental/GeometryCollectionEngine/Private/GeometryCollection/GeometryCollectionComponent.cpp
@@ -2208,6 +2208,7 @@ void UGeometryCollectionComponent::RegisterAndInitializePhysicsProxy()
 		SimulationParameters.MaxClusterLevel = MaxClusterLevel;
 		SimulationParameters.bUseSizeSpecificDamageThresholds = bUseSizeSpecificDamageThreshold;
 		SimulationParameters.DamageThreshold = DamageThreshold;
+		SimulationParameters.bUsePerClusterOnlyDamageThreshold = RestCollection? RestCollection->PerClusterOnlyDamageThreshold: false; 
 		SimulationParameters.ClusterConnectionMethod = (Chaos::FClusterCreationParameters::EConnectionMethod)ClusterCollectionType;
 		SimulationParameters.CollisionGroup = CollisionGroup;
 		SimulationParameters.CollisionSampleFraction = CollisionSampleFraction;
diff --git a/Engine/Source/Runtime/Experimental/GeometryCollectionEngine/Private/GeometryCollection/GeometryCollectionObject.cpp b/Engine/Source/Runtime/Experimental/GeometryCollectionEngine/Private/GeometryCollection/GeometryCollectionObject.cpp
--- a/Engine/Source/Runtime/Experimental/GeometryCollectionEngine/Private/GeometryCollection/GeometryCollectionObject.cpp
+++ b/Engine/Source/Runtime/Experimental/GeometryCollectionEngine/Private/GeometryCollection/GeometryCollectionObject.cpp
@@ -73,6 +73,7 @@ UGeometryCollection::UGeometryCollection(const FObjectInitializer& ObjectInitial
 	, ClusterGroupIndex(0)
 	, MaxClusterLevel(100)
 	, DamageThreshold({ 500000.f, 50000.f, 5000.f })
+	, PerClusterOnlyDamageThreshold(false)
 	, ClusterConnectionType(EClusterConnectionTypeEnum::Chaos_MinimalSpanningSubsetDelaunayTriangulation)
 	, bUseFullPrecisionUVs(false)
 	, bStripOnCook(false)
@@ -812,6 +813,12 @@ void UGeometryCollection::Serialize(FArchive& Ar)
 #endif
 	}
 
+	if (Ar.CustomVer(FUE5MainStreamObjectVersion::GUID) < FUE5MainStreamObjectVersion::GeometryCollectionPerChildDamageThreshold)
+	{
+		// prior this version, damage threshold were computed per cluster and propagated to children
+		PerClusterOnlyDamageThreshold = true; 
+	}
+
 #if WITH_EDITORONLY_DATA
 	if (bCreateSimulationData)
 	{
diff --git a/Engine/Source/Runtime/Experimental/GeometryCollectionEngine/Public/GeometryCollection/GeometryCollectionObject.h b/Engine/Source/Runtime/Experimental/GeometryCollectionEngine/Public/GeometryCollection/GeometryCollectionObject.h
--- a/Engine/Source/Runtime/Experimental/GeometryCollectionEngine/Public/GeometryCollection/GeometryCollectionObject.h
+++ b/Engine/Source/Runtime/Experimental/GeometryCollectionEngine/Public/GeometryCollection/GeometryCollectionObject.h
@@ -316,6 +316,9 @@ public:
 	/** End UObject Interface */
 
 	void Serialize(FArchive& Ar);
+#if WITH_EDITORONLY_DATA
+	void PostSerialize(const FArchive& Ar);
+#endif
 
 #if WITH_EDITOR
 	void EnsureDataIsCooked(bool bInitResources = true);
@@ -427,6 +430,12 @@ public:
 	UPROPERTY(EditAnywhere, Category = "Clustering")
 	TArray<float> DamageThreshold;
 
+	/** compatibility check, when true, only cluster compute damage from parameters and propagate to direct children
+	 *  when false, each child will compute it's damage threshold allowing for more precise and intuitive destruction behavior
+	 */
+	UPROPERTY(EditAnywhere, Category = "Clustering")
+	bool PerClusterOnlyDamageThreshold;
+
 	/** */
 	UPROPERTY(EditAnywhere, Category = "Clustering")
 	EClusterConnectionTypeEnum ClusterConnectionType;
