InterchangeFramework - Fix infinite recursive loop in FNodeDependencyCache when there are circular node dependencies.

#prefligh 626b3c3e584430685931568a

diff --git a/Engine/Source/Runtime/Experimental/Interchange/Engine/Private/Tasks/InterchangeTaskParsing.cpp b/Engine/Source/Runtime/Experimental/Interchange/Engine/Private/Tasks/InterchangeTaskParsing.cpp
--- a/Engine/Source/Runtime/Experimental/Interchange/Engine/Private/Tasks/InterchangeTaskParsing.cpp
+++ b/Engine/Source/Runtime/Experimental/Interchange/Engine/Private/Tasks/InterchangeTaskParsing.cpp
@@ -34,6 +34,14 @@
 struct FNodeDependencyCache
 {
 	const TSet<FString>& GetAccumulatedDependencies(const UInterchangeBaseNodeContainer* NodeContainer, const FString& NodeID)
+	{
+		TSet<FString> NodeStack;
+		return GetAccumulatedDependencies(NodeContainer, NodeID, NodeStack);
+	}
+
+private:
+
+	const TSet<FString>& GetAccumulatedDependencies(const UInterchangeBaseNodeContainer* NodeContainer, const FString& NodeID, TSet<FString>& NodeStack)
 	{
 		if (const TSet<FString>* DependenciesPtr = CachedDependencies.Find(NodeID))
 		{
@@ -41,13 +49,11 @@ struct FNodeDependencyCache
 		}
 
 		TSet<FString> Dependencies;
-		AccumulateDependencies(NodeContainer, NodeID, Dependencies);
+		AccumulateDependencies(NodeContainer, NodeID, Dependencies, NodeStack);
 		return CachedDependencies.Add(NodeID, MoveTemp(Dependencies));
 	}
 
-private:
-
-	void AccumulateDependencies(const UInterchangeBaseNodeContainer* NodeContainer, const FString& NodeID, TSet<FString>& OutDependenciesSet)
+	void AccumulateDependencies(const UInterchangeBaseNodeContainer* NodeContainer, const FString& NodeID, TSet<FString>& OutDependenciesSet, TSet<FString>& NodeStack)
 	{
 		const UInterchangeFactoryBaseNode* FactoryNode = NodeContainer->GetFactoryNode(NodeID);
 		if (!FactoryNode)
@@ -55,17 +61,24 @@ private:
 			return;
 		}
 
-		TArray<FString> FactoryDependencies;
-		FactoryNode->GetFactoryDependencies(FactoryDependencies);
-		for (const FString& DependencyID : FactoryDependencies)
+		bool bAlreadyInSet = false;
+		NodeStack.Add(NodeID, &bAlreadyInSet);
+		if (ensureMsgf(!bAlreadyInSet, TEXT("FNodeDependencyCache::AccumulateDependencies - Node \"%s\" is in a circular dependency, assets may not be imported properly."), *NodeID))
 		{
-			bool bAlreadyInSet = false;
-			OutDependenciesSet.Add(DependencyID, &bAlreadyInSet);
-			// Avoid infinite recursion.
-			if (!bAlreadyInSet)
+			TArray<FString> FactoryDependencies;
+			FactoryNode->GetFactoryDependencies(FactoryDependencies);
+			OutDependenciesSet.Reserve(OutDependenciesSet.Num() + FactoryDependencies.Num());
+			for (const FString& DependencyID : FactoryDependencies)
 			{
-				OutDependenciesSet.Append(GetAccumulatedDependencies(NodeContainer, DependencyID));
+				bAlreadyInSet = false;
+				OutDependenciesSet.Add(DependencyID, &bAlreadyInSet);
+				// Avoid infinite recursion.
+				if (!bAlreadyInSet)
+				{
+					OutDependenciesSet.Append(GetAccumulatedDependencies(NodeContainer, DependencyID, NodeStack));
+				}
 			}
+			NodeStack.Remove(NodeID);
 		}
 	}
 
