Chaos - joint stabilization
- add Inertia Conditioning system to increase inertia for long, thin, or small objects (enabled by default)
- also considiers joints with connectors outside the shape
- move some particle settings into a bitfield to reduce size
- add FVec3f as shorthand for TVec3-FRealSingle

diff --git a/Engine/Source/Runtime/Core/Public/UObject/UE5MainStreamObjectVersion.h b/Engine/Source/Runtime/Core/Public/UObject/UE5MainStreamObjectVersion.h
--- a/Engine/Source/Runtime/Core/Public/UObject/UE5MainStreamObjectVersion.h
+++ b/Engine/Source/Runtime/Core/Public/UObject/UE5MainStreamObjectVersion.h
@@ -218,6 +218,9 @@ struct CORE_API FUE5MainStreamObjectVersion
 		// Fix material feature level nodes to account for new SM6 input pin.
 		GeometryCollectionPerChildDamageThreshold,
 
+		// Move some Chaos flags into a bitfield
+		AddRigidParticleControlFlags,
+
 		// -----<new versions can be added above this line>-------------------------------------------------
 		VersionPlusOne,
 		LatestVersion = VersionPlusOne - 1
diff --git a/Engine/Source/Runtime/Engine/Classes/PhysicsEngine/BodyInstance.h b/Engine/Source/Runtime/Engine/Classes/PhysicsEngine/BodyInstance.h
--- a/Engine/Source/Runtime/Engine/Classes/PhysicsEngine/BodyInstance.h
+++ b/Engine/Source/Runtime/Engine/Classes/PhysicsEngine/BodyInstance.h
@@ -418,6 +418,24 @@ protected:
 	/** Whether we are pending a collision profile setup */
 	uint8 bPendingCollisionProfileSetup : 1;
 
+	/** 
+	 * @brief Enable automatic inertia conditioning to stabilize constraints.
+	 * 
+	 * Inertia conitioning increases inertia when an object is long and thin and also when it has joints that are outside the
+	 * collision shapes of the body. Increasing the inertia reduces the amount of rotation applied at joints which helps stabilize
+	 * joint chains, especially when bodies are small. In principle you can get the same behaviour by setting the InertiaTensorScale
+	 * appropriately, but this takes some of the guesswork out of it.
+	 * 
+	 * @note This only changes the inertia used in the low-level solver. That inertia is not visible to the BodyInstance
+	 * which will still report the inertia calculated from the mass, shapes, and InertiaTensorScale.
+	 * 
+	 * @note When enabled, the effective inertia depends on the joints attached to the body so the inertia will change when
+	 * joints are added or removed (automatically - no user action required).
+	 */
+	UPROPERTY(EditAnywhere, AdvancedDisplay, BlueprintReadWrite, Category = Physics)
+	uint8 bInertiaConditioning : 1;
+
+
 public:
 	/** Current scale of physics - used to know when and how physics must be rescaled to match current transform of OwnerComponent. */
 	FVector Scale3D;
@@ -744,6 +762,12 @@ public:
 	/** Return the body's inertia tensor. This is returned in local mass space */
 	FVector GetBodyInertiaTensor() const;
 
+	/** Whether inertia conditioning is enabled. @see bInertiaConditioning */
+	bool IsInertiaConditioningEnabled() { return bInertiaConditioning; }
+
+	/** Enable or disable inertia conditionin.  @see bInertiaConditioning */
+	void SetInertiaConditioningEnabled(bool bEnabled);
+
 
 	/** Set this body to be fixed (kinematic) or not. */
 	void SetInstanceSimulatePhysics(bool bSimulate, bool bMaintainPhysicsBlending=false);
diff --git a/Engine/Source/Runtime/Engine/Private/PhysicsEngine/BodyInstance.cpp b/Engine/Source/Runtime/Engine/Private/PhysicsEngine/BodyInstance.cpp
--- a/Engine/Source/Runtime/Engine/Private/PhysicsEngine/BodyInstance.cpp
+++ b/Engine/Source/Runtime/Engine/Private/PhysicsEngine/BodyInstance.cpp
@@ -352,6 +352,7 @@ FBodyInstance::FBodyInstance()
 	, bOverrideWalkableSlopeOnInstance(false)
 	, bInterpolateWhenSubStepping(true)
 	, bPendingCollisionProfileSetup(false)
+	, bInertiaConditioning(true)
 	, Scale3D(1.0f)
 	, CollisionProfileName(UCollisionProfile::CustomCollisionProfileName)
 	, PositionSolverIterationCount(8)
@@ -1220,7 +1221,8 @@ void FInitBodiesHelperBase::CreateActor_AssumesLocked(FBodyInstance* Instance, c
 		FPhysicsInterface::SetIsKinematic_AssumesLocked(Instance->ActorHandle, !Instance->ShouldInstanceSimulatingPhysics());
 #if WITH_CHAOS
 		FPhysicsInterface::SetMaxLinearVelocity_AssumesLocked(Instance->ActorHandle, TNumericLimits<float>::Max());
-		FPhysicsInterface::SetSmoothEdgeCollisionsEnabled(Instance->ActorHandle, Instance->bSmoothEdgeCollisions);
+		FPhysicsInterface::SetSmoothEdgeCollisionsEnabled_AssumesLocked(Instance->ActorHandle, Instance->bSmoothEdgeCollisions);
+		FPhysicsInterface::SetInertiaConditioningEnabled_AssumesLocked(Instance->ActorHandle, Instance->bInertiaConditioning);
 #endif
 
 		// Set sleep even notification
@@ -2895,6 +2897,22 @@ FVector FBodyInstance::GetBodyInertiaTensor() const
 	return OutTensor;
 }
 
+void FBodyInstance::SetInertiaConditioningEnabled(bool bInEnabled)
+{
+	if (bInEnabled != bInertiaConditioning)
+	{
+		bInertiaConditioning = bInEnabled;
+		FPhysicsCommand::ExecuteWrite(ActorHandle,
+			[bInEnabled](const FPhysicsActorHandle& Actor)
+			{
+				if (FPhysicsInterface::IsRigidBody(Actor))
+				{
+					FPhysicsInterface::SetInertiaConditioningEnabled_AssumesLocked(Actor, bInEnabled);
+				}
+			});
+	}
+}
+
 FBox FBodyInstance::GetBodyBounds() const
 {
 	FBox OutBox(EForceInit::ForceInitToZero);
@@ -3759,7 +3777,7 @@ void FBodyInstance::SetSmoothEdgeCollisionsEnabled(bool bNewSmoothEdgeCollisions
 			{
 				if (FPhysicsInterface::IsRigidBody(Actor))
 				{
-					FPhysicsInterface::SetSmoothEdgeCollisionsEnabled(Actor, bNewSmoothEdgeCollisions);
+					FPhysicsInterface::SetSmoothEdgeCollisionsEnabled_AssumesLocked(Actor, bNewSmoothEdgeCollisions);
 				}
 			});
 #endif
diff --git a/Engine/Source/Runtime/Engine/Private/PhysicsEngine/ImmediatePhysics/ImmediatePhysicsChaos/ImmediatePhysicsActorHandle_Chaos.cpp b/Engine/Source/Runtime/Engine/Private/PhysicsEngine/ImmediatePhysics/ImmediatePhysicsChaos/ImmediatePhysicsActorHandle_Chaos.cpp
--- a/Engine/Source/Runtime/Engine/Private/PhysicsEngine/ImmediatePhysics/ImmediatePhysicsChaos/ImmediatePhysicsActorHandle_Chaos.cpp
+++ b/Engine/Source/Runtime/Engine/Private/PhysicsEngine/ImmediatePhysics/ImmediatePhysicsChaos/ImmediatePhysicsActorHandle_Chaos.cpp
@@ -296,6 +296,7 @@ namespace ImmediatePhysics_Chaos
 					Dynamic->SetRotationOfMass(CoMTransform.GetRotation());
 					Dynamic->SetI(TVec3<FRealSingle>( Inertia.X, Inertia.Y, Inertia.Z ));
 					Dynamic->SetInvI(TVec3<FRealSingle>(InertiaInv.X, InertiaInv.Y, InertiaInv.Z ));
+					Dynamic->SetInertiaConditioningEnabled(BodyInstance->IsInertiaConditioningEnabled());
 					if (BodyInstance != nullptr)
 					{
 						Dynamic->SetLinearEtherDrag(BodyInstance->LinearDamping);
@@ -636,6 +637,11 @@ namespace ImmediatePhysics_Chaos
 			}
 			Dynamic->SetI(TVec3<FRealSingle>(NewInertia.X, NewInertia.Y, NewInertia.Z ));
 			Dynamic->SetInvI(TVec3<FRealSingle>(NewInverseInertia.X, NewInverseInertia.Y, NewInverseInertia.Z ));
+			
+			if (Dynamic->InertiaConditioningEnabled())
+			{
+				Dynamic->SetInertiaConditioningDirty();
+			}
 		}
 	}
 
diff --git a/Engine/Source/Runtime/Engine/Private/PhysicsEngine/ImmediatePhysics/ImmediatePhysicsChaos/ImmediatePhysicsJointHandle_Chaos.cpp b/Engine/Source/Runtime/Engine/Private/PhysicsEngine/ImmediatePhysics/ImmediatePhysicsChaos/ImmediatePhysicsJointHandle_Chaos.cpp
--- a/Engine/Source/Runtime/Engine/Private/PhysicsEngine/ImmediatePhysics/ImmediatePhysicsChaos/ImmediatePhysicsJointHandle_Chaos.cpp
+++ b/Engine/Source/Runtime/Engine/Private/PhysicsEngine/ImmediatePhysics/ImmediatePhysicsChaos/ImmediatePhysicsJointHandle_Chaos.cpp
@@ -185,6 +185,15 @@ namespace ImmediatePhysics_Chaos
 		ConstraintSettings.Sanitize();
 
 		ConstraintHandle = Constraints->AddConstraint({ Actor1->ParticleHandle, Actor2->ParticleHandle }, ConstraintSettings);
+
+		if (Actor1->ParticleHandle != nullptr)
+		{
+			FGenericParticleHandle(Actor1->ParticleHandle)->SetInertiaConditioningDirty();
+		}
+		if (Actor2->ParticleHandle != nullptr)
+		{
+			FGenericParticleHandle(Actor2->ParticleHandle)->SetInertiaConditioningDirty();
+		}
 	}
 
 	FJointHandle::~FJointHandle()
diff --git a/Engine/Source/Runtime/Engine/Private/PhysicsEngine/ImmediatePhysics/ImmediatePhysicsChaos/ImmediatePhysicsSimulation_Chaos.cpp b/Engine/Source/Runtime/Engine/Private/PhysicsEngine/ImmediatePhysics/ImmediatePhysicsChaos/ImmediatePhysicsSimulation_Chaos.cpp
--- a/Engine/Source/Runtime/Engine/Private/PhysicsEngine/ImmediatePhysics/ImmediatePhysicsChaos/ImmediatePhysicsSimulation_Chaos.cpp
+++ b/Engine/Source/Runtime/Engine/Private/PhysicsEngine/ImmediatePhysics/ImmediatePhysicsChaos/ImmediatePhysicsSimulation_Chaos.cpp
@@ -16,6 +16,7 @@
 #include "Chaos/DebugDrawQueue.h"
 #include "Chaos/Evolution/PBDMinEvolution.h"
 #include "Chaos/Joint/ChaosJointLog.h"
+#include "Chaos/MassConditioning.h"
 #include "Chaos/ParticleHandle.h"
 #include "Chaos/PBDCollisionConstraints.h"
 #include "Chaos/PBDConstraintRule.h"
@@ -144,6 +145,17 @@ FAutoConsoleVariableRef CVarChaosImmPhysJointAngularDriveDamping(TEXT("p.Chaos.I
 FAutoConsoleVariableRef CVarChaosImmPhysJointMinParentMassRatio(TEXT("p.Chaos.ImmPhys.Joint.MinParentMassRatio"), ChaosImmediate_Joint_MinParentMassRatio, TEXT("6Dof joint MinParentMassRatio (if > 0)"));
 FAutoConsoleVariableRef CVarChaosImmPhysJointMaxInertiaRatio(TEXT("p.Chaos.ImmPhys.Joint.MaxInertiaRatio"), ChaosImmediate_Joint_MaxInertiaRatio, TEXT("6Dof joint MaxInertiaRatio (if > 0)"));
 
+bool bChaosImmediate_InertiaConditioningEnabled = true;
+FAutoConsoleVariableRef  CVarChaosImmPhysParticleInertiaConditioningEnabled(TEXT("p.Chaos.ImmPhys.InertiaConditioning.Enabled"), bChaosImmediate_InertiaConditioningEnabled, TEXT("Enable/Disable constraint stabilization through inertia conditioning"));
+
+Chaos::FRealSingle ChaosImmediate_InertiaConditioningDistance = 10;
+FAutoConsoleVariableRef  CVarChaosImmPhysParticleInertiaConditioningDistance(TEXT("p.Chaos.ImmPhys.InertiaConditioning.Distance"), ChaosImmediate_InertiaConditioningDistance, TEXT(""));
+
+Chaos::FRealSingle ChaosImmediate_InertiaConditioningRotationRatio = 2;
+FAutoConsoleVariableRef  CVarChaosImmPhysParticleInertiaConditioningRotationRatio(TEXT("p.Chaos.ImmPhys.InertiaConditioning.RotationRatio"), ChaosImmediate_InertiaConditioningRotationRatio, TEXT(""));
+
+Chaos::FRealSingle ChaosImmediate_MaxInvInertiaComponentRatio = 0;
+FAutoConsoleVariableRef  CVarChaosImmPhysInertiaConditioningMaxInvInertiaComponentRatio(TEXT("p.Chaos.ImmPhys.InertiaConditioning.MaxInvInertiaComponentRatio"), ChaosImmediate_MaxInvInertiaComponentRatio, TEXT(""));
 
 //
 // Select the solver technique to use until we settle on the final one...
@@ -238,6 +250,7 @@ FAutoConsoleVariableRef CVarChaosImmPhysVelScale(TEXT("p.Chaos.ImmPhys.DebugDraw
 FAutoConsoleVariableRef CVarChaosImmPhysAngVelScale(TEXT("p.Chaos.ImmPhys.DebugDraw.AngVelScale"), ChaosImmPhysDebugDebugDrawSettings.AngVelScale, TEXT("If >0 show angular velocity when drawing particle transforms."));
 FAutoConsoleVariableRef CVarChaosImmPhysImpulseScale(TEXT("p.Chaos.ImmPhys.DebugDraw.ImpulseScale"), ChaosImmPhysDebugDebugDrawSettings.ImpulseScale, TEXT("If >0 show impulses when drawing collisions."));
 FAutoConsoleVariableRef CVarChaosImmPhysPushOutScale(TEXT("p.Chaos.ImmPhys.DebugDraw.PushOutScale"), ChaosImmPhysDebugDebugDrawSettings.PushOutScale, TEXT("If >0 show pushouts when drawing collisions."));
+FAutoConsoleVariableRef CVarChaosImmPhysInertiaScale(TEXT("p.Chaos.ImmPhys.DebugDraw.InertiaScale"), ChaosImmPhysDebugDebugDrawSettings.InertiaScale, TEXT("If >0 show inertia when drawing particles."));
 FAutoConsoleVariableRef CVarChaosImmPhysScale(TEXT("p.Chaos.ImmPhys.DebugDraw.Scale"), ChaosImmPhysDebugDebugDrawSettings.DrawScale, TEXT("Scale applied to all Chaos Debug Draw line lengths etc."));
 #endif
 
@@ -956,6 +969,8 @@ namespace ImmediatePhysics_Chaos
 			Implementation->bActorsDirty = false;
 		}
 
+		UpdateInertiaConditioning(InGravity);
+
 		UE_LOG(LogChaosJoint, Verbose, TEXT("Simulate Dt = %f Steps %d x %f (Rewind %f)"), DeltaTime, NumSteps, StepTime, RewindTime);
 		Implementation->Evolution.SetGravity(InGravity);
 		Implementation->Evolution.SetSimulationSpace(Implementation->SimulationSpace);
@@ -969,6 +984,64 @@ namespace ImmediatePhysics_Chaos
 #endif
 	}
 
+	void FSimulation::UpdateInertiaConditioning(const FVector& Gravity)
+	{
+		using namespace Chaos;
+
+		if (bChaosImmediate_InertiaConditioningEnabled)
+		{
+			// The maximum contribution to error correction from rotation
+			const FRealSingle MaxRotationRatio = ChaosImmediate_InertiaConditioningRotationRatio;
+
+			// The error distance that the constraint correction must be stable for
+			// @todo(chaos): should probably be tied to constraint teleport threshold?
+			const FRealSingle MaxDistance = ChaosImmediate_InertiaConditioningDistance;
+
+			// A limit on the relative sizes of the inertia components (inverse)
+			const FRealSingle MaxInvInertiaComponentRatio = ChaosImmediate_MaxInvInertiaComponentRatio;
+
+			for (FActorHandle* Actor : Implementation->ActorHandles)
+			{
+				FGenericParticleHandle Particle = Actor->GetParticle();
+				if (Particle->InertiaConditioningDirty())
+				{
+					const bool bWantInertiaConditioning = Particle->IsDynamic() && Particle->InertiaConditioningEnabled();
+					if (bWantInertiaConditioning)
+					{
+						// Calculate the extents of all the constraints on the body (in CoM space)
+						// @chaos(todo): We need a way to iterate over joints attached to a specific actor
+						const FVec3 CollisionExtents = FReal(0.5) * Particle->LocalBounds().InverseTransformedAABB(Actor->GetLocalCoMTransform()).Extents();
+						FVec3 ConstraintExtents = FVec3(0);
+						for (const FJointHandle* Joint : Implementation->JointHandles)
+						{
+							const Chaos::TVec2<const FActorHandle*>& JointActors = Joint->GetActorHandles();
+							FVec3 JointArm = FVec3(0);
+							if (JointActors[0] == Actor)
+							{
+								JointArm = Actor->GetLocalCoMTransform().InverseTransformPositionNoScale(Joint->GetConstraint()->GetSettings().ConnectorTransforms[0].GetTranslation());
+							}
+							else if (JointActors[1] == Actor)
+							{
+								JointArm = Actor->GetLocalCoMTransform().InverseTransformPositionNoScale(Joint->GetConstraint()->GetSettings().ConnectorTransforms[1].GetTranslation());
+							}
+							ConstraintExtents = FVec3::Max(ConstraintExtents, JointArm.GetAbs());
+						}
+
+						const FVec3f CoMExtents = FVec3f::Max(FVec3f(CollisionExtents), FVec3f(ConstraintExtents));
+						const FVec3f ActorInertiaConditioning = CalculateInertiaConditioning(FRealSingle(Actor->GetInverseMass()), FVec3f(Actor->GetInverseInertia()), CoMExtents, MaxDistance, MaxRotationRatio, MaxInvInertiaComponentRatio);
+				
+						Particle->SetInvIConditioning(ActorInertiaConditioning);
+					}
+					else
+					{
+						Particle->SetInvIConditioning(FVec3f(1));
+					}
+					Particle->ClearInertiaConditioningDirty();
+				}
+			}
+		}
+	}
+
 	void FSimulation::DebugDrawStaticParticles()
 	{
 #if CHAOS_DEBUG_DRAW
diff --git a/Engine/Source/Runtime/Engine/Public/Physics/ImmediatePhysics/ImmediatePhysicsChaos/ImmediatePhysicsSimulation_Chaos.h b/Engine/Source/Runtime/Engine/Public/Physics/ImmediatePhysics/ImmediatePhysicsChaos/ImmediatePhysicsSimulation_Chaos.h
--- a/Engine/Source/Runtime/Engine/Public/Physics/ImmediatePhysics/ImmediatePhysicsChaos/ImmediatePhysicsSimulation_Chaos.h
+++ b/Engine/Source/Runtime/Engine/Public/Physics/ImmediatePhysics/ImmediatePhysicsChaos/ImmediatePhysicsSimulation_Chaos.h
@@ -99,6 +99,7 @@ namespace ImmediatePhysics_Chaos
 
 	private:
 		void RemoveFromCollidingPairs(FActorHandle* ActorHandle);
+		void UpdateInertiaConditioning(const FVector& Gravity);
 		void PackCollidingPairs();
 		void UpdateActivePotentiallyCollidingPairs();
 		FReal UpdateStepTime(const FReal DeltaTime, const FReal MaxStepTime);
diff --git a/Engine/Source/Runtime/Experimental/Chaos/Private/Chaos/ChaosDebugDraw.cpp b/Engine/Source/Runtime/Experimental/Chaos/Private/Chaos/ChaosDebugDraw.cpp
--- a/Engine/Source/Runtime/Experimental/Chaos/Private/Chaos/ChaosDebugDraw.cpp
+++ b/Engine/Source/Runtime/Experimental/Chaos/Private/Chaos/ChaosDebugDraw.cpp
@@ -823,8 +823,19 @@ namespace Chaos
 			{
 				if (const TPBDRigidParticleHandle<FReal, 3>* Rigid = InParticle->CastToRigidParticle())
 				{
-					const FVec3 EquivalentBoxSize = Utilities::BoxSizeFromInertia(Rigid->I(), Rigid->M());
-					FDebugDrawQueue::GetInstance().DrawDebugBox(PCOM, 0.5f * Settings.InertiaScale * EquivalentBoxSize, QCOM, FColor::Magenta, false, 0.0f, 0, Settings.LineThickness);
+					if (Rigid->InvM() != 0)
+					{
+						// Show the raw inertia in black
+						if (!FVec3::IsNearlyEqual(Rigid->InvIConditioning(), TVec3<FRealSingle>(1), UE_SMALL_NUMBER))
+						{
+							const FVec3 EquivalentBoxSize = Utilities::BoxSizeFromInertia(Rigid->I(), Rigid->M());
+							FDebugDrawQueue::GetInstance().DrawDebugBox(PCOM, 0.5f * Settings.InertiaScale * EquivalentBoxSize, QCOM, FColor::Black, false, 0.0f, 0, Settings.LineThickness);
+						}
+
+						// Show the inertia used by the solver in magenta
+						const FVec3 EquivalentConditionedBoxSize = Utilities::BoxSizeFromInertia(Rigid->ConditionedI(), Rigid->M());
+						FDebugDrawQueue::GetInstance().DrawDebugBox(PCOM, 0.5f * Settings.InertiaScale * EquivalentConditionedBoxSize, QCOM, FColor::Magenta, false, 0.0f, 0, Settings.LineThickness);
+					}
 				}
 			}
 		}
diff --git a/Engine/Source/Runtime/Experimental/Chaos/Private/Chaos/MassConditioning.cpp b/Engine/Source/Runtime/Experimental/Chaos/Private/Chaos/MassConditioning.cpp
new file mode 100644
--- /dev/null
+++ b/Engine/Source/Runtime/Experimental/Chaos/Private/Chaos/MassConditioning.cpp
@@ -0,0 +1,119 @@
+// Copyright Epic Games, Inc. All Rights Reserved.
+#include "Chaos/MassConditioning.h"
+
+#include "Chaos/ConstraintHandle.h"
+#include "Chaos/ParticleHandle.h"
+#include "Chaos/PBDJointConstraints.h"
+
+namespace Chaos
+{
+	FVec3f CalculateInertiaConditioning(const FRealSingle InvM, const FVec3f& InvI, const FVec3f& ConstraintExtents, const FRealSingle MaxDistance, const FRealSingle MaxRotationRatio, const FRealSingle MaxInvInertiaComponentRatio)
+	{
+		check(InvM > 0);
+
+		FVec3f InvInertiaScale = FVec3f(1);
+
+		// Sanity check for zero dimensions along at least one axis
+		// We should never get this unless there are asset errors
+		if (ConstraintExtents.GetMin() > UE_SMALL_NUMBER)
+		{
+			// The I,J elements of this matrix can be interpreted the ratio of the rotation correction versus the translation 
+			// correction for a constraint arm along the Ith axis and a correction along the Jth axis. It has zero diagonal 
+			// because you get no rotation when correcting an error in the same direction as the constraint arm.
+			//
+			// The math: assume we have a constraint with a constraint arm along the X axis at MaxArm.X. Then assume we need to
+			// correct an error, and consider that error to be along each of the principle axes. The corrections perpendicular 
+			// to the arm (Y or Z axis) have worst-case rotation versus linear correction. The ratio of rotation correction 
+			// to linear correction would be (InvI . R x N x R) / InvM,  where R is the X axis and N is the X, Y or Z axis.
+			// For a cube of uniform density with no attached joints, each of these non-zero elements will be the same: 6
+			const FVec3f ConstraintExtentsSq = ConstraintExtents * ConstraintExtents;
+			const FVec3f MassInertiaRatio = InvI / InvM;
+			const FVec3f IRxNxRX = ConstraintExtentsSq.X * FVec3f(0, MassInertiaRatio.Y, MassInertiaRatio.Z);
+			const FVec3f IRxNxRY = ConstraintExtentsSq.Y * FVec3f(MassInertiaRatio.X, 0, MassInertiaRatio.Z);
+			const FVec3f IRxNxRZ = ConstraintExtentsSq.Z * FVec3f(MassInertiaRatio.X, MassInertiaRatio.Y, 0);
+
+			// Now we need to calculate a scale to apply to each component of the inverse inertia so that each of the elements in
+			// the above matrix is greater than some value. E.g., we modify the inertia so that the "worst case" corrections
+			// will always have a rotation correction that supplies less that MaxRotationRatio of the total correction.
+			// We also reduce the maximum rotation ratio (i.e., increase inertia) for small objects with size less than MaxDistance
+			FRealSingle ScaledMaxRotationRatio = MaxRotationRatio;
+			if (MaxDistance > 0)
+			{
+				const FRealSingle MaxRotationRatioScale = ConstraintExtents.GetMax() / MaxDistance;
+				const FRealSingle ClampedMaxRotationRatioScale = FMath::Min(MaxRotationRatioScale, FRealSingle(1));
+				ScaledMaxRotationRatio = ClampedMaxRotationRatioScale * MaxRotationRatio;
+			}
+
+			const FVec3f RotationRatio = FVec3f::Max3(IRxNxRX, IRxNxRY, IRxNxRZ);
+			InvInertiaScale = (FVec3f(ScaledMaxRotationRatio) / RotationRatio);
+
+			// We don't ever shrink the inertia. I.e., inv inertia scale must be less than 1
+			InvInertiaScale = InvInertiaScale.ComponentwiseMin(FVec3(1));
+
+			// Ensure that all components are at most MaxInvInertiaComponentRatio times the smallest component
+			if (MaxInvInertiaComponentRatio > 1)
+			{
+				const FVec3f ScaledInvI = InvInertiaScale * InvI;
+				const FRealSingle MinInvInertia = ScaledInvI.GetMin();
+				const FRealSingle MaxInvInertia = MaxInvInertiaComponentRatio * MinInvInertia;
+				if (ScaledInvI.X > MaxInvInertia)
+				{
+					InvInertiaScale.X = MaxInvInertia / ScaledInvI.X;
+				}
+				if (ScaledInvI.Y > MaxInvInertia)
+				{
+					InvInertiaScale.Y = MaxInvInertia / ScaledInvI.Y;
+				}
+				if (ScaledInvI.Z > MaxInvInertia)
+				{
+					InvInertiaScale.Z = MaxInvInertia / ScaledInvI.Z;
+				}
+
+				InvInertiaScale = InvInertiaScale.ComponentwiseMin(FVec3(1));
+			}
+		}
+		
+		return InvInertiaScale;
+	}
+
+	// Get the extents of all constraints in centre-of-mass space
+	FVec3f CalculateCoMConstraintExtents(const FPBDRigidParticleHandle* Rigid)
+	{
+		const FRigidTransform3 CoMTransform = FRigidTransform3(Rigid->CenterOfMass(), Rigid->RotationOfMass());
+
+		// Calculate the largest constraint arm from collisions and joints.
+		// For collisions, we assume contacts at the max extents of the bounds
+		// For joints, we use the actual constraint arm of all attached joints
+		const FAABB3 CoMBounds = Rigid->LocalBounds().InverseTransformedAABB(CoMTransform);
+		const FVec3f ColisionExtents = FReal(0.5) * CoMBounds.Extents();
+
+		FVec3 ConstraintExtents = FVec3(0);
+		for (FConstraintHandle* Constraint : Rigid->ParticleConstraints())
+		{
+			FVec3 ConstraintArm = FVec3(0);
+			if (const FPBDJointConstraintHandle* Joint = Constraint->As<FPBDJointConstraintHandle>())
+			{
+				const TVector<FGeometryParticleHandle*, 2> JointParticles = Joint->GetConstrainedParticles();
+				check((Rigid == JointParticles[0]) || (Rigid == JointParticles[1]));
+				const int32 JointParticleIndex = (Rigid == JointParticles[0]) ? 0 : 1;
+				const FVec3& ActorJointArm = Joint->GetSettings().ConnectorTransforms[JointParticleIndex].GetTranslation();
+				ConstraintArm = CoMTransform.InverseTransformPosition(ActorJointArm);
+			}
+
+			ConstraintExtents = FVec3::Max(ConstraintExtents, ConstraintArm.GetAbs());
+		}
+
+		return FVec3::Max(ColisionExtents, ConstraintExtents);
+	}
+
+	FVec3f CalculateParticleInertiaConditioning(const FPBDRigidParticleHandle* Rigid, const FRealSingle MaxDistance, const FRealSingle MaxRotationRatio, const FRealSingle MaxInvInertiaComponentRatio)
+	{
+		if (Rigid->InvM() == 0)
+		{
+			return FVec3f(1);
+		}
+
+		const FVec3 ConstraintExtents = CalculateCoMConstraintExtents(Rigid);
+		return CalculateInertiaConditioning(FRealSingle(Rigid->InvM()), Rigid->InvI(), ConstraintExtents, MaxDistance, MaxRotationRatio, MaxInvInertiaComponentRatio);
+	}
+}
\ No newline at end of file
diff --git a/Engine/Source/Runtime/Experimental/Chaos/Private/Chaos/PBDRigidsEvolution.cpp b/Engine/Source/Runtime/Experimental/Chaos/Private/Chaos/PBDRigidsEvolution.cpp
--- a/Engine/Source/Runtime/Experimental/Chaos/Private/Chaos/PBDRigidsEvolution.cpp
+++ b/Engine/Source/Runtime/Experimental/Chaos/Private/Chaos/PBDRigidsEvolution.cpp
@@ -1023,6 +1023,7 @@ namespace Chaos
 
 		Particle->SetObjectStateLowLevel(ObjectState);
 
+		// Put the particle into the correct array. This will change the location of the particle data
 		if (auto ClusteredParticle = Particle->CastToClustered())
 		{
 			Particles.SetClusteredParticleSOA(ClusteredParticle);
@@ -1032,6 +1033,7 @@ namespace Chaos
 			Particles.SetDynamicParticleSOA(Particle);
 		}
 
+		// @todo(chaos): this doesn't seem to handle Sleeping -> Kinematic, or Kinematic -> Sleeping
 		if (InitialState != ObjectState && !Particle->Disabled())
 		{
 			if (InitialState == EObjectStateType::Sleeping)
@@ -1048,6 +1050,17 @@ namespace Chaos
 				Particles.MarkTransientDirtyParticle(Particle);
 			}
 		}
+
+		// If we switched to/from dynamic the inertia conditioning may need to be refreshed
+		if (Particle->InertiaConditioningEnabled())
+		{
+			const bool bWasDynamic = (InitialState == EObjectStateType::Dynamic) || (InitialState == EObjectStateType::Sleeping);
+			const bool bIsDynamic = (ObjectState == EObjectStateType::Dynamic) || (ObjectState == EObjectStateType::Sleeping);
+			if (bWasDynamic != bIsDynamic)
+			{
+				Particle->SetInertiaConditioningDirty();
+			}
+		}
 	}
 
 	void FPBDRigidsEvolutionBase::SetParticleSleepType(FPBDRigidParticleHandle* Particle, ESleepType InSleepType)
diff --git a/Engine/Source/Runtime/Experimental/Chaos/Private/Chaos/PBDRigidsEvolutionGBF.cpp b/Engine/Source/Runtime/Experimental/Chaos/Private/Chaos/PBDRigidsEvolutionGBF.cpp
--- a/Engine/Source/Runtime/Experimental/Chaos/Private/Chaos/PBDRigidsEvolutionGBF.cpp
+++ b/Engine/Source/Runtime/Experimental/Chaos/Private/Chaos/PBDRigidsEvolutionGBF.cpp
@@ -5,6 +5,7 @@
 #include "Chaos/Framework/Parallel.h"
 #include "Chaos/ImplicitObjectTransformed.h"
 #include "Chaos/ImplicitObjectUnion.h"
+#include "Chaos/MassConditioning.h"
 #include "Chaos/PBDCollisionConstraints.h"
 #include "Chaos/PBDCollisionSpringConstraints.h"
 #include "Chaos/PerParticleEtherDrag.h"
@@ -91,6 +92,22 @@ namespace Chaos
 		bool DoFinalProbeNarrowPhase = true;
 		FAutoConsoleVariableRef CVarDoFinalProbeNarrowPhase(TEXT("p.Chaos.Solver.DoFinalProbeNarrowPhase"), DoFinalProbeNarrowPhase, TEXT(""));
 
+		// Enable inertia conditioning for joint and collisions stabilization for small and thin objects
+		bool bChaosSolverInertiaConditioningEnabled = true;
+		FAutoConsoleVariableRef  CVarChaosSolverInertiaConditioningEnabled(TEXT("p.Chaos.Solver.InertiaConditioning.Enabled"), bChaosSolverInertiaConditioningEnabled, TEXT("Enable/Disable constraint stabilization through inertia conditioning"));
+
+		// The largest joint error we expect to resolve in a moderately stable way
+		FRealSingle ChaosSolverInertiaConditioningDistance = 20;
+		FAutoConsoleVariableRef  CVarChaosSolverInertiaConditioningDistance(TEXT("p.Chaos.Solver.InertiaConditioning.Distance"), ChaosSolverInertiaConditioningDistance, TEXT(""));
+
+		// The ratio of joint error correction that comes from particle rotation versus translation
+		FRealSingle ChaosSolverInertiaConditioningRotationRatio = 2;
+		FAutoConsoleVariableRef  CVarChaosSolverInertiaConditioningRotationRatio(TEXT("p.Chaos.Solver.InertiaConditioning.RotationRatio"), ChaosSolverInertiaConditioningRotationRatio, TEXT(""));
+
+		// If > 1, limits the inverse inertia components to be mo more than this multiple of the smallest component. Makes the objects more round. Set to 0 to disable.
+		FRealSingle ChaosSolverMaxInvInertiaComponentRatio = 0;
+		FAutoConsoleVariableRef  CVarChaosSolverInertiaConditioningMaxInvInertiaComponentRatio(TEXT("p.Chaos.Solver.InertiaConditioning.MaxInvInertiaComponentRatio"), ChaosSolverMaxInvInertiaComponentRatio, TEXT(""));
+
 		DECLARE_CYCLE_STAT(TEXT("FPBDRigidsEvolutionGBF::AdvanceOneTimeStep"), STAT_Evolution_AdvanceOneTimeStep, STATGROUP_Chaos);
 		DECLARE_CYCLE_STAT(TEXT("FPBDRigidsEvolutionGBF::UnclusterUnions"), STAT_Evolution_UnclusterUnions, STATGROUP_Chaos);
 		DECLARE_CYCLE_STAT(TEXT("FPBDRigidsEvolutionGBF::Integrate"), STAT_Evolution_Integrate, STATGROUP_Chaos);
@@ -117,6 +134,7 @@ namespace Chaos
 		DECLARE_CYCLE_STAT(TEXT("FPBDRigidsEvolutionGBF::ParallelSolve"), STAT_Evolution_ParallelSolve, STATGROUP_Chaos);
 		DECLARE_CYCLE_STAT(TEXT("FPBDRigidsEvolutionGBF::SaveParticlePostSolve"), STAT_Evolution_SavePostSolve, STATGROUP_Chaos);
 		DECLARE_CYCLE_STAT(TEXT("FPBDRigidsEvolutionGBF::DeactivateSleep"), STAT_Evolution_DeactivateSleep, STATGROUP_Chaos);
+		DECLARE_CYCLE_STAT(TEXT("FPBDRigidsEvolutionGBF::InertiaConditioning"), STAT_Evolution_InertiaConditioning, STATGROUP_Chaos);
 
 		int32 SerializeEvolution = 0;
 		FAutoConsoleVariableRef CVarSerializeEvolution(TEXT("p.SerializeEvolution"), SerializeEvolution, TEXT(""));
@@ -397,6 +415,11 @@ void FPBDRigidsEvolutionGBF::AdvanceOneTimeStepImpl(const FReal Dt, const FSubSt
 		CollisionConstraints.ApplyCollisionModifier(*CollisionModifiers, Dt);
 	}
 
+	{
+		SCOPE_CYCLE_COUNTER(STAT_Evolution_InertiaConditioning);
+		UpdateInertiaConditioning();
+	}
+
 	if (bChaosUseCCD)
 	{
 		SCOPE_CYCLE_COUNTER(STAT_Evolution_CCD);
@@ -999,5 +1022,37 @@ void FPBDRigidsEvolutionGBF::DestroyParticleCollisionsInAllocator(FGeometryParti
 	}
 }
 
+void FPBDRigidsEvolutionGBF::UpdateInertiaConditioning()
+{
+	// The maximum contribution to error correction from rotation
+	const FRealSingle MaxRotationRatio = ChaosSolverInertiaConditioningRotationRatio;
+
+	// The error distance that the constraint correction must be stable for
+	// @todo(chaos): should probably be tied to constraint teleport threshold?
+	const FRealSingle MaxDistance = ChaosSolverInertiaConditioningDistance;
+
+	// A limit on the relative sizes of the inertia components (inverse)
+	const FRealSingle MaxInvInertiaComponentRatio = ChaosSolverMaxInvInertiaComponentRatio;
+
+	// @todo(chaos): we should only visit particles that are dirty
+	for (auto& Rigid : Particles.GetNonDisabledDynamicView())
+	{
+		if (Rigid.InertiaConditioningDirty())
+		{
+			const bool bWantInertiaConditioning = Rigid.IsDynamic() && Rigid.InertiaConditioningEnabled() && bChaosSolverInertiaConditioningEnabled;
+			if (bWantInertiaConditioning)
+			{
+				const FVec3f InvInertiaScale = CalculateParticleInertiaConditioning(Rigid.Handle(), MaxDistance, MaxRotationRatio, MaxInvInertiaComponentRatio);
+				Rigid.SetInvIConditioning(InvInertiaScale);
+			}
+			else
+			{
+				Rigid.SetInvIConditioning(FVec3f(1));
+			}
+			Rigid.ClearInertiaConditioningDirty();
+		}
+	}
+}
+
 }
 
diff --git a/Engine/Source/Runtime/Experimental/Chaos/Private/Chaos/ParticleHandle.cpp b/Engine/Source/Runtime/Experimental/Chaos/Private/Chaos/ParticleHandle.cpp
--- a/Engine/Source/Runtime/Experimental/Chaos/Private/Chaos/ParticleHandle.cpp
+++ b/Engine/Source/Runtime/Experimental/Chaos/Private/Chaos/ParticleHandle.cpp
@@ -253,16 +253,20 @@ namespace Chaos
 			}
 		}		
 
+		const bool bDirtyInertiaConditioning = (ObjectState() != DynamicMisc.ObjectState());
+		if (bDirtyInertiaConditioning)
+		{
+			SetInertiaConditioningDirty();
+		}
+
 		SetLinearEtherDrag(DynamicMisc.LinearEtherDrag());
 		SetAngularEtherDrag(DynamicMisc.AngularEtherDrag());
 		SetMaxLinearSpeedSq(DynamicMisc.MaxLinearSpeedSq());
 		SetMaxAngularSpeedSq(DynamicMisc.MaxAngularSpeedSq());
 		SetCollisionGroup(DynamicMisc.CollisionGroup());
-		SetGravityEnabled(DynamicMisc.GravityEnabled());
-		SetCCDEnabled(DynamicMisc.CCDEnabled());
 		SetDisabled(DynamicMisc.Disabled());
-		SetOneWayInteraction(DynamicMisc.OneWayInteraction());
 		SetCollisionConstraintFlags(DynamicMisc.CollisionConstraintFlags());
+		SetControlFlags(DynamicMisc.ControlFlags());
 
 		Evolution.SetParticleObjectState(this, DynamicMisc.ObjectState());
 		Evolution.SetParticleSleepType(this, DynamicMisc.SleepType());
diff --git a/Engine/Source/Runtime/Experimental/Chaos/Private/Chaos/RigidParticleControlFlags.cpp b/Engine/Source/Runtime/Experimental/Chaos/Private/Chaos/RigidParticleControlFlags.cpp
new file mode 100644
--- /dev/null
+++ b/Engine/Source/Runtime/Experimental/Chaos/Private/Chaos/RigidParticleControlFlags.cpp
@@ -0,0 +1,20 @@
+// Copyright Epic Games, Inc. All Rights Reserved.
+#include "Chaos/RigidParticleControlFlags.h"
+
+#include "Chaos/ChaosArchive.h"
+
+#include "Containers/UnrealString.h"
+
+namespace Chaos
+{
+	FString FRigidParticleControlFlags::ToString() const
+	{
+		return FString::Printf(TEXT("0x%x"), Bits);
+	}
+
+	FChaosArchive& operator<<(FChaosArchive& Ar, FRigidParticleControlFlags& Flags)
+	{
+		Ar << Flags.Bits;
+		return Ar;
+	}
+}
\ No newline at end of file
diff --git a/Engine/Source/Runtime/Experimental/Chaos/Private/PhysicsProxy/JointConstraintProxy.cpp b/Engine/Source/Runtime/Experimental/Chaos/Private/PhysicsProxy/JointConstraintProxy.cpp
--- a/Engine/Source/Runtime/Experimental/Chaos/Private/PhysicsProxy/JointConstraintProxy.cpp
+++ b/Engine/Source/Runtime/Experimental/Chaos/Private/PhysicsProxy/JointConstraintProxy.cpp
@@ -93,6 +93,10 @@ void FJointConstraintPhysicsProxy::InitializeOnPhysicsThread(FPBDRigidsSolver* I
 
 					Handle0->AddConstraintHandle(Constraint_PT);
 					Handle1->AddConstraintHandle(Constraint_PT);
+
+					// We added a joint to the particles, so we need to modify the inertia (See FPBDRigidsEvolutionGBF::UpdateInertiaConditioning)
+					FGenericParticleHandle(Handle0)->SetInertiaConditioningDirty();
+					FGenericParticleHandle(Handle1)->SetInertiaConditioningDirty();
 				}
 			}
 		}
diff --git a/Engine/Source/Runtime/Experimental/Chaos/Public/Chaos/BoundingVolumeUtilities.h b/Engine/Source/Runtime/Experimental/Chaos/Public/Chaos/BoundingVolumeUtilities.h
--- a/Engine/Source/Runtime/Experimental/Chaos/Public/Chaos/BoundingVolumeUtilities.h
+++ b/Engine/Source/Runtime/Experimental/Chaos/Public/Chaos/BoundingVolumeUtilities.h
@@ -31,7 +31,7 @@ inline FVec3 ComputeBoundsThickness(FVec3 Vel, FReal Dt, FReal MinBoundsThicknes
 inline FVec3 ComputeBoundsThickness(const TPBDRigidParticles<FReal, 3>& InParticles, FReal Dt, int32 BodyIndex, FReal MinBoundsThickness, FReal BoundsVelocityInflation)
 {
 	const bool bIsBounded = InParticles.HasBounds(BodyIndex);
-	const bool bIsCCD = InParticles.CCDEnabled(BodyIndex);
+	const bool bIsCCD = InParticles.ControlFlags(BodyIndex).GetCCDEnabled();
 
 	if (!bIsCCD && (BoundsVelocityInflation == FReal(0)))
 	{
diff --git a/Engine/Source/Runtime/Experimental/Chaos/Public/Chaos/Evolution/SolverBodyContainer.h b/Engine/Source/Runtime/Experimental/Chaos/Public/Chaos/Evolution/SolverBodyContainer.h
--- a/Engine/Source/Runtime/Experimental/Chaos/Public/Chaos/Evolution/SolverBodyContainer.h
+++ b/Engine/Source/Runtime/Experimental/Chaos/Public/Chaos/Evolution/SolverBodyContainer.h
@@ -50,7 +50,7 @@ namespace Chaos
 					SolverBody.SetR(PrevCoMTransform.GetRotation());
 
 					SolverBody.SetInvM(Particle->InvM());
-					SolverBody.SetInvILocal(Particle->InvI());
+					SolverBody.SetInvILocal(Particle->ConditionedInvI());
 				}
 				else
 				{
diff --git a/Engine/Source/Runtime/Experimental/Chaos/Public/Chaos/MassConditioning.h b/Engine/Source/Runtime/Experimental/Chaos/Public/Chaos/MassConditioning.h
new file mode 100644
--- /dev/null
+++ b/Engine/Source/Runtime/Experimental/Chaos/Public/Chaos/MassConditioning.h
@@ -0,0 +1,35 @@
+// Copyright Epic Games, Inc. All Rights Reserved.
+#pragma once
+
+#include "Chaos/Core.h"
+#include "Chaos/ParticleHandleFwd.h"
+
+#include "Math/UnrealMathUtility.h"
+
+namespace Chaos
+{
+
+
+	/**
+	 * Calculate an inertia scale so that rotation corrections when applying a constraint are no larger than some fraction of the total correction.
+	 * It has the net effect of making the object "more round" and also increasing the inertia relative to the mass when the object is not of "regular" proportions.
+	 * 
+	 * @param ConstraintExtents the Maximum constraint arm distance along each axis. This should be the extents of the object, including all shapes and joint connectors.
+	 * @param MaxDistance the constraint error that we want to be stable. Corrections above this may still contain large rotation components.
+	 * @param MaxRotationRatio the contribution to the constraint correction from rotation of the object will be less that this fraction of the total error.
+	*/
+	FVec3f CHAOS_API CalculateInertiaConditioning(const FRealSingle InvM, const FVec3f& InvI, const FVec3f& ConstraintExtents, const FRealSingle MaxDistance, const FRealSingle MaxRotationRatio, const FRealSingle MaxInvInertiaComponentRatio);
+
+	/**
+	 * Generate a CoM-relative extent that encompasses all constraints (collision and joints for now) on the particle. For collisions, we approximate the max extend
+	 * using the CoM-space bounds. For joints we use the actual joint attachment positions.
+	*/
+	FVec3f CHAOS_API CalculateParticleConstraintExtents(const FPBDRigidParticleHandle* Rigid);
+
+	/**
+	 * Calculate an inertia scale for a particle based on object size and all joint connector offsets that should stabilize the constraints.
+	 * @see CalculateInertiaConditioning()
+	*/
+	FVec3f CHAOS_API CalculateParticleInertiaConditioning(const FPBDRigidParticleHandle* Rigid, const FRealSingle MaxDistance, const FRealSingle MaxRotationRatio, const FRealSingle MaxInvInertiaComponentRatio);
+
+}
\ No newline at end of file
diff --git a/Engine/Source/Runtime/Experimental/Chaos/Public/Chaos/PBDRigidsEvolutionGBF.h b/Engine/Source/Runtime/Experimental/Chaos/Public/Chaos/PBDRigidsEvolutionGBF.h
--- a/Engine/Source/Runtime/Experimental/Chaos/Public/Chaos/PBDRigidsEvolutionGBF.h
+++ b/Engine/Source/Runtime/Experimental/Chaos/Public/Chaos/PBDRigidsEvolutionGBF.h
@@ -390,6 +390,8 @@ namespace Chaos
 		void GatherSolverInput(FReal Dt, int32 GroupIndex);
 		void ScatterSolverOutput(FReal Dt, int32 GroupIndex);
 
+		void UpdateInertiaConditioning();
+		void UpdateParticleInertiaConditioning(FPBDRigidParticleHandle* Rigid, const FReal MaxDistance, const FReal MaxRotationRatio);
 
 		FEvolutionResimCache* GetCurrentStepResimCache()
 		{
diff --git a/Engine/Source/Runtime/Experimental/Chaos/Public/Chaos/ParticleDirtyFlags.h b/Engine/Source/Runtime/Experimental/Chaos/Public/Chaos/ParticleDirtyFlags.h
--- a/Engine/Source/Runtime/Experimental/Chaos/Public/Chaos/ParticleDirtyFlags.h
+++ b/Engine/Source/Runtime/Experimental/Chaos/Public/Chaos/ParticleDirtyFlags.h
@@ -11,6 +11,7 @@
 #include "Chaos/CollisionFilterData.h"
 #include "Chaos/Collision/CollisionConstraintFlags.h"
 #include "Chaos/KinematicTargets.h"
+#include "Chaos/RigidParticleControlFlags.h"
 #include "UObject/ExternalPhysicsCustomObjectVersion.h"
 #include "UObject/ExternalPhysicsMaterialCustomObjectVersion.h"
 #include "UObject/PhysicsObjectVersion.h"
@@ -304,23 +305,37 @@ public:
 		Ar.UsingCustomVersion(FExternalPhysicsCustomObjectVersion::GUID);
 		Ar.UsingCustomVersion(FPhysicsObjectVersion::GUID);
 
+		// Flags moved into a bitmask
+		const bool bAddControlFlags = (Ar.CustomVer(FUE5MainStreamObjectVersion::GUID) >= FUE5MainStreamObjectVersion::AddRigidParticleControlFlags);
+
 		Ar << MLinearEtherDrag;
 		Ar << MAngularEtherDrag;
 		Ar << MObjectState;
-		Ar << MGravityEnabled;
-		Ar << MSleepType;
-		if (Ar.CustomVer(FExternalPhysicsCustomObjectVersion::GUID) >= FExternalPhysicsCustomObjectVersion::AddOneWayInteraction)
+		if (!bAddControlFlags && Ar.IsLoading())
 		{
-			Ar << MOneWayInteraction;
+			bool bGravityEnabled;
+			Ar << bGravityEnabled;
+			MControlFlags.SetGravityEnabled(bGravityEnabled);
 		}
-		else
+		Ar << MSleepType;
+		if (!bAddControlFlags && Ar.IsLoading())
 		{
-			MOneWayInteraction = false;
+			bool bOneWayInteraction = false;
+			if (Ar.CustomVer(FExternalPhysicsCustomObjectVersion::GUID) >= FExternalPhysicsCustomObjectVersion::AddOneWayInteraction)
+			{
+				Ar << bOneWayInteraction;
+			}
+			MControlFlags.SetOneWayInteractionEnabled(bOneWayInteraction);
 		}
 
-		if (Ar.CustomVer(FPhysicsObjectVersion::GUID) >= FPhysicsObjectVersion::AddCCDEnableFlag)
+		if (!bAddControlFlags && Ar.IsLoading())
 		{
-			Ar << bCCDEnabled;
+			if (Ar.CustomVer(FPhysicsObjectVersion::GUID) >= FPhysicsObjectVersion::AddCCDEnableFlag)
+			{
+				bool bCCDEnabled;
+				Ar << bCCDEnabled;
+				MControlFlags.SetCCDEnabled(bCCDEnabled);
+			}
 		}
 
 		const bool bAddCollisionConstraintFlagUE4 = (Ar.CustomVer(FPhysicsObjectVersion::GUID) >= FPhysicsObjectVersion::AddCollisionConstraintFlag);
@@ -344,6 +359,11 @@ public:
 			Ar << MMaxLinearSpeedSq;
 			Ar << MMaxAngularSpeedSq;
 		}
+
+		if (bAddControlFlags)
+		{
+			Ar << MControlFlags;
+		}
 	}
 
 	template <typename TOther>
@@ -354,12 +374,10 @@ public:
 		SetMaxLinearSpeedSq(Other.MaxLinearSpeedSq());
 		SetMaxAngularSpeedSq(Other.MaxAngularSpeedSq());
 		SetObjectState(Other.ObjectState());
-		SetGravityEnabled(Other.GravityEnabled());
 		SetCollisionGroup(Other.CollisionGroup());
 		SetSleepType(Other.SleepType());
-		SetOneWayInteraction(Other.OneWayInteraction());
 		SetCollisionConstraintFlags(Other.CollisionConstraintFlags());
-		SetCCDEnabled(Other.CCDEnabled());
+		SetControlFlags(Other.ControlFlags());
 		SetDisabled(Other.Disabled());
 	}
 
@@ -371,12 +389,10 @@ public:
 			&& AngularEtherDrag() == Other.AngularEtherDrag()
 			&& MaxLinearSpeedSq() == Other.MaxLinearSpeedSq()
 			&& MaxAngularSpeedSq() == Other.MaxAngularSpeedSq()
-			&& GravityEnabled() == Other.GravityEnabled()
 			&& CollisionGroup() == Other.CollisionGroup()
 			&& SleepType() == Other.SleepType()
-			&& OneWayInteraction() == Other.OneWayInteraction() 
 			&& CollisionConstraintFlags() == Other.CollisionConstraintFlags()
-			&& CCDEnabled() == Other.CCDEnabled()
+			&& ControlFlags() == Other.ControlFlags()
 			&& Disabled() == Other.Disabled();
 	}
 
@@ -400,11 +416,11 @@ public:
 	EObjectStateType ObjectState() const { return MObjectState; }
 	void SetObjectState(EObjectStateType InState){ MObjectState = InState; }
 
-	bool GravityEnabled() const { return MGravityEnabled; }
-	void SetGravityEnabled(bool InGravity){ MGravityEnabled = InGravity; }
+	bool GravityEnabled() const { return MControlFlags.GetGravityEnabled(); }
+	void SetGravityEnabled(bool bInGravity){ MControlFlags.SetGravityEnabled(bInGravity); }
 
-	bool CCDEnabled() const { return bCCDEnabled; }
-	void SetCCDEnabled(bool bInCCDEnabled) { bCCDEnabled = bInCCDEnabled; }
+	bool CCDEnabled() const { return MControlFlags.GetCCDEnabled(); }
+	void SetCCDEnabled(bool bInCCDEnabled) { MControlFlags.SetCCDEnabled(bInCCDEnabled); }
 
 	bool Disabled() const { return bDisabled; }
 	void SetDisabled(bool bInDisabled) { bDisabled = bInDisabled; }
@@ -420,8 +436,14 @@ public:
 	void AddCollisionConstraintFlag(const ECollisionConstraintFlags Flag) { MCollisionConstraintFlag |= uint32(Flag); }
 	void RemoveCollisionConstraintFlag(const ECollisionConstraintFlags Flag) { MCollisionConstraintFlag &= ~uint32(Flag); }
 	
-	bool OneWayInteraction() const { return MOneWayInteraction; }
-	void SetOneWayInteraction(bool InOneWayInteraction) { MOneWayInteraction = InOneWayInteraction; }
+	bool OneWayInteraction() const { return MControlFlags.GetOneWayInteractionEnabled(); }
+	void SetOneWayInteraction(bool bInOneWayInteraction) { MControlFlags.SetOneWayInteractionEnabled(bInOneWayInteraction); }
+
+	bool InertiaConditioningEnabled() const { return MControlFlags.GetInertiaConditioningEnabled(); }
+	void SetInertiaConditioningEnabled(bool bInEnabled) { MControlFlags.SetInertiaConditioningEnabled(bInEnabled); }
+
+	FRigidParticleControlFlags ControlFlags() const { return MControlFlags; }
+	void SetControlFlags(const FRigidParticleControlFlags& InFlags) { MControlFlags = InFlags; }
 
 private:
 	//NOTE: MObjectState is the only sim-writable data in this struct
@@ -437,11 +459,9 @@ private:
 	EResimType MResimType;
 	ESleepType MSleepType;
 
-	bool MGravityEnabled;
-	bool MOneWayInteraction = false;
 	uint32 MCollisionConstraintFlag = 0;
+	FRigidParticleControlFlags MControlFlags;
 
-	bool bCCDEnabled;
 	bool bDisabled;
 };
 
diff --git a/Engine/Source/Runtime/Experimental/Chaos/Public/Chaos/ParticleHandle.h b/Engine/Source/Runtime/Experimental/Chaos/Public/Chaos/ParticleHandle.h
--- a/Engine/Source/Runtime/Experimental/Chaos/Public/Chaos/ParticleHandle.h
+++ b/Engine/Source/Runtime/Experimental/Chaos/Public/Chaos/ParticleHandle.h
@@ -850,6 +850,7 @@ protected:
 		SetConstraintGraphIndex(INDEX_NONE);
 		SetToBeRemovedOnFracture(false);
 		SetSleepType(ESleepType::MaterialSleep);
+		SetInvIConditioning(TVec3<FRealSingle>(1));
 	}
 public:
 
@@ -867,6 +868,9 @@ public:
 
 	operator TPBDRigidParticleHandleImp<T, d, false>& () { return reinterpret_cast<TPBDRigidParticleHandleImp<T, d, false>&>(*this); }
 
+	bool IsKinematic() const { return (ObjectState() == EObjectStateType::Kinematic) || (ObjectState() == EObjectStateType::Static); }
+	bool IsDynamic() const { return (ObjectState() == EObjectStateType::Dynamic) || (ObjectState() == EObjectStateType::Sleeping); }
+
 	const TUniquePtr<TBVHParticles<T, d>>& CollisionParticles() const { return PBDRigidParticles->CollisionParticles(ParticleIdx); }
 	TUniquePtr<TBVHParticles<T, d>>& CollisionParticles() { return PBDRigidParticles->CollisionParticles(ParticleIdx); }
 
@@ -958,6 +962,8 @@ public:
 		SetInvI(Props.InvI());
 		SetM(Props.M());
 		SetInvM(Props.InvM());
+
+		SetInertiaConditioningDirty();
 	}
 
 	void SetDynamicMisc(const FParticleDynamicMisc& DynamicMisc, FPBDRigidsEvolutionBase& Evolution);
@@ -968,10 +974,12 @@ public:
 		SetWSmooth(W());
 	}
 
+	// Get the raw inertia. @see ConditionedInvI()
 	const TVec3<FRealSingle>& I() const { return PBDRigidParticles->I(ParticleIdx); }
 	TVec3<FRealSingle>& I() { return PBDRigidParticles->I(ParticleIdx); }
 	void SetI(const TVec3<FRealSingle>& InI) { PBDRigidParticles->I(ParticleIdx) = InI; }
 
+	// Get the raw inverse inertia. @see ConditionedInvI()
 	const TVec3<FRealSingle>& InvI() const { return PBDRigidParticles->InvI(ParticleIdx); }
 	TVec3<FRealSingle>& InvI() { return PBDRigidParticles->InvI(ParticleIdx); }
 	void SetInvI(const TVec3<FRealSingle>& InInvI) { PBDRigidParticles->InvI(ParticleIdx) = InInvI; }
@@ -990,6 +998,14 @@ public:
 	const TRotation<T,d>& RotationOfMass() const { return PBDRigidParticles->RotationOfMass(ParticleIdx); }
 	void SetRotationOfMass(const TRotation<T,d>& InRotationOfMass, bool bInvalidate = false) { PBDRigidParticles->RotationOfMass(ParticleIdx) = InRotationOfMass; }
 
+	// Get the inertia conditioning scales. This is a scale applied to the inverse inertia for use by the constraint solvers to improve stability
+	// and is calculated based on the attached joints and potential collision positions (approximated by object size)
+	const TVec3<FRealSingle>& InvIConditioning() const { return PBDRigidParticles->InvIConditioning(ParticleIdx); }
+	void SetInvIConditioning(const TVec3<FRealSingle>& InInvIConditioning) { PBDRigidParticles->InvIConditioning(ParticleIdx) = InInvIConditioning; }
+
+	// Get the conditioned inertia for use in constraint solvers
+	TVec3<FRealSingle> ConditionedInvI() const { return InvIConditioning() * InvI(); }
+	TVec3<FRealSingle> ConditionedI() const { return I() / InvIConditioning(); }
 
 	T LinearEtherDrag() const { return PBDRigidParticles->LinearEtherDrag(ParticleIdx); }
 	T& LinearEtherDrag() { return PBDRigidParticles->LinearEtherDrag(ParticleIdx); }
@@ -1032,21 +1048,77 @@ public:
 	const TPBDRigidParticleHandleImp<T, d, true>* Handle() const { return PBDRigidParticles->Handle(ParticleIdx); }
 	TPBDRigidParticleHandleImp<T, d, true>* Handle() { return PBDRigidParticles->Handle(ParticleIdx); }
 
-	bool GravityEnabled() const { return PBDRigidParticles->GravityEnabled(ParticleIdx); }
+	inline FRigidParticleControlFlags ControlFlags() const
+	{ 
+		return PBDRigidParticles->ControlFlags(ParticleIdx);
+	}
+	
+	// NOTE: ControlFlags should not be changed by the solver during the tick. These are externally controlled settings.
+	inline void SetControlFlags(const FRigidParticleControlFlags Flags)
+	{
+		PBDRigidParticles->ControlFlags(ParticleIdx) = Flags;
+	}
 
-	void SetGravityEnabled(bool bEnabled){ PBDRigidParticles->GravityEnabled(ParticleIdx) = bEnabled; }
+	inline bool GravityEnabled() const
+	{ 
+		return ControlFlags().GetGravityEnabled();
+	}
 
-	bool CCDEnabled() const {
-		return PBDRigidParticles->CCDEnabled(ParticleIdx);
+	inline void SetGravityEnabled(bool bEnabled)
+	{ 
+		PBDRigidParticles->ControlFlags(ParticleIdx).SetGravityEnabled(bEnabled);
 	}
 
-	void SetCCDEnabled(bool bEnabled) {
-		PBDRigidParticles->CCDEnabled(ParticleIdx) = bEnabled;
+	inline bool CCDEnabled() const
+	{
+		return ControlFlags().GetCCDEnabled();
 	}
 
-	bool OneWayInteraction() const { return PBDRigidParticles->OneWayInteraction(ParticleIdx); }
+	inline void SetCCDEnabled(bool bEnabled)
+	{
+		PBDRigidParticles->ControlFlags(ParticleIdx).SetCCDEnabled(bEnabled);
+	}
 
-	void SetOneWayInteraction(bool bInOneWayInteraction) { PBDRigidParticles->OneWayInteraction(ParticleIdx) = bInOneWayInteraction; }
+	inline bool OneWayInteraction() const
+	{ 
+		return ControlFlags().GetOneWayInteractionEnabled();
+	}
+
+	inline void SetOneWayInteraction(bool bEnabled)
+	{ 
+		PBDRigidParticles->ControlFlags(ParticleIdx).SetOneWayInteractionEnabled(bEnabled);
+	}
+
+	inline bool InertiaConditioningEnabled() const
+	{
+		return ControlFlags().GetInertiaConditioningEnabled();
+	}
+
+	inline void SetInertiaConditioningEnabled(bool bEnabled)
+	{
+		// NOTE: We still set this flag even for kinematics because they may change to dynamic later. However we
+		// won't actually calculate the inertia until it gets changed to dynamic (which will also set the dirty flag)
+		if (bEnabled != InertiaConditioningEnabled())
+		{
+			PBDRigidParticles->ControlFlags(ParticleIdx).SetInertiaConditioningEnabled(bEnabled);
+			SetInertiaConditioningDirty();
+		}
+	}
+
+	inline bool InertiaConditioningDirty()
+	{
+		return PBDRigidParticles->TransientFlags(ParticleIdx).GetInertiaConditioningDirty();
+	}
+
+	inline void SetInertiaConditioningDirty()
+	{
+		PBDRigidParticles->TransientFlags(ParticleIdx).SetInertiaConditioningDirty();
+	}
+
+	inline void ClearInertiaConditioningDirty()
+	{
+		PBDRigidParticles->TransientFlags(ParticleIdx).ClearInertiaConditioningDirty();
+	}
 
 	ESleepType SleepType() const { return PBDRigidParticles->SleepType(ParticleIdx);}
 
@@ -1555,9 +1627,12 @@ public:
 
 	const TVec3<FRealSingle> I() const
 	{ 
-		if (MHandle->CastToRigidParticle() && MHandle->ObjectState() == EObjectStateType::Dynamic)
+		if (auto RigidHandle = MHandle->CastToRigidParticle())
 		{
-			return MHandle->CastToRigidParticle()->I();
+			if (IsDynamic())
+			{
+				return RigidHandle->I();
+			}
 		}
 
 		return TVec3<FRealSingle>(0);
@@ -1565,9 +1640,12 @@ public:
 
 	const TVec3<FRealSingle> InvI() const
 	{ 
-		if (MHandle->CastToRigidParticle() && MHandle->ObjectState() == EObjectStateType::Dynamic)
+		if (auto RigidHandle = MHandle->CastToRigidParticle())
 		{
-			return MHandle->CastToRigidParticle()->InvI();
+			if (IsDynamic())
+			{
+				return MHandle->CastToRigidParticle()->InvI();
+			}
 		}
 
 		return TVec3<FRealSingle>(0);
@@ -1575,9 +1653,12 @@ public:
 
 	FReal M() const
 	{
-		if (MHandle->CastToRigidParticle() && MHandle->ObjectState() == EObjectStateType::Dynamic)
+		if (auto RigidHandle = MHandle->CastToRigidParticle())
 		{
-			return MHandle->CastToRigidParticle()->M();
+			if (IsDynamic())
+			{
+				return MHandle->CastToRigidParticle()->M();
+			}
 		}
 
 		return (FReal)0;
@@ -1585,9 +1666,12 @@ public:
 
 	FReal InvM() const
 	{
-		if (MHandle->CastToRigidParticle() && MHandle->ObjectState() == EObjectStateType::Dynamic)
+		if (auto RigidHandle = MHandle->CastToRigidParticle())
 		{
-			return MHandle->CastToRigidParticle()->InvM();
+			if (IsDynamic())
+			{
+				return MHandle->CastToRigidParticle()->InvM();
+			}
 		}
 
 		return (FReal)0;
@@ -1613,6 +1697,78 @@ public:
 		return FRotation3::FromIdentity();
 	}
 
+	TVec3<FRealSingle> InvIConditioning() const
+	{ 
+		if (auto RigidHandle = MHandle->CastToRigidParticle())
+		{
+			return RigidHandle->InvIConditioning();
+		}
+
+		return TVec3<FRealSingle>(1);
+	}
+
+	void SetInvIConditioning(const TVec3<FRealSingle>& InInvIConditioning)
+	{ 
+		if (auto RigidHandle = MHandle->CastToRigidParticle())
+		{
+			RigidHandle->SetInvIConditioning(InInvIConditioning);
+		}
+	}
+
+	TVec3<FRealSingle> ConditionedInvI() const
+	{ 
+		if (auto RigidHandle = MHandle->CastToRigidParticle())
+		{
+			return RigidHandle->ConditionedInvI();
+		}
+
+		return TVec3<FRealSingle>(0);
+	}
+
+	TVec3<FRealSingle> ConditionedI() const
+	{ 
+		if (auto RigidHandle = MHandle->CastToRigidParticle())
+		{
+			return RigidHandle->ConditionedI();
+		}
+
+		return TVec3<FRealSingle>(0);
+	}
+
+	bool InertiaConditioningEnabled() const
+	{
+		if (auto RigidHandle = MHandle->CastToRigidParticle())
+		{
+			return RigidHandle->InertiaConditioningEnabled();
+		}
+		return false;
+	}
+
+	bool InertiaConditioningDirty() const
+	{
+		if (auto RigidHandle = MHandle->CastToRigidParticle())
+		{
+			return RigidHandle->InertiaConditioningDirty();
+		}
+		return false;
+	}
+
+	void SetInertiaConditioningDirty()
+	{
+		if (auto RigidHandle = MHandle->CastToRigidParticle())
+		{
+			RigidHandle->SetInertiaConditioningDirty();
+		}
+	}
+
+	void ClearInertiaConditioningDirty()
+	{
+		if (auto RigidHandle = MHandle->CastToRigidParticle())
+		{
+			RigidHandle->ClearInertiaConditioningDirty();
+		}
+	}
+
 	FReal LinearEtherDrag() const
 	{
 		if (auto RigidHandle = MHandle->CastToRigidParticle())
@@ -2417,15 +2573,15 @@ public:
 	}*/
 
 	bool GravityEnabled() const { return MMiscData.Read().GravityEnabled(); }
-	void SetGravityEnabled(const bool InGravityEnabled)
+	void SetGravityEnabled(const bool bInEnabled)
 	{
-		MMiscData.Modify(true,MDirtyFlags,Proxy,[InGravityEnabled](auto& Data){ Data.SetGravityEnabled (InGravityEnabled);});
+		MMiscData.Modify(true, MDirtyFlags, Proxy, [bInEnabled](auto& Data) { Data.SetGravityEnabled(bInEnabled); });
 	}
 	
 	bool OneWayInteraction() const { return MMiscData.Read().OneWayInteraction(); }
-	void SetOneWayInteraction(const bool InOneWayInteraction)
+	void SetOneWayInteraction(const bool bInEnabled)
 	{
-		MMiscData.Modify(true, MDirtyFlags, Proxy, [InOneWayInteraction](auto& Data) { Data.SetOneWayInteraction(InOneWayInteraction); });
+		MMiscData.Modify(true, MDirtyFlags, Proxy, [bInEnabled](auto& Data) { Data.SetOneWayInteraction(bInEnabled); });
 	}
 
 	// Enable a single flag
@@ -2447,10 +2603,15 @@ public:
 	}
 
 	bool CCDEnabled() const { return MMiscData.Read().CCDEnabled(); }
+	void SetCCDEnabled(bool bInEnabled)
+	{
+		MMiscData.Modify(true, MDirtyFlags, Proxy, [bInEnabled](auto& Data) { Data.SetCCDEnabled(bInEnabled); });
+	}
 
-	void SetCCDEnabled(bool bInCCDEnabled)
+	bool InertiaConditioningEnabled() const { return MMiscData.Read().InertiaConditioningEnabled(); }
+	void SetInertiaConditioningEnabled(bool bInEnabled)
 	{
-		MMiscData.Modify(true, MDirtyFlags, Proxy, [bInCCDEnabled](auto& Data) {Data.SetCCDEnabled(bInCCDEnabled); });
+		MMiscData.Modify(true, MDirtyFlags, Proxy, [bInEnabled](auto& Data) { Data.SetInertiaConditioningEnabled(bInEnabled); });
 	}
 
 	bool Disabled() const { return MMiscData.Read().Disabled(); }
diff --git a/Engine/Source/Runtime/Experimental/Chaos/Public/Chaos/RigidParticleControlFlags.h b/Engine/Source/Runtime/Experimental/Chaos/Public/Chaos/RigidParticleControlFlags.h
new file mode 100644
--- /dev/null
+++ b/Engine/Source/Runtime/Experimental/Chaos/Public/Chaos/RigidParticleControlFlags.h
@@ -0,0 +1,124 @@
+// Copyright Epic Games, Inc. All Rights Reserved.
+#pragma once
+
+#include "CoreMinimal.h"
+
+class FString;
+
+namespace Chaos
+{
+	class FChaosArchive;
+
+	/**
+	 * @brief Enable/Disable the features on a particle
+	 * These flags are externally controlled and should not be changed by the solver during the tick. The may be
+	 * bound directly to settings or game-side object state. These flags should be treated as read-only to the solver.
+	*/
+	class FRigidParticleControlFlags
+	{
+	public:
+		using FStorage = uint8;
+
+		FRigidParticleControlFlags()
+			: Bits(0)
+		{
+		}
+
+		FStorage GetFlags() const { return Bits; }
+		void SetFlags(const FStorage InBits) { Bits = InBits; }
+
+		bool GetGravityEnabled() const { return Flags.bGravityEnabled; }
+		FRigidParticleControlFlags& SetGravityEnabled(const bool bEnabled) { Flags.bGravityEnabled = bEnabled; return *this; }
+
+		bool GetCCDEnabled() const { return Flags.bCCDEnabled; }
+		FRigidParticleControlFlags& SetCCDEnabled(const bool bEnabled) { Flags.bCCDEnabled = bEnabled; return *this; }
+
+		bool GetOneWayInteractionEnabled() const { return Flags.bOneWayInteractionEnabled; }
+		FRigidParticleControlFlags& SetOneWayInteractionEnabled(const bool bEnabled) { Flags.bOneWayInteractionEnabled = bEnabled; return *this; }
+
+		bool GetMaxDepenetrationVelocityOverrideEnabled() const { return Flags.bMaxDepenetrationVelocityOverrideEnabled; }
+		FRigidParticleControlFlags& SetMaxDepenetrationVelocityOverrideEnabled(const bool bEnabled) { Flags.bMaxDepenetrationVelocityOverrideEnabled = bEnabled; return *this; }
+
+		bool GetInertiaConditioningEnabled() const { return Flags.bInertiaConditioningEnabled; }
+		FRigidParticleControlFlags& SetInertiaConditioningEnabled(const bool bEnabled) { Flags.bInertiaConditioningEnabled = bEnabled; return *this; }
+
+		// Used by rewind system debuging
+		FString ToString() const;
+
+		// Used by MarshallingManager and dirty data system
+		friend bool operator==(const FRigidParticleControlFlags& L, const FRigidParticleControlFlags& R) { return L.GetFlags() == R.GetFlags(); }
+		friend bool operator!=(const FRigidParticleControlFlags& L, const FRigidParticleControlFlags& R) { return L.GetFlags() != R.GetFlags(); }
+
+		// Serialization
+		friend FChaosArchive& operator<<(FChaosArchive& Ar, FRigidParticleControlFlags& Data);
+
+	private:
+		struct FFlags
+		{
+			FStorage bGravityEnabled : 1;
+			FStorage bCCDEnabled : 1;
+			FStorage bOneWayInteractionEnabled : 1;
+			FStorage bMaxDepenetrationVelocityOverrideEnabled : 1;
+			FStorage bInertiaConditioningEnabled : 1;
+			// Add new properties above this line
+			// Change FStorage typedef if we exceed the max bits
+		};
+		union
+		{
+			FFlags Flags;
+			FStorage Bits;
+		};
+	};
+
+	// If we add more bits and exceed the storage size...
+	static_assert(sizeof(FRigidParticleControlFlags) == sizeof(FRigidParticleControlFlags::FStorage), "Change FRigidParticleControlFlags::FStorage to be larger");
+
+
+	/**
+	 * @brief Transient flags for indicating somethings needs to be updated based on a change to the particle
+	 * Typically some of these flags will be set when something on the particle changes (e.g., when initially created,
+	 * a joint is connected, etc etc). The flags will be checked and reset by some system during the tick.
+	 * 
+	 * Transient flags are for use in the solver only. They should not be directly controlled by settings, or
+	 * sent back outside of the solver (e.g., not to the game thread). @see FRigidParticleControlFlags for
+	 * externally controlled flags.
+	 *
+	 * E.g., the bInertiaConditioningDirty flags is set whenever the particle changes in such a way that the
+	 * inertia conditioning needs to be recalculated and is reset by the evolution when it has recalculated it.
+	 * 
+	 * @note In general the transient flags can be set from multiple locations, but should only ever be cleared in one.
+	*/
+	class FRigidParticleTransientFlags
+	{
+	public:
+		using FStorage = uint8;
+
+		FRigidParticleTransientFlags()
+			: Bits(0)
+		{
+		}
+
+		FStorage GetFlags() const { return Bits; }
+		void SetFlags(const FStorage InFlags) { Bits = InFlags; }
+
+		bool GetInertiaConditioningDirty() const { return Flags.bInertiaConditioningDirty; }
+		void SetInertiaConditioningDirty() { Flags.bInertiaConditioningDirty = true; }
+		void ClearInertiaConditioningDirty() { Flags.bInertiaConditioningDirty = false; }
+
+	private:
+		struct FFlags
+		{
+			FStorage bInertiaConditioningDirty : 1;
+			// Add new properties above this line
+			// Change FStorage typedef if we exceed the max bits
+		};
+		union
+		{
+			FFlags Flags;
+			FStorage Bits;
+		};
+	};
+
+	// If we add more bits and exceed the storage size...
+	static_assert(sizeof(FRigidParticleTransientFlags) == sizeof(FRigidParticleTransientFlags::FStorage), "Change FRigidParticleTransientFlags::FStorage to be larger");
+}
\ No newline at end of file
diff --git a/Engine/Source/Runtime/Experimental/Chaos/Public/Chaos/RigidParticles.h b/Engine/Source/Runtime/Experimental/Chaos/Public/Chaos/RigidParticles.h
--- a/Engine/Source/Runtime/Experimental/Chaos/Public/Chaos/RigidParticles.h
+++ b/Engine/Source/Runtime/Experimental/Chaos/Public/Chaos/RigidParticles.h
@@ -11,6 +11,7 @@
 #include "Chaos/Matrix.h"
 #include "Chaos/Particles.h"
 #include "Chaos/Rotation.h"
+#include "Chaos/RigidParticleControlFlags.h"
 #include "HAL/LowLevelMemTracker.h"
 #include "UObject/FortniteMainBranchObjectVersion.h"
 
@@ -76,6 +77,7 @@ public:
 		TArrayCollection::AddArray(&MAngularImpulseVelocity);
 		TArrayCollection::AddArray(&MI);
 		TArrayCollection::AddArray(&MInvI);
+		TArrayCollection::AddArray(&MInvIConditioning);
 		TArrayCollection::AddArray(&MM);
 		TArrayCollection::AddArray(&MInvM);
 		TArrayCollection::AddArray(&MCenterOfMass);
@@ -87,16 +89,15 @@ public:
 		TArrayCollection::AddArray(&MCollisionParticles);
 		TArrayCollection::AddArray(&MCollisionGroup);
 		TArrayCollection::AddArray(&MCollisionConstraintFlags);
+		TArrayCollection::AddArray(&MControlFlags);
+		TArrayCollection::AddArray(&MTransientFlags);
 		TArrayCollection::AddArray(&MDisabled);
 		TArrayCollection::AddArray(&MObjectState);
 		TArrayCollection::AddArray(&MPreObjectState);
 		TArrayCollection::AddArray(&MIslandIndex);
 		TArrayCollection::AddArray(&MGraphIndex);
 		TArrayCollection::AddArray(&MToBeRemovedOnFracture);
-		TArrayCollection::AddArray(&MGravityEnabled);
-		TArrayCollection::AddArray(&MOneWayInteraction);
 		TArrayCollection::AddArray(&MSleepType);
-		TArrayCollection::AddArray(&bCCDEnabled);
 	}
 	TRigidParticles(const TRigidParticles<T, d>& Other) = delete;
 	CHAOS_API TRigidParticles(TRigidParticles<T, d>&& Other)
@@ -109,6 +110,7 @@ public:
 		, MAngularImpulseVelocity(MoveTemp(Other.MAngularImpulseVelocity))
 		, MI(MoveTemp(Other.MI))
 		, MInvI(MoveTemp(Other.MInvI))
+		, MInvIConditioning(MoveTemp(Other.MInvIConditioning))
 		, MM(MoveTemp(Other.MM))
 		, MInvM(MoveTemp(Other.MInvM))
 		, MCenterOfMass(MoveTemp(Other.MCenterOfMass))
@@ -120,15 +122,14 @@ public:
 		, MCollisionParticles(MoveTemp(Other.MCollisionParticles))
 		, MCollisionGroup(MoveTemp(Other.MCollisionGroup))
 		, MCollisionConstraintFlags(MoveTemp(Other.MCollisionConstraintFlags))
+		, MControlFlags(MoveTemp(Other.MControlFlags))
+		, MTransientFlags(MoveTemp(Other.MTransientFlags))
 		, MIslandIndex(MoveTemp(Other.MIslandIndex))
 		, MGraphIndex(MoveTemp(Other.MGraphIndex))
 		, MToBeRemovedOnFracture(MoveTemp(Other.MToBeRemovedOnFracture))
 		, MObjectState(MoveTemp(Other.MObjectState))
 		, MPreObjectState(MoveTemp(Other.MPreObjectState))
-		, MGravityEnabled(MoveTemp(Other.MGravityEnabled))
-		, MOneWayInteraction(MoveTemp(Other.MOneWayInteraction))
 		, MSleepType(MoveTemp(Other.MSleepType))
-		, bCCDEnabled(MoveTemp(Other.bCCDEnabled))
 		, MDisabled(MoveTemp(Other.MDisabled))
 	{
 		TArrayCollection::AddArray(&MVSmooth);
@@ -139,6 +140,7 @@ public:
 		TArrayCollection::AddArray(&MAngularImpulseVelocity);
 		TArrayCollection::AddArray(&MI);
 		TArrayCollection::AddArray(&MInvI);
+		TArrayCollection::AddArray(&MInvIConditioning);
 		TArrayCollection::AddArray(&MM);
 		TArrayCollection::AddArray(&MInvM);
 		TArrayCollection::AddArray(&MCenterOfMass);
@@ -150,16 +152,15 @@ public:
 		TArrayCollection::AddArray(&MCollisionParticles);
 		TArrayCollection::AddArray(&MCollisionGroup);
 		TArrayCollection::AddArray(&MCollisionConstraintFlags);
+		TArrayCollection::AddArray(&MControlFlags);
+		TArrayCollection::AddArray(&MTransientFlags);
 		TArrayCollection::AddArray(&MDisabled);
 		TArrayCollection::AddArray(&MObjectState);
 		TArrayCollection::AddArray(&MPreObjectState);
 		TArrayCollection::AddArray(&MIslandIndex);
 		TArrayCollection::AddArray(&MGraphIndex);
 		TArrayCollection::AddArray(&MToBeRemovedOnFracture);
-		TArrayCollection::AddArray(&MGravityEnabled);
-		TArrayCollection::AddArray(&MOneWayInteraction);
 		TArrayCollection::AddArray(&MSleepType);
-		TArrayCollection::AddArray(&bCCDEnabled);
 	}
 
 	CHAOS_API virtual ~TRigidParticles()
@@ -189,6 +190,9 @@ public:
 	FORCEINLINE const TVec3<FRealSingle>& InvI(const int32 Index) const { return MInvI[Index]; }
 	FORCEINLINE TVec3<FRealSingle>& InvI(const int32 Index) { return MInvI[Index]; }
 
+	FORCEINLINE const TVec3<FRealSingle>& InvIConditioning(const int32 Index) const { return MInvIConditioning[Index]; }
+	FORCEINLINE TVec3<FRealSingle>& InvIConditioning(const int32 Index) { return MInvIConditioning[Index]; }
+
 	FORCEINLINE const T M(const int32 Index) const { return MM[Index]; }
 	FORCEINLINE T& M(const int32 Index) { return MM[Index]; }
 
@@ -241,16 +245,11 @@ public:
 	FORCEINLINE const bool ToBeRemovedOnFracture(const int32 Index) const { return MToBeRemovedOnFracture[Index]; }
 	FORCEINLINE bool& ToBeRemovedOnFracture(const int32 Index) { return MToBeRemovedOnFracture[Index]; }
 
-	FORCEINLINE const bool& GravityEnabled(const int32 Index) const { return MGravityEnabled[Index]; }
-	FORCEINLINE bool& GravityEnabled(const int32 Index) { return MGravityEnabled[Index]; }
-
-	FORCEINLINE const bool& OneWayInteraction(const int32 Index) const { return MOneWayInteraction[Index]; }
-	FORCEINLINE bool& OneWayInteraction(const int32 Index) { return MOneWayInteraction[Index]; }
-
-
-	FORCEINLINE const bool& CCDEnabled(const int32 Index) const { return bCCDEnabled[Index]; }
-	FORCEINLINE bool& CCDEnabled(const int32 Index) { return bCCDEnabled[Index]; }
+	FORCEINLINE const FRigidParticleControlFlags& ControlFlags(const int32 Index) const { return MControlFlags[Index]; }
+	FORCEINLINE FRigidParticleControlFlags& ControlFlags(const int32 Index) { return MControlFlags[Index]; }
 
+	FORCEINLINE const FRigidParticleTransientFlags& TransientFlags(const int32 Index) const { return MTransientFlags[Index]; }
+	FORCEINLINE FRigidParticleTransientFlags& TransientFlags(const int32 Index) { return MTransientFlags[Index]; }
 
 	FORCEINLINE ESleepType SleepType(const int32 Index) const { return MSleepType[Index]; }
 	FORCEINLINE ESleepType& SleepType(const int32 Index) { return MSleepType[Index]; }
@@ -333,6 +332,12 @@ public:
 		{
 			Ar << MI << MInvI;
 		}
+
+		MInvIConditioning.Resize(MInvI.Num());
+		for (int32 Index = 0; Index < MInvI.Num(); ++Index)
+		{
+			MInvIConditioning[Index] = TVec3<FRealSingle>(1);
+		}
 		
 		Ar << MM << MInvM;
 
@@ -343,27 +348,9 @@ public:
 		}
 
 		Ar << MCollisionParticles << MCollisionGroup << MIslandIndex << MDisabled << MObjectState << MSleepType;
-		//todo: add gravity enabled when we decide how we want to handle serialization
-	}
 
-	FORCEINLINE TArray<TVector<T, d>>& AllAcceleration() { return MAcceleration; }
-	FORCEINLINE TArray<TVector<T, d>>& AllAngularAcceleration() { return MAngularAcceleration; }
-	FORCEINLINE TArray<TVector<T, d>>& AllLinearImpulseVelocity() { return MLinearImpulseVelocity; }
-	FORCEINLINE TArray<TVector<T, d>>& AllAngularImpulseVelocity() { return MAngularImpulseVelocity; }
-	FORCEINLINE TArray<TVec3<FRealSingle>>& AllI() { return MI; }
-	FORCEINLINE TArray<TVec3<FRealSingle>>& AllInvI() { return MInvI; }
-	FORCEINLINE TArray<FReal>& AllM() { return MM; }
-	FORCEINLINE TArray<FReal>& AllInvM() { return MInvM; }
-	FORCEINLINE TArray<TVector<T, d>>& AllCenterOfMass() { return MCenterOfMass; }
-	FORCEINLINE TArray<TRotation<T, d>>& AllRotationOfMass() { return MRotationOfMass; }
-	FORCEINLINE TArray<FReal>& AllLinearEtherDrag() { return MLinearEtherDrag; }
-	FORCEINLINE TArray<FReal>& AllAngularEtherDrag() { return MAngularEtherDrag; }
-	FORCEINLINE TArray<FReal>& AllMaxLinearSpeeds() { return MaxLinearSpeedsSq; }
-	FORCEINLINE TArray<FReal>& AllMaxAngularSpeeds() { return MaxAngularSpeedsSq; }
-	FORCEINLINE TArray<bool>& AllDisabled() { return MDisabled; }
-	FORCEINLINE TArray<EObjectStateType>& AllObjectState() { return MObjectState; }
-	FORCEINLINE TArray<bool>& AllGravityEnabled() { return MGravityEnabled; }
-	FORCEINLINE TArray<bool>& AllCCDEnabled() { return bCCDEnabled; }
+		// @todo(chaos): what about MControlFlags and MTransientFlags?
+	}
 
 private:
 	TArrayCollectionArray<TVector<T, d>> MVSmooth;
@@ -374,6 +361,7 @@ private:
 	TArrayCollectionArray<TVector<T, d>> MAngularImpulseVelocity;
 	TArrayCollectionArray<TVec3<FRealSingle>> MI;
 	TArrayCollectionArray<TVec3<FRealSingle>> MInvI;
+	TArrayCollectionArray<TVec3<FRealSingle>> MInvIConditioning;
 	TArrayCollectionArray<T> MM;
 	TArrayCollectionArray<T> MInvM;
 	TArrayCollectionArray<TVector<T,d>> MCenterOfMass;
@@ -385,15 +373,14 @@ private:
 	TArrayCollectionArray<TUniquePtr<TBVHParticles<T, d>>> MCollisionParticles;
 	TArrayCollectionArray<int32> MCollisionGroup;
 	TArrayCollectionArray<uint32> MCollisionConstraintFlags;
+	TArrayCollectionArray<FRigidParticleControlFlags> MControlFlags;
+	TArrayCollectionArray<FRigidParticleTransientFlags> MTransientFlags;
 	TArrayCollectionArray<int32> MIslandIndex;
 	TArrayCollectionArray<int32> MGraphIndex;
 	TArrayCollectionArray<bool> MToBeRemovedOnFracture;
 	TArrayCollectionArray<EObjectStateType> MObjectState;
 	TArrayCollectionArray<EObjectStateType> MPreObjectState;
-	TArrayCollectionArray<bool> MGravityEnabled;
-	TArrayCollectionArray<bool> MOneWayInteraction;
 	TArrayCollectionArray<ESleepType> MSleepType;
-	TArrayCollectionArray<bool> bCCDEnabled;
 	TArrayCollectionArray<bool> MDisabled;
 
 	TArray<TSleepData<T, d>> MSleepData;
diff --git a/Engine/Source/Runtime/Experimental/Chaos/Public/PhysicsProxy/SingleParticlePhysicsProxy.h b/Engine/Source/Runtime/Experimental/Chaos/Public/PhysicsProxy/SingleParticlePhysicsProxy.h
--- a/Engine/Source/Runtime/Experimental/Chaos/Public/PhysicsProxy/SingleParticlePhysicsProxy.h
+++ b/Engine/Source/Runtime/Experimental/Chaos/Public/PhysicsProxy/SingleParticlePhysicsProxy.h
@@ -396,17 +396,41 @@ public:
 		});
 	}
 
-	void SetOneWayInteraction(const bool InOneWayInteraction)
+	void SetOneWayInteraction(const bool bInOneWayInteraction)
 	{
-		Write([InOneWayInteraction](auto* Particle)
+		Write([bInOneWayInteraction](auto* Particle)
 		{
 			if (auto Rigid = Particle->CastToRigidParticle())
 			{
-				return Rigid->SetOneWayInteraction(InOneWayInteraction);
+				return Rigid->SetOneWayInteraction(bInOneWayInteraction);
 			}
 		});
 	}
 
+	bool InertiaConditioningEnabled() const
+	{
+		return Read([](auto* Particle)
+			{
+				if (auto Rigid = Particle->CastToRigidParticle())
+				{
+					return Rigid->InertiaConditioningEnabled();
+				}
+
+				return false;
+			});
+	}
+
+	void SetInertiaConditioningEnabled(const bool bInEnabled)
+	{
+		Write([bInEnabled](auto* Particle)
+			{
+				if (auto Rigid = Particle->CastToRigidParticle())
+				{
+					return Rigid->SetInertiaConditioningEnabled(bInEnabled);
+				}
+			});
+	}
+
 	void SetResimType(EResimType ResimType)
 	{
 		Write([ResimType](auto* Particle)
diff --git a/Engine/Source/Runtime/Experimental/Chaos/Public/RewindData.h b/Engine/Source/Runtime/Experimental/Chaos/Public/RewindData.h
--- a/Engine/Source/Runtime/Experimental/Chaos/Public/RewindData.h
+++ b/Engine/Source/Runtime/Experimental/Chaos/Public/RewindData.h
@@ -655,9 +655,8 @@ public:
 	REWIND_PARTICLE_RIGID_PROPERTY(DynamicsMisc, MaxLinearSpeedSq)
 	REWIND_PARTICLE_RIGID_PROPERTY(DynamicsMisc, MaxAngularSpeedSq)
 	REWIND_PARTICLE_RIGID_PROPERTY(DynamicsMisc, ObjectState)
-	REWIND_PARTICLE_RIGID_PROPERTY(DynamicsMisc, GravityEnabled)
-	REWIND_PARTICLE_RIGID_PROPERTY(DynamicsMisc, CCDEnabled)
 	REWIND_PARTICLE_RIGID_PROPERTY(DynamicsMisc, CollisionGroup)
+	REWIND_PARTICLE_RIGID_PROPERTY(DynamicsMisc, ControlFlags)
 
 	REWIND_PARTICLE_RIGID_PROPERTY(MassProps, CenterOfMass)
 	REWIND_PARTICLE_RIGID_PROPERTY(MassProps, RotationOfMass)
@@ -718,9 +717,8 @@ public:
 			REWIND_PARTICLE_TO_STR(MaxLinearSpeedSq)
 			REWIND_PARTICLE_TO_STR(MaxAngularSpeedSq)
 			REWIND_PARTICLE_TO_STR(ObjectState)
-			REWIND_PARTICLE_TO_STR(GravityEnabled)
-			REWIND_PARTICLE_TO_STR(CCDEnabled)
 			REWIND_PARTICLE_TO_STR(CollisionGroup)
+			REWIND_PARTICLE_TO_STR(ControlFlags)
 
 			REWIND_PARTICLE_TO_STR(CenterOfMass)
 			REWIND_PARTICLE_TO_STR(RotationOfMass)
diff --git a/Engine/Source/Runtime/Experimental/ChaosCore/Public/Chaos/Core.h b/Engine/Source/Runtime/Experimental/ChaosCore/Public/Chaos/Core.h
--- a/Engine/Source/Runtime/Experimental/ChaosCore/Public/Chaos/Core.h
+++ b/Engine/Source/Runtime/Experimental/ChaosCore/Public/Chaos/Core.h
@@ -20,10 +20,19 @@ namespace Chaos
 	using FMatrix33 = PMatrix<FReal, 3, 3>;
 	using FMatrix44 = PMatrix<FReal, 4, 4>;
 	using FRigidTransform3 = TRigidTransform<FReal, 3>;
-	using FRigidTransformRealSingle3 = TRigidTransform<FRealSingle, 3>;
 
 	using FAABB3 = TAABB<FReal, 3>;
 
+	using FVec2f = TVector<FRealSingle, 2>;
+	using FVec3f = TVector<FRealSingle, 3>;
+	using FRotation3f = TRotation<FRealSingle, 3>;
+	using FRigidTransform3f = TRigidTransform<FRealSingle, 3>;
+
+	using FAABB3f = TAABB<FRealSingle, 3>;
+
+	// @todo(chaos): deprecate this and use FRigidTransform3f
+	using FRigidTransformRealSingle3 = TRigidTransform<FRealSingle, 3>;
+
 	template <typename T>
 	using TVec2 = TVector<T, 2>;
 
@@ -42,6 +51,6 @@ namespace Chaos
 	// NOTE: if you get a merge conflict on the GUID, you must replace it with a new GUID - do not accept the source or target
 	// or you will likely get DDC version conflicts resulting in crashes during load.
 	// Core version string for Chaos data. Any DDC builder dependent on Chaos for serialization should depend on this version
-	static const FString ChaosVersionString = TEXT("F3C4A6F4-BF50-488B-9FE2-7968054E9C6B");
+	static const FString ChaosVersionString = TEXT("B3AEC941-1F1D-414A-A961-C0376566C98D");
 
 }
diff --git a/Engine/Source/Runtime/PhysicsCore/Private/ChaosEngineInterface.cpp b/Engine/Source/Runtime/PhysicsCore/Private/ChaosEngineInterface.cpp
--- a/Engine/Source/Runtime/PhysicsCore/Private/ChaosEngineInterface.cpp
+++ b/Engine/Source/Runtime/PhysicsCore/Private/ChaosEngineInterface.cpp
@@ -296,7 +296,7 @@ void FChaosEngineInterface::DetachShape(const FPhysicsActorHandle& InActor,FPhys
 	}
 }
 
-void FChaosEngineInterface::SetSmoothEdgeCollisionsEnabled(const FPhysicsActorHandle& InActor, const bool bSmoothEdgeCollisionsEnabled)
+void FChaosEngineInterface::SetSmoothEdgeCollisionsEnabled_AssumesLocked(const FPhysicsActorHandle& InActor, const bool bSmoothEdgeCollisionsEnabled)
 {
 	InActor->GetGameThreadAPI().SetSmoothEdgeCollisionsEnabled(bSmoothEdgeCollisionsEnabled);
 }
@@ -898,12 +898,20 @@ void FChaosEngineInterface::SetMassSpaceInertiaTensor_AssumesLocked(FPhysicsActo
 
 void FChaosEngineInterface::SetComLocalPose_AssumesLocked(const FPhysicsActorHandle& InHandle,const FTransform& InComLocalPose)
 {
-	//@todo(mlentine): What is InComLocalPose? If the center of an object is not the local pose then many things break including the three vector represtnation of inertia.
 	Chaos::FRigidBodyHandle_External& Body_External = InHandle->GetGameThreadAPI();
 	Body_External.SetCenterOfMass(InComLocalPose.GetLocation());
 	Body_External.SetRotationOfMass(InComLocalPose.GetRotation());
 }
 
+bool FChaosEngineInterface::IsInertiaConditioningEnabled_AssumesLocked(const FPhysicsActorHandle& InActorReference)
+{
+	return InActorReference->GetGameThreadAPI().InertiaConditioningEnabled();
+}
+void FChaosEngineInterface::SetInertiaConditioningEnabled_AssumesLocked(const FPhysicsActorHandle& InActorReference, bool bEnabled)
+{
+	InActorReference->GetGameThreadAPI().SetInertiaConditioningEnabled(bEnabled);
+}
+
 void FChaosEngineInterface::SetIsSimulationShape(const FPhysicsShapeHandle& InShape,bool bIsSimShape)
 {
 	InShape.Shape->SetSimEnabled(bIsSimShape);
diff --git a/Engine/Source/Runtime/PhysicsCore/Public/Chaos/ChaosEngineInterface.h b/Engine/Source/Runtime/PhysicsCore/Public/Chaos/ChaosEngineInterface.h
--- a/Engine/Source/Runtime/PhysicsCore/Public/Chaos/ChaosEngineInterface.h
+++ b/Engine/Source/Runtime/PhysicsCore/Public/Chaos/ChaosEngineInterface.h
@@ -375,7 +375,7 @@ public:
 	static void AttachShape(const FPhysicsActorHandle& InActor,const FPhysicsShapeHandle& InNewShape);
 	static void DetachShape(const FPhysicsActorHandle& InActor,FPhysicsShapeHandle& InShape,bool bWakeTouching = true);
 
-	static void SetSmoothEdgeCollisionsEnabled(const FPhysicsActorHandle& InActor, const bool bSmoothEdgeCollisionsEnabled);
+	static void SetSmoothEdgeCollisionsEnabled_AssumesLocked(const FPhysicsActorHandle& InActor, const bool bSmoothEdgeCollisionsEnabled);
 
 	static void AddDisabledCollisionsFor_AssumesLocked(const TMap<FPhysicsActorHandle, TArray< FPhysicsActorHandle > >& InMap);
 	static void RemoveDisabledCollisionsFor_AssumesLocked(TArray< FPhysicsActorHandle > & InPhysicsActors);
@@ -466,6 +466,9 @@ public:
 	static void SetMassSpaceInertiaTensor_AssumesLocked(FPhysicsActorHandle& InHandle,const FVector& InTensor);
 	static void SetComLocalPose_AssumesLocked(const FPhysicsActorHandle& InHandle,const FTransform& InComLocalPose);
 
+	static bool IsInertiaConditioningEnabled_AssumesLocked(const FPhysicsActorHandle& InActorReference);
+	static void SetInertiaConditioningEnabled_AssumesLocked(const FPhysicsActorHandle& InActorReference, bool bEnabled);
+
 	static float GetStabilizationEnergyThreshold_AssumesLocked(const FPhysicsActorHandle& InHandle);
 	static void SetStabilizationEnergyThreshold_AssumesLocked(const FPhysicsActorHandle& InHandle,float InThreshold);
 	static uint32 GetSolverPositionIterationCount_AssumesLocked(const FPhysicsActorHandle& InHandle);
