Chaos : Flesh
Add corotated constraints and corotated fiber constraints in xpbd style. Add dRdF computation for corotated model.

diff --git a/Engine/Source/Runtime/Experimental/Chaos/Public/Chaos/XPBDCorotatedConstraints.h b/Engine/Source/Runtime/Experimental/Chaos/Public/Chaos/XPBDCorotatedConstraints.h
new file mode 100644
--- /dev/null
+++ b/Engine/Source/Runtime/Experimental/Chaos/Public/Chaos/XPBDCorotatedConstraints.h
@@ -0,0 +1,438 @@
+// Copyright Epic Games, Inc. All Rights Reserved.
+#pragma once
+
+#include "Chaos/PBDSpringConstraintsBase.h"
+#include "ChaosStats.h"
+#include "Chaos/ImplicitQRSVD.h"
+#include "Chaos/GraphColoring.h"
+#include "Chaos/Framework/Parallel.h"
+
+DECLARE_CYCLE_STAT(TEXT("Chaos XPBD Corotated Constraint"), STAT_ChaosXPBDCorotated, STATGROUP_Chaos);
+
+namespace Chaos::Softs
+{
+
+	class FXPBDCorotatedConstraints 
+	{
+
+		//TODO(Yizhou Chen): COuld be optimized . The SVD is using double so float -> double -> float all the time
+		//should change data type in accordance 
+
+	public:
+		//this one only accepts tetmesh input and mesh
+		FXPBDCorotatedConstraints(
+			const Chaos::Softs::FSolverParticles& InParticles,
+			int32 ParticleOffset,
+			int32 ParticleCount,
+			const TArray<TVector<int32, 4>>& InMesh,
+			const bool bRecordMetricIn = true,
+			const Chaos::Softs::FSolverReal& EMesh = (FSolverReal)10.0,
+			const Chaos::Softs::FSolverReal& NuMesh = (FSolverReal).3
+			)
+			: bRecordMetric(bRecordMetricIn), MeshConstraints(InMesh)
+		{	
+
+			LambdaArray.Init((FSolverReal)0., 2 * MeshConstraints.Num());
+			DmInverse.Init((FSolverReal)0., 9 * MeshConstraints.Num());
+			Measure.Init((FSolverReal)0.,  MeshConstraints.Num());
+			Lambda = EMesh * NuMesh / (((FSolverReal)1. + NuMesh) * ((FSolverReal)1. - (FSolverReal)2. * NuMesh));
+			Mu = EMesh / ((FSolverReal)2. * ((FSolverReal)1. + NuMesh));
+			for (int e = 0; e < InMesh.Num(); e++)
+			{
+				PMatrix<FSolverReal, 3, 3> Dm = DsInit(e, InParticles);
+				PMatrix<FSolverReal, 3, 3> DmInv = Dm.Inverse();
+				for (int r = 0; r < 3; r++) {
+					for (int c = 0; c < 3; c++) {
+						DmInverse[(3 * 3) * e + 3 * r + c] = DmInv.GetAt(r, c);
+					}
+				}
+
+				Measure[e] = Dm.Determinant() / (FSolverReal)6.;
+
+				//part of preprocessing: if inverted element is found, 
+				//invert it so that the measure is positive
+				if (Measure[e] < (FSolverReal)0.)
+				{
+
+					Measure[e] = -Measure[e];
+				}
+
+			}
+
+			InitColor(InParticles);
+		}
+
+		virtual ~FXPBDCorotatedConstraints() {}
+
+		PMatrix<FSolverReal, 3, 3> DsInit(const int e, const FSolverParticles& InParticles) const {
+			PMatrix<FSolverReal, 3, 3> Result((FSolverReal)0.);
+			for (int i = 0; i < 3; i++) {
+				for (int c = 0; c < 3; c++) {
+					Result.SetAt(c, i, InParticles.X(MeshConstraints[e][i + 1])[c] - InParticles.X(MeshConstraints[e][0])[c]);
+				}
+			}
+			return Result;
+		}
+
+
+		PMatrix<FSolverReal, 3, 3> Ds(const int e, const FSolverParticles& InParticles) const {
+			PMatrix<FSolverReal, 3, 3> Result((FSolverReal)0.);
+			for (int i = 0; i < 3; i++) {
+				for (int c = 0; c < 3; c++) {
+					Result.SetAt(c, i, InParticles.P(MeshConstraints[e][i+1])[c] - InParticles.P(MeshConstraints[e][0])[c]);
+				}
+			}
+			return Result;
+		}
+
+
+		PMatrix<FSolverReal, 3, 3> F(const int e, const FSolverParticles& InParticles) const {
+			return ElementDmInv(e) * Ds(e, InParticles);
+		}
+
+		PMatrix<FSolverReal, 3, 3> ElementDmInv(const int e) const {
+			PMatrix<FSolverReal, 3, 3> DmInv((FSolverReal)0.);
+			for (int r = 0; r < 3; r++) {
+				for (int c = 0; c < 3; c++) {
+					DmInv.SetAt(r, c, DmInverse[(3 * 3) * e + 3 * r + c]);
+				}
+			}
+			return DmInv;
+		}
+		
+		void Init() const 
+		{
+			for (FReal& Lambdas : LambdaArray) { Lambdas = (FReal)0.; }
+		}
+
+		virtual void ApplyInSerial(FSolverParticles& Particles, const FSolverReal Dt, const int32 ElementIndex) const
+		{
+			TVec4<FSolverVec3> PolarDelta = GetPolarDelta(Particles, Dt, ElementIndex);
+
+			for (int i = 0; i < 4; i++) 
+			{
+				Particles.P(MeshConstraints[ElementIndex][i]) += PolarDelta[i];
+			}
+			
+			TVec4<FSolverVec3> DetDelta = GetDeterminantDelta(Particles, Dt, ElementIndex);
+
+			for (int i = 0; i < 4; i++)
+			{
+				Particles.P(MeshConstraints[ElementIndex][i]) += DetDelta[i];
+			}
+
+
+		}
+
+		void ApplyInSerial(FSolverParticles& Particles, const FSolverReal Dt) const
+		{
+			SCOPE_CYCLE_COUNTER(STAT_ChaosXPBDCorotated);
+
+			for (int32 ElementIndex = 0; ElementIndex < MeshConstraints.Num(); ++ElementIndex)
+			{
+				ApplyInSerial(Particles, Dt, ElementIndex);
+			}
+
+
+		}
+
+		void ApplyInParallel(FSolverParticles& Particles, const FSolverReal Dt) const
+		{	
+			//code for error metric:
+			if (bRecordMetric)
+			{
+				GError.Init((FSolverReal)0., 3 * Particles.Size());
+				HErrorArray.Init((FSolverReal)0., 2 * MeshConstraints.Num());
+			}
+			
+
+			SCOPE_CYCLE_COUNTER(STAT_ChaosXPBDCorotated);
+			if ((ConstraintsPerColorStartIndex.Num() > 1) )//&& (MeshConstraints.Num() > Chaos_Spring_ParallelConstraintCount))
+			{
+				const int32 ConstraintColorNum = ConstraintsPerColorStartIndex.Num() - 1;
+
+				for (int32 ConstraintColorIndex = 0; ConstraintColorIndex < ConstraintColorNum; ++ConstraintColorIndex)
+				{
+					const int32 ColorStart = ConstraintsPerColorStartIndex[ConstraintColorIndex];
+					const int32 ColorSize = ConstraintsPerColorStartIndex[ConstraintColorIndex + 1] - ColorStart;
+					PhysicsParallelFor(ColorSize, [&](const int32 Index)
+						{
+							const int32 ConstraintIndex = ColorStart + Index;
+							ApplyInSerial(Particles, Dt, ConstraintIndex);
+						});
+				}
+			}
+
+
+		}
+
+	private:
+
+		void InitColor(const FSolverParticles& Particles)
+		{
+			// In dev builds we always color so we can tune the system without restarting. See Apply()
+			#if UE_BUILD_SHIPPING || UE_BUILD_TEST
+			if (Constraints.Num() > Chaos_Spring_ParallelConstraintCount)
+			#endif
+			{
+				const TArray<TArray<int32>> ConstraintsPerColor = FGraphColoring::ComputeGraphColoring(MeshConstraints, Particles);
+
+				// Reorder constraints based on color so each array in ConstraintsPerColor contains contiguous elements.
+				TArray<TVec4<int32>> ReorderedConstraints;
+				TArray<FReal> ReorderedMeasure;
+				TArray<FReal> ReorderedDmInverse;
+				TArray<int32> OrigToReorderedIndices; // used to reorder stiffness indices
+				ReorderedConstraints.SetNumUninitialized(MeshConstraints.Num());
+				ReorderedMeasure.SetNumUninitialized(Measure.Num());
+				ReorderedDmInverse.SetNumUninitialized(DmInverse.Num());
+				OrigToReorderedIndices.SetNumUninitialized(MeshConstraints.Num());
+
+				ConstraintsPerColorStartIndex.Reset(ConstraintsPerColor.Num() + 1);
+
+				int32 ReorderedIndex = 0;
+				for (const TArray<int32>& ConstraintsBatch : ConstraintsPerColor)
+				{
+					ConstraintsPerColorStartIndex.Add(ReorderedIndex);
+					for (const int32& BatchConstraint : ConstraintsBatch)
+					{
+						const int32 OrigIndex = BatchConstraint;
+						ReorderedConstraints[ReorderedIndex] = MeshConstraints[OrigIndex];
+						ReorderedMeasure[ReorderedIndex] = Measure[OrigIndex];
+						for (int32 kk = 0; kk < 9; kk++)
+						{
+							ReorderedDmInverse[9 * ReorderedIndex + kk] = DmInverse[9 * OrigIndex + kk];
+						}
+						OrigToReorderedIndices[OrigIndex] = ReorderedIndex;
+
+						++ReorderedIndex;
+					}
+				}
+				ConstraintsPerColorStartIndex.Add(ReorderedIndex);
+
+				MeshConstraints = MoveTemp(ReorderedConstraints);
+				Measure = MoveTemp(ReorderedMeasure);
+				DmInverse = MoveTemp(ReorderedDmInverse);
+			}
+		}
+
+	protected:
+
+
+		TVec4<FSolverVec3> GetDeterminantDelta(const FSolverParticles& Particles, const FReal Dt, const int32 ElementIndex, const FReal Tol = 1e-3) const
+		{
+			const TVec4<int32>& Constraint = MeshConstraints[ElementIndex];
+
+			const PMatrix<FReal, 3, 3> Fe = F(ElementIndex, Particles);
+			PMatrix<FReal, 3, 3> DmInvT = ElementDmInv(ElementIndex).GetTransposed();
+			TVec4<FSolverVec3> dC2(FSolverVec3((FReal)0.));
+
+			FReal J = Fe.Determinant();
+			if (J - 1 < Tol)
+			{
+				return TVec4<FSolverVec3>(FSolverVec3((FReal)0.));
+			}
+
+			PMatrix<FReal, 3, 3> JFinvT;
+			JFinvT.SetAt(0, 0, Fe.GetAt(1, 1) * Fe.GetAt(2, 2) - Fe.GetAt(2, 1) * Fe.GetAt(1, 2));
+			JFinvT.SetAt(0, 1, Fe.GetAt(2, 0) * Fe.GetAt(1, 2) - Fe.GetAt(1, 0) * Fe.GetAt(2, 2));
+			JFinvT.SetAt(0, 2, Fe.GetAt(1, 0) * Fe.GetAt(2, 1) - Fe.GetAt(2, 0) * Fe.GetAt(1, 1));
+			JFinvT.SetAt(1, 0, Fe.GetAt(2, 1) * Fe.GetAt(0, 2) - Fe.GetAt(0, 1) * Fe.GetAt(2, 2));
+			JFinvT.SetAt(1, 1, Fe.GetAt(0, 0) * Fe.GetAt(2, 2) - Fe.GetAt(2, 0) * Fe.GetAt(0, 2));
+			JFinvT.SetAt(1, 2, Fe.GetAt(2, 0) * Fe.GetAt(0, 1) - Fe.GetAt(0, 0) * Fe.GetAt(2, 1));
+			JFinvT.SetAt(2, 0, Fe.GetAt(0, 1) * Fe.GetAt(1, 2) - Fe.GetAt(1, 1) * Fe.GetAt(0, 2));
+			JFinvT.SetAt(2, 1, Fe.GetAt(1, 0) * Fe.GetAt(0, 2) - Fe.GetAt(0, 0) * Fe.GetAt(1, 2));
+			JFinvT.SetAt(2, 2, Fe.GetAt(0, 0) * Fe.GetAt(1, 1) - Fe.GetAt(1, 0) * Fe.GetAt(0, 1));
+
+			PMatrix<FReal, 3, 3> JinvTDmInvT = DmInvT * JFinvT;
+
+			for (int ie = 0; ie < 3; ie++) {
+				for (int alpha = 0; alpha < 3; alpha++) {
+					dC2[ie + 1][alpha] = JinvTDmInvT.GetAt(alpha, ie);
+				}
+			}
+			for (int alpha = 0; alpha < 3; alpha++) {
+				for (int l = 0; l < 3; l++) {
+					dC2[0][alpha] -= JinvTDmInvT.GetAt(alpha, l);
+				}
+			}
+
+			FReal AlphaTilde = (FReal)2. / (Dt * Dt * Lambda * Measure[ElementIndex]);
+
+			if (bRecordMetric)
+			{
+				HError += J - 1 + AlphaTilde * LambdaArray[2 * ElementIndex + 1];
+				for (int32 ie = 0; ie < 4; ie++)
+				{
+					for (int32 Alpha = 0; Alpha < 3; Alpha++)
+					{
+						GError[MeshConstraints[ElementIndex][ie] * 3 + Alpha] -= dC2[ie][Alpha] * LambdaArray[2 * ElementIndex + 1];
+					}
+				}
+				HErrorArray[2 * ElementIndex + 1] = J - 1 + AlphaTilde * LambdaArray[2 * ElementIndex + 1];
+			}
+			
+
+			FReal DLambda = (1 - J) - AlphaTilde * LambdaArray[2 * ElementIndex + 1];
+
+			FReal Denom = AlphaTilde;
+			for (int i = 0; i < 4; i++)
+			{
+				for (int j = 0; j < 3; j++)
+				{
+					Denom += dC2[i][j] * Particles.InvM(MeshConstraints[ElementIndex][i]) * dC2[i][j];
+				}
+			}
+			DLambda /= Denom;
+			LambdaArray[2 * ElementIndex + 1] += DLambda;
+			TVec4<FSolverVec3> Delta(FSolverVec3((FReal)0.));
+			for (int i = 0; i < 4; i++)
+			{
+				for (int j = 0; j < 3; j++)
+				{
+					Delta[i][j] = Particles.InvM(MeshConstraints[ElementIndex][i]) * dC2[i][j] * DLambda;
+				}
+			}
+			return Delta;
+		}
+
+		TVec4<FSolverVec3> GetPolarDelta(const FSolverParticles& Particles, const FReal Dt, const int32 ElementIndex, const FReal Tol = 1e-3) const
+		{
+
+
+			const PMatrix<FReal, 3, 3> Fe = F(ElementIndex, Particles);
+
+			PMatrix<FReal, 3, 3> Re((FReal)0.), Se((FReal)0.);
+
+			Chaos::PolarDecomposition(Fe, Re, Se);
+
+			FReal C1 = (FReal)0.;
+			for (int i = 0; i < 3; i++)
+			{
+				for (int j = 0; j < 3; j++)
+				{
+					C1 += FMath::Square((Fe - Re).GetAt(i, j));
+				}
+			}
+			C1 = FMath::Sqrt(C1);
+
+			if (C1 < Tol)
+			{
+				return TVec4<FSolverVec3>(FSolverVec3((FReal)0.));
+			}
+
+			TVector<FReal, 81> dRdF((FReal)0.);
+			Chaos::dRdFCorotated(Fe, dRdF);
+
+			PMatrix<FReal, 3, 3> DmInvT = ElementDmInv(ElementIndex).GetTransposed();
+
+			//TODO: deifnitely test the initialization
+			TVec4<FSolverVec3> dC1(FSolverVec3((FReal)0.));
+			//dC1 = dC1dF * dFdX
+			for (int alpha = 0; alpha < 3; alpha++) {
+				for (int l = 0; l < 3; l++) {
+
+					dC1[0][alpha] += (DmInvT * Re).GetAt(alpha, l) - (DmInvT * Fe).GetAt(alpha, l);
+
+				}
+			}
+			for (int ie = 0; ie < 3; ie++) {
+				for (int alpha = 0; alpha < 3; alpha++) {
+					dC1[ie+1][alpha] = (DmInvT * Fe).GetAt(alpha, ie) - (DmInvT * Re).GetAt(alpha, ie);
+				}
+			}
+			//it's really ie-1 here
+			for (int ie = 0; ie < 3; ie++) {
+				for (int alpha = 0; alpha < 3; alpha++) {
+					for (int m = 0; m < 3; m++) {
+						for (int n = 0; n < 3; n++) {
+							for (int j = 0; j < 3; j++) {
+								dC1[ie + 1][alpha] -= (Fe.GetAt(m, n) - Re.GetAt(m, n)) * dRdF[9*(alpha*3+j)+ m * 3 + n] * DmInvT.GetAt(j, ie);
+							}
+						}
+					}
+				}
+			}
+			for (int alpha = 0; alpha < 3; alpha++) {
+				for (int m = 0; m < 3; m++) {
+					for (int n = 0; n < 3; n++) {
+						for (int l = 0; l < 3; l++) {
+							for (int j = 0; j < 3; j++) {
+								dC1[0][alpha] += (Fe.GetAt(m, n) - Re.GetAt(m, n)) * dRdF[9*(alpha*3+j) + m * 3 + n] * DmInvT.GetAt(j, l);
+							}
+						}
+					}
+				}
+			}
+
+			
+			if (C1 != 0)
+			{
+				for (int i = 0; i < 4; i++)
+				{
+					for (int j = 0; j < 3; j++)
+					{
+						dC1[i][j] /= C1;
+					}
+				}
+			}
+
+			FReal AlphaTilde = (FReal)1. / (Dt * Dt * Mu * Measure[ElementIndex]);
+
+			if (bRecordMetric)
+			{
+				HError += C1 + AlphaTilde * LambdaArray[2 * ElementIndex + 0];
+				for (int32 ie = 0; ie < 4; ie++)
+				{
+					for (int32 Alpha = 0; Alpha < 3; Alpha++)
+					{
+						GError[MeshConstraints[ElementIndex][ie] * 3 + Alpha] -= dC1[ie][Alpha] * LambdaArray[2 * ElementIndex + 0];
+					}
+				}
+				HErrorArray[2 * ElementIndex + 0] = C1 + AlphaTilde * LambdaArray[2 * ElementIndex + 0];
+			}
+
+			FReal DLambda = -C1 - AlphaTilde* LambdaArray[2 * ElementIndex + 0];
+
+			FReal Denom = AlphaTilde;
+			for (int i = 0; i < 4; i++)
+			{
+				for (int j = 0; j < 3; j++)
+				{
+					Denom += dC1[i][j] * Particles.InvM(MeshConstraints[ElementIndex][i]) * dC1[i][j];
+				}
+			}
+			DLambda /= Denom;
+			LambdaArray[2 * ElementIndex + 0] += DLambda;
+			TVec4<FSolverVec3> Delta(FSolverVec3((FReal)0.));
+			for (int i = 0; i < 4; i++) 
+			{
+				for (int j = 0; j < 3; j++)
+				{
+					Delta[i][j] = Particles.InvM(MeshConstraints[ElementIndex][i]) * dC1[i][j] * DLambda;
+				}
+			}
+			return Delta;
+
+		}
+
+
+	protected:
+		mutable TArray<FReal> LambdaArray;
+		mutable TArray<FReal> DmInverse;
+
+		//material constants calculated from E:
+		FReal Mu;
+		FReal Lambda;
+		mutable FSolverReal HError;
+		mutable TArray<FSolverReal> HErrorArray;
+		bool bRecordMetric;
+
+		TArray<TVector<int32, 4>> MeshConstraints;
+		mutable TArray<FReal> Measure;
+		FSolverParticles RestParticles;
+		TArray<int32> ConstraintsPerColorStartIndex; // Constraints are ordered so each batch is contiguous. This is ColorNum + 1 length so it can be used as start and end.
+		mutable TArray<FSolverReal> GError;
+};
+
+}  // End namespace Chaos::Softs
+
diff --git a/Engine/Source/Runtime/Experimental/Chaos/Public/Chaos/XPBDCorotatedFiberConstraints.h b/Engine/Source/Runtime/Experimental/Chaos/Public/Chaos/XPBDCorotatedFiberConstraints.h
new file mode 100644
--- /dev/null
+++ b/Engine/Source/Runtime/Experimental/Chaos/Public/Chaos/XPBDCorotatedFiberConstraints.h
@@ -0,0 +1,210 @@
+// Copyright Epic Games, Inc. All Rights Reserved.
+#pragma once
+
+#include "Chaos/PBDSpringConstraintsBase.h"
+#include "ChaosStats.h"
+#include "Chaos/ImplicitQRSVD.h"
+#include "Chaos/GraphColoring.h"
+#include "Chaos/Framework/Parallel.h"
+
+//DECLARE_CYCLE_STAT(TEXT("Chaos XPBD Corotated Constraint"), STAT_XPBD_Spring, STATGROUP_Chaos);
+
+namespace Chaos::Softs
+{
+
+	class FXPBDCorotatedFiberConstraints : public FXPBDCorotatedConstraints
+	{
+
+		typedef FXPBDCorotatedConstraints Base;
+		using Base::MeshConstraints;
+		using Base::LambdaArray;
+		using Base::Measure;
+		using Base::DmInverse;
+
+	public:
+		//this one only accepts tetmesh input and mesh
+		FXPBDCorotatedFiberConstraints(
+			const FSolverParticles& InParticles,
+			int32 ParticleOffset,
+			int32 ParticleCount,
+			const TArray<TVector<int32, 4>>& InMesh,
+			const bool bRecordMetricIn = true,
+			const FSolverReal& EMesh = (FReal)10.0,
+			const FSolverReal& NuMesh = (FReal).3,
+			const FSolverVec3 InFiberDir = FSolverVec3((FSolverReal)1., (FSolverReal)0., (FSolverReal)0.),
+			const FSolverReal InSigmaMax = (FSolverReal)3e5
+		)
+			: Base(InParticles, ParticleOffset, ParticleCount, InMesh, bRecordMetricIn, EMesh, NuMesh), SigmaMax(InSigmaMax), FiberDir(InFiberDir)
+		{
+			LambdaArray.Init((FSolverReal)0., 3 * MeshConstraints.Num());
+		}
+
+		virtual ~FXPBDCorotatedFiberConstraints() {}
+
+		void SetTime(const float Time) const {
+			float FinalTime = 4.0f;
+			float CurrentTime = Time;
+			while (CurrentTime > FinalTime)
+			{
+				CurrentTime -= FinalTime;
+			}
+			AlphaActivation = (FReal)1. - (FReal)4. / FinalTime * FMath::Abs(CurrentTime - FinalTime / (FReal)2.);
+		}
+
+		virtual void ApplyInSerial(FSolverParticles& Particles, const FSolverReal Dt, const int32 ElementIndex) const override
+		{
+			TVec4<FSolverVec3> PolarDelta = Base::GetPolarDelta(Particles, Dt, ElementIndex);
+
+			for (int i = 0; i < 4; i++)
+			{
+				Particles.P(MeshConstraints[ElementIndex][i]) += PolarDelta[i];
+			}
+
+			TVec4<FSolverVec3> DetDelta = Base::GetDeterminantDelta(Particles, Dt, ElementIndex);
+
+			for (int i = 0; i < 4; i++)
+			{
+				Particles.P(MeshConstraints[ElementIndex][i]) += DetDelta[i];
+			}
+
+			TVec4<FSolverVec3> FiberDelta = GetFiberDelta(Particles, Dt, ElementIndex);
+
+			for (int i = 0; i < 4; i++)
+			{
+				Particles.P(MeshConstraints[ElementIndex][i]) += FiberDelta[i];
+			}
+			
+
+		}
+
+
+
+	private:
+
+		TVec4<FSolverVec3> GetFiberDelta(const FSolverParticles& Particles, const FSolverReal Dt, const int32 ElementIndex, const FSolverReal Tol = 1e-3) const
+		{
+			const PMatrix<FReal, 3, 3> Fe = F(ElementIndex, Particles);
+
+			PMatrix<FReal, 3, 3> Re((FReal)0.), Se((FReal)0.);
+
+			Chaos::PolarDecomposition(Fe, Re, Se);
+
+			// l: fiber stretch, f = (vTCv)^(1/2)
+			FSolverVec3 FeV = Fe.GetTransposed() * FiberDir;
+			FSolverVec3 DmInverseV = ElementDmInv(ElementIndex).GetTransposed() * FiberDir;
+			FSolverReal L = FeV.Size();
+			TVec4<FSolverVec3> dLdX(FSolverVec3((FReal)0.));
+			for (int32 alpha = 0; alpha < 3; alpha++)
+			{
+				for (int32 s = 0; s < 3; s++)
+				{
+					dLdX[0][alpha] -= FeV[alpha] * DmInverseV[s] / L;
+				}
+
+			}
+			for (int32 ie = 1; ie < 4; ie++)
+			{
+				for (int32 alpha = 0; alpha < 3; alpha++)
+				{
+					dLdX[ie][alpha] = FeV[alpha] * DmInverseV[ie - 1] / L;
+				}
+			}
+
+			FReal LambdaOFL = (FReal)1.4;
+			FReal P1 = (FReal)0.05;
+			FReal P2 = (FReal)6.6;
+			FReal FpIntegral = (FReal)0.;
+			FReal FaIntegral = (FReal)0.;
+
+			FReal C3 = (FReal)0.;
+			FReal AlphaTilde = LambdaOFL / (SigmaMax * Dt * Dt * Measure[ElementIndex]);
+			FReal dFpdL = (FReal)0.;
+			FReal dFadL = (FReal)0.;
+
+			if (L > LambdaOFL)
+			{
+				FpIntegral = P1 * LambdaOFL / P2 * FMath::Exp(P2 * (L / LambdaOFL - (FReal)1.)) - P1 * (L - LambdaOFL);
+				dFpdL = P1 * FMath::Exp(P2 * (L / LambdaOFL - (FReal)1.)) - P1;
+			}
+			if (L > (FReal)0.4 * LambdaOFL && L < (FReal)0.6 * LambdaOFL)
+			{
+				FaIntegral = (FReal)3. * LambdaOFL * FMath::Pow((L / LambdaOFL - (FReal)0.4), 3);
+				dFadL = (FReal)9. * FMath::Pow((L / LambdaOFL - (FReal)0.4), 2);
+			}
+			else if (L >= (FReal)0.6 * LambdaOFL && L <= (FReal)1.4 * LambdaOFL)
+			{
+				FaIntegral = (FReal)3. * LambdaOFL * (FReal)0.008 + L - (FReal)4. / (FReal)3. * LambdaOFL * FMath::Pow(L / LambdaOFL - (FReal)1., 3) - (FReal)0.6 * LambdaOFL - (FReal)4. / (FReal)3. * LambdaOFL * FMath::Pow((FReal)0.4, 3);
+				dFadL = (FReal)1. - (FReal)4. * FMath::Pow(L / LambdaOFL - (FReal)1., 2);
+			}
+			else if (L > (FReal)1.4 * LambdaOFL && L <= (FReal)1.6 * LambdaOFL)
+			{
+				FaIntegral = (FReal)3. * LambdaOFL * (FReal)0.008 + (FReal)0.8 * LambdaOFL - (FReal)8. / (FReal)3. * LambdaOFL * FMath::Pow((FReal)0.4, 3) + (FReal)3. * LambdaOFL * FMath::Pow(L / LambdaOFL - (FReal)1.6, 3) + (FReal)3. * LambdaOFL * (FReal)0.008;
+				dFadL = (FReal)9. * FMath::Pow((L / LambdaOFL - (FReal)1.6), 2);
+
+			}
+			else if (L > (FReal)1.6 * LambdaOFL)
+			{
+				FaIntegral = (FReal)3. * LambdaOFL * (FReal)0.008 + (FReal)0.8 * LambdaOFL - (FReal)8. / (FReal)3. * LambdaOFL * FMath::Pow((FReal)0.4, 3) + (FReal)3. * LambdaOFL * (FReal)0.008;
+				dFadL = (FReal)0.;
+			}
+
+
+			C3 = FMath::Sqrt(FpIntegral + AlphaActivation * FaIntegral);
+
+
+			FReal dC3dL = (FReal)0.5 * (dFpdL + AlphaActivation * dFadL);
+			if (C3 == 0)
+			{
+				return TVec4<FSolverVec3>(FSolverVec3((FSolverReal)0.));
+			}
+			else
+			{
+				dC3dL /= C3;
+			}
+
+			TVec4<FSolverVec3> dC3(FSolverVec3((FReal)0.));
+			for (int32 i = 0; i < 4; i++)
+			{
+				for (int32 j = 0; j < 3; j++)
+				{
+					dC3[i][j] = dC3dL * dLdX[i][j];
+				}
+			}
+
+			FReal DLambda = -C3 - AlphaTilde * LambdaArray[2 * ElementIndex + 2];
+
+			FReal Denom = AlphaTilde;
+			for (int i = 0; i < 4; i++)
+			{
+				for (int j = 0; j < 3; j++)
+				{
+					Denom += dC3[i][j] * Particles.InvM(MeshConstraints[ElementIndex][i]) * dC3[i][j];
+				}
+			}
+			DLambda /= Denom;
+			LambdaArray[2 * ElementIndex + 2] += DLambda;
+			TVec4<FSolverVec3> Delta(FSolverVec3((FReal)0.));
+			for (int i = 0; i < 4; i++)
+			{
+				for (int j = 0; j < 3; j++)
+				{
+					Delta[i][j] = Particles.InvM(MeshConstraints[ElementIndex][i]) * dC3[i][j] * DLambda;
+				}
+			}
+			return Delta;
+
+
+
+		}
+
+	private:
+
+		//material constants calculated from E:
+		FReal SigmaMax;
+		mutable FReal AlphaActivation;
+		FSolverVec3 FiberDir;
+
+	};
+
+}  // End namespace Chaos::Softs
+
diff --git a/Engine/Source/Runtime/Experimental/ChaosCore/Public/Chaos/ImplicitQRSVD.h b/Engine/Source/Runtime/Experimental/ChaosCore/Public/Chaos/ImplicitQRSVD.h
--- a/Engine/Source/Runtime/Experimental/ChaosCore/Public/Chaos/ImplicitQRSVD.h
+++ b/Engine/Source/Runtime/Experimental/ChaosCore/Public/Chaos/ImplicitQRSVD.h
@@ -85,7 +85,7 @@ namespace Chaos {
 		/**
 		  Fill the R with the entries of this rotation
 			*/
-		template <class T>
+		//template <class T>
 		inline void Fill(const PMatrix<T, 2, 2>& R) const {
 			PMatrix<T, 2, 2>& A = const_cast<PMatrix<T, 2, 2>&>(R);
 			A = PMatrix<T, 2, 2>((FReal)1., (FReal)0., (FReal)0., (FReal)1.);
@@ -102,7 +102,7 @@ namespace Chaos {
 			0  0  1
 			It only affects row i and row k of A.
 			*/
-		template <class T>
+		//template <class T>
 		inline void RowRotation(PMatrix<T, 2, 2>& A) const {
 			for (int j = 0; j < 2; j++) {
 				T tau1 = A.M[j*2+rowi];
@@ -112,7 +112,7 @@ namespace Chaos {
 			}
 		}
 
-		template <class T>
+		//template <class T>
 		inline void RowRotation(PMatrix<T, 3, 3>& A) const {
 			for (int j = 0; j < 3; j++) {
 				T tau1 = A.M[j][rowi];
@@ -129,7 +129,7 @@ namespace Chaos {
 			0  0  1
 			It only affects column i and column k of A.
 			*/
-		template <class T>
+		//template <class T>
 		inline void ColumnRotation(PMatrix<T, 2, 2>& A) const {
 			for (int j = 0; j < 2; j++) {
 				T tau1 = A.M[rowi*2+j];
@@ -139,7 +139,7 @@ namespace Chaos {
 			}
 		}
 
-		template <class T>
+		//template <class T>
 		inline void ColumnRotation(PMatrix<T, 3, 3>& A) const {
 			for (int j = 0; j < 3; j++) {
 				T tau1 = A.M[rowi][j];
@@ -233,8 +233,8 @@ namespace Chaos {
 	  */
 	template <class T>
 	inline void MakeUpperBidiag(PMatrix<T, 3, 3>& H, PMatrix<T, 3, 3>& U, PMatrix<T, 3, 3>& V) {
-		U = PMatrix<T, 3, 3>::Identity;
-		V = PMatrix<T, 3, 3>::Identity;
+		U = PMatrix<T, 3, 3>(1, 1, 1);
+		V = PMatrix<T, 3, 3>(1, 1, 1);
 
 		/**
 		  Reduce H to of form
@@ -261,8 +261,8 @@ namespace Chaos {
 	  */
 	template <class T>
 	inline void MakeLambdaShape(PMatrix<T, 3, 3>& H, PMatrix<T, 3, 3>& U, PMatrix<T, 3, 3>& V) {
-		U = PMatrix<T, 3, 3>::Identity;
-		V = PMatrix<T, 3, 3>::Identity;
+		U = PMatrix<T, 3, 3>(1, 1, 1);
+		V = PMatrix<T, 3, 3>(1, 1, 1);
 
 		/**
 		  Reduce H to of form
@@ -598,8 +598,8 @@ namespace Chaos {
 		T tol = 1024 * TMathUtilConstants<T>::Epsilon) {
 
 		PMatrix<T, 3, 3> B = A;
-		U = PMatrix<T, 3, 3>::Identity;
-		V = PMatrix<T, 3, 3>::Identity;
+		U = PMatrix<T, 3, 3>(1, 1, 1);
+		V = PMatrix<T, 3, 3>(1, 1, 1);
 
 		MakeUpperBidiag(B, U, V);
 
@@ -772,5 +772,97 @@ namespace Chaos {
 		S_Sym = V.GetTransposed() * PMatrix<T, 3, 3>(sigma) * V;
 	}
 
+	// 3X3 version of dRdF
+	// indexing of diff: dRdF[9*(i*3+j)+(m*3+n)] = dR.GetAt(m, n) /dF.GetAt (i, j):
+	template <class T>
+	void dRdFCorotated(const PMatrix<T, 3, 3>& F, TVector<T, 81>& dRdF) {
+		PMatrix<T, 3, 3> R, S, Dinv;
+		PolarDecomposition(F, R, S);
+		PMatrix<T, 3, 3> D = (S.M[0][0] + S.M[1][1] + S.M[2][2]) * PMatrix<T, 3, 3>(1, 1, 1) - S;
+		Dinv = D.Inverse();
+
+		dRdF[0 * 9 + 0] = (-R.GetAt(0, 1) * Dinv.GetAt(2, 1) * R.GetAt(0, 2) + R.GetAt(0, 1) * Dinv.GetAt(2, 2) * R.GetAt(0, 1) + R.GetAt(0, 2) * Dinv.GetAt(1, 1) * R.GetAt(0, 2) - R.GetAt(0, 2) * Dinv.GetAt(1, 2) * R.GetAt(0, 1));
+		dRdF[1 * 9 + 0] = (R.GetAt(0, 1) * Dinv.GetAt(2, 0) * R.GetAt(0, 2) - R.GetAt(0, 1) * Dinv.GetAt(2, 2) * R.GetAt(0, 0) - R.GetAt(0, 2) * Dinv.GetAt(1, 0) * R.GetAt(0, 2) + R.GetAt(0, 2) * Dinv.GetAt(1, 2) * R.GetAt(0, 0));
+		dRdF[2 * 9 + 0] = (-R.GetAt(0, 1) * Dinv.GetAt(2, 0) * R.GetAt(0, 1) + R.GetAt(0, 1) * Dinv.GetAt(2, 1) * R.GetAt(0, 0) + R.GetAt(0, 2) * Dinv.GetAt(1, 0) * R.GetAt(0, 1) - R.GetAt(0, 2) * Dinv.GetAt(1, 1) * R.GetAt(0, 0));
+		dRdF[3 * 9 + 0] = (-R.GetAt(0, 1) * Dinv.GetAt(2, 1) * R.GetAt(1, 2) + R.GetAt(0, 1) * Dinv.GetAt(2, 2) * R.GetAt(1, 1) + R.GetAt(0, 2) * Dinv.GetAt(1, 1) * R.GetAt(1, 2) - R.GetAt(0, 2) * Dinv.GetAt(1, 2) * R.GetAt(1, 1));
+		dRdF[4 * 9 + 0] = (R.GetAt(0, 1) * Dinv.GetAt(2, 0) * R.GetAt(1, 2) - R.GetAt(0, 1) * Dinv.GetAt(2, 2) * R.GetAt(1, 0) - R.GetAt(0, 2) * Dinv.GetAt(1, 0) * R.GetAt(1, 2) + R.GetAt(0, 2) * Dinv.GetAt(1, 2) * R.GetAt(1, 0));
+		dRdF[5 * 9 + 0] = (-R.GetAt(0, 1) * Dinv.GetAt(2, 0) * R.GetAt(1, 1) + R.GetAt(0, 1) *Dinv.GetAt(2, 1) * R.GetAt(1, 0) + R.GetAt(0, 2) * Dinv.GetAt(1, 0) * R.GetAt(1, 1) - R.GetAt(0, 2) * Dinv.GetAt(1, 1) * R.GetAt(1, 0));
+		dRdF[6 * 9 + 0] = (-R.GetAt(0, 1) * Dinv.GetAt(2, 1) * R.GetAt(2, 2) + R.GetAt(0, 1) * Dinv.GetAt(2, 2) * R.GetAt(2, 1) + R.GetAt(0, 2) * Dinv.GetAt(1, 1) * R.GetAt(2, 2) - R.GetAt(0, 2) * Dinv.GetAt(1, 2) * R.GetAt(2, 1));
+		dRdF[7 * 9 + 0] = (R.GetAt(0, 1) * Dinv.GetAt(2, 0) * R.GetAt(2, 2) - R.GetAt(0, 1) * Dinv.GetAt(2, 2) * R.GetAt(2, 0) - R.GetAt(0, 2) * Dinv.GetAt(1, 0) * R.GetAt(2, 2) + R.GetAt(0, 2) * Dinv.GetAt(1, 2) * R.GetAt(2, 0));
+		dRdF[8 * 9 + 0] = (-R.GetAt(0, 1) * Dinv.GetAt(2, 0) * R.GetAt(2, 1) + R.GetAt(0, 1) * Dinv.GetAt(2, 1) * R.GetAt(2, 0) + R.GetAt(0, 2) * Dinv.GetAt(1, 0) * R.GetAt(2, 1) - R.GetAt(0, 2) * Dinv.GetAt(1, 1) * R.GetAt(2, 0));
+		dRdF[0 * 9 + 1] = (R.GetAt(0, 0) * Dinv.GetAt(2, 1) * R.GetAt(0, 2) - R.GetAt(0, 0) * Dinv.GetAt(2, 2) * R.GetAt(0, 1) - R.GetAt(0, 2) * Dinv.GetAt(0, 1) * R.GetAt(0, 2) + R.GetAt(0, 2) * Dinv.GetAt(0, 2) * R.GetAt(0, 1));
+		dRdF[1 * 9 + 1] = (-R.GetAt(0, 0) * Dinv.GetAt(2, 0) * R.GetAt(0, 2) + R.GetAt(0, 0) * Dinv.GetAt(2, 2) * R.GetAt(0, 0) + R.GetAt(0, 2) * Dinv.GetAt(0, 0) * R.GetAt(0, 2) - R.GetAt(0, 2) * Dinv.GetAt(0, 2) * R.GetAt(0, 0));
+		dRdF[2 * 9 + 1] = (R.GetAt(0, 0) * Dinv.GetAt(2, 0) * R.GetAt(0, 1) - R.GetAt(0, 0) * Dinv.GetAt(2, 1) * R.GetAt(0, 0) - R.GetAt(0, 2) * Dinv.GetAt(0, 0) * R.GetAt(0, 1) + R.GetAt(0, 2) * Dinv.GetAt(0, 1) * R.GetAt(0, 0));
+		dRdF[3 * 9 + 1] = (R.GetAt(0, 0) * Dinv.GetAt(2, 1) * R.GetAt(1, 2) - R.GetAt(0, 0) * Dinv.GetAt(2, 2) * R.GetAt(1, 1) - R.GetAt(0, 2) * Dinv.GetAt(0, 1) * R.GetAt(1, 2) + R.GetAt(0, 2) * Dinv.GetAt(0, 2) * R.GetAt(1, 1));
+		dRdF[4 * 9 + 1] = (-R.GetAt(0, 0) * Dinv.GetAt(2, 0) * R.GetAt(1, 2) + R.GetAt(0, 0) * Dinv.GetAt(2, 2) * R.GetAt(1, 0) + R.GetAt(0, 2) * Dinv.GetAt(0, 0) * R.GetAt(1, 2) - R.GetAt(0, 2) * Dinv.GetAt(0, 2) * R.GetAt(1, 0));
+		dRdF[5 * 9 + 1] = (R.GetAt(0, 0) * Dinv.GetAt(2, 0) * R.GetAt(1, 1) - R.GetAt(0, 0) * Dinv.GetAt(2, 1) * R.GetAt(1, 0) - R.GetAt(0, 2) * Dinv.GetAt(0, 0) * R.GetAt(1, 1) + R.GetAt(0, 2) * Dinv.GetAt(0, 1) * R.GetAt(1, 0));
+		dRdF[6 * 9 + 1] = (R.GetAt(0, 0) * Dinv.GetAt(2, 1) * R.GetAt(2, 2) - R.GetAt(0, 0) * Dinv.GetAt(2, 2) * R.GetAt(2, 1) - R.GetAt(0, 2) * Dinv.GetAt(0, 1) * R.GetAt(2, 2) + R.GetAt(0, 2) * Dinv.GetAt(0, 2) * R.GetAt(2, 1));
+		dRdF[7 * 9 + 1] = (-R.GetAt(0, 0) * Dinv.GetAt(2, 0) * R.GetAt(2, 2) + R.GetAt(0, 0) * Dinv.GetAt(2, 2) * R.GetAt(2, 0) + R.GetAt(0, 2) * Dinv.GetAt(0, 0) * R.GetAt(2, 2) - R.GetAt(0, 2) * Dinv.GetAt(0, 2) * R.GetAt(2, 0));
+		dRdF[8 * 9 + 1] = (R.GetAt(0, 0) * Dinv.GetAt(2, 0) * R.GetAt(2, 1) - R.GetAt(0, 0) * Dinv.GetAt(2, 1) * R.GetAt(2, 0) - R.GetAt(0, 2) * Dinv.GetAt(0, 0) * R.GetAt(2, 1) + R.GetAt(0, 2) * Dinv.GetAt(0, 1) * R.GetAt(2, 0));
+		dRdF[0 * 9 + 2] = (-R.GetAt(0, 0) * Dinv.GetAt(1, 1) * R.GetAt(0, 2) + R.GetAt(0, 0) * Dinv.GetAt(1, 2) * R.GetAt(0, 1) + R.GetAt(0, 1) * Dinv.GetAt(0, 1) * R.GetAt(0, 2) - R.GetAt(0, 1) * Dinv.GetAt(0, 2) * R.GetAt(0, 1));
+		dRdF[1 * 9 + 2] = (R.GetAt(0, 0) * Dinv.GetAt(1, 0) * R.GetAt(0, 2) - R.GetAt(0, 0) * Dinv.GetAt(1, 2) * R.GetAt(0, 0) - R.GetAt(0, 1) * Dinv.GetAt(0, 0) * R.GetAt(0, 2) + R.GetAt(0, 1) * Dinv.GetAt(0, 2) * R.GetAt(0, 0));
+		dRdF[2 * 9 + 2] = (-R.GetAt(0, 0) * Dinv.GetAt(1, 0) * R.GetAt(0, 1) + R.GetAt(0, 0) * Dinv.GetAt(1, 1) * R.GetAt(0, 0) + R.GetAt(0, 1) * Dinv.GetAt(0, 0) * R.GetAt(0, 1) - R.GetAt(0, 1) * Dinv.GetAt(0, 1) * R.GetAt(0, 0));
+		dRdF[3 * 9 + 2] = (-R.GetAt(0, 0) * Dinv.GetAt(1, 1) * R.GetAt(1, 2) + R.GetAt(0, 0) * Dinv.GetAt(1, 2) * R.GetAt(1, 1) + R.GetAt(0, 1) * Dinv.GetAt(0, 1) * R.GetAt(1, 2) - R.GetAt(0, 1) * Dinv.GetAt(0, 2) * R.GetAt(1, 1));
+		dRdF[4 * 9 + 2] = (R.GetAt(0, 0) * Dinv.GetAt(1, 0) * R.GetAt(1, 2) - R.GetAt(0, 0) * Dinv.GetAt(1, 2) * R.GetAt(1, 0) - R.GetAt(0, 1) * Dinv.GetAt(0, 0) * R.GetAt(1, 2) + R.GetAt(0, 1) * Dinv.GetAt(0, 2) * R.GetAt(1, 0));
+		dRdF[5 * 9 + 2] = (-R.GetAt(0, 0) * Dinv.GetAt(1, 0) * R.GetAt(1, 1) + R.GetAt(0, 0) * Dinv.GetAt(1, 1) * R.GetAt(1, 0) + R.GetAt(0, 1) * Dinv.GetAt(0, 0) * R.GetAt(1, 1) - R.GetAt(0, 1) * Dinv.GetAt(0, 1) * R.GetAt(1, 0));
+		dRdF[6 * 9 + 2] = (-R.GetAt(0, 0) * Dinv.GetAt(1, 1) * R.GetAt(2, 2) + R.GetAt(0, 0) * Dinv.GetAt(1, 2) * R.GetAt(2, 1) + R.GetAt(0, 1) * Dinv.GetAt(0, 1) * R.GetAt(2, 2) - R.GetAt(0, 1) * Dinv.GetAt(0, 2) * R.GetAt(2, 1));
+		dRdF[7 * 9 + 2] = (R.GetAt(0, 0) * Dinv.GetAt(1, 0) * R.GetAt(2, 2) - R.GetAt(0, 0) * Dinv.GetAt(1, 2) * R.GetAt(2, 0) - R.GetAt(0, 1) * Dinv.GetAt(0, 0) * R.GetAt(2, 2) + R.GetAt(0, 1) * Dinv.GetAt(0, 2) * R.GetAt(2, 0));
+		dRdF[8 * 9 + 2] = (-R.GetAt(0, 0) * Dinv.GetAt(1, 0) * R.GetAt(2, 1) + R.GetAt(0, 0) * Dinv.GetAt(1, 1) * R.GetAt(2, 0) + R.GetAt(0, 1) * Dinv.GetAt(0, 0) * R.GetAt(2, 1) - R.GetAt(0, 1) * Dinv.GetAt(0, 1) * R.GetAt(2, 0));
+		dRdF[0 * 9 + 3] = (-R.GetAt(1, 1) * Dinv.GetAt(2, 1) * R.GetAt(0, 2) + R.GetAt(1, 1) * Dinv.GetAt(2, 2) * R.GetAt(0, 1) + R.GetAt(1, 2) * Dinv.GetAt(1, 1) * R.GetAt(0, 2) - R.GetAt(1, 2) * Dinv.GetAt(1, 2) * R.GetAt(0, 1));
+		dRdF[1 * 9 + 3] = (R.GetAt(1, 1) * Dinv.GetAt(2, 0) * R.GetAt(0, 2) - R.GetAt(1, 1) * Dinv.GetAt(2, 2) * R.GetAt(0, 0) - R.GetAt(1, 2) * Dinv.GetAt(1, 0) * R.GetAt(0, 2) + R.GetAt(1, 2) * Dinv.GetAt(1, 2) * R.GetAt(0, 0));
+		dRdF[2 * 9 + 3] = (-R.GetAt(1, 1) * Dinv.GetAt(2, 0) * R.GetAt(0, 1) + R.GetAt(1, 1) * Dinv.GetAt(2, 1) * R.GetAt(0, 0) + R.GetAt(1, 2) * Dinv.GetAt(1, 0) * R.GetAt(0, 1) - R.GetAt(1, 2) * Dinv.GetAt(1, 1) * R.GetAt(0, 0));
+		dRdF[3 * 9 + 3] = (-R.GetAt(1, 1) * Dinv.GetAt(2, 1) * R.GetAt(1, 2) + R.GetAt(1, 1) * Dinv.GetAt(2, 2) * R.GetAt(1, 1) + R.GetAt(1, 2) * Dinv.GetAt(1, 1) * R.GetAt(1, 2) - R.GetAt(1, 2) * Dinv.GetAt(1, 2) * R.GetAt(1, 1));
+		dRdF[4 * 9 + 3] = (R.GetAt(1, 1) * Dinv.GetAt(2, 0) * R.GetAt(1, 2) - R.GetAt(1, 1) * Dinv.GetAt(2, 2) * R.GetAt(1, 0) - R.GetAt(1, 2) * Dinv.GetAt(1, 0) * R.GetAt(1, 2) + R.GetAt(1, 2) * Dinv.GetAt(1, 2) * R.GetAt(1, 0));
+		dRdF[5 * 9 + 3] = (-R.GetAt(1, 1) * Dinv.GetAt(2, 0) * R.GetAt(1, 1) + R.GetAt(1, 1) * Dinv.GetAt(2, 1) * R.GetAt(1, 0) + R.GetAt(1, 2) * Dinv.GetAt(1, 0) * R.GetAt(1, 1) - R.GetAt(1, 2) * Dinv.GetAt(1, 1) * R.GetAt(1, 0));
+		dRdF[6 * 9 + 3] = (-R.GetAt(1, 1) * Dinv.GetAt(2, 1) * R.GetAt(2, 2) + R.GetAt(1, 1) * Dinv.GetAt(2, 2) * R.GetAt(2, 1) + R.GetAt(1, 2) * Dinv.GetAt(1, 1) * R.GetAt(2, 2) - R.GetAt(1, 2) * Dinv.GetAt(1, 2) * R.GetAt(2, 1));
+		dRdF[7 * 9 + 3] = (R.GetAt(1, 1) * Dinv.GetAt(2, 0) * R.GetAt(2, 2) - R.GetAt(1, 1) * Dinv.GetAt(2, 2) * R.GetAt(2, 0) - R.GetAt(1, 2) * Dinv.GetAt(1, 0) * R.GetAt(2, 2) + R.GetAt(1, 2) * Dinv.GetAt(1, 2) * R.GetAt(2, 0));
+		dRdF[8 * 9 + 3] = (-R.GetAt(1, 1) * Dinv.GetAt(2, 0) * R.GetAt(2, 1) + R.GetAt(1, 1) * Dinv.GetAt(2, 1) * R.GetAt(2, 0) + R.GetAt(1, 2) * Dinv.GetAt(1, 0) * R.GetAt(2, 1) - R.GetAt(1, 2) * Dinv.GetAt(1, 1) * R.GetAt(2, 0));
+		dRdF[0 * 9 + 4] = (R.GetAt(1, 0) * Dinv.GetAt(2, 1) * R.GetAt(0, 2) - R.GetAt(1, 0) * Dinv.GetAt(2, 2) * R.GetAt(0, 1) - R.GetAt(1, 2) * Dinv.GetAt(0, 1) * R.GetAt(0, 2) + R.GetAt(1, 2) * Dinv.GetAt(0, 2) * R.GetAt(0, 1));
+		dRdF[1 * 9 + 4] = (-R.GetAt(1, 0) * Dinv.GetAt(2, 0) * R.GetAt(0, 2) + R.GetAt(1, 0) * Dinv.GetAt(2, 2) * R.GetAt(0, 0) + R.GetAt(1, 2) * Dinv.GetAt(0, 0) * R.GetAt(0, 2) - R.GetAt(1, 2) * Dinv.GetAt(0, 2) * R.GetAt(0, 0));
+		dRdF[2 * 9 + 4] = (R.GetAt(1, 0) * Dinv.GetAt(2, 0) * R.GetAt(0, 1) - R.GetAt(1, 0) * Dinv.GetAt(2, 1) * R.GetAt(0, 0) - R.GetAt(1, 2) * Dinv.GetAt(0, 0) * R.GetAt(0, 1) + R.GetAt(1, 2) * Dinv.GetAt(0, 1) * R.GetAt(0, 0));
+		dRdF[3 * 9 + 4] = (R.GetAt(1, 0) * Dinv.GetAt(2, 1) * R.GetAt(1, 2) - R.GetAt(1, 0) * Dinv.GetAt(2, 2) * R.GetAt(1, 1) - R.GetAt(1, 2) * Dinv.GetAt(0, 1) * R.GetAt(1, 2) + R.GetAt(1, 2) * Dinv.GetAt(0, 2) * R.GetAt(1, 1));
+		dRdF[4 * 9 + 4] = (-R.GetAt(1, 0) * Dinv.GetAt(2, 0) * R.GetAt(1, 2) + R.GetAt(1, 0) * Dinv.GetAt(2, 2) * R.GetAt(1, 0) + R.GetAt(1, 2) * Dinv.GetAt(0, 0) * R.GetAt(1, 2) - R.GetAt(1, 2) * Dinv.GetAt(0, 2) * R.GetAt(1, 0));
+		dRdF[5 * 9 + 4] = (R.GetAt(1, 0) * Dinv.GetAt(2, 0) * R.GetAt(1, 1) - R.GetAt(1, 0) * Dinv.GetAt(2, 1) * R.GetAt(1, 0) - R.GetAt(1, 2) * Dinv.GetAt(0, 0) * R.GetAt(1, 1) + R.GetAt(1, 2) * Dinv.GetAt(0, 1) * R.GetAt(1, 0));
+		dRdF[6 * 9 + 4] = (R.GetAt(1, 0) * Dinv.GetAt(2, 1) * R.GetAt(2, 2) - R.GetAt(1, 0) * Dinv.GetAt(2, 2) * R.GetAt(2, 1) - R.GetAt(1, 2) * Dinv.GetAt(0, 1) * R.GetAt(2, 2) + R.GetAt(1, 2) * Dinv.GetAt(0, 2) * R.GetAt(2, 1));
+		dRdF[7 * 9 + 4] = (-R.GetAt(1, 0) * Dinv.GetAt(2, 0) * R.GetAt(2, 2) + R.GetAt(1, 0) * Dinv.GetAt(2, 2) * R.GetAt(2, 0) + R.GetAt(1, 2) * Dinv.GetAt(0, 0) * R.GetAt(2, 2) - R.GetAt(1, 2) * Dinv.GetAt(0, 2) * R.GetAt(2, 0));
+		dRdF[8 * 9 + 4] = (R.GetAt(1, 0) * Dinv.GetAt(2, 0) * R.GetAt(2, 1) - R.GetAt(1, 0) * Dinv.GetAt(2, 1) * R.GetAt(2, 0) - R.GetAt(1, 2) * Dinv.GetAt(0, 0) * R.GetAt(2, 1) + R.GetAt(1, 2) * Dinv.GetAt(0, 1) * R.GetAt(2, 0));
+		dRdF[0 * 9 + 5] = (-R.GetAt(1, 0) * Dinv.GetAt(1, 1) * R.GetAt(0, 2) + R.GetAt(1, 0) * Dinv.GetAt(1, 2) * R.GetAt(0, 1) + R.GetAt(1, 1) * Dinv.GetAt(0, 1) * R.GetAt(0, 2) - R.GetAt(1, 1) * Dinv.GetAt(0, 2) * R.GetAt(0, 1));
+		dRdF[1 * 9 + 5] = (R.GetAt(1, 0) * Dinv.GetAt(1, 0) * R.GetAt(0, 2) - R.GetAt(1, 0) * Dinv.GetAt(1, 2) * R.GetAt(0, 0) - R.GetAt(1, 1) * Dinv.GetAt(0, 0) * R.GetAt(0, 2) + R.GetAt(1, 1) * Dinv.GetAt(0, 2) * R.GetAt(0, 0));
+		dRdF[2 * 9 + 5] = (-R.GetAt(1, 0) * Dinv.GetAt(1, 0) * R.GetAt(0, 1) + R.GetAt(1, 0) * Dinv.GetAt(1, 1) * R.GetAt(0, 0) + R.GetAt(1, 1) * Dinv.GetAt(0, 0) * R.GetAt(0, 1) - R.GetAt(1, 1) * Dinv.GetAt(0, 1) * R.GetAt(0, 0));
+		dRdF[3 * 9 + 5] = (-R.GetAt(1, 0) * Dinv.GetAt(1, 1) * R.GetAt(1, 2) + R.GetAt(1, 0) * Dinv.GetAt(1, 2) * R.GetAt(1, 1) + R.GetAt(1, 1) * Dinv.GetAt(0, 1) * R.GetAt(1, 2) - R.GetAt(1, 1) * Dinv.GetAt(0, 2) * R.GetAt(1, 1));
+		dRdF[4 * 9 + 5] = (R.GetAt(1, 0) * Dinv.GetAt(1, 0) * R.GetAt(1, 2) - R.GetAt(1, 0) * Dinv.GetAt(1, 2) * R.GetAt(1, 0) - R.GetAt(1, 1) * Dinv.GetAt(0, 0) * R.GetAt(1, 2) + R.GetAt(1, 1) * Dinv.GetAt(0, 2) * R.GetAt(1, 0));
+		dRdF[5 * 9 + 5] = (-R.GetAt(1, 0) * Dinv.GetAt(1, 0) * R.GetAt(1, 1) + R.GetAt(1, 0) * Dinv.GetAt(1, 1) * R.GetAt(1, 0) + R.GetAt(1, 1) * Dinv.GetAt(0, 0) * R.GetAt(1, 1) - R.GetAt(1, 1) * Dinv.GetAt(0, 1) * R.GetAt(1, 0));
+		dRdF[6 * 9 + 5] = (-R.GetAt(1, 0) * Dinv.GetAt(1, 1) * R.GetAt(2, 2) + R.GetAt(1, 0) * Dinv.GetAt(1, 2) * R.GetAt(2, 1) + R.GetAt(1, 1) * Dinv.GetAt(0, 1) * R.GetAt(2, 2) - R.GetAt(1, 1) * Dinv.GetAt(0, 2) * R.GetAt(2, 1));
+		dRdF[7 * 9 + 5] = (R.GetAt(1, 0) * Dinv.GetAt(1, 0) * R.GetAt(2, 2) - R.GetAt(1, 0) * Dinv.GetAt(1, 2) * R.GetAt(2, 0) - R.GetAt(1, 1) * Dinv.GetAt(0, 0) * R.GetAt(2, 2) + R.GetAt(1, 1) * Dinv.GetAt(0, 2) * R.GetAt(2, 0));
+		dRdF[8 * 9 + 5] = (-R.GetAt(1, 0) * Dinv.GetAt(1, 0) * R.GetAt(2, 1) + R.GetAt(1, 0) * Dinv.GetAt(1, 1) * R.GetAt(2, 0) + R.GetAt(1, 1) * Dinv.GetAt(0, 0) * R.GetAt(2, 1) - R.GetAt(1, 1) * Dinv.GetAt(0, 1) * R.GetAt(2, 0));
+		dRdF[0 * 9 + 6] = (-R.GetAt(2, 1) * Dinv.GetAt(2, 1) * R.GetAt(0, 2) + R.GetAt(2, 1) * Dinv.GetAt(2, 2) * R.GetAt(0, 1) + R.GetAt(2, 2) * Dinv.GetAt(1, 1) * R.GetAt(0, 2) - R.GetAt(2, 2) * Dinv.GetAt(1, 2) * R.GetAt(0, 1));
+		dRdF[1 * 9 + 6] = (R.GetAt(2, 1) * Dinv.GetAt(2, 0) * R.GetAt(0, 2) - R.GetAt(2, 1) * Dinv.GetAt(2, 2) * R.GetAt(0, 0) - R.GetAt(2, 2) * Dinv.GetAt(1, 0) * R.GetAt(0, 2) + R.GetAt(2, 2) * Dinv.GetAt(1, 2) * R.GetAt(0, 0));
+		dRdF[2 * 9 + 6] = (-R.GetAt(2, 1) * Dinv.GetAt(2, 0) * R.GetAt(0, 1) + R.GetAt(2, 1) * Dinv.GetAt(2, 1) * R.GetAt(0, 0) + R.GetAt(2, 2) * Dinv.GetAt(1, 0) * R.GetAt(0, 1) - R.GetAt(2, 2) * Dinv.GetAt(1, 1) * R.GetAt(0, 0));
+		dRdF[3 * 9 + 6] = (-R.GetAt(2, 1) * Dinv.GetAt(2, 1) * R.GetAt(1, 2) + R.GetAt(2, 1) * Dinv.GetAt(2, 2) * R.GetAt(1, 1) + R.GetAt(2, 2) * Dinv.GetAt(1, 1) * R.GetAt(1, 2) - R.GetAt(2, 2) * Dinv.GetAt(1, 2) * R.GetAt(1, 1));
+		dRdF[4 * 9 + 6] = (R.GetAt(2, 1) * Dinv.GetAt(2, 0) * R.GetAt(1, 2) - R.GetAt(2, 1) * Dinv.GetAt(2, 2) * R.GetAt(1, 0) - R.GetAt(2, 2) * Dinv.GetAt(1, 0) * R.GetAt(1, 2) + R.GetAt(2, 2) * Dinv.GetAt(1, 2) * R.GetAt(1, 0));
+		dRdF[5 * 9 + 6] = (-R.GetAt(2, 1) * Dinv.GetAt(2, 0) * R.GetAt(1, 1) + R.GetAt(2, 1) * Dinv.GetAt(2, 1) * R.GetAt(1, 0) + R.GetAt(2, 2) * Dinv.GetAt(1, 0) * R.GetAt(1, 1) - R.GetAt(2, 2) * Dinv.GetAt(1, 1) * R.GetAt(1, 0));
+		dRdF[6 * 9 + 6] = (-R.GetAt(2, 1) * Dinv.GetAt(2, 1) * R.GetAt(2, 2) + R.GetAt(2, 1) * Dinv.GetAt(2, 2) * R.GetAt(2, 1) + R.GetAt(2, 2) * Dinv.GetAt(1, 1) * R.GetAt(2, 2) - R.GetAt(2, 2) * Dinv.GetAt(1, 2) * R.GetAt(2, 1));
+		dRdF[7 * 9 + 6] = (R.GetAt(2, 1) * Dinv.GetAt(2, 0) * R.GetAt(2, 2) - R.GetAt(2, 1) * Dinv.GetAt(2, 2) * R.GetAt(2, 0) - R.GetAt(2, 2) * Dinv.GetAt(1, 0) * R.GetAt(2, 2) + R.GetAt(2, 2) * Dinv.GetAt(1, 2) * R.GetAt(2, 0));
+		dRdF[8 * 9 + 6] = (-R.GetAt(2, 1) * Dinv.GetAt(2, 0) * R.GetAt(2, 1) + R.GetAt(2, 1) * Dinv.GetAt(2, 1) * R.GetAt(2, 0) + R.GetAt(2, 2) * Dinv.GetAt(1, 0) * R.GetAt(2, 1) - R.GetAt(2, 2) * Dinv.GetAt(1, 1) * R.GetAt(2, 0));
+		dRdF[0 * 9 + 7] = (R.GetAt(2, 0) * Dinv.GetAt(2, 1) * R.GetAt(0, 2) - R.GetAt(2, 0) * Dinv.GetAt(2, 2) * R.GetAt(0, 1) - R.GetAt(2, 2) * Dinv.GetAt(0, 1) * R.GetAt(0, 2) + R.GetAt(2, 2) * Dinv.GetAt(0, 2) * R.GetAt(0, 1));
+		dRdF[1 * 9 + 7] = (-R.GetAt(2, 0) * Dinv.GetAt(2, 0) * R.GetAt(0, 2) + R.GetAt(2, 0) * Dinv.GetAt(2, 2) * R.GetAt(0, 0) + R.GetAt(2, 2) * Dinv.GetAt(0, 0) * R.GetAt(0, 2) - R.GetAt(2, 2) * Dinv.GetAt(0, 2) * R.GetAt(0, 0));
+		dRdF[2 * 9 + 7] = (R.GetAt(2, 0) * Dinv.GetAt(2, 0) * R.GetAt(0, 1) - R.GetAt(2, 0) * Dinv.GetAt(2, 1) * R.GetAt(0, 0) - R.GetAt(2, 2) * Dinv.GetAt(0, 0) * R.GetAt(0, 1) + R.GetAt(2, 2) * Dinv.GetAt(0, 1) * R.GetAt(0, 0));
+		dRdF[3 * 9 + 7] = (R.GetAt(2, 0) * Dinv.GetAt(2, 1) * R.GetAt(1, 2) - R.GetAt(2, 0) * Dinv.GetAt(2, 2) * R.GetAt(1, 1) - R.GetAt(2, 2) * Dinv.GetAt(0, 1) * R.GetAt(1, 2) + R.GetAt(2, 2) * Dinv.GetAt(0, 2) * R.GetAt(1, 1));
+		dRdF[4 * 9 + 7] = (-R.GetAt(2, 0) * Dinv.GetAt(2, 0) * R.GetAt(1, 2) + R.GetAt(2, 0) * Dinv.GetAt(2, 2) * R.GetAt(1, 0) + R.GetAt(2, 2) * Dinv.GetAt(0, 0) * R.GetAt(1, 2) - R.GetAt(2, 2) * Dinv.GetAt(0, 2) * R.GetAt(1, 0));
+		dRdF[5 * 9 + 7] = (R.GetAt(2, 0) * Dinv.GetAt(2, 0) * R.GetAt(1, 1) - R.GetAt(2, 0) * Dinv.GetAt(2, 1) * R.GetAt(1, 0) - R.GetAt(2, 2) * Dinv.GetAt(0, 0) * R.GetAt(1, 1) + R.GetAt(2, 2) * Dinv.GetAt(0, 1) * R.GetAt(1, 0));
+		dRdF[6 * 9 + 7] = (R.GetAt(2, 0) * Dinv.GetAt(2, 1) * R.GetAt(2, 2) - R.GetAt(2, 0) * Dinv.GetAt(2, 2) * R.GetAt(2, 1) - R.GetAt(2, 2) * Dinv.GetAt(0, 1) * R.GetAt(2, 2) + R.GetAt(2, 2) * Dinv.GetAt(0, 2) * R.GetAt(2, 1));
+		dRdF[7 * 9 + 7] = (-R.GetAt(2, 0) * Dinv.GetAt(2, 0) * R.GetAt(2, 2) + R.GetAt(2, 0) * Dinv.GetAt(2, 2) * R.GetAt(2, 0) + R.GetAt(2, 2) * Dinv.GetAt(0, 0) * R.GetAt(2, 2) - R.GetAt(2, 2) * Dinv.GetAt(0, 2) * R.GetAt(2, 0));
+		dRdF[8 * 9 + 7] = (R.GetAt(2, 0) * Dinv.GetAt(2, 0) * R.GetAt(2, 1) - R.GetAt(2, 0) * Dinv.GetAt(2, 1) * R.GetAt(2, 0) - R.GetAt(2, 2) * Dinv.GetAt(0, 0) * R.GetAt(2, 1) + R.GetAt(2, 2) * Dinv.GetAt(0, 1) * R.GetAt(2, 0));
+		dRdF[0 * 9 + 8] = (-R.GetAt(2, 0) * Dinv.GetAt(1, 1) * R.GetAt(0, 2) + R.GetAt(2, 0) * Dinv.GetAt(1, 2) * R.GetAt(0, 1) + R.GetAt(2, 1) * Dinv.GetAt(0, 1) * R.GetAt(0, 2) - R.GetAt(2, 1) * Dinv.GetAt(0, 2) * R.GetAt(0, 1));
+		dRdF[1 * 9 + 8] = (R.GetAt(2, 0) * Dinv.GetAt(1, 0) * R.GetAt(0, 2) - R.GetAt(2, 0) * Dinv.GetAt(1, 2) * R.GetAt(0, 0) - R.GetAt(2, 1) * Dinv.GetAt(0, 0) * R.GetAt(0, 2) + R.GetAt(2, 1) * Dinv.GetAt(0, 2) * R.GetAt(0, 0));
+		dRdF[2 * 9 + 8] = (-R.GetAt(2, 0) * Dinv.GetAt(1, 0) * R.GetAt(0, 1) + R.GetAt(2, 0) * Dinv.GetAt(1, 1) * R.GetAt(0, 0) + R.GetAt(2, 1) * Dinv.GetAt(0, 0) * R.GetAt(0, 1) - R.GetAt(2, 1) * Dinv.GetAt(0, 1) * R.GetAt(0, 0));
+		dRdF[3 * 9 + 8] = (-R.GetAt(2, 0) * Dinv.GetAt(1, 1) * R.GetAt(1, 2) + R.GetAt(2, 0) * Dinv.GetAt(1, 2) * R.GetAt(1, 1) + R.GetAt(2, 1) * Dinv.GetAt(0, 1) * R.GetAt(1, 2) - R.GetAt(2, 1) * Dinv.GetAt(0, 2) * R.GetAt(1, 1));
+		dRdF[4 * 9 + 8] = (R.GetAt(2, 0) * Dinv.GetAt(1, 0) * R.GetAt(1, 2) - R.GetAt(2, 0) * Dinv.GetAt(1, 2) * R.GetAt(1, 0) - R.GetAt(2, 1) * Dinv.GetAt(0, 0) * R.GetAt(1, 2) + R.GetAt(2, 1) * Dinv.GetAt(0, 2) * R.GetAt(1, 0));
+		dRdF[5 * 9 + 8] = (-R.GetAt(2, 0) * Dinv.GetAt(1, 0) * R.GetAt(1, 1) + R.GetAt(2, 0) * Dinv.GetAt(1, 1) * R.GetAt(1, 0) + R.GetAt(2, 1) * Dinv.GetAt(0, 0) * R.GetAt(1, 1) - R.GetAt(2, 1) * Dinv.GetAt(0, 1) * R.GetAt(1, 0));
+		dRdF[6 * 9 + 8] = (-R.GetAt(2, 0) * Dinv.GetAt(1, 1) * R.GetAt(2, 2) + R.GetAt(2, 0) * Dinv.GetAt(1, 2) * R.GetAt(2, 1) + R.GetAt(2, 1) * Dinv.GetAt(0, 1) * R.GetAt(2, 2) - R.GetAt(2, 1) * Dinv.GetAt(0, 2) * R.GetAt(2, 1));
+		dRdF[7 * 9 + 8] = (R.GetAt(2, 0) * Dinv.GetAt(1, 0) * R.GetAt(2, 2) - R.GetAt(2, 0) * Dinv.GetAt(1, 2) * R.GetAt(2, 0) - R.GetAt(2, 1) * Dinv.GetAt(0, 0) * R.GetAt(2, 2) + R.GetAt(2, 1) * Dinv.GetAt(0, 2) * R.GetAt(2, 0));
+		dRdF[8 * 9 + 8] = (-R.GetAt(2, 0) * Dinv.GetAt(1, 0) * R.GetAt(2, 1) + R.GetAt(2, 0) * Dinv.GetAt(1, 1) * R.GetAt(2, 0) + R.GetAt(2, 1) * Dinv.GetAt(0, 0) * R.GetAt(2, 1) - R.GetAt(2, 1) * Dinv.GetAt(0, 1) * R.GetAt(2, 0));
+	}
+
 }  // namespace Chaos
 
