Chaos : Flesh
Fix XPBD Corotated Constraint compiler issue on server

diff --git a/Engine/Source/Runtime/Experimental/Chaos/Public/Chaos/XPBDCorotatedConstraints.h b/Engine/Source/Runtime/Experimental/Chaos/Public/Chaos/XPBDCorotatedConstraints.h
--- a/Engine/Source/Runtime/Experimental/Chaos/Public/Chaos/XPBDCorotatedConstraints.h
+++ b/Engine/Source/Runtime/Experimental/Chaos/Public/Chaos/XPBDCorotatedConstraints.h
@@ -170,10 +170,7 @@ namespace Chaos::Softs
 
 		void InitColor(const FSolverParticles& Particles)
 		{
-			// In dev builds we always color so we can tune the system without restarting. See Apply()
-			#if UE_BUILD_SHIPPING || UE_BUILD_TEST
-			if (Constraints.Num() > Chaos_Spring_ParallelConstraintCount)
-			#endif
+
 			{
 				const TArray<TArray<int32>> ConstraintsPerColor = FGraphColoring::ComputeGraphColoring(MeshConstraints, Particles);
 
