Fix kinematic moving clustered object not being properly woken up by a dynamic field

diff --git a/Engine/Source/Runtime/Experimental/Chaos/Private/Chaos/PBDRigidClusteringAlgo.cpp b/Engine/Source/Runtime/Experimental/Chaos/Private/Chaos/PBDRigidClusteringAlgo.cpp
--- a/Engine/Source/Runtime/Experimental/Chaos/Private/Chaos/PBDRigidClusteringAlgo.cpp
+++ b/Engine/Source/Runtime/Experimental/Chaos/Private/Chaos/PBDRigidClusteringAlgo.cpp
@@ -263,6 +263,10 @@ namespace Chaos
 				}
 
 				MEvolution.SetParticleObjectState(ClusteredCurrentNode, ObjectState);
+				if (ObjectState == Chaos::EObjectStateType::Dynamic)
+				{
+					MEvolution.SetParticleKinematicTarget(ClusteredCurrentNode, FKinematicTarget());
+				}
 			}
 		}
 	}
diff --git a/Engine/Source/Runtime/Experimental/Chaos/Public/PhysicsProxy/FieldSystemProxyHelper.h b/Engine/Source/Runtime/Experimental/Chaos/Public/PhysicsProxy/FieldSystemProxyHelper.h
--- a/Engine/Source/Runtime/Experimental/Chaos/Public/PhysicsProxy/FieldSystemProxyHelper.h
+++ b/Engine/Source/Runtime/Experimental/Chaos/Public/PhysicsProxy/FieldSystemProxyHelper.h
@@ -11,6 +11,8 @@
 #include "Chaos/Particles.h"
 #include <limits>
 
+#include "GeometryCollectionPhysicsProxy.h"
+
 namespace Chaos
 {
 	/**
@@ -192,10 +194,6 @@ namespace Chaos
 		else
 		{
 			RigidHandle->SetObjectStateLowLevel(FieldState);
-			if (FieldState == Chaos::EObjectStateType::Dynamic)
-			{
-				RigidSolver->GetEvolution()->SetParticleKinematicTarget(RigidHandle, FKinematicTarget());
-			}
 		}
 	}
 
@@ -273,9 +271,21 @@ namespace Chaos
 			const FRigidClustering& Clustering = RigidSolver->GetEvolution()->GetRigidClustering();
 			for (FPBDRigidClusteredParticleHandle* TopParentClusteredHandle: TopParentClusters)
 			{
-				if (TopParentClusteredHandle && TopParentClusteredHandle->ClusterIds().NumChildren && !TopParentClusteredHandle->Disabled())
+				if (TopParentClusteredHandle && !TopParentClusteredHandle->Disabled())
 				{
-					UpdateKinematicProperties(TopParentClusteredHandle, Clustering.GetChildrenMap(), *RigidSolver->GetEvolution());
+					if (TopParentClusteredHandle->ClusterIds().NumChildren)
+					{
+						UpdateKinematicProperties(TopParentClusteredHandle, Clustering.GetChildrenMap(), *RigidSolver->GetEvolution());
+					}
+					else
+					{
+						// if the cluster is dynamic let's make sure we clear kinematic target to avoid animated ones to have their velocity reset by the kinematic target update
+						// for particles with children this is taken care in UpdateKinematicProperties
+						if (TopParentClusteredHandle->ObjectState() == EObjectStateType::Dynamic)
+						{
+							RigidSolver->GetEvolution()->SetParticleKinematicTarget(TopParentClusteredHandle, FKinematicTarget());
+						}
+					}
 				}
 			}
 		}
