Chaos : fix static analysis warning

diff --git a/Engine/Source/Runtime/Experimental/Chaos/Private/Chaos/PBDRigidClustering.cpp b/Engine/Source/Runtime/Experimental/Chaos/Private/Chaos/PBDRigidClustering.cpp
--- a/Engine/Source/Runtime/Experimental/Chaos/Private/Chaos/PBDRigidClustering.cpp
+++ b/Engine/Source/Runtime/Experimental/Chaos/Private/Chaos/PBDRigidClustering.cpp
@@ -432,31 +432,32 @@ namespace Chaos
 
 	void FRigidClustering::RemoveChildFromParent(FPBDRigidParticleHandle* Child, const FPBDRigidClusteredParticleHandle* ClusteredParent)
 	{
-		ensure(Child);
-		ensure(ClusteredParent);
-		FPBDRigidClusteredParticleHandle* ClusteredChild = Child->CastToClustered();
-
-		MEvolution.EnableParticle(Child, ClusteredParent);
-		TopLevelClusterParents.Add(ClusteredChild);
-
-		ClusteredChild->SetClusterId(ClusterId(nullptr, ClusteredChild->ClusterIds().NumChildren)); // clear Id but retain number of children
-
-		const FRigidTransform3 PreSolveTM(ClusteredParent->P(), ClusteredParent->Q());
-		const FRigidTransform3 ChildFrame = ClusteredChild->ChildToParent() * PreSolveTM;
-		Child->SetX(ChildFrame.GetTranslation());
-		Child->SetR(ChildFrame.GetRotation());
-
-		Child->SetP(Child->X());
-		Child->SetQ(Child->R());
-
-		//todo(ocohen): for now just inherit velocity at new COM. This isn't quite right for rotation
-		//todo(ocohen): in the presence of collisions, this will leave all children with the post-collision
-		// velocity. This should be controlled by material properties so we can allow the broken pieces to
-		// maintain the clusters pre-collision velocity.
-		Child->SetV(ClusteredParent->V());
-		Child->SetW(ClusteredParent->W());
-		Child->SetPreV(ClusteredParent->PreV());
-		Child->SetPreW(ClusteredParent->PreW());
+		if (ensure(Child != nullptr && ClusteredParent != nullptr))
+		{
+			FPBDRigidClusteredParticleHandle* ClusteredChild = Child->CastToClustered();
+
+			MEvolution.EnableParticle(Child, ClusteredParent);
+			TopLevelClusterParents.Add(ClusteredChild);
+
+			ClusteredChild->SetClusterId(ClusterId(nullptr, ClusteredChild->ClusterIds().NumChildren)); // clear Id but retain number of children
+
+			const FRigidTransform3 PreSolveTM(ClusteredParent->P(), ClusteredParent->Q());
+			const FRigidTransform3 ChildFrame = ClusteredChild->ChildToParent() * PreSolveTM;
+			Child->SetX(ChildFrame.GetTranslation());
+			Child->SetR(ChildFrame.GetRotation());
+
+			Child->SetP(Child->X());
+			Child->SetQ(Child->R());
+
+			//todo(ocohen): for now just inherit velocity at new COM. This isn't quite right for rotation
+			//todo(ocohen): in the presence of collisions, this will leave all children with the post-collision
+			// velocity. This should be controlled by material properties so we can allow the broken pieces to
+			// maintain the clusters pre-collision velocity.
+			Child->SetV(ClusteredParent->V());
+			Child->SetW(ClusteredParent->W());
+			Child->SetPreV(ClusteredParent->PreV());
+			Child->SetPreW(ClusteredParent->PreW());
+		}
 	};
 
 	TArray<FPBDRigidParticleHandle*> FRigidClustering::CreateClustersFromNewIslands(
