Compile UnrealGame HoloLens - SimplexVectorized.h warning: illegal copy-initialization

diff --git a/Engine/Source/Runtime/Experimental/Chaos/Public/Chaos/VectorUtility.h b/Engine/Source/Runtime/Experimental/Chaos/Public/Chaos/VectorUtility.h
--- a/Engine/Source/Runtime/Experimental/Chaos/Public/Chaos/VectorUtility.h
+++ b/Engine/Source/Runtime/Experimental/Chaos/Public/Chaos/VectorUtility.h
@@ -8,18 +8,18 @@
  * @param V	vector
  * @return		VectorRegister4Float( B.x, A.y, A.z, A.w)
  */
-#if !defined(_MSC_VER)|| PLATFORM_ENABLE_VECTORINTRINSICS_NEON
+
 FORCEINLINE VectorRegister4Float VectorCast4IntTo4Float(const VectorRegister4Int& V)
 {
+#if !defined(_MSC_VER)|| PLATFORM_ENABLE_VECTORINTRINSICS_NEON
 	return VectorRegister4Float(V);
-}
-#else
-FORCEINLINE VectorRegister4Float VectorCast4IntTo4Float(const VectorRegister4Int& V)
-{
+#elif PLATFORM_ENABLE_VECTORINTRINSICS
 	return _mm_castsi128_ps(V);
+#else
+	return VectorCastIntToFloat(Vec);
+#endif
 }
 
-#endif
 
 /**
  * Cast VectorRegister4Float in VectorRegister4Int
@@ -27,7 +27,17 @@ FORCEINLINE VectorRegister4Float VectorCast4IntTo4Float(const VectorRegister4Int
  * @param V	vector
  * @return		VectorCast4FloatTo4Int( B.x, A.y, A.z, A.w)
  */
-#define VectorCast4FloatTo4Int(Vec)		VectorCastFloatToInt(Vec)
+FORCEINLINE VectorRegister4Int VectorCast4FloatTo4Int(const VectorRegister4Float& V)
+{
+#if !defined(_MSC_VER) || PLATFORM_ENABLE_VECTORINTRINSICS_NEON
+	return VectorRegister4Int(V);
+#elif PLATFORM_ENABLE_VECTORINTRINSICS
+	return _mm_castps_si128(V);
+#else
+	return VectorCastFloatToInt(Vec);
+#endif
+
+}
 
 /**
  * Selects and interleaves the lower two SP FP values from A and B.
