Height Field Overlap Optimization :
Check bounding volume of several cells, instead of checking all individual cells bounding volume.

diff --git a/Engine/Source/Programs/HeadlessChaos/Private/HeadlessChaosTestHeightfield.cpp b/Engine/Source/Programs/HeadlessChaos/Private/HeadlessChaosTestHeightfield.cpp
--- a/Engine/Source/Programs/HeadlessChaos/Private/HeadlessChaosTestHeightfield.cpp
+++ b/Engine/Source/Programs/HeadlessChaos/Private/HeadlessChaosTestHeightfield.cpp
@@ -868,6 +868,23 @@ namespace ChaosTest {
 				}
 			}
 		}
+		// Thin Box on the top
+		{
+			TBox<FReal, 3> Box(FVec3(-10.0, -10.0, -0.001), FVec3(10.0, 10.0, 0.001));
+			for (int32 Row = 0; Row < Rows; ++Row)
+			{
+				for (int32 Col = 0; Col < Columns; ++Col)
+				{
+					const FVec3 Translation(Col, Row, 9.5);
+					FRigidTransform3 QueryTM(Translation, TRotation<FReal, 3>(UE::Math::TQuat<FReal>(TVector<FReal, 3>(1.0, 0.0, 0.0), 0.0)));
+					FVec3 Dir(1, 0, 0);
+					bool Result = Heightfield.OverlapGeom(Box, QueryTM, 0.0, nullptr);
+					// Collision with the mountain
+					EXPECT_TRUE(Result);
+				}
+			}
+		}
+
 		// Inclined Box
 		{
 			TBox<FReal, 3> Box(FVec3(-1.0, -1.0, -0.0001), FVec3(1.0, 1.0, 10.0));
diff --git a/Engine/Source/Runtime/Experimental/Chaos/Private/Chaos/HeightField.cpp b/Engine/Source/Runtime/Experimental/Chaos/Private/Chaos/HeightField.cpp
--- a/Engine/Source/Runtime/Experimental/Chaos/Private/Chaos/HeightField.cpp
+++ b/Engine/Source/Runtime/Experimental/Chaos/Private/Chaos/HeightField.cpp
@@ -1354,6 +1354,59 @@ namespace Chaos
 		return OutIntersections.Num() > 0;
 	}
 
+
+	bool Chaos::FHeightField::GetGridIntersectionsBatch(FBounds2D InFlatBounds, TArray<TVec2<int32>>& OutIntersections, const FAABBVectorized& Bounds) const
+	{
+		OutIntersections.Reset();
+
+		const FBounds2D FlatBounds = GetFlatBounds();
+		FVec2 Scale2D(GeomData.Scale[0], GeomData.Scale[1]);
+
+		InFlatBounds = FBounds2D::FromPoints(FlatBounds.Clamp(InFlatBounds.Min) / Scale2D, FlatBounds.Clamp(InFlatBounds.Max) / Scale2D);
+
+		TVec2<int32> MinCell = FlatGrid.Cell(InFlatBounds.Min);
+		TVec2<int32> MaxCell = FlatGrid.Cell(InFlatBounds.Max);
+
+		// We want to capture the first cell (delta == 0) as well
+		const int32 MaxX = MaxCell[0] + 1;
+		const int32 MaxY = MaxCell[1] + 1;
+
+		const int32 NumX = MaxCell[0] - MaxX;
+		const int32 NumY = MaxCell[1] - MaxY;
+		OutIntersections.Reserve(NumX * NumY);
+		
+		constexpr int32 Jump = 4;
+		
+		FAABBVectorized CellsBounds;
+		for (int32 CurrFastX = MinCell[0]; CurrFastX < MaxX; CurrFastX += Jump)
+		{
+			for (int32 CurrFastY = MinCell[1]; CurrFastY < MaxY; CurrFastY += Jump)
+			{
+				const TVec2<int32> CellIdx(CurrFastX, CurrFastY);
+				const TVec2<int32> Area(FMath::Min(MaxX - CurrFastX, Jump), FMath::Min(MaxY - CurrFastY, Jump));
+				GeomData.GetBoundsScaled(CellIdx, Area, CellsBounds);
+				if (CellsBounds.Intersects(Bounds))
+				{
+					const int32 CurrMaxX = CurrFastX + Area[0];
+					const int32 CurrMaxY = CurrFastY + Area[1];
+
+					for (int32 CurrX = CurrFastX; CurrX < CurrMaxX; ++CurrX)
+					{
+						for (int32 CurrY = CurrFastY; CurrY < CurrMaxY; ++CurrY)
+						{
+							const TVec2<int32> Cell(CurrX, CurrY);
+							check (FlatGrid.IsValid(Cell))
+							OutIntersections.Add(Cell);
+						}
+					}
+				}
+			}
+		}
+		return OutIntersections.Num() > 0;
+	}
+	
+
+
 	typename Chaos::FHeightField::FBounds2D Chaos::FHeightField::GetFlatBounds() const
 	{
 		FBounds2D Result;
@@ -1966,65 +2019,65 @@ namespace Chaos
 
 		TArray<TVec2<int32>> Intersections;
 		
-
-		GetGridIntersections(FlatQueryBounds, Intersections);
-
-		bool bOverlaps = false;
-		if (OutMTD)
+		const FAABBVectorized QueryBoundsSimd(QueryBounds);
+		if (GetGridIntersectionsBatch(FlatQueryBounds, Intersections, QueryBoundsSimd))
 		{
-			OutMTD->Normal = FVec3(0);
-			OutMTD->Penetration = TNumericLimits<FReal>::Lowest();
-			FVec3 Points[4];
-			for (const TVec2<int32>& Cell : Intersections)
+			bool bOverlaps = false;
+			if (OutMTD)
 			{
-				const int32 SingleIndex = Cell[1] * (GeomData.NumCols) + Cell[0];
-				GeomData.GetPointsScaled(SingleIndex, Points);
+				OutMTD->Normal = FVec3(0);
+				OutMTD->Penetration = TNumericLimits<FReal>::Lowest();
+				FVec3 Points[4];
+				for (const TVec2<int32>& Cell : Intersections)
+				{
+					const int32 SingleIndex = Cell[1] * (GeomData.NumCols) + Cell[0];
+					GeomData.GetPointsScaled(SingleIndex, Points);
 
-				bOverlaps |= OverlapTriangleMTD(Points[0], Points[1], Points[3], OutMTD);
-				bOverlaps |= OverlapTriangleMTD(Points[0], Points[3], Points[2], OutMTD);
+					bOverlaps |= OverlapTriangleMTD(Points[0], Points[1], Points[3], OutMTD);
+					bOverlaps |= OverlapTriangleMTD(Points[0], Points[3], Points[2], OutMTD);
+				}
+				return bOverlaps;
 			}
-			return bOverlaps;
-		}
-		else
-		{
-			VectorRegister4Float Points[4];
-			const UE::Math::TQuat<FReal>& RotationDouble = QueryTM.GetRotation();
-			VectorRegister4Float Rotation = MakeVectorRegisterFloatFromDouble(MakeVectorRegister(RotationDouble.X, RotationDouble.Y, RotationDouble.Z, RotationDouble.W));
-			// Normalize rotation
-			Rotation = VectorNormalizeSafe(Rotation, GlobalVectorConstants::Float0001);
-			VectorRegister4Float InvRotation = VectorQuaternionInverse(Rotation);
-
-			const UE::Math::TVector<FReal>& TranslationDouble = QueryTM.GetTranslation();
-			const VectorRegister4Float Translation = MakeVectorRegisterFloatFromDouble(MakeVectorRegister(TranslationDouble.X, TranslationDouble.Y, TranslationDouble.Z, 0.0));
-			
-			FAABBVectorized CellBounds;
-			const FAABBVectorized QueryBoundsSimd(QueryBounds);
-			for (const TVec2<int32>& Cell : Intersections)
+			else
 			{
-				const int32 SingleIndex = Cell[1] * (GeomData.NumCols) + Cell[0];
-				GeomData.GetPointsAndBoundsScaledSimd(SingleIndex, Points, CellBounds);
-				if (CellBounds.Intersects(QueryBoundsSimd))
+				VectorRegister4Float Points[4];
+				const UE::Math::TQuat<FReal>& RotationDouble = QueryTM.GetRotation();
+				VectorRegister4Float Rotation = MakeVectorRegisterFloatFromDouble(MakeVectorRegister(RotationDouble.X, RotationDouble.Y, RotationDouble.Z, RotationDouble.W));
+				// Normalize rotation
+				Rotation = VectorNormalizeSafe(Rotation, GlobalVectorConstants::Float0001);
+				VectorRegister4Float InvRotation = VectorQuaternionInverse(Rotation);
+
+				const UE::Math::TVector<FReal>& TranslationDouble = QueryTM.GetTranslation();
+				const VectorRegister4Float Translation = MakeVectorRegisterFloatFromDouble(MakeVectorRegister(TranslationDouble.X, TranslationDouble.Y, TranslationDouble.Z, 0.0));
+
+				FAABBVectorized CellBounds;
+				for (const TVec2<int32>& Cell : Intersections)
 				{
-					// pre-transform the triangle in overlap geometry space
-					for (int32 Index = 0; Index < 4; ++Index)
+					const int32 SingleIndex = Cell[1] * (GeomData.NumCols) + Cell[0];
+					GeomData.GetPointsAndBoundsScaledSimd(SingleIndex, Points, CellBounds);
+					if (CellBounds.Intersects(QueryBoundsSimd))
 					{
-						VectorRegister4Float& P = Points[Index];
-						P = VectorSubtract(P, Translation);
-						P = VectorQuaternionRotateVector(InvRotation, P);
-					}
+						// pre-transform the triangle in overlap geometry space
+						for (int32 Index = 0; Index < 4; ++Index)
+						{
+							VectorRegister4Float& P = Points[Index];
+							P = VectorSubtract(P, Translation);
+							P = VectorQuaternionRotateVector(InvRotation, P);
+						}
 
-					if (OverlapTriangleNoMTD(Points[0], Points[1], Points[3]))
-					{
-						return true;
-					}
-					if (OverlapTriangleNoMTD(Points[0], Points[3], Points[2]))
-					{
-						return true;
+						if (OverlapTriangleNoMTD(Points[0], Points[1], Points[3]))
+						{
+							return true;
+						}
+						if (OverlapTriangleNoMTD(Points[0], Points[3], Points[2]))
+						{
+							return true;
+						}
 					}
 				}
 			}
-			return false;
 		}
+		return false;
 	}
 
 	bool FHeightField::OverlapGeom(const TSphere<FReal, 3>& QueryGeom, const FRigidTransform3& QueryTM, const FReal Thickness, FMTDInfo* OutMTD) const
@@ -2366,4 +2419,4 @@ namespace Chaos
 
 		return OutPhi;
 	}
-}
+}
\ No newline at end of file
diff --git a/Engine/Source/Runtime/Experimental/Chaos/Public/Chaos/HeightField.h b/Engine/Source/Runtime/Experimental/Chaos/Public/Chaos/HeightField.h
--- a/Engine/Source/Runtime/Experimental/Chaos/Public/Chaos/HeightField.h
+++ b/Engine/Source/Runtime/Experimental/Chaos/Public/Chaos/HeightField.h
@@ -710,6 +710,7 @@ namespace Chaos
 		bool GridSweep(const FVec3& StartPoint, const FVec3& Dir, const FReal Length, const FVec3 InHalfExtents, SQVisitor& Visitor) const;
 		bool GridCast(const FVec3& StartPoint, const FVec3& Dir, const FReal Length, FHeightfieldRaycastVisitor& Visitor) const;
 		bool GetGridIntersections(FBounds2D InFlatBounds, TArray<TVec2<int32>>& OutInterssctions) const;
+		bool GetGridIntersectionsBatch(FBounds2D InFlatBounds, TArray<TVec2<int32>>& OutIntersections, const FAABBVectorized& Bounds) const;
 		
 		FBounds2D GetFlatBounds() const;
 
