Chaos Cloth - Fix stiffness exponentiation calculations when values are close to 0 and 1.

Not only discontinuities in the exponent functions caused issues in some platforms, but also floating point approximations meant that a 0 input value was turned into a small stiffness whose effect was perceivable, most notably on the AnimDrive constraint.
Replacing the clamp operation by using a SMALL_NUMBER (1e-8) comparison window gives now the expected behavior on both boundary values, while providing some guaranty that the behavioral change isn't gong to impact most existing setups.

diff --git a/Engine/Source/Runtime/Experimental/Chaos/Public/Chaos/PBDStiffness.h b/Engine/Source/Runtime/Experimental/Chaos/Public/Chaos/PBDStiffness.h
--- a/Engine/Source/Runtime/Experimental/Chaos/Public/Chaos/PBDStiffness.h
+++ b/Engine/Source/Runtime/Experimental/Chaos/Public/Chaos/PBDStiffness.h
@@ -196,18 +196,21 @@ void FPBDStiffness::ApplyValues(const FSolverReal Dt, const int32 NumIterations)
 	// Define the stiffness mapping function
 	auto SimulationValue = [this, Exponent](const FSolverReal InValue)->FSolverReal
 	{
-		// If InValue is 1 then LogValue = -inf and the output becomes -inf as well
-		// In order for this function to be continuous, We want the output to be 1 when InValue = 1
-		const FSolverReal ClampedInValue = FMath::Clamp(InValue, (FSolverReal)0., (FSolverReal)1.);
-
-		if (ClampedInValue == (FSolverReal)1.)
+		// If InValue is 1 then LogValue = -inf and the output becomes -inf as well,
+		// in order for this function to be continuous, we want the output to be 1 when InValue = 1
+		// and need the stiffness to be exactly 0 when the input is 0
+		if (InValue <= (FSolverReal)UE_DOUBLE_SMALL_NUMBER)
+		{
+			return (FSolverReal)0.;
+		}
+		if (InValue >= (FSolverReal)(1. - UE_DOUBLE_SMALL_NUMBER))
 		{
 			return (FSolverReal)1.;
 		}
 		// Get a very steep exponential curve between the [0, 1] range to make easier to set the parameter
 		// The base has been chosen empirically
 		// ParameterValue = Pow(ParameterFitBase, ParameterValue - 1)
-		const FSolverReal ParameterFit = FMath::Exp(ParameterFitLogBase * (ClampedInValue - (FSolverReal)1.));
+		const FSolverReal ParameterFit = FMath::Exp(ParameterFitLogBase * (InValue - (FSolverReal)1.));
 
 		// Use simulation dependent stiffness exponent to alleviate the variations in effect when Dt and NumIterations change
 		// This is based on the Position-Based Simulation Methods paper (page 8),
