Chaos
- projectionphase  in main solver
- add teleport projection
- recombine projection flags into one flag
- remove unused iteration counts
- correct the iteration count names

diff --git a/Engine/Source/Editor/DetailCustomizations/Private/PhysicsConstraintComponentDetails.cpp b/Engine/Source/Editor/DetailCustomizations/Private/PhysicsConstraintComponentDetails.cpp
--- a/Engine/Source/Editor/DetailCustomizations/Private/PhysicsConstraintComponentDetails.cpp
+++ b/Engine/Source/Editor/DetailCustomizations/Private/PhysicsConstraintComponentDetails.cpp
@@ -175,28 +175,18 @@ void FPhysicsConstraintComponentDetails::AddConstraintBehaviorProperties(IDetail
 
 	//Add properties we want in specific order
 	ConstraintCat.AddProperty(ProfilePropertiesProperty->GetChildHandle(GET_MEMBER_NAME_CHECKED(FConstraintProfileProperties, bDisableCollision)));
-#if WITH_CHAOS
-	ConstraintCat.AddProperty(ProfilePropertiesProperty->GetChildHandle(GET_MEMBER_NAME_CHECKED(FConstraintProfileProperties, bEnableLinearProjection)));
-	ConstraintCat.AddProperty(ProfilePropertiesProperty->GetChildHandle(GET_MEMBER_NAME_CHECKED(FConstraintProfileProperties, bEnableAngularProjection)));
-	ConstraintCat.AddProperty(ProfilePropertiesProperty->GetChildHandle(GET_MEMBER_NAME_CHECKED(FConstraintProfileProperties, bEnableShockPropagation)));
-	ConstraintCat.AddProperty(ProfilePropertiesProperty->GetChildHandle(GET_MEMBER_NAME_CHECKED(FConstraintProfileProperties, ProjectionLinearAlpha)));
-	ConstraintCat.AddProperty(ProfilePropertiesProperty->GetChildHandle(GET_MEMBER_NAME_CHECKED(FConstraintProfileProperties, ProjectionAngularAlpha)));
-	ConstraintCat.AddProperty(ProfilePropertiesProperty->GetChildHandle(GET_MEMBER_NAME_CHECKED(FConstraintProfileProperties, ShockPropagationAlpha)));
-#endif
-
-#if WITH_CHAOS
-	// Hide PhysX-Only settings in Chaos
-	//ProfilePropertiesProperty->GetChildHandle(GET_MEMBER_NAME_CHECKED(FConstraintProfileProperties, ProjectionLinearTolerance))->MarkHiddenByCustomization();
-	//ProfilePropertiesProperty->GetChildHandle(GET_MEMBER_NAME_CHECKED(FConstraintProfileProperties, ProjectionAngularTolerance))->MarkHiddenByCustomization();
-#endif
-
-#if !WITH_CHAOS
-	// Hide Chaos-Only settings in PhysX
-	ProfilePropertiesProperty->GetChildHandle(GET_MEMBER_NAME_CHECKED(FConstraintProfileProperties, bEnableSoftProjection))->MarkHiddenByCustomization();
-	ProfilePropertiesProperty->GetChildHandle(GET_MEMBER_NAME_CHECKED(FConstraintProfileProperties, ProjectionLinearAlpha))->MarkHiddenByCustomization();
-	ProfilePropertiesProperty->GetChildHandle(GET_MEMBER_NAME_CHECKED(FConstraintProfileProperties, ProjectionAngularAlpha))->MarkHiddenByCustomization();
-	ProfilePropertiesProperty->GetChildHandle(GET_MEMBER_NAME_CHECKED(FConstraintProfileProperties, ShockPropagationAlpha))->MarkHiddenByCustomization();
-#endif
+	ConstraintCat.AddProperty(ProfilePropertiesProperty->GetChildHandle(GET_MEMBER_NAME_CHECKED(FConstraintProfileProperties, bParentDominates)));
+
+	IDetailGroup& ProjectionGroup = ConstraintCat.AddGroup("Projection", LOCTEXT("Projection", "Projection"), false, true);
+	ProjectionGroup.AddPropertyRow(ProfilePropertiesProperty->GetChildHandle(GET_MEMBER_NAME_CHECKED(FConstraintProfileProperties, bEnableProjection)).ToSharedRef());
+	ProjectionGroup.AddPropertyRow(ProfilePropertiesProperty->GetChildHandle(GET_MEMBER_NAME_CHECKED(FConstraintProfileProperties, ProjectionLinearTolerance)).ToSharedRef());
+	ProjectionGroup.AddPropertyRow(ProfilePropertiesProperty->GetChildHandle(GET_MEMBER_NAME_CHECKED(FConstraintProfileProperties, ProjectionAngularTolerance)).ToSharedRef());
+	ProjectionGroup.AddPropertyRow(ProfilePropertiesProperty->GetChildHandle(GET_MEMBER_NAME_CHECKED(FConstraintProfileProperties, ProjectionLinearAlpha)).ToSharedRef());
+	ProjectionGroup.AddPropertyRow(ProfilePropertiesProperty->GetChildHandle(GET_MEMBER_NAME_CHECKED(FConstraintProfileProperties, ProjectionAngularAlpha)).ToSharedRef());
+
+	IDetailGroup& ShockPropGroup = ConstraintCat.AddGroup("ShockPropagation", LOCTEXT("ShockPropagation", "Shock Propagation"), false, false);
+	ShockPropGroup.AddPropertyRow(ProfilePropertiesProperty->GetChildHandle(GET_MEMBER_NAME_CHECKED(FConstraintProfileProperties, bEnableShockPropagation)).ToSharedRef());
+	ShockPropGroup.AddPropertyRow(ProfilePropertiesProperty->GetChildHandle(GET_MEMBER_NAME_CHECKED(FConstraintProfileProperties, ShockPropagationAlpha)).ToSharedRef());
 
 	//Add the rest
 	uint32 NumProfileProperties = 0;
diff --git a/Engine/Source/Programs/HeadlessChaos/Private/HeadlessChaosTestJointForces.cpp b/Engine/Source/Programs/HeadlessChaos/Private/HeadlessChaosTestJointForces.cpp
--- a/Engine/Source/Programs/HeadlessChaos/Private/HeadlessChaosTestJointForces.cpp
+++ b/Engine/Source/Programs/HeadlessChaos/Private/HeadlessChaosTestJointForces.cpp
@@ -72,7 +72,6 @@ namespace ChaosTest {
 		{
 			Test.JointSettings[JointIndex].AngularMotionTypes = { EJointMotionType::Locked, EJointMotionType::Locked, EJointMotionType::Locked };
 			Test.JointSettings[JointIndex].bProjectionEnabled = false;
-			Test.JointSettings[JointIndex].bSoftProjectionEnabled = false;
 		}
 
 		Test.Create();
@@ -137,7 +136,6 @@ namespace ChaosTest {
 		{
 			Test.JointSettings[JointIndex].AngularMotionTypes = { EJointMotionType::Locked, EJointMotionType::Locked, EJointMotionType::Locked };
 			Test.JointSettings[JointIndex].bProjectionEnabled = false;
-			Test.JointSettings[JointIndex].bSoftProjectionEnabled = false;
 		}
 
 		Test.Create();
diff --git a/Engine/Source/Runtime/Engine/Classes/PhysicsEngine/ConstraintInstance.h b/Engine/Source/Runtime/Engine/Classes/PhysicsEngine/ConstraintInstance.h
--- a/Engine/Source/Runtime/Engine/Classes/PhysicsEngine/ConstraintInstance.h
+++ b/Engine/Source/Runtime/Engine/Classes/PhysicsEngine/ConstraintInstance.h
@@ -26,26 +26,26 @@ struct ENGINE_API FConstraintProfileProperties
 {
 	GENERATED_USTRUCT_BODY()
 
-	/** [PhysX only] Linear tolerance value in world units. If the distance error exceeds this tolerence limit, the body will be projected. */
-	UPROPERTY()
+	/** If the joint error is above this distance after the solve phase, the child body will be teleported to fix the error. Only used if bEnableProjection is true. */
+	UPROPERTY(EditAnywhere, Category = Projection, meta = (ClampMin = "0.0"))
 	float ProjectionLinearTolerance;
 
-	/** [PhysX only] Angular tolerance value in world units. If the distance error exceeds this tolerence limit, the body will be projected. */
-	UPROPERTY()
+	/** If the joint error is above this distance after the solve phase, the child body will be teleported to fix the error. Only used if bEnableProjection is true. */
+	UPROPERTY(EditAnywhere, Category = Projection, meta = (ClampMin = "0.0"))
 	float ProjectionAngularTolerance;
 
-	/**  How much linear projection to apply [0-1] */
+	/**  How much semi-physical linear projection correction to apply [0-1]. Only used if bEnableProjection is true and if hard limits are used. */
 	UPROPERTY(EditAnywhere, Category = Projection, meta = (ClampMin = "0.0", ClampMax = "1.0"))
 	float ProjectionLinearAlpha;
 
-	/** How much angular projection to apply [0-1] */
+	/** How much semi-physical angular projection correction to apply [0-1]. Only used if bEnableProjection is true and if hard limits are used. */
 	UPROPERTY(EditAnywhere, Category = Projection, meta = (ClampMin = "0.0", ClampMax = "1.0"))
 	float ProjectionAngularAlpha;
 
 	/** 
 	 * How much shock propagation to apply [0-1]. Shock propagation increases the mass of the parent body for the last iteration of the
 	 * position and velocity solve phases. This can help stiffen up joint chains, but is also prone to introducing energy down the chain
-	 * especially at high alpha.
+	 * especially at high alpha. Only used in bEnableShockPropagation is true.
 	 */
 	UPROPERTY(EditAnywhere, Category = Projection, meta = (ClampMin = "0.0", ClampMax = "1.0"))
 	float ShockPropagationAlpha;
@@ -87,49 +87,32 @@ struct ENGINE_API FConstraintProfileProperties
 	UPROPERTY(EditAnywhere, Category = Constraint)
 	uint8 bParentDominates : 1;
 
-	/**
-	* NOTE: RigidBody AnimNode only. Projection is not applied to ragdoll physics in the main scene. For Ragdolls, ShockPropagation has a similar effect
-	* and is more compatible with the world solver.
-	* 
-	* Projection is a post-solve position and angular fixup where the parent body in the constraint is treated as having infinite mass and the child body is
-	* translated and rotated to resolve any remaining errors using a semi-physical correction. This can be used to make constraint chains significantly stiffer 
-	* at lower iteration counts. Increasing iterations would have the same effect, but be much more expensive. Projection only works well if the chain is not 
-	* interacting with other objects (e.g., through collisions) because the projection of the bodies in the chain will cause other constraints to be violated. 
-	* Likewise, if a body is influenced by multiple constraints, then enabling projection on more than one constraint may lead to unexpected results - the 
-	* "last" constraint would win but the order in which constraints are solved cannot be directly controlled.
-	* 
-	* Projection is fairly expensive compared to a single position iteration, so if you can get the behaviour you need by adding a couple iterations, that
-	* is probably a better approach.
-	*
-	*
-	* Note: projection is only applied to hard-limit constraints, anf not applied to constraints with soft limits.
-	*/
-	UPROPERTY(EditAnywhere, Category = Projection)
-	uint8 bEnableLinearProjection : 1;
-
-	/**
-	 * NOTE: RigidBody AnimNode only. See coments on bEnableLinearProjection
-	*/
-	UPROPERTY(EditAnywhere, Category = Projection)
-	uint8 bEnableAngularProjection : 1;
-
 	/**
 	 * Shock propagation increases the mass of the parent body for the last iteration of the position and velocity solve phases. 
 	 * This can help stiffen up joint chains, but is also prone to introducing energy down the chain especially at high alpha.
-	 * 
-	 * NOTE: This is intended to be used for world constraints, not RigidBody AnimNodes which have the Projection system.
+	 * It also does not work well if there are collisions on the bodies preventing the joints from correctly resolving - this 
+	 * can lead to jitter, especially if collision shock propagation is also enabled.
 	 */
 	UPROPERTY(EditAnywhere, Category = Projection)
 	uint8 bEnableShockPropagation : 1;
 
-	// HIDDEN - TO BE DEPRECATED
-	UPROPERTY()
+	/**
+	* Projection is a post-solve position and angular fixup consisting of two correction procedures. First, if the constraint limits are exceeded
+	* by more that the Linear or Angular Tolerance, the bodies are teleported to eliminate the error. Second, if the constraint limits are exceeded
+	* by less than the tolerance, a semi-physical correction is applied,  with the parent body in the constraint is treated as having infinite mass.
+	* The teleport tolerance are controlled by ProjectionLinearTolerance and ProjectionAngularTolerance. The semi-physical correction is controlled
+	* by ProjectionLinearAlpha and ProjectionAnguilarAlpha. You may have one, none, or both systems enabled at the same time.
+	*
+	* Projection only works well if the chain is not interacting with other objects (e.g., through collisions) because the projection of the bodies in
+	* the chain will cause other constraints to be violated. Likewise, if a body is influenced by multiple constraints, then enabling projection on more
+	* than one constraint may lead to unexpected results - the  "last" constraint would win but the order in which constraints are solved cannot be directly controlled.
+	*
+	* Note that the semi-physical projection (ProjectionLinearAlpha and ProjectionAngularAlpha) is only applied to hard-limit constraints and not those with
+	* soft limits because the soft limit is the point at which the soft-constraint (spring) kicks in, and not really a limit on how far the joint can be separated.
+	*/
+	UPROPERTY(EditAnywhere, Category = Projection)
 	uint8 bEnableProjection : 1;
 
-	// HIDDEN - TO BE DEPRECATED
-	UPROPERTY()
-	uint8 bEnableSoftProjection : 1;
-
 	/** Whether it is possible to break the joint with angular force. */
 	UPROPERTY(EditAnywhere, AdvancedDisplay, Category = Angular)
 	uint8 bAngularBreakable : 1;
@@ -928,10 +911,10 @@ public:
 	void DisableProjection();
 
 	/** Set projection parameters */
-	void SetProjectionParams(bool bEnableLinearProjection, bool bEnableAngularProjection, float ProjectionLinearAlphaOrTolerance, float ProjectionAngularAlphaOrTolerance);
+	void SetProjectionParams(bool bEnableProjection, float ProjectionLinearAlpha, float ProjectionAngularAlpha, float ProjectionLinearTolerance, float ProjectionAngularTolerance);
 
 	/** Get projection parameters */
-	void GetProjectionAlphasOrTolerances(float& ProjectionLinearAlphaOrTolerance, float& ProjectionAngularAlphaOrTolerance) const;
+	void GetProjectionParams(float& ProjectionLinearAlpha, float& ProjectionAngularAlpha, float& ProjectionLinearTolerance, float& ProjectionAngularTolerance) const;
 
 	/** Set the shock propagation amount [0, 1] */
 	void SetShockPropagationParams(bool bEnableShockPropagation, float ShockPropagationAlpha);
diff --git a/Engine/Source/Runtime/Engine/Private/PhysicsEngine/ConstraintInstance.cpp b/Engine/Source/Runtime/Engine/Private/PhysicsEngine/ConstraintInstance.cpp
--- a/Engine/Source/Runtime/Engine/Private/PhysicsEngine/ConstraintInstance.cpp
+++ b/Engine/Source/Runtime/Engine/Private/PhysicsEngine/ConstraintInstance.cpp
@@ -143,11 +143,8 @@ FConstraintProfileProperties::FConstraintProfileProperties()
 	, ContactTransferScale(0.f)
 	, bDisableCollision(false)
 	, bParentDominates(false)
-	, bEnableLinearProjection(false)
-	, bEnableAngularProjection(false)
 	, bEnableShockPropagation(false)
-	, bEnableProjection(true)		// to be deprecated
-	, bEnableSoftProjection(false)	// to be deprecated
+	, bEnableProjection(true)
 	, bAngularBreakable(false)
 	, bAngularPlasticity(false)
 	, bLinearBreakable(false)
@@ -469,7 +466,7 @@ void FConstraintProfileProperties::UpdateConstraintFlags_AssumesLocked(const FPh
 
 	FPhysicsInterface::SetCollisionEnabled(InConstraintRef, !bDisableCollision);
 #if WITH_CHAOS
-	FPhysicsInterface::SetProjectionEnabled_AssumesLocked(InConstraintRef, bEnableLinearProjection || bEnableAngularProjection, ProjectionLinearAlpha, ProjectionAngularAlpha);
+	FPhysicsInterface::SetProjectionEnabled_AssumesLocked(InConstraintRef, bEnableProjection, ProjectionLinearAlpha, ProjectionAngularAlpha, ProjectionLinearTolerance, ProjectionAngularTolerance);
 	FPhysicsInterface::SetShockPropagationEnabled_AssumesLocked(InConstraintRef, bEnableShockPropagation, ShockPropagationAlpha);
 #else
 	FPhysicsInterface::SetProjectionEnabled_AssumesLocked(InConstraintRef, bEnableProjection, ProjectionLinearTolerance, ProjectionAngularTolerance);
@@ -1183,41 +1180,34 @@ FConstraintInstance * FConstraintInstance::Alloc()
 	return new (Memory)FConstraintInstance();
 }
 
-void FConstraintInstance::GetProjectionAlphasOrTolerances(float& ProjectionLinearAlphaOrTolerance, float& ProjectionAngularAlphaOrTolerance) const
+void FConstraintInstance::GetProjectionParams(float& ProjectionLinearAlpha, float& ProjectionAngularAlpha, float& ProjectionLinearTolerance, float& ProjectionAngularTolerance) const
 {
-#if WITH_CHAOS
-	ProjectionLinearAlphaOrTolerance = ProfileInstance.ProjectionLinearAlpha;
-	ProjectionAngularAlphaOrTolerance = ProfileInstance.ProjectionAngularAlpha;
-#else
-	ProjectionLinearAlphaOrTolerance = ProfileInstance.ProjectionLinearTolerance;
-	ProjectionAngularAlphaOrTolerance = ProfileInstance.ProjectionAngularTolerance;
-
-#endif
+	ProjectionLinearAlpha = ProfileInstance.ProjectionLinearAlpha;
+	ProjectionAngularAlpha = ProfileInstance.ProjectionAngularAlpha;
+	ProjectionLinearTolerance = ProfileInstance.ProjectionLinearTolerance;
+	ProjectionAngularTolerance = ProfileInstance.ProjectionAngularTolerance;
 }
 
-void FConstraintInstance::SetProjectionParams(bool bEnableLinearProjection, bool bEnableAngularProjection, float ProjectionLinearAlphaOrTolerance, float ProjectionAngularAlphaOrTolerance)
+void FConstraintInstance::SetProjectionParams(bool bEnableProjection, float ProjectionLinearAlpha, float ProjectionAngularAlpha, float ProjectionLinearTolerance, float ProjectionAngularTolerance)
 {
-	ProfileInstance.bEnableLinearProjection = bEnableLinearProjection;
-	ProfileInstance.bEnableAngularProjection = bEnableAngularProjection;
-#if WITH_CHAOS
-	ProfileInstance.ProjectionLinearAlpha = ProjectionLinearAlphaOrTolerance;
-	ProfileInstance.ProjectionAngularAlpha = ProjectionAngularAlphaOrTolerance;
-#else
-	ProfileInstance.ProjectionLinearTolerance = ProjectionLinearAlphaOrTolerance;
-	ProfileInstance.ProjectionAngularTolerance = ProjectionAngularAlphaOrTolerance;
-#endif
+	ProfileInstance.bEnableProjection = bEnableProjection;
+	ProfileInstance.ProjectionLinearAlpha = ProjectionLinearAlpha;
+	ProfileInstance.ProjectionAngularAlpha = ProjectionAngularAlpha;
+	ProfileInstance.ProjectionLinearTolerance = ProjectionLinearTolerance;
+	ProfileInstance.ProjectionAngularTolerance = ProjectionAngularTolerance;
 
 	FPhysicsCommand::ExecuteWrite(ConstraintHandle, [&](const FPhysicsConstraintHandle& Constraint)
 		{
 #if WITH_CHAOS
-			const bool bEnableProjection = bEnableAngularProjection || ProjectionLinearAlphaOrTolerance;
-			const float LinearAlpha = bEnableLinearProjection ? ProjectionLinearAlphaOrTolerance : 0.0f;
-			const float AngularAlpha = bEnableAngularProjection ? ProjectionAngularAlphaOrTolerance : 0.0f;
-			FPhysicsInterface::SetProjectionEnabled_AssumesLocked(Constraint, bEnableProjection, LinearAlpha, AngularAlpha);
+			const float LinearAlpha = bEnableProjection ? ProjectionLinearAlpha : 0.0f;
+			const float AngularAlpha = bEnableProjection ? ProjectionAngularAlpha : 0.0f;
+			const float TeleportDistance = bEnableProjection ? ProjectionLinearTolerance : -1.0f;
+			const float TeleportAngle = bEnableProjection ? ProjectionAngularTolerance : -1.0f;
+			FPhysicsInterface::SetProjectionEnabled_AssumesLocked(Constraint, bEnableProjection, LinearAlpha, AngularAlpha, TeleportDistance, TeleportAngle);
 #else
-			if (bEnableLinearProjection)
+			if (bEnableProjection)
 			{
-				FPhysicsInterface::SetProjectionEnabled_AssumesLocked(Constraint, true, ProjectionLinearAlphaOrTolerance, ProjectionAngularAlphaOrTolerance);
+				FPhysicsInterface::SetProjectionEnabled_AssumesLocked(Constraint, true, ProjectionLinearTolerance, ProjectionAngularTolerance);
 			}
 			else
 			{
@@ -1251,7 +1241,7 @@ void FConstraintInstance::SetShockPropagationParams(bool bEnableShockPropagation
 void FConstraintInstance::EnableProjection()
 {
 #if WITH_CHAOS
-	SetProjectionParams(true, true, ProfileInstance.ProjectionLinearAlpha, ProfileInstance.ProjectionAngularAlpha);
+	SetProjectionParams(true, ProfileInstance.ProjectionLinearAlpha, ProfileInstance.ProjectionAngularAlpha, ProfileInstance.ProjectionLinearTolerance, ProfileInstance.ProjectionAngularTolerance);
 #else
 	SetProjectionParams(true, ProfileInstance.ProjectionLinearTolerance, ProfileInstance.ProjectionAngularTolerance);
 #endif
@@ -1260,7 +1250,7 @@ void FConstraintInstance::EnableProjection()
 void FConstraintInstance::DisableProjection()
 {
 #if WITH_CHAOS
-	SetProjectionParams(false, false, ProfileInstance.ProjectionLinearAlpha, ProfileInstance.ProjectionAngularAlpha);
+	SetProjectionParams(false, ProfileInstance.ProjectionLinearAlpha, ProfileInstance.ProjectionAngularAlpha, ProfileInstance.ProjectionLinearTolerance, ProfileInstance.ProjectionAngularTolerance);
 #else
 	SetProjectionParams(false, ProfileInstance.ProjectionLinearTolerance, ProfileInstance.ProjectionAngularTolerance);
 #endif
diff --git a/Engine/Source/Runtime/Engine/Private/PhysicsEngine/ConstraintInstanceBlueprintLibrary.cpp b/Engine/Source/Runtime/Engine/Private/PhysicsEngine/ConstraintInstanceBlueprintLibrary.cpp
--- a/Engine/Source/Runtime/Engine/Private/PhysicsEngine/ConstraintInstanceBlueprintLibrary.cpp
+++ b/Engine/Source/Runtime/Engine/Private/PhysicsEngine/ConstraintInstanceBlueprintLibrary.cpp
@@ -53,10 +53,14 @@ void UConstraintInstanceBlueprintLibrary::SetProjectionParams(
 	float ProjectionAngularAlpha
 	)
 {
+	// @todo: Need to add ProjectionLinearTolerance and ProjectionAngularTolerance to SetProjectionParams
 	if (FConstraintInstance* ConstraintInstance = Accessor.Get())
 	{
+		const float ProjectionLinearTolerance = ConstraintInstance->ProfileInstance.ProjectionLinearTolerance;
+		const float ProjectionAngularTolerance = ConstraintInstance->ProfileInstance.ProjectionAngularTolerance;
+
 		Accessor.Modify();
-		ConstraintInstance->SetProjectionParams(bEnableProjection, bEnableProjection, ProjectionLinearAlpha, ProjectionAngularAlpha);
+		ConstraintInstance->SetProjectionParams(bEnableProjection, ProjectionLinearAlpha, ProjectionAngularAlpha, ProjectionLinearTolerance, ProjectionAngularTolerance);
 	}
 }
 
@@ -67,10 +71,13 @@ void UConstraintInstanceBlueprintLibrary::GetProjectionParams(
 	float& ProjectionAngularAlpha
 	)
 {
+	// @todo: Need to add ProjectionLinearTolerance and ProjectionAngularTolerance to GetProjectionParams
 	if (FConstraintInstance* ConstraintInstance = Accessor.Get())
 	{
 		bEnableProjection = ConstraintInstance->IsProjectionEnabled();
-		ConstraintInstance->GetProjectionAlphasOrTolerances(ProjectionLinearAlpha, ProjectionAngularAlpha);
+		float ProjectionLinearTolerance = 0;
+		float ProjectionAngularTolerance = 0;
+		ConstraintInstance->GetProjectionParams(ProjectionLinearAlpha, ProjectionAngularAlpha, ProjectionLinearTolerance, ProjectionAngularTolerance);
 	}
 	else
 	{
diff --git a/Engine/Source/Runtime/Engine/Private/PhysicsEngine/ImmediatePhysics/ImmediatePhysicsChaos/ImmediatePhysicsJointHandle_Chaos.cpp b/Engine/Source/Runtime/Engine/Private/PhysicsEngine/ImmediatePhysics/ImmediatePhysicsChaos/ImmediatePhysicsJointHandle_Chaos.cpp
--- a/Engine/Source/Runtime/Engine/Private/PhysicsEngine/ImmediatePhysics/ImmediatePhysicsChaos/ImmediatePhysicsJointHandle_Chaos.cpp
+++ b/Engine/Source/Runtime/Engine/Private/PhysicsEngine/ImmediatePhysics/ImmediatePhysicsChaos/ImmediatePhysicsJointHandle_Chaos.cpp
@@ -45,11 +45,13 @@ namespace ImmediatePhysics_Chaos
 		ConstraintSettings.AngularLimits[(int32)EJointAngularConstraintIndex::Swing1] = FMath::DegreesToRadians(ConstraintInstance->GetAngularSwing1Limit());
 		ConstraintSettings.AngularLimits[(int32)EJointAngularConstraintIndex::Swing2] = FMath::DegreesToRadians(ConstraintInstance->GetAngularSwing2Limit());
 
-		ConstraintSettings.bProjectionEnabled = Profile.bEnableLinearProjection || Profile.bEnableAngularProjection;
+		ConstraintSettings.bProjectionEnabled = Profile.bEnableProjection;
 		ConstraintSettings.bShockPropagationEnabled = Profile.bEnableShockPropagation;
 
-		ConstraintSettings.LinearProjection = Profile.bEnableLinearProjection ? Profile.ProjectionLinearAlpha : 0.0f;
-		ConstraintSettings.AngularProjection = Profile.bEnableAngularProjection ? Profile.ProjectionAngularAlpha : 0.0f;
+		ConstraintSettings.LinearProjection = Profile.bEnableProjection ? Profile.ProjectionLinearAlpha : 0.0f;
+		ConstraintSettings.AngularProjection = Profile.bEnableProjection ? Profile.ProjectionAngularAlpha : 0.0f;
+		ConstraintSettings.TeleportDistance = Profile.bEnableProjection ? Profile.ProjectionLinearTolerance : -1.0f;
+		ConstraintSettings.TeleportAngle = Profile.bEnableProjection ? FMath::DegreesToRadians(Profile.ProjectionAngularTolerance) : -1.0f;
 		ConstraintSettings.ShockPropagation = Profile.bEnableShockPropagation ? Profile.ShockPropagationAlpha : 0.0f;
 		ConstraintSettings.ParentInvMassScale = Profile.bParentDominates ? (FReal)0 : (FReal)1;
 
@@ -180,6 +182,8 @@ namespace ImmediatePhysics_Chaos
 			ConstraintSettings.LinearSoftForceMode = EJointForceMode::Acceleration;
 			ConstraintSettings.LinearProjection = 0.0f;
 			ConstraintSettings.AngularProjection = 0.0f;
+			ConstraintSettings.TeleportDistance = -1.0f;
+			ConstraintSettings.TeleportAngle = -1.0f;
 		}
 
 		ConstraintSettings.Sanitize();
diff --git a/Engine/Source/Runtime/Experimental/Chaos/Private/Chaos/Evolution/SolverBodyContainer.cpp b/Engine/Source/Runtime/Experimental/Chaos/Private/Chaos/Evolution/SolverBodyContainer.cpp
--- a/Engine/Source/Runtime/Experimental/Chaos/Private/Chaos/Evolution/SolverBodyContainer.cpp
+++ b/Engine/Source/Runtime/Experimental/Chaos/Private/Chaos/Evolution/SolverBodyContainer.cpp
@@ -8,10 +8,13 @@ namespace Chaos
 	{
 		if (Particle.IsValid())
 		{
+			// Set the particle state
 			if (SolverBody.IsDynamic())
 			{
-				// Set the particle state
-				FParticleUtilities::SetCoMWorldTransform(Particle, SolverBody.CorrectedP(), SolverBody.CorrectedQ());
+				// Apply DP and DQ to P and Q
+				SolverBody.ApplyCorrections();
+
+				FParticleUtilities::SetCoMWorldTransform(Particle, SolverBody.P(), SolverBody.Q());
 				Particle->SetV(SolverBody.V());
 				Particle->SetW(SolverBody.W());
 			}
diff --git a/Engine/Source/Runtime/Experimental/Chaos/Private/Chaos/Joint/PBDJointCachedSolverGaussSeidel.cpp b/Engine/Source/Runtime/Experimental/Chaos/Private/Chaos/Joint/PBDJointCachedSolverGaussSeidel.cpp
--- a/Engine/Source/Runtime/Experimental/Chaos/Private/Chaos/Joint/PBDJointCachedSolverGaussSeidel.cpp
+++ b/Engine/Source/Runtime/Experimental/Chaos/Private/Chaos/Joint/PBDJointCachedSolverGaussSeidel.cpp
@@ -233,16 +233,11 @@ void FPBDJointCachedSolver::InitProjection(
 	const FPBDJointSolverSettings& SolverSettings,
 	const FPBDJointSettings& JointSettings)
 {
-	bool bHasLinearProjection = false;
-	{
-		const FReal LinearProjection = FPBDJointUtilities::GetLinearProjection(SolverSettings, JointSettings);
-		bHasLinearProjection = JointSettings.bProjectionEnabled && LinearProjection > 0.0;
-	}
-	bool bHasAngularProjection = false;
-	{
-		const FReal AngularProjection = FPBDJointUtilities::GetAngularProjection(SolverSettings, JointSettings);
-		bHasAngularProjection = JointSettings.bProjectionEnabled && AngularProjection > 0.0;
-	}
+	const FReal LinearProjection = FPBDJointUtilities::GetLinearProjection(SolverSettings, JointSettings);
+	const FReal AngularProjection = FPBDJointUtilities::GetAngularProjection(SolverSettings, JointSettings);
+	const bool bHasLinearProjection = JointSettings.bProjectionEnabled && ((LinearProjection > 0) || (JointSettings.TeleportDistance > 0));
+	const bool bHasAngularProjection = JointSettings.bProjectionEnabled && ((AngularProjection > 0) || (JointSettings.TeleportAngle > 0));
+
 	if(bHasLinearProjection || bHasAngularProjection)
 	{
 		ComputeBodyState(0);
@@ -250,11 +245,14 @@ void FPBDJointCachedSolver::InitProjection(
 
 		ConnectorRs[1].EnforceShortestArcWith(ConnectorRs[0]);
 
-		InvMScales[0] = 0.0;
-		InvMScales[1] = 1.0;
+		InvMScales[0] = 0;
+		InvMScales[1] = 1;
 
-		UpdateMass0();
-		UpdateMass1();
+		// Fake spherical inertia for angular projection (avoid cost of recomputing inertia at current world space rotation)
+		InvMs[0] = 0;
+		InvIs[0] = FMatrix33(0, 0, 0);
+		InvMs[1] = ConditionedInvMs[1];
+		InvIs[1] = FMatrix33::FromDiagonal(FVec3(ConditionedInvILs[1].GetMax()));
 
 		if(bHasLinearProjection)
 		{
@@ -1323,19 +1321,22 @@ void FPBDJointCachedSolver::InitAxisPositionDrive(
 	
 void FPBDJointCachedSolver::ApplyProjections(
 		const FReal Dt,
-		const FReal InSolverStiffness,
 		const FPBDJointSolverSettings& SolverSettings,
 		const FPBDJointSettings& JointSettings,
 		const bool bLastIteration)
 {
+	if (!JointSettings.bProjectionEnabled)
+	{
+		return;
+	}
+
 	if (!IsDynamic(1))
 	{
 		// If child is kinematic, return. 
 		return;
 	}
 
-	SolverStiffness = InSolverStiffness;
-	
+	SolverStiffness = 1;
 
 	if (SolverSettings.bSolvePositionLast)
 	{
@@ -1347,29 +1348,18 @@ void FPBDJointCachedSolver::ApplyProjections(
 		ApplyPositionProjection(Dt, SolverSettings, JointSettings);
 		ApplyRotationProjection(Dt, SolverSettings, JointSettings);
 	}
+
 	if(bLastIteration)
 	{
-		//Final position fixup
-		const TVec3<EJointMotionType>& LinearMotion = JointSettings.LinearMotionTypes;
-		const bool bLinearLocked = (LinearMotion[0] == EJointMotionType::Locked) && (LinearMotion[1] == EJointMotionType::Locked) && (LinearMotion[2] == EJointMotionType::Locked);
-		if (bLinearLocked)
+		// Add velocity correction from the net projection motion
+		// @todo(chaos): this should be a joint setting?
+		if (Chaos_Joint_VelProjectionAlpha > 0.0f)
 		{
-			const FReal LinearProjection = FPBDJointUtilities::GetLinearProjection(SolverSettings, JointSettings);
-			if (JointSettings.bProjectionEnabled && (LinearProjection > 0))
-			{
-				const FVec3 DP1 = -LinearProjection * (ConnectorXs[1] - ConnectorXs[0]);
-				ApplyPositionDelta(1, DP1);
-			}
-		
-			// Add velocity correction from the net projection motion
-			if (Chaos_Joint_VelProjectionAlpha > 0.0f)
-			{
-				const FSolverReal VelocityScale = Chaos_Joint_VelProjectionAlpha  / static_cast<FSolverReal>(Dt);
-				const FSolverVec3 DV1 = Body1().DP() * VelocityScale;
-				const FSolverVec3 DW1 = Body1().DQ()* VelocityScale;
+			const FSolverReal VelocityScale = Chaos_Joint_VelProjectionAlpha  / static_cast<FSolverReal>(Dt);
+			const FSolverVec3 DV1 = Body1().DP() * VelocityScale;
+			const FSolverVec3 DW1 = Body1().DQ()* VelocityScale;
 		
-				Body(1).ApplyVelocityDelta(DV1, DW1);
-			}
+			Body(1).ApplyVelocityDelta(DV1, DW1);
 		}
 	}
 }
@@ -1380,6 +1370,10 @@ void FPBDJointCachedSolver::ApplyRotationProjection(
 			const FPBDJointSettings& JointSettings)
 {
 	const FReal AngularProjection = FPBDJointUtilities::GetAngularProjection(SolverSettings, JointSettings);
+	if (AngularProjection == 0)
+	{
+		return;
+	}
 	
 	const TVec3<EJointMotionType>& LinearMotion = JointSettings.LinearMotionTypes;
 	const bool bLinearLocked = (LinearMotion[0] == EJointMotionType::Locked) && (LinearMotion[1] == EJointMotionType::Locked) && (LinearMotion[2] == EJointMotionType::Locked);
@@ -1388,40 +1382,37 @@ void FPBDJointCachedSolver::ApplyRotationProjection(
 	{
 		if(RotationConstraints.bValidDatas[ConstraintIndex])
 		{
-			if (JointSettings.bProjectionEnabled && (AngularProjection > 0.0f))
-			{
-				FReal DeltaAngle = RotationConstraints.ConstraintCX[ConstraintIndex] +
-					FVec3::DotProduct(Body(1).DQ() - Body(0).DQ(), RotationConstraints.ConstraintAxis[ConstraintIndex]);
+			FReal DeltaAngle = RotationConstraints.ConstraintCX[ConstraintIndex] +
+				FVec3::DotProduct(Body(1).DQ() - Body(0).DQ(), RotationConstraints.ConstraintAxis[ConstraintIndex]);
 	
-				bool NeedsSolve = false;
-				if(RotationConstraints.bLimitsCheck[ConstraintIndex])
+			bool NeedsSolve = false;
+			if(RotationConstraints.bLimitsCheck[ConstraintIndex])
+			{
+				if(DeltaAngle > RotationConstraints.ConstraintLimits[ConstraintIndex] )
 				{
-					if(DeltaAngle > RotationConstraints.ConstraintLimits[ConstraintIndex] )
-					{
-						DeltaAngle -= RotationConstraints.ConstraintLimits[ConstraintIndex];
-						NeedsSolve = true;
-					}
-					else if(DeltaAngle < -RotationConstraints.ConstraintLimits[ConstraintIndex])
-					{
-						DeltaAngle += RotationConstraints.ConstraintLimits[ConstraintIndex];
-						NeedsSolve = true;
+					DeltaAngle -= RotationConstraints.ConstraintLimits[ConstraintIndex];
+					NeedsSolve = true;
+				}
+				else if(DeltaAngle < -RotationConstraints.ConstraintLimits[ConstraintIndex])
+				{
+					DeltaAngle += RotationConstraints.ConstraintLimits[ConstraintIndex];
+					NeedsSolve = true;
 			
-					}
 				}
+			}
 
-				if (!RotationConstraints.bLimitsCheck[ConstraintIndex] || (RotationConstraints.bLimitsCheck[ConstraintIndex] && NeedsSolve && FMath::Abs(DeltaAngle) > AngleTolerance))
-				{
-					const FReal IM = -FVec3::DotProduct(RotationConstraints.ConstraintAxis[ConstraintIndex], RotationConstraints.ConstraintDRAxis[ConstraintIndex][1]);
-					const FReal DeltaLambda = SolverStiffness * RotationConstraints.ConstraintHardStiffness[ConstraintIndex] * DeltaAngle / IM;
+			if (!RotationConstraints.bLimitsCheck[ConstraintIndex] || (RotationConstraints.bLimitsCheck[ConstraintIndex] && NeedsSolve && FMath::Abs(DeltaAngle) > AngleTolerance))
+			{
+				const FReal IM = -FVec3::DotProduct(RotationConstraints.ConstraintAxis[ConstraintIndex], RotationConstraints.ConstraintDRAxis[ConstraintIndex][1]);
+				const FReal DeltaLambda = SolverStiffness * RotationConstraints.ConstraintHardStiffness[ConstraintIndex] * DeltaAngle / IM;
 					
-					const FVec3 DR1 = AngularProjection * RotationConstraints.ConstraintDRAxis[ConstraintIndex][1] * DeltaLambda;
-					ApplyRotationDelta(1, DR1);
+				const FVec3 DR1 = AngularProjection * RotationConstraints.ConstraintDRAxis[ConstraintIndex][1] * DeltaLambda;
+				ApplyRotationDelta(1, DR1);
 
-					if(bLinearLocked)
-					{
-						const FVec3 DP1 = -AngularProjection * FVec3::CrossProduct(DR1, PositionConstraints.ConstraintArms[ConstraintIndex][1]);
-						ApplyPositionDelta(1,DP1);
-					}
+				if(bLinearLocked)
+				{
+					const FVec3 DP1 = -AngularProjection * FVec3::CrossProduct(DR1, PositionConstraints.ConstraintArms[ConstraintIndex][1]);
+					ApplyPositionDelta(1,DP1);
 				}
 			}
 		}
@@ -1434,52 +1425,117 @@ void FPBDJointCachedSolver::ApplyPositionProjection(
 		const FPBDJointSettings& JointSettings)
 {
 	const FReal LinearProjection = FPBDJointUtilities::GetLinearProjection(SolverSettings, JointSettings);
+	if (LinearProjection == 0)
+	{
+		return;
+	}
 
 	for(int32 ConstraintIndex = 0; ConstraintIndex < 3; ++ConstraintIndex)
 	{
 		if(PositionConstraints.bValidDatas[ConstraintIndex])
 		{
-			if (JointSettings.bProjectionEnabled && (LinearProjection > 0.0f))
-			{
-				const FVec3 CX = Body(1).DP()  - Body(0).DP() +
-					FVec3::CrossProduct(Body(1).DQ(), PositionConstraints.ConstraintArms[ConstraintIndex][1]) -
-					FVec3::CrossProduct(Body(0).DQ(), PositionConstraints.ConstraintArms[ConstraintIndex][0]);
+			const FVec3 CX = Body(1).DP()  - Body(0).DP() +
+				FVec3::CrossProduct(Body(1).DQ(), PositionConstraints.ConstraintArms[ConstraintIndex][1]) -
+				FVec3::CrossProduct(Body(0).DQ(), PositionConstraints.ConstraintArms[ConstraintIndex][0]);
 				
-				FReal DeltaPosition = PositionConstraints.ConstraintCX[ConstraintIndex] + FVec3::DotProduct(CX, PositionConstraints.ConstraintAxis[ConstraintIndex]);
+			FReal DeltaPosition = PositionConstraints.ConstraintCX[ConstraintIndex] + FVec3::DotProduct(CX, PositionConstraints.ConstraintAxis[ConstraintIndex]);
 				
-				bool NeedsSolve = false;
-				if(PositionConstraints.bLimitsCheck[ConstraintIndex])
+			bool NeedsSolve = false;
+			if(PositionConstraints.bLimitsCheck[ConstraintIndex])
+			{
+				if(DeltaPosition > PositionConstraints.ConstraintLimits[ConstraintIndex] )
 				{
-					if(DeltaPosition > PositionConstraints.ConstraintLimits[ConstraintIndex] )
-					{
-						DeltaPosition -= PositionConstraints.ConstraintLimits[ConstraintIndex];
-						NeedsSolve = true;
-					}
-					else if(DeltaPosition < -PositionConstraints.ConstraintLimits[ConstraintIndex])
-					{
-						DeltaPosition += PositionConstraints.ConstraintLimits[ConstraintIndex];
-						NeedsSolve = true;
-					}
+					DeltaPosition -= PositionConstraints.ConstraintLimits[ConstraintIndex];
+					NeedsSolve = true;
 				}
-				if (!PositionConstraints.bLimitsCheck[ConstraintIndex] || (PositionConstraints.bLimitsCheck[ConstraintIndex] && NeedsSolve && FMath::Abs(DeltaPosition ) > PositionTolerance))
+				else if(DeltaPosition < -PositionConstraints.ConstraintLimits[ConstraintIndex])
 				{
-					const FVec3 AngularAxis1 = FVec3::CrossProduct(PositionConstraints.ConstraintArms[ConstraintIndex][1], PositionConstraints.ConstraintAxis[ConstraintIndex]);
-					const FReal IM = InvM(1) - FVec3::DotProduct(AngularAxis1, PositionConstraints.ConstraintDRAxis[ConstraintIndex][1]);
-					const FReal DeltaLambda = SolverStiffness * PositionConstraints.ConstraintHardStiffness[ConstraintIndex] * DeltaPosition / IM;
+					DeltaPosition += PositionConstraints.ConstraintLimits[ConstraintIndex];
+					NeedsSolve = true;
+				}
+			}
+			if (!PositionConstraints.bLimitsCheck[ConstraintIndex] || (PositionConstraints.bLimitsCheck[ConstraintIndex] && NeedsSolve && FMath::Abs(DeltaPosition ) > PositionTolerance))
+			{
+				const FVec3 AngularAxis1 = FVec3::CrossProduct(PositionConstraints.ConstraintArms[ConstraintIndex][1], PositionConstraints.ConstraintAxis[ConstraintIndex]);
+				const FReal IM = InvM(1) - FVec3::DotProduct(AngularAxis1, PositionConstraints.ConstraintDRAxis[ConstraintIndex][1]);
+				const FReal DeltaLambda = SolverStiffness * PositionConstraints.ConstraintHardStiffness[ConstraintIndex] * DeltaPosition / IM;
 	
-					const FVec3 DX = PositionConstraints.ConstraintAxis[ConstraintIndex] * DeltaLambda;
+				const FVec3 DX = PositionConstraints.ConstraintAxis[ConstraintIndex] * DeltaLambda;
 					
-					const FVec3 DP1 = -LinearProjection * InvM(1) * DX;
-					const FVec3 DR1 = LinearProjection * PositionConstraints.ConstraintDRAxis[ConstraintIndex][1] * DeltaLambda;
-					
-					ApplyPositionDelta(1,DP1);
-					ApplyRotationDelta(1,DR1);
+				const FVec3 DP1 = -LinearProjection * InvM(1) * DX;
+				const FVec3 DR1 = LinearProjection * PositionConstraints.ConstraintDRAxis[ConstraintIndex][1] * DeltaLambda;
+
+				ApplyPositionDelta(1, DP1);
+				ApplyRotationDelta(1, DR1);
+			}
+		}
+	}
+}
+
+void FPBDJointCachedSolver::ApplyTeleports(
+	const FReal Dt,
+	const FPBDJointSolverSettings& SolverSettings,
+	const FPBDJointSettings& JointSettings)
+{
+	ApplyRotationTeleport(Dt, SolverSettings, JointSettings);
+	ApplyPositionTeleport(Dt, SolverSettings, JointSettings);
+}
+
+
+void FPBDJointCachedSolver::ApplyPositionTeleport(
+	const FReal Dt,
+	const FPBDJointSolverSettings& SolverSettings,
+	const FPBDJointSettings& JointSettings)
+{
+	if (JointSettings.TeleportDistance <= 0)
+	{
+		return;
+	}
+
+	for (int32 ConstraintIndex = 0; ConstraintIndex < 3; ++ConstraintIndex)
+	{
+		if (PositionConstraints.bValidDatas[ConstraintIndex])
+		{
+			FReal DeltaPosition = PositionConstraints.ConstraintCX[ConstraintIndex];
+
+			bool NeedsSolve = false;
+			if (PositionConstraints.bLimitsCheck[ConstraintIndex])
+			{
+				if (DeltaPosition > PositionConstraints.ConstraintLimits[ConstraintIndex])
+				{
+					DeltaPosition -= PositionConstraints.ConstraintLimits[ConstraintIndex];
+					NeedsSolve = true;
+				}
+				else if (DeltaPosition < -PositionConstraints.ConstraintLimits[ConstraintIndex])
+				{
+					DeltaPosition += PositionConstraints.ConstraintLimits[ConstraintIndex];
+					NeedsSolve = true;
+				}
+			}
+			if (!PositionConstraints.bLimitsCheck[ConstraintIndex] || (PositionConstraints.bLimitsCheck[ConstraintIndex] && NeedsSolve))
+			{
+				if (FMath::Abs(DeltaPosition) > JointSettings.TeleportDistance)
+				{
+					const FVec3 DP1 = -DeltaPosition * PositionConstraints.ConstraintAxis[ConstraintIndex];
+					ApplyPositionDelta(1, DP1);
 				}
 			}
 		}
 	}
 }
 
+void FPBDJointCachedSolver::ApplyRotationTeleport(
+	const FReal Dt,
+	const FPBDJointSolverSettings& SolverSettings,
+	const FPBDJointSettings& JointSettings)
+{
+	if (JointSettings.TeleportAngle <= 0)
+	{
+		return;
+	}
+}
+
+
 /** APPLY POSITION  DRIVES *********************************************************************************/
 
 void FPBDJointCachedSolver::ApplyPositionDrives(
diff --git a/Engine/Source/Runtime/Experimental/Chaos/Private/Chaos/Joint/PBDJointSolverGaussSeidel.cpp b/Engine/Source/Runtime/Experimental/Chaos/Private/Chaos/Joint/PBDJointSolverGaussSeidel.cpp
--- a/Engine/Source/Runtime/Experimental/Chaos/Private/Chaos/Joint/PBDJointSolverGaussSeidel.cpp
+++ b/Engine/Source/Runtime/Experimental/Chaos/Private/Chaos/Joint/PBDJointSolverGaussSeidel.cpp
@@ -314,7 +314,7 @@ namespace Chaos
 		// Position Projection
 		const FReal LinearProjection = FPBDJointUtilities::GetLinearProjection(SolverSettings, JointSettings);
 		const bool bLinearSoft = FPBDJointUtilities::GetSoftLinearLimitEnabled(SolverSettings, JointSettings);
-		const bool bLinearProjectionEnabled = (bLinearSoft && JointSettings.bSoftProjectionEnabled) || (!bLinearSoft && JointSettings.bProjectionEnabled);
+		const bool bLinearProjectionEnabled = (!bLinearSoft && JointSettings.bProjectionEnabled);
 		const TVec3<EJointMotionType>& LinearMotion = JointSettings.LinearMotionTypes;
 		const bool bLinearLocked = (LinearMotion[0] == EJointMotionType::Locked) && (LinearMotion[1] == EJointMotionType::Locked) && (LinearMotion[2] == EJointMotionType::Locked);
 		const bool bLinearLimited = (LinearMotion[0] == EJointMotionType::Limited) && (LinearMotion[1] == EJointMotionType::Limited) && (LinearMotion[2] == EJointMotionType::Limited);
@@ -345,7 +345,7 @@ namespace Chaos
 
 		// Twist projection
 		const bool bTwistSoft = FPBDJointUtilities::GetSoftTwistLimitEnabled(SolverSettings, JointSettings);
-		const bool bTwistProjectionEnabled = SolverSettings.bEnableTwistLimits && ((bTwistSoft && JointSettings.bSoftProjectionEnabled) || (!bTwistSoft && JointSettings.bProjectionEnabled));
+		const bool bTwistProjectionEnabled = SolverSettings.bEnableTwistLimits && (!bTwistSoft && JointSettings.bProjectionEnabled);
 		if (bTwistProjectionEnabled && (AngularProjection > 0.0f))
 		{
 			const EJointMotionType TwistMotion = JointSettings.AngularMotionTypes[(int32)EJointAngularConstraintIndex::Twist];
@@ -357,7 +357,7 @@ namespace Chaos
 
 		// Swing projection
 		const bool bSwingSoft = FPBDJointUtilities::GetSoftSwingLimitEnabled(SolverSettings, JointSettings);
-		const bool bSwingProjectionEnabled = SolverSettings.bEnableSwingLimits && ((bSwingSoft && JointSettings.bSoftProjectionEnabled) || (!bSwingSoft && JointSettings.bProjectionEnabled));
+		const bool bSwingProjectionEnabled = SolverSettings.bEnableSwingLimits && (!bSwingSoft && JointSettings.bProjectionEnabled);
 		if (bSwingProjectionEnabled && (AngularProjection > 0.0f))
 		{
 			const EJointMotionType Swing1Motion = JointSettings.AngularMotionTypes[(int32)EJointAngularConstraintIndex::Swing1];
@@ -401,7 +401,6 @@ namespace Chaos
 
 	void FPBDJointSolver::ApplyProjections(
 		const FReal Dt,
-		const FReal InSolverStiffness,
 		const FPBDJointSolverSettings& SolverSettings,
 		const FPBDJointSettings& JointSettings)
 	{
@@ -412,7 +411,7 @@ namespace Chaos
 			return;
 		}
 
-		SolverStiffness = InSolverStiffness;
+		SolverStiffness = 1;
 
 
 		FVec3 DP1 = FVec3(0);
@@ -436,7 +435,7 @@ namespace Chaos
 		{
 			const FReal LinearProjection = FPBDJointUtilities::GetLinearProjection(SolverSettings, JointSettings);
 			const bool bLinearSoft = FPBDJointUtilities::GetSoftLinearLimitEnabled(SolverSettings, JointSettings);
-			const bool bLinearProjectionEnabled = (bLinearSoft && JointSettings.bSoftProjectionEnabled) || (!bLinearSoft && JointSettings.bProjectionEnabled);
+			const bool bLinearProjectionEnabled = (!bLinearSoft && JointSettings.bProjectionEnabled);
 			if (bLinearProjectionEnabled && (LinearProjection > 0))
 			{
 				ApplyTranslateProjection(Dt, SolverSettings, JointSettings, LinearProjection, DP1, DR1);
diff --git a/Engine/Source/Runtime/Experimental/Chaos/Private/Chaos/PBDConstraintRule.cpp b/Engine/Source/Runtime/Experimental/Chaos/Private/Chaos/PBDConstraintRule.cpp
--- a/Engine/Source/Runtime/Experimental/Chaos/Private/Chaos/PBDConstraintRule.cpp
+++ b/Engine/Source/Runtime/Experimental/Chaos/Private/Chaos/PBDConstraintRule.cpp
@@ -247,6 +247,13 @@ namespace Chaos
 		return IslandGroup  ? Constraints.ApplyPhase2Serial(Dt, It, NumIts, *IslandGroup) : false;
 	}
 
+	template<class ConstraintType>
+	bool TPBDConstraintIslandRule<ConstraintType>::ApplyProjection(const FReal Dt, int32 GroupIndex, const int32 It, const int32 NumIts)
+	{
+		FPBDIslandGroup* IslandGroup = ConstraintGraph->GetIslandGroup(GroupIndex);
+		return IslandGroup ? Constraints.ApplyPhase3Serial(Dt, It, NumIts, *IslandGroup) : false;
+	}
+
 	template<class ConstraintType>
 	void TPBDConstraintIslandRule<ConstraintType>::InitializeAccelerationStructures()
 	{
diff --git a/Engine/Source/Runtime/Experimental/Chaos/Private/Chaos/PBDJointConstraints.cpp b/Engine/Source/Runtime/Experimental/Chaos/Private/Chaos/PBDJointConstraints.cpp
--- a/Engine/Source/Runtime/Experimental/Chaos/Private/Chaos/PBDJointConstraints.cpp
+++ b/Engine/Source/Runtime/Experimental/Chaos/Private/Chaos/PBDJointConstraints.cpp
@@ -194,11 +194,12 @@ namespace Chaos
 		, LinearProjection(0)
 		, AngularProjection(0)
 		, ShockPropagation(0)
+		, TeleportDistance(0)
+		, TeleportAngle(0)
 		, ParentInvMassScale(1)
 		, bCollisionEnabled(true)
 		, bProjectionEnabled(false)
 		, bShockPropagationEnabled(false)
-		, bSoftProjectionEnabled(false)
 		, LinearMotionTypes({ EJointMotionType::Locked, EJointMotionType::Locked, EJointMotionType::Locked })
 		, LinearLimit(FLT_MAX)
 		, AngularMotionTypes({ EJointMotionType::Free, EJointMotionType::Free, EJointMotionType::Free })
@@ -1056,12 +1057,12 @@ namespace Chaos
 		return (NumActive > 0);
 	}
 	
-	void FPBDJointConstraints::PreparePhase3Serial(const FReal Dt, const int32 It, FPBDIslandSolverData& SolverData)
+	void FPBDJointConstraints::PreparePhase3Serial(const FReal Dt, FPBDIslandSolverData& SolverData)
 	{
 		CSV_SCOPED_TIMING_STAT(Chaos, ProjectJointConstraints);
 		SCOPE_CYCLE_COUNTER(STAT_Joints_ApplyProjection);
 
-		if(It == 0 && Settings.bUseLinearSolver)
+		if(Settings.bUseLinearSolver)
 		{
 			for (int32 ConstraintIndex : SolverData.GetConstraintIndices(ContainerId))
 			{
@@ -1086,7 +1087,10 @@ namespace Chaos
 		SCOPE_CYCLE_COUNTER(STAT_Joints_ApplyProjection);
 
 		// Prepare phase 3 for the linear solver in order to partially re-init the solver
-		PreparePhase3Serial(Dt, It, SolverData);
+		if (It == 0)
+		{
+			PreparePhase3Serial(Dt, SolverData);
+		}
 
 		// UpdateConstraintProjection(It, NumIts, SolverData);
 		for (int32 ConstraintIndex : SolverData.GetConstraintIndices(ContainerId))
@@ -1395,7 +1399,7 @@ namespace Chaos
 		UE_LOG(LogChaosJoint, VeryVerbose, TEXT("Project Joint Constraint %d %s %s (dt = %f; it = %d / %d)"), ConstraintIndex, *Constraint[0]->ToString(), *Constraint[1]->ToString(), Dt, It, NumIts);
 
 		const FPBDJointSettings& JointSettings = ConstraintSettings[ConstraintIndex];
-		if (!JointSettings.bProjectionEnabled && !JointSettings.bSoftProjectionEnabled)
+		if (!JointSettings.bProjectionEnabled)
 		{
 			return false;
 		}
@@ -1420,8 +1424,12 @@ namespace Chaos
 				return false;
 			}
 
-			const FReal IterationStiffness = CalculateIterationStiffness(It, NumIts);
-			Solver.ApplyProjections(Dt, IterationStiffness, Settings, JointSettings, (It==(NumIts-1)));
+			if (It == 0)
+			{
+				Solver.ApplyTeleports(Dt, Settings, JointSettings);
+			}
+
+			Solver.ApplyProjections(Dt, Settings, JointSettings, (It==(NumIts-1)));
 		}
 		else
 		{
@@ -1439,8 +1447,7 @@ namespace Chaos
 				Solver.UpdateMasses(FReal(0), FReal(1));
 			}
 
-			const FReal IterationStiffness = CalculateIterationStiffness(It, NumIts);
-			Solver.ApplyProjections(Dt, IterationStiffness, Settings, JointSettings);
+			Solver.ApplyProjections(Dt, Settings, JointSettings);
 
 			if (SolverType == EConstraintSolverType::StandardPbd)
 			{
diff --git a/Engine/Source/Runtime/Experimental/Chaos/Private/Chaos/PBDRigidsEvolution.cpp b/Engine/Source/Runtime/Experimental/Chaos/Private/Chaos/PBDRigidsEvolution.cpp
--- a/Engine/Source/Runtime/Experimental/Chaos/Private/Chaos/PBDRigidsEvolution.cpp
+++ b/Engine/Source/Runtime/Experimental/Chaos/Private/Chaos/PBDRigidsEvolution.cpp
@@ -207,7 +207,7 @@ namespace Chaos
 		}
 	};
 
-	FPBDRigidsEvolutionBase::FPBDRigidsEvolutionBase(FPBDRigidsSOAs& InParticles, THandleArray<FChaosPhysicsMaterial>& InSolverPhysicsMaterials, int32 InNumIterations, int32 InNumPushOutIterations, bool InIsSingleThreaded)
+	FPBDRigidsEvolutionBase::FPBDRigidsEvolutionBase(FPBDRigidsSOAs& InParticles, THandleArray<FChaosPhysicsMaterial>& InSolverPhysicsMaterials, bool InIsSingleThreaded)
 	    : Particles(InParticles)
 		, SolverPhysicsMaterials(InSolverPhysicsMaterials)
 		, InternalAcceleration(nullptr)
@@ -216,8 +216,9 @@ namespace Chaos
 		, bIsSingleThreaded(InIsSingleThreaded)
 		, bCanStartAsyncTasks(true)
 		, LatestExternalTimestampConsumed_Internal(-1)
-		, NumIterations(InNumIterations)
-		, NumPushOutIterations(InNumPushOutIterations)
+		, NumPositionIterations(0)
+		, NumVelocityIterations(0)
+		, NumProjectionIterations(0)
 		, SpatialCollectionFactory(new FDefaultCollectionFactory())
 	{
 		Particles.GetParticleHandles().AddArray(&PhysicsMaterials);
diff --git a/Engine/Source/Runtime/Experimental/Chaos/Private/Chaos/PBDRigidsEvolutionGBF.cpp b/Engine/Source/Runtime/Experimental/Chaos/Private/Chaos/PBDRigidsEvolutionGBF.cpp
--- a/Engine/Source/Runtime/Experimental/Chaos/Private/Chaos/PBDRigidsEvolutionGBF.cpp
+++ b/Engine/Source/Runtime/Experimental/Chaos/Private/Chaos/PBDRigidsEvolutionGBF.cpp
@@ -121,6 +121,7 @@ namespace Chaos
 		DECLARE_CYCLE_STAT(TEXT("FPBDRigidsEvolutionGBF::ApplyConstraintsPhase1"), STAT_Evolution_ApplyConstraintsPhase1, STATGROUP_Chaos);
 		DECLARE_CYCLE_STAT(TEXT("FPBDRigidsEvolutionGBF::UpdateVelocities"), STAT_Evolution_UpdateVelocites, STATGROUP_Chaos);
 		DECLARE_CYCLE_STAT(TEXT("FPBDRigidsEvolutionGBF::ApplyConstraintsPhase2"), STAT_Evolution_ApplyConstraintsPhase2, STATGROUP_Chaos);
+		DECLARE_CYCLE_STAT(TEXT("FPBDRigidsEvolutionGBF::ApplyConstraintsPhase3"), STAT_Evolution_ApplyConstraintsPhase3, STATGROUP_Chaos);
 		DECLARE_CYCLE_STAT(TEXT("FPBDRigidsEvolutionGBF::DetectCollisions"), STAT_Evolution_DetectCollisions, STATGROUP_Chaos);
 		DECLARE_CYCLE_STAT(TEXT("FPBDRigidsEvolutionGBF::TransferJointCollisions"), STAT_Evolution_TransferJointCollisions, STATGROUP_Chaos);
 		DECLARE_CYCLE_STAT(TEXT("FPBDRigidsEvolutionGBF::PostDetectCollisionsCallback"), STAT_Evolution_PostDetectCollisionsCallback, STATGROUP_Chaos);
@@ -290,6 +291,7 @@ struct alignas(64) FPerIslandStats
 		PerIslandSolve_ApplyTotalSerialized,
 		PerIslandSolve_UpdateVelocitiesTotalSerialized,
 		PerIslandSolve_ApplyPushOutTotalSerialized,
+		PerIslandSolve_ApplyProjectionTotalSerialized,
 		PerIslandSolve_ScatterTotalSerialized,
 		NumStats
 	};
@@ -301,6 +303,7 @@ struct alignas(64) FPerIslandStats
 		CSV_CUSTOM_STAT_PER_ISLAND_HELPER(PerIslandSolve_ApplyTotalSerialized);
 		CSV_CUSTOM_STAT_PER_ISLAND_HELPER(PerIslandSolve_UpdateVelocitiesTotalSerialized);
 		CSV_CUSTOM_STAT_PER_ISLAND_HELPER(PerIslandSolve_ApplyPushOutTotalSerialized);
+		CSV_CUSTOM_STAT_PER_ISLAND_HELPER(PerIslandSolve_ApplyProjectionTotalSerialized);
 		CSV_CUSTOM_STAT_PER_ISLAND_HELPER(PerIslandSolve_ScatterTotalSerialized);
 	}
 
@@ -467,12 +470,13 @@ void FPBDRigidsEvolutionGBF::AdvanceOneTimeStepImpl(const FReal Dt, const FSubSt
 		PerIslandStats.AddDefaulted(NumGroups);
 #endif
 
-		PhysicsParallelFor(NumGroups, [&](int32 GroupIndex) {
-
-			if(GetConstraintGraph().GetIslandGroups().IsValidIndex(GroupIndex) &&
-			  !GetConstraintGraph().GetIslandGroup(GroupIndex)->IsValid())
-			  	
+		PhysicsParallelFor(NumGroups, [&](int32 GroupIndex) 
+		{
+			const bool bValidIsland = GetConstraintGraph().GetIslandGroups().IsValidIndex(GroupIndex) && GetConstraintGraph().GetIslandGroup(GroupIndex)->IsValid();
+		  	if (!bValidIsland)
+			{
 				return;
+			}
 			
 			// Reload cache if necessary
 			{
@@ -528,6 +532,14 @@ void FPBDRigidsEvolutionGBF::AdvanceOneTimeStepImpl(const FReal Dt, const FSubSt
 				ApplyConstraintsPhase2(Dt, GroupIndex);
 			}
 
+			// Apply projections. This is a phase only implemented by joints that applies a semi-physical
+			// correction and/or teleport to fix any errors remaining after the position and velocity solve steps
+			{
+				SCOPE_CYCLE_COUNTER(STAT_Evolution_ApplyConstraintsPhase3);
+				CSV_SCOPED_PER_ISLAND_TIMING_STAT(PerIslandSolve_ApplyProjectionTotalSerialized);
+				ApplyConstraintsPhase3(Dt, GroupIndex);
+			}
+
 			// Update the particles with the results of the constraint solvers, and also update constraint data
 			// that is accessed externally (net impusles, break info, etc)
 			{
@@ -712,14 +724,13 @@ void FPBDRigidsEvolutionGBF::ScatterSolverOutput(FReal Dt, int32 GroupIndex)
 
 void FPBDRigidsEvolutionGBF::ApplyConstraintsPhase1(const FReal Dt, int32 GroupIndex)
 {
-	int32 LocalNumIterations = ChaosNumContactIterationsOverride >= 0 ? ChaosNumContactIterationsOverride : NumIterations;
 	// @todo(ccaulfield): track whether we are sufficiently solved and can early-out
-	for (int i = 0; i < LocalNumIterations; ++i)
+	for (int i = 0; i < NumPositionIterations; ++i)
 	{
 		bool bNeedsAnotherIteration = false;
 		for (FPBDConstraintGraphRule* ConstraintRule : PrioritizedConstraintRules)
 		{
-			bNeedsAnotherIteration |= ConstraintRule->ApplyConstraints(Dt, GroupIndex, i, LocalNumIterations);
+			bNeedsAnotherIteration |= ConstraintRule->ApplyConstraints(Dt, GroupIndex, i, NumPositionIterations);
 		}
 
 		if (ChaosRigidsEvolutionApplyAllowEarlyOutCVar && !bNeedsAnotherIteration)
@@ -739,14 +750,13 @@ void FPBDRigidsEvolutionGBF::SetImplicitVelocities(const FReal Dt, int32 GroupIn
 
 void FPBDRigidsEvolutionGBF::ApplyConstraintsPhase2(const FReal Dt, int32 GroupIndex)
 {
-	int32 LocalNumPushOutIterations = ChaosNumPushOutIterationsOverride >= 0 ? ChaosNumPushOutIterationsOverride : NumPushOutIterations;
 	bool bNeedsAnotherIteration = true;
-	for (int32 It = 0; It < LocalNumPushOutIterations; ++It)
+	for (int32 It = 0; It < NumVelocityIterations; ++It)
 	{
 		bNeedsAnotherIteration = false;
 		for (FPBDConstraintGraphRule* ConstraintRule : PrioritizedConstraintRules)
 		{
-			bNeedsAnotherIteration |= ConstraintRule->ApplyPushOut(Dt, GroupIndex, It, LocalNumPushOutIterations);
+			bNeedsAnotherIteration |= ConstraintRule->ApplyPushOut(Dt, GroupIndex, It, NumVelocityIterations);
 		}
 
 		if (ChaosRigidsEvolutionApplyPushoutAllowEarlyOutCVar && !bNeedsAnotherIteration)
@@ -756,6 +766,25 @@ void FPBDRigidsEvolutionGBF::ApplyConstraintsPhase2(const FReal Dt, int32 GroupI
 	}
 }
 
+void FPBDRigidsEvolutionGBF::ApplyConstraintsPhase3(const FReal Dt, int32 GroupIndex)
+{
+	// Update the body transforms from the deltas calculated in the constraint solve phases 1 and 2
+	// NOTE: deliberately not updating the world-space inertia as it is not currently needed by joint projection
+	// and no other constraints implement projection
+	if (FPBDIslandGroup* IslandGroup = ConstraintGraph.GetIslandGroups()[GroupIndex].Get())
+	{
+		IslandGroup->GetBodyContainer().ApplyCorrections();
+	}
+
+	for (int32 It = 0; It < NumProjectionIterations; ++It)
+	{
+		for (FPBDConstraintGraphRule* ConstraintRule : PrioritizedConstraintRules)
+		{
+			ConstraintRule->ApplyProjection(Dt, GroupIndex, It, NumProjectionIterations);
+		}
+	}
+}
+
 void FPBDRigidsEvolutionGBF::SetIsDeterministic(const bool bInIsDeterministic)
 {
 	// We detect collisions in parallel, so order is non-deterministic without additional processing
@@ -763,11 +792,11 @@ void FPBDRigidsEvolutionGBF::SetIsDeterministic(const bool bInIsDeterministic)
 }
 
 FPBDRigidsEvolutionGBF::FPBDRigidsEvolutionGBF(FPBDRigidsSOAs& InParticles,THandleArray<FChaosPhysicsMaterial>& SolverPhysicsMaterials, const TArray<ISimCallbackObject*>* InCollisionModifiers, bool InIsSingleThreaded)
-	: Base(InParticles, SolverPhysicsMaterials, DefaultNumIterations, DefaultNumPushOutIterations, InIsSingleThreaded)
+	: Base(InParticles, SolverPhysicsMaterials, InIsSingleThreaded)
 	, Clustering(*this, Particles.GetClusteredParticles())
 	, JointConstraintRule(JointConstraints, ChaosSolverJointPriority)
 	, SuspensionConstraintRule(SuspensionConstraints, ChaosSolverSuspensionPriority)
-	, CollisionConstraints(InParticles, Collided, PhysicsMaterials, PerParticlePhysicsMaterials, &SolverPhysicsMaterials, DefaultNumCollisionPairIterations, DefaultNumCollisionPushOutPairIterations, DefaultRestitutionThreshold)
+	, CollisionConstraints(InParticles, Collided, PhysicsMaterials, PerParticlePhysicsMaterials, &SolverPhysicsMaterials, 0, 0, DefaultRestitutionThreshold)
 	, CollisionRule(CollisionConstraints, ChaosSolverCollisionPriority)
 	, BroadPhase(InParticles)
 	, NarrowPhase(DefaultCollisionCullDistance, BoundsThicknessVelocityMultiplier, CollisionConstraints.GetConstraintAllocator())
@@ -781,6 +810,10 @@ FPBDRigidsEvolutionGBF::FPBDRigidsEvolutionGBF(FPBDRigidsSOAs& InParticles,THand
 	, CCDManager()
 	, bIsDeterministic(false)
 {
+	SetNumPositionIterations(DefaultNumPositionIterations);
+	SetNumVelocityIterations(DefaultNumVelocityIterations);
+	SetNumProjectionIterations(DefaultNumProjectionIterations);
+
 	CollisionConstraints.SetCanDisableContacts(!!CollisionDisableCulledContacts);
 
 	SetParticleUpdatePositionFunction([this](const TParticleView<FPBDRigidParticles>& ParticlesInput, const FReal Dt)
diff --git a/Engine/Source/Runtime/Experimental/Chaos/Private/ChaosSolverConfiguration.cpp b/Engine/Source/Runtime/Experimental/Chaos/Private/ChaosSolverConfiguration.cpp
--- a/Engine/Source/Runtime/Experimental/Chaos/Private/ChaosSolverConfiguration.cpp
+++ b/Engine/Source/Runtime/Experimental/Chaos/Private/ChaosSolverConfiguration.cpp
@@ -3,24 +3,33 @@
 #include "ChaosSolverConfiguration.h"
 #include "Chaos/PBDRigidsEvolutionGBF.h"
 
-using FEvolution = Chaos::FPBDRigidsEvolutionGBF;
-
 FChaosSolverConfiguration::FChaosSolverConfiguration()
-	: Iterations(FEvolution::DefaultNumIterations)
-	, CollisionPairIterations(FEvolution::DefaultNumCollisionPairIterations)
-	, PushOutIterations(FEvolution::DefaultNumPushOutIterations)
-	, CollisionPushOutPairIterations(FEvolution::DefaultNumCollisionPushOutPairIterations)
-	, CollisionMarginFraction(FEvolution::DefaultCollisionMarginFraction)
-	, CollisionMarginMax(FEvolution::DefaultCollisionMarginMax)
-	, CollisionCullDistance(FEvolution::DefaultCollisionCullDistance)
-	, CollisionMaxPushOutVelocity(FEvolution::DefaultCollisionMaxPushOutVelocity)
-	, JointPairIterations(FEvolution::DefaultNumJointPairIterations)
-	, JointPushOutPairIterations(FEvolution::DefaultNumJointPushOutPairIterations)
+	: PositionIterations(Chaos::FPBDRigidsEvolutionGBF::DefaultNumPositionIterations)
+	, VelocityIterations(Chaos::FPBDRigidsEvolutionGBF::DefaultNumVelocityIterations)
+	, ProjectionIterations(Chaos::FPBDRigidsEvolutionGBF::DefaultNumProjectionIterations)
+	, CollisionMarginFraction(Chaos::FPBDRigidsEvolutionGBF::DefaultCollisionMarginFraction)
+	, CollisionMarginMax(Chaos::FPBDRigidsEvolutionGBF::DefaultCollisionMarginMax)
+	, CollisionCullDistance(Chaos::FPBDRigidsEvolutionGBF::DefaultCollisionCullDistance)
+	, CollisionMaxPushOutVelocity(Chaos::FPBDRigidsEvolutionGBF::DefaultCollisionMaxPushOutVelocity)
 	, ClusterConnectionFactor(1.0f)
 	, ClusterUnionConnectionType(EClusterUnionMethod::DelaunayTriangulation)
 	, bGenerateCollisionData(false)
 	, bGenerateBreakData(false)
 	, bGenerateTrailingData(false)
 	, bGenerateContactGraph(true)
+	, Iterations_DEPRECATED(Chaos::FPBDRigidsEvolutionGBF::DefaultNumPositionIterations)
+	, PushOutIterations_DEPRECATED(Chaos::FPBDRigidsEvolutionGBF::DefaultNumVelocityIterations)
+{
+}
+
+void FChaosSolverConfiguration::MoveRenamedPropertyValues()
 {
+	if (Iterations_DEPRECATED != Chaos::FPBDRigidsEvolutionGBF::DefaultNumPositionIterations)
+	{
+		PositionIterations = Iterations_DEPRECATED;
+	}
+	if (PushOutIterations_DEPRECATED != Chaos::FPBDRigidsEvolutionGBF::DefaultNumVelocityIterations)
+	{
+		VelocityIterations = PushOutIterations_DEPRECATED;
+	}
 }
diff --git a/Engine/Source/Runtime/Experimental/Chaos/Private/PBDRigidsSolver.cpp b/Engine/Source/Runtime/Experimental/Chaos/Private/PBDRigidsSolver.cpp
--- a/Engine/Source/Runtime/Experimental/Chaos/Private/PBDRigidsSolver.cpp
+++ b/Engine/Source/Runtime/Experimental/Chaos/Private/PBDRigidsSolver.cpp
@@ -51,12 +51,9 @@ DECLARE_DWORD_COUNTER_STAT(TEXT("NumUpdatedManifoldPoints"), STAT_ChaosCounter_N
 DECLARE_DWORD_COUNTER_STAT(TEXT("NumJoints"), STAT_ChaosCounter_NumJoints, STATGROUP_ChaosCounters);
 
 // Stat Iteration counters
-DECLARE_DWORD_COUNTER_STAT(TEXT("NumIterations"), STAT_ChaosCounter_NumIterations, STATGROUP_ChaosCounters);
-DECLARE_DWORD_COUNTER_STAT(TEXT("NumCollisionIterations"), STAT_ChaosCounter_NumCollisionIterations, STATGROUP_ChaosCounters);
-DECLARE_DWORD_COUNTER_STAT(TEXT("NumJointIterations"), STAT_ChaosCounter_NumJointIterations, STATGROUP_ChaosCounters);
-DECLARE_DWORD_COUNTER_STAT(TEXT("NumPushOutIterations"), STAT_ChaosCounter_NumPushOutIterations, STATGROUP_ChaosCounters);
-DECLARE_DWORD_COUNTER_STAT(TEXT("NumPushOutCollisionIterations"), STAT_ChaosCounter_NumPushOutCollisionIterations, STATGROUP_ChaosCounters);
-DECLARE_DWORD_COUNTER_STAT(TEXT("NumPushOutJointIterations"), STAT_ChaosCounter_NumPushOutJointIterations, STATGROUP_ChaosCounters);
+DECLARE_DWORD_COUNTER_STAT(TEXT("NumPositionIterations"), STAT_ChaosCounter_NumPositionIterations, STATGROUP_ChaosCounters);
+DECLARE_DWORD_COUNTER_STAT(TEXT("NumVelocityIterations"), STAT_ChaosCounter_NumVelocityIterations, STATGROUP_ChaosCounters);
+DECLARE_DWORD_COUNTER_STAT(TEXT("NumProjectionIterations"), STAT_ChaosCounter_NumProjectionIterations, STATGROUP_ChaosCounters);
 
 
 // DebugDraw CVars
@@ -199,26 +196,17 @@ namespace Chaos
 		FAutoConsoleVariableRef CVarChaosSolverCollisionPositionShockPropagationIterations(TEXT("p.Chaos.Solver.Collision.PositionShockPropagationIterations"), ChaosSolverCollisionPositionShockPropagationIterations, TEXT("Override number of position iterations where shock propagation is applied (if >= 0)"));
 		FAutoConsoleVariableRef CVarChaosSolverCollisionVelocityShockPropagationIterations(TEXT("p.Chaos.Solver.Collision.VelocityShockPropagationIterations"), ChaosSolverCollisionVelocityShockPropagationIterations, TEXT("Override number of velocity iterations where shock propagation is applied (if >= 0)"));
 
-		int32 ChaosSolverIterations = -1;
-		FAutoConsoleVariableRef CVarChaosSolverIterations(TEXT("p.Chaos.Solver.Iterations"), ChaosSolverIterations, TEXT("Override number of solver iterations (-1 to use config)"));
+		int32 ChaosSolverPositionIterations = -1;
+		FAutoConsoleVariableRef CVarChaosSolverIterations(TEXT("p.Chaos.Solver.Iterations.Position"), ChaosSolverPositionIterations, TEXT("Override number of solver position iterations (-1 to use config)"));
 
-		int32 ChaosSolverCollisionIterations = -1;
-		FAutoConsoleVariableRef CVarChaosSolverCollisionIterations(TEXT("p.Chaos.Solver.Collision.Iterations"), ChaosSolverCollisionIterations, TEXT("Override number of collision iterations per solver iteration (-1 to use config)"));
+		int32 ChaosSolverVelocityIterations = -1;
+		FAutoConsoleVariableRef CVarChaosSolverPushOutIterations(TEXT("p.Chaos.Solver.Iterations.Velocity"), ChaosSolverVelocityIterations, TEXT("Override number of solver velocity iterations (-1 to use config)"));
 
-		int32 ChaosSolverPushOutIterations = -1;
-		FAutoConsoleVariableRef CVarChaosSolverPushOutIterations(TEXT("p.Chaos.Solver.PushoutIterations"), ChaosSolverPushOutIterations, TEXT("Override number of solver pushout iterations (-1 to use config)"));
-
-		int32 ChaosSolverCollisionPushOutIterations = -1;
-		FAutoConsoleVariableRef CVarChaosSolverCollisionPushOutIterations(TEXT("p.Chaos.Solver.Collision.PushOutIterations"), ChaosSolverCollisionPushOutIterations, TEXT("Override number of collision iterations per solver iteration (-1 to use config)"));
-
-		int32 ChaosSolverJointPairIterations = -1;
-		FAutoConsoleVariableRef CVarChaosSolverJointPairIterations(TEXT("p.Chaos.Solver.Joint.PairIterations"), ChaosSolverJointPairIterations, TEXT("Override number of iterations per joint pair during a solver iteration (-1 to use config)"));
-
-		int32 ChaosSolverJointPushOutPairIterations = -1;
-		FAutoConsoleVariableRef CVarChaosSolverJointPushOutPairIterations(TEXT("p.Chaos.Solver.Joint.PushOutPairIterations"), ChaosSolverJointPushOutPairIterations, TEXT("Override number of push out iterations per joint during a solver iteration (-1 to use config)"));
+		int32 ChaosSolverProjectionIterations = -1;
+		FAutoConsoleVariableRef CVarChaosSolverProjectionIterations(TEXT("p.Chaos.Solver.Iterations.Projection"), ChaosSolverProjectionIterations, TEXT("Override number of solver projection iterations (-1 to use config)"));
 
 		int32 ChaosSolverDeterministic = -1;
-		FAutoConsoleVariableRef CVarChaosSolverDeterministic(TEXT("p.Chaos.Solver.Deterministic"), ChaosSolverDeterministic, TEXT("Override determinism. 0: disabled; 1: enabled; -1: default(disabled)"));
+		FAutoConsoleVariableRef CVarChaosSolverDeterministic(TEXT("p.Chaos.Solver.Deterministic"), ChaosSolverDeterministic, TEXT("Override determinism. 0: disabled; 1: enabled; -1: use config"));
 
 		// Copied from RBAN
 		Chaos::FRealSingle ChaosSolverJointPositionTolerance = 0.025f;
@@ -861,29 +849,17 @@ namespace Chaos
 			{
 				MEvolution->GetCollisionConstraints().SetVelocityShockPropagationIterations(ChaosSolverCollisionVelocityShockPropagationIterations);
 			}
-			if (ChaosSolverIterations >= 0)
-			{
-				SetIterations(ChaosSolverIterations);
-			}
-			if (ChaosSolverCollisionIterations >= 0)
-			{
-				SetCollisionPairIterations(ChaosSolverCollisionIterations);
-			}
-			if (ChaosSolverPushOutIterations >= 0)
-			{
-				SetPushOutIterations(ChaosSolverPushOutIterations);
-			}
-			if (ChaosSolverCollisionPushOutIterations >= 0)
+			if (ChaosSolverPositionIterations >= 0)
 			{
-				SetCollisionPushOutPairIterations(ChaosSolverCollisionPushOutIterations);
+				SetPositionIterations(ChaosSolverPositionIterations);
 			}
-			if (ChaosSolverJointPairIterations >= 0)
+			if (ChaosSolverVelocityIterations >= 0)
 			{
-				SetJointPairIterations(ChaosSolverJointPairIterations);
+				SetVelocityIterations(ChaosSolverVelocityIterations);
 			}
-			if (ChaosSolverJointPushOutPairIterations >= 0)
+			if (ChaosSolverProjectionIterations >= 0)
 			{
-				SetJointPushOutPairIterations(ChaosSolverJointPushOutPairIterations);
+				SetProjectionIterations(ChaosSolverProjectionIterations);
 			}
 			if (ChaosSolverCullDistance >= 0.0f)
 			{
@@ -1469,12 +1445,9 @@ CSV_CUSTOM_STAT(PhysicsCounters, Name, Value, ECsvCustomStatOp::Set);
 #endif
 
 		// Iterations
-		CHAOS_COUNTER_STAT(NumIterations, GetEvolution()->GetNumIterations());
-		CHAOS_COUNTER_STAT(NumCollisionIterations, GetEvolution()->GetCollisionConstraints().GetPairIterations());
-		CHAOS_COUNTER_STAT(NumJointIterations, GetJointConstraints().GetSettings().ApplyPairIterations);
-		CHAOS_COUNTER_STAT(NumPushOutIterations, GetEvolution()->GetNumPushOutIterations());
-		CHAOS_COUNTER_STAT(NumPushOutCollisionIterations, GetEvolution()->GetCollisionConstraints().GetPushOutPairIterations());
-		CHAOS_COUNTER_STAT(NumPushOutJointIterations, GetJointConstraints().GetSettings().ApplyPushOutPairIterations);
+		CHAOS_COUNTER_STAT(NumPositionIterations, GetEvolution()->GetNumPositionIterations());
+		CHAOS_COUNTER_STAT(NumVelocityIterations, GetEvolution()->GetNumVelocityIterations());
+		CHAOS_COUNTER_STAT(NumProjectionIterations, GetEvolution()->GetNumProjectionIterations());
 	}
 
 	void FPBDRigidsSolver::UpdateExpensiveStatCounters() const
@@ -1726,12 +1699,9 @@ CSV_CUSTOM_STAT(PhysicsCounters, Name, Value, ECsvCustomStatOp::Set);
 	{
 		GetEvolution()->GetRigidClustering().SetClusterConnectionFactor(InConfig.ClusterConnectionFactor);
 		GetEvolution()->GetRigidClustering().SetClusterUnionConnectionType(ToInternalConnectionMethod(InConfig.ClusterUnionConnectionType));
-		SetIterations(InConfig.Iterations);
-		SetCollisionPairIterations(InConfig.CollisionPairIterations);
-		SetPushOutIterations(InConfig.PushOutIterations);
-		SetCollisionPushOutPairIterations(InConfig.CollisionPushOutPairIterations);
-		SetJointPairIterations(InConfig.JointPairIterations);
-		SetJointPushOutPairIterations(InConfig.JointPushOutPairIterations);
+		SetPositionIterations(InConfig.PositionIterations);
+		SetVelocityIterations(InConfig.VelocityIterations);
+		SetProjectionIterations(InConfig.ProjectionIterations);
 		SetCollisionCullDistance(InConfig.CollisionCullDistance);
 		SetCollisionMaxPushOutVelocity(InConfig.CollisionMaxPushOutVelocity);
 		SetGenerateCollisionData(InConfig.bGenerateCollisionData);
diff --git a/Engine/Source/Runtime/Experimental/Chaos/Public/Chaos/Joint/PBDJointCachedSolverGaussSeidel.h b/Engine/Source/Runtime/Experimental/Chaos/Public/Chaos/Joint/PBDJointCachedSolverGaussSeidel.h
--- a/Engine/Source/Runtime/Experimental/Chaos/Public/Chaos/Joint/PBDJointCachedSolverGaussSeidel.h
+++ b/Engine/Source/Runtime/Experimental/Chaos/Public/Chaos/Joint/PBDJointCachedSolverGaussSeidel.h
@@ -233,7 +233,6 @@ struct FAxisConstraintDatas
 		// @todo(chaos): this can be build into the Apply phase when the whole solver is PBD
 		void ApplyProjections(
 			const FReal Dt,
-			const FReal InSolverStiffness,
 			const FPBDJointSolverSettings& SolverSettings,
 			const FPBDJointSettings& JointSettings,
 			const bool bLastIteration);
@@ -248,6 +247,21 @@ struct FAxisConstraintDatas
 			const FPBDJointSolverSettings& SolverSettings,
 			const FPBDJointSettings& JointSettings);
 
+		void ApplyTeleports(
+			const FReal Dt,
+			const FPBDJointSolverSettings& SolverSettings,
+			const FPBDJointSettings& JointSettings);
+
+		void ApplyPositionTeleport(
+			const FReal Dt,
+			const FPBDJointSolverSettings& SolverSettings,
+			const FPBDJointSettings& JointSettings);
+
+		void ApplyRotationTeleport(
+			const FReal Dt,
+			const FPBDJointSolverSettings& SolverSettings,
+			const FPBDJointSettings& JointSettings);
+
 		void EnableProjection();
 
 		void SetInvMassScales(
diff --git a/Engine/Source/Runtime/Experimental/Chaos/Public/Chaos/Joint/PBDJointSolverGaussSeidel.h b/Engine/Source/Runtime/Experimental/Chaos/Public/Chaos/Joint/PBDJointSolverGaussSeidel.h
--- a/Engine/Source/Runtime/Experimental/Chaos/Public/Chaos/Joint/PBDJointSolverGaussSeidel.h
+++ b/Engine/Source/Runtime/Experimental/Chaos/Public/Chaos/Joint/PBDJointSolverGaussSeidel.h
@@ -180,7 +180,6 @@ namespace Chaos
 		// @todo(chaos): this can be build into the Apply phase when the whole solver is PBD
 		void ApplyProjections(
 			const FReal Dt,
-			const FReal InSolverStiffness,
 			const FPBDJointSolverSettings& SolverSettings,
 			const FPBDJointSettings& JointSettings);
 
diff --git a/Engine/Source/Runtime/Experimental/Chaos/Public/Chaos/JointProperties.inl b/Engine/Source/Runtime/Experimental/Chaos/Public/Chaos/JointProperties.inl
--- a/Engine/Source/Runtime/Experimental/Chaos/Public/Chaos/JointProperties.inl
+++ b/Engine/Source/Runtime/Experimental/Chaos/Public/Chaos/JointProperties.inl
@@ -7,6 +7,8 @@ CHAOS_INNER_JOINT_PROPERTY(JointSettings, ProjectionEnabled, bProjectionEnabled,
 CHAOS_INNER_JOINT_PROPERTY(JointSettings, ShockPropagationEnabled, bShockPropagationEnabled, bool)
 CHAOS_INNER_JOINT_PROPERTY(JointSettings, ProjectionLinearAlpha, LinearProjection, FReal)
 CHAOS_INNER_JOINT_PROPERTY(JointSettings, ProjectionAngularAlpha, AngularProjection, FReal)
+CHAOS_INNER_JOINT_PROPERTY(JointSettings, ProjectionLinearTolerance, TeleportDistance, FReal)
+CHAOS_INNER_JOINT_PROPERTY(JointSettings, ProjectionAngularTolerance, TeleportAngle, FReal)
 CHAOS_INNER_JOINT_PROPERTY(JointSettings, ShockPropagationAlpha, ShockPropagation, FReal)
 CHAOS_INNER_JOINT_PROPERTY(JointSettings, ParentInvMassScale, ParentInvMassScale, FReal)
 CHAOS_INNER_JOINT_PROPERTY(JointSettings, LinearBreakForce, LinearBreakForce, FReal)
diff --git a/Engine/Source/Runtime/Experimental/Chaos/Public/Chaos/PBDCollisionConstraints.h b/Engine/Source/Runtime/Experimental/Chaos/Public/Chaos/PBDCollisionConstraints.h
--- a/Engine/Source/Runtime/Experimental/Chaos/Public/Chaos/PBDCollisionConstraints.h
+++ b/Engine/Source/Runtime/Experimental/Chaos/Public/Chaos/PBDCollisionConstraints.h
@@ -177,6 +177,7 @@ public:
 
 	bool ApplyPhase1Serial(const FReal Dt, const int32 It, const int32 NumIts, FPBDIslandSolverData& SolverData);
 	bool ApplyPhase2Serial(const FReal Dt, const int32 It, const int32 NumIts, FPBDIslandSolverData& SolverData);
+	bool ApplyPhase3Serial(const FReal Dt, const int32 It, const int32 NumIts, FPBDIslandSolverData& SolverData) { return false; }
 
 	//
 	// Color Rule API
@@ -342,8 +343,11 @@ private:
 	const TArrayCollectionArray<TSerializablePtr<FChaosPhysicsMaterial>>& MPhysicsMaterials;
 	const TArrayCollectionArray<TUniquePtr<FChaosPhysicsMaterial>>& MPerParticlePhysicsMaterials;
 	const THandleArray<FChaosPhysicsMaterial>* const SimMaterials;
+
+	// @todo(chaos): remove pair iteration counts when legacy PBD solver is removed (RBAN-only)
 	int32 MApplyPairIterations;
 	int32 MApplyPushOutPairIterations;
+
 	FReal RestitutionThreshold;
 	bool bEnableCollisions;
 	bool bEnableRestitution;
diff --git a/Engine/Source/Runtime/Experimental/Chaos/Public/Chaos/PBDConstraintRule.h b/Engine/Source/Runtime/Experimental/Chaos/Public/Chaos/PBDConstraintRule.h
--- a/Engine/Source/Runtime/Experimental/Chaos/Public/Chaos/PBDConstraintRule.h
+++ b/Engine/Source/Runtime/Experimental/Chaos/Public/Chaos/PBDConstraintRule.h
@@ -132,6 +132,9 @@ namespace Chaos
 		/** Apply push out for constraints in the specified island. Return true if more iterations are needed. */
 		virtual bool ApplyPushOut(const FReal Dt, int32 GroupIndex, const int32 It, const int32 NumIts) { return false; }
 
+		/** Apply push out for constraints in the specified island. Return true if more iterations are needed. */
+		virtual bool ApplyProjection(const FReal Dt, int32 GroupIndex, const int32 It, const int32 NumIts) { return false; }
+
 		/** Add all constraints to the connectivity graph */
 		virtual void AddToGraph() {}
 
@@ -248,6 +251,8 @@ namespace Chaos
 
 		virtual bool ApplyPushOut(const FReal Dt, int32 GroupIndex, const int32 It, const int32 NumIts) override;
 
+		virtual bool ApplyProjection(const FReal Dt, int32 GroupIndex, const int32 It, const int32 NumIts) override;
+
 		virtual void InitializeAccelerationStructures() override;
 		
 		virtual void UpdateAccelerationStructures(const FReal Dt, const int32 GroupIndex) override;
diff --git a/Engine/Source/Runtime/Experimental/Chaos/Public/Chaos/PBDJointConstraintTypes.h b/Engine/Source/Runtime/Experimental/Chaos/Public/Chaos/PBDJointConstraintTypes.h
--- a/Engine/Source/Runtime/Experimental/Chaos/Public/Chaos/PBDJointConstraintTypes.h
+++ b/Engine/Source/Runtime/Experimental/Chaos/Public/Chaos/PBDJointConstraintTypes.h
@@ -132,12 +132,13 @@ namespace Chaos
 		FReal LinearProjection;
 		FReal AngularProjection;
 		FReal ShockPropagation;
+		FReal TeleportDistance;
+		FReal TeleportAngle;			// Radians
 		FReal ParentInvMassScale;
 
 		bool bCollisionEnabled;
-		bool bProjectionEnabled;		// @chaos(todo): remove - implied by alpha
+		bool bProjectionEnabled;		// @chaos(todo): remove - implied by alpha and teleport settings
 		bool bShockPropagationEnabled;	// @chaos(todo): remove - implied by alpha
-		bool bSoftProjectionEnabled;	// @chaos(todo): remove - old PBD solver only
 
 		TVector<EJointMotionType, 3> LinearMotionTypes;
 		FReal LinearLimit;
@@ -204,6 +205,7 @@ namespace Chaos
 		FPBDJointSolverSettings();
 
 		// Iterations
+		// @todo(chaos): remove pair iteration counts when legacy PBD solver is removed (RBAN-only)
 		int32 ApplyPairIterations;
 		int32 ApplyPushOutPairIterations;
 
diff --git a/Engine/Source/Runtime/Experimental/Chaos/Public/Chaos/PBDJointConstraints.h b/Engine/Source/Runtime/Experimental/Chaos/Public/Chaos/PBDJointConstraints.h
--- a/Engine/Source/Runtime/Experimental/Chaos/Public/Chaos/PBDJointConstraints.h
+++ b/Engine/Source/Runtime/Experimental/Chaos/Public/Chaos/PBDJointConstraints.h
@@ -282,7 +282,6 @@ namespace Chaos
 		bool ApplyPhase1Serial(const FReal Dt, const int32 It, const int32 NumIts, FPBDIslandSolverData& SolverData);
 		bool ApplyPhase2Serial(const FReal Dt, const int32 It, const int32 NumIts, FPBDIslandSolverData& SolverData);
 		bool ApplyPhase3Serial(const FReal Dt, const int32 It, const int32 NumIts, FPBDIslandSolverData& SolverData);
-		void PreparePhase3Serial(const FReal Dt, const int32 It, FPBDIslandSolverData& SolverData);
 
 		/**
 		 * Set the solver method to use
@@ -310,6 +309,8 @@ namespace Chaos
 
 		bool CanEvaluate(const int32 ConstraintIndex) const;
 
+		void PreparePhase3Serial(const FReal Dt, FPBDIslandSolverData& SolverData);
+
 		bool ApplyPhase1Single(const FReal Dt, const int32 ConstraintIndex, const int32 NumPairIts, const int32 It, const int32 NumIts);
 		bool ApplyPhase2Single(const FReal Dt, const int32 ConstraintIndex, const int32 It, const int32 NumIts);
 		bool ApplyPhase3Single(const FReal Dt, const int32 ConstraintIndex, const int32 It, const int32 NumIts);
diff --git a/Engine/Source/Runtime/Experimental/Chaos/Public/Chaos/PBDPositionConstraints.h b/Engine/Source/Runtime/Experimental/Chaos/Public/Chaos/PBDPositionConstraints.h
--- a/Engine/Source/Runtime/Experimental/Chaos/Public/Chaos/PBDPositionConstraints.h
+++ b/Engine/Source/Runtime/Experimental/Chaos/Public/Chaos/PBDPositionConstraints.h
@@ -193,6 +193,7 @@ namespace Chaos
 
 		bool ApplyPhase1Serial(const FReal Dt, const int32 It, const int32 NumIts, FPBDIslandSolverData& SolverData);
 		bool ApplyPhase2Serial(const FReal Dt, const int32 It, const int32 NumIts, FPBDIslandSolverData& SolverData) { return false; }
+		bool ApplyPhase3Serial(const FReal Dt, const int32 It, const int32 NumIts, FPBDIslandSolverData& SolverData) { return false; }
 
 	protected:
 		using Base::GetConstraintIndex;
diff --git a/Engine/Source/Runtime/Experimental/Chaos/Public/Chaos/PBDRigidDynamicSpringConstraints.h b/Engine/Source/Runtime/Experimental/Chaos/Public/Chaos/PBDRigidDynamicSpringConstraints.h
--- a/Engine/Source/Runtime/Experimental/Chaos/Public/Chaos/PBDRigidDynamicSpringConstraints.h
+++ b/Engine/Source/Runtime/Experimental/Chaos/Public/Chaos/PBDRigidDynamicSpringConstraints.h
@@ -214,6 +214,7 @@ namespace Chaos
 
 		bool ApplyPhase1Serial(const FReal Dt, const int32 It, const int32 NumIts, FPBDIslandSolverData& SolverData);
 		bool ApplyPhase2Serial(const FReal Dt, const int32 It, const int32 NumIts, FPBDIslandSolverData& SolverData) { return false; }
+		bool ApplyPhase3Serial(const FReal Dt, const int32 It, const int32 NumIts, FPBDIslandSolverData& SolverData) { return false; }
 
 	protected:
 		using Base::GetConstraintIndex;
diff --git a/Engine/Source/Runtime/Experimental/Chaos/Public/Chaos/PBDRigidSpringConstraints.h b/Engine/Source/Runtime/Experimental/Chaos/Public/Chaos/PBDRigidSpringConstraints.h
--- a/Engine/Source/Runtime/Experimental/Chaos/Public/Chaos/PBDRigidSpringConstraints.h
+++ b/Engine/Source/Runtime/Experimental/Chaos/Public/Chaos/PBDRigidSpringConstraints.h
@@ -185,6 +185,7 @@ namespace Chaos
 
 		bool ApplyPhase1Serial(const FReal Dt, const int32 It, const int32 NumIts, FPBDIslandSolverData& SolverData);
 		bool ApplyPhase2Serial(const FReal Dt, const int32 It, const int32 NumIts, FPBDIslandSolverData& SolverData) { return false; }
+		bool ApplyPhase3Serial(const FReal Dt, const int32 It, const int32 NumIts, FPBDIslandSolverData& SolverData) { return false; }
 
 	protected:
 		using Base::GetConstraintIndex;
diff --git a/Engine/Source/Runtime/Experimental/Chaos/Public/Chaos/PBDRigidsEvolution.h b/Engine/Source/Runtime/Experimental/Chaos/Public/Chaos/PBDRigidsEvolution.h
--- a/Engine/Source/Runtime/Experimental/Chaos/Public/Chaos/PBDRigidsEvolution.h
+++ b/Engine/Source/Runtime/Experimental/Chaos/Public/Chaos/PBDRigidsEvolution.h
@@ -242,7 +242,7 @@ public:
 	typedef TFunction<void(FPBDRigidParticles&, const FReal, const FReal, const int32)> FKinematicUpdateRule;
 	typedef TFunction<void(TParticleView<FPBDRigidParticles>&)> FCaptureRewindRule;
 
-	CHAOS_API FPBDRigidsEvolutionBase(FPBDRigidsSOAs& InParticles, THandleArray<FChaosPhysicsMaterial>& InSolverPhysicsMaterials, int32 InNumIterations = 1, int32 InNumPushOutIterations = 1, bool InIsSingleThreaded = false);
+	CHAOS_API FPBDRigidsEvolutionBase(FPBDRigidsSOAs& InParticles, THandleArray<FChaosPhysicsMaterial>& InSolverPhysicsMaterials, bool InIsSingleThreaded = false);
 	CHAOS_API virtual ~FPBDRigidsEvolutionBase();
 
 	CHAOS_API TArray<FGeometryParticleHandle*> CreateStaticParticles(int32 NumParticles, const FUniqueIdx* ExistingIndices = nullptr, const FGeometryParticleParameters& Params = FGeometryParticleParameters())
@@ -313,24 +313,34 @@ public:
 		ConstraintRule->BindToGraph(ConstraintGraph, ContainerId);
 	}
 
-	CHAOS_API void SetNumIterations(int32 InNumIterations)
+	CHAOS_API void SetNumPositionIterations(int32 InNumIterations)
 	{
-		NumIterations = InNumIterations;
+		NumPositionIterations = InNumIterations;
 	}
 
-	CHAOS_API int32 GetNumIterations() const
+	CHAOS_API int32 GetNumPositionIterations() const
 	{
-		return NumIterations;
+		return NumPositionIterations;
 	}
 
-	CHAOS_API void SetNumPushOutIterations(int32 InNumIterations)
+	CHAOS_API void SetNumVelocityIterations(int32 InNumIterations)
 	{
-		NumPushOutIterations = InNumIterations;
+		NumVelocityIterations = InNumIterations;
 	}
 
-	CHAOS_API int32 GetNumPushOutIterations() const
+	CHAOS_API int32 GetNumVelocityIterations() const
 	{
-		return NumPushOutIterations;
+		return NumVelocityIterations;
+	}
+
+	CHAOS_API void SetNumProjectionIterations(int32 InNumIterations)
+	{
+		NumProjectionIterations = InNumIterations;
+	}
+
+	CHAOS_API int32 GetNumProjectionIterations() const
+	{
+		return NumProjectionIterations;
 	}
 
 	CHAOS_API void SetParticleKinematicTarget(FKinematicGeometryParticleHandle* KinematicHandle, const FKinematicTarget& NewKinematicTarget)
@@ -1024,8 +1034,9 @@ protected:
 	};
 	FGraphEventRef AccelerationStructureTaskComplete;
 
-	int32 NumIterations;
-	int32 NumPushOutIterations;
+	int32 NumPositionIterations;
+	int32 NumVelocityIterations;
+	int32 NumProjectionIterations;
 	TUniquePtr<ISpatialAccelerationCollectionFactory> SpatialCollectionFactory;
 
 	FAccelerationStructure* GetFreeSpatialAcceleration_Internal();
diff --git a/Engine/Source/Runtime/Experimental/Chaos/Public/Chaos/PBDRigidsEvolutionGBF.h b/Engine/Source/Runtime/Experimental/Chaos/Public/Chaos/PBDRigidsEvolutionGBF.h
--- a/Engine/Source/Runtime/Experimental/Chaos/Public/Chaos/PBDRigidsEvolutionGBF.h
+++ b/Engine/Source/Runtime/Experimental/Chaos/Public/Chaos/PBDRigidsEvolutionGBF.h
@@ -50,42 +50,6 @@ namespace Chaos
 	{
 	public:
 		using Base = FPBDRigidsEvolutionBase;
-		using Base::Particles;
-		using typename Base::FForceRule;
-		using Base::ForceRules;
-		using Base::PrepareTick;
-		using Base::UnprepareTick;
-		using Base::ApplyKinematicTargets;
-		using Base::UpdateConstraintPositionBasedState;
-		using Base::InternalAcceleration;
-		using Base::CreateConstraintGraph;
-		using Base::CreateIslands;
-		using Base::GetParticles;
-		using Base::DirtyParticle;
-		using Base::SetPhysicsMaterial;
-		using Base::SetPerParticlePhysicsMaterial;
-		using Base::GetPerParticlePhysicsMaterial;
-		using Base::CreateParticle;
-		using Base::GenerateUniqueIdx;
-		using Base::DestroyParticle;
-		using Base::CreateClusteredParticles;
-		using Base::EnableParticle;
-		using Base::DisableParticles;
-		using Base::NumIslands;
-		using Base::GetNonDisabledClusteredView;
-		using Base::DisableParticle;
-		using Base::GetConstraintGraph;
-		using Base::PhysicsMaterials;
-		using Base::PerParticlePhysicsMaterials;
-		using Base::ParticleDisableCount;
-		using Base::SolverPhysicsMaterials;
-		using Base::CaptureRewindData;
-		using Base::Collided;
-		using Base::SetParticleUpdatePositionFunction;
-		using Base::AddForceFunction;
-		using Base::AddConstraintRule;
-		using Base::ParticleUpdatePosition;
-		using Base::GetAllRemovals;
 
 		using FGravityForces = FPerParticleGravity;
 		using FCollisionConstraints = FPBDCollisionConstraints;
@@ -98,16 +62,13 @@ namespace Chaos
 		using FJointConstraintRule = TPBDConstraintIslandRule<FJointConstraints>;
 
 		// Default iteration counts
-		static constexpr int32 DefaultNumIterations = 8;
-		static constexpr int32 DefaultNumCollisionPairIterations = 1;
-		static constexpr int32 DefaultNumPushOutIterations = 1;
-		static constexpr int32 DefaultNumCollisionPushOutPairIterations = 1;
+		static constexpr int32 DefaultNumPositionIterations = 8;
+		static constexpr int32 DefaultNumVelocityIterations = 1;
+		static constexpr int32 DefaultNumProjectionIterations = 1;
 		static constexpr FRealSingle DefaultCollisionMarginFraction = 0.05f;
 		static constexpr FRealSingle DefaultCollisionMarginMax = 10.0f;
 		static constexpr FRealSingle DefaultCollisionCullDistance = 3.0f;
 		static constexpr FRealSingle DefaultCollisionMaxPushOutVelocity = 1000.0f;
-		static constexpr int32 DefaultNumJointPairIterations = 1;
-		static constexpr int32 DefaultNumJointPushOutPairIterations = 1;
 		static constexpr int32 DefaultRestitutionThreshold = 1000;
 
 		CHAOS_API FPBDRigidsEvolutionGBF(FPBDRigidsSOAs& InParticles, THandleArray<FChaosPhysicsMaterial>& SolverPhysicsMaterials, const TArray<ISimCallbackObject*>* InCollisionModifiers = nullptr, bool InIsSingleThreaded = false);
@@ -347,10 +308,16 @@ namespace Chaos
 		void SetImplicitVelocities(const FReal Dt, int32 GroupIndex);
 		
 		// Second phase of constraint solver (after implicit velocity calculation following results of phase 1)
-		// For GBF this is the pushout phase
+		// For StandardPBD this does nothing
 		// For QuasiPBD this is the velocity solve phase
 		void ApplyConstraintsPhase2(const FReal Dt, int32 GroupIndex);
 
+		// Apply the accumulated corrections to the solver bodies
+		void ApplyCorrections(const FReal Dt, int32 GroupIndex);
+
+		// Third phase of constraint solver
+		// This is the projection phase (only joints implement this atm)
+		void ApplyConstraintsPhase3(const FReal Dt, int32 GroupIndex);
 
 		CHAOS_API void Serialize(FChaosArchive& Ar);
 
diff --git a/Engine/Source/Runtime/Experimental/Chaos/Public/Chaos/PBDSuspensionConstraints.h b/Engine/Source/Runtime/Experimental/Chaos/Public/Chaos/PBDSuspensionConstraints.h
--- a/Engine/Source/Runtime/Experimental/Chaos/Public/Chaos/PBDSuspensionConstraints.h
+++ b/Engine/Source/Runtime/Experimental/Chaos/Public/Chaos/PBDSuspensionConstraints.h
@@ -240,6 +240,7 @@ namespace Chaos
 
 		bool ApplyPhase1Serial(const FReal Dt, const int32 It, const int32 NumIts, FPBDIslandSolverData& SolverData);
 		bool ApplyPhase2Serial(const FReal Dt, const int32 It, const int32 NumIts, FPBDIslandSolverData& SolverData);
+		bool ApplyPhase3Serial(const FReal Dt, const int32 It, const int32 NumIts, FPBDIslandSolverData& SolverData) { return false; }
 
 	protected:
 		using Base::GetConstraintIndex;
diff --git a/Engine/Source/Runtime/Experimental/Chaos/Public/ChaosSolverConfiguration.h b/Engine/Source/Runtime/Experimental/Chaos/Public/ChaosSolverConfiguration.h
--- a/Engine/Source/Runtime/Experimental/Chaos/Public/ChaosSolverConfiguration.h
+++ b/Engine/Source/Runtime/Experimental/Chaos/Public/ChaosSolverConfiguration.h
@@ -26,38 +26,35 @@ struct FChaosSolverConfiguration
 
 	CHAOS_API FChaosSolverConfiguration();
 
-	// The number of iterations to run during the constraint solver step
-	UPROPERTY(EditAnywhere, Category = "SolverConfiguration|Iterations")
-	int32 Iterations;
+	// Handle renamed properties
+	CHAOS_API void MoveRenamedPropertyValues();
 
-	// During solver iterations we solve each constraint in turn. For each constraint
-	// we run the solve step CollisionPairIterations times in a row.
-	UPROPERTY(EditAnywhere, Category = "SolverConfiguration|Iterations")
-	int32 CollisionPairIterations;
+	// The number of position iterations to run during the constraint solver step
+	UPROPERTY(EditAnywhere, Category = "SolverConfiguration|Iterations", meta = (ClampMin = "0"))
+	int32 PositionIterations;
+	
+	// The number of velocity iterations to run during the constraint solver step
+	UPROPERTY(EditAnywhere, Category = "SolverConfiguration|Iterations", meta = (ClampMin = "0"))
+	int32 VelocityIterations;
 
-	// The number of iterations to run during the constraint fixup step. This applies a post-solve
-	// correction that can address errors left behind during the main solver iterations.
-	UPROPERTY(EditAnywhere, Category = "SolverConfiguration|Iterations")
-	int32 PushOutIterations;
+	// The number of projection iterations to run during the constraint solver step
+	UPROPERTY(EditAnywhere, Category = "SolverConfiguration|Iterations", meta = (ClampMin = "0"))
+	int32 ProjectionIterations;
 
-	// During pushout iterations we pushout each constraint in turn. For each constraint
-	// we run the pushout step CollisionPairIterations times in a row.
-	UPROPERTY(EditAnywhere, Category = "SolverConfiguration|Iterations")
-	int32 CollisionPushOutPairIterations;
 
 	// A collision margin as a fraction of size used by some boxes and convex shapes to improve collision detection results.
 	// The core geometry of shapes that support a margin are reduced in size by the margin, and the margin
 	// is added back on during collision detection. The net result is a shape of the same size but with rounded corners.
-	UPROPERTY(EditAnywhere, Category = "SolverConfiguration|Collision")
+	UPROPERTY(EditAnywhere, Category = "SolverConfiguration|Collision", meta = (ClampMin = "0.0"))
 	float CollisionMarginFraction;
 
 	// An upper limit on the collision margin that will be subtracted from boxes and convex shapes. See CollisionMarginFraction
-	UPROPERTY(EditAnywhere, Category = "SolverConfiguration|Collision")
+	UPROPERTY(EditAnywhere, Category = "SolverConfiguration|Collision", meta = (ClampMin = "0.0"))
 	float CollisionMarginMax;
 
 	// During collision detection, if tweo shapes are at least this far apart we do not calculate their nearest features
 	// during the collision detection step.
-	UPROPERTY(EditAnywhere, Category = "SolverConfiguration|Collision")
+	UPROPERTY(EditAnywhere, Category = "SolverConfiguration|Collision", meta = (ClampMin = "0.0"))
 	float CollisionCullDistance;
 
 	// The maximum speed at which two bodies can be extracted from each other when they start a frame inter-penetrating. This can
@@ -65,18 +62,9 @@ struct FChaosSolverConfiguration
 	// zero means "no limit". A non-zero value can be used to prevent explosive behaviour when bodies start deeply penetrating. 
 	// An alternative to using this approach is to increase the number of Velocity Iterations, which is more expensive but will 
 	// ensure the bdoies are depenetrated in a single frame without explosive behaviour.
-	UPROPERTY(EditAnywhere, Category = "SolverConfiguration|Collision")
+	UPROPERTY(EditAnywhere, Category = "SolverConfiguration|Collision", meta = (ClampMin = "0.0"))
 	float CollisionMaxPushOutVelocity;
 
-	// The number of iterations to run on each constraint during the constraint solver step
-	UPROPERTY(EditAnywhere, Category = "SolverConfiguration|Iterations")
-	int32 JointPairIterations;
-
-	// The number of iterations to run during the constraint fixup step for each joint. This applies a post-solve
-	// correction that can address errors left behind during the main solver iterations.
-	UPROPERTY(EditAnywhere, Category = "SolverConfiguration|Iterations")
-	int32 JointPushOutPairIterations;
-
 	UPROPERTY(EditAnywhere, Category = "SolverConfiguration|Clustering")
 	float ClusterConnectionFactor;
 
@@ -103,4 +91,14 @@ struct FChaosSolverConfiguration
 
 	UPROPERTY(EditAnywhere, Category = "SolverConfiguration|Contacts")
 	bool bGenerateContactGraph;
+
+private:
+
+	// Renamed to PositionIterations
+	UPROPERTY()
+	int32 Iterations_DEPRECATED;
+
+	// Renamed to VelocityIterations
+	UPROPERTY()
+	int32 PushOutIterations_DEPRECATED;
 };
diff --git a/Engine/Source/Runtime/Experimental/Chaos/Public/PBDRigidsSolver.h b/Engine/Source/Runtime/Experimental/Chaos/Public/PBDRigidsSolver.h
--- a/Engine/Source/Runtime/Experimental/Chaos/Public/PBDRigidsSolver.h
+++ b/Engine/Source/Runtime/Experimental/Chaos/Public/PBDRigidsSolver.h
@@ -175,12 +175,9 @@ namespace Chaos
 		int32& GetCurrentFrame() { return CurrentFrame; }
 
 		/**/
-		void SetIterations(const int32 InNumIterations) { GetEvolution()->SetNumIterations(InNumIterations); }
-		void SetPushOutIterations(const int32 InNumIterations) {  GetEvolution()->SetNumPushOutIterations(InNumIterations); }
-		void SetCollisionPairIterations(const int32 InNumIterations) { GetEvolution()->GetCollisionConstraints().SetPairIterations(InNumIterations); }
-		void SetCollisionPushOutPairIterations(const int32 InNumIterations) { GetEvolution()->GetCollisionConstraints().SetPushOutPairIterations(InNumIterations); }
-		void SetJointPairIterations(const int32 InNumIterations) { GetJointConstraints().SetNumPairIterations(InNumIterations); }
-		void SetJointPushOutPairIterations(const int32 InNumIterations) {GetJointConstraints().SetNumPushOutPairIterations(InNumIterations); }
+		void SetPositionIterations(const int32 InNumIterations) { GetEvolution()->SetNumPositionIterations(InNumIterations); }
+		void SetVelocityIterations(const int32 InNumIterations) { GetEvolution()->SetNumVelocityIterations(InNumIterations); }
+		void SetProjectionIterations(const int32 InNumIterations) { GetEvolution()->SetNumProjectionIterations(InNumIterations); }
 		void SetCollisionCullDistance(const FReal InCullDistance) { GetEvolution()->GetNarrowPhase().SetBoundsExpansion(InCullDistance); }
 		void SetCollisionMaxPushOutVelocity(const FReal InMaxPushOutVelocity) { GetEvolution()->GetCollisionConstraints().SetMaxPushOutVelocity(InMaxPushOutVelocity); }
 		void SetUseContactGraph(const bool bInUseContactGraph) { GetEvolution()->GetCollisionConstraintsRule().SetUseContactGraph(bInUseContactGraph); }
diff --git a/Engine/Source/Runtime/Experimental/ChaosSolverEngine/Private/Chaos/ChaosSolverActor.cpp b/Engine/Source/Runtime/Experimental/ChaosSolverEngine/Private/Chaos/ChaosSolverActor.cpp
--- a/Engine/Source/Runtime/Experimental/ChaosSolverEngine/Private/Chaos/ChaosSolverActor.cpp
+++ b/Engine/Source/Runtime/Experimental/ChaosSolverEngine/Private/Chaos/ChaosSolverActor.cpp
@@ -356,10 +356,8 @@ void AChaosSolverActor::PostLoad()
 			return EClusterUnionMethod::None;
 		};
 
-		Properties.Iterations = CollisionIterations_DEPRECATED;
-		Properties.CollisionPairIterations = 2;
-		Properties.PushOutIterations = PushOutIterations_DEPRECATED;
-		Properties.CollisionPushOutPairIterations = PushOutPairIterations_DEPRECATED;
+		Properties.PositionIterations = CollisionIterations_DEPRECATED;
+		Properties.VelocityIterations = PushOutIterations_DEPRECATED;
 		Properties.ClusterConnectionFactor = ClusterConnectionFactor_DEPRECATED;
 		Properties.ClusterUnionConnectionType = ConvertDeprecatedConnectionType(ClusterUnionConnectionType_DEPRECATED);
 		Properties.bGenerateBreakData = DoGenerateBreakingData_DEPRECATED;
diff --git a/Engine/Source/Runtime/PhysicsCore/Private/ChaosEngineInterface.cpp b/Engine/Source/Runtime/PhysicsCore/Private/ChaosEngineInterface.cpp
--- a/Engine/Source/Runtime/PhysicsCore/Private/ChaosEngineInterface.cpp
+++ b/Engine/Source/Runtime/PhysicsCore/Private/ChaosEngineInterface.cpp
@@ -1297,7 +1297,7 @@ void FChaosEngineInterface::SetCollisionEnabled(const FPhysicsConstraintHandle&
 	}
 }
 
-void FChaosEngineInterface::SetProjectionEnabled_AssumesLocked(const FPhysicsConstraintHandle& InConstraintRef,bool bInProjectionEnabled,float InLinearAlpha,float InAngularAlpha)
+void FChaosEngineInterface::SetProjectionEnabled_AssumesLocked(const FPhysicsConstraintHandle& InConstraintRef,bool bInProjectionEnabled,float InLinearAlpha,float InAngularAlpha, float InLinearTolerance, float InAngularToleranceDeg)
 {
 	if (InConstraintRef.IsValid() && InConstraintRef.Constraint->IsType(Chaos::EConstraintType::JointConstraintType))
 	{
@@ -1306,6 +1306,8 @@ void FChaosEngineInterface::SetProjectionEnabled_AssumesLocked(const FPhysicsCon
 			Constraint->SetProjectionEnabled(bInProjectionEnabled);
 			Constraint->SetProjectionLinearAlpha(InLinearAlpha);
 			Constraint->SetProjectionAngularAlpha(InAngularAlpha);
+			Constraint->SetProjectionLinearTolerance(InLinearTolerance);
+			Constraint->SetProjectionAngularTolerance(FMath::DegreesToRadians(InAngularToleranceDeg));
 		}
 	}
 }
diff --git a/Engine/Source/Runtime/PhysicsCore/Private/PhysicsSettingsCore.cpp b/Engine/Source/Runtime/PhysicsCore/Private/PhysicsSettingsCore.cpp
--- a/Engine/Source/Runtime/PhysicsCore/Private/PhysicsSettingsCore.cpp
+++ b/Engine/Source/Runtime/PhysicsCore/Private/PhysicsSettingsCore.cpp
@@ -55,5 +55,7 @@ void UPhysicsSettingsCore::PostInitProperties()
 	{
 		DefaultShapeComplexity = bDefaultHasComplexCollision_DEPRECATED ? CTF_UseSimpleAndComplex : CTF_UseSimpleAsComplex;
 	}
+
+	SolverOptions.MoveRenamedPropertyValues();
 }
 
diff --git a/Engine/Source/Runtime/PhysicsCore/Public/Chaos/ChaosEngineInterface.h b/Engine/Source/Runtime/PhysicsCore/Public/Chaos/ChaosEngineInterface.h
--- a/Engine/Source/Runtime/PhysicsCore/Public/Chaos/ChaosEngineInterface.h
+++ b/Engine/Source/Runtime/PhysicsCore/Public/Chaos/ChaosEngineInterface.h
@@ -500,7 +500,7 @@ public:
 
 	static void SetCanVisualize(const FPhysicsConstraintHandle& InConstraintRef,bool bInCanVisualize);
 	static void SetCollisionEnabled(const FPhysicsConstraintHandle& InConstraintRef,bool bInCollisionEnabled);
-	static void SetProjectionEnabled_AssumesLocked(const FPhysicsConstraintHandle& InConstraintRef,bool bInProjectionEnabled,float InLinearAlpha = 1.0f,float InAngularAlpha = 0.0f);
+	static void SetProjectionEnabled_AssumesLocked(const FPhysicsConstraintHandle& InConstraintRef,bool bInProjectionEnabled,float InLinearAlpha = 1.0f,float InAngularAlpha = 0.0f, float InLinearTolerance = 0.0f, float InAngularToleranceDeg = 0.0f);
 	static void SetShockPropagationEnabled_AssumesLocked(const FPhysicsConstraintHandle& InConstraintRef, bool bInShockPropagationEnabled, float InShockPropagationAlpha);
 	static void SetParentDominates_AssumesLocked(const FPhysicsConstraintHandle& InConstraintRef,bool bInParentDominates);
 	static void SetBreakForces_AssumesLocked(const FPhysicsConstraintHandle& InConstraintRef,float InLinearBreakForce,float InAngularBreakForce);
