Overlap BVH optimization

diff --git a/Engine/Source/Runtime/Experimental/Chaos/Private/Chaos/TriangleMeshImplicitObject.cpp b/Engine/Source/Runtime/Experimental/Chaos/Private/Chaos/TriangleMeshImplicitObject.cpp
--- a/Engine/Source/Runtime/Experimental/Chaos/Private/Chaos/TriangleMeshImplicitObject.cpp
+++ b/Engine/Source/Runtime/Experimental/Chaos/Private/Chaos/TriangleMeshImplicitObject.cpp
@@ -818,7 +818,7 @@ bool FTriangleMeshImplicitObject::OverlapGeomImp(const QueryGeomType& QueryGeom,
 	// IMPORTANT QueryTM comes with a invscaled translation so we need a version of the TM with world space translation to properly compute the bounds
 	FRigidTransform3 TriMeshToGeomNoScale{ QueryTM };
 	TriMeshToGeomNoScale.SetTranslation(TriMeshToGeomNoScale.GetTranslation() * TriMeshScale);
-	// NOTE: BVH test is done in tri-mesh local space (whereas collision detection is done in world space becaused you can't non-uniformly scale all shapes)
+	// NOTE: BVH test is done in tri-mesh local space (whereas collision detection is done in world space because you can't non-uniformly scale all shapes)
 	FAABB3 QueryBounds = WorldScaleQueryGeom.BoundingBox();
 	QueryBounds = QueryBounds.TransformedAABB(TriMeshToGeomNoScale);
 	QueryBounds.ThickenSymmetrically(FVec3(Thickness));
@@ -867,25 +867,35 @@ bool FTriangleMeshImplicitObject::OverlapGeomImp(const QueryGeomType& QueryGeom,
 		}
 		else
 		{
+
+			const UE::Math::TQuat<FReal>& RotationDouble = WorldScaleQueryTM.GetRotation();
+			VectorRegister4Float Rotation = MakeVectorRegisterFloatFromDouble(MakeVectorRegister(RotationDouble.X, RotationDouble.Y, RotationDouble.Z, RotationDouble.W));
+			// Normalize rotation
+			Rotation = VectorNormalizeSafe(Rotation, GlobalVectorConstants::Float0001);
+
+			const UE::Math::TVector<FReal>& TranslationDouble = WorldScaleQueryTM.GetTranslation();
+			const VectorRegister4Float Translation = MakeVectorRegisterFloatFromDouble(MakeVectorRegister(TranslationDouble.X, TranslationDouble.Y, TranslationDouble.Z, 0.0));
+		
+
 			for (int32 TriIdx : PotentialIntersections)
 			{
 				FVec3 A, B, C;
 				TriangleMeshTransformVertsHelper(TriMeshScale, TriIdx, MParticles, Elements, A, B, C);
 
-				const FVec3 AB = B - A;
-				const FVec3 AC = C - A;
+				const VectorRegister4Float ASimd = MakeVectorRegisterFloatFromDouble(MakeVectorRegister(A.X, A.Y, A.Z, 0.0));
+				const VectorRegister4Float BSimd = MakeVectorRegisterFloatFromDouble(MakeVectorRegister(B.X, B.Y, B.Z, 0.0));
+				const VectorRegister4Float CSimd = MakeVectorRegisterFloatFromDouble(MakeVectorRegister(C.X, C.Y, C.Z, 0.0));
+
+				const VectorRegister4Float AB = VectorSubtract(BSimd, ASimd);
+				const VectorRegister4Float AC = VectorSubtract(CSimd, ASimd);
 
 				//It's most likely that the query object is in front of the triangle since queries tend to be on the outside.
 				//However, maybe we should check if it's behind the triangle plane. Also, we should enforce this winding in some way
-				const FVec3 Offset = FVec3::CrossProduct(AB, AC);
-
-				VectorRegister4Float ASimd = MakeVectorRegisterFloat((float)A.X, (float)A.Y, (float)A.Z, 0.0f);
-				VectorRegister4Float BSimd = MakeVectorRegisterFloat((float)B.X, (float)B.Y, (float)B.Z, 0.0f);
-				VectorRegister4Float CSimd = MakeVectorRegisterFloat((float)C.X, (float)C.Y, (float)C.Z, 0.0f);
-
+				const VectorRegister4Float InitialDir = VectorCross(AB, AC);
+				
 				FTriangleRegister Tri(ASimd, BSimd, CSimd);
 
-				if (GJKIntersection(Tri, WorldScaleQueryGeom, WorldScaleQueryTM, Thickness, Offset))
+				if (GJKIntersectionSimd(Tri, WorldScaleQueryGeom, Translation, Rotation, Thickness, InitialDir))
 				{
 					return true;
 				}
diff --git a/Engine/Source/Runtime/Experimental/Chaos/Public/Chaos/GJK.h b/Engine/Source/Runtime/Experimental/Chaos/Public/Chaos/GJK.h
--- a/Engine/Source/Runtime/Experimental/Chaos/Public/Chaos/GJK.h
+++ b/Engine/Source/Runtime/Experimental/Chaos/Public/Chaos/GJK.h
@@ -54,13 +54,11 @@ namespace Chaos
 		@param BToATM The transform of B in A's local space
 		@param ThicknessA The amount of geometry inflation for Geometry A(for example if the surface distance of two geometries with thickness 0 would be 2, a thickness of 0.5 would give a distance of 1.5)
 		@param InitialDir The first direction we use to search the CSO
-		@param ThicknessB The amount of geometry inflation for Geometry B(for example if the surface distance of two geometries with thickness 0 would be 2, a thickness of 0.5 would give a distance of 1.5)
 		@return True if the geometries overlap, False otherwise 
 	 */
 	template <typename T, typename TGeometryA, typename TGeometryB>
-	bool GJKIntersection(const TGeometryA& RESTRICT A, const TGeometryB& RESTRICT B, const TRigidTransform<T, 3>& BToATM, const T InThicknessA = 0, const TVector<T, 3>& InitialDir = TVector<T, 3>(1, 0, 0), const T InThicknessB = 0)
+	bool GJKIntersection(const TGeometryA& RESTRICT A, const TGeometryB& RESTRICT B, const TRigidTransform<T, 3>& BToATM, const T InThicknessA = 0, const TVector<T, 3>& InitialDir = TVector<T, 3>(1, 0, 0))
 	{
-
 		const UE::Math::TQuat<T>& RotationDouble = BToATM.GetRotation();
 		VectorRegister4Float RotationSimd = MakeVectorRegisterFloatFromDouble(MakeVectorRegister(RotationDouble.X, RotationDouble.Y, RotationDouble.Z, RotationDouble.W));
 
@@ -71,10 +69,18 @@ namespace Chaos
 
 		const VectorRegister4Float InitialDirSimd = MakeVectorRegisterFloatFromDouble(MakeVectorRegister(InitialDir[0], InitialDir[1], InitialDir[2], 0.0));
 		
-		VectorRegister4Float VSimd = VectorNegate(InitialDirSimd);
-		VSimd = VectorNormalizeSafe(VSimd, MakeVectorRegisterFloatConstant(-1.f, 0.f, 0.f, 0.f));
+		return GJKIntersectionSimd(A, B, TranslationSimd, RotationSimd, InThicknessA, InitialDirSimd);
+	}
+
+
+	template <typename TGeometryA, typename TGeometryB>
+	bool GJKIntersectionSimd(const TGeometryA& RESTRICT A, const TGeometryB& RESTRICT B, const VectorRegister4Float& Translation, const VectorRegister4Float& Rotation, FReal InThicknessA, const VectorRegister4Float& InitialDir)
+	{
+
+		VectorRegister4Float V = VectorNegate(InitialDir);
+		V = VectorNormalizeSafe(V, MakeVectorRegisterFloatConstant(-1.f, 0.f, 0.f, 0.f));
 
-		const VectorRegister4Float AToBRotationSimd = VectorQuaternionInverse(RotationSimd);
+		const VectorRegister4Float AToBRotation = VectorQuaternionInverse(Rotation);
 		bool bTerminate;
 		bool bNearZero = false;
 		int NumIterations = 0;
@@ -84,15 +90,14 @@ namespace Chaos
 		VectorRegister4Float BarycentricSimd;
 		VectorRegister4Int NumVerts = GlobalVectorConstants::IntZero;
 
-		T ThicknessA;
-		T ThicknessB;
+		FReal ThicknessA;
+		FReal ThicknessB;
 		CalculateQueryMargins(A, B, ThicknessA, ThicknessB);
 		ThicknessA += InThicknessA;
-		ThicknessB += InThicknessB;
-
 
-		const T Inflation = ThicknessA + ThicknessB + static_cast<T>(1e-3);
-		const VectorRegister4Float InflationSimd = MakeVectorRegisterFloat((float)Inflation, (float)Inflation, (float)Inflation, (float)Inflation);
+		const FReal InflationReal = ThicknessA + ThicknessB + static_cast<FReal>(1e-3);
+		const VectorRegister4Float Inflation = VectorSet1(static_cast<FRealSingle>(InflationReal));
+		const VectorRegister4Float Inflation2 = VectorMultiply(Inflation, Inflation);
 
 		int32 VertexIndexA = INDEX_NONE, VertexIndexB = INDEX_NONE;
 		do
@@ -101,14 +106,14 @@ namespace Chaos
 			{
 				break;	//if taking too long just stop. This should never happen
 			}
-			const VectorRegister4Float NegVSimd = VectorNegate(VSimd);
+			const VectorRegister4Float NegVSimd = VectorNegate(V);
 			const VectorRegister4Float SupportASimd = A.SupportCoreSimd(NegVSimd, ThicknessA);
-			const VectorRegister4Float VInBSimd = VectorQuaternionRotateVector(AToBRotationSimd, VSimd);
+			const VectorRegister4Float VInBSimd = VectorQuaternionRotateVector(AToBRotation, V);
 			const VectorRegister4Float SupportBLocalSimd = B.SupportCoreSimd(VInBSimd, ThicknessB);
-			const VectorRegister4Float SupportBSimd = VectorAdd(VectorQuaternionRotateVector(RotationSimd, SupportBLocalSimd), TranslationSimd);
+			const VectorRegister4Float SupportBSimd = VectorAdd(VectorQuaternionRotateVector(Rotation, SupportBLocalSimd), Translation);
 			const VectorRegister4Float WSimd = VectorSubtract(SupportASimd, SupportBSimd);
 
-			if (VectorMaskBits(VectorCompareGT(VectorDot3(VSimd, WSimd), InflationSimd)))
+			if (VectorMaskBits(VectorCompareGT(VectorDot3(V, WSimd), Inflation)))
 			{
 				return false;
 			}
@@ -123,10 +128,10 @@ namespace Chaos
 
 			NumVerts = VectorIntAdd(NumVerts, GlobalVectorConstants::IntOne);
 
-			VSimd = VectorSimplexFindClosestToOrigin<false>(SimplexSimd, NumVerts, BarycentricSimd, nullptr, nullptr);
+			V = VectorSimplexFindClosestToOrigin<false>(SimplexSimd, NumVerts, BarycentricSimd, nullptr, nullptr);
 
-			const VectorRegister4Float NewDist2Simd = VectorDot3(VSimd, VSimd);///
-			bNearZero = VectorMaskBits(VectorCompareLT(NewDist2Simd, VectorMultiply(InflationSimd, InflationSimd))) != 0;
+			const VectorRegister4Float NewDist2Simd = VectorDot3(V, V);///
+			bNearZero = VectorMaskBits(VectorCompareLT(NewDist2Simd, Inflation2)) != 0;
 
 			//as simplices become degenerate we will stop making progress. This is a side-effect of precision, in that case take V as the current best approximation
 			//question: should we take previous v in case it's better?
@@ -136,7 +141,7 @@ namespace Chaos
 
 			if (!bTerminate)
 			{
-				VSimd = VectorDivide(VSimd,VectorSqrt(NewDist2Simd));
+				V = VectorDivide(V, VectorSqrt(NewDist2Simd));
 			}
 
 		} while (!bTerminate);
