Chaos Cloth: Implement ISPC for Bending Constraints and XPBDSpringConstraints.

In the Shirring Test (in development editor), the Bending Constraints change drops the single iteration BendingApply from 2.1ms --> 350us. For reference, Bending Springs (with ISPC) are about 150us.

For XPBDSpringConstraints, there is no way to enable XPBD from the UI, so this does not affect existing behavior, but I use XPBD for testing high res cloth, so the ISPC XPBD speedups are important there.

This also includes code for using PhysicsParallelFor (following other parallel constraint examples). This is significantly slower than ISPC on PC (similar to other cloth constraints). I don't know if this would be better on other platforms, but the default is to use ISPC (following same convention as other cloth constraints).

diff --git a/Engine/Plugins/Experimental/ChaosCloth/Source/Chaos/Private/ChaosCloth/ChaosClothingSimulation.cpp b/Engine/Plugins/Experimental/ChaosCloth/Source/Chaos/Private/ChaosCloth/ChaosClothingSimulation.cpp
--- a/Engine/Plugins/Experimental/ChaosCloth/Source/Chaos/Private/ChaosCloth/ChaosClothingSimulation.cpp
+++ b/Engine/Plugins/Experimental/ChaosCloth/Source/Chaos/Private/ChaosCloth/ChaosClothingSimulation.cpp
@@ -32,6 +32,7 @@
 #include "Chaos/PBDCollisionSpringConstraints.h"
 #include "Chaos/PBDTriangleMeshCollisions.h"
 #include "Chaos/PBDBendingConstraints.h"
+#include "Chaos/XPBDSpringConstraints.h"
 #include "Chaos/VelocityField.h"
 #endif  // #if WITH_EDITOR || CHAOS_DEBUG_DRAW
 
@@ -148,6 +149,8 @@ namespace ClothingSimulationConsole
 					bChaos_PreSimulationTransforms_ISPC_Enabled =
 					bChaos_CalculateBounds_ISPC_Enabled =
 					bChaos_PostIterationUpdates_ISPC_Enabled =
+					bChaos_Bending_ISPC_Enabled =
+					bChaos_XPBDSpring_ISPC_Enabled =
 					bEnableISPC;
 				return;
 			}
diff --git a/Engine/Source/Runtime/Core/Public/Math/Vector.isph b/Engine/Source/Runtime/Core/Public/Math/Vector.isph
--- a/Engine/Source/Runtime/Core/Public/Math/Vector.isph
+++ b/Engine/Source/Runtime/Core/Public/Math/Vector.isph
@@ -110,6 +110,19 @@ struct FIntVector
 #endif
 };
 
+#ifndef EXPLICIT_INTVECTOR4
+#define EXPLICIT_INTVECTOR4 0
+#endif
+
+struct FIntVector4
+{
+#if EXPLICIT_INTVECTOR4 == 0
+	int V[4];
+#else
+	int<4> V;
+#endif
+};
+
 // Define types
 typedef FVector2 FVector2d;
 typedef FVector FVector3d;
@@ -347,6 +360,37 @@ inline void VectorStore(uniform FIntVector *uniform DstPtr, const varying FIntVe
 	}
 }
 
+inline varying FIntVector4 VectorLoad(const uniform FIntVector4 *uniform SrcPtr)
+{
+	varying FIntVector4 Result;
+
+	if(((1<<TARGET_WIDTH)-1 ^ lanemask()) == 0)
+	{
+		aos_to_soa4_ispc((uniform int32 * uniform)SrcPtr, &Result.V[0], &Result.V[1], &Result.V[2], &Result.V[3]);
+	}
+	else
+	{
+		#pragma ignore warning(perf)
+		Result = SrcPtr[programIndex];
+	}
+
+	return Result;
+}
+
+inline void VectorStore(uniform FIntVector4 *uniform DstPtr, const varying FIntVector4 &V)
+{
+	if(((1<<TARGET_WIDTH)-1 ^ lanemask()) == 0)
+	{
+		soa_to_aos4_ispc(V.V[0], V.V[1], V.V[2], V.V[3], (uniform int32 *uniform)DstPtr);
+	}
+	else
+	{
+		#pragma ignore warning(perf)
+		DstPtr[programIndex] = V;
+	}
+}
+
+
 inline varying FVector2 VectorLoad(const uniform FVector2 *uniform SrcPtr)
 {
 	varying FVector2 Result;
diff --git a/Engine/Source/Runtime/Experimental/Chaos/Private/Chaos/PBDBendingConstraints.cpp b/Engine/Source/Runtime/Experimental/Chaos/Private/Chaos/PBDBendingConstraints.cpp
new file mode 100644
--- /dev/null
+++ b/Engine/Source/Runtime/Experimental/Chaos/Private/Chaos/PBDBendingConstraints.cpp
@@ -0,0 +1,202 @@
+// Copyright Epic Games, Inc. All Rights Reserved.
+#include "Chaos/PBDBendingConstraints.h"
+#include "Chaos/GraphColoring.h"
+#include "Chaos/Framework/Parallel.h"
+#include "ChaosStats.h"
+#if INTEL_ISPC
+#include "PBDBendingConstraints.ispc.generated.h"
+#endif
+
+DECLARE_CYCLE_STAT(TEXT("Chaos PBD Bending Constraint"), STAT_PBD_Bending, STATGROUP_Chaos);
+
+#if INTEL_ISPC && !UE_BUILD_SHIPPING
+bool bChaos_Bending_ISPC_Enabled = true;
+FAutoConsoleVariableRef CVarChaosBendingISPCEnabled(TEXT("p.Chaos.Bending.ISPC"), bChaos_Bending_ISPC_Enabled, TEXT("Whether to use ISPC optimizations in Bending constraints"));
+
+static_assert(sizeof(ispc::FVector4f) == sizeof(Chaos::Softs::FPAndInvM), "sizeof(ispc::FVector4f) != sizeof(Chaos::Softs::FPAndInvM");
+static_assert(sizeof(ispc::FIntVector4) == sizeof(Chaos::TVec4<int32>), "sizeof(ispc::FIntVector4) != sizeof(Chaos::TVec4<int32>");
+#endif
+
+namespace Chaos::Softs {
+
+// @todo(chaos): the parallel threshold (or decision to run parallel) should probably be owned by the solver and passed to the constraint container
+int32 Chaos_Bending_ParallelConstraintCount = 100;
+FAutoConsoleVariableRef CVarChaosBendingParallelConstraintCount(TEXT("p.Chaos.Bending.ParallelConstraintCount"), Chaos_Bending_ParallelConstraintCount, TEXT("If we have more constraints than this, use parallel-for in Apply."));
+
+void FPBDBendingConstraints::InitColor(const FSolverParticles& InParticles)
+{
+	// In dev builds we always color so we can tune the system without restarting. See Apply()
+#if UE_BUILD_SHIPPING || UE_BUILD_TEST
+	if (Constraints.Num() > Chaos_Bending_ParallelConstraintCount)
+#endif
+	{
+		const TArray<TArray<int32>> ConstraintsPerColor = FGraphColoring::ComputeGraphColoring(Constraints, InParticles);
+
+		// Reorder constraints based on color so each array in ConstraintsPerColor contains contiguous elements.
+		TArray<TVec4<int32>> ReorderedConstraints; 
+		TArray<TVec2<int32>> ReorderedConstraintSharedEdges;
+		TArray<FSolverReal> ReorderedRestAngles;
+		TArray<int32> OrigToReorderedIndices; // used to reorder stiffness indices
+		ReorderedConstraints.SetNumUninitialized(Constraints.Num());
+		ReorderedConstraintSharedEdges.SetNumUninitialized(Constraints.Num());
+		ReorderedRestAngles.SetNumUninitialized(Constraints.Num());
+		OrigToReorderedIndices.SetNumUninitialized(Constraints.Num());
+
+		ConstraintsPerColorStartIndex.Reset(ConstraintsPerColor.Num() + 1);
+
+		int32 ReorderedIndex = 0;
+		for (const TArray<int32>& ConstraintsBatch : ConstraintsPerColor)
+		{
+			ConstraintsPerColorStartIndex.Add(ReorderedIndex);
+			for (const int32& BatchConstraint : ConstraintsBatch)
+			{
+				const int32 OrigIndex = BatchConstraint;
+				ReorderedConstraints[ReorderedIndex] = Constraints[OrigIndex];
+				ReorderedConstraintSharedEdges[ReorderedIndex] = ConstraintSharedEdges[OrigIndex];
+				ReorderedRestAngles[ReorderedIndex] = RestAngles[OrigIndex];
+				OrigToReorderedIndices[OrigIndex] = ReorderedIndex;
+				++ReorderedIndex;
+			}
+		}
+		ConstraintsPerColorStartIndex.Add(ReorderedIndex);
+
+		Constraints = MoveTemp(ReorderedConstraints);
+		ConstraintSharedEdges = MoveTemp(ReorderedConstraintSharedEdges);
+		RestAngles = MoveTemp(ReorderedRestAngles);
+		Stiffness.ReorderIndices(OrigToReorderedIndices);
+		BucklingStiffness.ReorderIndices(OrigToReorderedIndices);
+	}
+}
+
+void FPBDBendingConstraints::ApplyHelper(FSolverParticles& Particles, const FSolverReal Dt, const int32 ConstraintIndex, const FSolverReal ExpStiffnessValue, const FSolverReal ExpBucklingValue) const
+{
+
+	const TVec4<int32>& Constraint = Constraints[ConstraintIndex];
+	const int32 i1 = Constraint[0];
+	const int32 i2 = Constraint[1];
+	const int32 i3 = Constraint[2];
+	const int32 i4 = Constraint[3];
+	const TStaticArray<FSolverVec3, 4> Grads = Base::GetGradients(Particles, ConstraintIndex);
+	const FSolverReal S = Base::GetScalingFactor(Particles, ConstraintIndex, Grads, ExpStiffnessValue, ExpBucklingValue);
+	Particles.P(i1) -= S * Particles.InvM(i1) * Grads[0];
+	Particles.P(i2) -= S * Particles.InvM(i2) * Grads[1];
+	Particles.P(i3) -= S * Particles.InvM(i3) * Grads[2];
+	Particles.P(i4) -= S * Particles.InvM(i4) * Grads[3];
+}
+
+void FPBDBendingConstraints::Apply(FSolverParticles& Particles, const FSolverReal Dt) const
+{
+	TRACE_CPUPROFILER_EVENT_SCOPE(FPBDBendingConstraints_Apply);
+	SCOPE_CYCLE_COUNTER(STAT_PBD_Bending);
+
+	const bool StiffnessHasWeightMap = Stiffness.HasWeightMap();
+	const bool BucklingStiffnessHasWeightMap = BucklingStiffness.HasWeightMap();
+
+	if (ConstraintsPerColorStartIndex.Num() > 0 && Constraints.Num() > Chaos_Bending_ParallelConstraintCount)
+	{
+		const int32 ConstraintColorNum = ConstraintsPerColorStartIndex.Num() - 1;
+		if (!StiffnessHasWeightMap && !BucklingStiffnessHasWeightMap)
+		{
+			const FSolverReal ExpStiffnessValue = (FSolverReal)Stiffness;
+			const FSolverReal ExpBucklingValue = (FSolverReal)BucklingStiffness;
+#if INTEL_ISPC
+			if (bRealTypeCompatibleWithISPC && bChaos_Bending_ISPC_Enabled)
+			{
+				for (int32 ConstraintColorIndex = 0; ConstraintColorIndex < ConstraintColorNum; ++ConstraintColorIndex)
+				{
+					const int32 ColorStart = ConstraintsPerColorStartIndex[ConstraintColorIndex];
+					const int32 ColorSize = ConstraintsPerColorStartIndex[ConstraintColorIndex + 1] - ColorStart;
+					ispc::ApplyBendingConstraints(
+						(ispc::FVector4f*)Particles.GetPAndInvM().GetData(),
+						(ispc::FIntVector4*)&Constraints.GetData()[ColorStart],
+						&RestAngles.GetData()[ColorStart],
+						ExpStiffnessValue,
+						BucklingRatio,
+						ExpBucklingValue,
+						ColorSize);
+				}
+			}
+			else
+#endif
+			{
+				for (int32 ConstraintColorIndex = 0; ConstraintColorIndex < ConstraintColorNum; ++ConstraintColorIndex)
+				{
+					const int32 ColorStart = ConstraintsPerColorStartIndex[ConstraintColorIndex];
+					const int32 ColorSize = ConstraintsPerColorStartIndex[ConstraintColorIndex + 1] - ColorStart;
+					PhysicsParallelFor(ColorSize, [&](const int32 Index)
+					{
+						const int32 ConstraintIndex = ColorStart + Index;
+						ApplyHelper(Particles, Dt, ConstraintIndex, ExpStiffnessValue, ExpBucklingValue);
+					});
+				}
+			}
+		}
+		else  // Has weight maps
+		{
+#if INTEL_ISPC
+			if (bRealTypeCompatibleWithISPC && bChaos_Bending_ISPC_Enabled)
+			{
+				for (int32 ConstraintColorIndex = 0; ConstraintColorIndex < ConstraintColorNum; ++ConstraintColorIndex)
+				{
+					const int32 ColorStart = ConstraintsPerColorStartIndex[ConstraintColorIndex];
+					const int32 ColorSize = ConstraintsPerColorStartIndex[ConstraintColorIndex + 1] - ColorStart;
+					ispc::ApplyBendingConstraintsWithMaps(
+						(ispc::FVector4f*)Particles.GetPAndInvM().GetData(),
+						(ispc::FIntVector4*)&Constraints.GetData()[ColorStart],
+						&RestAngles.GetData()[ColorStart],
+						StiffnessHasWeightMap,
+						&Stiffness.GetIndices().GetData()[ColorStart],
+						&Stiffness.GetTable().GetData()[0],
+						BucklingRatio,
+						BucklingStiffnessHasWeightMap,
+						&BucklingStiffness.GetIndices().GetData()[ColorStart],
+						&BucklingStiffness.GetTable().GetData()[0],
+						ColorSize);
+				}
+			}
+			else
+#endif
+			{
+				const FSolverReal StiffnessNoMap = (FSolverReal)Stiffness;
+				const FSolverReal BucklingStiffnessNoMap = (FSolverReal)BucklingStiffness;
+				for (int32 ConstraintColorIndex = 0; ConstraintColorIndex < ConstraintColorNum; ++ConstraintColorIndex)
+				{
+					const int32 ColorStart = ConstraintsPerColorStartIndex[ConstraintColorIndex];
+					const int32 ColorSize = ConstraintsPerColorStartIndex[ConstraintColorIndex + 1] - ColorStart;
+					PhysicsParallelFor(ColorSize, [&](const int32 Index)
+					{
+						const int32 ConstraintIndex = ColorStart + Index;
+						const FSolverReal ExpStiffnessValue = StiffnessHasWeightMap ? Stiffness[ConstraintIndex] : StiffnessNoMap;
+						const FSolverReal ExpBucklingValue = BucklingStiffnessHasWeightMap ? BucklingStiffness[ConstraintIndex] : BucklingStiffnessNoMap;
+						ApplyHelper(Particles, Dt, ConstraintIndex, ExpStiffnessValue, ExpBucklingValue);
+					});
+				}
+			}
+		}
+	}
+	else
+	{
+		if (!StiffnessHasWeightMap && !BucklingStiffnessHasWeightMap)
+		{
+			const FSolverReal ExpStiffnessValue = (FSolverReal)Stiffness;
+			const FSolverReal ExpBucklingValue = (FSolverReal)BucklingStiffness;
+			for (int32 ConstraintIndex = 0; ConstraintIndex < Constraints.Num(); ++ConstraintIndex)
+			{
+				ApplyHelper(Particles, Dt, ConstraintIndex, ExpStiffnessValue, ExpBucklingValue);
+			}
+		}
+		else
+		{
+			const FSolverReal StiffnessNoMap = (FSolverReal)Stiffness;
+			const FSolverReal BucklingStiffnessNoMap = (FSolverReal)BucklingStiffness;
+			for (int32 ConstraintIndex = 0; ConstraintIndex < Constraints.Num(); ++ConstraintIndex)
+			{
+				const FSolverReal ExpStiffnessValue = StiffnessHasWeightMap ? Stiffness[ConstraintIndex] : StiffnessNoMap;
+				const FSolverReal ExpBucklingValue = BucklingStiffnessHasWeightMap ? BucklingStiffness[ConstraintIndex] : BucklingStiffnessNoMap;
+				ApplyHelper(Particles, Dt, ConstraintIndex, ExpStiffnessValue, ExpBucklingValue);
+			}
+		}
+	}
+}
+
+}  // End namespace Chaos::Softs
diff --git a/Engine/Source/Runtime/Experimental/Chaos/Private/Chaos/PBDBendingConstraints.ispc b/Engine/Source/Runtime/Experimental/Chaos/Private/Chaos/PBDBendingConstraints.ispc
new file mode 100644
--- /dev/null
+++ b/Engine/Source/Runtime/Experimental/Chaos/Private/Chaos/PBDBendingConstraints.ispc
@@ -0,0 +1,190 @@
+// Copyright Epic Games, Inc. All Rights Reserved.
+
+#define EXPLICIT_VECTOR4 1
+
+#include "Math/Vector.isph"
+#include "Chaos/PBDSofts.isph"
+
+static const uniform float SingleStepAngleLimit = FLOAT_PI * .25f;
+
+static inline FVector3f SafeDivide(const FVector3f& Numerator, const float Denominator)
+{
+	return VectorSelect(Denominator <= FLOAT_SMALL_NUMBER, FloatZeroVector, Numerator / Denominator);
+}
+
+static inline varying float SafeDivide(const varying float Numerator, const varying float Denominator)
+{
+	return VectorSelect(Denominator <= FLOAT_SMALL_NUMBER, (uniform float)0, Numerator / Denominator);
+}
+
+// Returns Angle
+static inline float CalcGradientsAndAngle(const FVector3f& P1, const FVector3f& P2, const FVector3f& P3, const FVector3f& P4, FVector3f& Grad1, FVector3f& Grad2, FVector3f& Grad3, FVector3f& Grad4)
+{
+	const FVector3f SharedEdgeNormalized = VectorGetSafeNormal(P2 - P1);
+	const FVector3f P13CrossP23 = VectorCross(P1 - P3, P2 - P3);
+	const float Normal1Len = VectorSize(P13CrossP23);
+	const FVector3f Normal1 = SafeDivide(P13CrossP23, Normal1Len);
+	const FVector3f P24CrossP14 = VectorCross(P2 - P4, P1 - P4);
+	const float Normal2Len = VectorSize(P24CrossP14);
+	const FVector3f Normal2 = SafeDivide(P24CrossP14, Normal2Len);
+
+	const FVector3f N2CrossN1 = VectorCross(Normal2, Normal1);
+
+	const float CosPhi = clamp(VectorDot(Normal1, Normal2), (uniform float)(-1), (uniform float)(1));
+	const float SinPhi = clamp(VectorDot(N2CrossN1, SharedEdgeNormalized), (uniform float)(-1), (uniform float)(1));
+
+	const FVector3f DPhiDN1_OverNormal1Len = SafeDivide(CosPhi * VectorCross(SharedEdgeNormalized, Normal2) - SinPhi * Normal2, Normal1Len);
+	const FVector3f DPhiDN2_OverNormal2Len = SafeDivide(CosPhi * VectorCross(Normal1, SharedEdgeNormalized) - SinPhi * Normal1, Normal2Len);
+
+	const FVector3f DPhiDP13 = VectorCross(P2 - P3, DPhiDN1_OverNormal1Len);
+	const FVector3f DPhiDP23 = VectorCross(DPhiDN1_OverNormal1Len, P1 - P3);
+	const FVector3f DPhiDP24 = VectorCross(P1 - P4, DPhiDN2_OverNormal2Len);
+	const FVector3f DPhiDP14 = VectorCross(DPhiDN2_OverNormal2Len, P2 - P4);
+
+	Grad1 = DPhiDP13 + DPhiDP14;
+	Grad2 = DPhiDP23 + DPhiDP24;
+	Grad3 = -1 * DPhiDP13 - DPhiDP23;
+	Grad4 = -1 * DPhiDP14 - DPhiDP24;
+
+	return atan2(SinPhi, CosPhi);
+}
+
+
+export void ApplyBendingConstraints(uniform FVector4f PandInvM[],
+									const uniform FIntVector4 Constraints[],
+									const uniform float RestAngles[],
+									const uniform float Stiffness,
+									const uniform float BucklingRatio,
+									const uniform float BucklingStiffness,
+									const uniform int32 NumConstraints)
+{
+
+	foreach(i = 0 ... NumConstraints)
+	{
+		const varying FIntVector4 Constraint = VectorLoad(&Constraints[extract(i,0)]);
+
+		const varying int32 i1 = Constraint.V[0];
+		const varying int32 i2 = Constraint.V[1];
+		const varying int32 i3 = Constraint.V[2];
+		const varying int32 i4 = Constraint.V[3];
+
+		const varying FVector4f PandInvM2 = VectorGather(&PandInvM[i2]);
+		const varying FVector4f PandInvM3 = VectorGather(&PandInvM[i3]);
+		const varying FVector4f PandInvM1 = VectorGather(&PandInvM[i1]);
+		const varying FVector4f PandInvM4 = VectorGather(&PandInvM[i4]);
+
+		varying FVector3f P1, P2, P3, P4;
+		varying float M1, M2, M3, M4;
+
+		UnzipPandInvM(PandInvM1, P1, M1);
+		UnzipPandInvM(PandInvM2, P2, M2);
+		UnzipPandInvM(PandInvM3, P3, M3);
+		UnzipPandInvM(PandInvM4, P4, M4);
+
+
+		varying FVector3f Grad1, Grad2, Grad3, Grad4;
+		const varying float Angle = CalcGradientsAndAngle(P1, P2, P3, P4, Grad1, Grad2, Grad3, Grad4);
+
+		const varying float Denom = M1 * VectorSizeSquared(Grad1) + M2 * VectorSizeSquared(Grad2) + M3 * VectorSizeSquared(Grad3) + M4 * VectorSizeSquared(Grad4);
+		
+		const varying float RestAngle = RestAngles[i];
+		const varying bool AngleIsBuckled = FLOAT_PI - abs(Angle) < BucklingRatio * (FLOAT_PI - abs(RestAngle));
+		const varying float StiffnessValue = VectorSelect(AngleIsBuckled, BucklingStiffness, Stiffness);
+
+		const varying float Delta = clamp(StiffnessValue * (Angle - RestAngle), -SingleStepAngleLimit, SingleStepAngleLimit);
+
+		const varying float S = SafeDivide(Delta, Denom);
+
+		if (M1 > 0)
+		{
+			VectorScatter(&PandInvM[i1], SetVector4(P1 - S * M1 * Grad1, M1));
+		}
+		if (M2 > 0)
+		{
+			VectorScatter(&PandInvM[i2], SetVector4(P2 - S * M2 * Grad2, M2));
+		}
+		if (M3 > 0)
+		{
+			VectorScatter(&PandInvM[i3], SetVector4(P3 - S * M3 * Grad3, M3));
+		}
+		if (M4 > 0)
+		{
+			VectorScatter(&PandInvM[i4], SetVector4(P4 - S * M4 * Grad4, M4));
+		}
+	}
+}
+
+export void ApplyBendingConstraintsWithMaps(uniform FVector4f PandInvM[],
+											const uniform FIntVector4 Constraints[],
+											const uniform float RestAngles[],
+											const uniform bool StiffnessHasMap,
+											const uniform uint8 StiffnessIndices[],
+											const uniform float StiffnessTable[],
+											const uniform float BucklingRatio,
+											const uniform bool BucklingStiffnessHasMap,
+											const uniform uint8 BucklingStiffnessIndices[],
+											const uniform float BucklingStiffnessTable[],
+											const uniform int32 NumConstraints)
+{
+
+	foreach(i = 0 ... NumConstraints)
+	{
+		const varying FIntVector4 Constraint = VectorLoad(&Constraints[extract(i,0)]);
+
+		const varying int32 i1 = Constraint.V[0];
+		const varying int32 i2 = Constraint.V[1];
+		const varying int32 i3 = Constraint.V[2];
+		const varying int32 i4 = Constraint.V[3];
+		const varying int8 StiffnessIndex = VectorSelect(StiffnessHasMap, StiffnessIndices[i], FLOAT_ZERO);
+		const varying int8 BucklingStiffnessIndex = VectorSelect(BucklingStiffnessHasMap, BucklingStiffnessIndices[i], FLOAT_ZERO);
+
+		const varying FVector4f PandInvM2 = VectorGather(&PandInvM[i2]);
+		const varying FVector4f PandInvM3 = VectorGather(&PandInvM[i3]);
+		const varying FVector4f PandInvM1 = VectorGather(&PandInvM[i1]);
+		const varying FVector4f PandInvM4 = VectorGather(&PandInvM[i4]);
+
+		varying FVector3f P1, P2, P3, P4;
+		varying float M1, M2, M3, M4;
+
+		UnzipPandInvM(PandInvM1, P1, M1);
+		UnzipPandInvM(PandInvM2, P2, M2);
+		UnzipPandInvM(PandInvM3, P3, M3);
+		UnzipPandInvM(PandInvM4, P4, M4);
+
+
+		varying FVector3f Grad1, Grad2, Grad3, Grad4;
+		const varying float Angle = CalcGradientsAndAngle(P1, P2, P3, P4, Grad1, Grad2, Grad3, Grad4);
+
+		const varying float Denom = M1 * VectorSizeSquared(Grad1) + M2 * VectorSizeSquared(Grad2) + M3 * VectorSizeSquared(Grad3) + M4 * VectorSizeSquared(Grad4);
+		
+		#pragma ignore warning(perf)
+		const varying float Stiffness = StiffnessTable[StiffnessIndex];
+		#pragma ignore warning(perf)
+		const varying float BucklingStiffness = BucklingStiffnessTable[BucklingStiffnessIndex];
+
+		const varying float RestAngle = RestAngles[i];
+		const varying bool AngleIsBuckled = FLOAT_PI - abs(Angle) < BucklingRatio * (FLOAT_PI - abs(RestAngle));
+		const varying float StiffnessValue = VectorSelect(AngleIsBuckled, BucklingStiffness, Stiffness);
+
+		const varying float Delta = clamp(StiffnessValue * (Angle - RestAngle), -SingleStepAngleLimit, SingleStepAngleLimit);
+
+		const varying float S = SafeDivide(Delta, Denom);
+
+		if (M1 > 0)
+		{
+			VectorScatter(&PandInvM[i1], SetVector4(P1 - S * M1 * Grad1, M1));
+		}
+		if (M2 > 0)
+		{
+			VectorScatter(&PandInvM[i2], SetVector4(P2 - S * M2 * Grad2, M2));
+		}
+		if (M3 > 0)
+		{
+			VectorScatter(&PandInvM[i3], SetVector4(P3 - S * M3 * Grad3, M3));
+		}
+		if (M4 > 0)
+		{
+			VectorScatter(&PandInvM[i4], SetVector4(P4 - S * M4 * Grad4, M4));
+		}
+	}
+}
diff --git a/Engine/Source/Runtime/Experimental/Chaos/Private/Chaos/XPBDSpringConstraints.cpp b/Engine/Source/Runtime/Experimental/Chaos/Private/Chaos/XPBDSpringConstraints.cpp
new file mode 100644
--- /dev/null
+++ b/Engine/Source/Runtime/Experimental/Chaos/Private/Chaos/XPBDSpringConstraints.cpp
@@ -0,0 +1,190 @@
+// Copyright Epic Games, Inc. All Rights Reserved.
+#include "Chaos/XPBDSpringConstraints.h"
+#include "Chaos/PBDSoftsSolverParticles.h"
+#include "Chaos/GraphColoring.h"
+#include "Chaos/Framework/Parallel.h"
+#include "ChaosStats.h"
+#include "ChaosLog.h"
+#include "HAL/IConsoleManager.h"
+
+#if INTEL_ISPC
+#include "XPBDSpringConstraints.ispc.generated.h"
+#endif
+
+DECLARE_CYCLE_STAT(TEXT("Chaos XPBD Spring Constraint"), STAT_XPBD_Spring, STATGROUP_Chaos);
+
+#if INTEL_ISPC && !UE_BUILD_SHIPPING
+static_assert(sizeof(ispc::FVector4f) == sizeof(Chaos::Softs::FPAndInvM), "sizeof(ispc::FVector4f) != sizeof(Chaos::Softs::FPAndInvM)");
+static_assert(sizeof(ispc::FIntVector2) == sizeof(Chaos::TVec2<int32>), "sizeof(ispc::FIntVector2) != sizeof(Chaos::TVec2<int32>)");
+
+bool bChaos_XPBDSpring_ISPC_Enabled = true;
+FAutoConsoleVariableRef CVarChaosXPBDSpringISPCEnabled(TEXT("p.Chaos.XPBDSpring.ISPC"), bChaos_XPBDSpring_ISPC_Enabled, TEXT("Whether to use ISPC optimizations in XPBD Spring constraints"));
+#endif
+
+namespace Chaos::Softs {
+
+// @todo(chaos): the parallel threshold (or decision to run parallel) should probably be owned by the solver and passed to the constraint container
+static int32 Chaos_XPBDSpring_ParallelConstraintCount = 100;
+#if !UE_BUILD_SHIPPING
+FAutoConsoleVariableRef CVarChaosXPBDSpringParallelConstraintCount(TEXT("p.Chaos.XPBDSpring.ParallelConstraintCount"), Chaos_XPBDSpring_ParallelConstraintCount, TEXT("If we have more constraints than this, use parallel-for in Apply."));
+#endif
+
+void FXPBDSpringConstraints::InitColor(const FSolverParticles& Particles)
+{
+	// In dev builds we always color so we can tune the system without restarting. See Apply()
+#if UE_BUILD_SHIPPING || UE_BUILD_TEST
+	if (Constraints.Num() > Chaos_Spring_ParallelConstraintCount)
+#endif
+	{
+		const TArray<TArray<int32>> ConstraintsPerColor = FGraphColoring::ComputeGraphColoring(Constraints, Particles);
+		
+		// Reorder constraints based on color so each array in ConstraintsPerColor contains contiguous elements.
+		TArray<TVec2<int32>> ReorderedConstraints;
+		TArray<FSolverReal> ReorderedDists;
+		TArray<int32> OrigToReorderedIndices; // used to reorder stiffness indices
+		ReorderedConstraints.SetNumUninitialized(Constraints.Num());
+		ReorderedDists.SetNumUninitialized(Dists.Num());
+		OrigToReorderedIndices.SetNumUninitialized(Constraints.Num());
+
+		ConstraintsPerColorStartIndex.Reset(ConstraintsPerColor.Num() + 1);
+
+		int32 ReorderedIndex = 0;
+		for (const TArray<int32>& ConstraintsBatch : ConstraintsPerColor)
+		{
+			ConstraintsPerColorStartIndex.Add(ReorderedIndex);
+			for (const int32& BatchConstraint : ConstraintsBatch)
+			{
+				const int32 OrigIndex = BatchConstraint;
+				ReorderedConstraints[ReorderedIndex] = Constraints[OrigIndex];
+				ReorderedDists[ReorderedIndex] = Dists[OrigIndex];
+				OrigToReorderedIndices[OrigIndex] = ReorderedIndex;
+
+				++ReorderedIndex;
+			}
+		}
+		ConstraintsPerColorStartIndex.Add(ReorderedIndex);
+
+		Constraints = MoveTemp(ReorderedConstraints);
+		Dists = MoveTemp(ReorderedDists);
+		Stiffness.ReorderIndices(OrigToReorderedIndices);
+	}
+}
+
+void FXPBDSpringConstraints::ApplyHelper(FSolverParticles& Particles, const FSolverReal Dt, const int32 ConstraintIndex, const FSolverReal ExpStiffnessValue) const
+{
+	const TVec2<int32>& Constraint = Constraints[ConstraintIndex];
+	const int32 i1 = Constraint[0];
+	const int32 i2 = Constraint[1];
+	const FSolverVec3 Delta =  GetDelta(Particles, Dt, ConstraintIndex, ExpStiffnessValue);
+	if (Particles.InvM(i1) > (FSolverReal)0.)
+	{
+		Particles.P(i1) -= Particles.InvM(i1) * Delta;
+	}
+	if (Particles.InvM(i2) > (FSolverReal)0.)
+	{
+		Particles.P(i2) += Particles.InvM(i2) * Delta;
+	}
+}
+
+void FXPBDSpringConstraints::Apply(FSolverParticles& Particles, const FSolverReal Dt) const
+{
+	TRACE_CPUPROFILER_EVENT_SCOPE(FXPBDSpringConstraints_Apply);
+	SCOPE_CYCLE_COUNTER(STAT_XPBD_Spring);
+	if ((ConstraintsPerColorStartIndex.Num() > 1) && (Constraints.Num() > Chaos_XPBDSpring_ParallelConstraintCount))
+	{
+		const int32 ConstraintColorNum = ConstraintsPerColorStartIndex.Num() - 1;
+		if (!Stiffness.HasWeightMap())
+		{
+			const FSolverReal ExpStiffnessValue = (FSolverReal)Stiffness;
+
+#if INTEL_ISPC
+			if (bRealTypeCompatibleWithISPC && bChaos_XPBDSpring_ISPC_Enabled)
+			{
+				for (int32 ConstraintColorIndex = 0; ConstraintColorIndex < ConstraintColorNum; ++ConstraintColorIndex)
+				{
+					const int32 ColorStart = ConstraintsPerColorStartIndex[ConstraintColorIndex];
+					const int32 ColorSize = ConstraintsPerColorStartIndex[ConstraintColorIndex + 1] - ColorStart;
+					ispc::ApplyXPBDSpringConstraints(
+						(ispc::FVector4f*)Particles.GetPAndInvM().GetData(),
+						(ispc::FIntVector2*)&Constraints.GetData()[ColorStart],
+						&Dists.GetData()[ColorStart],
+						&Lambdas.GetData()[ColorStart],
+						Dt,
+						ExpStiffnessValue,
+						ColorSize);
+				}
+			}
+			else
+#endif
+			{
+				for (int32 ConstraintColorIndex = 0; ConstraintColorIndex < ConstraintColorNum; ++ConstraintColorIndex)
+				{
+					const int32 ColorStart = ConstraintsPerColorStartIndex[ConstraintColorIndex];
+					const int32 ColorSize = ConstraintsPerColorStartIndex[ConstraintColorIndex + 1] - ColorStart;
+					PhysicsParallelFor(ColorSize, [&](const int32 Index)
+					{
+						const int32 ConstraintIndex = ColorStart + Index;
+						ApplyHelper(Particles, Dt, ConstraintIndex, ExpStiffnessValue);
+					});
+				}
+			}
+		}
+		else  // Has weight maps
+		{
+#if INTEL_ISPC
+			if (bRealTypeCompatibleWithISPC && bChaos_XPBDSpring_ISPC_Enabled)
+			{
+				for (int32 ConstraintColorIndex = 0; ConstraintColorIndex < ConstraintColorNum; ++ConstraintColorIndex)
+				{
+					const int32 ColorStart = ConstraintsPerColorStartIndex[ConstraintColorIndex];
+					const int32 ColorSize = ConstraintsPerColorStartIndex[ConstraintColorIndex + 1] - ColorStart;
+					ispc::ApplyXPBDSpringConstraintsWithWeightMaps(
+						(ispc::FVector4f*)Particles.GetPAndInvM().GetData(),
+						(ispc::FIntVector2*) & Constraints.GetData()[ColorStart],
+						&Dists.GetData()[ColorStart],
+						&Lambdas.GetData()[ColorStart],
+						Dt,
+						&Stiffness.GetIndices().GetData()[ColorStart],
+						&Stiffness.GetTable().GetData()[0],
+						ColorSize);
+				}
+			}
+			else
+#endif
+			{
+				for (int32 ConstraintColorIndex = 0; ConstraintColorIndex < ConstraintColorNum; ++ConstraintColorIndex)
+				{
+					const int32 ColorStart = ConstraintsPerColorStartIndex[ConstraintColorIndex];
+					const int32 ColorSize = ConstraintsPerColorStartIndex[ConstraintColorIndex + 1] - ColorStart;
+					PhysicsParallelFor(ColorSize, [&](const int32 Index)
+					{
+						const int32 ConstraintIndex = ColorStart + Index;
+						const FSolverReal ExpStiffnessValue = Stiffness[ConstraintIndex];
+						ApplyHelper(Particles, Dt, ConstraintIndex, ExpStiffnessValue);
+					});
+				}
+			}
+		}
+	}
+	else
+	{
+		if (!Stiffness.HasWeightMap())
+		{
+			const FSolverReal ExpStiffnessValue = (FSolverReal)Stiffness;
+			for (int32 ConstraintIndex = 0; ConstraintIndex < Constraints.Num(); ++ConstraintIndex)
+			{
+				ApplyHelper(Particles, Dt, ConstraintIndex, ExpStiffnessValue);
+			}
+		}
+		else
+		{
+			for (int32 ConstraintIndex = 0; ConstraintIndex < Constraints.Num(); ++ConstraintIndex)
+			{
+				const FSolverReal ExpStiffnessValue = Stiffness[ConstraintIndex];
+				ApplyHelper(Particles, Dt, ConstraintIndex, ExpStiffnessValue);
+			}
+		}
+	}
+}
+
+} // End namespace Chaos::Softs
diff --git a/Engine/Source/Runtime/Experimental/Chaos/Private/Chaos/XPBDSpringConstraints.ispc b/Engine/Source/Runtime/Experimental/Chaos/Private/Chaos/XPBDSpringConstraints.ispc
new file mode 100644
--- /dev/null
+++ b/Engine/Source/Runtime/Experimental/Chaos/Private/Chaos/XPBDSpringConstraints.ispc
@@ -0,0 +1,122 @@
+// Copyright Epic Games, Inc. All Rights Reserved.
+
+#define EXPLICIT_VECTOR4 1
+
+#include "Math/Vector.isph"
+#include "Chaos/PBDSofts.isph"
+
+static inline float SafeNormalize(FVector3f &Direction)
+{
+	const float SizeSquared = VectorSizeSquared(Direction);
+	const float Size = sqrt(SizeSquared);
+	Direction = VectorSelect((SizeSquared < FLOAT_KINDA_SMALL_NUMBER), FloatForwardVector, Direction / Size);
+	return (SizeSquared < FLOAT_KINDA_SMALL_NUMBER) ? FLOAT_ZERO : Size;
+}
+
+export void ApplyXPBDSpringConstraints(uniform FVector4f PandInvM[],
+									const uniform FIntVector2 Constraints[],
+									const uniform float Dists[],
+									uniform float Lambdas[],
+									const uniform float Dt,
+									const uniform float Stiffness,
+									const uniform int32 NumConstraints)
+{
+	foreach(i = 0 ... NumConstraints)
+	{
+		const varying FIntVector2 Constraint = VectorLoad(&Constraints[extract(i,0)]);
+		const varying int32 i1 = Constraint.V[0];
+		const varying int32 i2 = Constraint.V[1];
+
+		const varying FVector4f PandInvM1 = VectorGather(&PandInvM[i1]);
+		const varying FVector4f PandInvM2 = VectorGather(&PandInvM[i2]);
+
+		varying FVector3f P1, P2;
+		varying float InvM1, InvM2;
+		UnzipPandInvM(PandInvM1, P1, InvM1);
+		UnzipPandInvM(PandInvM2, P2, InvM2);
+
+		FVector3f Direction = P1 - P2;
+		const float Distance = SafeNormalize(Direction);
+
+		const varying float Dist = Dists[i];
+		const varying float Lambda = Lambdas[i];
+		const uniform float Alpha = (uniform float)1.f / (Stiffness * Dt * Dt);
+		const varying float Offset = Distance - Dist;
+
+		const float CombinedInvMass = InvM2 + InvM1;
+
+		const varying float DLambda = (Offset - Alpha * Lambda) / (CombinedInvMass + Alpha);
+
+		const FVector3f ComputedDelta = DLambda * Direction;
+
+		const FVector3f Delta = VectorSelect((InvM2 == FLOAT_ZERO && InvM1 == FLOAT_ZERO), FloatZeroVector, ComputedDelta);
+
+		Lambdas[i] = Lambda + DLambda;
+
+		if (InvM1 > FLOAT_ZERO)
+		{
+			VectorScatter(&PandInvM[i1], SetVector4( P1 - (InvM1 * Delta), InvM1 ));
+		}
+		if (InvM2 > FLOAT_ZERO)
+		{
+			VectorScatter(&PandInvM[i2], SetVector4( P2 + (InvM2 * Delta), InvM2 ));
+		}
+	}
+}
+
+export void ApplyXPBDSpringConstraintsWithWeightMaps(uniform FVector4f PandInvM[],
+									const uniform FIntVector2 Constraints[],
+									const uniform float Dists[],
+									uniform float Lambdas[],
+									const uniform float Dt,
+									const uniform uint8 StiffnessIndices[],
+									const uniform float StiffnessTable[],
+									const uniform int32 NumConstraints)
+{
+
+	foreach(i = 0 ... NumConstraints)
+	{
+		const varying FIntVector2 Constraint = VectorLoad(&Constraints[extract(i,0)]);
+		const varying int32 i1 = Constraint.V[0];
+		const varying int32 i2 = Constraint.V[1];
+
+		const varying FVector4f PandInvM1 = VectorGather(&PandInvM[i1]);
+		const varying FVector4f PandInvM2 = VectorGather(&PandInvM[i2]);
+
+		varying FVector3f P1, P2;
+		varying float InvM1, InvM2;
+		UnzipPandInvM(PandInvM1, P1, InvM1);
+		UnzipPandInvM(PandInvM2, P2, InvM2);
+
+		FVector3f Direction = P1 - P2;
+		const float Distance = SafeNormalize(Direction);
+
+		const varying float Dist = Dists[i];
+		const varying float Lambda = Lambdas[i];
+		const varying float Offset = Distance - Dist;
+
+		const varying uint8 StiffnessIndex = StiffnessIndices[i];
+
+		#pragma ignore warning(perf)
+		const varying float Stiffness = StiffnessTable[StiffnessIndex];
+
+		const varying float Alpha = (uniform float)1.f / (Stiffness * Dt * Dt);
+
+		const float CombinedInvMass = InvM2 + InvM1;
+
+		const varying float DLambda = (Offset - Alpha * Lambda) / (CombinedInvMass + Alpha);
+
+		const FVector3f ComputedDelta = DLambda * Direction;
+
+		const FVector3f Delta = VectorSelect((InvM2 == FLOAT_ZERO && InvM1 == FLOAT_ZERO), FloatZeroVector, ComputedDelta);
+
+		if (InvM1 > FLOAT_ZERO)
+		{
+			VectorScatter(&PandInvM[i1], SetVector4( P1 - (InvM1 * Delta), InvM1 ));
+		}
+		if (InvM2 > FLOAT_ZERO)
+		{
+			VectorScatter(&PandInvM[i2], SetVector4( P2 + (InvM2 * Delta), InvM2 ));
+		}
+	}
+}
diff --git a/Engine/Source/Runtime/Experimental/Chaos/Public/Chaos/PBDBendingConstraints.h b/Engine/Source/Runtime/Experimental/Chaos/Public/Chaos/PBDBendingConstraints.h
--- a/Engine/Source/Runtime/Experimental/Chaos/Public/Chaos/PBDBendingConstraints.h
+++ b/Engine/Source/Runtime/Experimental/Chaos/Public/Chaos/PBDBendingConstraints.h
@@ -6,7 +6,7 @@
 namespace Chaos::Softs
 {
 
-class FPBDBendingConstraints : public FPBDBendingConstraintsBase
+class CHAOS_API FPBDBendingConstraints : public FPBDBendingConstraintsBase
 {
 	typedef FPBDBendingConstraintsBase Base;
 	using Base::Constraints;
@@ -22,7 +22,10 @@ public:
 		const FSolverReal InBucklingRatio,
 		const FSolverVec2& InBucklingStiffness,
 		bool bTrimKinematicConstraints = false)
-		:Base(InParticles, ParticleOffset, ParticleCount, MoveTemp(InConstraints), StiffnessMultipliers, BucklingStiffnessMultipliers, InStiffness, InBucklingRatio, InBucklingStiffness, bTrimKinematicConstraints) {}
+		:Base(InParticles, ParticleOffset, ParticleCount, MoveTemp(InConstraints), StiffnessMultipliers, BucklingStiffnessMultipliers, InStiffness, InBucklingRatio, InBucklingStiffness, bTrimKinematicConstraints) 
+	{
+		InitColor(InParticles);
+	}
 
 	FPBDBendingConstraints(const FSolverParticles& InParticles, TArray<TVec4<int32>>&& InConstraints, const FSolverReal InStiffness = (FSolverReal)1.)
 	    : Base(InParticles, MoveTemp(InConstraints), InStiffness) {}
@@ -30,52 +33,22 @@ public:
 
 	virtual ~FPBDBendingConstraints() override {}
 
-	void Apply(FSolverParticles& InParticles, const FSolverReal Dt) const
-	{
-		const bool StiffnessHasWeightMap = Stiffness.HasWeightMap();
-		const bool BucklingStiffnessHasWeightMap = BucklingStiffness.HasWeightMap();
-		if (!StiffnessHasWeightMap && !BucklingStiffnessHasWeightMap)
-		{
-			const FSolverReal ExpStiffnessValue = (FSolverReal)Stiffness;
-			const FSolverReal ExpBucklingValue = (FSolverReal)BucklingStiffness;
+	void Apply(FSolverParticles& InParticles, const FSolverReal Dt) const;
 
-			for (int i = 0; i < Constraints.Num(); ++i)
-			{
-				const TVec4<int32>& Constraint = Constraints[i];
-				const int32 i1 = Constraint[0];
-				const int32 i2 = Constraint[1];
-				const int32 i3 = Constraint[2];
-				const int32 i4 = Constraint[3];
-				const TStaticArray<FSolverVec3, 4> Grads = Base::GetGradients(InParticles, i);
-				const FSolverReal S = Base::GetScalingFactor(InParticles, i, Grads, ExpStiffnessValue, ExpBucklingValue);
-				InParticles.P(i1) -= S * InParticles.InvM(i1) * Grads[0];
-				InParticles.P(i2) -= S * InParticles.InvM(i2) * Grads[1];
-				InParticles.P(i3) -= S * InParticles.InvM(i3) * Grads[2];
-				InParticles.P(i4) -= S * InParticles.InvM(i4) * Grads[3];
-			}
-		}
-		else
-		{
-			const FSolverReal StiffnessNoMap = (FSolverReal)Stiffness;
-			const FSolverReal BucklingStiffnessNoMap = (FSolverReal)BucklingStiffness;
-			for (int i = 0; i < Constraints.Num(); ++i)
-			{
-				const FSolverReal ExpStiffnessValue = StiffnessHasWeightMap ? Stiffness[i] : StiffnessNoMap;
-				const FSolverReal ExpBucklingValue = BucklingStiffnessHasWeightMap ? BucklingStiffness[i] : BucklingStiffnessNoMap;
-				const TVec4<int32>& Constraint = Constraints[i];
-				const int32 i1 = Constraint[0];
-				const int32 i2 = Constraint[1];
-				const int32 i3 = Constraint[2];
-				const int32 i4 = Constraint[3];
-				const TStaticArray<FSolverVec3, 4> Grads = Base::GetGradients(InParticles, i);
-				const FSolverReal S = Base::GetScalingFactor(InParticles, i, Grads, ExpStiffnessValue, ExpBucklingValue);
-				InParticles.P(i1) -= S * InParticles.InvM(i1) * Grads[0];
-				InParticles.P(i2) -= S * InParticles.InvM(i2) * Grads[1];
-				InParticles.P(i3) -= S * InParticles.InvM(i3) * Grads[2];
-				InParticles.P(i4) -= S * InParticles.InvM(i4) * Grads[3];
-			}
-		}
-	}
+private:
+	void InitColor(const FSolverParticles& InParticles);
+	void ApplyHelper(FSolverParticles& InParticles, const FSolverReal Dt, const int32 ConstraintIndex, const FSolverReal ExpStiffnessValue, const FSolverReal ExpBucklingValue) const;
+
+	TArray<int32> ConstraintsPerColorStartIndex; // Constraints are ordered so each batch is contiguous. This is ColorNum + 1 length so it can be used as start and end.
 };
 
 }  // End namespace Chaos::Softs
+
+// Support ISPC enable/disable in non-shipping builds
+#if !INTEL_ISPC
+const bool bChaos_Bending_ISPC_Enabled = false;
+#elif UE_BUILD_SHIPPING
+const bool bChaos_Bending_ISPC_Enabled = true;
+#else
+extern CHAOS_API bool bChaos_Bending_ISPC_Enabled;
+#endif
diff --git a/Engine/Source/Runtime/Experimental/Chaos/Public/Chaos/PBDBendingConstraintsBase.h b/Engine/Source/Runtime/Experimental/Chaos/Public/Chaos/PBDBendingConstraintsBase.h
--- a/Engine/Source/Runtime/Experimental/Chaos/Public/Chaos/PBDBendingConstraintsBase.h
+++ b/Engine/Source/Runtime/Experimental/Chaos/Public/Chaos/PBDBendingConstraintsBase.h
@@ -220,7 +220,6 @@ protected:
 	FSolverReal BucklingRatio;
 	FPBDStiffness BucklingStiffness;
 
-private:
 	TArray<FSolverReal> RestAngles;
 };
 
diff --git a/Engine/Source/Runtime/Experimental/Chaos/Public/Chaos/XPBDSpringConstraints.h b/Engine/Source/Runtime/Experimental/Chaos/Public/Chaos/XPBDSpringConstraints.h
--- a/Engine/Source/Runtime/Experimental/Chaos/Public/Chaos/XPBDSpringConstraints.h
+++ b/Engine/Source/Runtime/Experimental/Chaos/Public/Chaos/XPBDSpringConstraints.h
@@ -4,7 +4,6 @@
 #include "Chaos/PBDSpringConstraintsBase.h"
 #include "ChaosStats.h"
 
-DECLARE_CYCLE_STAT(TEXT("Chaos XPBD Spring Constraint"), STAT_XPBD_Spring, STATGROUP_Chaos);
 
 namespace Chaos::Softs
 {
@@ -13,7 +12,7 @@ namespace Chaos::Softs
 static const FSolverReal XPBDSpringMinStiffness = (FSolverReal)1e-1; // Min stiffness: 1e3 N/M^2 = 1e-1 N/CM^2. Empirically defined to have similar look to XPBDSpringConstraints.
 static const FSolverReal XPBDSpringMaxStiffness = (FSolverReal)1e7;  // Max stiffness: 1e+11 N/M^2 = 1e+7 N/CM^2
 
-class FXPBDSpringConstraints final : public FPBDSpringConstraintsBase
+class CHAOS_API FXPBDSpringConstraints final : public FPBDSpringConstraintsBase
 {
 	typedef FPBDSpringConstraintsBase Base;
 	using Base::Constraints;
@@ -34,6 +33,7 @@ public:
 		: Base(Particles, ParticleOffset, ParticleCount, InConstraints, StiffnessMultipliers, InStiffness, bTrimKinematicConstraints)
 	{
 		Lambdas.Init((FSolverReal)0., Constraints.Num());
+		InitColor(Particles);
 	}
 
 	virtual ~FXPBDSpringConstraints() override {}
@@ -43,45 +43,12 @@ public:
 	// Update stiffness table, as well as the simulation stiffness exponent
 	inline void ApplyProperties(const FSolverReal Dt, const int32 NumIterations) { Stiffness.ApplyXPBDValues(XPBDSpringMinStiffness, XPBDSpringMaxStiffness); }
 
-	void Apply(FSolverParticles& Particles, const FSolverReal Dt, const int32 ConstraintIndex, const FSolverReal ExpStiffnessValue) const
-	{
-		const TVec2<int32>& Constraint = Constraints[ConstraintIndex];
-		const int32 i1 = Constraint[0];
-		const int32 i2 = Constraint[1];
-		const FSolverVec3 Delta = GetDelta(Particles, Dt, ConstraintIndex, ExpStiffnessValue);
-		if (Particles.InvM(i1) > (FSolverReal)0.)
-		{
-			Particles.P(i1) -= Particles.InvM(i1) * Delta;
-		}
-		if (Particles.InvM(i2) > (FSolverReal)0.)
-		{
-			Particles.P(i2) += Particles.InvM(i2) * Delta;
-		}
-	}
-
-	void Apply(FSolverParticles& Particles, const FSolverReal Dt) const
-	{
-		SCOPE_CYCLE_COUNTER(STAT_XPBD_Spring);
-
-		if (!Stiffness.HasWeightMap())
-		{
-			const FSolverReal ExpStiffnessValue = (FSolverReal)Stiffness;
-			for (int32 ConstraintIndex = 0; ConstraintIndex < Constraints.Num(); ++ConstraintIndex)
-			{
-				Apply(Particles, Dt, ConstraintIndex, ExpStiffnessValue);
-			}
-		}
-		else
-		{
-			for (int32 ConstraintIndex = 0; ConstraintIndex < Constraints.Num(); ++ConstraintIndex)
-			{
-				const FSolverReal ExpStiffnessValue = Stiffness[ConstraintIndex];
-				Apply(Particles, Dt, ConstraintIndex, ExpStiffnessValue);
-			}
-		}
-	}
+	void Apply(FSolverParticles& Particles, const FSolverReal Dt) const;
 
 private:
+	void InitColor(const FSolverParticles& InParticles);
+	void ApplyHelper(FSolverParticles& Particles, const FSolverReal Dt, const int32 ConstraintIndex, const FSolverReal ExpStiffnessValue) const;
+
 	FSolverVec3 GetDelta(const FSolverParticles& Particles, const FSolverReal Dt, const int32 ConstraintIndex, const FSolverReal ExpStiffnessValue) const
 	{
 		const TVec2<int32>& Constraint = Constraints[ConstraintIndex];
@@ -113,6 +80,16 @@ private:
 
 private:
 	mutable TArray<FSolverReal> Lambdas;
+	TArray<int32> ConstraintsPerColorStartIndex; // Constraints are ordered so each batch is contiguous. This is ColorNum + 1 length so it can be used as start and end.
 };
 
 }  // End namespace Chaos::Softs
+
+// Support ISPC enable/disable in non-shipping builds
+#if !INTEL_ISPC
+const bool bChaos_XPBDSpring_ISPC_Enabled = false;
+#elif UE_BUILD_SHIPPING
+const bool bChaos_XPBDSpring_ISPC_Enabled = true;
+#else
+extern CHAOS_API bool bChaos_XPBDSpring_ISPC_Enabled;
+#endif
