Fix compile error from CL 20034474 (Chaos_Spring_ParallelConstraintCount -> Chaos_XPBDSpring_ParallelConstraintCount)

diff --git a/Engine/Source/Runtime/Experimental/Chaos/Private/Chaos/XPBDSpringConstraints.cpp b/Engine/Source/Runtime/Experimental/Chaos/Private/Chaos/XPBDSpringConstraints.cpp
--- a/Engine/Source/Runtime/Experimental/Chaos/Private/Chaos/XPBDSpringConstraints.cpp
+++ b/Engine/Source/Runtime/Experimental/Chaos/Private/Chaos/XPBDSpringConstraints.cpp
@@ -33,7 +33,7 @@ void FXPBDSpringConstraints::InitColor(const FSolverParticles& Particles)
 {
 	// In dev builds we always color so we can tune the system without restarting. See Apply()
 #if UE_BUILD_SHIPPING || UE_BUILD_TEST
-	if (Constraints.Num() > Chaos_Spring_ParallelConstraintCount)
+	if (Constraints.Num() > Chaos_XPBDSpring_ParallelConstraintCount)
 #endif
 	{
 		const TArray<TArray<int32>> ConstraintsPerColor = FGraphColoring::ComputeGraphColoring(Constraints, Particles);
