Memory leak fix: Make sure that unused sleep data is cleared

diff --git a/Engine/Source/Runtime/Experimental/Chaos/Private/EventDefaults.cpp b/Engine/Source/Runtime/Experimental/Chaos/Private/EventDefaults.cpp
--- a/Engine/Source/Runtime/Experimental/Chaos/Private/EventDefaults.cpp
+++ b/Engine/Source/Runtime/Experimental/Chaos/Private/EventDefaults.cpp
@@ -505,7 +505,13 @@ namespace Chaos
 
 			NonConstSolver->Particles.GetDynamicParticles().ClearSleepData();
 
+			// We don't care about sleep data added to these
+			NonConstSolver->Particles.GetDynamicKinematicParticles().ClearSleepData();
+			NonConstSolver->Particles.GetDynamicDisabledParticles().ClearSleepData();
 
+			// Sleep data is not supported for these particles yet, clear the buffers here so that we don't leak memory
+			NonConstSolver->Particles.GetClusteredParticles().ClearSleepData();
+			NonConstSolver->Particles.GetGeometryCollectionParticles().ClearSleepData(); 
 		});
 	}
 
diff --git a/Engine/Source/Runtime/Experimental/Chaos/Public/Chaos/PBDRigidsSOAs.h b/Engine/Source/Runtime/Experimental/Chaos/Public/Chaos/PBDRigidsSOAs.h
--- a/Engine/Source/Runtime/Experimental/Chaos/Public/Chaos/PBDRigidsSOAs.h
+++ b/Engine/Source/Runtime/Experimental/Chaos/Public/Chaos/PBDRigidsSOAs.h
@@ -779,6 +779,13 @@ public:
 	const FPBDRigidParticles& GetDynamicParticles() const { return *DynamicParticles; }
 	FPBDRigidParticles& GetDynamicParticles() { return *DynamicParticles; }
 
+	const FPBDRigidParticles& GetDynamicKinematicParticles() const { return *DynamicKinematicParticles; }
+	FPBDRigidParticles& GetDynamicKinematicParticles() { return *DynamicKinematicParticles; }
+
+	// Disabled Dynamic and DynamicKinematic Particles
+	const FPBDRigidParticles& GetDynamicDisabledParticles() const { return *DynamicDisabledParticles; }
+	FPBDRigidParticles& GetDynamicDisabledParticles() { return *DynamicDisabledParticles; }
+
 	const FGeometryParticles& GetNonDisabledStaticParticles() const { return *StaticParticles; }
 	FGeometryParticles& GetNonDisabledStaticParticles() { return *StaticParticles; }
 
