Avoid dangling handle pointers during resim

diff --git a/Engine/Source/Runtime/Experimental/Chaos/Public/Chaos/PBDRigidsSOAs.h b/Engine/Source/Runtime/Experimental/Chaos/Public/Chaos/PBDRigidsSOAs.h
--- a/Engine/Source/Runtime/Experimental/Chaos/Public/Chaos/PBDRigidsSOAs.h
+++ b/Engine/Source/Runtime/Experimental/Chaos/Public/Chaos/PBDRigidsSOAs.h
@@ -314,9 +314,25 @@ public:
 
 	void DestroyParticle(FGeometryParticleHandle* Particle)
 	{
+		if (bResimulating)
+		{
+			ResimStaticParticles.Remove(Particle);
+			FKinematicGeometryParticleHandle* Kinematic = Particle->CastToKinematicParticle();
+			if (Kinematic)
+			{
+				ResimKinematicParticles.Remove(Kinematic);
+			}
+		}
+
 		auto PBDRigid = Particle->CastToRigidParticle();
 		if(PBDRigid)
 		{
+			if (bResimulating)
+			{
+				ResimDynamicParticles.Remove(PBDRigid);
+				ResimDynamicKinematicParticles.Remove(PBDRigid);
+			}
+
 			RemoveFromActiveArray(PBDRigid, /*bStillDirty=*/ false);
 			MovingKinematicsMapArray.Remove(PBDRigid->CastToKinematicParticle());
 
