Made it so geometry collections can replicate from blueprints as well as handle initial dormancy

diff --git a/Engine/Source/Runtime/Experimental/Chaos/Public/PhysicsProxy/GeometryCollectionPhysicsProxy.h b/Engine/Source/Runtime/Experimental/Chaos/Public/PhysicsProxy/GeometryCollectionPhysicsProxy.h
--- a/Engine/Source/Runtime/Experimental/Chaos/Public/PhysicsProxy/GeometryCollectionPhysicsProxy.h
+++ b/Engine/Source/Runtime/Experimental/Chaos/Public/PhysicsProxy/GeometryCollectionPhysicsProxy.h
@@ -245,6 +245,9 @@ public:
 	FProxyInterpolationData& GetInterpolationData() { return InterpolationData; }
 	const FProxyInterpolationData& GetInterpolationData() const { return InterpolationData; }
 
+	void SetInitializedForReplication() { bInitializedForReplication = true; }
+	bool GetInitializedForReplication() const { return bInitializedForReplication; }
+
 protected:
 	/**
 	* Compute damage threshold for a specific transform
@@ -308,6 +311,8 @@ private:
 	bool IsObjectLoading; // Indicate when loaded
 	bool IsObjectDeleting; // Indicatge when pending deletion
 
+	bool bInitializedForReplication = false;	//Indicates that initialization for replication has finished
+
 	TManagedArray<TUniquePtr<Chaos::FGeometryParticle>> GTParticles;
 	FCollisionFilterData SimFilter;
 	FCollisionFilterData QueryFilter;
diff --git a/Engine/Source/Runtime/Experimental/GeometryCollectionEngine/Private/GeometryCollection/GeometryCollectionActor.cpp b/Engine/Source/Runtime/Experimental/GeometryCollectionEngine/Private/GeometryCollection/GeometryCollectionActor.cpp
--- a/Engine/Source/Runtime/Experimental/GeometryCollectionEngine/Private/GeometryCollection/GeometryCollectionActor.cpp
+++ b/Engine/Source/Runtime/Experimental/GeometryCollectionEngine/Private/GeometryCollection/GeometryCollectionActor.cpp
@@ -38,6 +38,7 @@ AGeometryCollectionActor::AGeometryCollectionActor(const FObjectInitializer& Obj
 	SetActorTickEnabled(true);
 
 	bReplicates = true;
+	NetDormancy = DORM_Initial;
 }
 
 
diff --git a/Engine/Source/Runtime/Experimental/GeometryCollectionEngine/Private/GeometryCollection/GeometryCollectionComponent.cpp b/Engine/Source/Runtime/Experimental/GeometryCollectionEngine/Private/GeometryCollection/GeometryCollectionComponent.cpp
--- a/Engine/Source/Runtime/Experimental/GeometryCollectionEngine/Private/GeometryCollection/GeometryCollectionComponent.cpp
+++ b/Engine/Source/Runtime/Experimental/GeometryCollectionEngine/Private/GeometryCollection/GeometryCollectionComponent.cpp
@@ -397,10 +397,12 @@ void UGeometryCollectionComponent::GetLifetimeReplicatedProps(TArray<FLifetimePr
 {
 	Super::GetLifetimeReplicatedProps(OutLifetimeProps);
 	
+	/*
 	FDoRepLifetimeParams Params;
 	Params.bIsPushBased = true;
 	Params.RepNotifyCondition = REPNOTIFY_OnChanged;
-	DOREPLIFETIME_WITH_PARAMS_FAST(UGeometryCollectionComponent, RepData, Params);
+	DOREPLIFETIME_WITH_PARAMS_FAST(UGeometryCollectionComponent, RepData, Params);*/
+	DOREPLIFETIME(UGeometryCollectionComponent, RepData);
 }
 
 FBoxSphereBounds UGeometryCollectionComponent::CalcBounds(const FTransform& LocalToWorldIn) const
@@ -1136,7 +1138,7 @@ void UGeometryCollectionComponent::InitializeComponent()
 
 			if(CurrSolver)
 			{
-				CurrSolver->RegisterSimOneShotCallback([Prox = PhysicsProxy, ReplicationLevel = ReplicationAbandonClusterLevel, AbandonAfterLevel = bEnableAbandonAfterLevel]()
+				CurrSolver->RegisterSimOneShotCallback([Prox = PhysicsProxy, ReplicationLevel = ReplicationAbandonClusterLevel, AbandonAfterLevel = bEnableAbandonAfterLevel, this]()
 				{
 					// As we're not in control we make it so our simulated proxy cannot break clusters
 					// We have to set the strain to a high value but be below the max for the data type
@@ -1168,6 +1170,8 @@ void UGeometryCollectionComponent::InitializeComponent()
 							P->SetStrain(MaxStrain);
 						}
 					}
+
+					Prox->SetInitializedForReplication();
 				});
 			}
 		}
@@ -1322,10 +1326,12 @@ void UGeometryCollectionComponent::UpdateRepData()
 	
 	if (Owner && GetIsReplicated() && Owner->GetLocalRole() == ROLE_Authority)
 	{
+		bool bFirstUpdate = false;
 		if(ClustersToRep == nullptr)
 		{
 			//we only allocate set if needed because it's pretty big to have per components that don't replicate
 			ClustersToRep = MakeUnique<TSet<Chaos::FPBDRigidClusteredParticleHandle*>>();
+			bFirstUpdate = true;
 		}
 
 		//We need to build a snapshot of the GC
@@ -1372,7 +1378,7 @@ void UGeometryCollectionComponent::UpdateRepData()
 				//first time in here so needs a new count
 				//TODO: check root needs to replicate if abandon by level is enabled
 				ClustersToRep->Add(Root);
-				if(Root->InternalCluster() == false)
+				if(Root->InternalCluster() == false && bFirstUpdate == false)	//if bFirstUpdate it must be that these are the initial roots of the GC. These did not break off so no need to replicate
 				{
 					//a one off so record it
 					const int32 TransformGroupIdx = PhysicsProxy->GetTransformGroupIndexFromHandle(Root);
@@ -1447,6 +1453,13 @@ void UGeometryCollectionComponent::UpdateRepData()
 
 			MARK_PROPERTY_DIRTY_FROM_NAME(UGeometryCollectionComponent, RepData, this);
 			++RepData.Version;
+
+			if(Owner->NetDormancy == DORM_Initial)
+			{
+				//If net dormancy is Initial it must be for perf reasons, but since a cluster changed we need to replicate down
+				//TODO: set back to dormant when sim goes to sleep
+				Owner->SetNetDormancy(DORM_Awake);
+			}
 		}
 	}
 }
@@ -1458,7 +1471,7 @@ FAutoConsoleVariableRef CVarGeometryCollectionHardMissingUpdatesSnapThreshold(TE
 void UGeometryCollectionComponent::ProcessRepData()
 {
 	using namespace Chaos;
-	if(VersionProcessed == RepData.Version)
+	if(VersionProcessed == RepData.Version || PhysicsProxy->GetInitializedForReplication() == false)
 	{
 		return;
 	}
@@ -1950,8 +1963,15 @@ void UGeometryCollectionComponent::AsyncPhysicsTickComponent(float DeltaTime, fl
 {
 	Super::AsyncPhysicsTickComponent(DeltaTime, SimTime);
 
-	UpdateRepData();
-	ProcessRepData();
+	const ENetRole LocalRole = GetOwnerRole();
+	if (LocalRole == ROLE_Authority)
+	{
+		UpdateRepData();
+	}
+	else
+	{
+		ProcessRepData();
+	}
 }
 
 void UGeometryCollectionComponent::OnRegister()
diff --git a/Engine/Source/Runtime/Experimental/GeometryCollectionEngine/Public/GeometryCollection/GeometryCollectionComponent.h b/Engine/Source/Runtime/Experimental/GeometryCollectionEngine/Public/GeometryCollection/GeometryCollectionComponent.h
--- a/Engine/Source/Runtime/Experimental/GeometryCollectionEngine/Public/GeometryCollection/GeometryCollectionComponent.h
+++ b/Engine/Source/Runtime/Experimental/GeometryCollectionEngine/Public/GeometryCollection/GeometryCollectionComponent.h
@@ -782,14 +782,14 @@ protected:
 
 
 	/* Per-instance override to enable/disable replication for the geometry collection */
-	UPROPERTY(EditInstanceOnly, BlueprintReadOnly, Category=Network)
+	UPROPERTY(EditAnywhere, BlueprintReadOnly, Category=Network)
 	bool bEnableReplication;
 
 	/** 
 	 * Enables use of ReplicationAbandonClusterLevel to stop providing network updates to
 	 * clients when the updated particle is of a level higher then specified.
 	 */
-	UPROPERTY(EditInstanceOnly, BlueprintReadOnly, Category = Network)
+	UPROPERTY(EditAnywhere, BlueprintReadOnly, Category = Network)
 	bool bEnableAbandonAfterLevel;
 
 	/**
@@ -798,7 +798,7 @@ protected:
 	 * gameplay relevant to cut down on required bandwidth to update a collection.
 	 * @see bEnableAbandonAfterLevel
 	 */ 
-	UPROPERTY(EditInstanceOnly, BlueprintReadOnly, Category = Network)
+	UPROPERTY(EditAnywhere, BlueprintReadOnly, Category = Network)
 	int32 ReplicationAbandonClusterLevel;
 
 	UPROPERTY(Replicated)
