fix performance of GeometryCollectionProximityUtility on large examples, especially w/ regular face orientations

+ move implementation details to cpp file and expose proximity tolerance as a parameter

diff --git a/Engine/Source/Runtime/Experimental/Chaos/Private/GeometryCollection/GeometryCollectionProximityUtility.cpp b/Engine/Source/Runtime/Experimental/Chaos/Private/GeometryCollection/GeometryCollectionProximityUtility.cpp
--- a/Engine/Source/Runtime/Experimental/Chaos/Private/GeometryCollection/GeometryCollectionProximityUtility.cpp
+++ b/Engine/Source/Runtime/Experimental/Chaos/Private/GeometryCollection/GeometryCollectionProximityUtility.cpp
@@ -2,319 +2,419 @@
 
 #include "GeometryCollection/GeometryCollectionProximityUtility.h"
 
+#include "Spatial/SparseDynamicOctree3.h"
+
 #include "Async/ParallelFor.h"
 #include "GeometryCollection/GeometryCollection.h"
 #include "GeometryCollection/GeometryCollectionAlgo.h"
 
+
+#include "Spatial/PointHashGrid3.h"
+
 DEFINE_LOG_CATEGORY_STATIC(LogChaosProximity, Verbose, All);
 
-FGeometryCollectionProximityUtility::FGeometryCollectionProximityUtility(FGeometryCollection* InCollection)
-	: Collection(InCollection)
+namespace UE { namespace GeometryCollectionInternal	{
+
+static inline bool AreNormalsOpposite(const FVector3f& Normal0, const FVector3f& Normal1)
 {
-	check(Collection);
-	
-	// We quantize surface normals into 20 uniform bins on a unit sphere surface, ie an icosahedron
-	BinNormals.SetNum(NumBins);
-	
-	BinNormals[0] = FVector3f(0.171535f, -0.793715f, 0.583717f);
-	BinNormals[1] = FVector3f(0.627078f, -0.778267f, 0.034524f);
-	BinNormals[2] = FVector3f(-0.491358f, -0.810104f, 0.319894f);
-	BinNormals[3] = FVector3f(-0.445554f, -0.804788f, -0.392214f);
-	BinNormals[4] = FVector3f(0.245658f, -0.785111f, -0.568669f);
-	BinNormals[5] = FVector3f(0.984880f, -0.161432f, 0.062144f);
-	BinNormals[6] = FVector3f(0.247864f, -0.186425f, 0.950708f);
-	BinNormals[7] = FVector3f(-0.824669f, -0.212942f, 0.523975f);
-	BinNormals[8] = FVector3f(-0.750546f, -0.204339f, -0.628411f);
-	BinNormals[9] = FVector3f(0.367791f, -0.172505f, -0.913787f);
-	BinNormals[10] = -BinNormals[0];
-	BinNormals[11] = -BinNormals[1];
-	BinNormals[12] = -BinNormals[2];
-	BinNormals[13] = -BinNormals[3];
-	BinNormals[14] = -BinNormals[4];
-	BinNormals[15] = -BinNormals[5];
-	BinNormals[16] = -BinNormals[6];
-	BinNormals[17] = -BinNormals[7];
-	BinNormals[18] = -BinNormals[8];
-	BinNormals[19] = -BinNormals[9];
-	
-	NumFaces = Collection->NumElements(FGeometryCollection::FacesGroup);
-	PrepFaceData();
+	return FVector3f::DotProduct(Normal0, Normal1) < (-1.0f + UE_KINDA_SMALL_NUMBER);
 }
 
-void FGeometryCollectionProximityUtility::UpdateProximity()
+static inline bool TriOutsideEdge(const TStaticArray<FVector2f, 3>& Points, const FVector2f& B, const FVector2f& C, float Normal)
 {
-	/*
-	*	For each face, we generate the surface normal and put the normal into a set of bins that discretize the surface
-	*	of the unit sphere. We know that any face can only contact other faces with opposite surface normal, which 
-	*	means that we need only compare a face with the bin of faces on the antipode of the sphere. We determine contact first
-	*	by testing if the normal is opposite, then if the faces are coplanar, then if the 2D triangles overlap.
-	*	Once we have face pairs, we can extend this to geometry indices and form the proximity structure
-	*	which is a map from geometry to geometry, one-to-many.
-	*/
-
-	BinFaces();
-	FindContactingFaces();
-	ExtendFaceProximityToGeometry();
+	const FVector2f CBPerp = UE::Geometry::PerpCW(C - B) * Normal;
+	const FVector2f P0B = Points[0] - B;
+	const FVector2f P1B = Points[1] - B;
+	const FVector2f P2B = Points[2] - B;
+	// Note: If we change this comparison to use a tolerance instead of 0, need to normalize CBPerp and also filter out degenerate tris
+	bool bAnyInside =
+		CBPerp.Dot(P0B) < 0 ||
+		CBPerp.Dot(P1B) < 0 ||
+		CBPerp.Dot(P2B) < 0;
+	return !bAnyInside;
 }
 
+static inline bool TrianglesIntersect(const TStaticArray<FVector2f,3>& T0, const TStaticArray<FVector2f,3>& T1)
+{
+	// Test if one of the triangles has a side with all of the other triangle's points on the outside.
 
-void FGeometryCollectionProximityUtility::PrepFaceData()
-{	
-	GenerateFaceToGeometry();
-	TransformVertices();
-	GenerateSurfaceNormals();
+	// Orientation of T0 and T1, as sign of area (computed as sign of the Z component of a cross product of the triangle edges)
+	float Normal0 = FMathf::Sign(T0[1].X - T0[0].X) * (T0[2].Y - T0[0].Y) - (T0[1].Y - T0[0].Y) * (T0[2].X - T0[0].X);
+	float Normal1 = FMathf::Sign(T1[1].X - T1[0].X) * (T1[2].Y - T1[0].Y) - (T1[1].Y - T1[0].Y) * (T1[2].X - T1[0].X);
+	
+	// Triangles overlap if there is no edge that the other tri is completely on the 'outside' side of
+	return !(
+		TriOutsideEdge(T1, T0[0], T0[1], Normal0) ||
+		TriOutsideEdge(T1, T0[1], T0[2], Normal0) ||
+		TriOutsideEdge(T1, T0[2], T0[0], Normal0) ||
+		TriOutsideEdge(T0, T1[0], T1[1], Normal1) ||
+		TriOutsideEdge(T0, T1[1], T1[2], Normal1) ||
+		TriOutsideEdge(T0, T1[2], T1[0], Normal1)
+		);
 }
 
-void FGeometryCollectionProximityUtility::BinFaces()
+// We'll bin triangles normals based on 20 dodecahedron directions, to filter which ones to test for overlap
+struct FBinNormals
 {
-	// We estimate that each bin will contain approximately 1/20th of the faces -- we double that to provide a decent buffer
-	Bins.SetNum(NumBins);
-	for (int32 BinIdx = 0; BinIdx < NumBins; ++BinIdx)
+	static constexpr int NumBins = 20;
+
+	TArray<FVector3f, TFixedAllocator<NumBins>> Bins;
+
+	FBinNormals()
 	{
-		Bins[BinIdx].Reserve(2 * (NumFaces/NumBins));
+		// We quantize surface normals into 20 uniform bins on a unit sphere surface, ie an icosahedron
+		Bins.SetNum(NumBins);
+
+		Bins[0] = FVector3f(0.171535f, -0.793715f, 0.583717f);
+		Bins[1] = FVector3f(0.627078f, -0.778267f, 0.034524f);
+		Bins[2] = FVector3f(0.491358f, 0.810104f, -0.319894f);
+		Bins[3] = FVector3f(0.445554f, 0.804788f, 0.392214f);
+		Bins[4] = FVector3f(0.245658f, -0.785111f, -0.568669f);
+		Bins[5] = FVector3f(0.984880f, -0.161432f, 0.062144f);
+		Bins[6] = FVector3f(0.247864f, -0.186425f, 0.950708f);
+		Bins[7] = FVector3f(0.824669f, 0.212942f, -0.523975f);
+		Bins[8] = FVector3f(0.750546f, 0.204339f, 0.628411f);
+		Bins[9] = FVector3f(0.367791f, -0.172505f, -0.913787f);
+		Bins[10] = -Bins[0];
+		Bins[11] = -Bins[1];
+		Bins[12] = -Bins[2];
+		Bins[13] = -Bins[3];
+		Bins[14] = -Bins[4];
+		Bins[15] = -Bins[5];
+		Bins[16] = -Bins[6];
+		Bins[17] = -Bins[7];
+		Bins[18] = -Bins[8];
+		Bins[19] = -Bins[9];
 	}
 
-	for (int32 FaceIdx = 0; FaceIdx < NumFaces; ++FaceIdx)
+	// Note: a normal could point at an edge or corner of the dodecahedron, resulting in an ambiguous case that misses matching normals ...
+	// Currently we just hope that the other proximity tests (vertex proximity and other faces) will find the proximity in these cases.
+	// TODO: Consider testing multiple bins in ambiguous cases.
+	int32 FindBestBin(const FVector3f& SurfaceNormal) const
 	{
-		if (SurfaceNormals[FaceIdx].IsUnit())
-		{ 
-			int32 BestFitBin = FindBestBin(SurfaceNormals[FaceIdx]);
-			Bins[BestFitBin].Add(FaceIdx);
+		// We select the bin with the highest alignment with the surface normal
+		float BestAlignment = -1.0;
+		int32 BestBin = INDEX_NONE;
+
+		for (int32 BinIdx = 0; BinIdx < NumBins; ++BinIdx)
+		{
+			float Alignment = FVector3f::DotProduct(SurfaceNormal, Bins[BinIdx]);
+			if (Alignment > BestAlignment)
+			{
+				BestAlignment = Alignment;
+				BestBin = BinIdx;
+			}
 		}
+
+		return BestBin;
 	}
-}
+};
 
-int32 FGeometryCollectionProximityUtility::FindBestBin(const FVector3f& SurfaceNormal) const
+// Per-Geometry pre-computed spatial data
+struct FPerGeometrySpatial
 {
-	// We select the bin with the highest alignment with the surface normal
-	float BestAlignment = -1.0;
-	int32 BestBin = INDEX_NONE;
+	TArray<TArray<int32>, TFixedAllocator<FBinNormals::NumBins>> Bins;
+	UE::Math::TBox<float> Bounds; // bounds in a shared space
 
-	for (int32 BinIdx = 0; BinIdx < NumBins; ++BinIdx)
+	void InitBins(const FGeometryCollection* Collection, const FBinNormals& Binner, int32 GeoIdx, const TArray<FVector3f>& SurfaceNormals)
 	{
-		float Alignment = FVector3f::DotProduct(SurfaceNormal, BinNormals[BinIdx]);
-		if (Alignment > BestAlignment)
+		Bins.SetNum(FBinNormals::NumBins);
+		int32 FaceStart = Collection->FaceStart[GeoIdx];
+		int32 FaceEnd = FaceStart + Collection->FaceCount[GeoIdx];
+		for (int32 FaceIdx = FaceStart; FaceIdx < FaceEnd; ++FaceIdx)
 		{
-			BestAlignment = Alignment;
-			BestBin = BinIdx;
+			int32 BestBin = Binner.FindBestBin(SurfaceNormals[FaceIdx]);
+			Bins[BestBin].Add(FaceIdx);
 		}
 	}
 
-	return BestBin;
-}
+	// Note: Only contains mappings to geometry w/ *higher* indices
+	TMap<int32, bool> CandidateContacts; // bool indicates if the contact has been confirmed
+};
 
-void FGeometryCollectionProximityUtility::FindContactingFaces()
+// Overall geometry-collection spatial data
+struct FGeometryCollectionProximitySpatial
 {
-	FaceContacts.SetNum(NumFaces);
-	
-	// The bin hemispheres are symmetrical so we only need to iterate half of them.
-	// We are comparing antipodal bins.
-	ParallelFor((NumBins / 2), [this](int32 BinIdx)
-		{ 
-			FindContactPairs(BinIdx, (NumBins/2) + BinIdx); 
-		});
-}
+	TArray<FVector3f> TransformedVertices;
+	TArray<FVector3f> SurfaceNormals;
+	TArray<FPerGeometrySpatial> GeoInfo;
+	TArray<TSet<int32>> KnownProximity;
+	UE::Math::TBox<float> OverallBounds;
 
-void FGeometryCollectionProximityUtility::FindContactPairs(int32 Zenith, int32 Nadir)
-{
-	for (int32 ZenithIdx : Bins[Zenith])
+	FGeometryCollectionProximitySpatial(const FGeometryCollection* Collection, float ProximityTolerance = KINDA_SMALL_NUMBER)
 	{
-		for (int32 NadirIdx : Bins[Nadir])
-		{
-			// Are the faces parallel?
-			if (AreNormalsOpposite(SurfaceNormals[ZenithIdx], SurfaceNormals[NadirIdx]))
+		TransformVertices(Collection);
+		GenerateSurfaceNormals(Collection);
+		InitProximityFromVertices(Collection, ProximityTolerance);
+		InitGeoNormalBins(Collection);
+		InitCandidateContacts(Collection, ProximityTolerance);
+		ComputeCoplanarContacts(Collection, ProximityTolerance);
+	}
+
+	void ComputeCoplanarContacts(const FGeometryCollection* Collection, float ProximityTolerance)
+	{
+		ParallelFor(GeoInfo.Num(), [this, &Collection, ProximityTolerance](int32 GeoIdx)
 			{
-				// Are faces co-planar?
-				if (AreFacesCoPlanar(ZenithIdx, NadirIdx))
+				for (TPair<int32, bool>& PossibleContact : GeoInfo[GeoIdx].CandidateContacts)
 				{
-					// Do triangles overlap?
-					if (DoFacesOverlap(ZenithIdx, NadirIdx))
+					if (PossibleContact.Value || KnownProximity[GeoIdx].Contains(PossibleContact.Key))
 					{
-						// Add contacts to both faces
-						FaceContacts[ZenithIdx].Add(NadirIdx);
-						FaceContacts[NadirIdx].Add(ZenithIdx);
+						continue;
 					}
+
+					int32 OtherGeoIdx = PossibleContact.Key;
+					for (int32 BinIdx = 0; BinIdx < FBinNormals::NumBins && !PossibleContact.Value; ++BinIdx)
+					{
+						int32 OtherBinIdx = (BinIdx + (FBinNormals::NumBins / 2)) % FBinNormals::NumBins;
+						for (int32 FaceIdx : GeoInfo[GeoIdx].Bins[BinIdx])
+						{
+							// Skip if the face bounds don't overlap the geo bounds
+							FIntVector3 Face = Collection->Indices[FaceIdx];
+							UE::Math::TBox<float> FaceBox(EForceInit::ForceInit);
+							FaceBox += TransformedVertices[Face.X];
+							FaceBox += TransformedVertices[Face.Y];
+							FaceBox += TransformedVertices[Face.Z];
+							if (!GeoInfo[OtherGeoIdx].Bounds.Intersect(FaceBox.ExpandBy(ProximityTolerance)))
+							{
+								continue;
+							}
+							for (int32 OtherFaceIdx : GeoInfo[OtherGeoIdx].Bins[OtherBinIdx])
+							{
+								if (AreNormalsOpposite(SurfaceNormals[FaceIdx], SurfaceNormals[OtherFaceIdx]))
+								{
+									if (AreFacesCoPlanar(Collection, FaceIdx, OtherFaceIdx, ProximityTolerance))
+									{
+										if (DoFacesOverlap(Collection, FaceIdx, OtherFaceIdx))
+										{
+											PossibleContact.Value = true; // verified contact
+											break; // don't need to consider the rest of the geometry
+										}
+									}
+								}
+							}
+							if (PossibleContact.Value)
+							{
+								break;
+							}
+						}
+					}
+				}
+			});
+	}
+
+	// Note: This function will destroy/move over the computed proximity data when building it in the output collection
+	void MoveProximityToCollection(FGeometryCollection* Collection)
+	{
+		if (!Collection->HasAttribute("Proximity", FGeometryCollection::GeometryGroup))
+		{
+			const FManagedArrayCollection::FConstructionParameters GeometryDependency(FGeometryCollection::GeometryGroup);
+			Collection->AddAttribute<TSet<int32>>("Proximity", FGeometryCollection::GeometryGroup, GeometryDependency);
+		}
+
+		TManagedArray<TSet<int32>>& Proximity = Collection->GetAttribute<TSet<int32>>("Proximity", FGeometryCollection::GeometryGroup);
+
+		for (int32 GeoIdx = 0; GeoIdx < KnownProximity.Num(); ++GeoIdx)
+		{
+			Proximity[GeoIdx] = MoveTemp(KnownProximity[GeoIdx]);
+		}
+		for (int32 GeoIdx = 0; GeoIdx < KnownProximity.Num(); ++GeoIdx)
+		{
+			for (const TPair<int32, bool>& PossibleContact : GeoInfo[GeoIdx].CandidateContacts)
+			{
+				if (PossibleContact.Value)
+				{
+					Proximity[GeoIdx].Add(PossibleContact.Key);
+					Proximity[PossibleContact.Key].Add(GeoIdx);
 				}
 			}
 		}
+
 	}
-}
 
-bool FGeometryCollectionProximityUtility::AreNormalsOpposite(const FVector3f& Normal0, const FVector3f& Normal1) const
-{
-	return FVector3f::DotProduct(Normal0, Normal1) < (-1.0f + UE_KINDA_SMALL_NUMBER);
-}
+	bool AreFacesCoPlanar(const FGeometryCollection* Collection, int32 Idx0, int32 Idx1, float ProximityTolerance) const
+	{
+		// Assumes that faces have already been determined to be parallel.
 
-bool FGeometryCollectionProximityUtility::AreFacesCoPlanar(int32 Idx0, int32 Idx1) const
-{
-	// Assumes that faces have already been determined to be parallel.
+		const TManagedArray<FIntVector>& Indices = Collection->Indices;
 
-	const TManagedArray<FIntVector>& Indices = Collection->Indices;
+		FVector3f SamplePoint = TransformedVertices[Indices[Idx0].X];
+		FVector3f PlaneOrigin = TransformedVertices[Indices[Idx1].X];
+		FVector3f PlaneNormal = SurfaceNormals[Idx1];
 
-	FVector3f SamplePoint = TransformedVertices[Indices[Idx0].X];
-	FVector3f PlaneOrigin = TransformedVertices[Indices[Idx1].X];
-	FVector3f PlaneNormal = SurfaceNormals[Idx1];
+		return FMath::Abs(FVector3f::DotProduct((SamplePoint - PlaneOrigin), PlaneNormal)) < ProximityTolerance;
+	}
 
-	return FMath::Abs(FVector3f::DotProduct((SamplePoint - PlaneOrigin), PlaneNormal)) < UE_KINDA_SMALL_NUMBER;
-}
+	bool DoFacesOverlap(const FGeometryCollection* Collection, int32 Idx0, int32 Idx1) const
+	{
+		// Assumes that faces have already been determined to be coplanar
 
-bool FGeometryCollectionProximityUtility::DoFacesOverlap(int32 Idx0, int32 Idx1) const
-{
-	// Assumes that the faces are coplanar
+		const TManagedArray<FIntVector>& Indices = Collection->Indices;
 
-	const TManagedArray<FIntVector>& Indices = Collection->Indices;
-	
-	// Project the first triangle into its normal plane
-	FVector3f Basis0 = TransformedVertices[Indices[Idx0].Y] - TransformedVertices[Indices[Idx0].X];
-	Basis0.Normalize();
-	FVector3f Basis1 = FVector3f::CrossProduct(SurfaceNormals[Idx0], Basis0);
-	Basis1.Normalize();
+		// Project the first triangle into its normal plane
+		FVector3f Basis0 = TransformedVertices[Indices[Idx0].Y] - TransformedVertices[Indices[Idx0].X];
+		Basis0.Normalize();
+		FVector3f Basis1 = FVector3f::CrossProduct(SurfaceNormals[Idx0], Basis0);
+		Basis1.Normalize();
 
-	FVector3f Origin = TransformedVertices[Indices[Idx0].X];
-	
-	TStaticArray<FVector2f,3> T0; 
-	// T0[0] is the origin of the system
-	T0[0] = FVector2f(0.f, 0.f);
-
-	T0[1] = FVector2f(FVector3f::DotProduct(TransformedVertices[Indices[Idx0].Y] - Origin, Basis0), FVector3f::DotProduct(TransformedVertices[Indices[Idx0].Y] - Origin, Basis1));
-	T0[2] = FVector2f(FVector3f::DotProduct(TransformedVertices[Indices[Idx0].Z] - Origin, Basis0), FVector3f::DotProduct(TransformedVertices[Indices[Idx0].Z] - Origin, Basis1));
-
-	// Project the second triangle into these coordinates. We reverse the winding order to flip the normal.
-	FVector3f Point0 = TransformedVertices[Indices[Idx1].Z] - Origin;
-	FVector3f Point1 = TransformedVertices[Indices[Idx1].Y] - Origin;
-	FVector3f Point2 = TransformedVertices[Indices[Idx1].X] - Origin;
-	TStaticArray<FVector2f,3> T1;
-	T1[0] = FVector2f(FVector3f::DotProduct(Point0, Basis0), FVector3f::DotProduct(Point0, Basis1));
-	T1[1] = FVector2f(FVector3f::DotProduct(Point1, Basis0), FVector3f::DotProduct(Point1, Basis1));
-	T1[2] = FVector2f(FVector3f::DotProduct(Point2, Basis0), FVector3f::DotProduct(Point2, Basis1));
-
-	return IdenticalTriangles(T0, T1) || TrianglesIntersect(T0, T1);
-}
+		FVector3f Origin = TransformedVertices[Indices[Idx0].X];
 
-bool FGeometryCollectionProximityUtility::IdenticalTriangles(const TStaticArray<FVector2f, 3>& T0, const TStaticArray<FVector2f, 3>& T1)
-{
-	int32 FirstMatch = 0;
-	for (; FirstMatch < 3; ++FirstMatch)
+		TStaticArray<FVector2f, 3> T0;
+		// T0[0] is the origin of the system
+		T0[0] = FVector2f(0.f, 0.f);
+
+		T0[1] = FVector2f(FVector3f::DotProduct(TransformedVertices[Indices[Idx0].Y] - Origin, Basis0), FVector3f::DotProduct(TransformedVertices[Indices[Idx0].Y] - Origin, Basis1));
+		T0[2] = FVector2f(FVector3f::DotProduct(TransformedVertices[Indices[Idx0].Z] - Origin, Basis0), FVector3f::DotProduct(TransformedVertices[Indices[Idx0].Z] - Origin, Basis1));
+
+		// Project the second triangle into these coordinates. We reverse the winding order to flip the normal.
+		FVector3f Point0 = TransformedVertices[Indices[Idx1].Z] - Origin;
+		FVector3f Point1 = TransformedVertices[Indices[Idx1].Y] - Origin;
+		FVector3f Point2 = TransformedVertices[Indices[Idx1].X] - Origin;
+		TStaticArray<FVector2f, 3> T1;
+		T1[0] = FVector2f(FVector3f::DotProduct(Point0, Basis0), FVector3f::DotProduct(Point0, Basis1));
+		T1[1] = FVector2f(FVector3f::DotProduct(Point1, Basis0), FVector3f::DotProduct(Point1, Basis1));
+		T1[2] = FVector2f(FVector3f::DotProduct(Point2, Basis0), FVector3f::DotProduct(Point2, Basis1));
+
+		// Note: If this was the only proximity test, we would also need to check for identical triangles here,
+		// but we don't need to in practice because we cover that case w/ the proximity-from-vertices that we compute first
+		return TrianglesIntersect(T0, T1);
+	}
+
+	void InitCandidateContacts(const FGeometryCollection* Collection, float ProximityTolerance)
 	{
-		if (T0[0] == T1[FirstMatch])
+		UE::Geometry::FSparseDynamicOctree3 GeoOctree;
+		GeoOctree.RootDimension = OverallBounds.GetExtent().GetAbsMax();
+		FVector3d Center = (FVector3d)OverallBounds.GetCenter();
+		TArray<int32> GeoIndices;
+		for (int32 GeoIdx = 0; GeoIdx < GeoInfo.Num(); ++GeoIdx)
 		{
-			break;
+			UE::Math::TBox<float> ExpandedBounds = GeoInfo[GeoIdx].Bounds.ExpandBy(ProximityTolerance);
+
+			// Center the boxes to work better with the sparse dynamic octree
+			UE::Geometry::FAxisAlignedBox3d CenteredBox3d((FVector3d)GeoInfo[GeoIdx].Bounds.Min, (FVector3d)GeoInfo[GeoIdx].Bounds.Max);
+			CenteredBox3d.Min -= Center;
+			CenteredBox3d.Max -= Center;
+			if (GeoIdx > 0)
+			{
+				GeoIndices.Reset();
+				UE::Geometry::FAxisAlignedBox3d CenteredBox3dExpanded = CenteredBox3d;
+				CenteredBox3dExpanded.Expand(ProximityTolerance);
+				GeoOctree.RangeQuery(CenteredBox3dExpanded, GeoIndices);
+				for (int32 CandidateIdx : GeoIndices)
+				{
+					if (!KnownProximity[CandidateIdx].Contains(GeoIdx) &&
+						// FSparseDynamicOctree3 doesn't doesn't filter for actual bounding box overlap, so we need to do so here
+						ExpandedBounds.Intersect(GeoInfo[CandidateIdx].Bounds))
+					{
+						// Note: Only add the lower idx -> higher idx mapping
+						GeoInfo[CandidateIdx].CandidateContacts.Add(GeoIdx, false);
+					}
+				}
+			}
+			GeoOctree.InsertObject(GeoIdx, CenteredBox3d);
 		}
 	}
 
-	if (FirstMatch == 3) // No match
+	void InitGeoNormalBins(const FGeometryCollection* Collection)
 	{
-		return false;
+		FBinNormals Binner;
+		ParallelFor(GeoInfo.Num(), [this, &Collection, &Binner](int32 GeoIdx)
+			{
+				GeoInfo[GeoIdx].InitBins(Collection, Binner, GeoIdx, SurfaceNormals);
+			});
 	}
 
-	if (T0[1] != T1[(FirstMatch + 1) % 3])
+	void InitProximityFromVertices(const FGeometryCollection* Collection, float ProximityTolerance)
 	{
-		return false;
+		UE::Geometry::TPointHashGrid3f<int32> VertHash(ProximityTolerance * 3, -1);
+		TArray<int32> NearPts;
+		int32 NumGeo = Collection->NumElements(FGeometryCollection::GeometryGroup);
+		KnownProximity.SetNum(NumGeo);
+		GeoInfo.SetNum(NumGeo);
+		for (int32 GeoIdx = 0; GeoIdx < NumGeo; ++GeoIdx)
+		{
+			GeoInfo[GeoIdx].Bounds = UE::Math::TBox<float>(EForceInit::ForceInit);
+			int32 VertStart = Collection->VertexStart[GeoIdx];
+			int32 VertEnd = VertStart + Collection->VertexCount[GeoIdx];
+			int32 TransformIdx = Collection->TransformIndex[GeoIdx];
+			if (!Collection->IsRigid(TransformIdx))
+			{
+				continue;
+			}
+			for (int32 VertIdx = VertStart; VertIdx < VertEnd; VertIdx++)
+			{
+				const FVector3f& Vertex = TransformedVertices[VertIdx];
+				GeoInfo[GeoIdx].Bounds += Vertex;
+				NearPts.Reset();
+				VertHash.FindPointsInBall(Vertex, ProximityTolerance, [&Vertex, this](const int32& Other) -> float
+					{
+						return FVector3f::DistSquared(TransformedVertices[Other], Vertex);
+					}, NearPts);
+				for (int32 NearPtIdx : NearPts)
+				{
+					int32 NearTransformIdx = Collection->BoneMap[NearPtIdx];
+					if (NearTransformIdx != TransformIdx)
+					{
+						int32 NearGeoIdx = Collection->TransformToGeometryIndex[NearTransformIdx];
+						KnownProximity[NearGeoIdx].Add(GeoIdx);
+						KnownProximity[GeoIdx].Add(NearGeoIdx);
+					}
+				}
+				VertHash.InsertPointUnsafe(VertIdx, Vertex);
+			}
+			OverallBounds += GeoInfo[GeoIdx].Bounds;
+		}
 	}
 
-	if (T0[2] != T1[(FirstMatch + 2) % 3])
+	void TransformVertices(const FGeometryCollection* Collection)
 	{
-		return false;
-	}
-
-	return true;
-}
-
-bool FGeometryCollectionProximityUtility::TrianglesIntersect(const TStaticArray<FVector2f,3>& T0, const TStaticArray<FVector2f,3>& T1)
-{
-	// Test if one of the triangles has a side with all of the other triangle's points on the outside.
-	float Normal0 = (T0[1].X - T0[0].X) * (T0[2].Y - T0[0].Y) - (T0[1].Y - T0[0].Y) * (T0[2].X - T0[0].X);
-	float Normal1 = (T1[1].X - T1[0].X) * (T1[2].Y - T1[0].Y) - (T1[1].Y - T1[0].Y) * (T1[2].X - T1[0].X);
-
-	return !(Cross(T1, T0[0], T0[1], Normal0) ||
-		Cross(T1, T0[1], T0[2], Normal0) ||
-		Cross(T1, T0[2], T0[0], Normal0) ||
-		Cross(T0, T1[0], T1[1], Normal1) ||
-		Cross(T0, T1[1], T1[2], Normal1) ||
-		Cross(T0, T1[1], T1[0], Normal1));
-}
+		TransformedVertices.SetNum(Collection->NumElements(FGeometryCollection::VerticesGroup));
 
-bool FGeometryCollectionProximityUtility::Cross(const TStaticArray<FVector2f,3>& Points, const FVector2f& B, const FVector2f& C, float Normal)
-{
-	float CyBy = C.Y - B.Y;
-	float CxBx = C.X - B.X;
-	const FVector2f& Pa = Points[0];
-	const FVector2f& Pb = Points[1];
-	const FVector2f& Pc = Points[2];
+		TArray<FTransform> GlobalTransformArray;
+		GeometryCollectionAlgo::GlobalMatrices(Collection->Transform, Collection->Parent, GlobalTransformArray);
 
-	return !(
-		(((Pa.X - B.X) * CyBy - (Pa.Y - B.Y) * CxBx) * Normal < 0) ||
-		(((Pb.X - B.X) * CyBy - (Pb.Y - B.Y) * CxBx) * Normal < 0) ||
-		(((Pc.X - B.X) * CyBy - (Pc.Y - B.Y) * CxBx) * Normal < 0));
-}
-
-void FGeometryCollectionProximityUtility::GenerateSurfaceNormals()
-{
-	// Generate surface normal for each face
-	SurfaceNormals.SetNum(NumFaces);
-	ParallelFor(NumFaces, [this](int32 FaceIdx)
-		{
-			const TManagedArray<FIntVector>& Indices = Collection->Indices;
+		ParallelFor(Collection->NumElements(FGeometryCollection::VerticesGroup), [this, &Collection, &GlobalTransformArray](int32 VertIdx)
+			{
+				const TManagedArray<int32>& BoneMap = Collection->BoneMap;
+				const TManagedArray<FVector3f>& Vertex = Collection->Vertex;
 
-			FVector3f Edge0 = (TransformedVertices[Indices[FaceIdx].X] - TransformedVertices[Indices[FaceIdx].Y]);
-			FVector3f Edge1 = (TransformedVertices[Indices[FaceIdx].Z] - TransformedVertices[Indices[FaceIdx].Y]);
-			SurfaceNormals[FaceIdx] = FVector3f::CrossProduct(Edge0, Edge1);
-			SurfaceNormals[FaceIdx].Normalize();
-		});
-}
+				const FTransform& GlobalTransform = GlobalTransformArray[BoneMap[VertIdx]];
+				TransformedVertices[VertIdx] = (FVector3f)GlobalTransform.TransformPosition(FVector(Vertex[VertIdx]));
+			});
 
-void FGeometryCollectionProximityUtility::GenerateFaceToGeometry()
-{
-	// Create a map from face back to geometry
-	const TManagedArray<int32>& FaceStart = Collection->FaceStart;
-	const TManagedArray<int32>& FaceCount = Collection->FaceCount;
+	}
 
-	FaceToGeometry.Init(INDEX_NONE, NumFaces);
-	for (int32 GeometryIdx = 0; GeometryIdx < Collection->NumElements(FGeometryCollection::GeometryGroup); ++GeometryIdx)
+	void GenerateSurfaceNormals(const FGeometryCollection* Collection)
 	{
-		for (int32 FaceOffset = 0; FaceOffset < FaceCount[GeometryIdx]; ++FaceOffset)
-		{
-			int32 FaceIdx = FaceStart[GeometryIdx] + FaceOffset;
-			FaceToGeometry[FaceIdx] = GeometryIdx;
-		}
+		const int32 NumFaces = Collection->NumElements(FGeometryCollection::FacesGroup);
+		// Generate surface normal for each face
+		SurfaceNormals.SetNum(NumFaces);
+		ParallelFor(NumFaces, [this, &Collection](int32 FaceIdx)
+			{
+				const TManagedArray<FIntVector>& Indices = Collection->Indices;
+
+				FVector3f Edge0 = (TransformedVertices[Indices[FaceIdx].X] - TransformedVertices[Indices[FaceIdx].Y]);
+				FVector3f Edge1 = (TransformedVertices[Indices[FaceIdx].Z] - TransformedVertices[Indices[FaceIdx].Y]);
+				SurfaceNormals[FaceIdx] = FVector3f::CrossProduct(Edge0, Edge1);
+				SurfaceNormals[FaceIdx].Normalize();
+			});
 	}
-}
 
-void FGeometryCollectionProximityUtility::TransformVertices()
-{
-	TransformedVertices.SetNum(Collection->NumElements(FGeometryCollection::VerticesGroup));
+};
 
-	TArray<FTransform> GlobalTransformArray;
-	GeometryCollectionAlgo::GlobalMatrices(Collection->Transform, Collection->Parent, GlobalTransformArray);
-	
-	ParallelFor(Collection->NumElements(FGeometryCollection::VerticesGroup), [this, &GlobalTransformArray](int32 VertIdx)
-		{
-			const TManagedArray<int32>& BoneMap = Collection->BoneMap;
-			const TManagedArray<FVector3f>& Vertex = Collection->Vertex;
 
-			const FTransform& GlobalTransform = GlobalTransformArray[BoneMap[VertIdx]];
-			TransformedVertices[VertIdx] = (FVector3f)GlobalTransform.TransformPosition(FVector(Vertex[VertIdx]));
-		});
 
-}
+}} // namespace UE::GeometryCollectionInternal
 
-void FGeometryCollectionProximityUtility::ExtendFaceProximityToGeometry()
+FGeometryCollectionProximityUtility::FGeometryCollectionProximityUtility(FGeometryCollection* InCollection)
+	: Collection(InCollection)
 {
-	if (!Collection->HasAttribute("Proximity", FGeometryCollection::GeometryGroup))
-	{
-		const FManagedArrayCollection::FConstructionParameters GeometryDependency(FGeometryCollection::GeometryGroup);
-		Collection->AddAttribute<TSet<int32>>("Proximity", FGeometryCollection::GeometryGroup, GeometryDependency);
-	}
+	check(Collection);
+}
 
-	TManagedArray<TSet<int32>>& Proximity = Collection->GetAttribute<TSet<int32>>("Proximity", FGeometryCollection::GeometryGroup);
+void FGeometryCollectionProximityUtility::UpdateProximity()
+{
+	using namespace UE::GeometryCollectionInternal;
 
-	for (int32 FaceIdx = 0; FaceIdx < NumFaces; ++FaceIdx)
-	{
-		TSet<int32>& CurrentContactList = Proximity[FaceToGeometry[FaceIdx]];
-		for (int32 ContactFaceIdx : FaceContacts[FaceIdx])
-		{
-			CurrentContactList.Add(FaceToGeometry[ContactFaceIdx]);
-		}
-	}
+	FGeometryCollectionProximitySpatial Spatial(Collection, ProximityThreshold);
+	Spatial.MoveProximityToCollection(Collection);
 }
-
diff --git a/Engine/Source/Runtime/Experimental/Chaos/Public/GeometryCollection/GeometryCollectionProximityUtility.h b/Engine/Source/Runtime/Experimental/Chaos/Public/GeometryCollection/GeometryCollectionProximityUtility.h
--- a/Engine/Source/Runtime/Experimental/Chaos/Public/GeometryCollection/GeometryCollectionProximityUtility.h
+++ b/Engine/Source/Runtime/Experimental/Chaos/Public/GeometryCollection/GeometryCollectionProximityUtility.h
@@ -2,13 +2,8 @@
 
 #pragma once
 
-#include "CoreMinimal.h"
-#include "Containers/StaticArray.h"
-
 class FGeometryCollection;
 
-
-
 class CHAOS_API FGeometryCollectionProximityUtility
 {
 public:
@@ -16,36 +11,10 @@ public:
 
 	void UpdateProximity();
 
-private:
-	void PrepFaceData();
-	void BinFaces();
-	void FindContactingFaces();
-	void ExtendFaceProximityToGeometry();
-
-	void TransformVertices();
-	void GenerateSurfaceNormals();
-	void GenerateFaceToGeometry();
-	int32 FindBestBin(const FVector3f& SurfaceNormal) const;
-	void FindContactPairs(int32 Zenith, int32 Nadir);
-	bool AreNormalsOpposite(const FVector3f& Normal0, const FVector3f& Normal1) const;
-	bool AreFacesCoPlanar(int32 Idx0, int32 Idx1) const;
-	bool DoFacesOverlap(int32 Idx0, int32 Idx1) const;
-	static bool IdenticalTriangles(const TStaticArray<FVector2f, 3>& T0, const TStaticArray<FVector2f, 3>& T1);
-	static bool TrianglesIntersect(const TStaticArray<FVector2f,3>& T0, const TStaticArray<FVector2f,3>& T1);
-	static bool Cross(const TStaticArray<FVector2f,3>& Points, const FVector2f& B, const FVector2f& C, float Normal);
+	// maximum separation to still count as 'contact'
+	float ProximityThreshold = 0.01f;
 
 private:
 	FGeometryCollection* Collection;
-
-	int32 NumFaces;
-	TArray<FVector3f> SurfaceNormals;
-	TArray<int32> FaceToGeometry;
-	TArray<FVector3f> TransformedVertices;
-
-	constexpr static int32 NumBins = 20;
-	TArray<TArray<int32>> Bins;
-	TArray<FVector3f> BinNormals;
-
-	TArray<TArray<int32>> FaceContacts;
 };
 
