Chaos - remove unnecessary IsKinematic functrion from particle handle

diff --git a/Engine/Source/Runtime/Experimental/Chaos/Public/Chaos/ParticleHandle.h b/Engine/Source/Runtime/Experimental/Chaos/Public/Chaos/ParticleHandle.h
--- a/Engine/Source/Runtime/Experimental/Chaos/Public/Chaos/ParticleHandle.h
+++ b/Engine/Source/Runtime/Experimental/Chaos/Public/Chaos/ParticleHandle.h
@@ -868,7 +868,6 @@ public:
 
 	operator TPBDRigidParticleHandleImp<T, d, false>& () { return reinterpret_cast<TPBDRigidParticleHandleImp<T, d, false>&>(*this); }
 
-	bool IsKinematic() const { return (ObjectState() == EObjectStateType::Kinematic) || (ObjectState() == EObjectStateType::Static); }
 	bool IsDynamic() const { return (ObjectState() == EObjectStateType::Dynamic) || (ObjectState() == EObjectStateType::Sleeping); }
 
 	const TUniquePtr<TBVHParticles<T, d>>& CollisionParticles() const { return PBDRigidParticles->CollisionParticles(ParticleIdx); }
