Chaos Cloth: Add new HierarchicalSpatialHash acceleration structure for cloth self collision/intersections.

The HierarchicalSpatialHash can be built in parallel using a simple lock-free HashMap. Currently the HashMap uses a linked list for buckets. It might be worth considering linear probing or something similar so that bucket data is stored together in memory for queries.

In tests on various cloth assets on our current PC hardware, building the Hierarchical Spatial Hash map is over 2x faster than building the old AABBTree. Querying for the purposes of point-point proximity and edge-face intersections is within a few percent of BVH queries (but consistently slightly slower). The BVH build is serial, so speed up will depend on parallelism, but the fully serial version of the spatial hashing is not significantly slower than BVH building (this may actually speak to other issues with parallel throughput). I did try various batch sizes, but given that we typically have several thousand tris, the default batching (split evenly between all available worker threads) was the best in most cases.

The TriangleMesh proximity and intersection queries have the same interface for BVH and SpatialHash query, so calling code can be templated on which spatial acceleration structure you're using. The PBDCollisionSpringConstraints have been updated in this way to retain the previous BVH API for anyone using it.

diff --git a/Engine/Plugins/Experimental/ChaosCloth/Source/Chaos/Private/ChaosCloth/ChaosClothConstraints.cpp b/Engine/Plugins/Experimental/ChaosCloth/Source/Chaos/Private/ChaosCloth/ChaosClothConstraints.cpp
--- a/Engine/Plugins/Experimental/ChaosCloth/Source/Chaos/Private/ChaosCloth/ChaosClothConstraints.cpp
+++ b/Engine/Plugins/Experimental/ChaosCloth/Source/Chaos/Private/ChaosCloth/ChaosClothConstraints.cpp
@@ -255,7 +255,7 @@ void FClothConstraints::CreateRules()
 			[this](Softs::FSolverParticles& Particles, const Softs::FSolverReal /*Dt*/)
 			{
 				SelfCollisionInit->Init(Particles);
-				SelfCollisionConstraints->Init(Particles, SelfCollisionInit->GetBVH(), SelfCollisionInit->GetVertexGIAColors(), SelfCollisionInit->GetTriangleGIAColors());
+				SelfCollisionConstraints->Init(Particles, SelfCollisionInit->GetSpatialHash(), SelfCollisionInit->GetVertexGIAColors(), SelfCollisionInit->GetTriangleGIAColors());
 			};
 
 		PostCollisionConstraintRules[PostCollisionConstraintRuleIndex++] =
diff --git a/Engine/Source/Runtime/Experimental/Chaos/Private/Chaos/PBDCollisionSpringConstraintsBase.cpp b/Engine/Source/Runtime/Experimental/Chaos/Private/Chaos/PBDCollisionSpringConstraintsBase.cpp
--- a/Engine/Source/Runtime/Experimental/Chaos/Private/Chaos/PBDCollisionSpringConstraintsBase.cpp
+++ b/Engine/Source/Runtime/Experimental/Chaos/Private/Chaos/PBDCollisionSpringConstraintsBase.cpp
@@ -53,7 +53,8 @@ void FPBDCollisionSpringConstraintsBase::Init(const FSolverParticles& Particles)
 	return Init(Particles, BVH, static_cast<TConstArrayView<FPBDTriangleMeshCollisions::FGIAColor>>(EmptyGIAColors), EmptyGIAColors);
 }
 
-void FPBDCollisionSpringConstraintsBase::Init(const FSolverParticles& Particles, const FTriangleMesh::TBVHType<FSolverReal>& BVH, 
+template<typename SpatialAccelerator>
+void FPBDCollisionSpringConstraintsBase::Init(const FSolverParticles& Particles, const SpatialAccelerator& Spatial, 
 	const TConstArrayView<FPBDTriangleMeshCollisions::FGIAColor>& VertexGIAColors, const TArray<FPBDTriangleMeshCollisions::FGIAColor>& TriangleGIAColors)
 {
 	if (!Elements.Num())
@@ -78,13 +79,13 @@ void FPBDCollisionSpringConstraintsBase::Init(const FSolverParticles& Particles,
 		const FSolverReal HeightSq = FMath::Square(Thickness + Thickness);
 
 		PhysicsParallelFor(NumParticles,
-			[this, &BVH, &Particles, &ConstraintIndex, HeightSq, MaxConnectionsPerPoint, &VertexGIAColors, &TriangleGIAColors](int32 i)
+			[this, &Spatial, &Particles, &ConstraintIndex, HeightSq, MaxConnectionsPerPoint, &VertexGIAColors, &TriangleGIAColors](int32 i)
 			{
 				const int32 Index = i + Offset;
 				constexpr FSolverReal ExtraThicknessMult = 1.5f;
 
 				TArray< TTriangleCollisionPoint<FSolverReal> > Result;
-				if (TriangleMesh.PointProximityQuery(BVH, static_cast<const TArrayView<const FSolverVec3>&>(Particles.XArray()), Index, Particles.X(Index), Thickness * ExtraThicknessMult, Thickness * ExtraThicknessMult,
+				if (TriangleMesh.PointProximityQuery(Spatial, static_cast<const TArrayView<const FSolverVec3>&>(Particles.XArray()), Index, Particles.X(Index), Thickness * ExtraThicknessMult, Thickness * ExtraThicknessMult,
 					[this, &VertexGIAColors, &TriangleGIAColors](const int32 PointIndex, const int32 TriangleIndex)->bool
 					{
 						const TVector<int32, 3>& Elem = Elements[TriangleIndex];
@@ -165,6 +166,10 @@ void FPBDCollisionSpringConstraintsBase::Init(const FSolverParticles& Particles,
 		FlipNormal.SetNum(ConstraintNum, /*bAllowShrinking*/ true);
 	}
 }
+template void CHAOS_API FPBDCollisionSpringConstraintsBase::Init<FTriangleMesh::TBVHType<FSolverReal>>(const FSolverParticles& Particles, const FTriangleMesh::TBVHType<FSolverReal>& Spatial, 
+	const TConstArrayView<FPBDTriangleMeshCollisions::FGIAColor>& VertexGIAColors, const TArray<FPBDTriangleMeshCollisions::FGIAColor>& TriangleGIAColors);
+template void CHAOS_API FPBDCollisionSpringConstraintsBase::Init<FTriangleMesh::TSpatialHashType<FSolverReal>>(const FSolverParticles& Particles, const FTriangleMesh::TSpatialHashType<FSolverReal>& Spatial, 
+	const TConstArrayView<FPBDTriangleMeshCollisions::FGIAColor>& VertexGIAColors, const TArray<FPBDTriangleMeshCollisions::FGIAColor>& TriangleGIAColors);
 
 FSolverVec3 FPBDCollisionSpringConstraintsBase::GetDelta(const FSolverParticles& Particles, const int32 i) const
 {
diff --git a/Engine/Source/Runtime/Experimental/Chaos/Private/Chaos/PBDTriangleMeshCollisions.cpp b/Engine/Source/Runtime/Experimental/Chaos/Private/Chaos/PBDTriangleMeshCollisions.cpp
--- a/Engine/Source/Runtime/Experimental/Chaos/Private/Chaos/PBDTriangleMeshCollisions.cpp
+++ b/Engine/Source/Runtime/Experimental/Chaos/Private/Chaos/PBDTriangleMeshCollisions.cpp
@@ -22,7 +22,8 @@ struct FEdgeFaceIntersection
 };
 
 // Returned array has NOT been shrunk
-static TArray<FEdgeFaceIntersection> FindEdgeFaceIntersections(const FTriangleMesh& TriangleMesh, const FTriangleMesh::TBVHType<FSolverReal>& BVH, const FSolverParticles& Particles)
+template<typename SpatialAccelerator>
+static TArray<FEdgeFaceIntersection> FindEdgeFaceIntersections(const FTriangleMesh& TriangleMesh, const SpatialAccelerator& Spatial, const FSolverParticles& Particles)
 {
 	TRACE_CPUPROFILER_EVENT_SCOPE(ChaosFPBDTriangleMeshCollisions_IntersectionQuery);
 
@@ -41,7 +42,7 @@ static TArray<FEdgeFaceIntersection> FindEdgeFaceIntersections(const FTriangleMe
 	TArray<FEdgeFaceIntersection> ExtraIntersections;
 	FCriticalSection CriticalSection;
 	PhysicsParallelFor(SegmentMesh.GetNumElements(),
-		[&BVH, &TriangleMesh, &Particles, &SegmentMesh, &EdgeToFaces, &Elements, &IntersectionIndex, &Intersections, PreallocatedIntersectionsPerEdge, &ExtraIntersections, &CriticalSection](int32 EdgeIndex)
+		[&Spatial, &TriangleMesh, &Particles, &SegmentMesh, &EdgeToFaces, &Elements, &IntersectionIndex, &Intersections, PreallocatedIntersectionsPerEdge, &ExtraIntersections, &CriticalSection](int32 EdgeIndex)
 		{
 
 			TArray< TTriangleCollisionPoint<FSolverReal> > Result;
@@ -52,7 +53,7 @@ static TArray<FEdgeFaceIntersection> FindEdgeFaceIntersections(const FTriangleMe
 			const FSolverVec3& EdgePosition1 = Particles.X(EdgePointIndex1);
 
 
-			if (TriangleMesh.EdgeIntersectionQuery(BVH, static_cast<const TArrayView<const FSolverVec3>&>(Particles.XArray()), EdgeIndex, EdgePosition0, EdgePosition1,
+			if (TriangleMesh.EdgeIntersectionQuery(Spatial, static_cast<const TArrayView<const FSolverVec3>&>(Particles.XArray()), EdgeIndex, EdgePosition0, EdgePosition1,
 				[&Elements, EdgePointIndex0, EdgePointIndex1](int32 EdgeIndex, int32 TriangleIndex)
 				{
 					if (EdgePointIndex0 == Elements[TriangleIndex][0] || EdgePointIndex0 == Elements[TriangleIndex][1] || EdgePointIndex0 == Elements[TriangleIndex][2] ||
@@ -1160,10 +1161,9 @@ void FPBDTriangleMeshCollisions::Init(const FSolverParticles& Particles)
 	}
 
 	{
-		TRACE_CPUPROFILER_EVENT_SCOPE(ChaosFPBDTriangleMeshCollisions_BuildBVH);
-		TriangleMesh.BuildBVH(static_cast<const TArrayView<const FSolverVec3>&>(Particles.XArray()), BVH);
+		TRACE_CPUPROFILER_EVENT_SCOPE(ChaosFPBDTriangleMeshCollisions_BuildSpatialHash);
+		TriangleMesh.BuildSpatialHash(static_cast<const TArrayView<const FSolverVec3>&>(Particles.XArray()), SpatialHash);
 	}
-
 	ContourMinimizationIntersections.Reset();
 	VertexGIAColors.Reset();
 	VertexGIAColors.SetNumZeroed(NumParticles);
@@ -1178,7 +1178,7 @@ void FPBDTriangleMeshCollisions::Init(const FSolverParticles& Particles)
 	}
 
 	// Detect all EdgeFace Intersections
-	TArray<FEdgeFaceIntersection> Intersections = FindEdgeFaceIntersections(TriangleMesh, BVH, Particles);
+	TArray<FEdgeFaceIntersection> Intersections = FindEdgeFaceIntersections(TriangleMesh, SpatialHash, Particles);
 	if (Intersections.Num() == 0)
 	{
 		return;
diff --git a/Engine/Source/Runtime/Experimental/Chaos/Private/Chaos/TriangleMesh.cpp b/Engine/Source/Runtime/Experimental/Chaos/Private/Chaos/TriangleMesh.cpp
--- a/Engine/Source/Runtime/Experimental/Chaos/Private/Chaos/TriangleMesh.cpp
+++ b/Engine/Source/Runtime/Experimental/Chaos/Private/Chaos/TriangleMesh.cpp
@@ -11,6 +11,7 @@
 #include "Math/RandomStream.h"
 #include "Templates/Sorting.h"
 #include "Templates/TypeHash.h"
+#include "Chaos/HierarchicalSpatialHash.h"
 
 #include <algorithm>
 #include <iostream>
@@ -1568,4 +1569,117 @@ template bool FTriangleMesh::EdgeIntersectionQuery<FRealSingle>(const TBVHType<F
 	TFunctionRef<bool(const int32 EdgeIndex, const int32 TriangleIndex)> BroadphaseTest, TArray<TTriangleCollisionPoint<FRealSingle>>& Result) const;
 template bool FTriangleMesh::EdgeIntersectionQuery<FRealDouble>(const TBVHType<FRealDouble>& BVH, const TConstArrayView<TVec3<FRealDouble>>& Points, const int32 EdgeIndex, const TVec3<FRealDouble>& EdgePosition1, const TVec3<FRealDouble>& EdgePosition2,
 	TFunctionRef<bool(const int32 EdgeIndex, const int32 TriangleIndex)> BroadphaseTest, TArray<TTriangleCollisionPoint<FRealDouble>>& Result) const;
+
+template<typename T>
+void FTriangleMesh::BuildSpatialHash(const TConstArrayView<TVec3<T>>& Points, TSpatialHashType<T>& SpatialHash) const
+{
+	TArray<TTriangleMeshBvEntry<T>> BVEntries;
+	const int32 NumTris = MElements.Num();
+	BVEntries.Reset(NumTris);
+	for (int32 Tri = 0; Tri < NumTris; ++Tri)
+	{
+		BVEntries.Add({ this, &Points, Tri });
+	}
+
+	SpatialHash.Initialize(BVEntries);
+}
+template void FTriangleMesh::BuildSpatialHash<FRealSingle>(const TConstArrayView<TVec3<FRealSingle>>& Points, TSpatialHashType<FRealSingle>& SpatialHash) const;
+template void FTriangleMesh::BuildSpatialHash<FRealDouble>(const TConstArrayView<TVec3<FRealDouble>>& Points, TSpatialHashType<FRealDouble>& SpatialHash) const;
+
+template<typename T>
+bool FTriangleMesh::PointProximityQuery(const TSpatialHashType<T>& SpatialHash, const TConstArrayView<TVec3<T>>& Points, const int32 PointIndex, const TVec3<T>& PointPosition, const T PointThickness, const T ThisThickness,
+	TFunctionRef<bool(const int32 PointIndex, const int32 TriangleIndex)> BroadphaseTest, TArray<TTriangleCollisionPoint<T>>& Result) const
+{
+	const T TotalThickness = ThisThickness + PointThickness;
+	const T TotalThicknessSq = TotalThickness * TotalThickness;
+	TAABB<T,3> QueryBounds(PointPosition, PointPosition);
+	QueryBounds.Thicken(TotalThickness);
+
+	const TArray<int32> PotentialIntersections = SpatialHash.FindAllIntersections(QueryBounds);
+
+	Result.Reset(PotentialIntersections.Num());
+
+	for (int32 TriIdx : PotentialIntersections)
+	{
+		if (!BroadphaseTest(PointIndex, TriIdx))
+		{
+			continue;
+		}
+
+		const TVec3<T>& A = Points[MElements[TriIdx][0]];
+		const TVec3<T>& B = Points[MElements[TriIdx][1]];
+		const TVec3<T>& C = Points[MElements[TriIdx][2]];
+		TVec3<T> Bary;
+		const TVec3<T> ClosestPoint = FindClosestPointAndBaryOnTriangle(A, B, C, PointPosition, Bary);
+
+		const T DistSq = (PointPosition - ClosestPoint).SizeSquared();
+		if (DistSq > TotalThicknessSq)
+		{
+			// Failed narrow test.
+			continue;
+		}
+
+		TVec3<T> Normal = TVec3<T>::CrossProduct(B - A, C - A).GetSafeNormal();
+		Normal = (TVec3<T>::DotProduct(Normal, PointPosition - A) > 0) ? Normal : -Normal;
+
+		TTriangleCollisionPoint<T> CollisionPoint;
+		CollisionPoint.ContactType = TTriangleCollisionPoint<T>::EContactType::PointFace;
+		CollisionPoint.Indices[0] = PointIndex;
+		CollisionPoint.Indices[1] = TriIdx;
+		CollisionPoint.Bary = TVec4<T>((T)1., Bary.X, Bary.Y, Bary.Z);
+		CollisionPoint.Location = ClosestPoint;
+		CollisionPoint.Normal = Normal;
+		CollisionPoint.Phi = FMath::Sqrt(DistSq);
+		Result.Add(CollisionPoint);
+	}
+	return Result.Num() > 0;
+}
+template bool FTriangleMesh::PointProximityQuery<FRealSingle>(const TSpatialHashType<FRealSingle>& SpatialHash, const TConstArrayView<TVector<FRealSingle, 3>>& Points, const int32 PointIndex, const TVector<FRealSingle, 3>& PointPosition, const FRealSingle PointThickness, const FRealSingle ThisThickness, TFunctionRef<bool(const int32 PointIndex, const int32 TriangleIndex)> BroadphaseTest, TArray<TTriangleCollisionPoint<FRealSingle>>& Result) const;
+template bool FTriangleMesh::PointProximityQuery<FRealDouble>(const TSpatialHashType<FRealDouble>& SpatialHash, const TConstArrayView<TVector<FRealDouble, 3>>& Points, const int32 PointIndex, const TVector<FRealDouble, 3>& PointPosition, const FRealDouble PointThickness, const FRealDouble ThisThickness, TFunctionRef<bool(const int32 PointIndex, const int32 TriangleIndex)> BroadphaseTest, TArray<TTriangleCollisionPoint<FRealDouble>>& Result) const;
+
+template<typename T>
+bool FTriangleMesh::EdgeIntersectionQuery(const TSpatialHashType<T>& SpatialHash, const TConstArrayView<TVec3<T>>& Points, const int32 EdgeIndex, const TVec3<T>& EdgePosition1, const TVec3<T>& EdgePosition2,
+	TFunctionRef<bool(const int32 EdgeIndex, const int32 TriangleIndex)> BroadphaseTest, TArray<TTriangleCollisionPoint<T>>& Result) const
+{
+	TAABB<T, 3> QueryBounds(EdgePosition1, EdgePosition1);
+	QueryBounds.GrowToInclude(EdgePosition2);
+
+	const TArray<int32> PotentialIntersections = SpatialHash.FindAllIntersections(QueryBounds);
+
+	Result.Reset(PotentialIntersections.Num());
+
+	for (int32 TriIdx : PotentialIntersections)
+	{
+		if (!BroadphaseTest(EdgeIndex, TriIdx))
+		{
+			continue;
+		}
+
+		const TVec3<T>& A = Points[MElements[TriIdx][0]];
+		const TVec3<T>& B = Points[MElements[TriIdx][1]];
+		const TVec3<T>& C = Points[MElements[TriIdx][2]];
+
+		const TTriangle<T> Triangle(A, B, C);
+
+		T Time;
+		TVector<T, 2> Bary;
+		if (Triangle.LineIntersection(EdgePosition1, EdgePosition2, Bary, Time))
+		{
+			TTriangleCollisionPoint<T> CollisionPoint;
+			CollisionPoint.ContactType = TTriangleCollisionPoint<T>::EContactType::EdgeFace;
+			CollisionPoint.Indices[0] = EdgeIndex;
+			CollisionPoint.Indices[1] = TriIdx;
+			CollisionPoint.Bary = TVec4<T>(Time, (T)1. - Bary.X - Bary.Y, Bary.X, Bary.Y);
+			CollisionPoint.Location = ((T)1. - Time) * EdgePosition1 + Time * EdgePosition2;
+			CollisionPoint.Normal = Triangle.GetNormal();
+			CollisionPoint.Phi = 0;
+			Result.Add(CollisionPoint);
+		}
+	}
+	return Result.Num() > 0;
+}
+template bool FTriangleMesh::EdgeIntersectionQuery<FRealSingle>(const TSpatialHashType<FRealSingle>& SpatialHash, const TConstArrayView<TVec3<FRealSingle>>& Points, const int32 EdgeIndex, const TVec3<FRealSingle>& EdgePosition1, const TVec3<FRealSingle>& EdgePosition2,
+	TFunctionRef<bool(const int32 EdgeIndex, const int32 TriangleIndex)> BroadphaseTest, TArray<TTriangleCollisionPoint<FRealSingle>>& Result) const;
+template bool FTriangleMesh::EdgeIntersectionQuery<FRealDouble>(const TSpatialHashType<FRealDouble>& SpatialHash, const TConstArrayView<TVec3<FRealDouble>>& Points, const int32 EdgeIndex, const TVec3<FRealDouble>& EdgePosition1, const TVec3<FRealDouble>& EdgePosition2,
+	TFunctionRef<bool(const int32 EdgeIndex, const int32 TriangleIndex)> BroadphaseTest, TArray<TTriangleCollisionPoint<FRealDouble>>& Result) const;
 }  // End namespace Chaos
diff --git a/Engine/Source/Runtime/Experimental/Chaos/Public/Chaos/HierarchicalSpatialHash.h b/Engine/Source/Runtime/Experimental/Chaos/Public/Chaos/HierarchicalSpatialHash.h
new file mode 100644
--- /dev/null
+++ b/Engine/Source/Runtime/Experimental/Chaos/Public/Chaos/HierarchicalSpatialHash.h
@@ -0,0 +1,307 @@
+// Copyright Epic Games, Inc. All Rights Reserved.
+#pragma once
+
+#include "Chaos/Box.h"
+#include "Chaos/Vector.h"
+#include "Containers/Map.h"
+#include "Chaos/Framework/Parallel.h"
+#include "Math/VectorRegister.h"
+
+#include <atomic>
+
+namespace Chaos
+{
+struct FHierarchicalSpatialHashCellIdx
+{
+	TVec3<int32> Cell;
+	int32 Lod;
+
+	bool operator==(const FHierarchicalSpatialHashCellIdx& OtherIdx) const
+	{
+		return Cell[0] == OtherIdx.Cell[0] && Cell[1] == OtherIdx.Cell[1] && Cell[2] == OtherIdx.Cell[2] && Lod == OtherIdx.Lod;
+	}
+	
+	template<typename T>
+	static int32 LodForAABB(const TAABB<T, 3>& Box)
+	{
+		const T LongestSideLen = Box.Extents().Max();
+		// want Ceil( log2( LongestSideLen) )
+		if (LongestSideLen <= (T).5)
+		{
+			return -(int32)FMath::FloorLog2((uint32)FMath::FloorToInt32((T)1 / LongestSideLen));
+		}
+		else
+		{
+			return (int32)FMath::CeilLogTwo((uint32)FMath::CeilToInt32(LongestSideLen));
+		}
+	}
+
+	template<typename T>
+	static T CellSizeForLod(const int32 Lod)
+	{
+		if (Lod >= 0)
+		{
+			return (T)(1 << Lod);
+		}
+		return (T)1 / (T)(1 << (-Lod));
+	}
+
+	template<typename T>
+	static TVec3<int32> CellIdxForPoint(const TVec3<T>& Point, const T CellSize)
+	{
+		return { FMath::FloorToInt32(Point[0] / CellSize), FMath::FloorToInt32(Point[1] / CellSize), FMath::FloorToInt32(Point[2] / CellSize) };
+	}
+};
+
+inline uint32 GetTypeHash(const Chaos::FHierarchicalSpatialHashCellIdx& Idx)
+{
+	// xorHash
+	int32 Hash = Idx.Cell[0] * 73856093;
+	Hash = Hash ^ (Idx.Cell[1] * 19349663);
+	Hash = Hash ^ (Idx.Cell[2] * 83492791);
+	Hash = Hash ^ (Idx.Lod * 67867979);
+	return (uint32)Hash;
+}
+
+// Currently this assumes a 3 dimensional grid. 
+template<typename PayloadType, typename T>
+class THierarchicalSpatialHash
+{
+	struct FPayloadAndBounds
+	{
+		PayloadType Payload;
+		int32 BoundsLookupIdx;
+	};
+
+	// This hashmap is designed to be fast at just what we need to do to build all at once in exactly the order it is done in Initialize and then query. 
+	// Based off of Core/HashTable.h
+	class FHashMap
+	{
+	public:
+		void Reset()
+		{
+			Hash.Reset();
+			Elements.Reset();
+		}
+
+		void PreallocateElementsForConcurrentAdd(int32 MaxNumElements, int32 HashSize)
+		{
+			HashSize = FMath::RoundUpToPowerOfTwo(HashSize);
+			Hash.SetNumUninitialized(HashSize);
+			for (std::atomic<int32>& H : Hash)
+			{
+				H.store(INDEX_NONE, std::memory_order_relaxed);
+			}
+			Elements.SetNumUninitialized(MaxNumElements);
+			ConcurrentElementAddIdx.store(0);
+		}
+
+		void ConcurrentAddElementRange(const TVec3<int32>& MinCellIdx, const TVec3<int32>& MaxCellIdx, int32 Lod, const FPayloadAndBounds& Value)
+		{
+			const int32 NumElementsToAdd = (MaxCellIdx[0] - MinCellIdx[0] + 1) * (MaxCellIdx[1] - MinCellIdx[1] + 1) * (MaxCellIdx[2] - MinCellIdx[2] + 1);
+
+			int32 IndexToWrite = ConcurrentElementAddIdx.fetch_add(NumElementsToAdd, std::memory_order_relaxed);
+
+			for (int32 XIdx = MinCellIdx[0]; XIdx <= MaxCellIdx[0]; ++XIdx)
+			{
+				for (int32 YIdx = MinCellIdx[1]; YIdx <= MaxCellIdx[1]; ++YIdx)
+				{
+					for (int32 ZIdx = MinCellIdx[2]; ZIdx <= MaxCellIdx[2]; ++ZIdx)
+					{
+						Elements[IndexToWrite].Key = FHierarchicalSpatialHashCellIdx{ {XIdx, YIdx, ZIdx}, Lod };
+						Elements[IndexToWrite].Value = Value;
+						const uint32 HashKey = GetTypeHash(Elements[IndexToWrite].Key) & (Hash.Num() - 1);
+						Elements[IndexToWrite].NextIndex = Hash[HashKey].exchange(IndexToWrite);
+						++IndexToWrite;
+					}
+				}
+			}			
+		}
+
+		void ShrinkElementsAfterConcurrentAdd()
+		{
+			constexpr bool bAllowShrinking = false;
+			Elements.SetNum(ConcurrentElementAddIdx.load(), bAllowShrinking);
+		}
+
+		int32 First(const FHierarchicalSpatialHashCellIdx& Key) const
+		{
+			checkSlow(FMath::IsPowerOfTwo(Hash.Num()));
+			const uint32 HashKey = GetTypeHash(Key) & (Hash.Num() - 1);
+			return Hash[HashKey].load(std::memory_order_relaxed);
+		}
+
+		int32 Next(int32 Index) const
+		{
+			return Elements[Index].NextIndex;
+		}
+
+		static bool IsValid(int32 Index)
+		{
+			return Index != INDEX_NONE;
+		}
+
+		bool KeyMatches(int32 Index, const FHierarchicalSpatialHashCellIdx& Key) const
+		{
+			return Elements[Index].Key == Key;
+		}
+
+		const FPayloadAndBounds& GetValue(int32 Index) const
+		{
+			return Elements[Index].Value;
+		}
+
+	private:
+
+		TArray<std::atomic<int32>> Hash;
+		struct alignas(32) FElement
+		{
+			FHierarchicalSpatialHashCellIdx Key;
+			FPayloadAndBounds Value;
+			int32 NextIndex;
+		};
+
+		TArray<FElement> Elements;
+		std::atomic<int32> ConcurrentElementAddIdx;
+	};
+
+	class FVectorAABB
+	{
+	public:
+
+		FVectorAABB() = default;
+
+		explicit FVectorAABB(const TAABB<T, 3>& BBox)
+			: Min(MakeVectorRegister(BBox.Min()[0], BBox.Min()[1], BBox.Min()[2], (T)0.))
+			, Max(MakeVectorRegister(BBox.Max()[0], BBox.Max()[1], BBox.Max()[2], (T)0.))
+		{}
+
+		bool Intersects(const FVectorAABB& Other) const
+		{
+			if (VectorAnyLesserThan(Other.Max, Min) || VectorAnyGreaterThan(Other.Min, Max))
+			{
+				return false;
+			}
+			return true;
+		}
+
+	private:
+
+		TVectorRegisterType<T> Min;
+		TVectorRegisterType<T> Max;
+	};
+
+public:
+
+	void Reset()
+	{
+		HashMap.Reset();
+		UsedLods.Reset();
+	}
+
+	template<typename ParticlesType>
+	void Initialize(const ParticlesType& Particles)
+	{
+		TRACE_CPUPROFILER_EVENT_SCOPE(HierarchicalSpatialHash_Init);
+		Reset();
+
+		Bounds.SetNumUninitialized(Particles.Num());
+		HashMap.PreallocateElementsForConcurrentAdd(Particles.Num() * 8, Particles.Num() * 2);
+
+		// For most common UsedLods, store whether or not the Lod is used in an array without a lock.
+		// For uncommon Lods, use a critical section to store directly into the UsedLods map.
+		constexpr int32 MaxPreAllocatedUsedLodValue = 10; // i.e., a single element 2^10 = 1024 across
+		constexpr int32 MinPreAllocatedUsedLodValue = -10; // i.e., a single element 1 / (2^10) ~= 0.001 across.
+		TStaticArray<bool, MaxPreAllocatedUsedLodValue - MinPreAllocatedUsedLodValue + 1> UsedLodsArray;
+		memset(UsedLodsArray.GetData(), 0, UsedLodsArray.Num() * sizeof(bool));
+		FCriticalSection UsedLodsCriticalSection;
+
+		PhysicsParallelFor(Particles.Num(),
+			[this, &Particles, MaxPreAllocatedUsedLodValue, MinPreAllocatedUsedLodValue, &UsedLodsArray, &UsedLodsCriticalSection](int32 ParticleIdx)
+			{
+				const auto& Particle = Particles[ParticleIdx];
+				const TAABB<T, 3> ParticleBounds = Particle.BoundingBox();
+
+				const int32 Lod = FHierarchicalSpatialHashCellIdx::LodForAABB(ParticleBounds);
+				const T CellSize = FHierarchicalSpatialHashCellIdx::CellSizeForLod<T>(Lod);
+
+				const PayloadType Payload = Particle.template GetPayload<PayloadType>(ParticleIdx);
+				Bounds[ParticleIdx] = FVectorAABB(ParticleBounds);
+
+				const TVec3<int32> MinCellIdx = FHierarchicalSpatialHashCellIdx::CellIdxForPoint(ParticleBounds.Min(), CellSize);
+				const TVec3<int32> MaxCellIdx = FHierarchicalSpatialHashCellIdx::CellIdxForPoint(ParticleBounds.Max(), CellSize);
+				const FPayloadAndBounds Value = { Payload, ParticleIdx };
+				HashMap.ConcurrentAddElementRange(MinCellIdx, MaxCellIdx, Lod, Value);
+
+				if (Lod >= MinPreAllocatedUsedLodValue && Lod <= MaxPreAllocatedUsedLodValue)
+				{
+					UsedLodsArray[Lod - MinPreAllocatedUsedLodValue] = true;
+				}
+				else
+				{
+					UsedLodsCriticalSection.Lock();
+					UsedLods.Add(Lod, CellSize);
+					UsedLodsCriticalSection.Unlock();
+				}
+			}
+		);
+
+		HashMap.ShrinkElementsAfterConcurrentAdd();
+
+		for (int32 Lod = MinPreAllocatedUsedLodValue; Lod <= MaxPreAllocatedUsedLodValue; ++Lod)
+		{
+			if (UsedLodsArray[Lod - MinPreAllocatedUsedLodValue])
+			{
+				UsedLods.Add(Lod, FHierarchicalSpatialHashCellIdx::CellSizeForLod<T>(Lod));
+			}
+		}
+	}
+
+	TArray<PayloadType> FindAllIntersections(const TAABB<T, 3>& LookupBounds) const 
+	{
+		TArray<PayloadType> Result;
+		FVectorAABB LookupBoundsVec(LookupBounds);
+		for (const typename TMap<int32, T>::ElementType& UsedLod : UsedLods)
+		{
+			const int32 Lod = UsedLod.Key;
+			const T CellSize = UsedLod.Value;
+
+			const TVec3<int32> MinCellIdx = FHierarchicalSpatialHashCellIdx::CellIdxForPoint(LookupBounds.Min(), CellSize);
+			const TVec3<int32> MaxCellIdx = FHierarchicalSpatialHashCellIdx::CellIdxForPoint(LookupBounds.Max(), CellSize);
+
+			for (int32 XIdx = MinCellIdx[0]; XIdx <= MaxCellIdx[0]; ++XIdx)
+			{
+				for (int32 YIdx = MinCellIdx[1]; YIdx <= MaxCellIdx[1]; ++YIdx)
+				{
+					for (int32 ZIdx = MinCellIdx[2]; ZIdx <= MaxCellIdx[2]; ++ZIdx)
+					{
+						const FHierarchicalSpatialHashCellIdx Key{ {XIdx, YIdx, ZIdx}, Lod };
+						for (int32 HashIndex = HashMap.First(Key); HashMap.IsValid(HashIndex); HashIndex = HashMap.Next(HashIndex))
+						{
+							if (HashMap.KeyMatches(HashIndex, Key))
+							{
+								const FPayloadAndBounds Value = HashMap.GetValue(HashIndex);
+								const FVectorAABB& PayloadBounds = Bounds[Value.BoundsLookupIdx];
+
+								if (PayloadBounds.Intersects(LookupBoundsVec))
+								{
+									// Typical cloth examples do not find a lot of unique results. 
+									// In testing, the overhead to do this search is lower than using a TSet and converting the result to an array
+									Result.AddUnique(Value.Payload);
+								}
+							}
+						}
+					}
+				}
+			}
+		}
+		return Result;
+	}
+
+private:
+	TMap<int32, T> UsedLods;
+	FHashMap HashMap;
+	TArray<FVectorAABB> Bounds;
+};
+}
+
diff --git a/Engine/Source/Runtime/Experimental/Chaos/Public/Chaos/PBDCollisionSpringConstraintsBase.h b/Engine/Source/Runtime/Experimental/Chaos/Public/Chaos/PBDCollisionSpringConstraintsBase.h
--- a/Engine/Source/Runtime/Experimental/Chaos/Public/Chaos/PBDCollisionSpringConstraintsBase.h
+++ b/Engine/Source/Runtime/Experimental/Chaos/Public/Chaos/PBDCollisionSpringConstraintsBase.h
@@ -37,10 +37,11 @@ public:
 
 	virtual ~FPBDCollisionSpringConstraintsBase() {}
 
-	UE_DEPRECATED(5.0, "Use Init(Particles, BVH, GIAColors) instead.")
+	UE_DEPRECATED(5.0, "Use Init(Particles, Spatial, GIAColors) instead.")
 	void Init(const FSolverParticles& Particles);
 
-	void Init(const FSolverParticles& Particles, const FTriangleMesh::TBVHType<FSolverReal>& BVH, const TConstArrayView<FPBDTriangleMeshCollisions::FGIAColor>& VertexGIAColors, const TArray<FPBDTriangleMeshCollisions::FGIAColor>& TriangleGIAColors);
+	template<typename SpatialAccelerator>
+	void Init(const FSolverParticles& Particles, const SpatialAccelerator& Spatial, const TConstArrayView<FPBDTriangleMeshCollisions::FGIAColor>& VertexGIAColors, const TArray<FPBDTriangleMeshCollisions::FGIAColor>& TriangleGIAColors);
 
 	FSolverVec3 GetDelta(const FSolverParticles& InParticles, const int32 i) const;
 
diff --git a/Engine/Source/Runtime/Experimental/Chaos/Public/Chaos/PBDTriangleMeshCollisions.h b/Engine/Source/Runtime/Experimental/Chaos/Public/Chaos/PBDTriangleMeshCollisions.h
--- a/Engine/Source/Runtime/Experimental/Chaos/Public/Chaos/PBDTriangleMeshCollisions.h
+++ b/Engine/Source/Runtime/Experimental/Chaos/Public/Chaos/PBDTriangleMeshCollisions.h
@@ -4,6 +4,7 @@
 #include "Chaos/Core.h"
 #include "Chaos/PBDSoftsEvolutionFwd.h"
 #include "Chaos/TriangleMesh.h"
+#include "Chaos/HierarchicalSpatialHash.h"
 
 
 // This does initialization for PBDCollisionSpringConstraints and PBDTriangleMeshIntersections, 
@@ -116,7 +117,7 @@ public:
 	void SetGlobalIntersectionAnalysis(bool bInGlobalIntersectionAnalysis) { bGlobalIntersectionAnalysis = bInGlobalIntersectionAnalysis; }
 	void SetContourMinimization(bool bInContourMinimization) { bContourMinimization = bInContourMinimization; }
 
-	const FTriangleMesh::TBVHType<FSolverReal>& GetBVH() const { return BVH; }
+	const FTriangleMesh::TSpatialHashType<FSolverReal>& GetSpatialHash() const { return SpatialHash; }
 	const TArray<FContourMinimizationIntersection>& GetContourMinimizationIntersections() const { return ContourMinimizationIntersections; }
 	const TConstArrayView<FGIAColor> GetVertexGIAColors() const { return bGlobalIntersectionAnalysis && VertexGIAColors.Num() == NumParticles ? TConstArrayView<FGIAColor>(VertexGIAColors.GetData() - Offset, NumParticles + Offset) : TConstArrayView<FGIAColor>(); }
 	const TArray<FGIAColor>& GetTriangleGIAColors() const { return TriangleGIAColors; }
@@ -130,7 +131,7 @@ private:
 	bool bGlobalIntersectionAnalysis;
 	bool bContourMinimization;
 
-	FTriangleMesh::TBVHType<FSolverReal> BVH;
+	FTriangleMesh::TSpatialHashType<FSolverReal> SpatialHash;
 	TArray<FContourMinimizationIntersection> ContourMinimizationIntersections;
 	TArray<FGIAColor> VertexGIAColors;
 	TArray<FGIAColor> TriangleGIAColors;
diff --git a/Engine/Source/Runtime/Experimental/Chaos/Public/Chaos/TriangleMesh.h b/Engine/Source/Runtime/Experimental/Chaos/Public/Chaos/TriangleMesh.h
--- a/Engine/Source/Runtime/Experimental/Chaos/Public/Chaos/TriangleMesh.h
+++ b/Engine/Source/Runtime/Experimental/Chaos/Public/Chaos/TriangleMesh.h
@@ -11,6 +11,8 @@
 
 namespace Chaos
 {
+	template <typename TPayloadType, typename T> class THierarchicalSpatialHash;
+
 	template<typename T> struct TTriangleCollisionPoint;
 
 	class FTriangleMesh
@@ -277,6 +279,7 @@ namespace Chaos
 			TriMesh.Init(MoveTemp(Elements));
 		}
 
+		// BVH-based collision queries
 		template<typename T>
 		using TBVHType = TAABBTree<int32, TAABBTreeLeafArray<int32, /*bComputeBounds=*/false, T>, /*bMutable=*/true, T>;
 
@@ -291,6 +294,20 @@ namespace Chaos
 		template<typename T>
 		bool EdgeIntersectionQuery(const TBVHType<T>& BVH, const TConstArrayView<TVec3<T>>& Points, const int32 EdgeIndex, const TVec3<T>& EdgePosition1, const TVec3<T>& EdgePosition2,
 			TFunctionRef<bool(const int32 EdgeIndex, const int32 TriangleIndex)> BroadphaseTest, TArray<TTriangleCollisionPoint<T>>& Result) const;
+
+		template<typename T>
+		using TSpatialHashType = THierarchicalSpatialHash<int32, T>;
+
+		template<typename T>
+		void BuildSpatialHash(const TConstArrayView<TVec3<T>>& Points, TSpatialHashType<T>& SpatialHash) const;
+
+		template<typename T>
+		bool PointProximityQuery(const TSpatialHashType<T>& SpatialHash, const TConstArrayView<TVec3<T>>& Points, const int32 PointIndex, const TVec3<T>& PointPosition, const T PointThickness, const T ThisThickness,
+			TFunctionRef<bool(const int32 PointIndex, const int32 TriangleIndex)> BroadphaseTest, TArray<TTriangleCollisionPoint<T>>& Result) const;
+
+		template<typename T>
+		bool EdgeIntersectionQuery(const TSpatialHashType<T>& SpatialHash, const TConstArrayView<TVec3<T>>& Points, const int32 EdgeIndex, const TVec3<T>& EdgePosition1, const TVec3<T>& EdgePosition2,
+			TFunctionRef<bool(const int32 EdgeIndex, const int32 TriangleIndex)> BroadphaseTest, TArray<TTriangleCollisionPoint<T>>& Result) const;
 		
 	private:
 		CHAOS_API void InitHelper(const int32 StartIdx, const int32 EndIdx, const bool CullDegenerateElements=true);
