Chaos : Fix Implicit / scaled triangle mesh CCD contact generation not setting one of the ShapeContactPoints

diff --git a/Engine/Source/Runtime/Experimental/Chaos/Private/Chaos/Collision/ContactPointsMiscShapes.cpp b/Engine/Source/Runtime/Experimental/Chaos/Private/Chaos/Collision/ContactPointsMiscShapes.cpp
--- a/Engine/Source/Runtime/Experimental/Chaos/Private/Chaos/Collision/ContactPointsMiscShapes.cpp
+++ b/Engine/Source/Runtime/Experimental/Chaos/Private/Chaos/Collision/ContactPointsMiscShapes.cpp
@@ -166,6 +166,7 @@ namespace Chaos
 				// @todo(chaos): handle Instanced with margin
 				if (B.LowLevelSweepGeom(ADowncast, AToBTM, LocalDir, Length, OutTime, Location, Normal, FaceIndex, FaceNormal, 0.0f, true))
 				{
+					Contact.ShapeContactPoints[0] = AToBTM.InverseTransformPositionNoScale(Location);
 					Contact.ShapeContactPoints[1] = Location;
 					Contact.ShapeContactNormal = Normal;
 
