Chaos : Fixed crash in PIE when setting trailing event on certain geometry collection

diff --git a/Engine/Source/Runtime/Experimental/GeometryCollectionEngine/Private/GeometryCollection/GeometryCollectionComponent.cpp b/Engine/Source/Runtime/Experimental/GeometryCollectionEngine/Private/GeometryCollection/GeometryCollectionComponent.cpp
--- a/Engine/Source/Runtime/Experimental/GeometryCollectionEngine/Private/GeometryCollection/GeometryCollectionComponent.cpp
+++ b/Engine/Source/Runtime/Experimental/GeometryCollectionEngine/Private/GeometryCollection/GeometryCollectionComponent.cpp
@@ -1103,16 +1103,19 @@ void UGeometryCollectionComponent::InitializeComponent()
 {
 	Super::InitializeComponent();
 
-	if (bStoreVelocities || bNotifyTrailing)
+	if (DynamicCollection)
 	{
-		if (!DynamicCollection->FindAttributeTyped<FVector3f>("LinearVelocity", FTransformCollection::TransformGroup))
+		if (bStoreVelocities || bNotifyTrailing)
 		{
-			DynamicCollection->AddAttribute<FVector3f>("LinearVelocity", FTransformCollection::TransformGroup);
-		}
+			if (!DynamicCollection->FindAttributeTyped<FVector3f>("LinearVelocity", FTransformCollection::TransformGroup))
+			{
+				DynamicCollection->AddAttribute<FVector3f>("LinearVelocity", FTransformCollection::TransformGroup);
+			}
 
-		if (!DynamicCollection->FindAttributeTyped<FVector3f>("AngularVelocity", FTransformCollection::TransformGroup))
-		{
-			DynamicCollection->AddAttribute<FVector3f>("AngularVelocity", FTransformCollection::TransformGroup);
+			if (!DynamicCollection->FindAttributeTyped<FVector3f>("AngularVelocity", FTransformCollection::TransformGroup))
+			{
+				DynamicCollection->AddAttribute<FVector3f>("AngularVelocity", FTransformCollection::TransformGroup);
+			}
 		}
 	}
 
