Additional BVH Query Overlap optimization:
When MTD is not required we were looking for all bounding volumes of tringles intersecting with the query then computing all triangles intersection and once we found one returning true.
Now as soon as a bounding volume intersect the query we compute the triangle intersection and if an intersection is found we return true.

diff --git a/Engine/Source/Programs/HeadlessChaos/Private/HeadlessChaosApp.cpp b/Engine/Source/Programs/HeadlessChaos/Private/HeadlessChaosApp.cpp
--- a/Engine/Source/Programs/HeadlessChaos/Private/HeadlessChaosApp.cpp
+++ b/Engine/Source/Programs/HeadlessChaos/Private/HeadlessChaosApp.cpp
@@ -27,6 +27,7 @@
 #include "HeadlessChaosTestBP.h"
 #include "HeadlessChaosTestRaycast.h"
 #include "HeadlessChaosTestSweep.h"
+#include "HeadlessChaosTestOverlap.h"
 #include "HeadlessChaosTestGJK.h"
 #include "HeadlessChaosTestEPA.h"
 #include "HeadlessChaosTestBroadphase.h"
@@ -205,6 +206,12 @@ TEST(SweepTests, DISABLED_LargeSweep)
 	SUCCEED();
 }
 
+TEST(OverlapTests, Overlap) {
+	ChaosTest::OverlapTriMesh();
+
+	SUCCEED();
+}
+
 TEST(MostOpposingTests, MostOpposing) {
 	ChaosTest::TrimeshMostOpposing();
 	ChaosTest::ConvexMostOpposing();
diff --git a/Engine/Source/Programs/HeadlessChaos/Private/HeadlessChaosTestOverlap.cpp b/Engine/Source/Programs/HeadlessChaos/Private/HeadlessChaosTestOverlap.cpp
new file mode 100644
--- /dev/null
+++ b/Engine/Source/Programs/HeadlessChaos/Private/HeadlessChaosTestOverlap.cpp
@@ -0,0 +1,214 @@
+// Copyright Epic Games, Inc. All Rights Reserved.
+
+#include "HeadlessChaosTestOverlap.h"
+
+#include "HeadlessChaos.h"
+#include "Chaos/Capsule.h"
+#include "Chaos/Box.h"
+#include "Chaos/TriangleMeshImplicitObject.h"
+#include "Chaos/ImplicitObjectScaled.h"
+
+namespace ChaosTest
+{
+	using namespace Chaos;
+
+	void OverlapTriMesh()
+	{
+		// Trimesh is simple pyramid
+		using namespace Chaos;
+		FTriangleMeshImplicitObject::ParticlesType TrimeshParticles(
+		{
+			{-10.0, 0.0, 0.0},
+			{10.0, 0.0, 0.0},
+			{0.0, 10.0, 0.0},
+			{0.0, -10.0, 0.0},
+			{0.0,  0.0, 10.0},
+		});
+
+		TArray<TVec3<int32>> Indices(
+		{
+			{0, 1, 2},
+			{0, 3, 1},
+			{0, 1, 2},
+			{0, 2, 4},
+			{1, 2, 4},
+			{0, 3, 4},
+			{1, 3, 4}
+		});
+
+		TArray<uint16> Materials;
+		for (int32 i = 0; i < Indices.Num(); ++i)
+		{
+			Materials.Emplace(0);
+		}
+
+		TUniquePtr<FTriangleMeshImplicitObject> TriangleMesh = MakeUnique<FTriangleMeshImplicitObject>(MoveTemp(TrimeshParticles), MoveTemp(Indices), MoveTemp(Materials));
+		
+		// Capsule test
+		const FVec3 X1 = { 0.0, 0.0, -2.0 };
+		const FVec3 X2 = { 0.0, 0.0, 2.0 };
+		const FReal Radius = 1.0;
+		const FCapsule Capsule = FCapsule(X1, X2, Radius);
+
+		{
+			const FRigidTransform3 QueryTM(FVec3(0.0, 0.0, 0.0), FQuat::Identity);
+			bool bResult = TriangleMesh->OverlapGeom(Capsule, QueryTM, 0.0);
+			EXPECT_EQ(bResult, true);
+		}
+		{
+			const FRigidTransform3 QueryTM(FVec3(10.0, 10.0, 0.0), FQuat::Identity);
+			bool bResult = TriangleMesh->OverlapGeom(Capsule, QueryTM, 0.0);
+			EXPECT_EQ(bResult, false);
+		}
+		{
+			const FRigidTransform3 QueryTM(FVec3(-10.0, -10.0, 0.0), FQuat::Identity);
+			bool bResult = TriangleMesh->OverlapGeom(Capsule, QueryTM, 0.0);
+			EXPECT_EQ(bResult, false);
+		}
+		{
+			const FRigidTransform3 QueryTM(FVec3(10.0, 0.0, 0.0), FQuat::Identity);
+			bool bResult = TriangleMesh->OverlapGeom(Capsule, QueryTM, 0.0);
+			EXPECT_EQ(bResult, true);
+		}
+		{
+			const FRigidTransform3 QueryTM(FVec3(0.0, 10.0, 0.0), FQuat::Identity);
+			bool bResult = TriangleMesh->OverlapGeom(Capsule, QueryTM, 0.0);
+			EXPECT_EQ(bResult, true);
+		}
+		{
+			const FRigidTransform3 QueryTM(FVec3(0.0, -10.0, 0.0), FQuat::Identity);
+			bool bResult = TriangleMesh->OverlapGeom(Capsule, QueryTM, 0.0);
+			EXPECT_EQ(bResult, true);
+		}
+		{
+			const FRigidTransform3 QueryTM(FVec3(-10.0, 0.0, 0.0), FQuat::Identity);
+			bool bResult = TriangleMesh->OverlapGeom(Capsule, QueryTM, 0.0);
+			EXPECT_EQ(bResult, true);
+		}
+		{
+			const FRigidTransform3 QueryTM(FVec3(0.0, 0.0, 12.0), FQuat::Identity);
+			bool bResult = TriangleMesh->OverlapGeom(Capsule, QueryTM, 0.0);
+			EXPECT_EQ(bResult, true);
+		}
+		{
+			const FRigidTransform3 QueryTM(FVec3(0.0, 0.0, 12.0), FQuat(FVec3(1.0, 0.0, 0.0), 3.1415926 / 4.0));
+			bool bResult = TriangleMesh->OverlapGeom(Capsule, QueryTM, 0.0);
+			EXPECT_EQ(bResult, false);
+		}
+		{
+			// Inside Mesh
+			const FRigidTransform3 QueryTM(FVec3(0.0, 0.0, 4.0), FQuat::Identity);
+			bool bResult = TriangleMesh->OverlapGeom(Capsule, QueryTM, 0.0);
+			EXPECT_EQ(bResult, false);
+		}
+
+
+		// Box test
+		const TBox<FReal, 3> Box = TBox<FReal, 3>({-1.0, -1.0, -2.0 }, { 1.0, 1.0, 2.0 });
+		{
+			const FRigidTransform3 QueryTM(FVec3(0.0, 0.0, 0.0), FQuat::Identity);
+			bool bResult = TriangleMesh->OverlapGeom(Box, QueryTM, 0.0);
+			EXPECT_EQ(bResult, true);
+		}
+		{
+			const FRigidTransform3 QueryTM(FVec3(10.0, 10.0, 0.0), FQuat::Identity);
+			bool bResult = TriangleMesh->OverlapGeom(Box, QueryTM, 0.0);
+			EXPECT_EQ(bResult, false);
+		}
+		{
+			const FRigidTransform3 QueryTM(FVec3(-10.0, -10.0, 0.0), FQuat::Identity);
+			bool bResult = TriangleMesh->OverlapGeom(Box, QueryTM, 0.0);
+			EXPECT_EQ(bResult, false);
+		}
+		{
+			const FRigidTransform3 QueryTM(FVec3(10.0, 0.0, 0.0), FQuat::Identity);
+			bool bResult = TriangleMesh->OverlapGeom(Box, QueryTM, 0.0);
+			EXPECT_EQ(bResult, true);
+		}
+		{
+			const FRigidTransform3 QueryTM(FVec3(0.0, 10.0, 0.0), FQuat::Identity);
+			bool bResult = TriangleMesh->OverlapGeom(Box, QueryTM, 0.0);
+			EXPECT_EQ(bResult, true);
+		}
+		{
+			const FRigidTransform3 QueryTM(FVec3(0.0, -10.0, 0.0), FQuat::Identity);
+			bool bResult = TriangleMesh->OverlapGeom(Box, QueryTM, 0.0);
+			EXPECT_EQ(bResult, true);
+		}
+		{
+			const FRigidTransform3 QueryTM(FVec3(-10.0, 0.0, 0.0), FQuat::Identity);
+			bool bResult = TriangleMesh->OverlapGeom(Box, QueryTM, 0.0);
+			EXPECT_EQ(bResult, true);
+		}
+		{
+			const FRigidTransform3 QueryTM(FVec3(0.0, 0.0, 12.0), FQuat::Identity);
+			bool bResult = TriangleMesh->OverlapGeom(Box, QueryTM, 0.0);
+			EXPECT_EQ(bResult, true);
+		}
+		{
+			const FRigidTransform3 QueryTM(FVec3(0.0, 0.0, 12.0), FQuat(FVec3(1.0, 0.0, 0.0), 3.1415926 / 4.0));
+			bool bResult = TriangleMesh->OverlapGeom(Box, QueryTM, 0.0);
+			EXPECT_EQ(bResult, false);
+		}
+		{
+			// Inside Mesh
+			const FRigidTransform3 QueryTM(FVec3(0.0, 0.0, 4.0), FQuat::Identity);
+			bool bResult = TriangleMesh->OverlapGeom(Box, QueryTM, 0.0);
+			EXPECT_EQ(bResult, false);
+		}
+
+		// Sphere test
+		const TSphere<FReal, 3> Sphere = TSphere<FReal, 3>({ 0.0, 0.0, 0.0 }, 2.0);
+		{
+			const FRigidTransform3 QueryTM(FVec3(0.0, 0.0, 0.0), FQuat::Identity);
+			bool bResult = TriangleMesh->OverlapGeom(Sphere, QueryTM, 0.0);
+			EXPECT_EQ(bResult, true);
+		}
+		{
+			const FRigidTransform3 QueryTM(FVec3(10.0, 10.0, 0.0), FQuat::Identity);
+			bool bResult = TriangleMesh->OverlapGeom(Sphere, QueryTM, 0.0);
+			EXPECT_EQ(bResult, false);
+		}
+		{
+			const FRigidTransform3 QueryTM(FVec3(-10.0, -10.0, 0.0), FQuat::Identity);
+			bool bResult = TriangleMesh->OverlapGeom(Sphere, QueryTM, 0.0);
+			EXPECT_EQ(bResult, false);
+		}
+		{
+			const FRigidTransform3 QueryTM(FVec3(10.0, 0.0, 0.0), FQuat::Identity);
+			bool bResult = TriangleMesh->OverlapGeom(Sphere, QueryTM, 0.0);
+			EXPECT_EQ(bResult, true);
+		}
+		{
+			const FRigidTransform3 QueryTM(FVec3(0.0, 10.0, 0.0), FQuat::Identity);
+			bool bResult = TriangleMesh->OverlapGeom(Sphere, QueryTM, 0.0);
+			EXPECT_EQ(bResult, true);
+		}
+		{
+			const FRigidTransform3 QueryTM(FVec3(0.0, -10.0, 0.0), FQuat::Identity);
+			bool bResult = TriangleMesh->OverlapGeom(Sphere, QueryTM, 0.0);
+			EXPECT_EQ(bResult, true);
+		}
+		{
+			const FRigidTransform3 QueryTM(FVec3(-10.0, 0.0, 0.0), FQuat::Identity);
+			bool bResult = TriangleMesh->OverlapGeom(Sphere, QueryTM, 0.0);
+			EXPECT_EQ(bResult, true);
+		}
+		{
+			const FRigidTransform3 QueryTM(FVec3(0.0, 0.0, 12.0), FQuat::Identity);
+			bool bResult = TriangleMesh->OverlapGeom(Sphere, QueryTM, 0.0);
+			EXPECT_EQ(bResult, true);
+		}
+		{
+			const FRigidTransform3 QueryTM(FVec3(0.0, 0.0, 12.0), FQuat(FVec3(1.0, 0.0, 0.0), 3.1415926 / 4.0));
+			bool bResult = TriangleMesh->OverlapGeom(Sphere, QueryTM, 0.0);
+			EXPECT_EQ(bResult, true);
+		}
+		{
+			// Inside Mesh
+			const FRigidTransform3 QueryTM(FVec3(0.0, 0.0, 4.0), FQuat::Identity);
+			bool bResult = TriangleMesh->OverlapGeom(Sphere, QueryTM, 0.0);
+			EXPECT_EQ(bResult, false);
+		}
+	}
+}
\ No newline at end of file
diff --git a/Engine/Source/Programs/HeadlessChaos/Public/HeadlessChaosTestOverlap.h b/Engine/Source/Programs/HeadlessChaos/Public/HeadlessChaosTestOverlap.h
new file mode 100644
--- /dev/null
+++ b/Engine/Source/Programs/HeadlessChaos/Public/HeadlessChaosTestOverlap.h
@@ -0,0 +1,10 @@
+// Copyright Epic Games, Inc. All Rights Reserved.
+
+#pragma once
+
+#include "HeadlessChaosTestUtility.h"
+
+namespace ChaosTest
+{
+	void OverlapTriMesh();
+}
\ No newline at end of file
diff --git a/Engine/Source/Runtime/Experimental/Chaos/Private/Chaos/TriangleMeshImplicitObject.cpp b/Engine/Source/Runtime/Experimental/Chaos/Private/Chaos/TriangleMeshImplicitObject.cpp
--- a/Engine/Source/Runtime/Experimental/Chaos/Private/Chaos/TriangleMeshImplicitObject.cpp
+++ b/Engine/Source/Runtime/Experimental/Chaos/Private/Chaos/TriangleMeshImplicitObject.cpp
@@ -357,6 +357,87 @@ struct FTriangleMeshOverlapVisitor
 	TArray<int32>& CollectedResults;
 };
 
+template <typename QueryGeomType>
+struct FTriangleMeshOverlapVisitorNoMTD
+{
+	FTriangleMeshOverlapVisitorNoMTD(const TRigidTransform<FReal, 3>& WorldScaleQueryTM, const QueryGeomType& InQueryGeom, FReal InThickness, const FVec3& InTriMeshScale, const FTriangleMeshImplicitObject* InTriMesh)
+		: QueryGeom(InQueryGeom)
+		, TriMeshScale(InTriMeshScale)
+		, Thickness(InThickness)
+		, TriMesh(InTriMesh)
+		, bFoundIntersection(false)
+	{
+		const UE::Math::TQuat<FReal>& RotationDouble = WorldScaleQueryTM.GetRotation();
+		Rotation = MakeVectorRegisterFloatFromDouble(MakeVectorRegister(RotationDouble.X, RotationDouble.Y, RotationDouble.Z, RotationDouble.W));
+		// Normalize rotation
+		Rotation = VectorNormalizeSafe(Rotation, GlobalVectorConstants::Float0001);
+
+		const UE::Math::TVector<FReal>& TranslationDouble = WorldScaleQueryTM.GetTranslation();
+		Translation = MakeVectorRegisterFloatFromDouble(MakeVectorRegister(TranslationDouble.X, TranslationDouble.Y, TranslationDouble.Z, 0.0));
+	}
+	
+	bool VisitOverlap(int32 TriIdx)
+	{
+		FVec3 A, B, C;
+		if (TriMesh->MElements.RequiresLargeIndices())
+		{
+			TriangleMeshTransformVertsHelper(TriMeshScale, TriIdx, TriMesh->MParticles, TriMesh->MElements.GetLargeIndexBuffer(), A, B, C);
+		}
+		else
+		{
+			TriangleMeshTransformVertsHelper(TriMeshScale, TriIdx, TriMesh->MParticles, TriMesh->MElements.GetSmallIndexBuffer(), A, B, C);
+		}
+
+		const VectorRegister4Float ASimd = MakeVectorRegisterFloatFromDouble(MakeVectorRegister(A.X, A.Y, A.Z, 0.0));
+		const VectorRegister4Float BSimd = MakeVectorRegisterFloatFromDouble(MakeVectorRegister(B.X, B.Y, B.Z, 0.0));
+		const VectorRegister4Float CSimd = MakeVectorRegisterFloatFromDouble(MakeVectorRegister(C.X, C.Y, C.Z, 0.0));
+
+		const VectorRegister4Float AB = VectorSubtract(BSimd, ASimd);
+		const VectorRegister4Float AC = VectorSubtract(CSimd, ASimd);
+
+		//It's most likely that the query object is in front of the triangle since queries tend to be on the outside.
+		//However, maybe we should check if it's behind the triangle plane. Also, we should enforce this winding in some way
+		const VectorRegister4Float InitialDir = VectorCross(AB, AC);
+
+		FTriangleRegister Tri(ASimd, BSimd, CSimd);
+
+		if (GJKIntersectionSimd(Tri, QueryGeom, Translation, Rotation, Thickness, InitialDir))
+		{
+			bFoundIntersection = true;
+			return false;
+		}
+		return true;		
+	}
+
+	bool VisitSweep(int32 Instance, FRealSingle& CurDataLength)
+	{
+		check(false);
+		return true;
+	}
+	bool VisitRaycast(int32 Instance, FRealSingle& CurDataLength)
+	{
+		check(false);
+		return true;
+	}
+
+	VectorRegister4Float Rotation;
+	VectorRegister4Float Translation;
+	const QueryGeomType& QueryGeom;
+	const FVec3 TriMeshScale;
+	FReal Thickness;
+	const FTriangleMeshImplicitObject* TriMesh;
+	bool bFoundIntersection;
+};
+
+template <typename QueryGeomType>
+bool FTrimeshBVH::FindAllIntersectionsNoMTD(const FAABB3& Intersection, const TRigidTransform<FReal, 3>& Transform, const QueryGeomType& QueryGeom, FReal Thickness, const FVec3& TriMeshScale, const FTriangleMeshImplicitObject* TriMesh) const
+{
+	FTriangleMeshOverlapVisitorNoMTD Visitor(Transform, QueryGeom, Thickness, TriMeshScale, TriMesh);
+	Overlap(FAABBVectorized(Intersection), Visitor);
+	return Visitor.bFoundIntersection;
+}
+
+
 TArray<int32> FTrimeshBVH::FindAllIntersections (const FAABB3& Intersection) const
 {
 	TArray<int32> Results;
@@ -424,7 +505,7 @@ bool FTriangleMeshImplicitObject::ContactManifoldImp(const GeomType& QueryGeom,
 	ensure(TriMeshScale != FVec3(0.0f));
 	bool bResult = false;
 
-	const auto& WorldScaleGeom = ScaleGeomIntoWorldHelper(QueryGeom, TriMeshScale);
+	const GeomType& WorldScaleGeom = ScaleGeomIntoWorldHelper(QueryGeom, TriMeshScale);
 	const FVec3 InvTriMeshScale = FReal(1) / TriMeshScale;
 
 	// IMPORTANT QueryTM comes with a invscaled translation so we need a version of the TM with world space translation to properly compute the bounds
@@ -544,13 +625,13 @@ bool FTriangleMeshImplicitObject::GJKContactPointImp(const QueryGeomType& QueryG
 	ensure(TriMeshScale != FVec3(0.0f));
 	bool bResult = false;
 
-	const auto& WorldScaleGeom = ScaleGeomIntoWorldHelper(QueryGeom, TriMeshScale);
+	const QueryGeomType& WorldScaleGeom = ScaleGeomIntoWorldHelper(QueryGeom, TriMeshScale);
 	const FVec3 InvTriMeshScale = FVec3(FReal(1) / TriMeshScale.X, FReal(1) / TriMeshScale.Y, FReal(1) / TriMeshScale.Z);
 
 	// IMPORTANT QueryTM comes with a invscaled translation so we need a version of the TM with world space translation to properly compute the bounds
 	FRigidTransform3 TriMeshToGeomNoScale{ QueryTM };
 	TriMeshToGeomNoScale.SetTranslation(TriMeshToGeomNoScale.GetTranslation() * TriMeshScale);
-	// NOTE: BVH test is done in tri-mesh local space (whereas collision detection is done in world space becaused you can't non-uniformly scale all shapes)
+	// NOTE: BVH test is done in tri-mesh local space (whereas collision detection is done in world space because you can't non-uniformly scale all shapes)
 	FAABB3 QueryBounds = WorldScaleGeom.BoundingBox();
 	QueryBounds = QueryBounds.TransformedAABB(TriMeshToGeomNoScale);
 	QueryBounds.ThickenSymmetrically(FVec3(WorldThickness));
@@ -806,12 +887,13 @@ void FTriangleMeshImplicitObject::VisitTriangle(const int32 TriangleIndex, const
 	}
 }
 
+
 template <typename QueryGeomType>
 bool FTriangleMeshImplicitObject::OverlapGeomImp(const QueryGeomType& QueryGeom, const FRigidTransform3& QueryTM, const FReal Thickness, FMTDInfo* OutMTD, FVec3 TriMeshScale) const
 {
 	bool bResult = false;
 
-	const auto& WorldScaleQueryGeom = ScaleGeomIntoWorldHelper(QueryGeom, TriMeshScale);
+	const QueryGeomType& WorldScaleQueryGeom = ScaleGeomIntoWorldHelper(QueryGeom, TriMeshScale);
 
 	const FVec3 InvTriMeshScale = FVec3(FReal(1) / TriMeshScale.X, FReal(1) / TriMeshScale.Y, FReal(1) / TriMeshScale.Z);
 
@@ -824,20 +906,19 @@ bool FTriangleMeshImplicitObject::OverlapGeomImp(const QueryGeomType& QueryGeom,
 	QueryBounds.ThickenSymmetrically(FVec3(Thickness));
 	QueryBounds.ScaleWithNegative(InvTriMeshScale);
 
-	const TArray<int32> PotentialIntersections = FastBVH.FindAllIntersections(QueryBounds);
 
 	if (OutMTD)
 	{
+		const TArray<int32> PotentialIntersections = FastBVH.FindAllIntersections(QueryBounds);
+
 		OutMTD->Normal = FVec3(0.0);
 		OutMTD->Penetration = TNumericLimits<FReal>::Lowest();
-	}
 
-	TRigidTransform<FReal, 3> WorldScaleQueryTM;
-	ScaleTransformHelper(TriMeshScale, QueryTM, WorldScaleQueryTM);
 
-	auto LambdaHelper = [&](const auto& Elements, FMTDInfo* InnerMTD)
-	{
-		if (InnerMTD)
+		TRigidTransform<FReal, 3> WorldScaleQueryTM;
+		ScaleTransformHelper(TriMeshScale, QueryTM, WorldScaleQueryTM);
+
+		auto LambdaHelper = [&](const auto& Elements, FMTDInfo* InnerMTD)
 		{
 			bool bOverlap = false;
 			for (int32 TriIdx : PotentialIntersections)
@@ -864,54 +945,22 @@ bool FTriangleMeshImplicitObject::OverlapGeomImp(const QueryGeomType& QueryGeom,
 			}
 
 			return bOverlap;
+		};
+
+		if (MElements.RequiresLargeIndices())
+		{
+			return LambdaHelper(MElements.GetLargeIndexBuffer(), OutMTD);
 		}
 		else
 		{
-
-			const UE::Math::TQuat<FReal>& RotationDouble = WorldScaleQueryTM.GetRotation();
-			VectorRegister4Float Rotation = MakeVectorRegisterFloatFromDouble(MakeVectorRegister(RotationDouble.X, RotationDouble.Y, RotationDouble.Z, RotationDouble.W));
-			// Normalize rotation
-			Rotation = VectorNormalizeSafe(Rotation, GlobalVectorConstants::Float0001);
-
-			const UE::Math::TVector<FReal>& TranslationDouble = WorldScaleQueryTM.GetTranslation();
-			const VectorRegister4Float Translation = MakeVectorRegisterFloatFromDouble(MakeVectorRegister(TranslationDouble.X, TranslationDouble.Y, TranslationDouble.Z, 0.0));
-		
-
-			for (int32 TriIdx : PotentialIntersections)
-			{
-				FVec3 A, B, C;
-				TriangleMeshTransformVertsHelper(TriMeshScale, TriIdx, MParticles, Elements, A, B, C);
-
-				const VectorRegister4Float ASimd = MakeVectorRegisterFloatFromDouble(MakeVectorRegister(A.X, A.Y, A.Z, 0.0));
-				const VectorRegister4Float BSimd = MakeVectorRegisterFloatFromDouble(MakeVectorRegister(B.X, B.Y, B.Z, 0.0));
-				const VectorRegister4Float CSimd = MakeVectorRegisterFloatFromDouble(MakeVectorRegister(C.X, C.Y, C.Z, 0.0));
-
-				const VectorRegister4Float AB = VectorSubtract(BSimd, ASimd);
-				const VectorRegister4Float AC = VectorSubtract(CSimd, ASimd);
-
-				//It's most likely that the query object is in front of the triangle since queries tend to be on the outside.
-				//However, maybe we should check if it's behind the triangle plane. Also, we should enforce this winding in some way
-				const VectorRegister4Float InitialDir = VectorCross(AB, AC);
-				
-				FTriangleRegister Tri(ASimd, BSimd, CSimd);
-
-				if (GJKIntersectionSimd(Tri, WorldScaleQueryGeom, Translation, Rotation, Thickness, InitialDir))
-				{
-					return true;
-				}
-			}
-
-			return false;
+			return LambdaHelper(MElements.GetSmallIndexBuffer(), OutMTD);
 		}
-	};
-
-	if(MElements.RequiresLargeIndices())
-	{
-		return LambdaHelper(MElements.GetLargeIndexBuffer(), OutMTD);
 	}
 	else
 	{
-		return LambdaHelper(MElements.GetSmallIndexBuffer(), OutMTD);
+		TRigidTransform<FReal, 3> WorldScaleQueryTM;
+		ScaleTransformHelper(TriMeshScale, QueryTM, WorldScaleQueryTM);
+		return FastBVH.FindAllIntersectionsNoMTD(QueryBounds, WorldScaleQueryTM, WorldScaleQueryGeom, Thickness, TriMeshScale, this);
 	}
 }
 
diff --git a/Engine/Source/Runtime/Experimental/Chaos/Public/Chaos/TriangleMeshImplicitObject.h b/Engine/Source/Runtime/Experimental/Chaos/Public/Chaos/TriangleMeshImplicitObject.h
--- a/Engine/Source/Runtime/Experimental/Chaos/Public/Chaos/TriangleMeshImplicitObject.h
+++ b/Engine/Source/Runtime/Experimental/Chaos/Public/Chaos/TriangleMeshImplicitObject.h
@@ -24,6 +24,9 @@ namespace Chaos
 	extern CHAOS_API bool TriMeshPerPolySupport;
 
 	struct FMTDInfo;
+	template <typename QueryGeomType>
+	struct FTriangleMeshOverlapVisitorNoMTD;
+	class FTriangleMeshImplicitObject;
 	class FContactPoint;
 	FArchive& operator<<(FArchive& Ar, FAABBVectorized& Bounds);
 
@@ -182,9 +185,10 @@ namespace Chaos
 			VisitTree(BoundsFilter, FaceVisitor);
 		}
 
+		template <typename QueryGeomType>
+		bool FindAllIntersectionsNoMTD(const FAABB3& Intersection, const TRigidTransform<FReal, 3>& Transform, const QueryGeomType& QueryGeom, FReal Thickness, const FVec3& TriMeshScale, const FTriangleMeshImplicitObject* TriMesh) const;
 		TArray<int32> FindAllIntersections(const FAABB3& Intersection) const;
 
-
 		template <typename BoundsFilterType, typename FaceVisitorType>
 		FORCEINLINE_DEBUGGABLE EVisitorResult VisitFaces(int32 StartIndex, int32 IndexCount, BoundsFilterType& BoundsFilter, FaceVisitorType& FaceVisitor) const
 		{
@@ -739,6 +743,8 @@ namespace Chaos
 
 		template<typename InStorageType, typename InRealType>
 		friend struct FBvEntry;
+		template <typename QueryGeomType> 
+		friend struct FTriangleMeshOverlapVisitorNoMTD;
 
 		template<bool bRequiresLargeIndex>
 		struct FBvEntry
