Fix GJK bug by introducing an epsilon tolerance to consider a point in the tetrahedron

diff --git a/Engine/Source/Programs/HeadlessChaos/Private/HeadlessChaosTestGJK.cpp b/Engine/Source/Programs/HeadlessChaos/Private/HeadlessChaosTestGJK.cpp
--- a/Engine/Source/Programs/HeadlessChaos/Private/HeadlessChaosTestGJK.cpp
+++ b/Engine/Source/Programs/HeadlessChaos/Private/HeadlessChaosTestGJK.cpp
@@ -1994,4 +1994,30 @@ namespace ChaosTest
 		EXPECT_NEAR(ClosestB.Z, (FReal)50.0f, (FReal)KINDA_SMALL_NUMBER);
 
 	}
+
+
+	GTEST_TEST(GJKTests, GJKBug_BadBarycentricCoords)
+	{
+		FImplicitBox3 BoxA({ -31.999998092651367, -0.73166346549987793, -47.015655517578125 },
+			{ 31.999998092651367, 0.73166346549987793, 47.015655517578125 });
+		FImplicitBox3 BoxB({ -64.202491760253906, -64.269241333007812, 0.27499961853027344 },
+			{ -0.20249176025390625, -0.18923950195312500, 38.524999618530273 });
+
+		const FRigidTransform3 Transform(FVec3(4.28251314, -16.3213539, 32.0828743), FQuat(0.360390902, 0.00000000, 0.00000000, 0.932801366));
+
+		const FVec3 RayDir(0.00000000, 0.672346294, -0.740236819);
+		const FRealDouble Length = 37.388404846191406;
+		FRealDouble OutTime = 0;
+		FVec3 OutPosition(0);
+		FVec3 OutNormal(0);
+		const FRealDouble Thickness = 0;
+		const bool bComputeMTD = true;
+		const FVec3 Offset(-4.2825129036202441, -9.4891323727604799, -34.722524278655456);
+
+		bool bResult = GJKRaycast2<FRealDouble, FImplicitBox3, FImplicitBox3>(BoxA, BoxB, Transform, RayDir, Length, OutTime, OutPosition, OutNormal, Thickness, bComputeMTD, Offset, Thickness);
+		EXPECT_TRUE(bResult);
+		EXPECT_NEAR(OutPosition.X, -13.95, 1e-1);
+		EXPECT_NEAR(OutPosition.Y, -0.73, 1e-1);
+		EXPECT_NEAR(OutPosition.Z, 14.63, 1e-1);
+	}
 }
\ No newline at end of file
diff --git a/Engine/Source/Runtime/Experimental/Chaos/Public/Chaos/SimplexVectorized.h b/Engine/Source/Runtime/Experimental/Chaos/Public/Chaos/SimplexVectorized.h
--- a/Engine/Source/Runtime/Experimental/Chaos/Public/Chaos/SimplexVectorized.h
+++ b/Engine/Source/Runtime/Experimental/Chaos/Public/Chaos/SimplexVectorized.h
@@ -229,13 +229,13 @@ namespace Chaos
 
 		if (VectorMaskBits(IsBC))
 		{
-			const VectorRegister4Float w = VectorDivide(d4MinusD3, NormalizationDenominatorBC);
 			constexpr VectorRegister4Int two = MakeVectorRegisterIntConstant(2, 2, 2, 2);
 			NumVerts = two;
-			const VectorRegister4Float OneMinusW = VectorSubtract(GlobalVectorConstants::FloatOne, w);
-			// b0	a1	a2	a3
+			const VectorRegister4Float w = VectorDivide(d4MinusD3, NormalizationDenominatorBC);
 			if (CalculatExtraInformation)
 			{
+				const VectorRegister4Float OneMinusW = VectorSubtract(GlobalVectorConstants::FloatOne, w);
+				// b0	a1	a2	a3
 				OutBarycentric = VectorUnpackLo(OneMinusW, w);
 			}
 			const VectorRegister4Float CMinusB = VectorSubtract(C, B);
@@ -259,12 +259,12 @@ namespace Chaos
 		constexpr VectorRegister4Int three = MakeVectorRegisterIntConstant(3, 3, 3, 3);
 		NumVerts = three;
 
-		const VectorRegister4Float OneMinusVMinusW = VectorSubtract(VectorSubtract(GlobalVectorConstants::FloatOne, v), w);
-		// b0	a1	a2	a3
-		const VectorRegister4Float OneMinusVMinusW_W = VectorUnpackLo(OneMinusVMinusW, w);
-		// a0	b0	a1	b1
 		if (CalculatExtraInformation)
 		{
+			const VectorRegister4Float OneMinusVMinusW = VectorSubtract(VectorSubtract(GlobalVectorConstants::FloatOne, v), w);
+			// b0	a1	a2	a3
+			const VectorRegister4Float OneMinusVMinusW_W = VectorUnpackLo(OneMinusVMinusW, w);
+			// a0	b0	a1	b1
 			OutBarycentric = VectorUnpackLo(OneMinusVMinusW_W, v);
 		}
 
@@ -347,10 +347,12 @@ namespace Chaos
 													MakeVectorRegisterIntConstant(3, 3, 3, 3)
 		};
 
+		 VectorRegister4Float Eps = VectorMultiply(GlobalVectorConstants::KindaSmallNumber, VectorDivide(DetM, VectorSet1(4.0f)));
+
 		bool bInside = true;
 		for (int Idx = 0; Idx < 4; ++Idx)
 		{
-			bSignMatch[Idx] = VectorSignMatch(DetM, Cofactors[Idx]);
+			bSignMatch[Idx] = VectorSignMatch(DetM, VectorSubtract(Cofactors[Idx], Eps));
 			if (!bSignMatch[Idx])
 			{
 				bInside = false;
@@ -371,19 +373,18 @@ namespace Chaos
 
 		if (bInside)
 		{
-			VectorRegister4Float OutBarycentricVectors[4];
-			const VectorRegister4Float  InvDetM = VectorDivide(GlobalVectorConstants::FloatOne, DetM);
-			OutBarycentricVectors[0] = VectorDivide(Cofactors[0], InvDetM);
-			OutBarycentricVectors[1] = VectorDivide(Cofactors[1], InvDetM);
-			OutBarycentricVectors[2] = VectorDivide(Cofactors[2], InvDetM);
-			OutBarycentricVectors[3] = VectorDivide(Cofactors[3], InvDetM);
-
-			// a0	b0	a1	b1
-			const VectorRegister4Float OutBarycentric0101 = VectorUnpackLo(OutBarycentricVectors[0], OutBarycentricVectors[1]);
-			const VectorRegister4Float OutBarycentric2323 = VectorUnpackLo(OutBarycentricVectors[2], OutBarycentricVectors[3]);
-			// a0	a1	b0	b1
 			if (CalculatExtraInformation)
 			{
+				VectorRegister4Float OutBarycentricVectors[4];
+				const VectorRegister4Float  InvDetM = VectorDivide(GlobalVectorConstants::FloatOne, DetM);
+				OutBarycentricVectors[0] = VectorDivide(Cofactors[0], InvDetM);
+				OutBarycentricVectors[1] = VectorDivide(Cofactors[1], InvDetM);
+				OutBarycentricVectors[2] = VectorDivide(Cofactors[2], InvDetM);
+				OutBarycentricVectors[3] = VectorDivide(Cofactors[3], InvDetM);
+				// a0	b0	a1	b1
+				const VectorRegister4Float OutBarycentric0101 = VectorUnpackLo(OutBarycentricVectors[0], OutBarycentricVectors[1]);
+				const VectorRegister4Float OutBarycentric2323 = VectorUnpackLo(OutBarycentricVectors[2], OutBarycentricVectors[3]);
+				// a0	a1	b0	b1
 				OutBarycentric = VectorMoveLh(OutBarycentric0101, OutBarycentric2323);
 			}
 
