Fix warning in HierarchicalSpatialHash about Arithmetic overflow when compiling for doubles.

diff --git a/Engine/Source/Runtime/Experimental/Chaos/Public/Chaos/HierarchicalSpatialHash.h b/Engine/Source/Runtime/Experimental/Chaos/Public/Chaos/HierarchicalSpatialHash.h
--- a/Engine/Source/Runtime/Experimental/Chaos/Public/Chaos/HierarchicalSpatialHash.h
+++ b/Engine/Source/Runtime/Experimental/Chaos/Public/Chaos/HierarchicalSpatialHash.h
@@ -41,9 +41,9 @@ struct FHierarchicalSpatialHashCellIdx
 	{
 		if (Lod >= 0)
 		{
-			return (T)(1 << Lod);
+			return (T)(1ULL << Lod);
 		}
-		return (T)1 / (T)(1 << (-Lod));
+		return (T)1 / (T)(1ULL << (-Lod));
 	}
 
 	template<typename T>
