Chaos : make sure we only generate connection graphs once
- don't compute connection graph in BuildClusters_Internal
- fix the post creation generation by using the right particle array

diff --git a/Engine/Source/Runtime/Experimental/Chaos/Private/PhysicsProxy/GeometryCollectionPhysicsProxy.cpp b/Engine/Source/Runtime/Experimental/Chaos/Private/PhysicsProxy/GeometryCollectionPhysicsProxy.cpp
--- a/Engine/Source/Runtime/Experimental/Chaos/Private/PhysicsProxy/GeometryCollectionPhysicsProxy.cpp
+++ b/Engine/Source/Runtime/Experimental/Chaos/Private/PhysicsProxy/GeometryCollectionPhysicsProxy.cpp
@@ -979,14 +979,14 @@ void FGeometryCollectionPhysicsProxy::InitializeBodiesPT(Chaos::FPBDRigidsSolver
 			{
 				if (RestCollection->IsClustered(TransformGroupIndex))
 				{
-					if (SolverClusterHandles[TransformGroupIndex])
+					if (SolverParticleHandles[TransformGroupIndex])
 					{
 						Chaos::FClusterCreationParameters ClusterParams;
 						// #todo: should other parameters be set here?  Previously, there was no parameters being sent, and it is unclear
 						// where some of these parameters are defined (ie: CoillisionThicknessPercent)
 						ClusterParams.ConnectionMethod = Parameters.ClusterConnectionMethod;
 						
-						RigidsSolver->GetEvolution()->GetRigidClustering().GenerateConnectionGraph(SolverClusterHandles[TransformGroupIndex], ClusterParams);
+						RigidsSolver->GetEvolution()->GetRigidClustering().GenerateConnectionGraph(SolverParticleHandles[TransformGroupIndex], ClusterParams);
 					}
 				}
 			}
@@ -1119,7 +1119,8 @@ FGeometryCollectionPhysicsProxy::BuildClusters_Internal(
 	//Since the mass orientation is computed in world space in both cases we'd end up with a diagonal inertia matrix and identity rotation that looks like this: [big, small, big] or [small, big, big].
 	//Because of this we need to know how to rotate collision particles and geometry to match with original computation. If it was just geometry we could transform it before passing, but we need collision particles as well
 	Chaos::FClusterCreationParameters ClusterCreationParameters = ClusterParameters;
-	ClusterCreationParameters.bGenerateConnectionGraph = true;
+	// connectivity is taken care outside, no need to generate it here
+	ClusterCreationParameters.bGenerateConnectionGraph = false;
 	// fix... ClusterCreationParameters.CollisionParticles = Simplicials[CollectionClusterIndex];
 	ClusterCreationParameters.ConnectionMethod = Parameters.ClusterConnectionMethod;
 	if (ClusterCreationParameters.CollisionParticles)
