Interchange attribute storage: Added support for uint8 enum classes

diff --git a/Engine/Source/Runtime/Experimental/Interchange/Core/Private/Tests/StorageTest.cpp b/Engine/Source/Runtime/Experimental/Interchange/Core/Private/Tests/StorageTest.cpp
--- a/Engine/Source/Runtime/Experimental/Interchange/Core/Private/Tests/StorageTest.cpp
+++ b/Engine/Source/Runtime/Experimental/Interchange/Core/Private/Tests/StorageTest.cpp
@@ -303,6 +303,22 @@ bool FAttributeStorageTest::RunTest(const FString& Parameters)
 		TestEqual(TEXT("`AttributeStorage` must handle add and retrieve FRandomStream attribute"), StoredRandomStream.RandHelper(RandHelperUInt8), TestStoredSeed);
 	}
 
+	//Store an enum class
+	{
+		enum class EAttributeStorageEnumTest : uint8
+		{
+			EnumValue0 = 0,
+			EnumValue23 = 23,
+			EnumValue35 = 35,
+			EnumValue255 = 255,
+		};
+
+		ADD_ATTRIBUTE_STORAGE(EAttributeStorageEnumTest, EAttributeStorageEnumTest::EnumValue0);
+		ADD_ATTRIBUTE_STORAGE(EAttributeStorageEnumTest, EAttributeStorageEnumTest::EnumValue23);
+		ADD_ATTRIBUTE_STORAGE(EAttributeStorageEnumTest, EAttributeStorageEnumTest::EnumValue35);
+		ADD_ATTRIBUTE_STORAGE(EAttributeStorageEnumTest, EAttributeStorageEnumTest::EnumValue255);
+	}
+
 	//Get the Hash we will use it later to do some test
 	FGuid HashGuidRef = TestStorage.GetStorageHash();
 
diff --git a/Engine/Source/Runtime/Experimental/Interchange/Core/Public/Types/AttributeStorage.h b/Engine/Source/Runtime/Experimental/Interchange/Core/Public/Types/AttributeStorage.h
--- a/Engine/Source/Runtime/Experimental/Interchange/Core/Public/Types/AttributeStorage.h
+++ b/Engine/Source/Runtime/Experimental/Interchange/Core/Public/Types/AttributeStorage.h
@@ -10,11 +10,28 @@
 #include "Misc/ScopeLock.h"
 #include "UObject/SoftObjectPath.h"
 
+#include <type_traits>
+
 //Interchange namespace
 namespace UE
 {
 	namespace Interchange
 	{
+		namespace Private
+		{
+			/**
+			 * Wrapper around std::underlying_type to make sure we don't call it on non enum class types.
+			 */
+			template <typename T, typename Enable = void>
+			struct TUnderlyingType {};
+
+			template <typename T>
+			struct TUnderlyingType<T, typename std::enable_if<std::is_enum<T>::value>::type>
+			{
+				using type = typename std::underlying_type<T>::type;
+			};
+		}
+
 		struct FAttributeKey
 		{
 			FName Key;
@@ -234,7 +251,7 @@ namespace UE
 		 *
 		 * @param T The type to be used in FAttributeStorage.
 		 */
-		template<typename T> struct TAttributeTypeTraits
+		template<typename T, typename Enable = void> struct TAttributeTypeTraits
 		{
 			static CONSTEXPR EAttributeTypes GetType()
 			{
@@ -1259,6 +1276,27 @@ namespace UE
 		};
 
 
+		/** Implements variant type traits for uint8 enum classes. */
+		template<typename EnumClassType> struct TAttributeTypeTraits
+			<EnumClassType,
+				typename std::enable_if<
+					std::is_same<
+						typename UE::Interchange::Private::TUnderlyingType<EnumClassType>::type,
+						uint8
+					>::value
+				>::type
+			>
+		{
+			static CONSTEXPR EAttributeTypes GetType() { return EAttributeTypes::UInt8; }
+			static FString ToString(const EnumClassType& Value)
+			{
+				uint32 ValueConv = (uint8)Value;
+				FStringFormatOrderedArguments OrderedArguments;
+				OrderedArguments.Add(FStringFormatArg(ValueConv));
+				return FString::Format(TEXT("{0}"), OrderedArguments);
+			}
+		};
+
 		/** Implements variant type traits for the built-in uint16 type. */
 		template<> struct TAttributeTypeTraits<uint16>
 		{
