Chaos - support for runtime material changes

diff --git a/Engine/Source/Runtime/Experimental/Chaos/Private/Chaos/PBDRigidsEvolutionGBF.cpp b/Engine/Source/Runtime/Experimental/Chaos/Private/Chaos/PBDRigidsEvolutionGBF.cpp
--- a/Engine/Source/Runtime/Experimental/Chaos/Private/Chaos/PBDRigidsEvolutionGBF.cpp
+++ b/Engine/Source/Runtime/Experimental/Chaos/Private/Chaos/PBDRigidsEvolutionGBF.cpp
@@ -1055,6 +1055,16 @@ void FPBDRigidsEvolutionGBF::DestroyParticleCollisionsInAllocator(FGeometryParti
 	}
 }
 
+void FPBDRigidsEvolutionGBF::ParticleMaterialChanged(FGeometryParticleHandle* Particle)
+{
+	Particle->ParticleCollisions().VisitCollisions([this](FPBDCollisionConstraint& Collision)
+	{
+		CollisionConstraints.UpdateConstraintMaterialProperties(Collision);
+		return ECollisionVisitorResult::Continue;
+	});
+}
+
+
 void FPBDRigidsEvolutionGBF::UpdateInertiaConditioning()
 {
 	// The maximum contribution to error correction from rotation
diff --git a/Engine/Source/Runtime/Experimental/Chaos/Private/PhysicsProxy/SingleParticlePhysicsProxy.cpp b/Engine/Source/Runtime/Experimental/Chaos/Private/PhysicsProxy/SingleParticlePhysicsProxy.cpp
--- a/Engine/Source/Runtime/Experimental/Chaos/Private/PhysicsProxy/SingleParticlePhysicsProxy.cpp
+++ b/Engine/Source/Runtime/Experimental/Chaos/Private/PhysicsProxy/SingleParticlePhysicsProxy.cpp
@@ -119,6 +119,7 @@ void PushToPhysicsStateImp(const Chaos::FDirtyPropertiesManager& Manager, Chaos:
 		//shape properties
 		bool bUpdateCollisionData = false;
 		bool bHasCollision = false;
+		bool bHasMaterial = false;
 		for(int32 ShapeDataIdx : Dirty.ShapeDataIndices)
 		{
 			const FShapeDirtyData& ShapeData = ShapesData[ShapeDataIdx];
@@ -135,10 +136,15 @@ void PushToPhysicsStateImp(const Chaos::FDirtyPropertiesManager& Manager, Chaos:
 			if(auto NewData = ShapeData.FindMaterials(Manager, ShapeDataIdx))
 			{
 				Handle->ShapesArray()[ShapeIdx]->SetMaterialData(*NewData);
+				bHasMaterial = true;
 			}
-
 		}
 		
+		if (bHasMaterial)
+		{
+			// If materials changed, collisions need to recache their material data
+			Evolution.ParticleMaterialChanged(Handle);
+		}
 
 		if(bUpdateCollisionData && !ForceNoCollisionIntoSQ)
 		{
diff --git a/Engine/Source/Runtime/Experimental/Chaos/Public/Chaos/PBDRigidsEvolutionGBF.h b/Engine/Source/Runtime/Experimental/Chaos/Public/Chaos/PBDRigidsEvolutionGBF.h
--- a/Engine/Source/Runtime/Experimental/Chaos/Public/Chaos/PBDRigidsEvolutionGBF.h
+++ b/Engine/Source/Runtime/Experimental/Chaos/Public/Chaos/PBDRigidsEvolutionGBF.h
@@ -164,6 +164,9 @@ namespace Chaos
 			});
 		}
 
+		// Called when a the material changes one or more shapes on a particle. Required because collisions cache material properties
+		void ParticleMaterialChanged(FGeometryParticleHandle* Particle);
+
 		template<typename TParticleView>
 		void Integrate(const TParticleView& InParticles, FReal Dt)
 		{
