Chaos - workaround and logging for invalid vertex index from GJK

diff --git a/Engine/Source/Runtime/Experimental/Chaos/Private/Chaos/Collision/PBDCollisionSolver.h b/Engine/Source/Runtime/Experimental/Chaos/Private/Chaos/Collision/PBDCollisionSolver.h
--- a/Engine/Source/Runtime/Experimental/Chaos/Private/Chaos/Collision/PBDCollisionSolver.h
+++ b/Engine/Source/Runtime/Experimental/Chaos/Private/Chaos/Collision/PBDCollisionSolver.h
@@ -418,7 +418,7 @@ namespace Chaos
 		// Assume we stay in the friction cone...
 		OutStaticFrictionRatio = FSolverReal(1);
 
-		if (NetPushOutNormal < FSolverReal(UE_SMALL_NUMBER))
+		if (NetPushOutNormal < FSolverReal(UE_KINDA_SMALL_NUMBER))
 		{
 			// Note: we have already added the current iteration's PushOut to the NetPushOut but it has not been applied to the body
 			// so we must subtract it again to calculate the actual pushout we want to undo (i.e., the net pushout that has been applied to the body so far)
diff --git a/Engine/Source/Runtime/Experimental/Chaos/Private/Chaos/CollisionOneShotManifolds.cpp b/Engine/Source/Runtime/Experimental/Chaos/Private/Chaos/CollisionOneShotManifolds.cpp
--- a/Engine/Source/Runtime/Experimental/Chaos/Private/Chaos/CollisionOneShotManifolds.cpp
+++ b/Engine/Source/Runtime/Experimental/Chaos/Private/Chaos/CollisionOneShotManifolds.cpp
@@ -879,6 +879,54 @@ namespace Chaos
 			return BestPlaneIndex;
 		}
 
+		template<typename ConvexImplicitType>
+		bool CheckVertexIndex(const ConvexImplicitType& Convex, const int32 VertexIndex)
+		{
+			ensureMsgf(VertexIndex != INDEX_NONE, TEXT("GJKContactPointMargin invalid vertex index (type:%d) %d"), Convex.GetType(), VertexIndex);
+			return VertexIndex != INDEX_NONE;
+		}
+
+		template<>
+		bool CheckVertexIndex(const FTriangle& Triangle, const int32 VertexIndex)
+		{
+			ensureMsgf(VertexIndex != INDEX_NONE, TEXT("GJKContactPointMargin invalid Triangle VertexIndex %d"), VertexIndex);
+			return VertexIndex != INDEX_NONE;
+		}
+
+		template<>
+		bool CheckVertexIndex(const FImplicitCapsule3& Capsule, const int32 VertexIndex)
+		{
+			ensureMsgf(VertexIndex != INDEX_NONE, TEXT("GJKContactPointMargin invalid Capsule VertexIndex %d"), VertexIndex);
+			return VertexIndex != INDEX_NONE;
+		}
+
+		template<>
+		bool CheckVertexIndex(const FImplicitBox3& Box, const int32 VertexIndex)
+		{
+			ensureMsgf(VertexIndex != INDEX_NONE, TEXT("GJKContactPointMargin invalid Box VertexIndex %d"), VertexIndex);
+			return VertexIndex != INDEX_NONE;
+		}
+
+		template<>
+		bool CheckVertexIndex(const FImplicitConvex3& Convex, const int32 VertexIndex)
+		{
+			ensureMsgf(VertexIndex != INDEX_NONE, TEXT("GJKContactPointMargin invalid Convex VertexIndex %d [%d, %d, %d]"), VertexIndex, Convex.NumVertices(), Convex.NumEdges(), Convex.NumPlanes());
+			return VertexIndex != INDEX_NONE;
+		}
+
+		template<>
+		bool CheckVertexIndex(const TImplicitObjectInstanced<FImplicitConvex3>& Convex, const int32 VertexIndex)
+		{
+			ensureMsgf(VertexIndex != INDEX_NONE, TEXT("GJKContactPointMargin invalid InstancedConvex VertexIndex %d [%d, %d, %d]"), VertexIndex, Convex.NumVertices(), Convex.NumEdges(), Convex.NumPlanes());
+			return VertexIndex != INDEX_NONE;
+		}
+
+		template<>
+		bool CheckVertexIndex(const TImplicitObjectScaled<FImplicitConvex3>& Convex, const int32 VertexIndex)
+		{
+			ensureMsgf(VertexIndex != INDEX_NONE, TEXT("GJKContactPointMargin invalid ScaledConvex VertexIndex %d [%d, %d, %d]"), VertexIndex, Convex.NumVertices(), Convex.NumEdges(), Convex.NumPlanes());
+			return VertexIndex != INDEX_NONE;
+		}
 
 		template <typename ConvexImplicitType1, typename ConvexImplicitType2>
 		void ConstructConvexConvexOneShotManifold(
@@ -912,6 +960,13 @@ namespace Chaos
 			FContactPoint GJKContactPoint = GJKContactPointMargin(Convex1, Convex2, Convex2ToConvex1Transform, Margin1, Margin2, Constraint.GetGJKWarmStartData(), MaxMarginDelta, VertexIndexA, VertexIndexB);
 			PHYSICS_CSV_CUSTOM_EXPENSIVE(PhysicsCounters, NumManifoldsGJKCalled, 1, ECsvCustomStatOp::Accumulate);
 
+			// We are seeing very rare cases where the VertexIndex is not set from an FConvex, even though the convex has vertices
+			// @todo(chaos): track down this problem - there is no path through GJKContactPointMargin (for FConvex) which does not set the VertexIndex
+			if (!CheckVertexIndex(Convex1, VertexIndexA) || !CheckVertexIndex(Convex2, VertexIndexB))
+			{
+				return;
+			}
+
 			const bool bCanUpdateManifold = bChaos_Collision_EnableManifoldGJKReplace;
 			if (bCanUpdateManifold && Constraint.TryAddManifoldContact(GJKContactPoint))
 			{
