Chaos : Flesh
Add float 2*2 matrix and associated tests. Change XPBD corotated and XPBD corotated fiber to use float. Add detailed timing for XPBD corotated. Add PBD altitude springs.

diff --git a/Engine/Source/Runtime/Experimental/Chaos/Public/Chaos/PBDAltitudeSpringConstraints.h b/Engine/Source/Runtime/Experimental/Chaos/Public/Chaos/PBDAltitudeSpringConstraints.h
new file mode 100644
--- /dev/null
+++ b/Engine/Source/Runtime/Experimental/Chaos/Public/Chaos/PBDAltitudeSpringConstraints.h
@@ -0,0 +1,389 @@
+// Copyright Epic Games, Inc. All Rights Reserved.
+#pragma once
+
+#include "Chaos/PBDSpringConstraintsBase.h"
+#include "ChaosStats.h"
+#include "Chaos/ImplicitQRSVD.h"
+#include "Chaos/GraphColoring.h"
+#include "Chaos/Framework/Parallel.h"
+
+
+namespace Chaos::Softs
+{
+
+	class FPBDAltitudeSpringConstraints
+	{
+	public:
+		FPBDAltitudeSpringConstraints(const FSolverParticles& InParticles, TArray<TVec4<int32>>&& InConstraints, const FSolverReal InStiffness = (FSolverReal)1.)
+			: Constraints(InConstraints), Stiffness(InStiffness)
+		{	
+			
+			RestLength.Init((FSolverReal)0., 4 * Constraints.Num());
+			Volumes.Init((FSolverReal)0., Constraints.Num());
+			for (int32 ElementIndex = 0; ElementIndex < Constraints.Num(); ElementIndex++)
+			{	
+				TVec4<int32> Constraint = Constraints[ElementIndex];
+				for (int32 i = 0; i < 4; i++)
+				{
+					const FSolverVec3& P1 = InParticles.X(Constraint[i]);
+					const FSolverVec3& P2 = InParticles.X(Constraint[(i + 1) % 4]);
+					const FSolverVec3& P3 = InParticles.X(Constraint[(i + 2) % 4]);
+					const FSolverVec3& P4 = InParticles.X(Constraint[(i + 3) % 4]);
+					RestLength[4* ElementIndex +i] = FSolverVec3::DotProduct(FSolverVec3::CrossProduct(P3 - P2, P4 - P2), P1 - P2) / FSolverVec3::CrossProduct(P3 - P2, P4 - P2).Size();
+
+				}
+				const FSolverVec3& P1 = InParticles.X(Constraint[0]);
+				const FSolverVec3& P2 = InParticles.X(Constraint[1]);
+				const FSolverVec3& P3 = InParticles.X(Constraint[2]);
+				const FSolverVec3& P4 = InParticles.X(Constraint[3]);
+				Volumes[ElementIndex] = FSolverVec3::DotProduct(FSolverVec3::CrossProduct(P2 - P1, P3 - P1), P4 - P1) / (FSolverReal)6.;
+			}
+			
+			InitColor(InParticles);
+		}
+		virtual ~FPBDAltitudeSpringConstraints() {}
+
+		TVec4<FSolverVec3> GetGradients(const FSolverParticles& InParticles, const int32 ElementIndex, const int32 ie) const
+		{
+			TVec4<FSolverVec3> Grads(FSolverVec3((FSolverReal)0.));;
+			const TVec4<int32>& Constraint = Constraints[ElementIndex];
+			const int32 i1 = Constraint[ie];
+			const int32 i2 = Constraint[(ie + 1) % 4];
+			const int32 i3 = Constraint[(ie + 2) % 4];
+			const int32 i4 = Constraint[(ie + 3) % 4];
+			const FSolverVec3& P1 = InParticles.P(i1);
+			const FSolverVec3& P2 = InParticles.P(i2);
+			const FSolverVec3& P3 = InParticles.P(i3);
+			const FSolverVec3& P4 = InParticles.P(i4);
+			const FSolverVec3 P1P2 = P1 - P2;
+			const FSolverVec3 P3P2 = P3 - P2;
+			const FSolverVec3 P4P2 = P4 - P2;
+
+			FSolverVec3 ScaledBottomNormal = FSolverVec3::CrossProduct(P3 - P2, P4 - P2);
+			FSolverReal ScaledVolume = FSolverVec3::DotProduct(ScaledBottomNormal, P1 - P2);
+			FSolverReal ScaledArea = ScaledBottomNormal.Size();
+			TVec4<FSolverVec3> dAdX(FSolverVec3((FSolverReal)0.));
+			FSolverVec3 ScaleddAdV((FSolverReal)0.);
+			ScaleddAdV[0] = -ScaledBottomNormal[1] * P4P2[2] + ScaledBottomNormal[2] * P4P2[1];
+			ScaleddAdV[1] = ScaledBottomNormal[0] * P4P2[2] - ScaledBottomNormal[2] * P4P2[0];
+			ScaleddAdV[2] = -ScaledBottomNormal[0] * P4P2[1] + ScaledBottomNormal[1] * P4P2[0];
+			FSolverVec3 ScaleddAdW((FSolverReal)0.);
+			ScaleddAdW[0] = ScaledBottomNormal[1] * P3P2[2] - ScaledBottomNormal[2] * P3P2[1];
+			ScaleddAdW[1] = -ScaledBottomNormal[0] * P3P2[2] + ScaledBottomNormal[2] * P3P2[0];
+			ScaleddAdW[2] = ScaledBottomNormal[0] * P3P2[1] - ScaledBottomNormal[1] * P3P2[0];
+			dAdX[2] = (FSolverReal)1. / ScaledArea * ScaleddAdV;
+			dAdX[3] = (FSolverReal)1. / ScaledArea * ScaleddAdW;
+			dAdX[1] = -dAdX[2] - dAdX[3];
+
+			TVec4<FSolverVec3> dVdX(FSolverVec3((FSolverReal)0.));
+			dVdX[0] = FSolverVec3::CrossProduct(P3P2, P4P2);
+			dVdX[2] = FSolverVec3::CrossProduct(P4P2, P1P2);
+			dVdX[3] = FSolverVec3::CrossProduct(P1P2, P3P2);
+			dVdX[1] = -(dVdX[0] + dVdX[2] + dVdX[3]);
+			for (int32 je = 0; je < 4; je++)
+			{
+				Grads[je] = -ScaledVolume / (ScaledArea * ScaledArea) * dAdX[je] + (FSolverReal)1. / ScaledArea * dVdX[je];
+			}
+
+
+			return Grads;
+		}
+
+		FSolverReal GetScalingFactor(const FSolverParticles& InParticles, const int32 ElementIndex, const int32 ie, const TVec4<FSolverVec3>& Grads) const
+		{
+			const TVec4<int32>& Constraint = Constraints[ElementIndex];
+			const int32 i1 = Constraint[ie];
+			const int32 i2 = Constraint[(ie + 1) % 4];
+			const int32 i3 = Constraint[(ie + 2) % 4];
+			const int32 i4 = Constraint[(ie + 3) % 4];
+			const FSolverVec3& P1 = InParticles.P(i1);
+			const FSolverVec3& P2 = InParticles.P(i2);
+			const FSolverVec3& P3 = InParticles.P(i3);
+			const FSolverVec3& P4 = InParticles.P(i4);
+			FSolverVec3 ScaledBottomNormal = FSolverVec3::CrossProduct(P3 - P2, P4 - P2);
+			FSolverReal ScaledVolume = FSolverVec3::DotProduct(ScaledBottomNormal, P1 - P2);
+			FSolverReal CurrentLength = ScaledVolume / ScaledBottomNormal.Size();
+			const FSolverReal S = (CurrentLength - RestLength[4*ElementIndex+ie]) / (
+				InParticles.InvM(i1) * Grads[0].SizeSquared() +
+				InParticles.InvM(i2) * Grads[1].SizeSquared() +
+				InParticles.InvM(i3) * Grads[2].SizeSquared() +
+				InParticles.InvM(i4) * Grads[3].SizeSquared());
+			return Stiffness * S;
+		}
+
+		virtual void ApplySingleElement(FSolverParticles& Particles, const FSolverReal Dt, const int32 ElementIndex) const
+		{
+			TRACE_CPUPROFILER_EVENT_SCOPE(TEXT("STAT_ChaosPBDAltitudeSpringApplySingle"));
+			const TVec4<int32>& Constraint = Constraints[ElementIndex];
+			for (int ie = 0; ie < 4; ++ie)
+			{
+				
+				const int32 i1 = Constraint[ie];
+				const int32 i2 = Constraint[(ie + 1) % 4];
+				const int32 i3 = Constraint[(ie + 2) % 4];
+				const int32 i4 = Constraint[(ie + 3) % 4];
+				const TVec4<FSolverVec3> Grads = GetGradients(Particles, ElementIndex, ie);
+				const FSolverReal S = GetScalingFactor(Particles, ElementIndex, ie, Grads);
+				Particles.P(i1) -= S * Particles.InvM(i1) * Grads[0];
+				Particles.P(i2) -= S * Particles.InvM(i2) * Grads[1];
+				Particles.P(i3) -= S * Particles.InvM(i3) * Grads[2];
+				Particles.P(i4) -= S * Particles.InvM(i4) * Grads[3];
+			}
+		}
+
+		virtual void ApplySingleElementAllInOne(FSolverParticles& Particles, const FSolverReal Dt, const int32 ElementIndex) const
+		{
+			TRACE_CPUPROFILER_EVENT_SCOPE(TEXT("STAT_ChaosPBDAltitudeSpringApplySingle"));
+			const TVec4<int32>& Constraint = Constraints[ElementIndex];
+			for (int ie = 0; ie < 4; ++ie)
+			{
+				TVec4<FSolverVec3> Grads(FSolverVec3((FSolverReal)0.));
+				const int32 i1 = Constraint[ie];
+				const int32 i2 = Constraint[(ie + 1) % 4];
+				const int32 i3 = Constraint[(ie + 2) % 4];
+				const int32 i4 = Constraint[(ie + 3) % 4];
+				const FSolverVec3& P1 = Particles.P(i1);
+				const FSolverVec3& P2 = Particles.P(i2);
+				const FSolverVec3& P3 = Particles.P(i3);
+				const FSolverVec3& P4 = Particles.P(i4);
+				const FSolverVec3 P1P2 = P1 - P2;
+				const FSolverVec3 P3P2 = P3 - P2;
+				const FSolverVec3 P4P2 = P4 - P2;
+
+				FSolverVec3 ScaledBottomNormal = FSolverVec3::CrossProduct(P3 - P2, P4 - P2);
+				FSolverReal ScaledVolume = FSolverVec3::DotProduct(ScaledBottomNormal, P1 - P2);
+				FSolverReal ScaledArea = ScaledBottomNormal.Size();
+				TVec4<FSolverVec3> dAdX(FSolverVec3((FSolverReal)0.));
+				FSolverVec3 ScaleddAdV((FSolverReal)0.);
+				ScaleddAdV[0] = -ScaledBottomNormal[1] * P4P2[2] + ScaledBottomNormal[2] * P4P2[1];
+				ScaleddAdV[1] = ScaledBottomNormal[0] * P4P2[2] - ScaledBottomNormal[2] * P4P2[0];
+				ScaleddAdV[2] = -ScaledBottomNormal[0] * P4P2[1] + ScaledBottomNormal[1] * P4P2[0];
+				FSolverVec3 ScaleddAdW((FSolverReal)0.);
+				ScaleddAdW[0] = ScaledBottomNormal[1] * P3P2[2] - ScaledBottomNormal[2] * P3P2[1];
+				ScaleddAdW[1] = -ScaledBottomNormal[0] * P3P2[2] + ScaledBottomNormal[2] * P3P2[0];
+				ScaleddAdW[2] = ScaledBottomNormal[0] * P3P2[1] - ScaledBottomNormal[1] * P3P2[0];
+				dAdX[2] = (FSolverReal)1. / ScaledArea * ScaleddAdV;
+				dAdX[3] = (FSolverReal)1. / ScaledArea * ScaleddAdW;
+				dAdX[1] = -dAdX[2] - dAdX[3];
+
+				TVec4<FSolverVec3> dVdX(FSolverVec3((FSolverReal)0.));
+				dVdX[0] = FSolverVec3::CrossProduct(P3P2, P4P2);
+				dVdX[2] = FSolverVec3::CrossProduct(P4P2, P1P2);
+				dVdX[3] = FSolverVec3::CrossProduct(P1P2, P3P2);
+				dVdX[1] = -(dVdX[0] + dVdX[2] + dVdX[3]);
+				for (int32 je = 0; je < 4; je++)
+				{
+					Grads[je] = -ScaledVolume / (ScaledArea * ScaledArea) * dAdX[je] + (FSolverReal)1. / ScaledArea * dVdX[je];
+				}
+				FSolverReal CurrentLength = ScaledVolume / ScaledBottomNormal.Size();
+				const FSolverReal S = Stiffness * (CurrentLength - RestLength[4 * ElementIndex + ie]) / (
+					Particles.InvM(i1) * Grads[0].SizeSquared() +
+					Particles.InvM(i2) * Grads[1].SizeSquared() +
+					Particles.InvM(i3) * Grads[2].SizeSquared() +
+					Particles.InvM(i4) * Grads[3].SizeSquared());
+				Particles.P(i1) -= S * Particles.InvM(i1) * Grads[0];
+				Particles.P(i2) -= S * Particles.InvM(i2) * Grads[1];
+				Particles.P(i3) -= S * Particles.InvM(i3) * Grads[2];
+				Particles.P(i4) -= S * Particles.InvM(i4) * Grads[3];
+			}
+		}
+
+		virtual void ApplySingleElementVolumeAllInOne(FSolverParticles& Particles, const FSolverReal Dt, const int32 ElementIndex) const
+		{
+			TVec4<FSolverVec3> Grads;
+			const TVec4<int32>& Constraint = Constraints[ElementIndex];
+			const FSolverVec3& P1 = Particles.P(Constraint[0]);
+			const FSolverVec3& P2 = Particles.P(Constraint[1]);
+			const FSolverVec3& P3 = Particles.P(Constraint[2]);
+			const FSolverVec3& P4 = Particles.P(Constraint[3]);
+			const FSolverVec3 P2P1 = P2 - P1;
+			const FSolverVec3 P4P1 = P4 - P1;
+			const FSolverVec3 P3P1 = P3 - P1;
+			Grads[1] = FSolverVec3::CrossProduct(P3P1, P4P1) / (FSolverReal)6.;
+			Grads[2] = FSolverVec3::CrossProduct(P4P1, P2P1) / (FSolverReal)6.;
+			Grads[3] = FSolverVec3::CrossProduct(P2P1, P3P1) / (FSolverReal)6.;
+			Grads[0] = -(Grads[1] + Grads[2] + Grads[3]);
+			const int32 i1 = Constraint[0];
+			const int32 i2 = Constraint[1];
+			const int32 i3 = Constraint[2];
+			const int32 i4 = Constraint[3];
+
+			const FSolverReal Volume = FSolverVec3::DotProduct(FSolverVec3::CrossProduct(P2 - P1, P3 - P1), P4 - P1) / (FSolverReal)6.;
+			const FSolverReal S = Stiffness*(Volume - Volumes[ElementIndex]) / (
+				Particles.InvM(i1) * Grads[0].SizeSquared() +
+				Particles.InvM(i2) * Grads[1].SizeSquared() +
+				Particles.InvM(i3) * Grads[2].SizeSquared() +
+				Particles.InvM(i4) * Grads[3].SizeSquared());
+			
+			Particles.P(i1) -= S * Particles.InvM(i1) * Grads[0];
+			Particles.P(i2) -= S * Particles.InvM(i2) * Grads[1];
+			Particles.P(i3) -= S * Particles.InvM(i3) * Grads[2];
+			Particles.P(i4) -= S * Particles.InvM(i4) * Grads[3];
+		}
+
+		virtual void ApplySingleElementShortest(FSolverParticles& Particles, const FSolverReal Dt, const int32 ElementIndex) const
+		{
+			TRACE_CPUPROFILER_EVENT_SCOPE(TEXT("STAT_ChaosPBDAltitudeSpringApplySingleShortest"));
+			const TVec4<int32>& Constraint = Constraints[ElementIndex];
+			int32 ie = GetSingleElementShortestLength(Particles, Dt, ElementIndex);
+			
+			const int32 i1 = Constraint[ie];
+			const int32 i2 = Constraint[(ie + 1) % 4];
+			const int32 i3 = Constraint[(ie + 2) % 4];
+			const int32 i4 = Constraint[(ie + 3) % 4];
+			const TVec4<FSolverVec3> Grads = GetGradients(Particles, ElementIndex, ie);
+			const FSolverReal S = GetScalingFactor(Particles, ElementIndex, ie, Grads);
+			Particles.P(i1) -= S * Particles.InvM(i1) * Grads[0];
+			Particles.P(i2) -= S * Particles.InvM(i2) * Grads[1];
+			Particles.P(i3) -= S * Particles.InvM(i3) * Grads[2];
+			Particles.P(i4) -= S * Particles.InvM(i4) * Grads[3];
+
+		}
+
+		int32 GetSingleElementShortestLength(FSolverParticles& Particles, const FSolverReal Dt, const int32 ElementIndex) const 
+		{	
+			int32 ReturnIndex = -1;
+			FSolverReal MinConstraint = TMathUtilConstants<FSolverReal>::MaxReal;
+			const TVec4<int32>& Constraint = Constraints[ElementIndex];
+			for (int ie = 0; ie < 4; ++ie)
+			{
+				const int32 i1 = Constraint[ie];
+				const int32 i2 = Constraint[(ie + 1) % 4];
+				const int32 i3 = Constraint[(ie + 2) % 4];
+				const int32 i4 = Constraint[(ie + 3) % 4];
+				const FSolverVec3& P1 = Particles.P(i1);
+				const FSolverVec3& P2 = Particles.P(i2);
+				const FSolverVec3& P3 = Particles.P(i3);
+				const FSolverVec3& P4 = Particles.P(i4);
+				FSolverVec3 ScaledBottomNormal = FSolverVec3::CrossProduct(P3 - P2, P4 - P2);
+				FSolverReal ScaledVolume = FSolverVec3::DotProduct(ScaledBottomNormal, P1 - P2);
+				FSolverReal CurrentLength = ScaledVolume / ScaledBottomNormal.Size();
+				FSolverReal CurrentConstraint = CurrentLength - RestLength[4 * ElementIndex + ie];
+				if (CurrentConstraint < MinConstraint)
+				{
+					ReturnIndex = ie;
+					MinConstraint = CurrentConstraint;
+				}
+			}
+
+			//another stratey: find the face with largest area
+			/*FSolverReal MaxConstraint = (FSolverReal)-1.;
+			for (int ie = 0; ie < 4; ++ie)
+			{
+				const int32 i1 = Constraint[ie];
+				const int32 i2 = Constraint[(ie + 1) % 4];
+				const int32 i3 = Constraint[(ie + 2) % 4];
+				const int32 i4 = Constraint[(ie + 3) % 4];
+				const FSolverVec3& P1 = Particles.P(i1);
+				const FSolverVec3& P2 = Particles.P(i2);
+				const FSolverVec3& P3 = Particles.P(i3);
+				const FSolverVec3& P4 = Particles.P(i4);
+				FSolverReal ScaledBottomNormalArea = FSolverVec3::CrossProduct(P3 - P2, P4 - P2).Size();
+				if (ScaledBottomNormalArea > MaxConstraint)
+				{
+					ReturnIndex = ie;
+					MaxConstraint = ScaledBottomNormalArea;
+				}
+			}*/
+
+			ensure(ReturnIndex != -1);
+			return ReturnIndex;
+
+		}
+
+		void ApplyInSerial(FSolverParticles& Particles, const FSolverReal Dt) const
+		{
+			//SCOPE_CYCLE_COUNTER(STAT_ChaosXPBDCorotated);
+			TRACE_CPUPROFILER_EVENT_SCOPE(TEXT("STAT_ChaosPBDAltitudeSpringApplySerial"));
+			for (int32 ElementIndex = 0; ElementIndex < Constraints.Num(); ++ElementIndex)
+			{
+				ApplySingleElement(Particles, Dt, ElementIndex);
+				//ApplySingleElementShortest(Particles, Dt, ElementIndex);
+				//ApplySingleElementVolumeAllInOne(Particles, Dt, ElementIndex);
+			}
+		}
+
+		void ApplyInParallel(FSolverParticles& Particles, const FSolverReal Dt) const
+		{
+			{
+				TRACE_CPUPROFILER_EVENT_SCOPE(TEXT("STAT_ChaosPBDAltitudeSpringApply"));
+				if ((ConstraintsPerColorStartIndex.Num() > 1))//&& (MeshConstraints.Num() > Chaos_Spring_ParallelConstraintCount))
+				{
+					const int32 ConstraintColorNum = ConstraintsPerColorStartIndex.Num() - 1;
+
+					for (int32 ConstraintColorIndex = 0; ConstraintColorIndex < ConstraintColorNum; ++ConstraintColorIndex)
+					{
+						const int32 ColorStart = ConstraintsPerColorStartIndex[ConstraintColorIndex];
+						const int32 ColorSize = ConstraintsPerColorStartIndex[ConstraintColorIndex + 1] - ColorStart;
+						PhysicsParallelFor(ColorSize, [&](const int32 Index)
+							{
+								const int32 ConstraintIndex = ColorStart + Index;
+								ApplySingleElement(Particles, Dt, ConstraintIndex); 
+								//ApplySingleElementShortest(Particles, Dt, ConstraintIndex);
+								//ApplySingleElementVolumeAllInOne(Particles, Dt, ConstraintIndex);
+								//ApplySingleElementAllInOne(Particles, Dt, ConstraintIndex);
+							});
+					}
+				}
+			}
+		}
+
+	protected:
+		TArray<TVec4<int32>> Constraints;
+
+	protected:
+
+
+	private:
+		void InitColor(const FSolverParticles& Particles)
+		{
+
+			{
+				const TArray<TArray<int32>> ConstraintsPerColor = FGraphColoring::ComputeGraphColoring(Constraints, Particles);
+
+				// Reorder constraints based on color so each array in ConstraintsPerColor contains contiguous elements.
+				TArray<TVec4<int32>> ReorderedConstraints;
+				TArray<FSolverReal> ReorderedRestLength;
+				TArray<FSolverReal> ReorderedVolumes;
+				TArray<int32> OrigToReorderedIndices; // used to reorder stiffness indices
+				ReorderedConstraints.SetNumUninitialized(Constraints.Num());
+				ReorderedRestLength.SetNumUninitialized(RestLength.Num());
+				ReorderedVolumes.SetNumUninitialized(Volumes.Num());
+				OrigToReorderedIndices.SetNumUninitialized(Constraints.Num());
+
+				ConstraintsPerColorStartIndex.Reset(ConstraintsPerColor.Num() + 1);
+
+				int32 ReorderedIndex = 0;
+				for (const TArray<int32>& ConstraintsBatch : ConstraintsPerColor)
+				{
+					ConstraintsPerColorStartIndex.Add(ReorderedIndex);
+					for (const int32& BatchConstraint : ConstraintsBatch)
+					{
+						const int32 OrigIndex = BatchConstraint;
+						ReorderedConstraints[ReorderedIndex] = Constraints[OrigIndex];
+						//ReorderedMeasure[ReorderedIndex] = Measure[OrigIndex];
+						for (int32 kk = 0; kk < 4; kk++)
+						{
+							ReorderedRestLength[4 * ReorderedIndex + kk] = RestLength[4 * OrigIndex + kk];
+						}
+						ReorderedVolumes[ReorderedIndex] = Volumes[OrigIndex];
+						OrigToReorderedIndices[OrigIndex] = ReorderedIndex;
+
+						++ReorderedIndex;
+					}
+				}
+				ConstraintsPerColorStartIndex.Add(ReorderedIndex);
+
+				Constraints = MoveTemp(ReorderedConstraints);
+				RestLength = MoveTemp(ReorderedRestLength);
+				Volumes = MoveTemp(ReorderedVolumes);
+			}
+		}
+
+		TArray<FSolverReal> RestLength;
+		TArray<FSolverReal> Volumes;
+		FSolverReal Stiffness;
+		TArray<int32> ConstraintsPerColorStartIndex;
+	};
+
+}  // End namespace Chaos::Softs
diff --git a/Engine/Source/Runtime/Experimental/Chaos/Public/Chaos/XPBDCorotatedConstraints.h b/Engine/Source/Runtime/Experimental/Chaos/Public/Chaos/XPBDCorotatedConstraints.h
--- a/Engine/Source/Runtime/Experimental/Chaos/Public/Chaos/XPBDCorotatedConstraints.h
+++ b/Engine/Source/Runtime/Experimental/Chaos/Public/Chaos/XPBDCorotatedConstraints.h
@@ -8,6 +8,8 @@
 #include "Chaos/Framework/Parallel.h"
 
 DECLARE_CYCLE_STAT(TEXT("Chaos XPBD Corotated Constraint"), STAT_ChaosXPBDCorotated, STATGROUP_Chaos);
+DECLARE_CYCLE_STAT(TEXT("Chaos XPBD Corotated Constraint Polar Compute"), STAT_ChaosXPBDCorotatedPolar, STATGROUP_Chaos);
+DECLARE_CYCLE_STAT(TEXT("Chaos XPBD Corotated Constraint Det Compute"), STAT_ChaosXPBDCorotatedDet, STATGROUP_Chaos);
 
 namespace Chaos::Softs
 {
@@ -102,11 +104,12 @@ namespace Chaos::Softs
 		
 		void Init() const 
 		{
-			for (FReal& Lambdas : LambdaArray) { Lambdas = (FReal)0.; }
+			for (FSolverReal& Lambdas : LambdaArray) { Lambdas = (FSolverReal)0.; }
 		}
 
 		virtual void ApplyInSerial(FSolverParticles& Particles, const FSolverReal Dt, const int32 ElementIndex) const
 		{
+			TRACE_CPUPROFILER_EVENT_SCOPE(TEXT("STAT_ChaosXPBDCorotatedApplySingle"));
 			TVec4<FSolverVec3> PolarDelta = GetPolarDelta(Particles, Dt, ElementIndex);
 
 			for (int i = 0; i < 4; i++) 
@@ -127,7 +130,7 @@ namespace Chaos::Softs
 		void ApplyInSerial(FSolverParticles& Particles, const FSolverReal Dt) const
 		{
 			SCOPE_CYCLE_COUNTER(STAT_ChaosXPBDCorotated);
-
+			TRACE_CPUPROFILER_EVENT_SCOPE(TEXT("STAT_ChaosXPBDCorotatedApplySerial"));
 			for (int32 ElementIndex = 0; ElementIndex < MeshConstraints.Num(); ++ElementIndex)
 			{
 				ApplyInSerial(Particles, Dt, ElementIndex);
@@ -145,21 +148,23 @@ namespace Chaos::Softs
 				HErrorArray.Init((FSolverReal)0., 2 * MeshConstraints.Num());
 			}
 			
-
-			SCOPE_CYCLE_COUNTER(STAT_ChaosXPBDCorotated);
-			if ((ConstraintsPerColorStartIndex.Num() > 1) )//&& (MeshConstraints.Num() > Chaos_Spring_ParallelConstraintCount))
 			{
-				const int32 ConstraintColorNum = ConstraintsPerColorStartIndex.Num() - 1;
-
-				for (int32 ConstraintColorIndex = 0; ConstraintColorIndex < ConstraintColorNum; ++ConstraintColorIndex)
+				SCOPE_CYCLE_COUNTER(STAT_ChaosXPBDCorotated);
+				TRACE_CPUPROFILER_EVENT_SCOPE(TEXT("STAT_ChaosXPBDCorotatedApply"));
+				if ((ConstraintsPerColorStartIndex.Num() > 1))//&& (MeshConstraints.Num() > Chaos_Spring_ParallelConstraintCount))
 				{
-					const int32 ColorStart = ConstraintsPerColorStartIndex[ConstraintColorIndex];
-					const int32 ColorSize = ConstraintsPerColorStartIndex[ConstraintColorIndex + 1] - ColorStart;
-					PhysicsParallelFor(ColorSize, [&](const int32 Index)
-						{
-							const int32 ConstraintIndex = ColorStart + Index;
-							ApplyInSerial(Particles, Dt, ConstraintIndex);
-						});
+					const int32 ConstraintColorNum = ConstraintsPerColorStartIndex.Num() - 1;
+
+					for (int32 ConstraintColorIndex = 0; ConstraintColorIndex < ConstraintColorNum; ++ConstraintColorIndex)
+					{
+						const int32 ColorStart = ConstraintsPerColorStartIndex[ConstraintColorIndex];
+						const int32 ColorSize = ConstraintsPerColorStartIndex[ConstraintColorIndex + 1] - ColorStart;
+						PhysicsParallelFor(ColorSize, [&](const int32 Index)
+							{
+								const int32 ConstraintIndex = ColorStart + Index;
+								ApplyInSerial(Particles, Dt, ConstraintIndex);
+							});
+					}
 				}
 			}
 
@@ -176,8 +181,8 @@ namespace Chaos::Softs
 
 				// Reorder constraints based on color so each array in ConstraintsPerColor contains contiguous elements.
 				TArray<TVec4<int32>> ReorderedConstraints;
-				TArray<FReal> ReorderedMeasure;
-				TArray<FReal> ReorderedDmInverse;
+				TArray<FSolverReal> ReorderedMeasure;
+				TArray<FSolverReal> ReorderedDmInverse;
 				TArray<int32> OrigToReorderedIndices; // used to reorder stiffness indices
 				ReorderedConstraints.SetNumUninitialized(MeshConstraints.Num());
 				ReorderedMeasure.SetNumUninitialized(Measure.Num());
@@ -215,21 +220,23 @@ namespace Chaos::Softs
 	protected:
 
 
-		TVec4<FSolverVec3> GetDeterminantDelta(const FSolverParticles& Particles, const FReal Dt, const int32 ElementIndex, const FReal Tol = 1e-3) const
+		TVec4<FSolverVec3> GetDeterminantDelta(const FSolverParticles& Particles, const FSolverReal Dt, const int32 ElementIndex, const FSolverReal Tol = 1e-3) const
 		{
+			TRACE_CPUPROFILER_EVENT_SCOPE(TEXT("STAT_ChaosXPBDCorotatedApplyDet"));
+			SCOPE_CYCLE_COUNTER(STAT_ChaosXPBDCorotatedDet);
 			const TVec4<int32>& Constraint = MeshConstraints[ElementIndex];
 
-			const PMatrix<FReal, 3, 3> Fe = F(ElementIndex, Particles);
-			PMatrix<FReal, 3, 3> DmInvT = ElementDmInv(ElementIndex).GetTransposed();
-			TVec4<FSolverVec3> dC2(FSolverVec3((FReal)0.));
+			const PMatrix<FSolverReal, 3, 3> Fe = F(ElementIndex, Particles);
+			PMatrix<FSolverReal, 3, 3> DmInvT = ElementDmInv(ElementIndex).GetTransposed();
+			TVec4<FSolverVec3> dC2(FSolverVec3((FSolverReal)0.));
 
-			FReal J = Fe.Determinant();
+			FSolverReal J = Fe.Determinant();
 			if (J - 1 < Tol)
 			{
-				return TVec4<FSolverVec3>(FSolverVec3((FReal)0.));
+				return TVec4<FSolverVec3>(FSolverVec3((FSolverReal)0.));
 			}
 
-			PMatrix<FReal, 3, 3> JFinvT;
+			PMatrix<FSolverReal, 3, 3> JFinvT;
 			JFinvT.SetAt(0, 0, Fe.GetAt(1, 1) * Fe.GetAt(2, 2) - Fe.GetAt(2, 1) * Fe.GetAt(1, 2));
 			JFinvT.SetAt(0, 1, Fe.GetAt(2, 0) * Fe.GetAt(1, 2) - Fe.GetAt(1, 0) * Fe.GetAt(2, 2));
 			JFinvT.SetAt(0, 2, Fe.GetAt(1, 0) * Fe.GetAt(2, 1) - Fe.GetAt(2, 0) * Fe.GetAt(1, 1));
@@ -240,7 +247,7 @@ namespace Chaos::Softs
 			JFinvT.SetAt(2, 1, Fe.GetAt(1, 0) * Fe.GetAt(0, 2) - Fe.GetAt(0, 0) * Fe.GetAt(1, 2));
 			JFinvT.SetAt(2, 2, Fe.GetAt(0, 0) * Fe.GetAt(1, 1) - Fe.GetAt(1, 0) * Fe.GetAt(0, 1));
 
-			PMatrix<FReal, 3, 3> JinvTDmInvT = DmInvT * JFinvT;
+			PMatrix<FSolverReal, 3, 3> JinvTDmInvT = DmInvT * JFinvT;
 
 			for (int ie = 0; ie < 3; ie++) {
 				for (int alpha = 0; alpha < 3; alpha++) {
@@ -253,7 +260,7 @@ namespace Chaos::Softs
 				}
 			}
 
-			FReal AlphaTilde = (FReal)2. / (Dt * Dt * Lambda * Measure[ElementIndex]);
+			FSolverReal AlphaTilde = (FSolverReal)2. / (Dt * Dt * Lambda * Measure[ElementIndex]);
 
 			if (bRecordMetric)
 			{
@@ -269,9 +276,9 @@ namespace Chaos::Softs
 			}
 			
 
-			FReal DLambda = (1 - J) - AlphaTilde * LambdaArray[2 * ElementIndex + 1];
+			FSolverReal DLambda = (1 - J) - AlphaTilde * LambdaArray[2 * ElementIndex + 1];
 
-			FReal Denom = AlphaTilde;
+			FSolverReal Denom = AlphaTilde;
 			for (int i = 0; i < 4; i++)
 			{
 				for (int j = 0; j < 3; j++)
@@ -281,7 +288,7 @@ namespace Chaos::Softs
 			}
 			DLambda /= Denom;
 			LambdaArray[2 * ElementIndex + 1] += DLambda;
-			TVec4<FSolverVec3> Delta(FSolverVec3((FReal)0.));
+			TVec4<FSolverVec3> Delta(FSolverVec3((FSolverReal)0.));
 			for (int i = 0; i < 4; i++)
 			{
 				for (int j = 0; j < 3; j++)
@@ -292,17 +299,17 @@ namespace Chaos::Softs
 			return Delta;
 		}
 
-		TVec4<FSolverVec3> GetPolarDelta(const FSolverParticles& Particles, const FReal Dt, const int32 ElementIndex, const FReal Tol = 1e-3) const
-		{
-
-
-			const PMatrix<FReal, 3, 3> Fe = F(ElementIndex, Particles);
+		TVec4<FSolverVec3> GetPolarDelta(const FSolverParticles& Particles, const FSolverReal Dt, const int32 ElementIndex, const FSolverReal Tol = 1e-3) const
+		{	
+			TRACE_CPUPROFILER_EVENT_SCOPE(TEXT("STAT_ChaosXPBDCorotatedApplyPolar"));
+			SCOPE_CYCLE_COUNTER(STAT_ChaosXPBDCorotatedPolar);
+			const PMatrix<FSolverReal, 3, 3> Fe = F(ElementIndex, Particles);
 
-			PMatrix<FReal, 3, 3> Re((FReal)0.), Se((FReal)0.);
+			PMatrix<FSolverReal, 3, 3> Re((FSolverReal)0.), Se((FSolverReal)0.);
 
 			Chaos::PolarDecomposition(Fe, Re, Se);
 
-			FReal C1 = (FReal)0.;
+			FSolverReal C1 = (FSolverReal)0.;
 			for (int i = 0; i < 3; i++)
 			{
 				for (int j = 0; j < 3; j++)
@@ -314,16 +321,16 @@ namespace Chaos::Softs
 
 			if (C1 < Tol)
 			{
-				return TVec4<FSolverVec3>(FSolverVec3((FReal)0.));
+				return TVec4<FSolverVec3>(FSolverVec3((FSolverReal)0.));
 			}
 
-			TVector<FReal, 81> dRdF((FReal)0.);
+			TVector<FSolverReal, 81> dRdF((FSolverReal)0.);
 			Chaos::dRdFCorotated(Fe, dRdF);
 
-			PMatrix<FReal, 3, 3> DmInvT = ElementDmInv(ElementIndex).GetTransposed();
+			PMatrix<FSolverReal, 3, 3> DmInvT = ElementDmInv(ElementIndex).GetTransposed();
 
 			//TODO: deifnitely test the initialization
-			TVec4<FSolverVec3> dC1(FSolverVec3((FReal)0.));
+			TVec4<FSolverVec3> dC1(FSolverVec3((FSolverReal)0.));
 			//dC1 = dC1dF * dFdX
 			for (int alpha = 0; alpha < 3; alpha++) {
 				for (int l = 0; l < 3; l++) {
@@ -373,7 +380,7 @@ namespace Chaos::Softs
 				}
 			}
 
-			FReal AlphaTilde = (FReal)1. / (Dt * Dt * Mu * Measure[ElementIndex]);
+			FSolverReal AlphaTilde = (FSolverReal)1. / (Dt * Dt * Mu * Measure[ElementIndex]);
 
 			if (bRecordMetric)
 			{
@@ -388,9 +395,9 @@ namespace Chaos::Softs
 				HErrorArray[2 * ElementIndex + 0] = C1 + AlphaTilde * LambdaArray[2 * ElementIndex + 0];
 			}
 
-			FReal DLambda = -C1 - AlphaTilde* LambdaArray[2 * ElementIndex + 0];
+			FSolverReal DLambda = -C1 - AlphaTilde* LambdaArray[2 * ElementIndex + 0];
 
-			FReal Denom = AlphaTilde;
+			FSolverReal Denom = AlphaTilde;
 			for (int i = 0; i < 4; i++)
 			{
 				for (int j = 0; j < 3; j++)
@@ -400,7 +407,7 @@ namespace Chaos::Softs
 			}
 			DLambda /= Denom;
 			LambdaArray[2 * ElementIndex + 0] += DLambda;
-			TVec4<FSolverVec3> Delta(FSolverVec3((FReal)0.));
+			TVec4<FSolverVec3> Delta(FSolverVec3((FSolverReal)0.));
 			for (int i = 0; i < 4; i++) 
 			{
 				for (int j = 0; j < 3; j++)
@@ -414,18 +421,18 @@ namespace Chaos::Softs
 
 
 	protected:
-		mutable TArray<FReal> LambdaArray;
-		mutable TArray<FReal> DmInverse;
+		mutable TArray<FSolverReal> LambdaArray;
+		mutable TArray<FSolverReal> DmInverse;
 
 		//material constants calculated from E:
-		FReal Mu;
-		FReal Lambda;
+		FSolverReal Mu;
+		FSolverReal Lambda;
 		mutable FSolverReal HError;
 		mutable TArray<FSolverReal> HErrorArray;
 		bool bRecordMetric;
 
 		TArray<TVector<int32, 4>> MeshConstraints;
-		mutable TArray<FReal> Measure;
+		mutable TArray<FSolverReal> Measure;
 		FSolverParticles RestParticles;
 		TArray<int32> ConstraintsPerColorStartIndex; // Constraints are ordered so each batch is contiguous. This is ColorNum + 1 length so it can be used as start and end.
 		mutable TArray<FSolverReal> GError;
diff --git a/Engine/Source/Runtime/Experimental/Chaos/Public/Chaos/XPBDCorotatedFiberConstraints.h b/Engine/Source/Runtime/Experimental/Chaos/Public/Chaos/XPBDCorotatedFiberConstraints.h
--- a/Engine/Source/Runtime/Experimental/Chaos/Public/Chaos/XPBDCorotatedFiberConstraints.h
+++ b/Engine/Source/Runtime/Experimental/Chaos/Public/Chaos/XPBDCorotatedFiberConstraints.h
@@ -29,8 +29,8 @@ namespace Chaos::Softs
 			int32 ParticleCount,
 			const TArray<TVector<int32, 4>>& InMesh,
 			const bool bRecordMetricIn = true,
-			const FSolverReal& EMesh = (FReal)10.0,
-			const FSolverReal& NuMesh = (FReal).3,
+			const FSolverReal& EMesh = (FSolverReal)10.0,
+			const FSolverReal& NuMesh = (FSolverReal).3,
 			const FSolverVec3 InFiberDir = FSolverVec3((FSolverReal)1., (FSolverReal)0., (FSolverReal)0.),
 			const FSolverReal InSigmaMax = (FSolverReal)3e5
 		)
@@ -48,7 +48,7 @@ namespace Chaos::Softs
 			{
 				CurrentTime -= FinalTime;
 			}
-			AlphaActivation = (FReal)1. - (FReal)4. / FinalTime * FMath::Abs(CurrentTime - FinalTime / (FReal)2.);
+			AlphaActivation = (FSolverReal)1. - (FSolverReal)4. / FinalTime * FMath::Abs(CurrentTime - FinalTime / (FSolverReal)2.);
 		}
 
 		virtual void ApplyInSerial(FSolverParticles& Particles, const FSolverReal Dt, const int32 ElementIndex) const override
@@ -83,9 +83,9 @@ namespace Chaos::Softs
 
 		TVec4<FSolverVec3> GetFiberDelta(const FSolverParticles& Particles, const FSolverReal Dt, const int32 ElementIndex, const FSolverReal Tol = 1e-3) const
 		{
-			const PMatrix<FReal, 3, 3> Fe = F(ElementIndex, Particles);
+			const PMatrix<FSolverReal, 3, 3> Fe = F(ElementIndex, Particles);
 
-			PMatrix<FReal, 3, 3> Re((FReal)0.), Se((FReal)0.);
+			PMatrix<FSolverReal, 3, 3> Re((FSolverReal)0.), Se((FSolverReal)0.);
 
 			Chaos::PolarDecomposition(Fe, Re, Se);
 
@@ -93,7 +93,7 @@ namespace Chaos::Softs
 			FSolverVec3 FeV = Fe.GetTransposed() * FiberDir;
 			FSolverVec3 DmInverseV = ElementDmInv(ElementIndex).GetTransposed() * FiberDir;
 			FSolverReal L = FeV.Size();
-			TVec4<FSolverVec3> dLdX(FSolverVec3((FReal)0.));
+			TVec4<FSolverVec3> dLdX(FSolverVec3((FSolverReal)0.));
 			for (int32 alpha = 0; alpha < 3; alpha++)
 			{
 				for (int32 s = 0; s < 3; s++)
@@ -110,49 +110,49 @@ namespace Chaos::Softs
 				}
 			}
 
-			FReal LambdaOFL = (FReal)1.4;
-			FReal P1 = (FReal)0.05;
-			FReal P2 = (FReal)6.6;
-			FReal FpIntegral = (FReal)0.;
-			FReal FaIntegral = (FReal)0.;
+			FSolverReal LambdaOFL = (FSolverReal)1.4;
+			FSolverReal P1 = (FSolverReal)0.05;
+			FSolverReal P2 = (FSolverReal)6.6;
+			FSolverReal FpIntegral = (FSolverReal)0.;
+			FSolverReal FaIntegral = (FSolverReal)0.;
 
-			FReal C3 = (FReal)0.;
-			FReal AlphaTilde = LambdaOFL / (SigmaMax * Dt * Dt * Measure[ElementIndex]);
-			FReal dFpdL = (FReal)0.;
-			FReal dFadL = (FReal)0.;
+			FSolverReal C3 = (FSolverReal)0.;
+			FSolverReal AlphaTilde = LambdaOFL / (SigmaMax * Dt * Dt * Measure[ElementIndex]);
+			FSolverReal dFpdL = (FSolverReal)0.;
+			FSolverReal dFadL = (FSolverReal)0.;
 
 			if (L > LambdaOFL)
 			{
-				FpIntegral = P1 * LambdaOFL / P2 * FMath::Exp(P2 * (L / LambdaOFL - (FReal)1.)) - P1 * (L - LambdaOFL);
-				dFpdL = P1 * FMath::Exp(P2 * (L / LambdaOFL - (FReal)1.)) - P1;
+				FpIntegral = P1 * LambdaOFL / P2 * FMath::Exp(P2 * (L / LambdaOFL - (FSolverReal)1.)) - P1 * (L - LambdaOFL);
+				dFpdL = P1 * FMath::Exp(P2 * (L / LambdaOFL - (FSolverReal)1.)) - P1;
 			}
-			if (L > (FReal)0.4 * LambdaOFL && L < (FReal)0.6 * LambdaOFL)
+			if (L > (FSolverReal)0.4 * LambdaOFL && L < (FSolverReal)0.6 * LambdaOFL)
 			{
-				FaIntegral = (FReal)3. * LambdaOFL * FMath::Pow((L / LambdaOFL - (FReal)0.4), 3);
-				dFadL = (FReal)9. * FMath::Pow((L / LambdaOFL - (FReal)0.4), 2);
+				FaIntegral = (FSolverReal)3. * LambdaOFL * FMath::Pow((L / LambdaOFL - (FSolverReal)0.4), 3);
+				dFadL = (FSolverReal)9. * FMath::Pow((L / LambdaOFL - (FSolverReal)0.4), 2);
 			}
-			else if (L >= (FReal)0.6 * LambdaOFL && L <= (FReal)1.4 * LambdaOFL)
+			else if (L >= (FSolverReal)0.6 * LambdaOFL && L <= (FSolverReal)1.4 * LambdaOFL)
 			{
-				FaIntegral = (FReal)3. * LambdaOFL * (FReal)0.008 + L - (FReal)4. / (FReal)3. * LambdaOFL * FMath::Pow(L / LambdaOFL - (FReal)1., 3) - (FReal)0.6 * LambdaOFL - (FReal)4. / (FReal)3. * LambdaOFL * FMath::Pow((FReal)0.4, 3);
-				dFadL = (FReal)1. - (FReal)4. * FMath::Pow(L / LambdaOFL - (FReal)1., 2);
+				FaIntegral = (FSolverReal)3. * LambdaOFL * (FSolverReal)0.008 + L - (FSolverReal)4. / (FSolverReal)3. * LambdaOFL * FMath::Pow(L / LambdaOFL - (FSolverReal)1., 3) - (FSolverReal)0.6 * LambdaOFL - (FSolverReal)4. / (FSolverReal)3. * LambdaOFL * FMath::Pow((FSolverReal)0.4, 3);
+				dFadL = (FSolverReal)1. - (FSolverReal)4. * FMath::Pow(L / LambdaOFL - (FSolverReal)1., 2);
 			}
-			else if (L > (FReal)1.4 * LambdaOFL && L <= (FReal)1.6 * LambdaOFL)
+			else if (L > (FSolverReal)1.4 * LambdaOFL && L <= (FSolverReal)1.6 * LambdaOFL)
 			{
-				FaIntegral = (FReal)3. * LambdaOFL * (FReal)0.008 + (FReal)0.8 * LambdaOFL - (FReal)8. / (FReal)3. * LambdaOFL * FMath::Pow((FReal)0.4, 3) + (FReal)3. * LambdaOFL * FMath::Pow(L / LambdaOFL - (FReal)1.6, 3) + (FReal)3. * LambdaOFL * (FReal)0.008;
-				dFadL = (FReal)9. * FMath::Pow((L / LambdaOFL - (FReal)1.6), 2);
+				FaIntegral = (FSolverReal)3. * LambdaOFL * (FSolverReal)0.008 + (FSolverReal)0.8 * LambdaOFL - (FSolverReal)8. / (FSolverReal)3. * LambdaOFL * FMath::Pow((FSolverReal)0.4, 3) + (FSolverReal)3. * LambdaOFL * FMath::Pow(L / LambdaOFL - (FSolverReal)1.6, 3) + (FSolverReal)3. * LambdaOFL * (FSolverReal)0.008;
+				dFadL = (FSolverReal)9. * FMath::Pow((L / LambdaOFL - (FSolverReal)1.6), 2);
 
 			}
-			else if (L > (FReal)1.6 * LambdaOFL)
+			else if (L > (FSolverReal)1.6 * LambdaOFL)
 			{
-				FaIntegral = (FReal)3. * LambdaOFL * (FReal)0.008 + (FReal)0.8 * LambdaOFL - (FReal)8. / (FReal)3. * LambdaOFL * FMath::Pow((FReal)0.4, 3) + (FReal)3. * LambdaOFL * (FReal)0.008;
-				dFadL = (FReal)0.;
+				FaIntegral = (FSolverReal)3. * LambdaOFL * (FSolverReal)0.008 + (FSolverReal)0.8 * LambdaOFL - (FSolverReal)8. / (FSolverReal)3. * LambdaOFL * FMath::Pow((FSolverReal)0.4, 3) + (FSolverReal)3. * LambdaOFL * (FSolverReal)0.008;
+				dFadL = (FSolverReal)0.;
 			}
 
 
 			C3 = FMath::Sqrt(FpIntegral + AlphaActivation * FaIntegral);
 
 
-			FReal dC3dL = (FReal)0.5 * (dFpdL + AlphaActivation * dFadL);
+			FSolverReal dC3dL = (FSolverReal)0.5 * (dFpdL + AlphaActivation * dFadL);
 			if (C3 == 0)
 			{
 				return TVec4<FSolverVec3>(FSolverVec3((FSolverReal)0.));
@@ -162,7 +162,7 @@ namespace Chaos::Softs
 				dC3dL /= C3;
 			}
 
-			TVec4<FSolverVec3> dC3(FSolverVec3((FReal)0.));
+			TVec4<FSolverVec3> dC3(FSolverVec3((FSolverReal)0.));
 			for (int32 i = 0; i < 4; i++)
 			{
 				for (int32 j = 0; j < 3; j++)
@@ -171,9 +171,9 @@ namespace Chaos::Softs
 				}
 			}
 
-			FReal DLambda = -C3 - AlphaTilde * LambdaArray[2 * ElementIndex + 2];
+			FSolverReal DLambda = -C3 - AlphaTilde * LambdaArray[2 * ElementIndex + 2];
 
-			FReal Denom = AlphaTilde;
+			FSolverReal Denom = AlphaTilde;
 			for (int i = 0; i < 4; i++)
 			{
 				for (int j = 0; j < 3; j++)
@@ -183,7 +183,7 @@ namespace Chaos::Softs
 			}
 			DLambda /= Denom;
 			LambdaArray[2 * ElementIndex + 2] += DLambda;
-			TVec4<FSolverVec3> Delta(FSolverVec3((FReal)0.));
+			TVec4<FSolverVec3> Delta(FSolverVec3((FSolverReal)0.));
 			for (int i = 0; i < 4; i++)
 			{
 				for (int j = 0; j < 3; j++)
@@ -200,8 +200,8 @@ namespace Chaos::Softs
 	private:
 
 		//material constants calculated from E:
-		FReal SigmaMax;
-		mutable FReal AlphaActivation;
+		FSolverReal SigmaMax;
+		mutable FSolverReal AlphaActivation;
 		FSolverVec3 FiberDir;
 
 	};
diff --git a/Engine/Source/Runtime/Experimental/ChaosCore/Public/Chaos/Matrix.h b/Engine/Source/Runtime/Experimental/ChaosCore/Public/Chaos/Matrix.h
--- a/Engine/Source/Runtime/Experimental/ChaosCore/Public/Chaos/Matrix.h
+++ b/Engine/Source/Runtime/Experimental/ChaosCore/Public/Chaos/Matrix.h
@@ -143,6 +143,55 @@ namespace Chaos
 		}
 	};
 
+	template<>
+	class PMatrix<FRealSingle, 2, 2>
+	{
+	public:
+		FRealSingle M[4];
+
+		PMatrix(const FRealSingle x00, const FRealSingle x10, const FRealSingle x01, const FRealSingle x11)
+		{
+			M[0] = x00;
+			M[1] = x10;
+			M[2] = x01;
+			M[3] = x11;
+		}
+
+		PMatrix(const FRealSingle x00, const FRealSingle x10, const FRealSingle x11)
+		{
+			M[0] = x00;
+			M[1] = x10;
+			M[2] = x10;
+			M[3] = x11;
+		}
+
+		PMatrix<FRealSingle, 2, 2> SubtractDiagonal(const FRealSingle Scalar) const
+		{
+			return PMatrix<FRealSingle, 2, 2>(
+				M[0] - Scalar,
+				M[1],
+				M[2],
+				M[3] - Scalar);
+		}
+
+		TVector<FRealSingle, 2> TransformPosition(const TVector<FRealSingle, 2>& Other) const
+		{
+			return TVector<FRealSingle, 2>(
+				M[0] * Other.X + M[2] * Other.Y,
+				M[1] * Other.X + M[3] * Other.Y);
+		}
+
+		PMatrix<FRealSingle, 2, 2> Inverse() const
+		{
+			const FRealSingle OneOverDeterminant = static_cast<FRealSingle>(1.0) / (M[0] * M[3] - M[1] * M[2]);
+			return PMatrix<FRealSingle, 2, 2>(
+				OneOverDeterminant * M[3],
+				-OneOverDeterminant * M[1],
+				-OneOverDeterminant * M[2],
+				OneOverDeterminant * M[0]);
+		}
+	};
+
 	template<>
 	class PMatrix<FRealSingle, 4, 4> : public UE::Math::TMatrix<FRealSingle>
 	{
