Chaos Cloth: Update PBDStiffness so that exponential "ParameterFit" is 0 when Value = 0 and 1 when Value = 1.

This is just a slight scaling and translation from the original ParameterFit function. (Previous range was [1e-3, 1] instead of  [0,1]). There is no upgrade between old and new function because the difference largest at InValue = 0 (where ParameterFitOld = 1e-3 and ParameterFitNew = 0), and behavior differences are small at this end because the curve is so flat (with the exception of seeing a difference between true 0 and anything > 0).  I did some toggle tests between the two functions, and they didn't make a perceptable difference. I recognize this also means that maybe this change is pointless, but it makes the exponential stiffness value smooth rather than having a small discontinuity near zero with the way we clip values which feels like a good thing in general. (New function is 7e-11 at the clipping point vs 1e-3).

Also, I noticed in testing that when I added the Buckling, it caused BendingSprings to be created when BendingStiffness == 0 since BucklingStiffness (unused by BendingSprings) is defaulted to 1. This has been fixed in this CL.

Sidenote: XPBDStiffness suffers from similar issues of 0 not mapping to 0. However, we actually always use Alpha = 1 / (Stiffness * Dt * Dt), so additional changes are necessary to handle true 0. I am not making those changes now because we'd like to rework XPBDStiffness units soon and XPBD is still not actually exposed in the UI.

diff --git a/Engine/Plugins/Experimental/ChaosCloth/Source/Chaos/Private/ChaosCloth/ChaosClothingSimulationCloth.cpp b/Engine/Plugins/Experimental/ChaosCloth/Source/Chaos/Private/ChaosCloth/ChaosClothingSimulationCloth.cpp
--- a/Engine/Plugins/Experimental/ChaosCloth/Source/Chaos/Private/ChaosCloth/ChaosClothingSimulationCloth.cpp
+++ b/Engine/Plugins/Experimental/ChaosCloth/Source/Chaos/Private/ChaosCloth/ChaosClothingSimulationCloth.cpp
@@ -135,7 +135,7 @@ void FClothingSimulationCloth::FLODData::Add(FClothingSimulationSolver* Solver,
 	const TConstArrayView<FRealSingle>& BendingStiffnessMultipliers = WeightMaps[(int32)EChaosWeightMapTarget::BendingStiffness];
 	const TConstArrayView<FRealSingle>& BucklingStiffnessMultipliers = WeightMaps[(int32)EChaosWeightMapTarget::BucklingStiffness];
 	if (Cloth->BendingStiffness[0] > (FRealSingle)0. || (Cloth->BendingStiffness[1] > (FRealSingle)0. && BendingStiffnessMultipliers.Num() == NumParticles) ||
-		Cloth->BucklingStiffness[0] > (FRealSingle)0. || (Cloth->BucklingStiffness[1] > (FRealSingle)0. && BucklingStiffnessMultipliers.Num() == NumParticles))
+		(Cloth->bUseBendingElements && (Cloth->BucklingStiffness[0] > (FRealSingle)0. || (Cloth->BucklingStiffness[1] > (FRealSingle)0. && BucklingStiffnessMultipliers.Num() == NumParticles))))
 	{
 		if (Cloth->bUseBendingElements)
 		{
diff --git a/Engine/Source/Runtime/Experimental/Chaos/Public/Chaos/PBDStiffness.h b/Engine/Source/Runtime/Experimental/Chaos/Public/Chaos/PBDStiffness.h
--- a/Engine/Source/Runtime/Experimental/Chaos/Public/Chaos/PBDStiffness.h
+++ b/Engine/Source/Runtime/Experimental/Chaos/Public/Chaos/PBDStiffness.h
@@ -99,6 +99,7 @@ private:
 	TArray<uint8> Indices; // Per particle/constraints array of index to the stiffness table
 	TArray<FSolverReal> Table;  // Fixed lookup table of stiffness values, use uint8 indexation
 	FSolverVec2 WeightedValue;
+	const FSolverReal ParameterFitBase;
 	const FSolverReal ParameterFitLogBase;
 };
 
@@ -109,6 +110,7 @@ FPBDStiffness::FPBDStiffness(
 	int32 TableSize,
 	FSolverReal InParameterFitBase)
 	: WeightedValue(InWeightedValue.ClampAxes((FSolverReal)0., (FSolverReal)1.))
+	, ParameterFitBase(InParameterFitBase)
 	, ParameterFitLogBase(FMath::Loge(InParameterFitBase))
 {
 	check(TableSize > 0 && TableSize < 256);  // The Stiffness lookup table is restricted to uint8 sized indices
@@ -147,6 +149,7 @@ FPBDStiffness::FPBDStiffness(
 	FSolverReal InParameterFitBase,
 	typename TEnableIf<Valence >= 2 && Valence <= 4>::Type*)
 	: WeightedValue(InWeightedValue.ClampAxes((FSolverReal)0., (FSolverReal)1.))
+	, ParameterFitBase(InParameterFitBase)
 	, ParameterFitLogBase(FMath::Loge(InParameterFitBase))
 {
 	check(TableSize > 0 && TableSize < 256);  // The Stiffness lookup table is restricted to uint8 sized indices
@@ -185,6 +188,14 @@ FPBDStiffness::FPBDStiffness(
 	}
 }
 
+static inline FSolverReal CalcExponentialParameterFit(const FSolverReal ParameterFitBase, const FSolverReal ParameterFitLogBase, const FSolverReal InValue)
+{
+	// Get a very steep exponential curve between the [0, 1] range to make easier to set the parameter
+	// The base has been chosen empirically.
+	// ParameterFit = (pow(ParameterFitBase , InValue) - 1) / (ParameterFitBase - 1)
+	// Note that ParameterFit = 0 when InValue = 0, and 1 when InValue = 1.
+	return (FMath::Exp(ParameterFitLogBase * InValue) - (FSolverReal)1.) / (ParameterFitBase - (FSolverReal)1.);
+}
 
 void FPBDStiffness::ApplyValues(const FSolverReal Dt, const int32 NumIterations)
 {
@@ -207,10 +218,7 @@ void FPBDStiffness::ApplyValues(const FSolverReal Dt, const int32 NumIterations)
 		{
 			return (FSolverReal)1.;
 		}
-		// Get a very steep exponential curve between the [0, 1] range to make easier to set the parameter
-		// The base has been chosen empirically
-		// ParameterValue = Pow(ParameterFitBase, ParameterValue - 1)
-		const FSolverReal ParameterFit = FMath::Exp(ParameterFitLogBase * (InValue - (FSolverReal)1.));
+		const FSolverReal ParameterFit = CalcExponentialParameterFit(ParameterFitBase, ParameterFitLogBase, InValue);
 
 		// Use simulation dependent stiffness exponent to alleviate the variations in effect when Dt and NumIterations change
 		// This is based on the Position-Based Simulation Methods paper (page 8),
@@ -235,7 +243,7 @@ void FPBDStiffness::ApplyXPBDValues(const FSolverReal MinStiffness, const FSolve
 	SCOPE_CYCLE_COUNTER(STAT_PBD_StiffnessApplyValues);
 	
 	// XPBD internally handles the effects of iterations and dt. 
-	// PBD stiffness is more like a constraint compliance. This should scale like 1/ xpbd stiffness which 
+	// PBD stiffness is more like a constraint compliance. This should scale like 1/stiffness which 
 	// is an actual measure of stiffness (e.g., something with units like Pascals)	
 	const FSolverReal MinOverMaxStiffness = MinStiffness / MaxStiffness;
 	auto SimulationValue = [this, MinOverMaxStiffness, MinStiffness, MaxStiffness](const FSolverReal InValue)->FSolverReal
@@ -245,7 +253,7 @@ void FPBDStiffness::ApplyXPBDValues(const FSolverReal MinStiffness, const FSolve
 		{
 			return MaxStiffness;
 		}
-		const FSolverReal ParameterFit = FMath::Exp(ParameterFitLogBase * (ClampedInValue - (FSolverReal)1.));
+		const FSolverReal ParameterFit = CalcExponentialParameterFit(ParameterFitBase, ParameterFitLogBase, InValue);
 		return MinStiffness / ((FSolverReal)1. - ParameterFit + MinOverMaxStiffness);
 	};
 
