Chaos - fix EPA tolerance for doubles
- fix GJK warm start when motion causes the objects to overlap

diff --git a/Engine/Source/Runtime/Experimental/Chaos/Public/Chaos/EPA.h b/Engine/Source/Runtime/Experimental/Chaos/Public/Chaos/EPA.h
--- a/Engine/Source/Runtime/Experimental/Chaos/Public/Chaos/EPA.h
+++ b/Engine/Source/Runtime/Experimental/Chaos/Public/Chaos/EPA.h
@@ -39,6 +39,18 @@ struct TEPAEntry
 		return Distance > Other.Distance;
 	}
 
+	constexpr T SelectEpsilon(float FloatEpsilon, double DoubleEpsilon) const
+	{
+		if (sizeof(T) <= sizeof(float))
+		{
+			return FloatEpsilon;
+		}
+		else
+		{
+			return DoubleEpsilon;
+		}
+	}
+
 	FORCEINLINE_DEBUGGABLE bool Initialize(const TVec3<T>* VerticesA, const TVec3<T>* VerticesB, int32 InIdx0, int32 InIdx1, int32 InIdx2, const TVector<int32,3>& InAdjFaces, const TVector<int32,3>& InAdjEdges)
 	{
 		const TVec3<T> V0 = MinkowskiVert(VerticesA, VerticesB, InIdx0);
@@ -49,7 +61,7 @@ struct TEPAEntry
 		const TVec3<T> V0V2 = V2 - V0;
 		const TVec3<T> Norm = TVec3<T>::CrossProduct(V0V1, V0V2);
 		const T NormLenSq = Norm.SizeSquared();
-		constexpr T Eps = 1e-4f;
+		const T Eps = SelectEpsilon(1.e-4f, 1.e-8);
 		if (NormLenSq < Eps)
 		{
 			return false;
