Chaos Cloth: switch XPBD constraints to use traditional stiffness inputs rather than PBD exponential compliance inputs (PBDStiffness).  We might ultimately want to have a UI that supports both an exponential 0-1 input and the actual stiffness values, but for now, using the stiffness values. With the new stiffness handling, it is possible to get stiffness = 0 (as opposed to a very small number). Update XPBD code to handle this case since we do calculations in terms of alpha = 1 / (stiffness * dt *dt).

Add XPBDBendingConstraints. In adding the XPBDBendingConstraints, I realized that I had misused VectorSelect instead of just the builtin ispc select function for scalars. This caused weird bugs in XPBDBending. I went back and fixed PBDBendingConstraints.ispc as well.

The PBDBendingConstraints change is the only change that affects existing standard users because XPBD continues to be unavailable in the UI.

diff --git a/Engine/Plugins/Experimental/ChaosCloth/Source/Chaos/Private/ChaosCloth/ChaosClothConstraints.cpp b/Engine/Plugins/Experimental/ChaosCloth/Source/Chaos/Private/ChaosCloth/ChaosClothConstraints.cpp
--- a/Engine/Plugins/Experimental/ChaosCloth/Source/Chaos/Private/ChaosCloth/ChaosClothConstraints.cpp
+++ b/Engine/Plugins/Experimental/ChaosCloth/Source/Chaos/Private/ChaosCloth/ChaosClothConstraints.cpp
@@ -3,6 +3,7 @@
 #include "Chaos/PBDSpringConstraints.h"
 #include "Chaos/XPBDSpringConstraints.h"
 #include "Chaos/PBDBendingConstraints.h"
+#include "Chaos/XPBDBendingConstraints.h"
 #include "Chaos/PBDAxialSpringConstraints.h"
 #include "Chaos/XPBDAxialSpringConstraints.h"
 #include "Chaos/PBDVolumeConstraint.h"
@@ -167,6 +168,20 @@ void FClothConstraints::CreateRules()
 				BendingElementConstraints->Apply(Particles, Dt);
 			};
 	}
+	if (XBendingElementConstraints)
+	{
+		ConstraintInits[ConstraintInitIndex++] =
+			[this](Softs::FSolverParticles& /*Particles*/, const Softs::FSolverReal Dt)
+		{
+			XBendingElementConstraints->Init();
+			XBendingElementConstraints->ApplyProperties(Dt, Evolution->GetIterations());
+		};
+		ConstraintRules[ConstraintRuleIndex++] =
+			[this](Softs::FSolverParticles& Particles, const Softs::FSolverReal Dt)
+		{
+			XBendingElementConstraints->Apply(Particles, Dt);
+		};
+	}
 	if (XAreaConstraints)
 	{
 		ConstraintInits[ConstraintInitIndex++] =
@@ -356,17 +371,33 @@ void FClothConstraints::SetBendingConstraints(TArray<TVec4<int32>>&& BendingElem
 {
 	check(Evolution);
 
-	// TODO Add XPBD
-	BendingElementConstraints = MakeShared<Softs::FPBDBendingConstraints>(
-		Evolution->Particles(),
-		ParticleOffset, NumParticles,
-		MoveTemp(BendingElements),
-		StiffnessMultipliers,
-		BucklingStiffnessMultipliers,
-		/*InStiffness =*/ Softs::FSolverVec2::UnitVector,
-		/*InBucklingRatio=*/ (Softs::FSolverReal)0.f,
-		/*InBucklingStiffnessMultiplier =*/ Softs::FSolverVec2::UnitVector,
-		/*bTrimKinematicConstraints =*/ true);
+	if (bUseXPBDConstraints)
+	{
+		XBendingElementConstraints = MakeShared<Softs::FXPBDBendingConstraints>(
+			Evolution->Particles(),
+			ParticleOffset, NumParticles,
+			MoveTemp(BendingElements),
+			StiffnessMultipliers,
+			BucklingStiffnessMultipliers,
+			/*InStiffness =*/ Softs::FSolverVec2::UnitVector,
+			/*InBucklingRatio=*/ (Softs::FSolverReal)0.f,
+			/*InBucklingStiffnessMultiplier =*/ Softs::FSolverVec2::UnitVector,
+			/*bTrimKinematicConstraints =*/ true);
+	}
+	else
+	{
+		BendingElementConstraints = MakeShared<Softs::FPBDBendingConstraints>(
+			Evolution->Particles(),
+			ParticleOffset, NumParticles,
+			MoveTemp(BendingElements),
+			StiffnessMultipliers,
+			BucklingStiffnessMultipliers,
+			/*InStiffness =*/ Softs::FSolverVec2::UnitVector,
+			/*InBucklingRatio=*/ (Softs::FSolverReal)0.f,
+			/*InBucklingStiffnessMultiplier =*/ Softs::FSolverVec2::UnitVector,
+			/*bTrimKinematicConstraints =*/ true);
+
+	}
 	++NumConstraintInits;  // Uses init to update the property tables
 	++NumConstraintRules;
 }
@@ -559,6 +590,10 @@ void FClothConstraints::SetBendingProperties(const Softs::FSolverVec2& BendingSt
 	{
 		BendingElementConstraints->SetProperties(BendingStiffness, BucklingRatio, BucklingStiffness);
 	}
+	if (XBendingElementConstraints)
+	{
+		XBendingElementConstraints->SetProperties(BendingStiffness, BucklingRatio, BucklingStiffness);
+	}
 }
 
 void FClothConstraints::SetAreaProperties(const Softs::FSolverVec2& AreaStiffness)
diff --git a/Engine/Plugins/Experimental/ChaosCloth/Source/Chaos/Private/ChaosCloth/ChaosClothingSimulation.cpp b/Engine/Plugins/Experimental/ChaosCloth/Source/Chaos/Private/ChaosCloth/ChaosClothingSimulation.cpp
--- a/Engine/Plugins/Experimental/ChaosCloth/Source/Chaos/Private/ChaosCloth/ChaosClothingSimulation.cpp
+++ b/Engine/Plugins/Experimental/ChaosCloth/Source/Chaos/Private/ChaosCloth/ChaosClothingSimulation.cpp
@@ -33,6 +33,7 @@
 #include "Chaos/PBDTriangleMeshCollisions.h"
 #include "Chaos/PBDBendingConstraints.h"
 #include "Chaos/XPBDSpringConstraints.h"
+#include "Chaos/XPBDBendingConstraints.h"
 #include "Chaos/VelocityField.h"
 #endif  // #if WITH_EDITOR || CHAOS_DEBUG_DRAW
 
@@ -151,6 +152,7 @@ namespace ClothingSimulationConsole
 					bChaos_PostIterationUpdates_ISPC_Enabled =
 					bChaos_Bending_ISPC_Enabled =
 					bChaos_XPBDSpring_ISPC_Enabled =
+					bChaos_XPBDBending_ISPC_Enabled =
 					bEnableISPC;
 				return;
 			}
@@ -412,7 +414,9 @@ void FClothingSimulation::CreateActor(USkeletalMeshComponent* InOwnerComponent,
 		TVec2<FRealSingle>(ClothConfig->AnimDriveStiffness.Low, ClothConfig->AnimDriveStiffness.High),  // Animatable
 		TVec2<FRealSingle>(ClothConfig->AnimDriveDamping.Low, ClothConfig->AnimDriveDamping.High),  // Animatable
 		ClothConfig->ShapeTargetStiffness,  // TODO: This is now deprecated
-		/*bUseXPBDConstraints =*/ false,  // Experimental
+		false, // bUseXPBDEdgeSprings
+		false, // bUseXPBDBendingElements
+		false, // bUseXPBDAreaSprings
 		ClothConfig->GravityScale,
 		ClothConfig->bUseGravityOverride,
 		ClothConfig->Gravity,
@@ -839,7 +843,9 @@ void FClothingSimulation::RefreshClothConfig(const IClothingSimulationContext* I
 			TVec2<FRealSingle>(ClothConfig->AnimDriveStiffness.Low, ClothConfig->AnimDriveStiffness.High),  // Animatable
 			TVec2<FRealSingle>(ClothConfig->AnimDriveDamping.Low, ClothConfig->AnimDriveDamping.High),  // Animatable
 			ClothConfig->ShapeTargetStiffness,
-			/*bUseXPBDConstraints =*/ false,  // Experimental
+			false, // bUseXPBDEdgeSprings
+			false, // bUseXPBDBendingElements
+			false, // bUseXPBDAreaSprings
 			ClothConfig->GravityScale,
 			ClothConfig->bUseGravityOverride,
 			ClothConfig->Gravity,
@@ -1873,6 +1879,29 @@ void FClothingSimulation::DebugDrawBendingConstraint(FPrimitiveDrawInterface* PD
 				DrawLine(PDI, Pos0, Pos1, bIsBuckled ? FLinearColor::Red : FLinearColor::Blue);
 			}
 		}
+
+		if (const Softs::FXPBDBendingConstraints* const BendingConstraints = ClothConstraints.GetXBendingElementConstraints().Get())
+		{
+			const TArray<TVec4<int32>>& Constraints = BendingConstraints->GetConstraints();
+			const TArray<FSolverReal>& RestAngles = BendingConstraints->GetRestAngles();
+
+			// Color constraint edge with red or blue: Red = Buckled, Blue = Not Buckled. This is based on its current configuration, not the values used
+			// in the actual simulation.
+			for (int32 ConstraintIndex = 0; ConstraintIndex < Constraints.Num(); ++ConstraintIndex)
+			{
+				const Softs::FSolverVec3& P1 = Positions[Constraints[ConstraintIndex][0]];
+				const Softs::FSolverVec3& P2 = Positions[Constraints[ConstraintIndex][1]];
+				const Softs::FSolverVec3& P3 = Positions[Constraints[ConstraintIndex][2]];
+				const Softs::FSolverVec3& P4 = Positions[Constraints[ConstraintIndex][3]];
+				const Softs::FSolverReal Angle = Softs::FPBDBendingConstraintsBase::CalcAngle(P1, P2, P3, P4);
+
+				const bool bIsBuckled = BendingConstraints->AngleIsBuckled(Angle, RestAngles[ConstraintIndex]);
+
+				const FVec3 Pos0 = FVec3(P1) + LocalSpaceLocation;
+				const FVec3 Pos1 = FVec3(P2) + LocalSpaceLocation;
+				DrawLine(PDI, Pos0, Pos1, bIsBuckled ? FLinearColor::Red : FLinearColor::Blue);
+			}
+		}
 	}
 }
 
diff --git a/Engine/Plugins/Experimental/ChaosCloth/Source/Chaos/Private/ChaosCloth/ChaosClothingSimulationCloth.cpp b/Engine/Plugins/Experimental/ChaosCloth/Source/Chaos/Private/ChaosCloth/ChaosClothingSimulationCloth.cpp
--- a/Engine/Plugins/Experimental/ChaosCloth/Source/Chaos/Private/ChaosCloth/ChaosClothingSimulationCloth.cpp
+++ b/Engine/Plugins/Experimental/ChaosCloth/Source/Chaos/Private/ChaosCloth/ChaosClothingSimulationCloth.cpp
@@ -128,7 +128,7 @@ void FClothingSimulationCloth::FLODData::Add(FClothingSimulationSolver* Solver,
 	const TConstArrayView<FRealSingle>& EdgeStiffnessMultipliers = WeightMaps[(int32)EChaosWeightMapTarget::EdgeStiffness];
 	if (Cloth->EdgeStiffness[0] > (FRealSingle)0. || (Cloth->EdgeStiffness[1] > (FRealSingle)0. && EdgeStiffnessMultipliers.Num() == NumParticles))
 	{
-		ClothConstraints.SetEdgeConstraints(SurfaceElements, EdgeStiffnessMultipliers, Cloth->bUseXPBDConstraints);
+		ClothConstraints.SetEdgeConstraints(SurfaceElements, EdgeStiffnessMultipliers, Cloth->bUseXPBDEdgeConstraints);
 	}
 
 	// Bending constraints
@@ -140,12 +140,12 @@ void FClothingSimulationCloth::FLODData::Add(FClothingSimulationSolver* Solver,
 		if (Cloth->bUseBendingElements)
 		{
 			TArray<Chaos::TVec4<int32>> BendingElements = TriangleMesh.GetUniqueAdjacentElements();
-			ClothConstraints.SetBendingConstraints(MoveTemp(BendingElements), BendingStiffnessMultipliers, BucklingStiffnessMultipliers, Cloth->bUseXPBDConstraints);
+			ClothConstraints.SetBendingConstraints(MoveTemp(BendingElements), BendingStiffnessMultipliers, BucklingStiffnessMultipliers, Cloth->bUseXPBDBendingConstraints);
 		}
 		else
 		{
 			const TArray<Chaos::TVec2<int32>> Edges = TriangleMesh.GetUniqueAdjacentPoints();
-			ClothConstraints.SetBendingConstraints(Edges, BendingStiffnessMultipliers, Cloth->bUseXPBDConstraints);
+			ClothConstraints.SetBendingConstraints(Edges, BendingStiffnessMultipliers, Cloth->bUseXPBDBendingConstraints);
 		}
 	}
 
@@ -153,7 +153,7 @@ void FClothingSimulationCloth::FLODData::Add(FClothingSimulationSolver* Solver,
 	const TConstArrayView<FRealSingle>& AreaStiffnessMultipliers = WeightMaps[(int32)EChaosWeightMapTarget::AreaStiffness];
 	if (Cloth->AreaStiffness[0] > (FRealSingle)0. || (Cloth->AreaStiffness[1] > (FRealSingle)0. && AreaStiffnessMultipliers.Num() == NumParticles))
 	{
-		ClothConstraints.SetAreaConstraints(SurfaceElements, AreaStiffnessMultipliers, Cloth->bUseXPBDConstraints);
+		ClothConstraints.SetAreaConstraints(SurfaceElements, AreaStiffnessMultipliers, Cloth->bUseXPBDAreaConstraints);
 	}
 
 	// Volume constraints
@@ -355,7 +355,9 @@ FClothingSimulationCloth::FClothingSimulationCloth(
 	const TVec2<FRealSingle>& InAnimDriveStiffness,
 	const TVec2<FRealSingle>& InAnimDriveDamping,
 	FRealSingle InShapeTargetStiffness,
-	bool bInUseXPBDConstraints,
+	bool bInUseXPBDEdgeConstraints,
+	bool bInUseXPBDBendingConstraints,
+	bool bInUseXPBDAreaConstraints,
 	FRealSingle InGravityScale,
 	bool bInIsGravityOverridden,
 	const TVec3<FRealSingle>& InGravityOverride,
@@ -399,7 +401,9 @@ FClothingSimulationCloth::FClothingSimulationCloth(
 	, AnimDriveStiffness(InAnimDriveStiffness)
 	, AnimDriveDamping(InAnimDriveDamping)
 	, ShapeTargetStiffness(InShapeTargetStiffness)
-	, bUseXPBDConstraints(bInUseXPBDConstraints)
+	, bUseXPBDEdgeConstraints(bInUseXPBDEdgeConstraints)
+	, bUseXPBDBendingConstraints(bInUseXPBDBendingConstraints)
+	, bUseXPBDAreaConstraints(bInUseXPBDAreaConstraints)
 	, GravityScale(InGravityScale)
 	, bIsGravityOverridden(bInIsGravityOverridden)
 	, GravityOverride(InGravityOverride)
diff --git a/Engine/Plugins/Experimental/ChaosCloth/Source/Chaos/Public/ChaosCloth/ChaosClothConstraints.h b/Engine/Plugins/Experimental/ChaosCloth/Source/Chaos/Public/ChaosCloth/ChaosClothConstraints.h
--- a/Engine/Plugins/Experimental/ChaosCloth/Source/Chaos/Public/ChaosCloth/ChaosClothConstraints.h
+++ b/Engine/Plugins/Experimental/ChaosCloth/Source/Chaos/Public/ChaosCloth/ChaosClothConstraints.h
@@ -78,6 +78,7 @@ namespace Chaos
 		const TSharedPtr<Softs::FPBDSpringConstraints>& GetBendingConstraints() const { return BendingConstraints; }  
 		const TSharedPtr<Softs::FXPBDSpringConstraints>& GetXBendingConstraints() const { return XBendingConstraints; }
 		const TSharedPtr<Softs::FPBDBendingConstraints>& GetBendingElementConstraints() const { return BendingElementConstraints; }
+		const TSharedPtr<Softs::FXPBDBendingConstraints>& GetXBendingElementConstraints() const { return XBendingElementConstraints; }
 		const TSharedPtr<Softs::FPBDAxialSpringConstraints>& GetAreaConstraints() const { return AreaConstraints; }
 		const TSharedPtr<Softs::FXPBDAxialSpringConstraints>& GetXAreaConstraints() const { return XAreaConstraints; }
 		const TSharedPtr<Softs::FPBDSpringConstraints>& GetThinShellVolumeConstraints() const { return ThinShellVolumeConstraints; }
@@ -105,6 +106,7 @@ namespace Chaos
 		TSharedPtr<Softs::FPBDSpringConstraints> BendingConstraints;
 		TSharedPtr<Softs::FXPBDSpringConstraints> XBendingConstraints;
 		TSharedPtr<Softs::FPBDBendingConstraints> BendingElementConstraints;
+		TSharedPtr<Softs::FXPBDBendingConstraints> XBendingElementConstraints;
 		TSharedPtr<Softs::FPBDAxialSpringConstraints> AreaConstraints;
 		TSharedPtr<Softs::FXPBDAxialSpringConstraints> XAreaConstraints;
 		TSharedPtr<Softs::FPBDSpringConstraints> ThinShellVolumeConstraints;
diff --git a/Engine/Plugins/Experimental/ChaosCloth/Source/Chaos/Public/ChaosCloth/ChaosClothingSimulationCloth.h b/Engine/Plugins/Experimental/ChaosCloth/Source/Chaos/Public/ChaosCloth/ChaosClothingSimulationCloth.h
--- a/Engine/Plugins/Experimental/ChaosCloth/Source/Chaos/Public/ChaosCloth/ChaosClothingSimulationCloth.h
+++ b/Engine/Plugins/Experimental/ChaosCloth/Source/Chaos/Public/ChaosCloth/ChaosClothingSimulationCloth.h
@@ -48,7 +48,9 @@ namespace Chaos
 			const TVec2<FRealSingle>& InAnimDriveStiffness,
 			const TVec2<FRealSingle>& InAnimDriveDamping,
 			FRealSingle InShapeTargetStiffness,
-			bool bInUseXPBDConstraints,
+			bool bInUseXPBDEdgeConstraints,
+			bool bInUseXPBDBendingConstraints,
+			bool bInUseXPBDAreaConstraints,
 			FRealSingle InGravityScale,
 			bool bIsGravityOverridden,
 			const TVec3<FRealSingle>& InGravityOverride,
@@ -212,7 +214,9 @@ namespace Chaos
 		TVec2<FRealSingle> AnimDriveStiffness;  // Animatable
 		TVec2<FRealSingle> AnimDriveDamping;  // Animatable
 		FRealSingle ShapeTargetStiffness;
-		bool bUseXPBDConstraints;
+		bool bUseXPBDEdgeConstraints;
+		bool bUseXPBDBendingConstraints;
+		bool bUseXPBDAreaConstraints;
 		FRealSingle GravityScale;
 		bool bIsGravityOverridden;
 		TVec3<FRealSingle> GravityOverride;
diff --git a/Engine/Source/Runtime/Experimental/Chaos/Private/Chaos/PBDBendingConstraints.ispc b/Engine/Source/Runtime/Experimental/Chaos/Private/Chaos/PBDBendingConstraints.ispc
--- a/Engine/Source/Runtime/Experimental/Chaos/Private/Chaos/PBDBendingConstraints.ispc
+++ b/Engine/Source/Runtime/Experimental/Chaos/Private/Chaos/PBDBendingConstraints.ispc
@@ -14,7 +14,7 @@ static inline FVector3f SafeDivide(const FVector3f& Numerator, const float Denom
 
 static inline varying float SafeDivide(const varying float Numerator, const varying float Denominator)
 {
-	return VectorSelect(Denominator <= FLOAT_SMALL_NUMBER, (uniform float)0, Numerator / Denominator);
+	return select(Denominator <= FLOAT_SMALL_NUMBER, (uniform float)0, Numerator / Denominator);
 }
 
 // Returns Angle
@@ -89,7 +89,7 @@ export void ApplyBendingConstraints(uniform FVector4f PandInvM[],
 		
 		const varying float RestAngle = RestAngles[i];
 		const varying bool AngleIsBuckled = FLOAT_PI - abs(Angle) < BucklingRatio * (FLOAT_PI - abs(RestAngle));
-		const varying float StiffnessValue = VectorSelect(AngleIsBuckled, BucklingStiffness, Stiffness);
+		const varying float StiffnessValue = select(AngleIsBuckled, BucklingStiffness, Stiffness);
 
 		const varying float Delta = clamp(StiffnessValue * (Angle - RestAngle), -SingleStepAngleLimit, SingleStepAngleLimit);
 
@@ -135,8 +135,8 @@ export void ApplyBendingConstraintsWithMaps(uniform FVector4f PandInvM[],
 		const varying int32 i2 = Constraint.V[1];
 		const varying int32 i3 = Constraint.V[2];
 		const varying int32 i4 = Constraint.V[3];
-		const varying int8 StiffnessIndex = VectorSelect(StiffnessHasMap, StiffnessIndices[i], FLOAT_ZERO);
-		const varying int8 BucklingStiffnessIndex = VectorSelect(BucklingStiffnessHasMap, BucklingStiffnessIndices[i], FLOAT_ZERO);
+		const varying int8 StiffnessIndex = select(StiffnessHasMap, StiffnessIndices[i], 0);
+		const varying int8 BucklingStiffnessIndex = select(BucklingStiffnessHasMap, BucklingStiffnessIndices[i], 0);
 
 		const varying FVector4f PandInvM2 = VectorGather(&PandInvM[i2]);
 		const varying FVector4f PandInvM3 = VectorGather(&PandInvM[i3]);
@@ -164,7 +164,7 @@ export void ApplyBendingConstraintsWithMaps(uniform FVector4f PandInvM[],
 
 		const varying float RestAngle = RestAngles[i];
 		const varying bool AngleIsBuckled = FLOAT_PI - abs(Angle) < BucklingRatio * (FLOAT_PI - abs(RestAngle));
-		const varying float StiffnessValue = VectorSelect(AngleIsBuckled, BucklingStiffness, Stiffness);
+		const varying float StiffnessValue = select(AngleIsBuckled, BucklingStiffness, Stiffness);
 
 		const varying float Delta = clamp(StiffnessValue * (Angle - RestAngle), -SingleStepAngleLimit, SingleStepAngleLimit);
 
diff --git a/Engine/Source/Runtime/Experimental/Chaos/Private/Chaos/XPBDBendingConstraints.cpp b/Engine/Source/Runtime/Experimental/Chaos/Private/Chaos/XPBDBendingConstraints.cpp
new file mode 100644
--- /dev/null
+++ b/Engine/Source/Runtime/Experimental/Chaos/Private/Chaos/XPBDBendingConstraints.cpp
@@ -0,0 +1,231 @@
+// Copyright Epic Games, Inc. All Rights Reserved.
+#include "Chaos/XPBDBendingConstraints.h"
+#include "Chaos/GraphColoring.h"
+#include "Chaos/Framework/Parallel.h"
+#include "ChaosStats.h"
+#if INTEL_ISPC
+#include "XPBDBendingConstraints.ispc.generated.h"
+#endif
+
+DECLARE_CYCLE_STAT(TEXT("Chaos XPBD Bending Constraint"), STAT_XPBD_Bending, STATGROUP_Chaos);
+
+#if INTEL_ISPC && !UE_BUILD_SHIPPING
+bool bChaos_XPBDBending_ISPC_Enabled = true;
+FAutoConsoleVariableRef CVarChaosXPBDBendingISPCEnabled(TEXT("p.Chaos.XPBDBending.ISPC"), bChaos_XPBDBending_ISPC_Enabled, TEXT("Whether to use ISPC optimizations in XPBD Bending constraints"));
+
+static_assert(sizeof(ispc::FVector4f) == sizeof(Chaos::Softs::FPAndInvM), "sizeof(ispc::FVector4f) != sizeof(Chaos::Softs::FPAndInvM");
+static_assert(sizeof(ispc::FIntVector4) == sizeof(Chaos::TVec4<int32>), "sizeof(ispc::FIntVector4) != sizeof(Chaos::TVec4<int32>");
+#endif
+
+namespace Chaos::Softs {
+
+// @todo(chaos): the parallel threshold (or decision to run parallel) should probably be owned by the solver and passed to the constraint container
+int32 Chaos_XPBDBending_ParallelConstraintCount = 100;
+FAutoConsoleVariableRef CVarChaosXPBDBendingParallelConstraintCount(TEXT("p.Chaos.XPBDBending.ParallelConstraintCount"), Chaos_XPBDBending_ParallelConstraintCount, TEXT("If we have more constraints than this, use parallel-for in Apply."));
+
+void FXPBDBendingConstraints::InitColor(const FSolverParticles& InParticles)
+{
+	// In dev builds we always color so we can tune the system without restarting. See Apply()
+#if UE_BUILD_SHIPPING || UE_BUILD_TEST
+	if (Constraints.Num() > Chaos_XPBDBending_ParallelConstraintCount)
+#endif
+	{
+		const TArray<TArray<int32>> ConstraintsPerColor = FGraphColoring::ComputeGraphColoring(Constraints, InParticles);
+
+		// Reorder constraints based on color so each array in ConstraintsPerColor contains contiguous elements.
+		TArray<TVec4<int32>> ReorderedConstraints; 
+		TArray<TVec2<int32>> ReorderedConstraintSharedEdges;
+		TArray<FSolverReal> ReorderedRestAngles;
+		TArray<int32> OrigToReorderedIndices; // used to reorder stiffness indices
+		ReorderedConstraints.SetNumUninitialized(Constraints.Num());
+		ReorderedConstraintSharedEdges.SetNumUninitialized(Constraints.Num());
+		ReorderedRestAngles.SetNumUninitialized(Constraints.Num());
+		OrigToReorderedIndices.SetNumUninitialized(Constraints.Num());
+
+		ConstraintsPerColorStartIndex.Reset(ConstraintsPerColor.Num() + 1);
+
+		int32 ReorderedIndex = 0;
+		for (const TArray<int32>& ConstraintsBatch : ConstraintsPerColor)
+		{
+			ConstraintsPerColorStartIndex.Add(ReorderedIndex);
+			for (const int32& BatchConstraint : ConstraintsBatch)
+			{
+				const int32 OrigIndex = BatchConstraint;
+				ReorderedConstraints[ReorderedIndex] = Constraints[OrigIndex];
+				ReorderedConstraintSharedEdges[ReorderedIndex] = ConstraintSharedEdges[OrigIndex];
+				ReorderedRestAngles[ReorderedIndex] = RestAngles[OrigIndex];
+				OrigToReorderedIndices[OrigIndex] = ReorderedIndex;
+				++ReorderedIndex;
+			}
+		}
+		ConstraintsPerColorStartIndex.Add(ReorderedIndex);
+
+		Constraints = MoveTemp(ReorderedConstraints);
+		ConstraintSharedEdges = MoveTemp(ReorderedConstraintSharedEdges);
+		RestAngles = MoveTemp(ReorderedRestAngles);
+		Stiffness.ReorderIndices(OrigToReorderedIndices);
+		BucklingStiffness.ReorderIndices(OrigToReorderedIndices);
+	}
+}
+
+void FXPBDBendingConstraints::ApplyHelper(FSolverParticles& Particles, const FSolverReal Dt, const int32 ConstraintIndex, const FSolverReal StiffnessValue, const FSolverReal BucklingValue) const
+{
+	const TVec4<int32>& Constraint = Constraints[ConstraintIndex];
+	const int32 i1 = Constraint[0];
+	const int32 i2 = Constraint[1];
+	const int32 i3 = Constraint[2];
+	const int32 i4 = Constraint[3];
+	const FSolverReal Angle = CalcAngle(Particles.P(i1), Particles.P(i2), Particles.P(i3), Particles.P(i4));
+
+	const FSolverReal BiphasicStiffnessValue = AngleIsBuckled(Angle, RestAngles[ConstraintIndex]) ? BucklingValue : StiffnessValue;
+	if (BiphasicStiffnessValue < XPBDBendMinStiffness || (Particles.InvM(i1) == (FSolverReal)0. && Particles.InvM(i2) == (FSolverReal)0. && Particles.InvM(i3) == (FSolverReal)0. && Particles.InvM(i4) == (FSolverReal)0.))
+	{
+		return;
+	}
+
+	const TStaticArray<FSolverVec3, 4> Grads = Base::GetGradients(Particles, ConstraintIndex);
+
+	FSolverReal& Lambda = Lambdas[ConstraintIndex];
+	const FSolverReal Alpha = (FSolverReal)1.f / (BiphasicStiffnessValue * Dt * Dt);
+	const FSolverReal Denom = Particles.InvM(i1) * Grads[0].SizeSquared() + Particles.InvM(i2) * Grads[1].SizeSquared() + Particles.InvM(i3) * Grads[2].SizeSquared() + Particles.InvM(i4) * Grads[3].SizeSquared() + Alpha;
+	const FSolverReal DLambda = (Angle - RestAngles[ConstraintIndex] - Alpha * Lambda) / Denom;
+
+	Particles.P(i1) -= DLambda * Particles.InvM(i1) * Grads[0];
+	Particles.P(i2) -= DLambda * Particles.InvM(i2) * Grads[1];
+	Particles.P(i3) -= DLambda * Particles.InvM(i3) * Grads[2];
+	Particles.P(i4) -= DLambda * Particles.InvM(i4) * Grads[3];
+	Lambda += DLambda;
+}
+
+void FXPBDBendingConstraints::Apply(FSolverParticles& Particles, const FSolverReal Dt) const
+{
+	TRACE_CPUPROFILER_EVENT_SCOPE(FXPBDBendingConstraints_Apply);
+	SCOPE_CYCLE_COUNTER(STAT_XPBD_Bending);
+
+	const bool StiffnessHasWeightMap = Stiffness.HasWeightMap();
+	const bool BucklingStiffnessHasWeightMap = BucklingStiffness.HasWeightMap();
+
+	if (ConstraintsPerColorStartIndex.Num() > 0 && Constraints.Num() > Chaos_XPBDBending_ParallelConstraintCount)
+	{
+		const int32 ConstraintColorNum = ConstraintsPerColorStartIndex.Num() - 1;
+		if (!StiffnessHasWeightMap && !BucklingStiffnessHasWeightMap)
+		{
+			const FSolverReal ExpStiffnessValue = (FSolverReal)Stiffness;
+			const FSolverReal ExpBucklingValue = (FSolverReal)BucklingStiffness;
+
+			if (ExpStiffnessValue < XPBDBendMinStiffness && ExpBucklingValue < XPBDBendMinStiffness)
+			{
+				return;
+			}
+
+#if INTEL_ISPC
+			if (bRealTypeCompatibleWithISPC && bChaos_XPBDBending_ISPC_Enabled)
+			{
+				for (int32 ConstraintColorIndex = 0; ConstraintColorIndex < ConstraintColorNum; ++ConstraintColorIndex)
+				{
+					const int32 ColorStart = ConstraintsPerColorStartIndex[ConstraintColorIndex];
+					const int32 ColorSize = ConstraintsPerColorStartIndex[ConstraintColorIndex + 1] - ColorStart;
+					ispc::ApplyXPBDBendingConstraints(
+						(ispc::FVector4f*)Particles.GetPAndInvM().GetData(),
+						(ispc::FIntVector4*)&Constraints.GetData()[ColorStart],
+						&RestAngles.GetData()[ColorStart],
+						&Lambdas.GetData()[ColorStart],
+						Dt,
+						XPBDBendMinStiffness,
+						ExpStiffnessValue,
+						BucklingRatio,
+						ExpBucklingValue,
+						ColorSize);
+				}
+			}
+			else
+#endif
+			{
+				for (int32 ConstraintColorIndex = 0; ConstraintColorIndex < ConstraintColorNum; ++ConstraintColorIndex)
+				{
+					const int32 ColorStart = ConstraintsPerColorStartIndex[ConstraintColorIndex];
+					const int32 ColorSize = ConstraintsPerColorStartIndex[ConstraintColorIndex + 1] - ColorStart;
+					PhysicsParallelFor(ColorSize, [&](const int32 Index)
+					{
+						const int32 ConstraintIndex = ColorStart + Index;
+						ApplyHelper(Particles, Dt, ConstraintIndex, ExpStiffnessValue, ExpBucklingValue);
+					});
+				}
+			}
+		}
+		else  // Has weight maps
+		{
+#if INTEL_ISPC
+			if (bRealTypeCompatibleWithISPC && bChaos_XPBDBending_ISPC_Enabled)
+			{
+				for (int32 ConstraintColorIndex = 0; ConstraintColorIndex < ConstraintColorNum; ++ConstraintColorIndex)
+				{
+					const int32 ColorStart = ConstraintsPerColorStartIndex[ConstraintColorIndex];
+					const int32 ColorSize = ConstraintsPerColorStartIndex[ConstraintColorIndex + 1] - ColorStart;
+					ispc::ApplyXPBDBendingConstraintsWithMaps(
+						(ispc::FVector4f*)Particles.GetPAndInvM().GetData(),
+						(ispc::FIntVector4*)&Constraints.GetData()[ColorStart],
+						&RestAngles.GetData()[ColorStart],
+						&Lambdas.GetData()[ColorStart],
+						Dt,
+						XPBDBendMinStiffness,
+						StiffnessHasWeightMap,
+						&Stiffness.GetIndices().GetData()[ColorStart],
+						&Stiffness.GetTable().GetData()[0],
+						BucklingRatio,
+						BucklingStiffnessHasWeightMap,
+						&BucklingStiffness.GetIndices().GetData()[ColorStart],
+						&BucklingStiffness.GetTable().GetData()[0],
+						ColorSize);
+				}
+			}
+			else
+#endif
+			{
+				const FSolverReal StiffnessNoMap = (FSolverReal)Stiffness;
+				const FSolverReal BucklingStiffnessNoMap = (FSolverReal)BucklingStiffness;
+				for (int32 ConstraintColorIndex = 0; ConstraintColorIndex < ConstraintColorNum; ++ConstraintColorIndex)
+				{
+					const int32 ColorStart = ConstraintsPerColorStartIndex[ConstraintColorIndex];
+					const int32 ColorSize = ConstraintsPerColorStartIndex[ConstraintColorIndex + 1] - ColorStart;
+					PhysicsParallelFor(ColorSize, [&](const int32 Index)
+					{
+						const int32 ConstraintIndex = ColorStart + Index;
+						const FSolverReal ExpStiffnessValue = StiffnessHasWeightMap ? Stiffness[ConstraintIndex] : StiffnessNoMap;
+						const FSolverReal ExpBucklingValue = BucklingStiffnessHasWeightMap ? BucklingStiffness[ConstraintIndex] : BucklingStiffnessNoMap;
+						ApplyHelper(Particles, Dt, ConstraintIndex, ExpStiffnessValue, ExpBucklingValue);
+					});
+				}
+			}
+		}
+	}
+	else
+	{
+		if (!StiffnessHasWeightMap && !BucklingStiffnessHasWeightMap)
+		{
+			const FSolverReal ExpStiffnessValue = (FSolverReal)Stiffness;
+			const FSolverReal ExpBucklingValue = (FSolverReal)BucklingStiffness;
+			if (ExpStiffnessValue < XPBDBendMinStiffness && ExpBucklingValue < XPBDBendMinStiffness)
+			{
+				return;
+			}
+			for (int32 ConstraintIndex = 0; ConstraintIndex < Constraints.Num(); ++ConstraintIndex)
+			{
+				ApplyHelper(Particles, Dt, ConstraintIndex, ExpStiffnessValue, ExpBucklingValue);
+			}
+		}
+		else
+		{
+			const FSolverReal StiffnessNoMap = (FSolverReal)Stiffness;
+			const FSolverReal BucklingStiffnessNoMap = (FSolverReal)BucklingStiffness;
+			for (int32 ConstraintIndex = 0; ConstraintIndex < Constraints.Num(); ++ConstraintIndex)
+			{
+				const FSolverReal ExpStiffnessValue = StiffnessHasWeightMap ? Stiffness[ConstraintIndex] : StiffnessNoMap;
+				const FSolverReal ExpBucklingValue = BucklingStiffnessHasWeightMap ? BucklingStiffness[ConstraintIndex] : BucklingStiffnessNoMap;
+				ApplyHelper(Particles, Dt, ConstraintIndex, ExpStiffnessValue, ExpBucklingValue);
+			}
+		}
+	}
+}
+
+}  // End namespace Chaos::Softs
diff --git a/Engine/Source/Runtime/Experimental/Chaos/Private/Chaos/XPBDBendingConstraints.ispc b/Engine/Source/Runtime/Experimental/Chaos/Private/Chaos/XPBDBendingConstraints.ispc
new file mode 100644
--- /dev/null
+++ b/Engine/Source/Runtime/Experimental/Chaos/Private/Chaos/XPBDBendingConstraints.ispc
@@ -0,0 +1,197 @@
+// Copyright Epic Games, Inc. All Rights Reserved.
+
+#define EXPLICIT_VECTOR4 1
+
+#include "Math/Vector.isph"
+#include "Chaos/PBDSofts.isph"
+
+static inline FVector3f SafeDivide(const FVector3f& Numerator, const float Denominator)
+{
+	return VectorSelect(Denominator <= FLOAT_SMALL_NUMBER, FloatZeroVector, Numerator / Denominator);
+}
+
+static inline varying float SafeDivide(const varying float Numerator, const varying float Denominator)
+{
+	return select(Denominator <= FLOAT_SMALL_NUMBER, (uniform float)0, Numerator / Denominator);
+}
+
+// Returns Angle
+static inline float CalcGradientsAndAngle(const FVector3f& P1, const FVector3f& P2, const FVector3f& P3, const FVector3f& P4, FVector3f& Grad1, FVector3f& Grad2, FVector3f& Grad3, FVector3f& Grad4)
+{
+	const FVector3f SharedEdgeNormalized = VectorGetSafeNormal(P2 - P1);
+	const FVector3f P13CrossP23 = VectorCross(P1 - P3, P2 - P3);
+	const float Normal1Len = VectorSize(P13CrossP23);
+	const FVector3f Normal1 = SafeDivide(P13CrossP23, Normal1Len);
+	const FVector3f P24CrossP14 = VectorCross(P2 - P4, P1 - P4);
+	const float Normal2Len = VectorSize(P24CrossP14);
+	const FVector3f Normal2 = SafeDivide(P24CrossP14, Normal2Len);
+
+	const FVector3f N2CrossN1 = VectorCross(Normal2, Normal1);
+
+	const float CosPhi = clamp(VectorDot(Normal1, Normal2), (uniform float)(-1), (uniform float)(1));
+	const float SinPhi = clamp(VectorDot(N2CrossN1, SharedEdgeNormalized), (uniform float)(-1), (uniform float)(1));
+
+	const FVector3f DPhiDN1_OverNormal1Len = SafeDivide(CosPhi * VectorCross(SharedEdgeNormalized, Normal2) - SinPhi * Normal2, Normal1Len);
+	const FVector3f DPhiDN2_OverNormal2Len = SafeDivide(CosPhi * VectorCross(Normal1, SharedEdgeNormalized) - SinPhi * Normal1, Normal2Len);
+
+	const FVector3f DPhiDP13 = VectorCross(P2 - P3, DPhiDN1_OverNormal1Len);
+	const FVector3f DPhiDP23 = VectorCross(DPhiDN1_OverNormal1Len, P1 - P3);
+	const FVector3f DPhiDP24 = VectorCross(P1 - P4, DPhiDN2_OverNormal2Len);
+	const FVector3f DPhiDP14 = VectorCross(DPhiDN2_OverNormal2Len, P2 - P4);
+
+	Grad1 = DPhiDP13 + DPhiDP14;
+	Grad2 = DPhiDP23 + DPhiDP24;
+	Grad3 = -1 * DPhiDP13 - DPhiDP23;
+	Grad4 = -1 * DPhiDP14 - DPhiDP24;
+
+	return atan2(SinPhi, CosPhi);
+}
+
+
+export void ApplyXPBDBendingConstraints(uniform FVector4f PandInvM[],
+									const uniform FIntVector4 Constraints[],
+									const uniform float RestAngles[],
+									uniform float Lambdas[],
+									const uniform float Dt,
+									const uniform float MinStiffness,
+									const uniform float Stiffness,
+									const uniform float BucklingRatio,
+									const uniform float BucklingStiffness,
+									const uniform int32 NumConstraints)
+{
+
+	foreach(i = 0 ... NumConstraints)
+	{
+		const varying FIntVector4 Constraint = VectorLoad(&Constraints[extract(i,0)]);
+
+		const varying int32 i1 = Constraint.V[0];
+		const varying int32 i2 = Constraint.V[1];
+		const varying int32 i3 = Constraint.V[2];
+		const varying int32 i4 = Constraint.V[3];
+		
+		const varying FVector4f PandInvM1 = VectorGather(&PandInvM[i1]);
+		const varying FVector4f PandInvM2 = VectorGather(&PandInvM[i2]);
+		const varying FVector4f PandInvM3 = VectorGather(&PandInvM[i3]);
+		const varying FVector4f PandInvM4 = VectorGather(&PandInvM[i4]);
+
+		varying FVector3f P1, P2, P3, P4;
+		varying float M1, M2, M3, M4;
+
+		UnzipPandInvM(PandInvM1, P1, M1);
+		UnzipPandInvM(PandInvM2, P2, M2);
+		UnzipPandInvM(PandInvM3, P3, M3);
+		UnzipPandInvM(PandInvM4, P4, M4);
+
+		varying FVector3f Grad1, Grad2, Grad3, Grad4;
+		const varying float Angle = CalcGradientsAndAngle(P1, P2, P3, P4, Grad1, Grad2, Grad3, Grad4);
+		
+		const varying float RestAngle = RestAngles[i];
+		const varying bool AngleIsBuckled = FLOAT_PI - abs(Angle) < BucklingRatio * (FLOAT_PI - abs(RestAngle));
+		const varying float StiffnessValue = Stiffness;//VectorSelect(AngleIsBuckled, BucklingStiffness, Stiffness);
+		
+		const varying float Lambda = Lambdas[i];
+		const varying float Alpha  = 1.f / (StiffnessValue * Dt * Dt);
+		const varying float Denom = M1 * VectorSizeSquared(Grad1) + M2 * VectorSizeSquared(Grad2) + M3 * VectorSizeSquared(Grad3) + M4 * VectorSizeSquared(Grad4) + Alpha;
+
+		const varying float ComputedDLambda = (Angle - RestAngle - Alpha * Lambda) / Denom;
+		const varying float DLambda = select(StiffnessValue < MinStiffness, (uniform float)0, ComputedDLambda);
+		Lambdas[i] = Lambda + DLambda;
+
+		if (M1 > 0)
+		{
+			VectorScatter(&PandInvM[i1], SetVector4(P1 - DLambda * M1 * Grad1, M1));
+		}
+		if (M2 > 0)
+		{
+			VectorScatter(&PandInvM[i2], SetVector4(P2 - DLambda * M2 * Grad2, M2));
+		}
+		if (M3 > 0)
+		{
+			VectorScatter(&PandInvM[i3], SetVector4(P3 - DLambda * M3 * Grad3, M3));
+		}
+		if (M4 > 0)
+		{
+			VectorScatter(&PandInvM[i4], SetVector4(P4 - DLambda * M4 * Grad4, M4));
+		}
+	}
+}
+
+export void ApplyXPBDBendingConstraintsWithMaps(uniform FVector4f PandInvM[],
+											const uniform FIntVector4 Constraints[],
+											const uniform float RestAngles[],
+											uniform float Lambdas[],
+											const uniform float Dt,
+											const uniform float MinStiffness,
+											const uniform bool StiffnessHasMap,
+											const uniform uint8 StiffnessIndices[],
+											const uniform float StiffnessTable[],
+											const uniform float BucklingRatio,
+											const uniform bool BucklingStiffnessHasMap,
+											const uniform uint8 BucklingStiffnessIndices[],
+											const uniform float BucklingStiffnessTable[],
+											const uniform int32 NumConstraints)
+{
+
+	foreach(i = 0 ... NumConstraints)
+	{
+		const varying FIntVector4 Constraint = VectorLoad(&Constraints[extract(i,0)]);
+
+		const varying int32 i1 = Constraint.V[0];
+		const varying int32 i2 = Constraint.V[1];
+		const varying int32 i3 = Constraint.V[2];
+		const varying int32 i4 = Constraint.V[3];
+		const varying int8 StiffnessIndex = select(StiffnessHasMap, StiffnessIndices[i], 0);
+		const varying int8 BucklingStiffnessIndex = select(BucklingStiffnessHasMap, BucklingStiffnessIndices[i], 0);
+
+		const varying FVector4f PandInvM2 = VectorGather(&PandInvM[i2]);
+		const varying FVector4f PandInvM3 = VectorGather(&PandInvM[i3]);
+		const varying FVector4f PandInvM1 = VectorGather(&PandInvM[i1]);
+		const varying FVector4f PandInvM4 = VectorGather(&PandInvM[i4]);
+
+		varying FVector3f P1, P2, P3, P4;
+		varying float M1, M2, M3, M4;
+
+		UnzipPandInvM(PandInvM1, P1, M1);
+		UnzipPandInvM(PandInvM2, P2, M2);
+		UnzipPandInvM(PandInvM3, P3, M3);
+		UnzipPandInvM(PandInvM4, P4, M4);
+
+
+		varying FVector3f Grad1, Grad2, Grad3, Grad4;
+		const varying float Angle = CalcGradientsAndAngle(P1, P2, P3, P4, Grad1, Grad2, Grad3, Grad4);
+		
+		#pragma ignore warning(perf)
+		const varying float Stiffness = StiffnessTable[StiffnessIndex];
+		#pragma ignore warning(perf)
+		const varying float BucklingStiffness = BucklingStiffnessTable[BucklingStiffnessIndex];
+		
+		const varying float RestAngle = RestAngles[i];
+		const varying bool AngleIsBuckled = FLOAT_PI - abs(Angle) < BucklingRatio * (FLOAT_PI - abs(RestAngle));
+		const varying float StiffnessValue = select(AngleIsBuckled, BucklingStiffness, Stiffness);
+		
+		const varying float Lambda = Lambdas[i];
+		const varying float Alpha  = 1.f / (StiffnessValue * Dt * Dt);
+		const varying float Denom = M1 * VectorSizeSquared(Grad1) + M2 * VectorSizeSquared(Grad2) + M3 * VectorSizeSquared(Grad3) + M4 * VectorSizeSquared(Grad4) + Alpha;
+
+		const varying float ComputedDLambda = (Angle - RestAngle - Alpha * Lambda) / Denom;
+		const varying float DLambda = select( StiffnessValue < MinStiffness, (uniform float)0, ComputedDLambda);
+		Lambdas[i] = Lambda + DLambda;
+
+		if (M1 > 0)
+		{
+			VectorScatter(&PandInvM[i1], SetVector4(P1 - DLambda * M1 * Grad1, M1));
+		}
+		if (M2 > 0)
+		{
+			VectorScatter(&PandInvM[i2], SetVector4(P2 - DLambda * M2 * Grad2, M2));
+		}
+		if (M3 > 0)
+		{
+			VectorScatter(&PandInvM[i3], SetVector4(P3 - DLambda * M3 * Grad3, M3));
+		}
+		if (M4 > 0)
+		{
+			VectorScatter(&PandInvM[i4], SetVector4(P4 - DLambda * M4 * Grad4, M4));
+		}
+	}
+}
diff --git a/Engine/Source/Runtime/Experimental/Chaos/Private/Chaos/XPBDSpringConstraints.cpp b/Engine/Source/Runtime/Experimental/Chaos/Private/Chaos/XPBDSpringConstraints.cpp
--- a/Engine/Source/Runtime/Experimental/Chaos/Private/Chaos/XPBDSpringConstraints.cpp
+++ b/Engine/Source/Runtime/Experimental/Chaos/Private/Chaos/XPBDSpringConstraints.cpp
@@ -96,6 +96,10 @@ void FXPBDSpringConstraints::Apply(FSolverParticles& Particles, const FSolverRea
 		if (!Stiffness.HasWeightMap())
 		{
 			const FSolverReal ExpStiffnessValue = (FSolverReal)Stiffness;
+			if (ExpStiffnessValue < XPBDSpringMinStiffness)
+			{
+				return;
+			}
 
 #if INTEL_ISPC
 			if (bRealTypeCompatibleWithISPC && bChaos_XPBDSpring_ISPC_Enabled)
@@ -144,6 +148,7 @@ void FXPBDSpringConstraints::Apply(FSolverParticles& Particles, const FSolverRea
 						&Dists.GetData()[ColorStart],
 						&Lambdas.GetData()[ColorStart],
 						Dt,
+						XPBDSpringMinStiffness,
 						&Stiffness.GetIndices().GetData()[ColorStart],
 						&Stiffness.GetTable().GetData()[0],
 						ColorSize);
@@ -171,6 +176,10 @@ void FXPBDSpringConstraints::Apply(FSolverParticles& Particles, const FSolverRea
 		if (!Stiffness.HasWeightMap())
 		{
 			const FSolverReal ExpStiffnessValue = (FSolverReal)Stiffness;
+			if (ExpStiffnessValue < XPBDSpringMinStiffness)
+			{
+				return;
+			}
 			for (int32 ConstraintIndex = 0; ConstraintIndex < Constraints.Num(); ++ConstraintIndex)
 			{
 				ApplyHelper(Particles, Dt, ConstraintIndex, ExpStiffnessValue);
diff --git a/Engine/Source/Runtime/Experimental/Chaos/Private/Chaos/XPBDSpringConstraints.ispc b/Engine/Source/Runtime/Experimental/Chaos/Private/Chaos/XPBDSpringConstraints.ispc
--- a/Engine/Source/Runtime/Experimental/Chaos/Private/Chaos/XPBDSpringConstraints.ispc
+++ b/Engine/Source/Runtime/Experimental/Chaos/Private/Chaos/XPBDSpringConstraints.ispc
@@ -69,6 +69,7 @@ export void ApplyXPBDSpringConstraintsWithWeightMaps(uniform FVector4f PandInvM[
 									const uniform float Dists[],
 									uniform float Lambdas[],
 									const uniform float Dt,
+									const uniform float MinStiffness,
 									const uniform uint8 StiffnessIndices[],
 									const uniform float StiffnessTable[],
 									const uniform int32 NumConstraints)
@@ -108,7 +109,7 @@ export void ApplyXPBDSpringConstraintsWithWeightMaps(uniform FVector4f PandInvM[
 
 		const FVector3f ComputedDelta = DLambda * Direction;
 
-		const FVector3f Delta = VectorSelect((InvM2 == FLOAT_ZERO && InvM1 == FLOAT_ZERO), FloatZeroVector, ComputedDelta);
+		const FVector3f Delta = VectorSelect(Stiffness < MinStiffness || (InvM2 == FLOAT_ZERO && InvM1 == FLOAT_ZERO), FloatZeroVector, ComputedDelta);
 
 		if (InvM1 > FLOAT_ZERO)
 		{
diff --git a/Engine/Source/Runtime/Experimental/Chaos/Public/Chaos/PBDSoftsEvolutionFwd.h b/Engine/Source/Runtime/Experimental/Chaos/Public/Chaos/PBDSoftsEvolutionFwd.h
--- a/Engine/Source/Runtime/Experimental/Chaos/Public/Chaos/PBDSoftsEvolutionFwd.h
+++ b/Engine/Source/Runtime/Experimental/Chaos/Public/Chaos/PBDSoftsEvolutionFwd.h
@@ -47,6 +47,7 @@ namespace Chaos
 		class FPBDSpringConstraints;
 		class FXPBDSpringConstraints;
 		class FPBDBendingConstraints;
+		class FXPBDBendingConstraints;
 		class FPBDAxialSpringConstraints;
 		class FXPBDAxialSpringConstraints;
 		class FPBDVolumeConstraint;
diff --git a/Engine/Source/Runtime/Experimental/Chaos/Public/Chaos/PBDStiffness.h b/Engine/Source/Runtime/Experimental/Chaos/Public/Chaos/PBDStiffness.h
--- a/Engine/Source/Runtime/Experimental/Chaos/Public/Chaos/PBDStiffness.h
+++ b/Engine/Source/Runtime/Experimental/Chaos/Public/Chaos/PBDStiffness.h
@@ -56,9 +56,17 @@ public:
 	/**
 	 * Set the low and high values of the weight map.
 	 * The weight map table only gets updated after ApplyValues is called.
+	 * Low and high values are clamped between [0,1]
 	 */
 	void SetWeightedValue(const FSolverVec2& InWeightedValue) { WeightedValue = InWeightedValue.ClampAxes((FSolverReal)0., (FSolverReal)1.); }
 
+	/**
+	 * Set the low and high values of the weight map.
+	 * The weight map table only gets updated after ApplyValues is called.
+	 * Low and high values are not clamped. Commonly used for XPBD Stiffness values which are not [0,1]
+	 */
+	void SetWeightedValueUnclamped(const FSolverVec2& InWeightedValue) { WeightedValue = InWeightedValue; }
+
 	/**
 	 * Return the low and high values set for this weight map.
 	 * Both values will always be between 0 and 1 due to having been clamped in SetWeightedValue.
@@ -69,7 +77,7 @@ public:
 	inline void ApplyValues(const FSolverReal Dt, const int32 NumIterations);
 
 	/** Update the weight map table with the current simulation parameters. */
-	inline void ApplyXPBDValues(const FSolverReal MinStiffness, const FSolverReal MaxStiffnesss);
+	inline void ApplyXPBDValues(const FSolverReal MaxStiffnesss);
 
 	/**
 	 * Lookup for the exponential weighted value at the specified weight map index.
@@ -238,23 +246,14 @@ void FPBDStiffness::ApplyValues(const FSolverReal Dt, const int32 NumIterations)
 	}
 }
 
-void FPBDStiffness::ApplyXPBDValues(const FSolverReal MinStiffness, const FSolverReal MaxStiffness)
+void FPBDStiffness::ApplyXPBDValues(const FSolverReal MaxStiffness)
 {
 	SCOPE_CYCLE_COUNTER(STAT_PBD_StiffnessApplyValues);
 	
-	// XPBD internally handles the effects of iterations and dt. 
-	// PBD stiffness is more like a constraint compliance. This should scale like 1/stiffness which 
-	// is an actual measure of stiffness (e.g., something with units like Pascals)	
-	const FSolverReal MinOverMaxStiffness = MinStiffness / MaxStiffness;
-	auto SimulationValue = [this, MinOverMaxStiffness, MinStiffness, MaxStiffness](const FSolverReal InValue)->FSolverReal
+	auto SimulationValue = [this, MaxStiffness](const FSolverReal InValue)->FSolverReal
 	{
-		const FSolverReal ClampedInValue = FMath::Clamp(InValue, (FSolverReal)0., (FSolverReal)1.);
-		if (ClampedInValue == (FSolverReal)1.)
-		{
-			return MaxStiffness;
-		}
-		const FSolverReal ParameterFit = CalcExponentialParameterFit(ParameterFitBase, ParameterFitLogBase, InValue);
-		return MinStiffness / ((FSolverReal)1. - ParameterFit + MinOverMaxStiffness);
+		// Do not apply exponential to XPBDStiffnesses. They are authored in terms of true stiffness values (e.g., kg/s^2), not exponential compliance
+		return FMath::Clamp(InValue, (FSolverReal)0., MaxStiffness);
 	};
 
 	const FSolverReal Offset = WeightedValue[0];
diff --git a/Engine/Source/Runtime/Experimental/Chaos/Public/Chaos/XPBDAxialSpringConstraints.h b/Engine/Source/Runtime/Experimental/Chaos/Public/Chaos/XPBDAxialSpringConstraints.h
--- a/Engine/Source/Runtime/Experimental/Chaos/Public/Chaos/XPBDAxialSpringConstraints.h
+++ b/Engine/Source/Runtime/Experimental/Chaos/Public/Chaos/XPBDAxialSpringConstraints.h
@@ -9,9 +9,9 @@ DECLARE_CYCLE_STAT(TEXT("Chaos XPBD Axial Spring Constraint"), STAT_XPBD_AxialSp
 namespace Chaos::Softs
 {
 
-// Stiffness is in N/CM^2, so it needs to be adjusted from the PBD stiffness ranging between [0,1]
-static const FSolverReal XPBDAxialSpringMinStiffness = (FSolverReal)1e-1;
-static const FSolverReal XPBDAxialSpringMaxStiffness = (FSolverReal)1e7;  // Max stiffness: 1e+11 N/M^2 = 1e+7 N/CM^2
+// Stiffness is in kg/s^2
+static const FSolverReal XPBDAxialSpringMinStiffness = (FSolverReal)1e-4; // Stiffness below this will be considered 0 since all of our calculations are actually based on 1 / stiffness.
+static const FSolverReal XPBDAxialSpringMaxStiffness = (FSolverReal)1e7;
 
 class FXPBDAxialSpringConstraints final : public FPBDAxialSpringConstraintsBase
 {
@@ -37,7 +37,7 @@ public:
 
 	virtual ~FXPBDAxialSpringConstraints() override {}
 
-	void ApplyProperties(const FSolverReal Dt, const int32 NumIterations) { Stiffness.ApplyXPBDValues(XPBDAxialSpringMinStiffness, XPBDAxialSpringMaxStiffness); }
+	void ApplyProperties(const FSolverReal Dt, const int32 NumIterations) { Stiffness.ApplyXPBDValues(XPBDAxialSpringMaxStiffness); }
 
 	void Init() const { for (FSolverReal& Lambda : Lambdas) { Lambda = (FSolverReal)0.; } }
 
@@ -47,6 +47,10 @@ public:
 		if (!Stiffness.HasWeightMap())
 		{
 			const FSolverReal ExpStiffnessValue = (FSolverReal)Stiffness;
+			if (ExpStiffnessValue < XPBDAxialSpringMinStiffness)
+			{
+				return;
+			}
 			for (int32 ConstraintIndex = 0; ConstraintIndex < Constraints.Num(); ++ConstraintIndex)
 			{
 				const TVector<int32, 3>& constraint = Constraints[ConstraintIndex];
@@ -97,7 +101,7 @@ public:
 	}
 
 private:
-	FSolverVec3 GetDelta(const FSolverParticles& Particles, const FSolverReal Dt, const int32 InConstraintIndex, const FSolverReal ExpStiffnessValue) const
+	FSolverVec3 GetDelta(const FSolverParticles& Particles, const FSolverReal Dt, const int32 InConstraintIndex, const FSolverReal StiffnessValue) const
 	{
 		const TVector<int32, 3>& Constraint = Constraints[InConstraintIndex];
 		const int32 i1 = Constraint[0];
@@ -106,7 +110,7 @@ private:
 
 		const FSolverReal Bary = Barys[InConstraintIndex];
 		const FSolverReal PInvMass = Particles.InvM(i3) * ((FSolverReal)1. - Bary) + Particles.InvM(i2) * Bary;
-		if (Particles.InvM(i1) == (FSolverReal)0. && PInvMass == (FSolverReal)0.)
+		if (StiffnessValue < XPBDAxialSpringMinStiffness || ( Particles.InvM(i1) == (FSolverReal)0. && PInvMass == (FSolverReal)0.))
 		{
 			return FSolverVec3((FSolverReal)0.);
 		}
@@ -128,7 +132,7 @@ private:
 		const FSolverReal Offset = (Distance - Dists[InConstraintIndex]);
 
 		FSolverReal& Lambda = Lambdas[InConstraintIndex];
-		const FSolverReal Alpha = (FSolverReal)1 / (ExpStiffnessValue * Dt * Dt);
+		const FSolverReal Alpha = (FSolverReal)1 / (StiffnessValue * Dt * Dt);
 
 		const FSolverReal DLambda = (Offset - Alpha * Lambda) / (CombinedInvMass + Alpha);
 		const FSolverVec3 Delta = DLambda * Direction;
diff --git a/Engine/Source/Runtime/Experimental/Chaos/Public/Chaos/XPBDBendingConstraints.h b/Engine/Source/Runtime/Experimental/Chaos/Public/Chaos/XPBDBendingConstraints.h
new file mode 100644
--- /dev/null
+++ b/Engine/Source/Runtime/Experimental/Chaos/Public/Chaos/XPBDBendingConstraints.h
@@ -0,0 +1,75 @@
+// Copyright Epic Games, Inc. All Rights Reserved.
+#pragma once
+
+#include "Chaos/PBDBendingConstraintsBase.h"
+#include "ChaosStats.h"
+
+
+namespace Chaos::Softs
+{
+
+// Stiffness is in kg cm^2 / rad^2 s^2
+static const FSolverReal XPBDBendMinStiffness = (FSolverReal)1e-4; // Stiffness below this will be considered 0 since all of our calculations are actually based on 1 / stiffness.
+static const FSolverReal XPBDBendMaxStiffness = (FSolverReal)1e7; 
+
+class CHAOS_API FXPBDBendingConstraints final : public FPBDBendingConstraintsBase
+{
+	typedef FPBDBendingConstraintsBase Base;
+	using Base::Constraints;
+	using Base::RestAngles;
+	using Base::Stiffness;
+	using Base::BucklingStiffness;
+
+public:
+	FXPBDBendingConstraints(const FSolverParticles& InParticles,
+		int32 ParticleOffset,
+		int32 ParticleCount,
+		TArray<TVec4<int32>>&& InConstraints,
+		const TConstArrayView<FRealSingle>& StiffnessMultipliers,
+		const TConstArrayView<FRealSingle>& BucklingStiffnessMultipliers,
+		const FSolverVec2& InStiffness,
+		const FSolverReal InBucklingRatio,
+		const FSolverVec2& InBucklingStiffness,
+		bool bTrimKinematicConstraints = false)
+		: Base(InParticles, ParticleOffset, ParticleCount, MoveTemp(InConstraints), StiffnessMultipliers, BucklingStiffnessMultipliers, InStiffness, InBucklingRatio, InBucklingStiffness, bTrimKinematicConstraints)
+	{
+		Lambdas.Init((FSolverReal)0., Constraints.Num());
+		InitColor(InParticles);
+	}
+
+	virtual ~FXPBDBendingConstraints() override {}
+
+	void Init() const { for (FSolverReal& Lambda : Lambdas) { Lambda = (FSolverReal)0.; } }
+
+	// Update stiffness values
+	void SetProperties(const FSolverVec2& InStiffness, const FSolverReal InBucklingRatio, const FSolverVec2& InBucklingStiffness)
+	{
+		Stiffness.SetWeightedValueUnclamped(InStiffness);
+		BucklingRatio = InBucklingRatio;
+		BucklingStiffness.SetWeightedValueUnclamped(InBucklingStiffness);
+	}
+
+	// Update stiffness table, as well as the simulation stiffness exponent
+	void ApplyProperties(const FSolverReal Dt, const int32 NumIterations) { Stiffness.ApplyXPBDValues(XPBDBendMaxStiffness); BucklingStiffness.ApplyXPBDValues(XPBDBendMaxStiffness); }
+
+	void Apply(FSolverParticles& Particles, const FSolverReal Dt) const;
+
+private:
+	void InitColor(const FSolverParticles& InParticles);
+	void ApplyHelper(FSolverParticles& Particles, const FSolverReal Dt, const int32 ConstraintIndex, const FSolverReal ExpStiffnessValue, const FSolverReal ExpBucklingValue) const;
+
+private:
+	mutable TArray<FSolverReal> Lambdas;
+	TArray<int32> ConstraintsPerColorStartIndex; // Constraints are ordered so each batch is contiguous. This is ColorNum + 1 length so it can be used as start and end.
+};
+
+}  // End namespace Chaos::Softs
+
+// Support ISPC enable/disable in non-shipping builds
+#if !INTEL_ISPC
+const bool bChaos_XPBDBending_ISPC_Enabled = false;
+#elif UE_BUILD_SHIPPING
+const bool bChaos_XPBDBending_ISPC_Enabled = true;
+#else
+extern CHAOS_API bool bChaos_XPBDBending_ISPC_Enabled;
+#endif
diff --git a/Engine/Source/Runtime/Experimental/Chaos/Public/Chaos/XPBDLongRangeConstraints.h b/Engine/Source/Runtime/Experimental/Chaos/Public/Chaos/XPBDLongRangeConstraints.h
--- a/Engine/Source/Runtime/Experimental/Chaos/Public/Chaos/XPBDLongRangeConstraints.h
+++ b/Engine/Source/Runtime/Experimental/Chaos/Public/Chaos/XPBDLongRangeConstraints.h
@@ -10,7 +10,6 @@ DECLARE_CYCLE_STAT(TEXT("Chaos XPBD Long Range Constraint"), STAT_XPBD_LongRange
 namespace Chaos::Softs
 {
 
-// Stiffness is in N/CM^2, so it needs to be adjusted from the PBD stiffness ranging between [0,1]
 static const FSolverReal XPBDLongRangeMinStiffness = (FSolverReal)1e-1;
 static const FSolverReal XPBDLongRangeMaxStiffness = (FSolverReal)1e7;
 
@@ -42,7 +41,7 @@ public:
 	virtual ~FXPBDLongRangeConstraints() override {}
 
 	// Set stiffness offset and range, as well as the simulation stiffness exponent
-	void ApplyProperties(const FSolverReal Dt, const int32 NumIterations) { Stiffness.ApplyXPBDValues(XPBDLongRangeMinStiffness, XPBDLongRangeMaxStiffness); ApplyScale(); }
+	void ApplyProperties(const FSolverReal Dt, const int32 NumIterations) { Stiffness.ApplyXPBDValues(XPBDLongRangeMaxStiffness); ApplyScale(); }
 
 	void Init() const
 	{
@@ -130,6 +129,10 @@ public:
 private:
 	void Apply(FSolverParticles& Particles, const FSolverReal Dt, const FTether& Tether, int32 ConstraintIndex, const FSolverReal InStiffness, const FSolverReal InScale) const
 	{
+		if (InStiffness < XPBDLongRangeMinStiffness)
+		{
+			return;
+		}
 		FSolverVec3 Direction;
 		FSolverReal Offset;
 		GetDelta(Particles, Tether, InScale, Direction, Offset);
diff --git a/Engine/Source/Runtime/Experimental/Chaos/Public/Chaos/XPBDSpringConstraints.h b/Engine/Source/Runtime/Experimental/Chaos/Public/Chaos/XPBDSpringConstraints.h
--- a/Engine/Source/Runtime/Experimental/Chaos/Public/Chaos/XPBDSpringConstraints.h
+++ b/Engine/Source/Runtime/Experimental/Chaos/Public/Chaos/XPBDSpringConstraints.h
@@ -8,9 +8,9 @@
 namespace Chaos::Softs
 {
 
-// Stiffness is in N/CM^2, so it needs to be adjusted from the PBD stiffness ranging between [0,1]
-static const FSolverReal XPBDSpringMinStiffness = (FSolverReal)1e-1; // Min stiffness: 1e3 N/M^2 = 1e-1 N/CM^2. Empirically defined to have similar look to XPBDSpringConstraints.
-static const FSolverReal XPBDSpringMaxStiffness = (FSolverReal)1e7;  // Max stiffness: 1e+11 N/M^2 = 1e+7 N/CM^2
+// Stiffness is in kg/s^2
+static const FSolverReal XPBDSpringMinStiffness = (FSolverReal)1e-4; // Stiffness below this will be considered 0 since all of our calculations are actually based on 1 / stiffness.
+static const FSolverReal XPBDSpringMaxStiffness = (FSolverReal)1e7; 
 
 class CHAOS_API FXPBDSpringConstraints final : public FPBDSpringConstraintsBase
 {
@@ -39,24 +39,27 @@ public:
 	virtual ~FXPBDSpringConstraints() override {}
 
 	void Init() const { for (FSolverReal& Lambda : Lambdas) { Lambda = (FSolverReal)0.; } }
+
+	// Update stiffness values
+	void SetProperties(const FSolverVec2& InStiffness) { Stiffness.SetWeightedValueUnclamped(InStiffness); }
 	
 	// Update stiffness table, as well as the simulation stiffness exponent
-	inline void ApplyProperties(const FSolverReal Dt, const int32 NumIterations) { Stiffness.ApplyXPBDValues(XPBDSpringMinStiffness, XPBDSpringMaxStiffness); }
+	inline void ApplyProperties(const FSolverReal Dt, const int32 NumIterations) { Stiffness.ApplyXPBDValues(XPBDSpringMaxStiffness); }
 
 	void Apply(FSolverParticles& Particles, const FSolverReal Dt) const;
 
 private:
 	void InitColor(const FSolverParticles& InParticles);
-	void ApplyHelper(FSolverParticles& Particles, const FSolverReal Dt, const int32 ConstraintIndex, const FSolverReal ExpStiffnessValue) const;
+	void ApplyHelper(FSolverParticles& Particles, const FSolverReal Dt, const int32 ConstraintIndex, const FSolverReal StiffnessValue) const;
 
-	FSolverVec3 GetDelta(const FSolverParticles& Particles, const FSolverReal Dt, const int32 ConstraintIndex, const FSolverReal ExpStiffnessValue) const
+	FSolverVec3 GetDelta(const FSolverParticles& Particles, const FSolverReal Dt, const int32 ConstraintIndex, const FSolverReal StiffnessValue) const
 	{
 		const TVec2<int32>& Constraint = Constraints[ConstraintIndex];
 
 		const int32 i1 = Constraint[0];
 		const int32 i2 = Constraint[1];
 
-		if (Particles.InvM(i2) == (FSolverReal)0. && Particles.InvM(i1) == (FSolverReal)0.)
+		if (StiffnessValue < XPBDSpringMinStiffness || (Particles.InvM(i2) == (FSolverReal)0. && Particles.InvM(i1) == (FSolverReal)0.))
 		{
 			return FSolverVec3((FSolverReal)0.);
 		}
@@ -69,7 +72,7 @@ private:
 		const FSolverReal Offset = Distance - Dists[ConstraintIndex];
 
 		FSolverReal& Lambda = Lambdas[ConstraintIndex];
-		const FSolverReal Alpha = (FSolverReal)1.f / (ExpStiffnessValue * Dt * Dt);
+		const FSolverReal Alpha = (FSolverReal)1.f / (StiffnessValue * Dt * Dt);
 
 		const FSolverReal DLambda = (Offset - Alpha * Lambda) / (CombinedInvMass + Alpha);
 		const FSolverVec3 Delta = DLambda * Direction;
