Chaos : destruction
- add the ability for the rest collection to specify the connection graph ( if not the proxy will generate it as it is now )
- Proxy on going cleanup

diff --git a/Engine/Source/Runtime/Experimental/Chaos/Private/PhysicsProxy/GeometryCollectionPhysicsProxy.cpp b/Engine/Source/Runtime/Experimental/Chaos/Private/PhysicsProxy/GeometryCollectionPhysicsProxy.cpp
--- a/Engine/Source/Runtime/Experimental/Chaos/Private/PhysicsProxy/GeometryCollectionPhysicsProxy.cpp
+++ b/Engine/Source/Runtime/Experimental/Chaos/Private/PhysicsProxy/GeometryCollectionPhysicsProxy.cpp
@@ -74,6 +74,12 @@ FAutoConsoleVariableRef CVarGeometryCollectionAlwaysGenerateGTCollisionForCluste
 	bGeometryCollectionAlwaysGenerateGTCollisionForClusters,
 	TEXT("When enabled, always generate a game thread side collision for clusters.[def: true]"));
 
+bool bGeometryCollectionAlwaysGenerateConnectionGraph = false;
+FAutoConsoleVariableRef CVarGeometryCollectionAlwaysGenerateConnectionGraph(
+	TEXT("p.GeometryCollection.AlwaysGenerateConnectionGraph"),
+	bGeometryCollectionAlwaysGenerateConnectionGraph,
+	TEXT("When enabled, always  generate the cluster's connection graph instead of using the rest collection stored one - Note: this should only be used for troubleshooting.[def: false]"));
+
 DEFINE_LOG_CATEGORY_STATIC(UGCC_LOG, Error, All);
 
 static const FSharedSimulationSizeSpecificData& GetSizeSpecificData(const TArray<FSharedSimulationSizeSpecificData>& SizeSpecificData, const FGeometryCollection& RestCollection, const int32 TransformIndex, const FBox& BoundingBox);
@@ -680,12 +686,64 @@ void FGeometryCollectionPhysicsProxy::InitializeDynamicCollection(FGeometryDynam
 int32 ReportTooManyChildrenNum = -1;
 FAutoConsoleVariableRef CVarReportTooManyChildrenNum(TEXT("p.ReportTooManyChildrenNum"), ReportTooManyChildrenNum, TEXT("Issue warning if more than this many children exist in a single cluster"));
 
+void FGeometryCollectionPhysicsProxy::CreateNonClusteredParticles(Chaos::FPBDRigidsSolver* RigidsSolver, const FGeometryCollection& RestCollection, const FGeometryDynamicCollection& DynamicCollection)
+{
+	const TManagedArray<bool>& SimulatableParticles = DynamicCollection.SimulatableParticles;
+	
+	//const int NumRigids = 0; // ryan - Since we're doing SOA, we start at zero?
+	int NumRigids = 0;
+	BaseParticleIndex = NumRigids;
+
+	// Gather unique indices from GT to pass into PT handle creation
+	TArray<Chaos::FUniqueIdx> UniqueIndices;
+	UniqueIndices.Reserve(SimulatableParticles.Num());
+
+	// Count geometry collection leaf node particles to add
+	int NumSimulatedParticles = 0;
+	for (int32 Idx = 0; Idx < SimulatableParticles.Num(); ++Idx)
+	{
+		NumSimulatedParticles += SimulatableParticles[Idx];
+		if (SimulatableParticles[Idx] && !RestCollection.IsClustered(Idx) && RestCollection.IsGeometry(Idx))
+		{
+			NumRigids++;
+			UniqueIndices.Add(GTParticles[Idx]->UniqueIdx());
+		}
+	}
+
+	// Add entries into simulation array
+	RigidsSolver->GetEvolution()->ReserveParticles(NumSimulatedParticles);
+	TArray<Chaos::TPBDGeometryCollectionParticleHandle<Chaos::FReal, 3>*> Handles = RigidsSolver->GetEvolution()->CreateGeometryCollectionParticles(NumRigids, UniqueIndices.GetData());
+
+	int32 NextIdx = 0;
+	for (int32 Idx = 0; Idx < SimulatableParticles.Num(); ++Idx)
+	{
+		SolverParticleHandles[Idx] = nullptr;
+		if (SimulatableParticles[Idx] && !RestCollection.IsClustered(Idx))
+		{
+			// todo: Unblocked read access of game thread data on the physics thread.
+
+			Chaos::TPBDGeometryCollectionParticleHandle<Chaos::FReal, 3>* Handle = Handles[NextIdx++];
+
+			Handle->SetPhysicsProxy(this);
+
+			SolverParticleHandles[Idx] = Handle;
+			HandleToTransformGroupIndex.Add(Handle, Idx);
+
+			// We're on the physics thread here but we've already set up the GT particles and we're just linking here
+			Handle->GTGeometryParticle() = GTParticles[Idx].Get();
+
+			check(SolverParticleHandles[Idx]->GetParticleType() == Handle->GetParticleType());
+			RigidsSolver->GetEvolution()->CreateParticle(Handle);
+		}
+	}
+}
+
 void FGeometryCollectionPhysicsProxy::InitializeBodiesPT(Chaos::FPBDRigidsSolver* RigidsSolver, typename Chaos::FPBDRigidsSolver::FParticlesType& Particles)
 {
 	const FGeometryCollection* RestCollection = Parameters.RestCollection;
 	const FGeometryDynamicCollection& DynamicCollection = PhysicsThreadCollection;
 
-	if (Parameters.Simulating)
+	if (Parameters.Simulating && RestCollection)
 	{
 		const TManagedArray<int32>& TransformIndex = RestCollection->TransformIndex;
 		const TManagedArray<int32>& BoneMap = RestCollection->BoneMap;
@@ -709,52 +767,7 @@ void FGeometryCollectionPhysicsProxy::InitializeBodiesPT(Chaos::FPBDRigidsSolver
 		TArray<FTransform> Transform;
 		GeometryCollectionAlgo::GlobalMatrices(DynamicCollection.Transform, Parent, Transform);
 
-		//const int NumRigids = 0; // ryan - Since we're doing SOA, we start at zero?
-		int NumRigids = 0;
-		BaseParticleIndex = NumRigids;
-
-		// Gather unique indices from GT to pass into PT handle creation
-		TArray<Chaos::FUniqueIdx> UniqueIndices;
-		UniqueIndices.Reserve(SimulatableParticles.Num());
-
-		// Count geometry collection leaf node particles to add
-		int NumSimulatedParticles = 0;
-		for (int32 Idx = 0; Idx < SimulatableParticles.Num(); ++Idx)
-		{
-			NumSimulatedParticles += SimulatableParticles[Idx];
-			if (SimulatableParticles[Idx] && !RestCollection->IsClustered(Idx) && RestCollection->IsGeometry(Idx))
-			{
-				NumRigids++;
-				UniqueIndices.Add(GTParticles[Idx]->UniqueIdx());
-			}
-		}
-
-		// Add entries into simulation array
-		RigidsSolver->GetEvolution()->ReserveParticles(NumSimulatedParticles);
-		TArray<Chaos::TPBDGeometryCollectionParticleHandle<Chaos::FReal, 3>*> Handles = RigidsSolver->GetEvolution()->CreateGeometryCollectionParticles(NumRigids, UniqueIndices.GetData());
-
-		int32 NextIdx = 0;
-		for (int32 Idx = 0; Idx < SimulatableParticles.Num(); ++Idx)
-		{
-			SolverParticleHandles[Idx] = nullptr;
-			if (SimulatableParticles[Idx] && !RestCollection->IsClustered(Idx))
-			{
-				// todo: Unblocked read access of game thread data on the physics thread.
-
-				Chaos::TPBDGeometryCollectionParticleHandle<Chaos::FReal, 3>* Handle = Handles[NextIdx++];
-
-				Handle->SetPhysicsProxy(this);
-
-				SolverParticleHandles[Idx] = Handle;
-				HandleToTransformGroupIndex.Add(Handle, Idx);
-
-				// We're on the physics thread here but we've already set up the GT particles and we're just linking here
-				Handle->GTGeometryParticle() = GTParticles[Idx].Get();
-
-				check(SolverParticleHandles[Idx]->GetParticleType() == Handle->GetParticleType());
-				RigidsSolver->GetEvolution()->CreateParticle(Handle);
-			}
-		}
+		CreateNonClusteredParticles(RigidsSolver, *RestCollection, DynamicCollection);
 
 		const float StrainDefault = Parameters.DamageThreshold.Num() ? Parameters.DamageThreshold[0] : 0;
 		// Add the rigid bodies
@@ -973,23 +986,58 @@ void FGeometryCollectionPhysicsProxy::InitializeBodiesPT(Chaos::FPBDRigidsSolver
 			// as well as views of clustered particles.
 			Particles.UpdateGeometryCollectionViews(true); 
 
-			// Set cluster connectivity.  TPBDRigidClustering::CreateClusterParticle() 
-			// will optionally do this, but we switch that functionality off in BuildClusters_Internal().
-			for(int32 TransformGroupIndex = 0; TransformGroupIndex < NumTransforms; ++TransformGroupIndex)
+			const TManagedArray<TSet<int32>>* Connections = RestCollection->FindAttribute<TSet<int32>>("Connections", FGeometryCollection::TransformGroup);
+			const bool bGenerateConnectionGraph = (!Connections) || (Connections->Num() != RestCollection->Transform.Num()) || bGeometryCollectionAlwaysGenerateConnectionGraph;
+			if (bGenerateConnectionGraph)
 			{
-				if (RestCollection->IsClustered(TransformGroupIndex))
+				UE_LOG(LogChaos, Warning, TEXT("Generating GC connection graph on the physics thread"));
+				// Set cluster connectivity.  TPBDRigidClustering::CreateClusterParticle() 
+				// will optionally do this, but we switch that functionality off in BuildClusters_Internal().
+				for (int32 TransformGroupIndex = 0; TransformGroupIndex < NumTransforms; ++TransformGroupIndex)
 				{
-					if (SolverParticleHandles[TransformGroupIndex])
+					if (RestCollection->IsClustered(TransformGroupIndex))
 					{
-						Chaos::FClusterCreationParameters ClusterParams;
-						// #todo: should other parameters be set here?  Previously, there was no parameters being sent, and it is unclear
-						// where some of these parameters are defined (ie: CoillisionThicknessPercent)
-						ClusterParams.ConnectionMethod = Parameters.ClusterConnectionMethod;
+						if (FClusterHandle* ClusteredParticle = SolverParticleHandles[TransformGroupIndex])
+						{
+							Chaos::FClusterCreationParameters ClusterParams;
+							// #todo: should other parameters be set here?  Previously, there was no parameters being sent, and it is unclear
+							// where some of these parameters are defined (ie: CollisionThicknessPercent)
+							ClusterParams.ConnectionMethod = Parameters.ClusterConnectionMethod;
 						
-						RigidsSolver->GetEvolution()->GetRigidClustering().GenerateConnectionGraph(SolverParticleHandles[TransformGroupIndex], ClusterParams);
+							RigidsSolver->GetEvolution()->GetRigidClustering().GenerateConnectionGraph(ClusteredParticle, ClusterParams);
+						}
 					}
 				}
 			}
+			else
+			{
+				for (int32 TransformGroupIndex = 0; TransformGroupIndex < NumTransforms; ++TransformGroupIndex)
+				{
+					if (FClusterHandle* ClusteredParticle = SolverParticleHandles[TransformGroupIndex]) 
+					{
+						const TSet<int32>& Siblings = (*Connections)[TransformGroupIndex];
+						for (const int32 SiblingTransformIndex: Siblings)
+						{
+							if (FClusterHandle* OtherClusteredParticle = SolverParticleHandles[SiblingTransformIndex])
+							{
+								// we add both connection for integrity so let use the index order as a way to filter out the symmetrical part 
+								if (SiblingTransformIndex > TransformGroupIndex)
+								{
+									Chaos::TConnectivityEdge<Chaos::FReal> Edge;
+									Edge.Strain = (ClusteredParticle->Strains() + OtherClusteredParticle->Strains()) * 0.5f;
+									
+									Edge.Sibling = OtherClusteredParticle;
+									ClusteredParticle->ConnectivityEdges().Add(Edge);
+									
+									Edge.Sibling = ClusteredParticle;
+									OtherClusteredParticle->ConnectivityEdges().Add(Edge);
+								}
+							}
+						}
+					}
+				}
+				// load connection graph from RestCollection
+			}
 		} // end if EnableClustering
  
 
@@ -1719,8 +1767,8 @@ void FGeometryCollectionPhysicsProxy::BufferPhysicsResults_External(Chaos::FDirt
 	//	TargetResults.Parent;
 	//
 	//TargetResults.DynamicState; // ObjectState
-	TManagedArray<FVector3f>* LinearVelocities = GameThreadCollection.FindAttributeTyped<FVector3f>("LinearVelocity", FTransformCollection::TransformGroup);
-	TManagedArray<FVector3f>* AngularVelocities = GameThreadCollection.FindAttributeTyped<FVector3f>("AngularVelocity", FTransformCollection::TransformGroup);
+	const TManagedArray<FVector3f>* LinearVelocities = GameThreadCollection.FindAttributeTyped<FVector3f>("LinearVelocity", FTransformCollection::TransformGroup);
+	const TManagedArray<FVector3f>* AngularVelocities = GameThreadCollection.FindAttributeTyped<FVector3f>("AngularVelocity", FTransformCollection::TransformGroup);
 	
 	for (int32 TransformGroupIndex = 0; TransformGroupIndex < TargetResults.Transforms.Num(); ++TransformGroupIndex)
 	{
@@ -2277,24 +2325,13 @@ void FGeometryCollectionPhysicsProxy::InitializeSharedCollisionStructures(
 	const TManagedArray<int32>& Parent = RestCollection.Parent;
 	const TManagedArray<TSet<int32>>& Children = RestCollection.Children;
 	const TManagedArray<int32>& SimulationType = RestCollection.SimulationType;
-	TManagedArray<bool>& CollectionSimulatableParticles =
-		RestCollection.GetAttribute<bool>(
-			FGeometryCollection::SimulatableParticlesAttribute, FTransformCollection::TransformGroup);
-	TManagedArray<FVector3f>& CollectionInertiaTensor =
-		RestCollection.AddAttribute<FVector3f>(
-			TEXT("InertiaTensor"), FTransformCollection::TransformGroup);
-	TManagedArray<FRealSingle>& CollectionMass =
-		RestCollection.AddAttribute<FRealSingle>(
-			TEXT("Mass"), FTransformCollection::TransformGroup);
-	TManagedArray<TUniquePtr<FSimplicial>>& CollectionSimplicials =
-		RestCollection.AddAttribute<TUniquePtr<FSimplicial>>(
-			FGeometryDynamicCollection::SimplicialsAttribute, FTransformCollection::TransformGroup);
-
-	RestCollection.RemoveAttribute(
-		FGeometryDynamicCollection::ImplicitsAttribute, FTransformCollection::TransformGroup);
-	TManagedArray<FGeometryDynamicCollection::FSharedImplicit>& CollectionImplicits =
-		RestCollection.AddAttribute<FGeometryDynamicCollection::FSharedImplicit>(
-			FGeometryDynamicCollection::ImplicitsAttribute, FTransformCollection::TransformGroup);
+	TManagedArray<bool>& CollectionSimulatableParticles = RestCollection.ModifyAttribute<bool>(FGeometryCollection::SimulatableParticlesAttribute, FTransformCollection::TransformGroup);
+	TManagedArray<FVector3f>& CollectionInertiaTensor = RestCollection.AddAttribute<FVector3f>(TEXT("InertiaTensor"), FTransformCollection::TransformGroup);
+	TManagedArray<FRealSingle>& CollectionMass = RestCollection.AddAttribute<FRealSingle>(TEXT("Mass"), FTransformCollection::TransformGroup);
+	TManagedArray<TUniquePtr<FSimplicial>>& CollectionSimplicials =	RestCollection.AddAttribute<TUniquePtr<FSimplicial>>(FGeometryDynamicCollection::SimplicialsAttribute, FTransformCollection::TransformGroup);
+
+	RestCollection.RemoveAttribute(FGeometryDynamicCollection::ImplicitsAttribute, FTransformCollection::TransformGroup);
+	TManagedArray<FGeometryDynamicCollection::FSharedImplicit>& CollectionImplicits = RestCollection.AddAttribute<FGeometryDynamicCollection::FSharedImplicit>(FGeometryDynamicCollection::ImplicitsAttribute, FTransformCollection::TransformGroup);
 
 	TManagedArray<TSet<int32>>* TransformToConvexIndices = RestCollection.FindAttribute<TSet<int32>>("TransformToConvexIndices", FTransformCollection::TransformGroup);
 	TManagedArray<TUniquePtr<Chaos::FConvex>>* ConvexGeometry = RestCollection.FindAttribute<TUniquePtr<Chaos::FConvex>>("ConvexHull", "Convex");
@@ -2307,9 +2344,7 @@ void FGeometryCollectionPhysicsProxy::InitializeSharedCollisionStructures(
 
 
 	// @todo(chaos_transforms) : do we still use this?
-	TManagedArray<FTransform>& CollectionMassToLocal =
-		RestCollection.AddAttribute<FTransform>(
-			TEXT("MassToLocal"), FTransformCollection::TransformGroup);
+	TManagedArray<FTransform>& CollectionMassToLocal = RestCollection.AddAttribute<FTransform>(TEXT("MassToLocal"), FTransformCollection::TransformGroup);
 	FTransform IdentityXf(FQuat::Identity, FVector(0));
 	IdentityXf.NormalizeRotation();
 	CollectionMassToLocal.Fill(IdentityXf);
diff --git a/Engine/Source/Runtime/Experimental/Chaos/Public/PhysicsProxy/GeometryCollectionPhysicsProxy.h b/Engine/Source/Runtime/Experimental/Chaos/Public/PhysicsProxy/GeometryCollectionPhysicsProxy.h
--- a/Engine/Source/Runtime/Experimental/Chaos/Public/PhysicsProxy/GeometryCollectionPhysicsProxy.h
+++ b/Engine/Source/Runtime/Experimental/Chaos/Public/PhysicsProxy/GeometryCollectionPhysicsProxy.h
@@ -285,6 +285,7 @@ protected:
 
 	void PrepareBufferData(Chaos::FDirtyGeometryCollectionData& BufferData, const FGeometryDynamicCollection& ThreadCollection,  Chaos::FReal SolverLastDt = 0.0);
 
+	void CreateNonClusteredParticles(Chaos::FPBDRigidsSolver* RigidsSolver,	const FGeometryCollection& RestCollection, const FGeometryDynamicCollection& DynamicCollection);
 private:
 
 	FSimulationParameters Parameters;
