Chaos : destruction : fix compilation error from CL20116432

diff --git a/Engine/Source/Runtime/Experimental/Chaos/Private/PhysicsProxy/GeometryCollectionPhysicsProxy.cpp b/Engine/Source/Runtime/Experimental/Chaos/Private/PhysicsProxy/GeometryCollectionPhysicsProxy.cpp
--- a/Engine/Source/Runtime/Experimental/Chaos/Private/PhysicsProxy/GeometryCollectionPhysicsProxy.cpp
+++ b/Engine/Source/Runtime/Experimental/Chaos/Private/PhysicsProxy/GeometryCollectionPhysicsProxy.cpp
@@ -2325,7 +2325,7 @@ void FGeometryCollectionPhysicsProxy::InitializeSharedCollisionStructures(
 	const TManagedArray<int32>& Parent = RestCollection.Parent;
 	const TManagedArray<TSet<int32>>& Children = RestCollection.Children;
 	const TManagedArray<int32>& SimulationType = RestCollection.SimulationType;
-	TManagedArray<bool>& CollectionSimulatableParticles = RestCollection.ModifyAttribute<bool>(FGeometryCollection::SimulatableParticlesAttribute, FTransformCollection::TransformGroup);
+	TManagedArray<bool>& CollectionSimulatableParticles = RestCollection.GetAttribute<bool>(FGeometryCollection::SimulatableParticlesAttribute, FTransformCollection::TransformGroup);
 	TManagedArray<FVector3f>& CollectionInertiaTensor = RestCollection.AddAttribute<FVector3f>(TEXT("InertiaTensor"), FTransformCollection::TransformGroup);
 	TManagedArray<FRealSingle>& CollectionMass = RestCollection.AddAttribute<FRealSingle>(TEXT("Mass"), FTransformCollection::TransformGroup);
 	TManagedArray<TUniquePtr<FSimplicial>>& CollectionSimplicials =	RestCollection.AddAttribute<TUniquePtr<FSimplicial>>(FGeometryDynamicCollection::SimplicialsAttribute, FTransformCollection::TransformGroup);
