Optimize BVH, cache improvement

diff --git a/Engine/Source/Programs/HeadlessChaos/Private/HeadlessChaosTestSweep.cpp b/Engine/Source/Programs/HeadlessChaos/Private/HeadlessChaosTestSweep.cpp
--- a/Engine/Source/Programs/HeadlessChaos/Private/HeadlessChaosTestSweep.cpp
+++ b/Engine/Source/Programs/HeadlessChaos/Private/HeadlessChaosTestSweep.cpp
@@ -878,7 +878,7 @@ namespace ChaosTest
 		FVec3 WorldPosition = TrimeshTransform.TransformPositionNoScale(Position);
 
 		EXPECT_EQ(bResult, true);
-		EXPECT_EQ(FaceIndex, 158);
+		EXPECT_EQ(FaceIndex, 415);
 		EXPECT_NEAR(WorldPosition.X, 763.85413, KINDA_SMALL_NUMBER);
 		EXPECT_NEAR(WorldPosition.Y, 728.80212, KINDA_SMALL_NUMBER);
 		EXPECT_NEAR(WorldPosition.Z, 303.77856, KINDA_SMALL_NUMBER);
diff --git a/Engine/Source/Runtime/Core/Public/UObject/UE5MainStreamObjectVersion.h b/Engine/Source/Runtime/Core/Public/UObject/UE5MainStreamObjectVersion.h
--- a/Engine/Source/Runtime/Core/Public/UObject/UE5MainStreamObjectVersion.h
+++ b/Engine/Source/Runtime/Core/Public/UObject/UE5MainStreamObjectVersion.h
@@ -224,6 +224,9 @@ struct CORE_API FUE5MainStreamObjectVersion
 		// Allow each LiveLink controller to specify its own component to control
 		LiveLinkComponentPickerPerController,
 
+		// Remove Faces in Triangle Mesh BVH
+		RemoveTriangleMeshBVHFaces,
+
 		// -----<new versions can be added above this line>-------------------------------------------------
 		VersionPlusOne,
 		LatestVersion = VersionPlusOne - 1
diff --git a/Engine/Source/Runtime/Experimental/Chaos/Private/Chaos/TriangleMeshImplicitObject.cpp b/Engine/Source/Runtime/Experimental/Chaos/Private/Chaos/TriangleMeshImplicitObject.cpp
--- a/Engine/Source/Runtime/Experimental/Chaos/Private/Chaos/TriangleMeshImplicitObject.cpp
+++ b/Engine/Source/Runtime/Experimental/Chaos/Private/Chaos/TriangleMeshImplicitObject.cpp
@@ -1538,7 +1538,30 @@ void FTriangleMeshImplicitObject::UpdateVertices(const TArray<FVector>& NewPosit
 
 	RebuildFastBVH();
 }
-	
+
+namespace
+{
+	void AddTriangles(TArray<TVec3<FTrimeshIndexBuffer::LargeIdxType>>& LargeIndices, TArray<TVec3<FTrimeshIndexBuffer::SmallIdxType>>& SmallIndices, const FTrimeshIndexBuffer& Elements, int32 OldIndex)
+	{
+		if (Elements.RequiresLargeIndices())
+		{
+			FTrimeshIndexBuffer::LargeIdxType A = Elements.GetLargeIndexBuffer()[OldIndex][0];
+			FTrimeshIndexBuffer::LargeIdxType B = Elements.GetLargeIndexBuffer()[OldIndex][1];
+			FTrimeshIndexBuffer::LargeIdxType C = Elements.GetLargeIndexBuffer()[OldIndex][2];
+
+			LargeIndices.Add(TVec3<FTrimeshIndexBuffer::LargeIdxType>(A, B, C));
+		}
+		else
+		{
+			FTrimeshIndexBuffer::SmallIdxType A = Elements.GetSmallIndexBuffer()[OldIndex][0];
+			FTrimeshIndexBuffer::SmallIdxType B = Elements.GetSmallIndexBuffer()[OldIndex][1];
+			FTrimeshIndexBuffer::SmallIdxType C = Elements.GetSmallIndexBuffer()[OldIndex][2];
+
+			SmallIndices.Add(TVec3<FTrimeshIndexBuffer::SmallIdxType>(A, B, C));
+		}
+	}
+}
+
 void FTriangleMeshImplicitObject::RebuildFastBVHFromTree(const BVHType& TreeBVH)
 {
 	using NodeType = TAABBTreeNode<FRealSingle>;
@@ -1547,9 +1570,13 @@ void FTriangleMeshImplicitObject::RebuildFastBVHFromTree(const BVHType& TreeBVH)
 	const TArray<LeafType>& Leaves = TreeBVH.GetLeaves();
 
 	FastBVH.Nodes.Reset();
-	FastBVH.Faces.Reset();
 	FastBVH.FaceBounds.Reset();
 
+	int32 FaceNum = 0;
+	TArray<TVec3<FTrimeshIndexBuffer::LargeIdxType>> LargeIndices;
+	TArray<TVec3<FTrimeshIndexBuffer::SmallIdxType>> SmallIndices;
+		
+
 	// since we do skip leaf nodes, we need to handle the case where we have only one node that will be a leaf by default
 	if (Nodes.Num() == 1)
 	{
@@ -1568,7 +1595,17 @@ void FTriangleMeshImplicitObject::RebuildFastBVHFromTree(const BVHType& TreeBVH)
 		for (const TPayloadBoundsElement<int32, FRealSingle>& LeafPayload: Leaf.Elems)
 		{
 			FastBVH.FaceBounds.Add(FAABBVectorized(LeafPayload.Bounds));
-			FastBVH.Faces.Add(LeafPayload.Payload);
+			// Reorder triangle indices, triangles will be in the same order as the bounding volume in the BVH structure.
+			// And all triangles in a leaf will be contiguous in memory.
+			AddTriangles(LargeIndices, SmallIndices, MElements, LeafPayload.Payload);
+		}
+		if (MElements.RequiresLargeIndices())
+		{
+			MElements.Reinitialize(MoveTemp(LargeIndices));
+		}
+		else
+		{
+			MElements.Reinitialize(MoveTemp(SmallIndices));
 		}
 		return;
 	}
@@ -1626,7 +1663,7 @@ void FTriangleMeshImplicitObject::RebuildFastBVHFromTree(const BVHType& TreeBVH)
 						const LeafType& Leaf = Leaves[*LeafIndex];
 
 						// store face range in the node 
-						ChildData.SetChildOrFaceIndex(FastBVH.Faces.Num());
+						ChildData.SetChildOrFaceIndex(FaceNum);
 						ChildData.SetFaceCount(Leaf.Elems.Num());
 						check(ChildData.GetFaceCount() > 0);
 
@@ -1635,7 +1672,10 @@ void FTriangleMeshImplicitObject::RebuildFastBVHFromTree(const BVHType& TreeBVH)
 						for (const auto& LeafPayload: Leaf.Elems)
 						{
 							FastBVH.FaceBounds.Add(FAABBVectorized(LeafPayload.Bounds));
-							FastBVH.Faces.Add(LeafPayload.Payload);
+							// Reorder triangle indices, triangles will be in the same order as the bounding volume in the BVH structure.
+							// And all triangles in a leaf will be contiguous in memory.
+							AddTriangles(LargeIndices, SmallIndices, MElements, LeafPayload.Payload);
+							FaceNum++;
 						}
 					}
 					// push for further processing
@@ -1664,7 +1704,14 @@ void FTriangleMeshImplicitObject::RebuildFastBVHFromTree(const BVHType& TreeBVH)
 			}
 		}
 	}
-	
+	if (MElements.RequiresLargeIndices())
+	{
+		MElements.Reinitialize(MoveTemp(LargeIndices));
+	}
+	else
+	{
+		MElements.Reinitialize(MoveTemp(SmallIndices));
+	}
 }
 	
 }
diff --git a/Engine/Source/Runtime/Experimental/Chaos/Public/Chaos/TriangleMeshImplicitObject.h b/Engine/Source/Runtime/Experimental/Chaos/Public/Chaos/TriangleMeshImplicitObject.h
--- a/Engine/Source/Runtime/Experimental/Chaos/Public/Chaos/TriangleMeshImplicitObject.h
+++ b/Engine/Source/Runtime/Experimental/Chaos/Public/Chaos/TriangleMeshImplicitObject.h
@@ -197,7 +197,7 @@ namespace Chaos
 			{
 				if (BoundsFilter(FaceBounds[FaceIndex]) == EFilterResult::Keep)
 				{
-					if (FaceVisitor(Faces[FaceIndex]) == EVisitorResult::Stop)
+					if (FaceVisitor(FaceIndex) == EVisitorResult::Stop)
 					{
 						return EVisitorResult::Stop;
 					}
@@ -250,14 +250,18 @@ namespace Chaos
 
 		void Serialize(FChaosArchive& Ar)
 		{
+			Ar.UsingCustomVersion(FUE5MainStreamObjectVersion::GUID);
 			Ar << Nodes;
 			Ar << FaceBounds;
-			Ar << Faces;
+			if (Ar.CustomVer(FUE5MainStreamObjectVersion::GUID) < FUE5MainStreamObjectVersion::RemoveTriangleMeshBVHFaces)
+			{
+				TArray<int32> TmpFaces;
+				Ar << TmpFaces;
+			}
 		}
 		
 		TArray<FNode> Nodes;
 		TArray<FAABBVectorized> FaceBounds;
-		TArray<int32> Faces;
 	};
 
 	FORCEINLINE_DEBUGGABLE FChaosArchive& operator<<(FChaosArchive& Ar, FTrimeshBVH::FChildData& ChildData)
@@ -332,6 +336,23 @@ namespace Chaos
 		FTrimeshIndexBuffer(const FTrimeshIndexBuffer& Other) = delete;
 		FTrimeshIndexBuffer& operator=(const FTrimeshIndexBuffer& Other) = delete;
 
+
+		void Reinitialize(TArray<TVec3<LargeIdxType>>&& Elements)
+		{
+			LargeIdxBuffer.Empty();
+			SmallIdxBuffer.Empty();
+			bRequiresLargeIndices = true;
+			LargeIdxBuffer = MoveTemp(Elements);
+		}
+
+		void Reinitialize(TArray<TVec3<SmallIdxType>>&& Elements)
+		{
+			LargeIdxBuffer.Empty();
+			SmallIdxBuffer.Empty();
+			bRequiresLargeIndices = false;
+			SmallIdxBuffer = MoveTemp(Elements);
+		}
+
 		void Serialize(FArchive& Ar)
 		{
 			Ar.UsingCustomVersion(FExternalPhysicsCustomObjectVersion::GUID);
@@ -597,6 +618,12 @@ namespace Chaos
 				Ar << FastBVH;
 			}
 
+			if (Ar.CustomVer(FUE5MainStreamObjectVersion::GUID) < FUE5MainStreamObjectVersion::RemoveTriangleMeshBVHFaces)
+			{
+				// Force to rebuild the BVH
+				RebuildFastBVH();
+			}
+
 			if (Ar.CustomVer(FExternalPhysicsCustomObjectVersion::GUID) >= FExternalPhysicsCustomObjectVersion::AddTrimeshMaterialIndices)
 			{
 				Ar << MaterialIndices;
diff --git a/Engine/Source/Runtime/Experimental/ChaosCore/Public/Chaos/Core.h b/Engine/Source/Runtime/Experimental/ChaosCore/Public/Chaos/Core.h
--- a/Engine/Source/Runtime/Experimental/ChaosCore/Public/Chaos/Core.h
+++ b/Engine/Source/Runtime/Experimental/ChaosCore/Public/Chaos/Core.h
@@ -51,6 +51,6 @@ namespace Chaos
 	// NOTE: if you get a merge conflict on the GUID, you must replace it with a new GUID - do not accept the source or target
 	// or you will likely get DDC version conflicts resulting in crashes during load.
 	// Core version string for Chaos data. Any DDC builder dependent on Chaos for serialization should depend on this version
-	static const FString ChaosVersionString = TEXT("B3AEC941-1F1D-414A-A961-C0376566C98D");
+	static const FString ChaosVersionString = TEXT("841FF58C-9DD5-4AF7-BE5B-959F1F2BC4D7");
 
 }
