Chaos: remove warning left from local changes

diff --git a/Engine/Source/Runtime/Experimental/Chaos/Private/PhysicsProxy/GeometryCollectionPhysicsProxy.cpp b/Engine/Source/Runtime/Experimental/Chaos/Private/PhysicsProxy/GeometryCollectionPhysicsProxy.cpp
--- a/Engine/Source/Runtime/Experimental/Chaos/Private/PhysicsProxy/GeometryCollectionPhysicsProxy.cpp
+++ b/Engine/Source/Runtime/Experimental/Chaos/Private/PhysicsProxy/GeometryCollectionPhysicsProxy.cpp
@@ -990,7 +990,6 @@ void FGeometryCollectionPhysicsProxy::InitializeBodiesPT(Chaos::FPBDRigidsSolver
 			const bool bGenerateConnectionGraph = (!Connections) || (Connections->Num() != RestCollection->Transform.Num()) || bGeometryCollectionAlwaysGenerateConnectionGraph;
 			if (bGenerateConnectionGraph)
 			{
-				UE_LOG(LogChaos, Warning, TEXT("Generating GC connection graph on the physics thread"));
 				// Set cluster connectivity.  TPBDRigidClustering::CreateClusterParticle() 
 				// will optionally do this, but we switch that functionality off in BuildClusters_Internal().
 				for (int32 TransformGroupIndex = 0; TransformGroupIndex < NumTransforms; ++TransformGroupIndex)
