Chaos : harden manageArray Api for better detection of modification to the array
this is a step towards more automated recompoutation of some attributes dependent on others

diff --git a/Engine/Plugins/Experimental/ChaosEditor/Source/FractureEditor/Private/FractureEditorModeToolkit.cpp b/Engine/Plugins/Experimental/ChaosEditor/Source/FractureEditor/Private/FractureEditorModeToolkit.cpp
--- a/Engine/Plugins/Experimental/ChaosEditor/Source/FractureEditor/Private/FractureEditorModeToolkit.cpp
+++ b/Engine/Plugins/Experimental/ChaosEditor/Source/FractureEditor/Private/FractureEditorModeToolkit.cpp
@@ -929,7 +929,7 @@ void FFractureEditorModeToolkit::SetHideForUnselected(UGeometryCollectionCompone
 		// If a cluster is selected, set false for all children.
 		if (GeometryCollection->HasAttribute("Hide", FGeometryCollection::TransformGroup))
 		{
-			TManagedArray<bool>& Hide = GeometryCollection->GetAttribute<bool>("Hide", FGeometryCollection::TransformGroup);
+			TManagedArray<bool>& Hide = GeometryCollection->ModifyAttribute<bool>("Hide", FGeometryCollection::TransformGroup);
 			const TManagedArray<TSet<int32>>& Children = GeometryCollection->GetAttribute<TSet<int32>>("Children", FGeometryCollection::TransformGroup);
 			
 			const TArray<int32>& SelectedBones = GCComp->GetSelectedBones();
@@ -1080,7 +1080,7 @@ int32 FFractureEditorModeToolkit::GetLevelCount()
 				bool HasLevelAttribute = GeometryCollection->HasAttribute("Level", FTransformCollection::TransformGroup);
 				if (HasLevelAttribute)
 				{
-					TManagedArray<int32>& Levels = GeometryCollection->GetAttribute<int32>("Level", FTransformCollection::TransformGroup);
+					const TManagedArray<int32>& Levels = GeometryCollection->GetAttribute<int32>("Level", FTransformCollection::TransformGroup);
 
 					if(Levels.Num() > 0)
 					{
@@ -1675,7 +1675,7 @@ void FFractureEditorModeToolkit::UpdateExplodedVectors(UGeometryCollectionCompon
 
 		check(OutGeometryCollection->HasAttribute("ExplodedVector", FGeometryCollection::TransformGroup));
 
-		TManagedArray<FVector3f>& ExplodedVectors = OutGeometryCollection->GetAttribute<FVector3f>("ExplodedVector", FGeometryCollection::TransformGroup);
+		TManagedArray<FVector3f>& ExplodedVectors = OutGeometryCollection->ModifyAttribute<FVector3f>("ExplodedVector", FGeometryCollection::TransformGroup);
 		const TManagedArray<FTransform>& Transform = OutGeometryCollection->GetAttribute<FTransform>("Transform", FGeometryCollection::TransformGroup);
 		const TManagedArray<int32>& TransformToGeometryIndex = OutGeometryCollection->GetAttribute<int32>("TransformToGeometryIndex", FGeometryCollection::TransformGroup);
 		const TManagedArray<FBox>& BoundingBox = OutGeometryCollection->GetAttribute<FBox>("BoundingBox", FGeometryCollection::GeometryGroup);
diff --git a/Engine/Plugins/Experimental/ChaosEditor/Source/FractureEditor/Private/FractureToolConvex.cpp b/Engine/Plugins/Experimental/ChaosEditor/Source/FractureEditor/Private/FractureToolConvex.cpp
--- a/Engine/Plugins/Experimental/ChaosEditor/Source/FractureEditor/Private/FractureToolConvex.cpp
+++ b/Engine/Plugins/Experimental/ChaosEditor/Source/FractureEditor/Private/FractureToolConvex.cpp
@@ -221,8 +221,8 @@ void UFractureToolConvex::FractureContextChanged()
 			FTransform CombinedTransform = InnerTransform * OuterTransform;
 			bool bIsCustom = HasCustomConvex ? bool((*HasCustomConvex)[TransformIdx]) : false;
 
-			TManagedArray<TSet<int32>>& TransformToConvexIndices = Collection.GetAttribute<TSet<int32>>("TransformToConvexIndices", FTransformCollection::TransformGroup);
-			TManagedArray<TUniquePtr<Chaos::FConvex>>& ConvexHull = Collection.GetAttribute<TUniquePtr<Chaos::FConvex>>("ConvexHull", "Convex");
+			const TManagedArray<TSet<int32>>& TransformToConvexIndices = Collection.GetAttribute<TSet<int32>>("TransformToConvexIndices", FTransformCollection::TransformGroup);
+			const TManagedArray<TUniquePtr<Chaos::FConvex>>& ConvexHull = Collection.GetAttribute<TUniquePtr<Chaos::FConvex>>("ConvexHull", "Convex");
 
 			EdgesMappings.AddMapping(CollectionIdx, TransformIdx, HullEdges.Num());
 
diff --git a/Engine/Plugins/Experimental/ChaosEditor/Source/FractureEditor/Private/FractureToolEmbed.cpp b/Engine/Plugins/Experimental/ChaosEditor/Source/FractureEditor/Private/FractureToolEmbed.cpp
--- a/Engine/Plugins/Experimental/ChaosEditor/Source/FractureEditor/Private/FractureToolEmbed.cpp
+++ b/Engine/Plugins/Experimental/ChaosEditor/Source/FractureEditor/Private/FractureToolEmbed.cpp
@@ -232,7 +232,7 @@ void UFractureToolAutoEmbedGeometry::Execute(TWeakPtr<FFractureEditorModeToolkit
 				// Which bone points to the closest convex?
 				const TManagedArray<FTransform>& Transform = ClosestGeometryCollection->Transform;
 				const TManagedArray<int32>& Parent = ClosestGeometryCollection->Parent;
-				TManagedArray<TSet<int32>>& TransformToConvexIndices =
+				const TManagedArray<TSet<int32>>& TransformToConvexIndices =
 					ClosestGeometryCollection->GetAttribute<TSet<int32>>("TransformToConvexIndices", FTransformCollection::TransformGroup);
 
 				int32 BoneIndex = INDEX_NONE;
@@ -270,7 +270,7 @@ void UFractureToolAutoEmbedGeometry::Execute(TWeakPtr<FFractureEditorModeToolkit
 					}
 
 					// Get a guid generated for the new instance
-					TManagedArray<FGuid>& Guids = ClosestGeometryCollection->GetAttribute<FGuid>("GUID", FGeometryCollection::TransformGroup);
+					TManagedArray<FGuid>& Guids = ClosestGeometryCollection->ModifyAttribute<FGuid>("GUID", FGeometryCollection::TransformGroup);
 					Guids[Guids.Num()-1] = FGuid::NewGuid();
 
 					// #todo there might be a lot of these -- collect and put outside the loop.
diff --git a/Engine/Plugins/Experimental/ChaosEditor/Source/FractureEditor/Private/FractureToolProperties.cpp b/Engine/Plugins/Experimental/ChaosEditor/Source/FractureEditor/Private/FractureToolProperties.cpp
--- a/Engine/Plugins/Experimental/ChaosEditor/Source/FractureEditor/Private/FractureToolProperties.cpp
+++ b/Engine/Plugins/Experimental/ChaosEditor/Source/FractureEditor/Private/FractureToolProperties.cpp
@@ -65,7 +65,7 @@ void UFractureToolSetInitialDynamicState::Execute(TWeakPtr<FFractureEditorModeTo
 				TSharedPtr<FGeometryCollection, ESPMode::ThreadSafe> GeometryCollectionPtr = GCObject->GetGeometryCollection();
 				if (FGeometryCollection* GeometryCollection = GeometryCollectionPtr.Get())
 				{
-					TManagedArray<int32>& InitialDynamicState = GeometryCollection->GetAttribute<int32>("InitialDynamicState", FGeometryCollection::TransformGroup);
+					TManagedArray<int32>& InitialDynamicState = GeometryCollection->ModifyAttribute<int32>("InitialDynamicState", FGeometryCollection::TransformGroup);
 
 					TArray<int32> SelectedBones = GeometryCollectionComponent->GetSelectedBones();
 					for (int32 Index : SelectedBones)
diff --git a/Engine/Plugins/Experimental/ChaosEditor/Source/FractureEditor/Private/SGeometryCollectionOutliner.cpp b/Engine/Plugins/Experimental/ChaosEditor/Source/FractureEditor/Private/SGeometryCollectionOutliner.cpp
--- a/Engine/Plugins/Experimental/ChaosEditor/Source/FractureEditor/Private/SGeometryCollectionOutliner.cpp
+++ b/Engine/Plugins/Experimental/ChaosEditor/Source/FractureEditor/Private/SGeometryCollectionOutliner.cpp
@@ -457,7 +457,7 @@ void SGeometryCollectionOutliner::SetInitialDynamicState(int32 InDynamicState)
 				TSharedPtr<FGeometryCollection,ESPMode::ThreadSafe> GeometryCollection = GeometryCollectionObject->GetGeometryCollection();
 				if (GeometryCollection)
 				{
-					TManagedArray<int32>& InitialDynamicState = GeometryCollection->GetAttribute<int32>("InitialDynamicState", FGeometryCollection::TransformGroup);
+					TManagedArray<int32>& InitialDynamicState = GeometryCollection->ModifyAttribute<int32>("InitialDynamicState", FGeometryCollection::TransformGroup);
 					TArray<int32> SelectedBones = Component->GetSelectedBones();
 					int32 BoneIndex = SelectedItem->GetBoneIndex();
 					if (ensure(0 <= BoneIndex && BoneIndex < InitialDynamicState.Num()))
@@ -714,7 +714,7 @@ void FGeometryCollectionTreeItemBone::UpdateItemFromCollection()
 			const bool bHasLevelAttribute = GeometryCollectionPtr->HasAttribute("Level", FGeometryCollection::TransformGroup);
 			if (bHasLevelAttribute && OutlinerSettings->ColorByLevel)
 			{
-				TManagedArray<int32>& Level = GeometryCollectionPtr->GetAttribute<int32>("Level", FTransformCollection::TransformGroup);
+				const TManagedArray<int32>& Level = GeometryCollectionPtr->GetAttribute<int32>("Level", FTransformCollection::TransformGroup);
 				ItemColor = FSlateColor(FGeometryCollectionTreeItem::GetColorPerDepth((uint32)Level[ItemBoneIndex]));
 			}
 			else
diff --git a/Engine/Plugins/Experimental/GeometryCollectionPlugin/Source/GeometryCollectionEditor/Private/GeometryCollection/GeometryCollectionCommands.cpp b/Engine/Plugins/Experimental/GeometryCollectionPlugin/Source/GeometryCollectionEditor/Private/GeometryCollection/GeometryCollectionCommands.cpp
--- a/Engine/Plugins/Experimental/GeometryCollectionPlugin/Source/GeometryCollectionEditor/Private/GeometryCollection/GeometryCollectionCommands.cpp
+++ b/Engine/Plugins/Experimental/GeometryCollectionPlugin/Source/GeometryCollectionEditor/Private/GeometryCollection/GeometryCollectionCommands.cpp
@@ -254,7 +254,7 @@ void FGeometryCollectionCommands::SetNamedAttributeValues(const TArray<FString>&
 
 
 									const FRegexPattern RegexPattern(InPattern);
-									TManagedArray<bool>& AttributeArray = RestCollection->GetAttribute<bool>(InAttributeName, InGroupName);
+									TManagedArray<bool>& AttributeArray = RestCollection->ModifyAttribute<bool>(InAttributeName, InGroupName);
 
 									if (InPatternName.GetStringLength())
 									{
@@ -262,7 +262,7 @@ void FGeometryCollectionCommands::SetNamedAttributeValues(const TArray<FString>&
 										{
 											if (RestCollection->FindAttribute<FString>(InPatternName, InGroupName))
 											{
-												TManagedArray<FString>& NameArray = RestCollection->GetAttribute<FString>(InPatternName, InGroupName);
+												const TManagedArray<FString>& NameArray = RestCollection->GetAttribute<FString>(InPatternName, InGroupName);
 
 												for (int i = 0; i < NameArray.Num(); i++)
 												{
@@ -292,7 +292,7 @@ void FGeometryCollectionCommands::SetNamedAttributeValues(const TArray<FString>&
 
 
 									const FRegexPattern RegexPattern(InPattern);
-									TManagedArray<int32>& AttributeArray = RestCollection->GetAttribute<int32>(InAttributeName, InGroupName);
+									TManagedArray<int32>& AttributeArray = RestCollection->ModifyAttribute<int32>(InAttributeName, InGroupName);
 
 									if (InPatternName.GetStringLength())
 									{
@@ -300,7 +300,7 @@ void FGeometryCollectionCommands::SetNamedAttributeValues(const TArray<FString>&
 										{
 											if (RestCollection->FindAttribute<FString>(InPatternName, InGroupName))
 											{
-												TManagedArray<FString>& NameArray = RestCollection->GetAttribute<FString>(InPatternName, InGroupName);
+												const TManagedArray<FString>& NameArray = RestCollection->GetAttribute<FString>(InPatternName, InGroupName);
 
 												for (int i = 0; i < NameArray.Num(); i++)
 												{
@@ -330,7 +330,7 @@ void FGeometryCollectionCommands::SetNamedAttributeValues(const TArray<FString>&
 
 
 									const FRegexPattern RegexPattern(InPattern);
-									TManagedArray<float>& AttributeArray = RestCollection->GetAttribute<float>(InAttributeName, InGroupName);
+									TManagedArray<float>& AttributeArray = RestCollection->ModifyAttribute<float>(InAttributeName, InGroupName);
 
 									if (InPatternName.GetStringLength())
 									{
@@ -338,7 +338,7 @@ void FGeometryCollectionCommands::SetNamedAttributeValues(const TArray<FString>&
 										{
 											if (RestCollection->FindAttribute<FString>(InPatternName, InGroupName))
 											{
-												TManagedArray<FString>& NameArray = RestCollection->GetAttribute<FString>(InPatternName, InGroupName);
+												const TManagedArray<FString>& NameArray = RestCollection->GetAttribute<FString>(InPatternName, InGroupName);
 
 												for (int i = 0; i < NameArray.Num(); i++)
 												{
diff --git a/Engine/Plugins/Experimental/PlanarCutPlugin/Source/PlanarCut/Private/GeometryMeshConversion.cpp b/Engine/Plugins/Experimental/PlanarCutPlugin/Source/PlanarCut/Private/GeometryMeshConversion.cpp
--- a/Engine/Plugins/Experimental/PlanarCutPlugin/Source/PlanarCut/Private/GeometryMeshConversion.cpp
+++ b/Engine/Plugins/Experimental/PlanarCutPlugin/Source/PlanarCut/Private/GeometryMeshConversion.cpp
@@ -2586,7 +2586,7 @@ int32 FDynamicMeshCollection::CutWithMultiplePlanes(
 	// create proximity sets on geometry collection and populate using ToCut's Proximity multimap and the array ToCutIdxToGeometryIdx
 	if (bHasProximity)
 	{
-		TManagedArray<TSet<int32>>& GCProximity = Collection->GetAttribute<TSet<int32>>("Proximity", FGeometryCollection::GeometryGroup);
+		TManagedArray<TSet<int32>>& GCProximity = Collection->ModifyAttribute<TSet<int32>>("Proximity", FGeometryCollection::GeometryGroup);
 		for (TPair<int32, int32> Link : Proximity)
 		{
 			GCProximity[ToCutIdxToGeometryIdx[Link.Key]].Add(ToCutIdxToGeometryIdx[Link.Value]);
@@ -2740,7 +2740,7 @@ int32 FDynamicMeshCollection::CutWithCellMeshes(const FInternalSurfaceMaterials&
 			}
 			if (bHasProximity)
 			{
-				TManagedArray<TSet<int32>>& Proximity = Collection->GetAttribute<TSet<int32>>("Proximity", FGeometryCollection::GeometryGroup);
+				TManagedArray<TSet<int32>>& Proximity = Collection->ModifyAttribute<TSet<int32>>("Proximity", FGeometryCollection::GeometryGroup);
 				TArray<TUniquePtr<TPointHashGrid3d<int>>> VertexHashes;
 				TArray<FAxisAlignedBox3d> MeshBounds;
 				auto MakeHash = [this, &VertexHashes, &MeshBounds, &BooleanResults](int GID)
diff --git a/Engine/Plugins/Experimental/PlanarCutPlugin/Source/PlanarCut/Private/PlanarCut.cpp b/Engine/Plugins/Experimental/PlanarCutPlugin/Source/PlanarCut/Private/PlanarCut.cpp
--- a/Engine/Plugins/Experimental/PlanarCutPlugin/Source/PlanarCut/Private/PlanarCut.cpp
+++ b/Engine/Plugins/Experimental/PlanarCutPlugin/Source/PlanarCut/Private/PlanarCut.cpp
@@ -913,7 +913,7 @@ int32 MergeBones(
 
 	FGeometryCollectionProximityUtility ProximityUtility(&Collection);
 	ProximityUtility.UpdateProximity();
-	TManagedArray<TSet<int32>>& Proximity = Collection.GetAttribute<TSet<int32>>("Proximity", FGeometryCollection::GeometryGroup);
+	const TManagedArray<TSet<int32>>& Proximity = Collection.GetAttribute<TSet<int32>>("Proximity", FGeometryCollection::GeometryGroup);
 
 	// local array so we can populate it with all transforms if input was empty
 	TArray<int32> TransformIndices(TransformIndicesView);
diff --git a/Engine/Source/Runtime/Experimental/Chaos/Private/GeometryCollection/GeometryCollection.cpp b/Engine/Source/Runtime/Experimental/Chaos/Private/GeometryCollection/GeometryCollection.cpp
--- a/Engine/Source/Runtime/Experimental/Chaos/Private/GeometryCollection/GeometryCollection.cpp
+++ b/Engine/Source/Runtime/Experimental/Chaos/Private/GeometryCollection/GeometryCollection.cpp
@@ -943,10 +943,10 @@ void FGeometryCollection::Serialize(Chaos::FChaosArchive& Ar)
 			FS_Clustered = 0x00000002,
 		};
 
-		TManagedArray<FGeometryCollectionBoneNode>* BoneHierarchyPtr = FindAttribute<FGeometryCollectionBoneNode>("BoneHierarchy", FTransformCollection::TransformGroup);
+		const TManagedArray<FGeometryCollectionBoneNode>* BoneHierarchyPtr = FindAttribute<FGeometryCollectionBoneNode>("BoneHierarchy", FTransformCollection::TransformGroup);
 		if (BoneHierarchyPtr)
 		{
-			TManagedArray<FGeometryCollectionBoneNode>& BoneHierarchy = *BoneHierarchyPtr;
+			const TManagedArray<FGeometryCollectionBoneNode>& BoneHierarchy = *BoneHierarchyPtr;
 
 			for (int Idx = 0; Idx < BoneHierarchy.Num(); Idx++)
 			{
@@ -954,7 +954,7 @@ void FGeometryCollection::Serialize(Chaos::FChaosArchive& Ar)
 				{
 					AddAttribute<int32>("Level", FGeometryCollection::TransformGroup);
 				}
-				TManagedArray<int32>& Level = GetAttribute<int32>("Level", FGeometryCollection::TransformGroup);
+				TManagedArray<int32>& Level = ModifyAttribute<int32>("Level", FGeometryCollection::TransformGroup);
 				Level[Idx] = BoneHierarchy[Idx].Level;
 
 				SimulationType[Idx] = ESimulationTypes::FST_Rigid;
@@ -978,7 +978,7 @@ void FGeometryCollection::Serialize(Chaos::FChaosArchive& Ar)
 		if (Version < 5)
 		{
 			UE_LOG(FGeometryCollectionLogging, Log, TEXT("GeometryCollection has inaccurate simulation type tags. Updating tags based on transform topology."));
-			TManagedArray<bool>* SimulatableParticles = FindAttribute<bool>(FGeometryCollection::SimulatableParticlesAttribute, FTransformCollection::TransformGroup);
+			const TManagedArray<bool>* SimulatableParticles = FindAttribute<bool>(FGeometryCollection::SimulatableParticlesAttribute, FTransformCollection::TransformGroup);
 			TArray<bool> RigidChildren; RigidChildren.Init(false,NumElements(FTransformCollection::TransformGroup));
 			const TArray<int32> RecursiveOrder = GeometryCollectionAlgo::ComputeRecursiveOrder(*this);
 			for (const int32 TransformGroupIndex : RecursiveOrder)
@@ -1033,7 +1033,7 @@ void FGeometryCollection::Serialize(Chaos::FChaosArchive& Ar)
 		{
 			if (HasAttribute("TransformToConvexIndex", FTransformCollection::TransformGroup))
 			{
-				TManagedArray<int32> TransformToConvexIndex = MoveTemp(GetAttribute<int32>("TransformToConvexIndex", FTransformCollection::TransformGroup));
+				TManagedArray<int32> TransformToConvexIndex = MoveTemp(ModifyAttribute<int32>("TransformToConvexIndex", FTransformCollection::TransformGroup));
 				RemoveAttribute("TransformToConvexIndex", FTransformCollection::TransformGroup);
 				// if we don't already have the one-to-many version, convert the previous one-to-one mapping to the new format
 				if (!HasAttribute("TransformToConvexIndices", FTransformCollection::TransformGroup))
@@ -1062,22 +1062,22 @@ void FGeometryCollection::Serialize(Chaos::FChaosArchive& Ar)
 				AddAttribute<TArray<FVector2f>>("UVs", FGeometryCollection::VerticesGroup);				
 			}
 
-			TManagedArray<TArray<FVector2f>>* MultipleUVs = FindAttribute<TArray<FVector2f>>("UVs", FGeometryCollection::VerticesGroup);
+			TManagedArray<TArray<FVector2f>>& MultipleUVs = ModifyAttribute<TArray<FVector2f>>("UVs", FGeometryCollection::VerticesGroup);
 			if (NumUVLayers() < 1)
 			{
-				for (int32 VertIdx = 0; VertIdx < MultipleUVs->Num(); ++VertIdx)
+				for (int32 VertIdx = 0; VertIdx < MultipleUVs.Num(); ++VertIdx)
 				{
-					(*MultipleUVs)[VertIdx].SetNum(1);
+					MultipleUVs[VertIdx].SetNum(1);
 				}
 			}
 
-			if (TManagedArray<FVector2f>* SingleUV = FindAttribute<FVector2f>("UV", FGeometryCollection::VerticesGroup))
+			if (const TManagedArray<FVector2f>* SingleUV = FindAttribute<FVector2f>("UV", FGeometryCollection::VerticesGroup))
 			{
-				for (int32 VertIdx = 0; VertIdx < MultipleUVs->Num(); ++VertIdx)
+				for (int32 VertIdx = 0; VertIdx < MultipleUVs.Num(); ++VertIdx)
 				{
 					if (SingleUV)
 					{
-						(*MultipleUVs)[VertIdx][0] = (*SingleUV)[VertIdx];
+						MultipleUVs[VertIdx][0] = (*SingleUV)[VertIdx];
 					}
 				}
 
@@ -1638,7 +1638,7 @@ void FGeometryCollection::WriteDataToOBJFile(const FString &Name, const FString
 		DataFile.open(string(TCHAR_TO_UTF8(*FullPath)));
 		DataFile << "# Vertex Visibility - vertices whose visibility flag are true" << endl;
 
-		TManagedArray<bool>& VertexVisibility = GetAttribute<bool>("VertexVisibility", FGeometryCollection::VerticesGroup);
+		const TManagedArray<bool>& VertexVisibility = ModifyAttribute<bool>("VertexVisibility", FGeometryCollection::VerticesGroup);
 		int num = 0;
 		for (int32 IdxVertex = 0; IdxVertex < NumVertices; ++IdxVertex)
 		{
@@ -1905,4 +1905,4 @@ void FGeometryCollection::UpdateOldAttributeNames()
 		this->FaceCount[i] = OldFaceCount[i];
 	}
 	this->RemoveGroup("Structure");
-}
\ No newline at end of file
+}
diff --git a/Engine/Source/Runtime/Experimental/Chaos/Private/GeometryCollection/GeometryCollectionAlgo.cpp b/Engine/Source/Runtime/Experimental/Chaos/Private/GeometryCollection/GeometryCollectionAlgo.cpp
--- a/Engine/Source/Runtime/Experimental/Chaos/Private/GeometryCollection/GeometryCollectionAlgo.cpp
+++ b/Engine/Source/Runtime/Experimental/Chaos/Private/GeometryCollection/GeometryCollectionAlgo.cpp
@@ -282,8 +282,8 @@ namespace GeometryCollectionAlgo
 
 					if (0 < ChildIndex && ChildIndex < NumTransforms)
 					{
-						TManagedArray<int32>& Parent = Collection->GetAttribute<int32>(FTransformCollection::ParentAttribute, FTransformCollection::TransformGroup);
-						TManagedArray< TArray<int32> >& Children = Collection->GetAttribute< TArray<int32> >(FTransformCollection::ChildrenAttribute, FTransformCollection::TransformGroup);
+						TManagedArray<int32>& Parent = Collection->ModifyAttribute<int32>(FTransformCollection::ParentAttribute, FTransformCollection::TransformGroup);
+						TManagedArray< TArray<int32> >& Children = Collection->ModifyAttribute< TArray<int32> >(FTransformCollection::ChildrenAttribute, FTransformCollection::TransformGroup);
 
 						int32 ParentIndex = Parent[ChildIndex];
 						if (0 <= ParentIndex && ParentIndex < NumTransforms)
diff --git a/Engine/Source/Runtime/Experimental/Chaos/Private/GeometryCollection/GeometryCollectionClusteringUtility.cpp b/Engine/Source/Runtime/Experimental/Chaos/Private/GeometryCollection/GeometryCollectionClusteringUtility.cpp
--- a/Engine/Source/Runtime/Experimental/Chaos/Private/GeometryCollection/GeometryCollectionClusteringUtility.cpp
+++ b/Engine/Source/Runtime/Experimental/Chaos/Private/GeometryCollection/GeometryCollectionClusteringUtility.cpp
@@ -74,7 +74,7 @@ void FGeometryCollectionClusteringUtility::ClusterAllBonesUnderNewRoot(FGeometry
 
 	if (GeometryCollection->HasAttribute("Level", FGeometryCollection::TransformGroup))
 	{
-		TManagedArray<int32>& Levels = GeometryCollection->GetAttribute<int32>("Level", FGeometryCollection::TransformGroup);
+		TManagedArray<int32>& Levels = GeometryCollection->ModifyAttribute<int32>("Level", FGeometryCollection::TransformGroup);
 
 		// all bones shifted down one in hierarchy
 		for (int ChildIndex = 0; ChildIndex < NumElements; ChildIndex++)
@@ -95,8 +95,8 @@ void FGeometryCollectionClusteringUtility::ClusterAllBonesUnderNewRoot(FGeometry
 		GeometryCollection->HasAttribute("ExplodedTransform", FGeometryCollection::TransformGroup) )
 	{
 
-		TManagedArray<FVector3f>& ExplodedVectors = GeometryCollection->GetAttribute<FVector3f>("ExplodedVector", FGeometryCollection::TransformGroup);
-		TManagedArray<FTransform>& ExplodedTransforms = GeometryCollection->GetAttribute<FTransform>("ExplodedTransform", FGeometryCollection::TransformGroup);
+		TManagedArray<FVector3f>& ExplodedVectors = GeometryCollection->ModifyAttribute<FVector3f>("ExplodedVector", FGeometryCollection::TransformGroup);
+		TManagedArray<FTransform>& ExplodedTransforms = GeometryCollection->ModifyAttribute<FTransform>("ExplodedTransform", FGeometryCollection::TransformGroup);
 
 		FVector3f SumOfOffsets(0, 0, 0);
 		for (int32 ChildBoneIndex : ChildBones)
@@ -130,7 +130,7 @@ void FGeometryCollectionClusteringUtility::ClusterBonesUnderExistingRoot(FGeomet
 	check(GeometryCollection);
 	bool CalcNewLocalTransform = true;
 
-	TManagedArray<int32>& Levels = GeometryCollection->GetAttribute<int32>("Level", FGeometryCollection::TransformGroup);
+	TManagedArray<int32>& Levels = GeometryCollection->ModifyAttribute<int32>("Level", FGeometryCollection::TransformGroup);
 
 	TManagedArray<FTransform>& Transforms = GeometryCollection->Transform;
 	TManagedArray<FString>& BoneNames = GeometryCollection->BoneName;
@@ -299,7 +299,7 @@ void FGeometryCollectionClusteringUtility::CollapseHierarchyOneLevel(FGeometryCo
 	TManagedArray<TSet<int32>>& Children = GeometryCollection->Children;
 	TManagedArray<FTransform>& Transforms = GeometryCollection->Transform;
 	TManagedArray<FString>& BoneNames = GeometryCollection->BoneName;
-	TManagedArray<int32>& Levels = GeometryCollection->GetAttribute<int32>("Level", FGeometryCollection::TransformGroup);
+	TManagedArray<int32>& Levels = GeometryCollection->ModifyAttribute<int32>("Level", FGeometryCollection::TransformGroup);
 
 	for (int32 SourceElement : SourceElements)
 	{
@@ -633,7 +633,7 @@ void FGeometryCollectionClusteringUtility::UpdateHierarchyLevelOfChildren(FGeome
 	{
 		GeometryCollection->AddAttribute<int32>("Level", FGeometryCollection::TransformGroup, FManagedArrayCollection::FConstructionParameters(FName(),  false));
 	}
-	TManagedArray<int32>& Levels = GeometryCollection->GetAttribute<int32>("Level", FGeometryCollection::TransformGroup);
+	TManagedArray<int32>& Levels = GeometryCollection->ModifyAttribute<int32>("Level", FGeometryCollection::TransformGroup);
 	const TManagedArray<TSet<int32>>& Children = GeometryCollection->Children;
 	check(ParentElement < Levels.Num());
 	check(ParentElement < Children.Num());
@@ -710,7 +710,7 @@ void FGeometryCollectionClusteringUtility::CollapseLevelHierarchy(int8 Level, FG
 void FGeometryCollectionClusteringUtility::CollapseSelectedHierarchy(int8 Level, const TArray<int32>& SelectedBones, FGeometryCollection* GeometryCollection)
 {
 	check(GeometryCollection);
-	TManagedArray<int32>& Levels = GeometryCollection->GetAttribute<int32>("Level", FGeometryCollection::TransformGroup);
+	const TManagedArray<int32>& Levels = GeometryCollection->GetAttribute<int32>("Level", FGeometryCollection::TransformGroup);
 	const TManagedArray<TSet<int32>>& Children = GeometryCollection->Children;
 
 	// can't collapse root node away and doesn't make sense to operate when AllLevels selected
diff --git a/Engine/Source/Runtime/Experimental/Chaos/Private/GeometryCollection/GeometryCollectionConvexPropertiesInterface.cpp b/Engine/Source/Runtime/Experimental/Chaos/Private/GeometryCollection/GeometryCollectionConvexPropertiesInterface.cpp
--- a/Engine/Source/Runtime/Experimental/Chaos/Private/GeometryCollection/GeometryCollectionConvexPropertiesInterface.cpp
+++ b/Engine/Source/Runtime/Experimental/Chaos/Private/GeometryCollection/GeometryCollectionConvexPropertiesInterface.cpp
@@ -54,7 +54,7 @@ FGeometryCollectionConvexPropertiesInterface::InitializeInterface()
 void FGeometryCollectionConvexPropertiesInterface::SetDefaultProperty()
 {
 	FConvexCreationProperties DefaultProperty;
-	TManagedArray<int32>& Index = ManagedCollection->GetAttribute<int32>(ConvexIndexAttribute, ConvexPropertiesGroup);
+	TManagedArray<int32>& Index = ManagedCollection->ModifyAttribute<int32>(ConvexIndexAttribute, ConvexPropertiesGroup);
 	if (Index.Num() == 0 || !Index.Contains(INDEX_NONE))
 	{
 		int32 DefaultAttributeIndex = ManagedCollection->AddElements(1, ConvexPropertiesGroup);
@@ -70,7 +70,7 @@ FGeometryCollectionConvexPropertiesInterface::GetDefaultProperty() const
 	const TManagedArray<int32>& Index = ManagedCollection->GetAttribute<int32>(ConvexIndexAttribute, ConvexPropertiesGroup);
 	if (Index.Contains(INDEX_NONE))
 	{
-		TManagedArray<bool>& Enable = ManagedCollection->GetAttribute<bool>(ConvexEnable, ConvexPropertiesGroup);
+		const TManagedArray<bool>& Enable = ManagedCollection->GetAttribute<bool>(ConvexEnable, ConvexPropertiesGroup);
 		const TManagedArray<float>& FractionRemove = ManagedCollection->GetAttribute<float>(ConvexFractionRemoveAttribute, ConvexPropertiesGroup);
 		const TManagedArray<float>& SimplificationThreshold = ManagedCollection->GetAttribute<float>(ConvexSimplificationThresholdAttribute, ConvexPropertiesGroup);
 		const TManagedArray<float>& CanExceedFraction = ManagedCollection->GetAttribute<float>(ConvexCanExceedFractionAttribute, ConvexPropertiesGroup);
@@ -108,10 +108,10 @@ FGeometryCollectionConvexPropertiesInterface::GetConvexProperties(int TransformG
 	int32 PropIndex = Index.Find(TransformGroupIndex);
 	if (0 <= PropIndex && PropIndex < Index.Num())
 	{
-		TManagedArray<bool>& Enable = ManagedCollection->GetAttribute<bool>(ConvexEnable, ConvexPropertiesGroup);
+		const TManagedArray<bool>& Enable = ManagedCollection->GetAttribute<bool>(ConvexEnable, ConvexPropertiesGroup);
 		const TManagedArray<float>& FractionRemove = ManagedCollection->GetAttribute<float>(ConvexFractionRemoveAttribute, ConvexPropertiesGroup);
 		const TManagedArray<float>& SimplificationThreshold = ManagedCollection->GetAttribute<float>(ConvexSimplificationThresholdAttribute, ConvexPropertiesGroup);
-		TManagedArray<float>& CanExceedFraction = ManagedCollection->GetAttribute<float>(ConvexCanExceedFractionAttribute, ConvexPropertiesGroup);
+		const TManagedArray<float>& CanExceedFraction = ManagedCollection->GetAttribute<float>(ConvexCanExceedFractionAttribute, ConvexPropertiesGroup);
 
 		ConvexProperty.Enable = Enable[PropIndex];
 		ConvexProperty.FractionRemove = FractionRemove[PropIndex];
@@ -125,11 +125,11 @@ void
 FGeometryCollectionConvexPropertiesInterface::SetConvexProperties(const FConvexCreationProperties& InConvexAttributes, int TransformGroupIndex)
 {
 	FConvexCreationProperties Property;
-	TManagedArray<int32>& Index = ManagedCollection->GetAttribute<int32>(ConvexIndexAttribute, ConvexPropertiesGroup);
-	TManagedArray<bool>& Enable = ManagedCollection->GetAttribute<bool>(ConvexEnable, ConvexPropertiesGroup);
-	TManagedArray<float>& FractionRemove = ManagedCollection->GetAttribute<float>(ConvexFractionRemoveAttribute, ConvexPropertiesGroup);
-	TManagedArray<float>& SimplificationThreshold = ManagedCollection->GetAttribute<float>(ConvexSimplificationThresholdAttribute, ConvexPropertiesGroup);
-	TManagedArray<float>& CanExceedFraction = ManagedCollection->GetAttribute<float>(ConvexCanExceedFractionAttribute, ConvexPropertiesGroup);
+	TManagedArray<int32>& Index = ManagedCollection->ModifyAttribute<int32>(ConvexIndexAttribute, ConvexPropertiesGroup);
+	TManagedArray<bool>& Enable = ManagedCollection->ModifyAttribute<bool>(ConvexEnable, ConvexPropertiesGroup);
+	TManagedArray<float>& FractionRemove = ManagedCollection->ModifyAttribute<float>(ConvexFractionRemoveAttribute, ConvexPropertiesGroup);
+	TManagedArray<float>& SimplificationThreshold = ManagedCollection->ModifyAttribute<float>(ConvexSimplificationThresholdAttribute, ConvexPropertiesGroup);
+	TManagedArray<float>& CanExceedFraction = ManagedCollection->ModifyAttribute<float>(ConvexCanExceedFractionAttribute, ConvexPropertiesGroup);
 
 	int32 AttributeIndex = INDEX_NONE;
 	if (!Index.Contains(TransformGroupIndex))
diff --git a/Engine/Source/Runtime/Experimental/Chaos/Private/GeometryCollection/GeometryCollectionConvexUtility.cpp b/Engine/Source/Runtime/Experimental/Chaos/Private/GeometryCollection/GeometryCollectionConvexUtility.cpp
--- a/Engine/Source/Runtime/Experimental/Chaos/Private/GeometryCollection/GeometryCollectionConvexUtility.cpp
+++ b/Engine/Source/Runtime/Experimental/Chaos/Private/GeometryCollection/GeometryCollectionConvexUtility.cpp
@@ -21,8 +21,8 @@ TOptional<FGeometryCollectionConvexUtility::FGeometryCollectionConvexData> FGeom
 	}
 
 	FGeometryCollectionConvexUtility::FGeometryCollectionConvexData ConvexData{
-		GeometryCollection->GetAttribute<TSet<int32>>("TransformToConvexIndices", FTransformCollection::TransformGroup),
-		GeometryCollection->GetAttribute<TUniquePtr<Chaos::FConvex>>("ConvexHull", "Convex")
+		GeometryCollection->ModifyAttribute<TSet<int32>>("TransformToConvexIndices", FTransformCollection::TransformGroup),
+		GeometryCollection->ModifyAttribute<TUniquePtr<Chaos::FConvex>>("ConvexHull", "Convex")
 	};
 	return TOptional<FGeometryCollectionConvexUtility::FGeometryCollectionConvexData>(ConvexData);
 }
@@ -55,8 +55,8 @@ FGeometryCollectionConvexUtility::FGeometryCollectionConvexData FGeometryCollect
 	// Check for correct population. Make sure all rigid nodes should have a convex associated; leave convex hulls for transform nodes alone for now
 	const TManagedArray<int32>& SimulationType = GeometryCollection->GetAttribute<int32>("SimulationType", FTransformCollection::TransformGroup);
 	const TManagedArray<int32>& TransformToGeometryIndex = GeometryCollection->GetAttribute<int32>("TransformToGeometryIndex", FTransformCollection::TransformGroup);
-	TManagedArray<TSet<int32>>& TransformToConvexIndices = GeometryCollection->GetAttribute<TSet<int32>>("TransformToConvexIndices", FTransformCollection::TransformGroup);
-	TManagedArray<TUniquePtr<Chaos::FConvex>>& ConvexHull = GeometryCollection->GetAttribute<TUniquePtr<Chaos::FConvex>>("ConvexHull", "Convex");
+	TManagedArray<TSet<int32>>& TransformToConvexIndices = GeometryCollection->ModifyAttribute<TSet<int32>>("TransformToConvexIndices", FTransformCollection::TransformGroup);
+	TManagedArray<TUniquePtr<Chaos::FConvex>>& ConvexHull = GeometryCollection->ModifyAttribute<TUniquePtr<Chaos::FConvex>>("ConvexHull", "Convex");
 	
 	TArray<int32> ProduceConvexHulls;
 	ProduceConvexHulls.Reserve(SimulationType.Num());
@@ -1159,8 +1159,8 @@ FGeometryCollectionConvexUtility::FGeometryCollectionConvexData FGeometryCollect
 		GeometryCollection->AddAttribute<TUniquePtr<Chaos::FConvex>>("ConvexHull", "Convex");
 	}
 
-	TManagedArray<TSet<int32>>& TransformToConvexIndices = GeometryCollection->GetAttribute<TSet<int32>>("TransformToConvexIndices", FTransformCollection::TransformGroup);
-	TManagedArray<TUniquePtr<Chaos::FConvex>>& ConvexHull = GeometryCollection->GetAttribute<TUniquePtr<Chaos::FConvex>>("ConvexHull", "Convex");
+	TManagedArray<TSet<int32>>& TransformToConvexIndices = GeometryCollection->ModifyAttribute<TSet<int32>>("TransformToConvexIndices", FTransformCollection::TransformGroup);
+	TManagedArray<TUniquePtr<Chaos::FConvex>>& ConvexHull = GeometryCollection->ModifyAttribute<TUniquePtr<Chaos::FConvex>>("ConvexHull", "Convex");
 	TransformToConvexIndices = MoveTemp(TransformToConvexIndexArr);
 	GeometryCollection->EmptyGroup("Convex");
 	GeometryCollection->Resize(Convexes.Num(), "Convex");
@@ -1241,7 +1241,7 @@ void FGeometryCollectionConvexUtility::RemoveConvexHulls(FGeometryCollection* Ge
 {
 	if (GeometryCollection->HasGroup("Convex") && GeometryCollection->HasAttribute("TransformToConvexIndices", FTransformCollection::TransformGroup))
 	{
-		TManagedArray<TSet<int32>>& TransformToConvexIndices = GeometryCollection->GetAttribute<TSet<int32>>("TransformToConvexIndices", FTransformCollection::TransformGroup);
+		TManagedArray<TSet<int32>>& TransformToConvexIndices = GeometryCollection->ModifyAttribute<TSet<int32>>("TransformToConvexIndices", FTransformCollection::TransformGroup);
 		TArray<int32> ConvexIndices;
 		for (int32 TransformIdx : SortedTransformDeletes)
 		{
@@ -1279,7 +1279,7 @@ TManagedArray<int32>* FGeometryCollectionConvexUtility::GetCustomConvexFlags(FGe
 		}
 		return &GeometryCollection->AddAttribute<int32>("HasCustomConvex", FTransformCollection::TransformGroup);
 	}
-	return &GeometryCollection->GetAttribute<int32>("HasCustomConvex", FTransformCollection::TransformGroup);
+	return &GeometryCollection->ModifyAttribute<int32>("HasCustomConvex", FTransformCollection::TransformGroup);
 }
 
 
@@ -1418,8 +1418,8 @@ void FGeometryCollectionConvexUtility::SetVolumeAttributes(FGeometryCollection*
 	{
 		Collection->AddAttribute<float>("Size", FGeometryCollection::TransformGroup);
 	}
-	TManagedArray<float>& Volumes = Collection->GetAttribute<float>("Volume", FTransformCollection::TransformGroup);
-	TManagedArray<float>& Sizes = Collection->GetAttribute<float>("Size", FTransformCollection::TransformGroup);
+	TManagedArray<float>& Volumes = Collection->ModifyAttribute<float>("Volume", FTransformCollection::TransformGroup);
+	TManagedArray<float>& Sizes = Collection->ModifyAttribute<float>("Size", FTransformCollection::TransformGroup);
 
 	const TManagedArray<int32>& SimulationType = Collection->SimulationType;
 	const TManagedArray<int32>& TransformToGeometryIndex = Collection->TransformToGeometryIndex;
diff --git a/Engine/Source/Runtime/Experimental/Chaos/Private/GeometryCollection/GeometryCollectionProximityUtility.cpp b/Engine/Source/Runtime/Experimental/Chaos/Private/GeometryCollection/GeometryCollectionProximityUtility.cpp
--- a/Engine/Source/Runtime/Experimental/Chaos/Private/GeometryCollection/GeometryCollectionProximityUtility.cpp
+++ b/Engine/Source/Runtime/Experimental/Chaos/Private/GeometryCollection/GeometryCollectionProximityUtility.cpp
@@ -211,7 +211,7 @@ struct FGeometryCollectionProximitySpatial
 			Collection->AddAttribute<TSet<int32>>("Proximity", FGeometryCollection::GeometryGroup, GeometryDependency);
 		}
 
-		TManagedArray<TSet<int32>>& Proximity = Collection->GetAttribute<TSet<int32>>("Proximity", FGeometryCollection::GeometryGroup);
+		TManagedArray<TSet<int32>>& Proximity = Collection->ModifyAttribute<TSet<int32>>("Proximity", FGeometryCollection::GeometryGroup);
 
 		for (int32 GeoIdx = 0; GeoIdx < KnownProximity.Num(); ++GeoIdx)
 		{
diff --git a/Engine/Source/Runtime/Experimental/Chaos/Private/GeometryCollection/GeometryCollectionUtility.cpp b/Engine/Source/Runtime/Experimental/Chaos/Private/GeometryCollection/GeometryCollectionUtility.cpp
--- a/Engine/Source/Runtime/Experimental/Chaos/Private/GeometryCollection/GeometryCollectionUtility.cpp
+++ b/Engine/Source/Runtime/Experimental/Chaos/Private/GeometryCollection/GeometryCollectionUtility.cpp
@@ -488,7 +488,7 @@ namespace GeometryCollection
 
 		if (bNeedsInit || bForceInit)
 		{
-			TManagedArray<FGuid>& Guids = Collection->GetAttribute<FGuid>("GUID", FTransformCollection::TransformGroup);
+			TManagedArray<FGuid>& Guids = Collection->ModifyAttribute<FGuid>("GUID", FTransformCollection::TransformGroup);
 			for (int32 Idx = StartIdx; Idx < Guids.Num(); ++Idx)
 			{
 				Guids[Idx] = FGuid::NewGuid();
diff --git a/Engine/Source/Runtime/Experimental/Chaos/Private/GeometryCollection/ManagedArrayCollection.cpp b/Engine/Source/Runtime/Experimental/Chaos/Private/GeometryCollection/ManagedArrayCollection.cpp
--- a/Engine/Source/Runtime/Experimental/Chaos/Private/GeometryCollection/ManagedArrayCollection.cpp
+++ b/Engine/Source/Runtime/Experimental/Chaos/Private/GeometryCollection/ManagedArrayCollection.cpp
@@ -86,6 +86,16 @@ bool FManagedArrayCollection::HasAttribute(FName Name, FName Group) const
 	return Map.Contains(Key);
 }
 
+bool FManagedArrayCollection::IsAttributeDirty(FName Name, FName Group) const
+{
+	const FKeyType Key = FManagedArrayCollection::MakeMapKey(Name, Group);
+	if (const FValueType* Attribute = Map.Find(Key))
+	{
+		return (Attribute->Value)? Attribute->Value->IsDirty(): false;
+	}
+	return false;
+}
+
 TArray<FName> FManagedArrayCollection::AttributeNames(FName Group) const
 {
 	TArray<FName> AttributeNames;
diff --git a/Engine/Source/Runtime/Experimental/Chaos/Private/GeometryCollection/TransformCollection.cpp b/Engine/Source/Runtime/Experimental/Chaos/Private/GeometryCollection/TransformCollection.cpp
--- a/Engine/Source/Runtime/Experimental/Chaos/Private/GeometryCollection/TransformCollection.cpp
+++ b/Engine/Source/Runtime/Experimental/Chaos/Private/GeometryCollection/TransformCollection.cpp
@@ -39,10 +39,10 @@ void FTransformCollection::Serialize(Chaos::FChaosArchive& Ar)
 	if (Ar.IsLoading())
 	{
 		// for backwards compatibility convert old BoneHierarchy struct into split out arrays
-		TManagedArray<FGeometryCollectionBoneNode>* BoneHierarchyPtr = FindAttribute<FGeometryCollectionBoneNode>("BoneHierarchy", FTransformCollection::TransformGroup);
+		const TManagedArray<FGeometryCollectionBoneNode>* BoneHierarchyPtr = FindAttribute<FGeometryCollectionBoneNode>("BoneHierarchy", FTransformCollection::TransformGroup);
 		if (BoneHierarchyPtr)
 		{
-			TManagedArray<FGeometryCollectionBoneNode>& BoneHierarchy = *BoneHierarchyPtr;
+			const TManagedArray<FGeometryCollectionBoneNode>& BoneHierarchy = *BoneHierarchyPtr;
 
 			for (int Idx = 0; Idx < BoneHierarchy.Num(); Idx++)
 			{
diff --git a/Engine/Source/Runtime/Experimental/Chaos/Private/PhysicsProxy/GeometryCollectionPhysicsProxy.cpp b/Engine/Source/Runtime/Experimental/Chaos/Private/PhysicsProxy/GeometryCollectionPhysicsProxy.cpp
--- a/Engine/Source/Runtime/Experimental/Chaos/Private/PhysicsProxy/GeometryCollectionPhysicsProxy.cpp
+++ b/Engine/Source/Runtime/Experimental/Chaos/Private/PhysicsProxy/GeometryCollectionPhysicsProxy.cpp
@@ -2324,7 +2324,7 @@ void FGeometryCollectionPhysicsProxy::InitializeSharedCollisionStructures(
 	const TManagedArray<int32>& Parent = RestCollection.Parent;
 	const TManagedArray<TSet<int32>>& Children = RestCollection.Children;
 	const TManagedArray<int32>& SimulationType = RestCollection.SimulationType;
-	TManagedArray<bool>& CollectionSimulatableParticles = RestCollection.GetAttribute<bool>(FGeometryCollection::SimulatableParticlesAttribute, FTransformCollection::TransformGroup);
+	TManagedArray<bool>& CollectionSimulatableParticles = RestCollection.ModifyAttribute<bool>(FGeometryCollection::SimulatableParticlesAttribute, FTransformCollection::TransformGroup);
 	TManagedArray<FVector3f>& CollectionInertiaTensor = RestCollection.AddAttribute<FVector3f>(TEXT("InertiaTensor"), FTransformCollection::TransformGroup);
 	TManagedArray<FRealSingle>& CollectionMass = RestCollection.AddAttribute<FRealSingle>(TEXT("Mass"), FTransformCollection::TransformGroup);
 	TManagedArray<TUniquePtr<FSimplicial>>& CollectionSimplicials =	RestCollection.AddAttribute<TUniquePtr<FSimplicial>>(FGeometryDynamicCollection::SimplicialsAttribute, FTransformCollection::TransformGroup);
diff --git a/Engine/Source/Runtime/Experimental/Chaos/Public/GeometryCollection/ManagedArray.h b/Engine/Source/Runtime/Experimental/Chaos/Public/GeometryCollection/ManagedArray.h
--- a/Engine/Source/Runtime/Experimental/Chaos/Public/GeometryCollection/ManagedArray.h
+++ b/Engine/Source/Runtime/Experimental/Chaos/Public/GeometryCollection/ManagedArray.h
@@ -104,9 +104,27 @@ protected:
 	virtual void Init(const FManagedArrayBase& ) {};
 
 public:
-
+	FManagedArrayBase()
+	{
+		ClearDirtyFlag();
+	}
+	
 	virtual ~FManagedArrayBase() {}
 
+	FORCEINLINE void ClearDirtyFlag()
+	{
+		bIsDirty = false;
+	}
+	
+	FORCEINLINE_DEBUGGABLE void MarkDirty()
+	{
+		bIsDirty = true;
+	}
+	
+	FORCEINLINE bool IsDirty() const
+	{
+		return bIsDirty;
+	}
 	
 	//todo(ocohen): these should all be private with friend access to managed collection
 	
@@ -167,7 +185,8 @@ public:
 	{
 		check(false);
 	}
-
+private:
+	bool bIsDirty;
 };
 
 template <typename T>
diff --git a/Engine/Source/Runtime/Experimental/Chaos/Public/GeometryCollection/ManagedArrayCollection.h b/Engine/Source/Runtime/Experimental/Chaos/Public/GeometryCollection/ManagedArrayCollection.h
--- a/Engine/Source/Runtime/Experimental/Chaos/Public/GeometryCollection/ManagedArrayCollection.h
+++ b/Engine/Source/Runtime/Experimental/Chaos/Public/GeometryCollection/ManagedArrayCollection.h
@@ -106,7 +106,7 @@ public:
 			}
 			Map.Add(FManagedArrayCollection::MakeMapKey(Name, Group), MoveTemp(Value));
 		}
-		return GetAttribute<T>(Name, Group);
+		return ModifyAttribute<T>(Name, Group);
 	}
 
 	/**
@@ -237,24 +237,44 @@ public:
 	};
 
 	/**
-	* Returns attribute access of Type(T) from the group
+	* Returns attribute access of Type(T) from the group for modification
+	* this will mark the collection dirty
 	* @param Name - The name of the attribute
 	* @param Group - The group that manages the attribute
 	* @return ManagedArray<T> &
 	*/
 	template<typename T>
-	TManagedArray<T>& GetAttribute(FName Name, FName Group)
+	TManagedArray<T>& ModifyAttribute(FName Name, FName Group)
 	{
 		check(HasAttribute(Name, Group))
-		FKeyType Key = FManagedArrayCollection::MakeMapKey(Name, Group);
-		return *(static_cast<TManagedArray<T>*>(Map[Key].Value));
-	};
+		const FKeyType Key = FManagedArrayCollection::MakeMapKey(Name, Group);
+		FManagedArrayBase* ManagedArray = Map[Key].Value;
+		ManagedArray->MarkDirty();
+		UE_LOG(LogChaos, Warning, TEXT("Modifying attribute : [%s, %s]"), *Name.ToString(), *Group.ToString());
+		return *(static_cast<TManagedArray<T>*>(ManagedArray));
+	}
+	
+	/**
+	* Returns attribute access of Type(T) from the group
+	* @param Name - The name of the attribute
+	* @param Group - The group that manages the attribute
+	* @return ManagedArray<T> &
+	*/
+	
+	// template<typename T>
+	// UE_DEPRECATED(5.0, "non const GetAttribute() version is now deprecated, use ModifyAttribute instead")
+	// TManagedArray<T>& GetAttribute(FName Name, FName Group)
+	// {
+	// 	check(HasAttribute(Name, Group))
+	// 	const FKeyType Key = FManagedArrayCollection::MakeMapKey(Name, Group);
+	// 	return *(static_cast<TManagedArray<T>*>(Map[Key].Value));
+	// };
 
 	template<typename T>
 	const TManagedArray<T>& GetAttribute(FName Name, FName Group) const
 	{
 		check(HasAttribute(Name, Group))
-		FKeyType Key = FManagedArrayCollection::MakeMapKey(Name, Group);
+		const FKeyType Key = FManagedArrayCollection::MakeMapKey(Name, Group);
 		return *(static_cast<TManagedArray<T>*>(Map[Key].Value));
 	};
 
@@ -305,6 +325,13 @@ public:
 	*/
 	FORCEINLINE bool HasGroup(FName Group) const { return GroupInfo.Contains(Group); }
 
+	/**
+	* Check if an attribute is dirty
+	* @param Name - The name of the attribute
+	* @param Group - The group that manages the attribute
+	*/
+	bool IsAttributeDirty(FName Name, FName Group) const;
+	
 	/**
 	*
 	*/
diff --git a/Engine/Source/Runtime/Experimental/GeometryCollectionEngine/Private/GeometryCollection/GeometryCollectionComponent.cpp b/Engine/Source/Runtime/Experimental/GeometryCollectionEngine/Private/GeometryCollection/GeometryCollectionComponent.cpp
--- a/Engine/Source/Runtime/Experimental/GeometryCollectionEngine/Private/GeometryCollection/GeometryCollectionComponent.cpp
+++ b/Engine/Source/Runtime/Experimental/GeometryCollectionEngine/Private/GeometryCollection/GeometryCollectionComponent.cpp
@@ -443,7 +443,7 @@ FBoxSphereBounds UGeometryCollectionComponent::CalcBounds(const FTransform& Loca
 			TArray<FMatrix> TmpGlobalMatrices;
 			GeometryCollectionAlgo::GlobalMatrices(Transforms, ParentIndices, TmpGlobalMatrices);
 
-			TManagedArray<FBox>& TransformBounds = HackGeometryCollectionPtr->GetAttribute<FBox>("BoundingBox", "Transform");
+			TManagedArray<FBox>& TransformBounds = HackGeometryCollectionPtr->ModifyAttribute<FBox>("BoundingBox", "Transform");
 			for (int32 TransformIndex = 0; TransformIndex < HackGeometryCollectionPtr->NumElements(FGeometryCollection::TransformGroup); TransformIndex++)
 			{
 				BoundingBox += TransformBounds[TransformIndex].TransformBy(TmpGlobalMatrices[TransformIndex] * LocalToWorldWithScale);
@@ -2952,7 +2952,7 @@ void FScopedColorEdit::UpdateBoneColors()
 
 		const TManagedArray<int>& Parents = Collection->Parent;
 		bool HasLevelAttribute = Collection->HasAttribute("Level", FTransformCollection::TransformGroup);
-		TManagedArray<int>* Levels = nullptr;
+		const TManagedArray<int>* Levels = nullptr;
 		if (HasLevelAttribute)
 		{
 			Levels = &Collection->GetAttribute<int32>("Level", FTransformCollection::TransformGroup);
@@ -3185,7 +3185,7 @@ void UGeometryCollectionComponent::CalculateGlobalMatrices()
 				const TManagedArray<float>& SleepTimer = DynamicCollection->GetAttribute<float>("SleepTimer", FGeometryCollection::TransformGroup);
 				const TManagedArray<float>& MaxSleepTime = DynamicCollection->GetAttribute<float>("MaxSleepTime", FGeometryCollection::TransformGroup);
 				const TManagedArray<float>& RemovalDuration = DynamicCollection->GetAttribute<float>("RemovalDuration", FGeometryCollection::TransformGroup);
-				TManagedArray<FTransform>& UniformScale = DynamicCollection->GetAttribute<FTransform>("UniformScale", FGeometryCollection::TransformGroup);
+				TManagedArray<FTransform>& UniformScale = DynamicCollection->ModifyAttribute<FTransform>("UniformScale", FGeometryCollection::TransformGroup);
 
 				for (int32 Idx = 0; Idx < GetTransformArray().Num(); ++Idx)
 				{
@@ -3436,7 +3436,7 @@ void UGeometryCollectionComponent::IncrementSleepTimer(float DeltaTime)
 		&& DynamicCollection->HasAttribute("MaxSleepTime", FGeometryCollection::TransformGroup)
 		&& DynamicCollection->HasAttribute("RemovalDuration", FGeometryCollection::TransformGroup))
 	{
-		TManagedArray<float>& SleepTimer = DynamicCollection->GetAttribute<float>("SleepTimer", FGeometryCollection::TransformGroup);
+		TManagedArray<float>& SleepTimer = DynamicCollection->ModifyAttribute<float>("SleepTimer", FGeometryCollection::TransformGroup);
 		const TManagedArray<float>& RemovalDuration = DynamicCollection->GetAttribute<float>("RemovalDuration", FGeometryCollection::TransformGroup);
 		const TManagedArray<float>& MaxSleepTime = DynamicCollection->GetAttribute<float>("MaxSleepTime", FGeometryCollection::TransformGroup);
 		TArray<int32> ToDisable;
diff --git a/Engine/Source/Runtime/Experimental/GeometryCollectionEngine/Private/GeometryCollection/GeometryCollectionObject.cpp b/Engine/Source/Runtime/Experimental/GeometryCollectionEngine/Private/GeometryCollection/GeometryCollectionObject.cpp
--- a/Engine/Source/Runtime/Experimental/GeometryCollectionEngine/Private/GeometryCollection/GeometryCollectionObject.cpp
+++ b/Engine/Source/Runtime/Experimental/GeometryCollectionEngine/Private/GeometryCollection/GeometryCollectionObject.cpp
@@ -1042,7 +1042,7 @@ TSharedPtr<FGeometryCollection, ESPMode::ThreadSafe> UGeometryCollection::Genera
 	// We also need to resize geometry groups to ensure that they are empty.
 	if (bStripOnCook)
 	{
-		TManagedArray<int32>& TransformToGeometryIndex = DuplicateGeometryCollection->GetAttribute<int32>("TransformToGeometryIndex", FTransformCollection::TransformGroup);
+		const TManagedArray<int32>& TransformToGeometryIndex = DuplicateGeometryCollection->GetAttribute<int32>("TransformToGeometryIndex", FTransformCollection::TransformGroup);
 
 		//
 		// Copy the bounds to the TransformGroup.
@@ -1058,16 +1058,16 @@ TSharedPtr<FGeometryCollection, ESPMode::ThreadSafe> UGeometryCollection::Genera
 			DuplicateGeometryCollection->AddAttribute<FBox>("NaniteIndex", "Transform");
 		}
 
-		int32 NumTransforms = GeometryCollection->NumElements(FGeometryCollection::TransformGroup);
-		TManagedArray<int32>& NaniteIndex = DuplicateGeometryCollection->GetAttribute<int32>("NaniteIndex", "Transform");
-		TManagedArray<FBox>& TransformBounds = DuplicateGeometryCollection->GetAttribute<FBox>("BoundingBox", "Transform");
-		TManagedArray<FBox>& GeometryBounds = GeometryCollection->GetAttribute<FBox>("BoundingBox", "Geometry");
+		const int32 NumTransforms = GeometryCollection->NumElements(FGeometryCollection::TransformGroup);
+		TManagedArray<int32>& NaniteIndex = DuplicateGeometryCollection->ModifyAttribute<int32>("NaniteIndex", "Transform");
+		TManagedArray<FBox>& TransformBounds = DuplicateGeometryCollection->ModifyAttribute<FBox>("BoundingBox", "Transform");
+		const TManagedArray<FBox>& GeometryBounds = GeometryCollection->GetAttribute<FBox>("BoundingBox", "Geometry");
 
 		NaniteIndex.Fill(INDEX_NONE);
 		for (int TransformIndex = 0; TransformIndex < NumTransforms; TransformIndex++)
 		{
 			NaniteIndex[TransformIndex] = TransformToGeometryIndex[TransformIndex];
-			int32 GeometryIndex = TransformToGeometryIndex[TransformIndex];
+			const int32 GeometryIndex = TransformToGeometryIndex[TransformIndex];
 			if (GeometryIndex != INDEX_NONE)
 			{
 				TransformBounds[TransformIndex] = GeometryBounds[GeometryIndex];
diff --git a/Engine/Source/Runtime/Experimental/GeometryCollectionEngine/Public/GeometryCollection/GeometryCollectionComponent.h b/Engine/Source/Runtime/Experimental/GeometryCollectionEngine/Public/GeometryCollection/GeometryCollectionComponent.h
--- a/Engine/Source/Runtime/Experimental/GeometryCollectionEngine/Public/GeometryCollection/GeometryCollectionComponent.h
+++ b/Engine/Source/Runtime/Experimental/GeometryCollectionEngine/Public/GeometryCollection/GeometryCollectionComponent.h
@@ -262,7 +262,7 @@ FORCEINLINE TManagedArray<Type>& Get##Name##ArrayCopyOnWrite()					\
 		DynamicCollection->CopyAttribute(										\
 			*RestCollection->GetGeometryCollection(), StaticName, Group);		\
 		Indirect##Name##Array =													\
-			&DynamicCollection->GetAttribute<Type>(StaticName, Group);			\
+			&DynamicCollection->ModifyAttribute<Type>(StaticName, Group);		\
 		CopyOnWriteAttributeList.Add(											\
 			reinterpret_cast<FManagedArrayBase**>(&Indirect##Name##Array));		\
 	}																			\
