Chaos destruction : fix support for per-component gravity flag for geometry collections ( only accounted for at initialization, does not yet work at runtime )

diff --git a/Engine/Source/Runtime/Experimental/Chaos/Private/PhysicsProxy/GeometryCollectionPhysicsProxy.cpp b/Engine/Source/Runtime/Experimental/Chaos/Private/PhysicsProxy/GeometryCollectionPhysicsProxy.cpp
--- a/Engine/Source/Runtime/Experimental/Chaos/Private/PhysicsProxy/GeometryCollectionPhysicsProxy.cpp
+++ b/Engine/Source/Runtime/Experimental/Chaos/Private/PhysicsProxy/GeometryCollectionPhysicsProxy.cpp
@@ -738,6 +738,18 @@ void FGeometryCollectionPhysicsProxy::CreateNonClusteredParticles(Chaos::FPBDRig
 	}
 }
 
+void FGeometryCollectionPhysicsProxy::SetGravityEnabled(const Chaos::FPBDRigidsSolver& RigidsSolver, bool bEnabled)
+{
+	Chaos::FPerParticleGravity& GravityForces = RigidsSolver.GetEvolution()->GetGravityForces();
+	for (int32 HandleIdx = 0; HandleIdx < SolverParticleHandles.Num(); ++HandleIdx)
+	{
+		if (Chaos::TPBDRigidParticleHandle<Chaos::FReal, 3>* Handle = SolverParticleHandles[HandleIdx])
+		{
+			Handle->SetGravityEnabled(bEnabled);
+		}
+	}
+}
+
 void FGeometryCollectionPhysicsProxy::InitializeBodiesPT(Chaos::FPBDRigidsSolver* RigidsSolver, typename Chaos::FPBDRigidsSolver::FParticlesType& Particles)
 {
 	const FGeometryCollection* RestCollection = Parameters.RestCollection;
@@ -1051,25 +1063,8 @@ void FGeometryCollectionPhysicsProxy::InitializeBodiesPT(Chaos::FPBDRigidsSolver
 		}
 #endif // TODO_REIMPLEMENT_RIGID_CACHING
 
-
-
-		if (DisableGeometryCollectionGravity) // cvar
-		{
-			// Our assumption is that you'd only ever want to wholesale opt geometry 
-			// collections out of gravity for debugging, so we keep this conditional
-			// out of the loop above and on it's own.  This means we can't turn gravity
-			// back on once it's off, but even if we didn't enclose this in an if(),
-			// this function won't be called again unless something dirties the proxy.
-
-			Chaos::FPerParticleGravity& GravityForces = RigidsSolver->GetEvolution()->GetGravityForces();
-			for (int32 HandleIdx = 0; HandleIdx < SolverParticleHandles.Num(); ++HandleIdx)
-			{
-				if (Chaos::TPBDRigidParticleHandle<Chaos::FReal, 3>* Handle = SolverParticleHandles[HandleIdx])
-				{
-					Handle->SetGravityEnabled(false);
-				}
-			}
-		}
+		const bool bEnableGravity = Parameters.EnableGravity && !DisableGeometryCollectionGravity;
+		SetGravityEnabled(*RigidsSolver, bEnableGravity);
 
 		// call DirtyParticle to make sure the acceleration structure is up to date with all the changes happening here
 		for (int32 TransformGroupIndex = 0; TransformGroupIndex < NumTransforms; ++TransformGroupIndex)
diff --git a/Engine/Source/Runtime/Experimental/Chaos/Public/GeometryCollection/GeometryCollectionSimulationCoreTypes.h b/Engine/Source/Runtime/Experimental/Chaos/Public/GeometryCollection/GeometryCollectionSimulationCoreTypes.h
--- a/Engine/Source/Runtime/Experimental/Chaos/Public/GeometryCollection/GeometryCollectionSimulationCoreTypes.h
+++ b/Engine/Source/Runtime/Experimental/Chaos/Public/GeometryCollection/GeometryCollectionSimulationCoreTypes.h
@@ -193,6 +193,7 @@ struct FSimulationParameters
 		, bGenerateTrailingData(false)
 		, bGenerateRemovalsData(false)
 		, RemoveOnFractureEnabled(false)
+		, EnableGravity(true)
 		, SimulationFilterData()
 		, QueryFilterData()
 		, UserData(nullptr)
@@ -229,7 +230,8 @@ struct FSimulationParameters
 		, bGenerateTrailingData(Other.bGenerateTrailingData)
 		, bGenerateRemovalsData(Other.bGenerateRemovalsData)
 		, Shared(Other.Shared)
-		, RemoveOnFractureEnabled(false)
+		, RemoveOnFractureEnabled(Other.RemoveOnFractureEnabled)
+		, EnableGravity(Other.EnableGravity)
 		, SimulationFilterData(Other.SimulationFilterData)
 		, QueryFilterData(Other.QueryFilterData)
 		, UserData(Other.UserData)
@@ -289,6 +291,8 @@ struct FSimulationParameters
 	FSharedSimulationParameters Shared;
 
 	bool RemoveOnFractureEnabled;
+	
+	bool EnableGravity;
 
 	FCollisionFilterData SimulationFilterData;
 	FCollisionFilterData QueryFilterData;
diff --git a/Engine/Source/Runtime/Experimental/Chaos/Public/PhysicsProxy/GeometryCollectionPhysicsProxy.h b/Engine/Source/Runtime/Experimental/Chaos/Public/PhysicsProxy/GeometryCollectionPhysicsProxy.h
--- a/Engine/Source/Runtime/Experimental/Chaos/Public/PhysicsProxy/GeometryCollectionPhysicsProxy.h
+++ b/Engine/Source/Runtime/Experimental/Chaos/Public/PhysicsProxy/GeometryCollectionPhysicsProxy.h
@@ -273,6 +273,8 @@ protected:
 		const Chaos::FClusterCreationParameters & Parameters,
 		const Chaos::FUniqueIdx* ExistingIndex);
 
+	void SetGravityEnabled(const Chaos::FPBDRigidsSolver& RigidsSolver, bool bEnabled);
+	
 	/** 
 	 * Traverses the parents of \p TransformIndex in \p GeometryCollection, counting
 	 * the number of levels until the next parent is \c INDEX_NONE.
diff --git a/Engine/Source/Runtime/Experimental/GeometryCollectionEngine/Private/GeometryCollection/GeometryCollectionComponent.cpp b/Engine/Source/Runtime/Experimental/GeometryCollectionEngine/Private/GeometryCollection/GeometryCollectionComponent.cpp
--- a/Engine/Source/Runtime/Experimental/GeometryCollectionEngine/Private/GeometryCollection/GeometryCollectionComponent.cpp
+++ b/Engine/Source/Runtime/Experimental/GeometryCollectionEngine/Private/GeometryCollection/GeometryCollectionComponent.cpp
@@ -2247,6 +2247,7 @@ void UGeometryCollectionComponent::RegisterAndInitializePhysicsProxy()
 		SimulationParameters.bGenerateTrailingData = bNotifyTrailing;
 		SimulationParameters.bGenerateRemovalsData = bNotifyRemovals;
 		SimulationParameters.RemoveOnFractureEnabled = SimulationParameters.Shared.RemoveOnFractureIndices.Num() > 0;
+		SimulationParameters.EnableGravity = BodyInstance.bEnableGravity;
 		SimulationParameters.WorldTransform = GetComponentToWorld();
 		SimulationParameters.UserData = static_cast<void*>(&PhysicsUserData);
 
