BVH optimization : Rotate geometry instead of rotating AABB

diff --git a/Engine/Source/Programs/HeadlessChaos/Private/HeadlessChaosTestOverlap.cpp b/Engine/Source/Programs/HeadlessChaos/Private/HeadlessChaosTestOverlap.cpp
--- a/Engine/Source/Programs/HeadlessChaos/Private/HeadlessChaosTestOverlap.cpp
+++ b/Engine/Source/Programs/HeadlessChaos/Private/HeadlessChaosTestOverlap.cpp
@@ -102,6 +102,66 @@ namespace ChaosTest
 				EXPECT_EQ(bResult, false);
 			}
 		}
+		{
+			const FVec3 X1 = { 0.0, 0.0, -200.0 };
+			const FVec3 X2 = { 0.0, 0.0, 200.0 };
+			const FReal Radius = 100.0;
+			const FCapsule Capsule = FCapsule(X1, X2, Radius);
+			TSharedPtr<FCapsule> CapsuleShared = MakeShared<FCapsule>(X1, X2, Radius);
+			FVec3 TriMeshScale = { 0.01, 0.01, 0.01 };
+			TImplicitObjectScaled<FCapsule> ScaledCapsule = TImplicitObjectScaled<FCapsule>(CapsuleShared, TriMeshScale);
+			{
+				const FRigidTransform3 QueryTM(FVec3(0.0, 0.0, 0.0), FQuat::Identity);
+				bool bResult = TriangleMesh->OverlapGeom(ScaledCapsule, QueryTM, 0.0, nullptr, TriMeshScale);
+				EXPECT_EQ(bResult, true);
+			}
+			{
+				const FRigidTransform3 QueryTM(FVec3(10.0, 10.0, 0.0), FQuat::Identity);
+				bool bResult = TriangleMesh->OverlapGeom(ScaledCapsule, QueryTM, 0.0, nullptr, TriMeshScale);
+				EXPECT_EQ(bResult, false);
+			}
+			{
+				const FRigidTransform3 QueryTM(FVec3(-10.0, -10.0, 0.0), FQuat::Identity);
+				bool bResult = TriangleMesh->OverlapGeom(ScaledCapsule, QueryTM, 0.0, nullptr, TriMeshScale);
+				EXPECT_EQ(bResult, false);
+			}
+			{
+				const FRigidTransform3 QueryTM(FVec3(10.0, 0.0, 0.0), FQuat::Identity);
+				bool bResult = TriangleMesh->OverlapGeom(ScaledCapsule, QueryTM, 0.0, nullptr, TriMeshScale);
+				EXPECT_EQ(bResult, true);
+			}
+			{
+				const FRigidTransform3 QueryTM(FVec3(0.0, 10.0, 0.0), FQuat::Identity);
+				bool bResult = TriangleMesh->OverlapGeom(ScaledCapsule, QueryTM, 0.0, nullptr, TriMeshScale);
+				EXPECT_EQ(bResult, true);
+			}
+			{
+				const FRigidTransform3 QueryTM(FVec3(0.0, -10.0, 0.0), FQuat::Identity);
+				bool bResult = TriangleMesh->OverlapGeom(ScaledCapsule, QueryTM, 0.0, nullptr, TriMeshScale);
+				EXPECT_EQ(bResult, true);
+			}
+			{
+				const FRigidTransform3 QueryTM(FVec3(-10.0, 0.0, 0.0), FQuat::Identity);
+				bool bResult = TriangleMesh->OverlapGeom(ScaledCapsule, QueryTM, 0.0, nullptr, TriMeshScale);
+				EXPECT_EQ(bResult, true);
+			}
+			{
+				const FRigidTransform3 QueryTM(FVec3(0.0, 0.0, 12.0), FQuat::Identity);
+				bool bResult = TriangleMesh->OverlapGeom(ScaledCapsule, QueryTM, 0.0, nullptr, TriMeshScale);
+				EXPECT_EQ(bResult, true);
+			}
+			{
+				const FRigidTransform3 QueryTM(FVec3(0.0, 0.0, 12.0), FQuat(FVec3(1.0, 0.0, 0.0), 3.1415926 / 4.0));
+				bool bResult = TriangleMesh->OverlapGeom(ScaledCapsule, QueryTM, 0.0, nullptr, TriMeshScale);
+				EXPECT_EQ(bResult, false);
+			}
+			{
+				// Inside Mesh
+				const FRigidTransform3 QueryTM(FVec3(0.0, 0.0, 4.0), FQuat::Identity);
+				bool bResult = TriangleMesh->OverlapGeom(ScaledCapsule, QueryTM, 0.0);
+				EXPECT_EQ(bResult, false);
+			}
+		}
 		{
 			// Box test
 			const TBox<FReal, 3> Box = TBox<FReal, 3>({ -1.0, -1.0, -2.0 }, { 1.0, 1.0, 2.0 });
diff --git a/Engine/Source/Runtime/Experimental/Chaos/Private/Chaos/TriangleMeshImplicitObject.cpp b/Engine/Source/Runtime/Experimental/Chaos/Private/Chaos/TriangleMeshImplicitObject.cpp
--- a/Engine/Source/Runtime/Experimental/Chaos/Private/Chaos/TriangleMeshImplicitObject.cpp
+++ b/Engine/Source/Runtime/Experimental/Chaos/Private/Chaos/TriangleMeshImplicitObject.cpp
@@ -512,8 +512,7 @@ bool FTriangleMeshImplicitObject::ContactManifoldImp(const GeomType& QueryGeom,
 	FRigidTransform3 TriMeshToGeomNoScale{ QueryTM };
 	TriMeshToGeomNoScale.SetTranslation(TriMeshToGeomNoScale.GetTranslation() * TriMeshScale);
 	// NOTE: BVH test is done in tri-mesh local space (whereas collision detection is done in world space becaused you can't non-uniformly scale all shapes)
-	FAABB3 QueryBounds = WorldScaleGeom.BoundingBox();
-	QueryBounds = QueryBounds.TransformedAABB(TriMeshToGeomNoScale);
+	FAABB3 QueryBounds = QueryGeom.CalculateTransformedBounds(TriMeshToGeomNoScale);
 	QueryBounds.ThickenSymmetrically(FVec3(WorldThickness));
 	QueryBounds.ScaleWithNegative(InvTriMeshScale);
 
@@ -632,8 +631,7 @@ bool FTriangleMeshImplicitObject::GJKContactPointImp(const QueryGeomType& QueryG
 	FRigidTransform3 TriMeshToGeomNoScale{ QueryTM };
 	TriMeshToGeomNoScale.SetTranslation(TriMeshToGeomNoScale.GetTranslation() * TriMeshScale);
 	// NOTE: BVH test is done in tri-mesh local space (whereas collision detection is done in world space because you can't non-uniformly scale all shapes)
-	FAABB3 QueryBounds = WorldScaleGeom.BoundingBox();
-	QueryBounds = QueryBounds.TransformedAABB(TriMeshToGeomNoScale);
+	FAABB3 QueryBounds = QueryGeom.CalculateTransformedBounds(TriMeshToGeomNoScale);
 	QueryBounds.ThickenSymmetrically(FVec3(WorldThickness));
 	QueryBounds.ScaleWithNegative(InvTriMeshScale);
 
@@ -901,8 +899,7 @@ bool FTriangleMeshImplicitObject::OverlapGeomImp(const QueryGeomType& QueryGeom,
 	FRigidTransform3 TriMeshToGeomNoScale{ QueryTM };
 	TriMeshToGeomNoScale.SetTranslation(TriMeshToGeomNoScale.GetTranslation() * TriMeshScale);
 	// NOTE: BVH test is done in tri-mesh local space (whereas collision detection is done in world space because you can't non-uniformly scale all shapes)
-	FAABB3 QueryBounds = WorldScaleQueryGeom.BoundingBox();
-	QueryBounds = QueryBounds.TransformedAABB(TriMeshToGeomNoScale);
+	FAABB3 QueryBounds = QueryGeom.CalculateTransformedBounds(TriMeshToGeomNoScale);
 	QueryBounds.ThickenSymmetrically(FVec3(Thickness));
 	QueryBounds.ScaleWithNegative(InvTriMeshScale);
 
@@ -1206,7 +1203,7 @@ bool FTriangleMeshImplicitObject::SweepGeomImp(const QueryGeomType& QueryGeom, c
 		using VisitorType = FTriangleMeshSweepVisitor<QueryGeomType,decltype(Elements[0][0])>;
 		VisitorType SQVisitor(*this, Elements, QueryGeom, ScaledDirNormalized, LengthScale, ScaledStartTM, bComputeMTD, TriMeshScale, CullsBackFaceRaycastCode);
 
-		const FAABB3 QueryBounds = QueryGeom.BoundingBox().TransformedAABB(FRigidTransform3(FVec3::ZeroVector,StartTM.GetRotation()));
+		const FAABB3 QueryBounds = QueryGeom.CalculateTransformedBounds(FRigidTransform3(FVec3::ZeroVector, StartTM.GetRotation()));
 		const FVec3 InvTriMeshScale = SafeInvScale(TriMeshScale);
 		const FVec3 StartPoint = QueryBounds.Center() * InvTriMeshScale + StartTM.GetLocation();
 		const FVec3 Inflation = QueryBounds.Extents() * InvTriMeshScale.GetAbs() * 0.5 + FVec3(Thickness);
@@ -1715,4 +1712,3 @@ void FTriangleMeshImplicitObject::RebuildFastBVHFromTree(const BVHType& TreeBVH)
 }
 	
 }
-
diff --git a/Engine/Source/Runtime/Experimental/Chaos/Public/Chaos/TriangleMeshImplicitObject.h b/Engine/Source/Runtime/Experimental/Chaos/Public/Chaos/TriangleMeshImplicitObject.h
--- a/Engine/Source/Runtime/Experimental/Chaos/Public/Chaos/TriangleMeshImplicitObject.h
+++ b/Engine/Source/Runtime/Experimental/Chaos/Public/Chaos/TriangleMeshImplicitObject.h
@@ -122,9 +122,7 @@ namespace Chaos
 			const auto BoundsFilter = [&StartSimd, &InvDirSimd, &Parallel, &CurrentLength](const FAABBVectorized& Bounds) -> EFilterResult
 			{
 				const VectorRegister4Float CurDataLength = VectorLoadFloat1(&CurrentLength);
-
-				FAABBVectorized CleanBounds(Bounds.GetMin(), Bounds.GetMax());
-				const bool bHit = CleanBounds.RaycastFast(StartSimd, InvDirSimd, Parallel, CurDataLength);
+				const bool bHit = Bounds.RaycastFast(StartSimd, InvDirSimd, Parallel, CurDataLength);
 
 				return bHit ? EFilterResult::Keep : EFilterResult::Skip;
 			};
@@ -171,8 +169,7 @@ namespace Chaos
 		{
 			const auto BoundsFilter = [&AABB](const FAABBVectorized& Bounds) -> EFilterResult
 			{
-				FAABBVectorized CleanBounds(Bounds.GetMin(), Bounds.GetMax());
-				const bool bHit = CleanBounds.Intersects(AABB);
+				const bool bHit = Bounds.Intersects(AABB);
 				return bHit ? EFilterResult::Keep : EFilterResult::Skip;
 			};
 
