EvalGraph - Moved the Core and Engine modules from Experimental Plugin to Experimental Engine

diff --git a/Engine/Plugins/Experimental/EvalGraph/EvalGraph.uplugin b/Engine/Plugins/Experimental/EvalGraph/EvalGraph.uplugin
--- a/Engine/Plugins/Experimental/EvalGraph/EvalGraph.uplugin
+++ b/Engine/Plugins/Experimental/EvalGraph/EvalGraph.uplugin
@@ -10,27 +10,17 @@
     "DocsURL" : "",
     "MarketplaceURL" : "",
     "SupportURL" : "",
-    "EnabledByDefault" : false,
+    "EnabledByDefault" : true,
     "CanContainContent" : true,
     "IsBetaVersion" : true,
 	"IsExperimentalVersion": true,
-    "Installed" : false,
+    "Installed" : true,
 	"Modules": [
-		{
-			"Name": "EvalGraph",
-			"Type": "Runtime",
-			"LoadingPhase": "PreDefault"
-		},
-		{
-			"Name": "EvalGraphEngine",
-			"Type": "Runtime",
-			"LoadingPhase": "Default"
-		},
 		{
 			"Name": "EvalGraphEditor",
 			"Type": "Editor",
 			"LoadingPhase": "Default"
-		}
+		},
 	],
 	"Plugins": [
 	]
diff --git a/Engine/Plugins/Experimental/EvalGraph/Source/EvalGraph/EvalGraph.Build.cs b/Engine/Plugins/Experimental/EvalGraph/Source/EvalGraph/EvalGraph.Build.cs
deleted file mode 100644
diff --git a/Engine/Plugins/Experimental/EvalGraph/Source/EvalGraph/Private/EvalGraph/EvalGraph.cpp b/Engine/Plugins/Experimental/EvalGraph/Source/EvalGraph/Private/EvalGraph/EvalGraph.cpp
deleted file mode 100644
diff --git a/Engine/Plugins/Experimental/EvalGraph/Source/EvalGraph/Private/EvalGraph/EvalGraphInputOutput.cpp b/Engine/Plugins/Experimental/EvalGraph/Source/EvalGraph/Private/EvalGraph/EvalGraphInputOutput.cpp
deleted file mode 100644
diff --git a/Engine/Plugins/Experimental/EvalGraph/Source/EvalGraph/Private/EvalGraph/EvalGraphNode.cpp b/Engine/Plugins/Experimental/EvalGraph/Source/EvalGraph/Private/EvalGraph/EvalGraphNode.cpp
deleted file mode 100644
diff --git a/Engine/Plugins/Experimental/EvalGraph/Source/EvalGraph/Private/EvalGraph/EvalGraphNodeFactory.cpp b/Engine/Plugins/Experimental/EvalGraph/Source/EvalGraph/Private/EvalGraph/EvalGraphNodeFactory.cpp
deleted file mode 100644
diff --git a/Engine/Plugins/Experimental/EvalGraph/Source/EvalGraph/Private/EvalGraph/EvalGraphNodeParameters.cpp b/Engine/Plugins/Experimental/EvalGraph/Source/EvalGraph/Private/EvalGraph/EvalGraphNodeParameters.cpp
deleted file mode 100644
diff --git a/Engine/Plugins/Experimental/EvalGraph/Source/EvalGraph/Private/EvalGraph/EvalGraphPlugin.cpp b/Engine/Plugins/Experimental/EvalGraph/Source/EvalGraph/Private/EvalGraph/EvalGraphPlugin.cpp
deleted file mode 100644
diff --git a/Engine/Plugins/Experimental/EvalGraph/Source/EvalGraph/Private/EvalGraph/EvalGraphPluginPrivate.h b/Engine/Plugins/Experimental/EvalGraph/Source/EvalGraph/Private/EvalGraph/EvalGraphPluginPrivate.h
deleted file mode 100644
diff --git a/Engine/Plugins/Experimental/EvalGraph/Source/EvalGraph/Public/EvalGraph/EvalGraph.h b/Engine/Plugins/Experimental/EvalGraph/Source/EvalGraph/Public/EvalGraph/EvalGraph.h
deleted file mode 100644
diff --git a/Engine/Plugins/Experimental/EvalGraph/Source/EvalGraph/Public/EvalGraph/EvalGraphConnectionTypeValues.inl b/Engine/Plugins/Experimental/EvalGraph/Source/EvalGraph/Public/EvalGraph/EvalGraphConnectionTypeValues.inl
deleted file mode 100644
diff --git a/Engine/Plugins/Experimental/EvalGraph/Source/EvalGraph/Public/EvalGraph/EvalGraphConnectionTypes.h b/Engine/Plugins/Experimental/EvalGraph/Source/EvalGraph/Public/EvalGraph/EvalGraphConnectionTypes.h
deleted file mode 100644
diff --git a/Engine/Plugins/Experimental/EvalGraph/Source/EvalGraph/Public/EvalGraph/EvalGraphInputOutput.h b/Engine/Plugins/Experimental/EvalGraph/Source/EvalGraph/Public/EvalGraph/EvalGraphInputOutput.h
deleted file mode 100644
diff --git a/Engine/Plugins/Experimental/EvalGraph/Source/EvalGraph/Public/EvalGraph/EvalGraphNode.h b/Engine/Plugins/Experimental/EvalGraph/Source/EvalGraph/Public/EvalGraph/EvalGraphNode.h
deleted file mode 100644
diff --git a/Engine/Plugins/Experimental/EvalGraph/Source/EvalGraph/Public/EvalGraph/EvalGraphNodeFactory.h b/Engine/Plugins/Experimental/EvalGraph/Source/EvalGraph/Public/EvalGraph/EvalGraphNodeFactory.h
deleted file mode 100644
diff --git a/Engine/Plugins/Experimental/EvalGraph/Source/EvalGraph/Public/EvalGraph/EvalGraphNodeParameters.h b/Engine/Plugins/Experimental/EvalGraph/Source/EvalGraph/Public/EvalGraph/EvalGraphNodeParameters.h
deleted file mode 100644
diff --git a/Engine/Plugins/Experimental/EvalGraph/Source/EvalGraph/Public/EvalGraph/EvalGraphPlugin.h b/Engine/Plugins/Experimental/EvalGraph/Source/EvalGraph/Public/EvalGraph/EvalGraphPlugin.h
deleted file mode 100644
diff --git a/Engine/Plugins/Experimental/EvalGraph/Source/EvalGraphEditor/EvalGraphEditor.Build.cs b/Engine/Plugins/Experimental/EvalGraph/Source/EvalGraphEditor/EvalGraphEditor.Build.cs
--- a/Engine/Plugins/Experimental/EvalGraph/Source/EvalGraphEditor/EvalGraphEditor.Build.cs
+++ b/Engine/Plugins/Experimental/EvalGraph/Source/EvalGraphEditor/EvalGraphEditor.Build.cs
@@ -32,7 +32,7 @@ namespace UnrealBuildTool.Rules
 					"InputCore",
 					"AdvancedPreviewScene",
 					"GraphEditor",
-					"EvalGraph",
+					"EvalGraphCore",
 					"EvalGraphEngine",
 					"Slate",
 				}
diff --git a/Engine/Plugins/Experimental/EvalGraph/Source/EvalGraphEngine/EvalGraphEngine.Build.cs b/Engine/Plugins/Experimental/EvalGraph/Source/EvalGraphEngine/EvalGraphEngine.Build.cs
deleted file mode 100644
diff --git a/Engine/Plugins/Experimental/EvalGraph/Source/EvalGraphEngine/Private/EvalGraph/EvalGraphEnginePlugin.cpp b/Engine/Plugins/Experimental/EvalGraph/Source/EvalGraphEngine/Private/EvalGraph/EvalGraphEnginePlugin.cpp
deleted file mode 100644
diff --git a/Engine/Plugins/Experimental/EvalGraph/Source/EvalGraphEngine/Private/EvalGraph/EvalGraphEnginePluginPrivate.h b/Engine/Plugins/Experimental/EvalGraph/Source/EvalGraphEngine/Private/EvalGraph/EvalGraphEnginePluginPrivate.h
deleted file mode 100644
diff --git a/Engine/Plugins/Experimental/EvalGraph/Source/EvalGraphEngine/Private/EvalGraph/EvalGraphObject.cpp b/Engine/Plugins/Experimental/EvalGraph/Source/EvalGraphEngine/Private/EvalGraph/EvalGraphObject.cpp
deleted file mode 100644
diff --git a/Engine/Plugins/Experimental/EvalGraph/Source/EvalGraphEngine/Public/EvalGraph/EvalGraphEnginePlugin.h b/Engine/Plugins/Experimental/EvalGraph/Source/EvalGraphEngine/Public/EvalGraph/EvalGraphEnginePlugin.h
deleted file mode 100644
diff --git a/Engine/Plugins/Experimental/EvalGraph/Source/EvalGraphEngine/Public/EvalGraph/EvalGraphObject.h b/Engine/Plugins/Experimental/EvalGraph/Source/EvalGraphEngine/Public/EvalGraph/EvalGraphObject.h
deleted file mode 100644
diff --git a/Engine/Source/Runtime/Experimental/EvalGraph/Core/EvalGraphCore.Build.cs b/Engine/Source/Runtime/Experimental/EvalGraph/Core/EvalGraphCore.Build.cs
new file mode 100644
--- /dev/null
+++ b/Engine/Source/Runtime/Experimental/EvalGraph/Core/EvalGraphCore.Build.cs
@@ -0,0 +1,19 @@
+// Copyright Epic Games, Inc. All Rights Reserved.
+
+namespace UnrealBuildTool.Rules
+{
+	public class EvalGraphCore : ModuleRules
+	{
+        public EvalGraphCore(ReadOnlyTargetRules Target) : base(Target)
+		{
+			PublicDependencyModuleNames.AddRange(
+				new string[]
+				{
+					"Core",
+					"CoreUObject",
+					"Chaos",
+				}
+				);
+		}
+	}
+}
diff --git a/Engine/Source/Runtime/Experimental/EvalGraph/Core/Private/EvalGraph/EvalGraph.cpp b/Engine/Source/Runtime/Experimental/EvalGraph/Core/Private/EvalGraph/EvalGraph.cpp
new file mode 100644
--- /dev/null
+++ b/Engine/Source/Runtime/Experimental/EvalGraph/Core/Private/EvalGraph/EvalGraph.cpp
@@ -0,0 +1,165 @@
+// Copyright Epic Games, Inc. All Rights Reserved.
+
+#include "EvalGraph/EvalGraph.h"
+
+#include "EvalGraph/EvalGraphInputOutput.h"
+#include "EvalGraph/EvalGraphConnectionTypes.h"
+#include "EvalGraph/EvalGraphNodeParameters.h"
+#include "EvalGraph/EvalGraphNodeFactory.h"
+#include "Logging/LogMacros.h"
+
+DEFINE_LOG_CATEGORY_STATIC(EGGRAPH_LOG, Error, All);
+
+namespace Eg
+{
+
+	FGraph::FGraph(FGuid InGuid)
+		: Guid(InGuid)
+	{}
+
+
+	void FGraph::RemoveNode(TSharedPtr<FNode> Node)
+	{
+		for (FConnectionTypeBase* Output : Node->GetOutputs())
+		{
+			if (Output)
+			{
+				for (FConnectionTypeBase* Input : Output->GetBaseInputs())
+				{
+					if (Input)
+					{
+						Disconnect(Output, Input);
+					}
+				}
+			}
+		}
+		for (FConnectionTypeBase* Input : Node->GetInputs())
+		{
+			if (Input)
+			{
+				TArray<FConnectionTypeBase*> Outputs = Input->GetBaseOutputs();
+				for (FConnectionTypeBase* Out : Outputs)
+				{
+					if (Out)
+					{
+						Disconnect(Out, Input);
+					}
+				}
+			}
+		}
+		Nodes.Remove(Node);
+	}
+
+	void FGraph::Connect(FConnectionTypeBase* Input, FConnectionTypeBase* Output)
+	{
+		if (ensure(Input && Output))
+		{
+			Input->AddConnection(Output);
+			Output->AddConnection(Input);
+			Connections.Add(FConnection(Input->GetGuid(), Output->GetGuid()));
+		}
+	}
+
+	void FGraph::Disconnect(FConnectionTypeBase* Input, FConnectionTypeBase* Output)
+	{
+		Input->RemoveConnection(Output);
+		Output->RemoveConnection(Input);
+		Connections.RemoveSwap(FConnection(Input->GetGuid(), Output->GetGuid()));
+	}
+
+	void FGraph::Serialize(FArchive& Ar)
+	{
+		Ar << Guid;
+		if (Ar.IsSaving())
+		{
+			FGuid ArGuid;
+			FName ArType, ArName;
+			uint32 ArIntType;
+			int32 ArNum = Nodes.Num();
+			Ar << ArNum;
+			for (TSharedPtr<FNode> Node : Nodes)
+			{
+				ArGuid = Node->GetGuid();
+				ArType = Node->GetType();
+				ArName = Node->GetName();
+				Ar << ArGuid << ArType << ArName;
+
+
+				TArray< FConnectionTypeBase* > IO = Node->GetOutputs();  IO.Append(Node->GetInputs());
+				ArNum = IO.Num();
+				Ar << ArNum;
+				for (FConnectionTypeBase* Conn : IO)
+				{
+					ArGuid = Conn->GetGuid();
+					ArIntType = (int32)Conn->GetType();
+					ArName = Conn->GetName();
+					Ar << ArGuid << ArIntType << ArName;
+				}
+
+				Node->SerializeInternal(Ar);
+			}
+
+			Ar << Connections;
+
+		}
+		else if( Ar.IsLoading())
+		{
+			FGuid ArGuid;
+			FName ArType, ArName;
+			uint32 ArIntType;
+			int32 ArNum = 0;
+
+			TMap<FGuid, FConnectionTypeBase* > ConnectionGuidMap;
+
+			Ar << ArNum;
+			for (int32 Ndx = ArNum; Ndx > 0; Ndx--)
+			{
+				Ar << ArGuid << ArType << ArName;
+				if (TSharedPtr<FNode> Node = FNodeFactory::GetInstance()->NewNodeFromRegisteredType(*this, { ArGuid,ArType,ArName }))
+				{
+					int ArNumInner;
+					Ar << ArNumInner;
+					TArray< FConnectionTypeBase* > IO = Node->GetOutputs();  IO.Append(Node->GetInputs());
+					for (int Cdx=0;Cdx< ArNumInner; Cdx++)
+					{
+						Ar << ArGuid << ArIntType << ArName;
+						if (Cdx < IO.Num())
+						{
+							if ((EGraphConnectionType)ArIntType!=IO[Cdx]->GetType())
+							{
+								FName TypeName = GraphConnectionTypeName((EGraphConnectionType)ArIntType);
+								FName FromType = GraphConnectionTypeName(IO[Cdx]->GetType());
+								UE_LOG(EGGRAPH_LOG, Fatal, TEXT("Type mismatch in input file for Ed::FGraphNode Input/Output (%s,%s) to type (%s)"),
+									*TypeName.ToString(), *ArName.ToString(), *FromType.ToString());
+							}
+
+							IO[Cdx]->SetGuid(ArGuid);
+							ensure(!ConnectionGuidMap.Contains(ArGuid));
+							ConnectionGuidMap.Add(ArGuid, IO[Cdx]);
+						}
+					}
+
+					Node->SerializeInternal(Ar);
+				}
+				else
+				{
+					ensureMsgf(false, TEXT("Warning: Missing registered node type on load. Graph will fail to evaluate (%s %s)"), *ArType.ToString(), *ArName.ToString());
+				}
+				
+			}
+
+			TArray< FConnection > LocalConnections;
+			Ar << LocalConnections;
+			for (const FConnection& Con : LocalConnections)
+			{
+				ensure(ConnectionGuidMap.Contains(Con.Key));
+				ensure(ConnectionGuidMap.Contains(Con.Value));
+				ensure(ConnectionGuidMap[Con.Key]->GetType() == ConnectionGuidMap[Con.Value]->GetType());
+				Connect(ConnectionGuidMap[Con.Key], ConnectionGuidMap[Con.Value]);
+			}
+
+		}
+	}
+
+}
+
diff --git a/Engine/Source/Runtime/Experimental/EvalGraph/Core/Private/EvalGraph/EvalGraphCoreLogPrivate.h b/Engine/Source/Runtime/Experimental/EvalGraph/Core/Private/EvalGraph/EvalGraphCoreLogPrivate.h
new file mode 100644
--- /dev/null
+++ b/Engine/Source/Runtime/Experimental/EvalGraph/Core/Private/EvalGraph/EvalGraphCoreLogPrivate.h
@@ -0,0 +1,10 @@
+// Copyright Epic Games, Inc. All Rights Reserved.
+
+#include "CoreMinimal.h"
+#include "Stats/Stats.h"
+
+// You should place include statements to your module's private header files here.  You only need to
+// add includes for headers that are used in most of your module's source files though.
+
+
+DECLARE_STATS_GROUP(TEXT("EvalGraph"), STATGROUP_EvalGraph, STATCAT_Advanced);
diff --git a/Engine/Source/Runtime/Experimental/EvalGraph/Core/Private/EvalGraph/EvalGraphCoreModule.cpp b/Engine/Source/Runtime/Experimental/EvalGraph/Core/Private/EvalGraph/EvalGraphCoreModule.cpp
new file mode 100644
--- /dev/null
+++ b/Engine/Source/Runtime/Experimental/EvalGraph/Core/Private/EvalGraph/EvalGraphCoreModule.cpp
@@ -0,0 +1,31 @@
+// Copyright Epic Games, Inc. All Rights Reserved.
+
+#pragma once
+
+#include "CoreMinimal.h"
+#include "EvalGraphCoreLogPrivate.h"
+#include "Modules/ModuleInterface.h"
+#include "Modules/ModuleManager.h"
+
+
+/**
+ * The public interface to this module
+ */
+class FEvalGraphCoreModule : public IModuleInterface
+{
+
+public:
+
+	// IModuleInterface interface
+
+	virtual void StartupModule() override { }
+	virtual void ShutdownModule() override { }
+
+	virtual bool SupportsDynamicReloading() override
+	{
+		return false;
+	}
+
+};
+
+IMPLEMENT_MODULE( FEvalGraphCoreModule, EvalGraphCore )
diff --git a/Engine/Source/Runtime/Experimental/EvalGraph/Core/Private/EvalGraph/EvalGraphInputOutput.cpp b/Engine/Source/Runtime/Experimental/EvalGraph/Core/Private/EvalGraph/EvalGraphInputOutput.cpp
new file mode 100644
--- /dev/null
+++ b/Engine/Source/Runtime/Experimental/EvalGraph/Core/Private/EvalGraph/EvalGraphInputOutput.cpp
@@ -0,0 +1,100 @@
+// Copyright Epic Games, Inc. All Rights Reserved.
+
+#include "EvalGraph/EvalGraphInputOutput.h"
+
+#include "EvalGraph/EvalGraphConnectionTypes.h"
+#include "EvalGraph/EvalGraphNodeParameters.h"
+#include "EvalGraph/EvalGraphNode.h"
+
+namespace Eg
+{
+
+	FConnectionTypeBase::FConnectionTypeBase(EGraphConnectionType InType, FName InName, FNode* InOwningNode, FGuid InGuid)
+		: Type(InType)
+		, Name(InName)
+		, Guid(InGuid)
+		, OwningNode(InOwningNode)
+	{}
+
+	void FConnectionTypeBase::AddBaseInput(FNode* InNode, FConnectionTypeBase* That) 
+	{ InNode->AddBaseInput(That); }
+	void FConnectionTypeBase::AddBaseOutput(FNode* InNode, FConnectionTypeBase* That) 
+	{ InNode->AddBaseOutput(That); }
+
+
+	template<class T>
+	FInput<T>::FInput(const FInputParameters<T>& Param, FGuid InGuid )
+		: FConnectionTypeBase(GraphConnectionType<T>(), Param.Name, Param.Owner, InGuid)
+		, Default(Param.Default)
+		, Connection(nullptr)
+	{
+		Super::AddBaseInput(Param.Owner, this);
+	}
+
+	template<class T>
+	void FInput<T>::SetValue(const T& Value, const FContext& Context)
+	{
+		Default = Value;
+		if (!Connection)
+		{
+			OwningNode->InvalidateOutputs();
+		}
+	}
+
+	template<class T>
+	void FInput<T>::Invalidate()
+	{
+		OwningNode->InvalidateOutputs();
+	}
+
+	template<class T>
+	FOutput<T>::FOutput(const FOutputParameters& Param, FGuid InGuid )
+		: FConnectionTypeBase(GraphConnectionType<T>(), Param.Name, Param.Owner, InGuid)
+	{
+		Super::AddBaseOutput(Param.Owner, this);
+	}
+
+
+	template<class T>
+	void FOutput<T>::SetValue(T InVal, const FContext& Context)
+	{
+		CacheKey = Context.GetTypeHash();
+		Cache.Data = InVal;
+	}
+
+	template<class T>
+	T FOutput<T>::Evaluate(const FContext& Context)
+	{
+		if (CacheKey!=Context.GetTypeHash())
+		{
+			OwningNode->Evaluate(Context, this);
+		}
+		ensure(CacheKey==Context.GetTypeHash());
+		return Cache.Data;
+	}
+
+	template<class T>
+	void FOutput<T>::Invalidate( )
+	{
+		if (CacheKey != UINT_MAX)
+		{
+			CacheKey = UINT_MAX;
+			Cache = T();
+
+			for (FConnectionTypeBase* Con : GetConnections())
+			{
+				Con->Invalidate();
+			}
+		}
+	}
+
+
+#define EVAL_GRAPH_CONNECTION_TYPE(a,A) template class FOutput<a>;
+#include "../../Public/EvalGraph/EvalGraphConnectionTypeValues.inl"
+#undef EVAL_GRAPH_CONNECTION_TYPE
+
+#define EVAL_GRAPH_CONNECTION_TYPE(a,A) template class FInput<a>;
+#include "../../Public/EvalGraph/EvalGraphConnectionTypeValues.inl"
+#undef EVAL_GRAPH_CONNECTION_TYPE
+}
+
diff --git a/Engine/Source/Runtime/Experimental/EvalGraph/Core/Private/EvalGraph/EvalGraphNode.cpp b/Engine/Source/Runtime/Experimental/EvalGraph/Core/Private/EvalGraph/EvalGraphNode.cpp
new file mode 100644
--- /dev/null
+++ b/Engine/Source/Runtime/Experimental/EvalGraph/Core/Private/EvalGraph/EvalGraphNode.cpp
@@ -0,0 +1,18 @@
+// Copyright Epic Games, Inc. All Rights Reserved.
+
+#include "EvalGraph/EvalGraphNode.h"
+
+#include "EvalGraph/EvalGraphNode.h"
+
+namespace Eg
+{
+	void FNode::InvalidateOutputs()
+	{
+		for(FConnectionTypeBase * Output :  Outputs)
+		{
+			Output->Invalidate();
+		}
+	}
+
+}
+
diff --git a/Engine/Source/Runtime/Experimental/EvalGraph/Core/Private/EvalGraph/EvalGraphNodeFactory.cpp b/Engine/Source/Runtime/Experimental/EvalGraph/Core/Private/EvalGraph/EvalGraphNodeFactory.cpp
new file mode 100644
--- /dev/null
+++ b/Engine/Source/Runtime/Experimental/EvalGraph/Core/Private/EvalGraph/EvalGraphNodeFactory.cpp
@@ -0,0 +1,21 @@
+// Copyright Epic Games, Inc. All Rights Reserved.
+
+#include "EvalGraph/EvalGraphNodeFactory.h"
+
+#include "EvalGraph/EvalGraphConnectionTypes.h"
+#include "EvalGraph/EvalGraphNode.h"
+
+namespace Eg
+{
+	FNodeFactory* FNodeFactory::Instance = nullptr;
+
+	TSharedPtr<FNode> FNodeFactory::NewNodeFromRegisteredType(FGraph& Graph, const FNewNodeParameters& Param)
+	{ 
+		if (ClassMap.Contains(Param.Type))
+		{
+			return Graph.AddNode(ClassMap[Param.Type](Param));
+		}
+		return TSharedPtr<FNode>(nullptr);
+	}
+}
+
diff --git a/Engine/Source/Runtime/Experimental/EvalGraph/Core/Private/EvalGraph/EvalGraphNodeParameters.cpp b/Engine/Source/Runtime/Experimental/EvalGraph/Core/Private/EvalGraph/EvalGraphNodeParameters.cpp
new file mode 100644
--- /dev/null
+++ b/Engine/Source/Runtime/Experimental/EvalGraph/Core/Private/EvalGraph/EvalGraphNodeParameters.cpp
@@ -0,0 +1,4 @@
+// Copyright Epic Games, Inc. All Rights Reserved.
+
+#include "EvalGraph/EvalGraphNodeParameters.h"
+
diff --git a/Engine/Source/Runtime/Experimental/EvalGraph/Core/Public/EvalGraph/EvalGraph.h b/Engine/Source/Runtime/Experimental/EvalGraph/Core/Public/EvalGraph/EvalGraph.h
new file mode 100644
--- /dev/null
+++ b/Engine/Source/Runtime/Experimental/EvalGraph/Core/Public/EvalGraph/EvalGraph.h
@@ -0,0 +1,80 @@
+// Copyright Epic Games, Inc. All Rights Reserved.
+
+#pragma once
+
+#include "CoreMinimal.h"
+#include "Chaos/ChaosArchive.h"
+#include "Serialization/Archive.h"
+#include "EvalGraph/EvalGraphNode.h"
+
+namespace Eg
+{
+	class FConnectionTypeBase;
+
+	typedef TTuple<FGuid, FGuid> FConnection; // <Input, Output>
+
+
+	//
+	//
+	//
+	class EVALGRAPHCORE_API FGraph
+	{
+
+		FGuid  Guid;
+		TArray< TSharedPtr<FNode> > Nodes;
+		TArray< FConnection > Connections;
+		//TMap<FGuid, TSharedPtr<FConnectionTypeBase> > ConnectionMap;
+
+	public:
+		FGraph(FGuid InGuid = FGuid::NewGuid());
+		virtual ~FGraph() {}
+
+		int NumNodes() { return Nodes.Num(); }
+
+		template<class T> TSharedPtr<T> AddNode(T* InNode)
+		{
+			TSharedPtr<T> NewNode(InNode);
+			Nodes.AddUnique(NewNode);
+			return NewNode;
+		}
+
+		template<class T> T* FindNode(FName InName)
+		{
+			for (TSharedPtr<FNode> Node : Nodes)  
+			{
+				if (Node->GetName().ToString().Equals(InName.ToString())) 
+				{
+					if (Node->GetType().ToString().Equals(T::Type.ToString())) 
+					{
+						return reinterpret_cast<T*>(Node.Get()); // @todo(eg) : Can we do better here?
+					}
+				}
+			}
+			return nullptr;
+		}
+
+
+		void RemoveNode(TSharedPtr<FNode> Node);
+
+		void Connect(FConnectionTypeBase* Input, FConnectionTypeBase* Output);
+		void Disconnect(FConnectionTypeBase* Input, FConnectionTypeBase* Output);
+
+		virtual void Serialize(FArchive& Ar);
+
+	};
+
+}
+
+FORCEINLINE FArchive& operator<<(FArchive& Ar, Eg::FGraph& Value)
+{
+	Value.Serialize(Ar);
+	return Ar;
+}
+
+FORCEINLINE FArchive& operator<<(Chaos::FChaosArchive& Ar, Eg::FGraph& Value)
+{
+	Value.Serialize(Ar);
+	return Ar;
+}
+
+
diff --git a/Engine/Source/Runtime/Experimental/EvalGraph/Core/Public/EvalGraph/EvalGraphConnectionTypeValues.inl b/Engine/Source/Runtime/Experimental/EvalGraph/Core/Public/EvalGraph/EvalGraphConnectionTypeValues.inl
new file mode 100644
--- /dev/null
+++ b/Engine/Source/Runtime/Experimental/EvalGraph/Core/Public/EvalGraph/EvalGraphConnectionTypeValues.inl
@@ -0,0 +1,109 @@
+// Copyright Epic Games, Inc. All Rights Reserved.
+
+// usage
+//
+// General purpose Scene Graph Type definitions
+
+#ifndef EVAL_GRAPH_CONNECTION_TYPE
+#error EVAL_GRAPH_CONNECTION_TYPE macro is undefined.
+#endif
+
+// NOTE: new types must be added at the bottom to keep serialization from breaking
+
+//
+EVAL_GRAPH_CONNECTION_TYPE(bool, Bool)
+EVAL_GRAPH_CONNECTION_TYPE(char, Char)
+EVAL_GRAPH_CONNECTION_TYPE(int, Integer)
+EVAL_GRAPH_CONNECTION_TYPE(uint8, UInt8)
+EVAL_GRAPH_CONNECTION_TYPE(float, Float)
+EVAL_GRAPH_CONNECTION_TYPE(double, Double)
+//
+EVAL_GRAPH_CONNECTION_TYPE(FManagedArrayCollection*,FManagedArrayCollectionPtr)
+EVAL_GRAPH_CONNECTION_TYPE(TManagedArrayBase<double>*,FManagedDoubleArrayPtr)
+EVAL_GRAPH_CONNECTION_TYPE(TManagedArrayBase<float>*,FManagedFloatArrayPtr)
+EVAL_GRAPH_CONNECTION_TYPE(TManagedArrayBase<int32>*,FManagedInt32ArrayPtr)
+EVAL_GRAPH_CONNECTION_TYPE(TManagedArrayBase<bool>*,FManagedBoolArrayPtr)
+EVAL_GRAPH_CONNECTION_TYPE(TManagedArrayBase<FVector3d>*,FManagedVector3dArrayPtr)
+EVAL_GRAPH_CONNECTION_TYPE(TManagedArrayBase<FVector3f>*,FManagedVector3fArrayPtr)
+EVAL_GRAPH_CONNECTION_TYPE(TManagedArrayBase<FTransform>*,FManagedTransformArrayPtr)
+EVAL_GRAPH_CONNECTION_TYPE(TManagedArrayBase<FIntVector>*,FManagedIntVectorArrayPtr)
+EVAL_GRAPH_CONNECTION_TYPE(TManagedArrayBase<FQuat4f>*,FManagedQuat4fArrayPtr)
+EVAL_GRAPH_CONNECTION_TYPE(TManagedArrayBase<FVector2f>*,FManagedVector2fArrayPtr)
+EVAL_GRAPH_CONNECTION_TYPE(TManagedArrayBase<FLinearColor>*,FManagedLinearColorArrayPtr)
+EVAL_GRAPH_CONNECTION_TYPE(TManagedArrayBase<FString>*,FManagedStringArrayPtr)
+EVAL_GRAPH_CONNECTION_TYPE(TManagedArrayBase<FBox>*,FManagedBoxArrayPtr)
+EVAL_GRAPH_CONNECTION_TYPE(TManagedArrayBase<TSet<int32>>*,FManagedIntSetArrayPtr)
+EVAL_GRAPH_CONNECTION_TYPE(TManagedArrayBase<FIntVector4>*,FManagedIntVector4ArrayPtr)
+EVAL_GRAPH_CONNECTION_TYPE(TManagedArrayBase<TArray<FVector2f>>*,FManagedVector2fAArrayPtr)
+EVAL_GRAPH_CONNECTION_TYPE(TManagedArrayBase<TArray<FVector3f>*>*,FManagedVector3fAPtrArrayPtr)
+//
+EVAL_GRAPH_CONNECTION_TYPE(FManagedArrayCollectionSharedSafePtr, ManagedArrayCollectionSharedSafePtr)
+EVAL_GRAPH_CONNECTION_TYPE(FManagedDoubleArraySharedSafePtr, ManagedDoubleArraySharedSafePtr)
+EVAL_GRAPH_CONNECTION_TYPE(FManagedFloatArraySharedSafePtr, ManagedFloatArraySharedSafePtr)
+EVAL_GRAPH_CONNECTION_TYPE(FManagedInt32ArraySharedSafePtr, ManagedInt32ArraySharedSafePtr)
+EVAL_GRAPH_CONNECTION_TYPE(FManagedBoolArraySharedSafePtr, ManagedBoolArraySharedSafePtr)
+EVAL_GRAPH_CONNECTION_TYPE(FManagedVector3dArraySharedSafePtr, ManagedVector3dArraySharedSafePtr)
+EVAL_GRAPH_CONNECTION_TYPE(FManagedVector3fArraySharedSafePtr, ManagedVector3fArraySharedSafePtr)
+EVAL_GRAPH_CONNECTION_TYPE(FManagedTransformArraySharedSafePtr, ManagedTransformArraySharedSafePtr)
+EVAL_GRAPH_CONNECTION_TYPE(FManagedIntVectorArraySharedSafePtr, ManagedIntVectorArraySharedSafePtr)
+EVAL_GRAPH_CONNECTION_TYPE(FManagedQuat4fArraySharedSafePtr, ManagedQuat4fArraySharedSafePtr)
+EVAL_GRAPH_CONNECTION_TYPE(FManagedVector2fArraySharedSafePtr, ManagedVector2fArraySharedSafePtr)
+EVAL_GRAPH_CONNECTION_TYPE(FManagedLinearColorArraySharedSafePtr, ManagedLinearColorArraySharedSafePtr)
+EVAL_GRAPH_CONNECTION_TYPE(FManagedStringArraySharedSafePtr, ManagedStringArraySharedSafePtr)
+EVAL_GRAPH_CONNECTION_TYPE(FManagedBoxArraySharedSafePtr, ManagedBoxArraySharedSafePtr)
+EVAL_GRAPH_CONNECTION_TYPE(FManagedIntSetArraySharedSafePtr, ManagedIntSetArraySharedSafePtr)
+EVAL_GRAPH_CONNECTION_TYPE(FManagedIntVector4ArraySharedSafePtr, ManagedIntVector4ArraySharedSafePtr)
+EVAL_GRAPH_CONNECTION_TYPE(FManagedVector2fAArraySharedSafePtr, ManagedVector2fAArraySharedSafePtr)
+EVAL_GRAPH_CONNECTION_TYPE(FManagedVector3fAPtrArraySharedSafePtr, ManagedVector3fAPtrArraySharedSafePtr)
+//
+EVAL_GRAPH_CONNECTION_TYPE(FString, String)
+EVAL_GRAPH_CONNECTION_TYPE(FBox, Box)
+EVAL_GRAPH_CONNECTION_TYPE(FColor, Color)
+EVAL_GRAPH_CONNECTION_TYPE(FIntVector, IntVector)
+EVAL_GRAPH_CONNECTION_TYPE(FIntVector4, IntVector4)
+EVAL_GRAPH_CONNECTION_TYPE(FLinearColor, LinearColor)
+EVAL_GRAPH_CONNECTION_TYPE(FMatrix, Matrix)
+EVAL_GRAPH_CONNECTION_TYPE(FOrientedBox, OrientedBox)
+EVAL_GRAPH_CONNECTION_TYPE(FPlane, Plane)
+EVAL_GRAPH_CONNECTION_TYPE(FQuat, Quat)
+EVAL_GRAPH_CONNECTION_TYPE(FQuat4f, Quat4f)
+EVAL_GRAPH_CONNECTION_TYPE(FSphere, Sphere)
+EVAL_GRAPH_CONNECTION_TYPE(FTransform, Transform)
+EVAL_GRAPH_CONNECTION_TYPE(FVector, Vector)
+EVAL_GRAPH_CONNECTION_TYPE(FVector2D, Vector2D)
+EVAL_GRAPH_CONNECTION_TYPE(FVector3f, Vector3f)
+EVAL_GRAPH_CONNECTION_TYPE(FVector4, Vector4)
+//
+EVAL_GRAPH_CONNECTION_TYPE(TSet<int32>, IntSet)
+EVAL_GRAPH_CONNECTION_TYPE(TArray<bool>, BoolArray)
+EVAL_GRAPH_CONNECTION_TYPE(TArray<int>, IntArray)
+EVAL_GRAPH_CONNECTION_TYPE(TArray<double>, DoubleArray)
+EVAL_GRAPH_CONNECTION_TYPE(TArray<FBox>, BoxArray)
+EVAL_GRAPH_CONNECTION_TYPE(TArray<FMatrix>, MatrixArray)
+EVAL_GRAPH_CONNECTION_TYPE(TArray<FVector>, VectorArray)
+EVAL_GRAPH_CONNECTION_TYPE(TArray<FVector2f>, Vector2fArray)
+EVAL_GRAPH_CONNECTION_TYPE(TArray<FVector3f>, Vector3fArray)
+EVAL_GRAPH_CONNECTION_TYPE(TArray<FQuat>, QuatArray)
+EVAL_GRAPH_CONNECTION_TYPE(TArray<FQuat4f>, Quat4fArray)
+//
+EVAL_GRAPH_CONNECTION_TYPE(TSet<int32>*, IntSetPtr)
+EVAL_GRAPH_CONNECTION_TYPE(TArray<bool>*, BoolArrayPtr)
+EVAL_GRAPH_CONNECTION_TYPE(TArray<int>*, IntArrayPtr)
+EVAL_GRAPH_CONNECTION_TYPE(TArray<double>*, DoubleArrayPtr)
+EVAL_GRAPH_CONNECTION_TYPE(TArray<FBox>*, BoxArrayPtr)
+EVAL_GRAPH_CONNECTION_TYPE(TArray<FMatrix>*, MatrixArrayPtr)
+EVAL_GRAPH_CONNECTION_TYPE(TArray<FVector>*, VectorArrayPtr)
+EVAL_GRAPH_CONNECTION_TYPE(TArray<FVector2f>*, Vector2fArrayPtr)
+EVAL_GRAPH_CONNECTION_TYPE(TArray<FVector3f>*, Vector3fArrayPtr)
+EVAL_GRAPH_CONNECTION_TYPE(TArray<FQuat>*, QuatArrayPtr)
+EVAL_GRAPH_CONNECTION_TYPE(TArray<FQuat4f>*, Quat4fArrayPtr)
+//
+EVAL_GRAPH_CONNECTION_TYPE(Chaos::FImplicitObject3*, ImplicitObject3Ptr)
+EVAL_GRAPH_CONNECTION_TYPE(Chaos::FBVHParticlesFloat3*, BVHParticlesFloat3Ptr)
+EVAL_GRAPH_CONNECTION_TYPE(FImplicitObject3SerializablePtr, FImplicitObject3SerializablePtr)
+EVAL_GRAPH_CONNECTION_TYPE(FImplicitObjectSharedSafePtr, FImplicitObjectSharedSafePtr)
+
+// NOTE: new types must be added at the bottom to keep serialization from breaking
+
+
+#undef EVAL_GRAPH_CONNECTION_TYPE
diff --git a/Engine/Source/Runtime/Experimental/EvalGraph/Core/Public/EvalGraph/EvalGraphConnectionTypes.h b/Engine/Source/Runtime/Experimental/EvalGraph/Core/Public/EvalGraph/EvalGraphConnectionTypes.h
new file mode 100644
--- /dev/null
+++ b/Engine/Source/Runtime/Experimental/EvalGraph/Core/Public/EvalGraph/EvalGraphConnectionTypes.h
@@ -0,0 +1,127 @@
+// Copyright Epic Games, Inc. All Rights Reserved.
+#pragma once
+#include "GeometryCollection/ManagedArrayCollection.h"
+
+#include "Math/Box.h"
+#include "Math/Box2D.h"
+#include "Math/Color.h"
+#include "Math/IntVector.h"
+#include "Math/Matrix.h"
+#include "Math/TransformCalculus2D.h"
+#include "Math/OrientedBox.h"
+#include "Math/OrthoMatrix.h"
+#include "Math/Plane.h"
+#include "Math/Quat.h"
+#include "Math/TransformCalculus2D.h"
+#include "Math/QuatRotationTranslationMatrix.h"
+#include "Math/Sphere.h"
+#include "Math/TransformNonVectorized.h"
+#include "Math/TransformCalculus2D.h"
+#include "Math/TranslationMatrix.h"
+#include "Math/Vector.h"
+#include "Math/Vector2D.h"
+#include "Math/Vector4.h"
+#include "Math/Matrix.h"
+#include "Math/RotationMatrix.h"
+
+#include "Chaos/ImplicitObject.h"
+#include "Chaos/BVHParticles.h"
+#include "Chaos/ParticleHandle.h"
+#include "Chaos/ParticleHandleFwd.h"
+#include "Chaos/Convex.h"
+
+
+namespace Eg
+{
+	typedef TSharedPtr<FManagedArrayCollection, ESPMode::ThreadSafe> FManagedArrayCollectionSharedSafePtr;
+	typedef TSharedPtr<TManagedArrayBase<double>, ESPMode::ThreadSafe> FManagedDoubleArraySharedSafePtr;
+	typedef TSharedPtr<TManagedArrayBase<float>, ESPMode::ThreadSafe> FManagedFloatArraySharedSafePtr;
+	typedef TSharedPtr<TManagedArrayBase<int32>, ESPMode::ThreadSafe> FManagedInt32ArraySharedSafePtr;
+	typedef TSharedPtr<TManagedArrayBase<bool>, ESPMode::ThreadSafe> FManagedBoolArraySharedSafePtr;
+	typedef TSharedPtr<TManagedArrayBase<FVector3d>, ESPMode::ThreadSafe> FManagedVector3dArraySharedSafePtr;
+	typedef TSharedPtr<TManagedArrayBase<FVector3f>, ESPMode::ThreadSafe> FManagedVector3fArraySharedSafePtr;
+	typedef TSharedPtr<TManagedArrayBase<FTransform>, ESPMode::ThreadSafe> FManagedTransformArraySharedSafePtr;
+	typedef TSharedPtr<TManagedArrayBase<FIntVector>, ESPMode::ThreadSafe> FManagedIntVectorArraySharedSafePtr;
+	typedef TSharedPtr<TManagedArrayBase<FQuat4f>, ESPMode::ThreadSafe> FManagedQuat4fArraySharedSafePtr;
+	typedef TSharedPtr<TManagedArrayBase<FVector2f>, ESPMode::ThreadSafe> FManagedVector2fArraySharedSafePtr;
+	typedef TSharedPtr<TManagedArrayBase<FLinearColor>, ESPMode::ThreadSafe> FManagedLinearColorArraySharedSafePtr;
+	typedef TSharedPtr<TManagedArrayBase<FString>, ESPMode::ThreadSafe> FManagedStringArraySharedSafePtr;
+	typedef TSharedPtr<TManagedArrayBase<FBox>, ESPMode::ThreadSafe> FManagedBoxArraySharedSafePtr;
+	typedef TSharedPtr<TManagedArrayBase<TSet<int32>>, ESPMode::ThreadSafe> FManagedIntSetArraySharedSafePtr;
+	typedef TSharedPtr<TManagedArrayBase<FIntVector4>, ESPMode::ThreadSafe> FManagedIntVector4ArraySharedSafePtr;
+	typedef TSharedPtr<TManagedArrayBase<TArray<FVector2f>>, ESPMode::ThreadSafe> FManagedVector2fAArraySharedSafePtr;
+	typedef TSharedPtr<TManagedArrayBase<TArray<FVector3f>*>, ESPMode::ThreadSafe> FManagedVector3fAPtrArraySharedSafePtr;
+
+
+	typedef Chaos::TSerializablePtr<Chaos::FImplicitObject3> FImplicitObject3SerializablePtr;
+	typedef TSharedPtr<Chaos::FImplicitObject, ESPMode::ThreadSafe> FImplicitObjectSharedSafePtr;
+
+	// ---------------------------------------------------------
+	//
+	// General purpose EManagedArrayType definition. 
+	// This defines things like 
+	//     EGraphConnectionType::FManagedArrayCollectionType
+	// see EvalGraphConnectionTypeValues.inl for specific types.
+	//
+#define EVAL_GRAPH_CONNECTION_TYPE(a,A) F##A##Type,
+	enum class EGraphConnectionType : uint32
+	{
+		FNoneType,
+#include "EvalGraphConnectionTypeValues.inl"
+	};
+#undef EVAL_GRAPH_CONNECTION_TYPE
+
+	// ---------------------------------------------------------
+	//  GraphConnectionType<T>
+	//    Templated function to return a EGraphConnectionType.
+	//
+	template<class T> inline EGraphConnectionType GraphConnectionType();
+#define EVAL_GRAPH_CONNECTION_TYPE(a,A) template<> inline EGraphConnectionType GraphConnectionType<a>() { return EGraphConnectionType::F##A##Type; }
+#include "EvalGraphConnectionTypeValues.inl"
+#undef EVAL_GRAPH_CONNECTION_TYPE
+
+	// ---------------------------------------------------------
+	//  GraphConnectionTypeName
+	//    Templated function to return a EGraphConnectionType.
+	//
+	template<class T> inline FName GraphConnectionTypeName();
+#define EVAL_GRAPH_CONNECTION_TYPE(a,A) template<> inline FName GraphConnectionTypeName<a>() { return FName(TEXT(#A)); }
+#include "EvalGraphConnectionTypeValues.inl"
+#undef EVAL_GRAPH_CONNECTION_TYPE
+
+// ---------------------------------------------------------
+//  GraphConnectionTypeName
+//    Templated function to return a EGraphConnectionType as a FName
+//
+	inline FName GraphConnectionTypeName(EGraphConnectionType ValueType)
+	{
+		switch (ValueType)
+		{
+#define EVAL_GRAPH_CONNECTION_TYPE(a,A)	case EGraphConnectionType::F##A##Type:\
+		return FName(TEXT(#A));
+#include "EvalGraphConnectionTypeValues.inl"
+#undef EVAL_GRAPH_CONNECTION_TYPE
+		}
+		return FName("FNoneType");
+	}
+
+	// ---------------------------------------------------------
+	//  void*
+	//     Returns a new EGraphConnectionType pointer based on 
+	//     passed type.
+	//
+	inline void* NewGraphValueType(EGraphConnectionType ValueType)
+	{
+		switch (ValueType)
+		{
+#define EVAL_GRAPH_CONNECTION_TYPE(a,A)	case EGraphConnectionType::F##A##Type:\
+		return new a();
+#include "EvalGraphConnectionTypeValues.inl"
+#undef EVAL_GRAPH_CONNECTION_TYPE
+		}
+		check(false);
+		return nullptr;
+	}
+
+}
+
diff --git a/Engine/Source/Runtime/Experimental/EvalGraph/Core/Public/EvalGraph/EvalGraphInputOutput.h b/Engine/Source/Runtime/Experimental/EvalGraph/Core/Public/EvalGraph/EvalGraphInputOutput.h
new file mode 100644
--- /dev/null
+++ b/Engine/Source/Runtime/Experimental/EvalGraph/Core/Public/EvalGraph/EvalGraphInputOutput.h
@@ -0,0 +1,189 @@
+// Copyright Epic Games, Inc. All Rights Reserved.
+
+#pragma once
+
+#include "CoreMinimal.h"
+#include "EvalGraph/EvalGraphNodeParameters.h"
+
+namespace Eg
+{
+	class FNode;
+
+	//
+	// Input Output Base
+	//
+
+	class EVALGRAPHCORE_API FConnectionTypeBase
+	{
+
+	protected:
+		EGraphConnectionType Type;
+		FName Name;
+		FGuid  Guid;
+		FNode* OwningNode = nullptr;
+
+		friend class FNode;
+		static void AddBaseInput(FNode* InNode, FConnectionTypeBase*);
+		static void AddBaseOutput(FNode* InNode, FConnectionTypeBase*);
+
+
+	public:
+		FConnectionTypeBase(EGraphConnectionType InType, FName InName, FNode* OwningNode = nullptr, FGuid InGuid = FGuid::NewGuid());
+		virtual ~FConnectionTypeBase() {};
+
+		EGraphConnectionType GetType() const { return Type; }
+
+		FGuid GetGuid() const { return Guid; }
+		void SetGuid(FGuid InGuid) { Guid = InGuid; }
+
+		FName GetName() const { return Name; }
+		void SetName(FName InName) { Name = InName; }
+
+
+		virtual bool AddConnection(FConnectionTypeBase* In) { return false; };
+		virtual bool RemoveConnection(FConnectionTypeBase* In) { return false; }
+
+		virtual TArray< FConnectionTypeBase* > GetBaseInputs() { return TArray<FConnectionTypeBase* >(); }
+		virtual TArray< FConnectionTypeBase* > GetBaseOutputs() { return TArray<FConnectionTypeBase* >(); }
+		virtual void Invalidate() {};
+
+	};
+
+	template<typename T> class FOutput;
+
+	//
+	//  Input
+	//
+
+	template<class T>
+	struct EVALGRAPHCORE_API FInputParameters {
+		FInputParameters(FName InName, FNode* InOwner, T InDefault = T())
+			: Name(InName)
+			, Owner(InOwner)
+			, Default(InDefault) {}
+		FName Name;
+		FNode* Owner = nullptr;
+		T Default;
+	};
+
+	template<class T>
+	class EVALGRAPHCORE_API FInput : public FConnectionTypeBase
+	{
+		friend class FConnectionTypeBase;
+
+		typedef FConnectionTypeBase Super;
+		T Default;
+		FOutput<T>* Connection;
+	public:
+		FInput(const FInputParameters<T>& Param, FGuid InGuid = FGuid::NewGuid());
+	
+		const T& GetDefault() const { return Default; }
+		
+		const FOutput<T>* GetConnection() const { return Connection; }
+		FOutput<T>* GetConnection() { return Connection; }
+
+		virtual bool AddConnection(FConnectionTypeBase* InOutput) override
+		{ 
+			ensure(Connection == nullptr);
+			if (ensure(InOutput->GetType()==this->GetType()))
+			{
+				Connection = (FOutput<T>*)InOutput;
+				return true;
+			}
+			return false;
+		}
+
+		virtual bool RemoveConnection(FConnectionTypeBase* InOutput) override
+		{ 
+			if (ensure(Connection == (FOutput<T>*)InOutput))
+			{
+				Connection = nullptr;
+				return true;
+			}
+			return false;
+		}
+
+		T GetValue(const FContext& Context)
+		{
+			if (Connection)
+			{
+				return Connection->Evaluate(Context);
+			}
+			return Default;
+		}
+
+		void SetValue(const T& Value, const FContext& Context);
+
+		virtual TArray< FConnectionTypeBase* > GetBaseOutputs() override
+		{ 
+			TArray<FConnectionTypeBase* > RetList;
+			if (GetConnection())
+			{
+				RetList.Add(GetConnection());
+			}
+			return RetList;
+		}
+
+		virtual void Invalidate() override;
+
+	};
+
+
+	//
+	// Output
+	//
+
+	struct EVALGRAPHCORE_API FOutputParameters {
+		FOutputParameters(FName InName, FNode * InOwner) 
+			: Name(InName)
+			, Owner(InOwner) {}
+		FName Name;
+		FNode* Owner = nullptr;
+	};
+
+	template<class T>
+	class EVALGRAPHCORE_API FOutput : public FConnectionTypeBase
+	{
+		typedef FConnectionTypeBase Super;
+		friend class FConnectionTypeBase;
+
+		uint32 CacheKey = UINT_MAX;
+		TCacheValue<T> Cache;
+		TArray< FInput<T>* > Connections;
+	
+	public:
+		FOutput(const FOutputParameters& Param, FGuid InGuid = FGuid::NewGuid());
+
+		const TArray<FInput<T>*>& GetConnections() const { return Connections; }
+		TArray<FInput<T>* >& GetConnections() { return Connections; }
+
+		virtual TArray<FConnectionTypeBase*> GetBaseInputs() override
+		{ 
+			TArray<FConnectionTypeBase*> RetList;
+			RetList.Reserve(Connections.Num());
+			for (FInput<T>* Ptr : Connections) { RetList.Add(Ptr); }
+			return RetList;
+		}
+
+
+		virtual bool AddConnection(FConnectionTypeBase* InOutput) override
+		{
+			if (ensure(InOutput->GetType() == this->GetType()))
+			{
+				Connections.Add((FInput<T>*)InOutput);
+				return true;
+			}
+			return false;
+		}
+		
+		virtual bool RemoveConnection(FConnectionTypeBase* InInput) override { Connections.RemoveSwap((FInput<T>*)InInput); return true; }
+
+		void SetValue(T InVal, const FContext& Context);
+
+		T Evaluate(const FContext& Context);
+
+		virtual void Invalidate() override;
+
+	};
+}
+
diff --git a/Engine/Source/Runtime/Experimental/EvalGraph/Core/Public/EvalGraph/EvalGraphNode.h b/Engine/Source/Runtime/Experimental/EvalGraph/Core/Public/EvalGraph/EvalGraphNode.h
new file mode 100644
--- /dev/null
+++ b/Engine/Source/Runtime/Experimental/EvalGraph/Core/Public/EvalGraph/EvalGraphNode.h
@@ -0,0 +1,64 @@
+// Copyright Epic Games, Inc. All Rights Reserved.
+
+#pragma once
+
+#include "CoreMinimal.h"
+#include "EvalGraph/EvalGraphInputOutput.h"
+#include "EvalGraph/EvalGraphNodeParameters.h"
+
+class FConnectionTypeBase;
+
+namespace Eg
+{
+
+	struct EVALGRAPHCORE_API FNodeParameters {
+		FName Name;
+	};
+
+
+	//
+	// FNode
+	//
+	class FNode
+	{
+		friend class FGraph;
+		friend class FConnectionTypeBase;
+
+		FGuid Guid;
+		FName Name;
+
+		TArray< FConnectionTypeBase* > Inputs;
+		TArray< FConnectionTypeBase* > Outputs;
+	public:
+
+		FNode(const FNodeParameters& Param, FGuid InGuid = FGuid::NewGuid())
+			: Guid(InGuid)
+			, Name(Param.Name)
+		{}
+
+		virtual ~FNode() {}
+
+		FName GetName() const { return Name; }
+		void SetName(FName InName) { Name = InName; }
+
+		virtual FName GetType() const { check(true); return FName("invalid"); }
+		FGuid GetGuid() const { return Guid; }
+
+
+		virtual void Evaluate(const FContext& Context, FConnectionTypeBase*) { ensure(false); }
+		void InvalidateOutputs();
+
+		virtual void SerializeInternal(FArchive& Ar) {};
+
+	protected:
+		void AddBaseInput(FConnectionTypeBase* InPtr) { Inputs.Add(InPtr); }
+		const TArray< FConnectionTypeBase* >& GetInputs() const { return Inputs; }
+		TArray< FConnectionTypeBase* >& GetInputs() { return Inputs; }
+
+		void AddBaseOutput(FConnectionTypeBase* InPtr) { Outputs.Add(InPtr); }
+		const TArray< FConnectionTypeBase* >& GetOutputs() const { return Outputs; }
+		TArray< FConnectionTypeBase* >& GetOutputs() { return Outputs; }
+
+	};
+}
+
diff --git a/Engine/Source/Runtime/Experimental/EvalGraph/Core/Public/EvalGraph/EvalGraphNodeFactory.h b/Engine/Source/Runtime/Experimental/EvalGraph/Core/Public/EvalGraph/EvalGraphNodeFactory.h
new file mode 100644
--- /dev/null
+++ b/Engine/Source/Runtime/Experimental/EvalGraph/Core/Public/EvalGraph/EvalGraphNodeFactory.h
@@ -0,0 +1,59 @@
+// Copyright Epic Games, Inc. All Rights Reserved.
+
+#pragma once
+
+#include "CoreMinimal.h"
+#include "EvalGraph/EvalGraph.h"
+
+namespace Eg
+{
+	class FNode;
+	class FConnectionTypeBase;
+
+	struct EVALGRAPHCORE_API FNewNodeParameters {
+		FGuid Guid;
+		FName Type;
+		FName Name; 
+	};
+
+	//
+	//
+	//
+	class EVALGRAPHCORE_API FNodeFactory
+	{
+		typedef TFunction<FNode* (const FNewNodeParameters&)> FNewNodeFunction;
+
+		TMap<FName, FNewNodeFunction > ClassMap;
+
+		static FNodeFactory* Instance;
+		FNodeFactory() {}
+
+	public:
+		~FNodeFactory() { delete Instance; }
+
+		static FNodeFactory* GetInstance()
+		{
+			if (!Instance)
+			{
+				Instance = new FNodeFactory();
+			}
+			return Instance;
+		}
+
+		void RegisterNode(const FName& Type, FNewNodeFunction NewFunction)
+		{
+			ensure(!ClassMap.Contains(Type));
+			ClassMap.Add(Type, NewFunction);
+		}
+
+		TSharedPtr<FNode> NewNodeFromRegisteredType(FGraph& Graph, const FNewNodeParameters& Param);
+
+		void here() {}
+		template<class T> TSharedPtr<T> NewNode(FGraph& Graph, const FNewNodeParameters& Param)
+		{
+			return Graph.AddNode(new T(Param.Name, Param.Guid));
+		}
+	};
+
+}
+
diff --git a/Engine/Source/Runtime/Experimental/EvalGraph/Core/Public/EvalGraph/EvalGraphNodeParameters.h b/Engine/Source/Runtime/Experimental/EvalGraph/Core/Public/EvalGraph/EvalGraphNodeParameters.h
new file mode 100644
--- /dev/null
+++ b/Engine/Source/Runtime/Experimental/EvalGraph/Core/Public/EvalGraph/EvalGraphNodeParameters.h
@@ -0,0 +1,41 @@
+// Copyright Epic Games, Inc. All Rights Reserved.
+
+#pragma once
+
+#include "CoreMinimal.h"
+#include "EvalGraph/EvalGraphConnectionTypes.h"
+
+class  UEvalGraph;
+class  UEdGraphPin;
+
+namespace Eg
+{
+	struct EVALGRAPHCORE_API FContext
+	{
+		float Timestamp = 0.f;
+
+		uint32 GetTypeHash() const
+		{
+			return ::GetTypeHash(Timestamp);
+		}
+	};
+
+
+	struct FCacheValueBase {
+		FCacheValueBase(EGraphConnectionType InType) : Type(InType) {}
+		EGraphConnectionType Type;
+	};
+
+	template<class T>
+	struct TCacheValue : public FCacheValueBase {
+		TCacheValue(T InData = T())
+			: FCacheValueBase(GraphConnectionType<T>())
+			, Data(InData) {}
+		T Data;
+	};
+}
+
+FORCEINLINE uint32 GetTypeHash(const Eg::FContext& Context)
+{
+	return ::GetTypeHash(Context.Timestamp);
+}
diff --git a/Engine/Source/Runtime/Experimental/EvalGraph/Engine/EvalGraphEngine.Build.cs b/Engine/Source/Runtime/Experimental/EvalGraph/Engine/EvalGraphEngine.Build.cs
new file mode 100644
--- /dev/null
+++ b/Engine/Source/Runtime/Experimental/EvalGraph/Engine/EvalGraphEngine.Build.cs
@@ -0,0 +1,21 @@
+// Copyright Epic Games, Inc. All Rights Reserved.
+
+namespace UnrealBuildTool.Rules
+{
+	public class EvalGraphEngine : ModuleRules
+	{
+        public EvalGraphEngine(ReadOnlyTargetRules Target) : base(Target)
+		{
+			bTreatAsEngineModule = true;
+			PublicDependencyModuleNames.AddRange(
+				new string[]
+				{
+					"Core",
+					"CoreUObject",
+					"Engine",
+					"EvalGraphCore"
+				}
+				);
+		}
+	}
+}
diff --git a/Engine/Source/Runtime/Experimental/EvalGraph/Engine/Private/EvalGraph/EvalGraphEngineLogPrivate.h b/Engine/Source/Runtime/Experimental/EvalGraph/Engine/Private/EvalGraph/EvalGraphEngineLogPrivate.h
new file mode 100644
--- /dev/null
+++ b/Engine/Source/Runtime/Experimental/EvalGraph/Engine/Private/EvalGraph/EvalGraphEngineLogPrivate.h
@@ -0,0 +1,11 @@
+// Copyright Epic Games, Inc. All Rights Reserved.
+
+#pragma once
+
+#include "CoreMinimal.h"
+#include "Stats/Stats.h"
+
+// You should place include statements to your module's private header files here.  You only need to
+// add includes for headers that are used in most of your module's source files though.
+
+
diff --git a/Engine/Source/Runtime/Experimental/EvalGraph/Engine/Private/EvalGraph/EvalGraphEngineModule.cpp b/Engine/Source/Runtime/Experimental/EvalGraph/Engine/Private/EvalGraph/EvalGraphEngineModule.cpp
new file mode 100644
--- /dev/null
+++ b/Engine/Source/Runtime/Experimental/EvalGraph/Engine/Private/EvalGraph/EvalGraphEngineModule.cpp
@@ -0,0 +1,32 @@
+// Copyright Epic Games, Inc. All Rights Reserved.
+
+#pragma once
+
+#include "CoreMinimal.h"
+#include "EvalGraphEngineLogPrivate.h"
+#include "Modules/ModuleInterface.h"
+#include "Modules/ModuleManager.h"
+
+
+/**
+ * The public interface to this module
+ */
+class FEvalGraphEngineModule : public IModuleInterface
+{
+public:
+
+	// IModuleInterface interface
+
+	virtual void StartupModule() override { }
+	virtual void ShutdownModule() override { }
+
+	virtual bool SupportsDynamicReloading() override
+	{
+		return false;
+	}
+};
+
+IMPLEMENT_MODULE( FEvalGraphEngineModule, EvalGraphEngine )
+
+
+
diff --git a/Engine/Source/Runtime/Experimental/EvalGraph/Engine/Private/EvalGraph/EvalGraphObject.cpp b/Engine/Source/Runtime/Experimental/EvalGraph/Engine/Private/EvalGraph/EvalGraphObject.cpp
new file mode 100644
--- /dev/null
+++ b/Engine/Source/Runtime/Experimental/EvalGraph/Engine/Private/EvalGraph/EvalGraphObject.cpp
@@ -0,0 +1,33 @@
+// Copyright Epic Games, Inc. All Rights Reserved.
+
+#include "EvalGraph/EvalGraphObject.h"
+#include "EvalGraph/EvalGraph.h"
+
+
+#define LOCTEXT_NAMESPACE "UEvalGraph")
+
+UEvalGraph::UEvalGraph(const FObjectInitializer& ObjectInitializer)
+	: UObject(ObjectInitializer)
+	, EvalGraph(new Eg::FGraph())
+{}
+
+#if WITH_EDITOR
+
+void UEvalGraph::PostEditChangeProperty(struct FPropertyChangedEvent& PropertyChangedEvent)
+{
+	UObject::PostEditChangeProperty(PropertyChangedEvent);
+}
+
+#endif
+
+void UEvalGraph::PostLoad()
+{
+	UObject::PostLoad();
+}
+
+void UEvalGraph::Serialize(FArchive& Ar)
+{
+	Ar << *EvalGraph.Get();
+}
+
+#undef LOCTEXT_NAMESPACE
diff --git a/Engine/Source/Runtime/Experimental/EvalGraph/Engine/Public/EvalGraph/EvalGraphObject.h b/Engine/Source/Runtime/Experimental/EvalGraph/Engine/Public/EvalGraph/EvalGraphObject.h
new file mode 100644
--- /dev/null
+++ b/Engine/Source/Runtime/Experimental/EvalGraph/Engine/Public/EvalGraph/EvalGraphObject.h
@@ -0,0 +1,42 @@
+// Copyright Epic Games, Inc. All Rights Reserved.
+
+#pragma once
+
+#include "CoreMinimal.h"
+#include "EvalGraph/EvalGraph.h"
+
+#include "EvalGraphObject.generated.h"
+
+/**
+* UEvalGraph (UObject)
+*
+* UObject wrapper for the Eg::FGraph
+*
+*/
+UCLASS(BlueprintType, customconstructor)
+class EVALGRAPHENGINE_API UEvalGraph : public UObject
+{
+	GENERATED_UCLASS_BODY()
+
+	TSharedPtr<Eg::FGraph, ESPMode::ThreadSafe> EvalGraph;
+
+public:
+	UEvalGraph(const FObjectInitializer& ObjectInitializer = FObjectInitializer::Get());
+
+	/** UObject Interface */
+#if WITH_EDITOR
+	virtual void PostEditChangeProperty(struct FPropertyChangedEvent& PropertyChangedEvent) override;
+#endif
+	virtual void PostLoad() override;
+	/** End UObject Interface */
+
+	void Serialize(FArchive& Ar);
+
+	/** Accessors for internal geometry collection */
+	void SetEvalGraph(TSharedPtr<Eg::FGraph, ESPMode::ThreadSafe> EvalGraphIn) { EvalGraph = EvalGraphIn; }
+	TSharedPtr<Eg::FGraph, ESPMode::ThreadSafe>       GetEvalGraph() { return EvalGraph; }
+	const TSharedPtr<Eg::FGraph, ESPMode::ThreadSafe> GetEvalGraph() const { return EvalGraph; }
+
+
+};
+
