EvalGraph - Fix CIS

diff --git a/Engine/Source/Runtime/Experimental/EvalGraph/Core/Private/EvalGraph/EvalGraphCoreLogPrivate.h b/Engine/Source/Runtime/Experimental/EvalGraph/Core/Private/EvalGraph/EvalGraphCoreLogPrivate.h
--- a/Engine/Source/Runtime/Experimental/EvalGraph/Core/Private/EvalGraph/EvalGraphCoreLogPrivate.h
+++ b/Engine/Source/Runtime/Experimental/EvalGraph/Core/Private/EvalGraph/EvalGraphCoreLogPrivate.h
@@ -1,5 +1,7 @@
 // Copyright Epic Games, Inc. All Rights Reserved.
 
+#pragma once
+
 #include "CoreMinimal.h"
 #include "Stats/Stats.h"
 
