Set OwnerComponent on GeometryCollection's BodyInstance, fixes component ID not set in filter data, fixing ignore component SQ param for GC.

diff --git a/Engine/Source/Runtime/Experimental/GeometryCollectionEngine/Private/GeometryCollection/GeometryCollectionComponent.cpp b/Engine/Source/Runtime/Experimental/GeometryCollectionEngine/Private/GeometryCollection/GeometryCollectionComponent.cpp
--- a/Engine/Source/Runtime/Experimental/GeometryCollectionEngine/Private/GeometryCollection/GeometryCollectionComponent.cpp
+++ b/Engine/Source/Runtime/Experimental/GeometryCollectionEngine/Private/GeometryCollection/GeometryCollectionComponent.cpp
@@ -2166,6 +2166,7 @@ void UGeometryCollectionComponent::OnCreatePhysicsState()
 			// #BGTODO We need a dummy body setup for now to allow the body instance to generate filter information. Change body instance to operate independently.
 			DummyBodySetup = NewObject<UBodySetup>(this, UBodySetup::StaticClass());
 			BodyInstance.BodySetup = DummyBodySetup;
+			BodyInstance.OwnerComponent = this; // Required to make filter data include component/actor ID for ignored actors/components
 
 			FBodyCollisionFilterData FilterData;
 			FMaskFilter FilterMask = BodyInstance.GetMaskFilter();
@@ -2174,13 +2175,6 @@ void UGeometryCollectionComponent::OnCreatePhysicsState()
 			InitialSimFilter = FilterData.SimFilter;
 			InitialQueryFilter = FilterData.QuerySimpleFilter;
 
-			// since InitBody has not been called on the bodyInstance, OwnerComponent is nullptr
-			// we then need to set the owner on the query filters to allow for actor filtering 
-			if (const AActor* Owner = GetOwner())
-			{
-				InitialQueryFilter.Word0 = Owner->GetUniqueID();
-			}
-
 			// Enable for complex and simple (no dual representation currently like other meshes)
 			InitialQueryFilter.Word3 |= (EPDF_SimpleCollision | EPDF_ComplexCollision);
 			InitialSimFilter.Word3 |= (EPDF_SimpleCollision | EPDF_ComplexCollision);
