Chaos Cloth: fix accidentally submitted debug code.

I accidentally submitted debug code that disabled buckling for XPBDBendingConstraints when using ISPC in the initial commit 20107568.

diff --git a/Engine/Source/Runtime/Experimental/Chaos/Private/Chaos/XPBDBendingConstraints.ispc b/Engine/Source/Runtime/Experimental/Chaos/Private/Chaos/XPBDBendingConstraints.ispc
--- a/Engine/Source/Runtime/Experimental/Chaos/Private/Chaos/XPBDBendingConstraints.ispc
+++ b/Engine/Source/Runtime/Experimental/Chaos/Private/Chaos/XPBDBendingConstraints.ispc
@@ -87,7 +87,7 @@ export void ApplyXPBDBendingConstraints(uniform FVector4f PandInvM[],
 		
 		const varying float RestAngle = RestAngles[i];
 		const varying bool AngleIsBuckled = FLOAT_PI - abs(Angle) < BucklingRatio * (FLOAT_PI - abs(RestAngle));
-		const varying float StiffnessValue = Stiffness;//VectorSelect(AngleIsBuckled, BucklingStiffness, Stiffness);
+		const varying float StiffnessValue = select(AngleIsBuckled, BucklingStiffness, Stiffness);
 		
 		const varying float Lambda = Lambdas[i];
 		const varying float Alpha  = 1.f / (StiffnessValue * Dt * Dt);
