EvalGraph - Removed input output type declarations.

diff --git a/Engine/Source/Runtime/Experimental/EvalGraph/Core/Private/EvalGraph/EvalGraph.cpp b/Engine/Source/Runtime/Experimental/EvalGraph/Core/Private/EvalGraph/EvalGraph.cpp
--- a/Engine/Source/Runtime/Experimental/EvalGraph/Core/Private/EvalGraph/EvalGraph.cpp
+++ b/Engine/Source/Runtime/Experimental/EvalGraph/Core/Private/EvalGraph/EvalGraph.cpp
@@ -3,7 +3,6 @@
 #include "EvalGraph/EvalGraph.h"
 
 #include "EvalGraph/EvalGraphInputOutput.h"
-#include "EvalGraph/EvalGraphConnectionTypes.h"
 #include "EvalGraph/EvalGraphNodeParameters.h"
 #include "EvalGraph/EvalGraphNodeFactory.h"
 #include "Logging/LogMacros.h"
@@ -15,7 +14,8 @@ namespace Eg
 
 	FGraph::FGraph(FGuid InGuid)
 		: Guid(InGuid)
-	{}
+	{
+	}
 
 
 	void FGraph::RemoveNode(TSharedPtr<FNode> Node)
@@ -74,7 +74,6 @@ namespace Eg
 		{
 			FGuid ArGuid;
 			FName ArType, ArName;
-			uint32 ArIntType;
 			int32 ArNum = Nodes.Num();
 			Ar << ArNum;
 			for (TSharedPtr<FNode> Node : Nodes)
@@ -91,9 +90,9 @@ namespace Eg
 				for (FConnectionTypeBase* Conn : IO)
 				{
 					ArGuid = Conn->GetGuid();
-					ArIntType = (int32)Conn->GetType();
+					ArType = Conn->GetType();
 					ArName = Conn->GetName();
-					Ar << ArGuid << ArIntType << ArName;
+					Ar << ArGuid << ArType << ArName;
 				}
 
 				Node->SerializeInternal(Ar);
@@ -106,7 +105,6 @@ namespace Eg
 		{
 			FGuid ArGuid;
 			FName ArType, ArName;
-			uint32 ArIntType;
 			int32 ArNum = 0;
 
 			TMap<FGuid, FConnectionTypeBase* > ConnectionGuidMap;
@@ -122,16 +120,16 @@ namespace Eg
 					TArray< FConnectionTypeBase* > IO = Node->GetOutputs();  IO.Append(Node->GetInputs());
 					for (int Cdx=0;Cdx< ArNumInner; Cdx++)
 					{
-						Ar << ArGuid << ArIntType << ArName;
+						Ar << ArGuid << ArType << ArName;
 						if (Cdx < IO.Num())
 						{
-							if ((EGraphConnectionType)ArIntType!=IO[Cdx]->GetType())
-							{
-								FName TypeName = GraphConnectionTypeName((EGraphConnectionType)ArIntType);
-								FName FromType = GraphConnectionTypeName(IO[Cdx]->GetType());
-								UE_LOG(EGGRAPH_LOG, Fatal, TEXT("Type mismatch in input file for Ed::FGraphNode Input/Output (%s,%s) to type (%s)"),
-									*TypeName.ToString(), *ArName.ToString(), *FromType.ToString());
-							}
+							//if ((EGraphConnectionType)ArIntType!=IO[Cdx]->GetType())
+							//{
+							//	FName TypeName = GraphConnectionTypeName((EGraphConnectionType)ArIntType);
+							//	FName FromType = GraphConnectionTypeName(IO[Cdx]->GetType());
+							//	UE_LOG(EGGRAPH_LOG, Fatal, TEXT("Type mismatch in input file for Ed::FGraphNode Input/Output (%s,%s) to type (%s)"),
+							//		*TypeName.ToString(), *ArName.ToString(), *FromType.ToString());
+							//}
 
 							IO[Cdx]->SetGuid(ArGuid);
 							ensure(!ConnectionGuidMap.Contains(ArGuid));
diff --git a/Engine/Source/Runtime/Experimental/EvalGraph/Core/Private/EvalGraph/EvalGraphCoreModule.cpp b/Engine/Source/Runtime/Experimental/EvalGraph/Core/Private/EvalGraph/EvalGraphCoreModule.cpp
--- a/Engine/Source/Runtime/Experimental/EvalGraph/Core/Private/EvalGraph/EvalGraphCoreModule.cpp
+++ b/Engine/Source/Runtime/Experimental/EvalGraph/Core/Private/EvalGraph/EvalGraphCoreModule.cpp
@@ -1,31 +1,20 @@
 // Copyright Epic Games, Inc. All Rights Reserved.
 
-#pragma once
+#include "EvalGraph/EvalGraphCoreModule.h"
 
 #include "CoreMinimal.h"
-#include "EvalGraphCoreLogPrivate.h"
+#include "EvalGraph/EvalGraphNodeFactory.h"
 #include "Modules/ModuleInterface.h"
 #include "Modules/ModuleManager.h"
 
+#define LOCTEXT_NAMESPACE "EvalGraphCore"
 
-/**
- * The public interface to this module
- */
-class FEvalGraphCoreModule : public IModuleInterface
-{
+void IEvalGraphCoreModule::StartupModule()
+{/*Never Called ?*/}
 
-public:
+void IEvalGraphCoreModule::ShutdownModule()
+{/*Never Called ?*/ }
 
-	// IModuleInterface interface
+IMPLEMENT_MODULE(IEvalGraphCoreModule, EvalGraphCore)
 
-	virtual void StartupModule() override { }
-	virtual void ShutdownModule() override { }
-
-	virtual bool SupportsDynamicReloading() override
-	{
-		return false;
-	}
-
-};
-
-IMPLEMENT_MODULE( FEvalGraphCoreModule, EvalGraphCore )
+#undef LOCTEXT_NAMESPACE
diff --git a/Engine/Source/Runtime/Experimental/EvalGraph/Core/Private/EvalGraph/EvalGraphInputOutput.cpp b/Engine/Source/Runtime/Experimental/EvalGraph/Core/Private/EvalGraph/EvalGraphInputOutput.cpp
--- a/Engine/Source/Runtime/Experimental/EvalGraph/Core/Private/EvalGraph/EvalGraphInputOutput.cpp
+++ b/Engine/Source/Runtime/Experimental/EvalGraph/Core/Private/EvalGraph/EvalGraphInputOutput.cpp
@@ -2,14 +2,12 @@
 
 #include "EvalGraph/EvalGraphInputOutput.h"
 
-#include "EvalGraph/EvalGraphConnectionTypes.h"
 #include "EvalGraph/EvalGraphNodeParameters.h"
 #include "EvalGraph/EvalGraphNode.h"
 
 namespace Eg
 {
-
-	FConnectionTypeBase::FConnectionTypeBase(EGraphConnectionType InType, FName InName, FNode* InOwningNode, FGuid InGuid)
+	FConnectionTypeBase::FConnectionTypeBase(FName InType, FName InName, FNode* InOwningNode, FGuid InGuid)
 		: Type(InType)
 		, Name(InName)
 		, Guid(InGuid)
@@ -24,7 +22,7 @@ namespace Eg
 
 	template<class T>
 	FInput<T>::FInput(const FInputParameters<T>& Param, FGuid InGuid )
-		: FConnectionTypeBase(GraphConnectionType<T>(), Param.Name, Param.Owner, InGuid)
+		: FConnectionTypeBase(Param.Type, Param.Name, Param.Owner, InGuid)
 		, Default(Param.Default)
 		, Connection(nullptr)
 	{
@@ -48,8 +46,8 @@ namespace Eg
 	}
 
 	template<class T>
-	FOutput<T>::FOutput(const FOutputParameters& Param, FGuid InGuid )
-		: FConnectionTypeBase(GraphConnectionType<T>(), Param.Name, Param.Owner, InGuid)
+	FOutput<T>::FOutput(const FOutputParameters<T>& Param, FGuid InGuid )
+		: FConnectionTypeBase(Param.Type, Param.Name, Param.Owner, InGuid)
 	{
 		Super::AddBaseOutput(Param.Owner, this);
 	}
@@ -88,13 +86,13 @@ namespace Eg
 		}
 	}
 
+// Base Types
+EVAL_GRAPH_CONNECTION_TYPE(bool, Bool)
+EVAL_GRAPH_CONNECTION_TYPE(char, Char)
+EVAL_GRAPH_CONNECTION_TYPE(int, Integer)
+EVAL_GRAPH_CONNECTION_TYPE(uint8, UInt8)
+EVAL_GRAPH_CONNECTION_TYPE(float, Float)
+EVAL_GRAPH_CONNECTION_TYPE(double, Double)
 
-#define EVAL_GRAPH_CONNECTION_TYPE(a,A) template class FOutput<a>;
-#include "../../Public/EvalGraph/EvalGraphConnectionTypeValues.inl"
-#undef EVAL_GRAPH_CONNECTION_TYPE
-
-#define EVAL_GRAPH_CONNECTION_TYPE(a,A) template class FInput<a>;
-#include "../../Public/EvalGraph/EvalGraphConnectionTypeValues.inl"
-#undef EVAL_GRAPH_CONNECTION_TYPE
 }
 
diff --git a/Engine/Source/Runtime/Experimental/EvalGraph/Core/Private/EvalGraph/EvalGraphNodeFactory.cpp b/Engine/Source/Runtime/Experimental/EvalGraph/Core/Private/EvalGraph/EvalGraphNodeFactory.cpp
--- a/Engine/Source/Runtime/Experimental/EvalGraph/Core/Private/EvalGraph/EvalGraphNodeFactory.cpp
+++ b/Engine/Source/Runtime/Experimental/EvalGraph/Core/Private/EvalGraph/EvalGraphNodeFactory.cpp
@@ -2,7 +2,6 @@
 
 #include "EvalGraph/EvalGraphNodeFactory.h"
 
-#include "EvalGraph/EvalGraphConnectionTypes.h"
 #include "EvalGraph/EvalGraphNode.h"
 
 namespace Eg
diff --git a/Engine/Source/Runtime/Experimental/EvalGraph/Core/Public/EvalGraph/EvalGraphConnectionTypeValues.inl b/Engine/Source/Runtime/Experimental/EvalGraph/Core/Public/EvalGraph/EvalGraphConnectionTypeValues.inl
deleted file mode 100644
diff --git a/Engine/Source/Runtime/Experimental/EvalGraph/Core/Public/EvalGraph/EvalGraphConnectionTypes.h b/Engine/Source/Runtime/Experimental/EvalGraph/Core/Public/EvalGraph/EvalGraphConnectionTypes.h
deleted file mode 100644
diff --git a/Engine/Source/Runtime/Experimental/EvalGraph/Core/Public/EvalGraph/EvalGraphCoreModule.h b/Engine/Source/Runtime/Experimental/EvalGraph/Core/Public/EvalGraph/EvalGraphCoreModule.h
new file mode 100644
--- /dev/null
+++ b/Engine/Source/Runtime/Experimental/EvalGraph/Core/Public/EvalGraph/EvalGraphCoreModule.h
@@ -0,0 +1,30 @@
+// Copyright Epic Games, Inc. All Rights Reserved.
+
+#pragma once
+
+#include "CoreMinimal.h"
+#include "EvalGraph/EvalGraphNodeFactory.h"
+#include "Modules/ModuleInterface.h"
+#include "Modules/ModuleManager.h"
+
+
+/**
+ * The public interface to this module
+ */
+class EVALGRAPHCORE_API IEvalGraphCoreModule : public IModuleInterface
+{
+
+public:
+
+	// IModuleInterface interface
+
+	virtual void StartupModule() override;
+
+	virtual void ShutdownModule() override;
+
+	virtual bool SupportsDynamicReloading() override
+	{
+		return false;
+	}
+
+};
diff --git a/Engine/Source/Runtime/Experimental/EvalGraph/Core/Public/EvalGraph/EvalGraphInputOutput.h b/Engine/Source/Runtime/Experimental/EvalGraph/Core/Public/EvalGraph/EvalGraphInputOutput.h
--- a/Engine/Source/Runtime/Experimental/EvalGraph/Core/Public/EvalGraph/EvalGraphInputOutput.h
+++ b/Engine/Source/Runtime/Experimental/EvalGraph/Core/Public/EvalGraph/EvalGraphInputOutput.h
@@ -9,6 +9,9 @@ namespace Eg
 {
 	class FNode;
 
+	template<class T> inline EVALGRAPHCORE_API FName GraphConnectionTypeName();
+
+
 	//
 	// Input Output Base
 	//
@@ -17,7 +20,7 @@ namespace Eg
 	{
 
 	protected:
-		EGraphConnectionType Type;
+		FName Type;
 		FName Name;
 		FGuid  Guid;
 		FNode* OwningNode = nullptr;
@@ -28,10 +31,10 @@ namespace Eg
 
 
 	public:
-		FConnectionTypeBase(EGraphConnectionType InType, FName InName, FNode* OwningNode = nullptr, FGuid InGuid = FGuid::NewGuid());
+		FConnectionTypeBase(FName InType, FName InName, FNode* OwningNode = nullptr, FGuid InGuid = FGuid::NewGuid());
 		virtual ~FConnectionTypeBase() {};
 
-		EGraphConnectionType GetType() const { return Type; }
+		FName GetType() const { return Type; }
 
 		FGuid GetGuid() const { return Guid; }
 		void SetGuid(FGuid InGuid) { Guid = InGuid; }
@@ -58,9 +61,10 @@ namespace Eg
 	template<class T>
 	struct EVALGRAPHCORE_API FInputParameters {
 		FInputParameters(FName InName, FNode* InOwner, T InDefault = T())
-			: Name(InName)
+			: Type(GraphConnectionTypeName<T>())
 			, Owner(InOwner)
 			, Default(InDefault) {}
+		FName Type;
 		FName Name;
 		FNode* Owner = nullptr;
 		T Default;
@@ -133,10 +137,13 @@ namespace Eg
 	// Output
 	//
 
+	template<class T>
 	struct EVALGRAPHCORE_API FOutputParameters {
 		FOutputParameters(FName InName, FNode * InOwner) 
-			: Name(InName)
+			: Type(GraphConnectionTypeName<T>())
+			, Name(InName)
 			, Owner(InOwner) {}
+		FName Type;
 		FName Name;
 		FNode* Owner = nullptr;
 	};
@@ -152,7 +159,7 @@ namespace Eg
 		TArray< FInput<T>* > Connections;
 	
 	public:
-		FOutput(const FOutputParameters& Param, FGuid InGuid = FGuid::NewGuid());
+		FOutput(const FOutputParameters<T>& Param, FGuid InGuid = FGuid::NewGuid());
 
 		const TArray<FInput<T>*>& GetConnections() const { return Connections; }
 		TArray<FInput<T>* >& GetConnections() { return Connections; }
@@ -185,5 +192,11 @@ namespace Eg
 		virtual void Invalidate() override;
 
 	};
+
+
+#define EVAL_GRAPH_CONNECTION_TYPE(a,A) \
+template<> FName GraphConnectionTypeName<a>() { return FName(TEXT(#A)); }\
+template class FOutput<a>;\
+template class FInput<a>;
 }
 
diff --git a/Engine/Source/Runtime/Experimental/EvalGraph/Core/Public/EvalGraph/EvalGraphNode.h b/Engine/Source/Runtime/Experimental/EvalGraph/Core/Public/EvalGraph/EvalGraphNode.h
--- a/Engine/Source/Runtime/Experimental/EvalGraph/Core/Public/EvalGraph/EvalGraphNode.h
+++ b/Engine/Source/Runtime/Experimental/EvalGraph/Core/Public/EvalGraph/EvalGraphNode.h
@@ -8,6 +8,7 @@
 
 class FConnectionTypeBase;
 
+
 namespace Eg
 {
 
@@ -62,3 +63,14 @@ namespace Eg
 	};
 }
 
+#define EVAL_GRAPH_REGISTER_CREATION_FACTORY(A) Eg::FNodeFactory::GetInstance()->RegisterNode(Eg::A::Type, [](const Eg::FNewNodeParameters& InParam) {return new Eg::A({InParam.Name}, InParam.Guid); });
+
+#define EVAL_GRAPH_DEFINE_EXTERNAL(A) FName A::Type = #A;
+
+#define EVAL_GRAPH_DEFINE_INTERNAL(A) public:\
+ static FName Type;\
+ virtual FName GetType() const { return Type; }\
+ private:
+
+
+
diff --git a/Engine/Source/Runtime/Experimental/EvalGraph/Core/Public/EvalGraph/EvalGraphNodeFactory.h b/Engine/Source/Runtime/Experimental/EvalGraph/Core/Public/EvalGraph/EvalGraphNodeFactory.h
--- a/Engine/Source/Runtime/Experimental/EvalGraph/Core/Public/EvalGraph/EvalGraphNodeFactory.h
+++ b/Engine/Source/Runtime/Experimental/EvalGraph/Core/Public/EvalGraph/EvalGraphNodeFactory.h
@@ -42,8 +42,10 @@ namespace Eg
 
 		void RegisterNode(const FName& Type, FNewNodeFunction NewFunction)
 		{
-			ensure(!ClassMap.Contains(Type));
-			ClassMap.Add(Type, NewFunction);
+			if (ensure(!ClassMap.Contains(Type)))
+			{
+				ClassMap.Add(Type, NewFunction);
+			}
 		}
 
 		TSharedPtr<FNode> NewNodeFromRegisteredType(FGraph& Graph, const FNewNodeParameters& Param);
diff --git a/Engine/Source/Runtime/Experimental/EvalGraph/Core/Public/EvalGraph/EvalGraphNodeParameters.h b/Engine/Source/Runtime/Experimental/EvalGraph/Core/Public/EvalGraph/EvalGraphNodeParameters.h
--- a/Engine/Source/Runtime/Experimental/EvalGraph/Core/Public/EvalGraph/EvalGraphNodeParameters.h
+++ b/Engine/Source/Runtime/Experimental/EvalGraph/Core/Public/EvalGraph/EvalGraphNodeParameters.h
@@ -3,7 +3,6 @@
 #pragma once
 
 #include "CoreMinimal.h"
-#include "EvalGraph/EvalGraphConnectionTypes.h"
 
 class  UEvalGraph;
 class  UEdGraphPin;
@@ -21,16 +20,10 @@ namespace Eg
 	};
 
 
-	struct FCacheValueBase {
-		FCacheValueBase(EGraphConnectionType InType) : Type(InType) {}
-		EGraphConnectionType Type;
-	};
-
 	template<class T>
-	struct TCacheValue : public FCacheValueBase {
+	struct TCacheValue  {
 		TCacheValue(T InData = T())
-			: FCacheValueBase(GraphConnectionType<T>())
-			, Data(InData) {}
+			: Data(InData) {}
 		T Data;
 	};
 }
