Bug fix: Use scaled geometry to compute AABB in Triangle Mesh Implicit Object

diff --git a/Engine/Source/Runtime/Experimental/Chaos/Private/Chaos/TriangleMeshImplicitObject.cpp b/Engine/Source/Runtime/Experimental/Chaos/Private/Chaos/TriangleMeshImplicitObject.cpp
--- a/Engine/Source/Runtime/Experimental/Chaos/Private/Chaos/TriangleMeshImplicitObject.cpp
+++ b/Engine/Source/Runtime/Experimental/Chaos/Private/Chaos/TriangleMeshImplicitObject.cpp
@@ -512,7 +512,7 @@ bool FTriangleMeshImplicitObject::ContactManifoldImp(const GeomType& QueryGeom,
 	FRigidTransform3 TriMeshToGeomNoScale{ QueryTM };
 	TriMeshToGeomNoScale.SetTranslation(TriMeshToGeomNoScale.GetTranslation() * TriMeshScale);
 	// NOTE: BVH test is done in tri-mesh local space (whereas collision detection is done in world space becaused you can't non-uniformly scale all shapes)
-	FAABB3 QueryBounds = QueryGeom.CalculateTransformedBounds(TriMeshToGeomNoScale);
+	FAABB3 QueryBounds = WorldScaleGeom.CalculateTransformedBounds(TriMeshToGeomNoScale);
 	QueryBounds.ThickenSymmetrically(FVec3(WorldThickness));
 	QueryBounds.ScaleWithNegative(InvTriMeshScale);
 
@@ -631,7 +631,7 @@ bool FTriangleMeshImplicitObject::GJKContactPointImp(const QueryGeomType& QueryG
 	FRigidTransform3 TriMeshToGeomNoScale{ QueryTM };
 	TriMeshToGeomNoScale.SetTranslation(TriMeshToGeomNoScale.GetTranslation() * TriMeshScale);
 	// NOTE: BVH test is done in tri-mesh local space (whereas collision detection is done in world space because you can't non-uniformly scale all shapes)
-	FAABB3 QueryBounds = QueryGeom.CalculateTransformedBounds(TriMeshToGeomNoScale);
+	FAABB3 QueryBounds = WorldScaleGeom.CalculateTransformedBounds(TriMeshToGeomNoScale);
 	QueryBounds.ThickenSymmetrically(FVec3(WorldThickness));
 	QueryBounds.ScaleWithNegative(InvTriMeshScale);
 
@@ -899,7 +899,7 @@ bool FTriangleMeshImplicitObject::OverlapGeomImp(const QueryGeomType& QueryGeom,
 	FRigidTransform3 TriMeshToGeomNoScale{ QueryTM };
 	TriMeshToGeomNoScale.SetTranslation(TriMeshToGeomNoScale.GetTranslation() * TriMeshScale);
 	// NOTE: BVH test is done in tri-mesh local space (whereas collision detection is done in world space because you can't non-uniformly scale all shapes)
-	FAABB3 QueryBounds = QueryGeom.CalculateTransformedBounds(TriMeshToGeomNoScale);
+	FAABB3 QueryBounds = WorldScaleQueryGeom.CalculateTransformedBounds(TriMeshToGeomNoScale);
 	QueryBounds.ThickenSymmetrically(FVec3(Thickness));
 	QueryBounds.ScaleWithNegative(InvTriMeshScale);
 
