Fix CIS

diff --git a/Engine/Source/Runtime/Experimental/EvalGraph/Engine/Private/EvalGraph/EvalGraphEngineModule.cpp b/Engine/Source/Runtime/Experimental/EvalGraph/Engine/Private/EvalGraph/EvalGraphEngineModule.cpp
--- a/Engine/Source/Runtime/Experimental/EvalGraph/Engine/Private/EvalGraph/EvalGraphEngineModule.cpp
+++ b/Engine/Source/Runtime/Experimental/EvalGraph/Engine/Private/EvalGraph/EvalGraphEngineModule.cpp
@@ -1,7 +1,5 @@
 // Copyright Epic Games, Inc. All Rights Reserved.
 
-#pragma once
-
 #include "CoreMinimal.h"
 #include "EvalGraphEngineLogPrivate.h"
 #include "Modules/ModuleInterface.h"
