Interchange: made sure we don't generate multiple asset creation tasks for the same asset name as they would end up editing the same UObject.

diff --git a/Engine/Source/Runtime/Experimental/Interchange/Engine/Private/Tasks/InterchangeTaskParsing.cpp b/Engine/Source/Runtime/Experimental/Interchange/Engine/Private/Tasks/InterchangeTaskParsing.cpp
--- a/Engine/Source/Runtime/Experimental/Interchange/Engine/Private/Tasks/InterchangeTaskParsing.cpp
+++ b/Engine/Source/Runtime/Experimental/Interchange/Engine/Private/Tasks/InterchangeTaskParsing.cpp
@@ -225,7 +225,8 @@ void UE::Interchange::FTaskParsing::DoTask(ENamedThreads::Type CurrentThread, co
 	//Assets
 	FGraphEventArray AssetsCompletionPrerequistes;
 	{
-		TFunction<FGraphEventRef(FTaskData&)> CreateTasksForAsset = [this, &AsyncHelper](FTaskData& TaskData)
+		TSet<FString> CreatedTasksAssetNames; // Tracks for which asset name we have created a task so that we don't have 2 tasks for the same asset name
+		TFunction<FGraphEventRef(FTaskData&)> CreateTasksForAsset = [this, &AsyncHelper, &CreatedTasksAssetNames](FTaskData& TaskData)
 		{
 			check(TaskData.Nodes.Num() == 1); //We expect 1 node per asset task
 
@@ -234,20 +235,36 @@ void UE::Interchange::FTaskParsing::DoTask(ENamedThreads::Type CurrentThread, co
 			UInterchangeFactoryBaseNode* const FactoryNode = TaskData.Nodes[0];
 			const bool bFactoryCanRunOnAnyThread = FactoryClass->GetDefaultObject<UInterchangeFactoryBase>()->CanExecuteOnAnyThread();
 
-			//Add create package task has a prerequisite of FTaskCreateAsset. Create package task is a game thread task
-			FGraphEventArray CreatePackagePrerequistes;
-			int32 CreatePackageTaskIndex = AsyncHelper->CreatePackageTasks.Add(
-				TGraphTask<FTaskCreatePackage>::CreateTask(&(TaskData.Prerequisites)).ConstructAndDispatchWhenReady(PackageBasePath, SourceIndex, WeakAsyncHelper, FactoryNode, FactoryClass)
-			);
-			CreatePackagePrerequistes.Add(AsyncHelper->CreatePackageTasks[CreatePackageTaskIndex]);
+			FString PackageSubPath;
+			FactoryNode->GetCustomSubPath(PackageSubPath);
 
-			FGraphEventArray CreateAssetPrerequistes;
-			int32 CreateTaskIndex = AsyncHelper->CreateAssetTasks.Add(
-				TGraphTask<FTaskCreateAsset>::CreateTask(&(CreatePackagePrerequistes)).ConstructAndDispatchWhenReady(PackageBasePath, SourceIndex, WeakAsyncHelper, FactoryNode, bFactoryCanRunOnAnyThread)
-			);
-			CreateAssetPrerequistes.Add(AsyncHelper->CreateAssetTasks[CreateTaskIndex]);
+			const FString AssetFullPath = FPaths::Combine(PackageBasePath, PackageSubPath, FactoryNode->GetAssetName());
 
-			return AsyncHelper->CreateAssetTasks[CreateTaskIndex];
+			if ( ensureMsgf(!CreatedTasksAssetNames.Contains(AssetFullPath),
+				TEXT("Found multiple task data with the same asset name (%s). Only one will be executed."), *AssetFullPath) )
+			{
+				//Add create package task has a prerequisite of FTaskCreateAsset. Create package task is a game thread task
+				FGraphEventArray CreatePackagePrerequistes;
+				int32 CreatePackageTaskIndex = AsyncHelper->CreatePackageTasks.Add(
+					TGraphTask<FTaskCreatePackage>::CreateTask(&(TaskData.Prerequisites)).ConstructAndDispatchWhenReady(PackageBasePath, SourceIndex, WeakAsyncHelper, FactoryNode, FactoryClass)
+				);
+				CreatePackagePrerequistes.Add(AsyncHelper->CreatePackageTasks[CreatePackageTaskIndex]);
+
+				FGraphEventArray CreateAssetPrerequistes;
+				int32 CreateTaskIndex = AsyncHelper->CreateAssetTasks.Add(
+					TGraphTask<FTaskCreateAsset>::CreateTask(&(CreatePackagePrerequistes)).ConstructAndDispatchWhenReady(PackageBasePath, SourceIndex, WeakAsyncHelper, FactoryNode, bFactoryCanRunOnAnyThread)
+				);
+				CreateAssetPrerequistes.Add(AsyncHelper->CreateAssetTasks[CreateTaskIndex]);
+
+				CreatedTasksAssetNames.Add(AssetFullPath);
+				return AsyncHelper->CreateAssetTasks[CreateTaskIndex];
+			}
+			else
+			{
+				FGraphEventRef EmptyGraphEvent = FGraphEvent::CreateGraphEvent();
+				EmptyGraphEvent->DispatchSubsequents();
+				return EmptyGraphEvent;
+			}
 		};
 
 		AssetsCompletionPrerequistes = CreateTasksForEachTaskData(AssetTaskDatas, CreateTasksForAsset);
