Make CollisionObjectReductionPercentage also apply to convex collision for geometry collections

diff --git a/Engine/Source/Runtime/Experimental/Chaos/Private/GeometryCollection/GeometryCollectionCollisionStructureManager.cpp b/Engine/Source/Runtime/Experimental/Chaos/Private/GeometryCollection/GeometryCollectionCollisionStructureManager.cpp
--- a/Engine/Source/Runtime/Experimental/Chaos/Private/GeometryCollection/GeometryCollectionCollisionStructureManager.cpp
+++ b/Engine/Source/Runtime/Experimental/Chaos/Private/GeometryCollection/GeometryCollectionCollisionStructureManager.cpp
@@ -253,7 +253,8 @@ FCollisionStructureManager::NewImplicitConvex(
 	const TManagedArray<TUniquePtr<Chaos::FConvex>>* ConvexGeometry,
 	const ECollisionTypeEnum CollisionType,
 	const FTransform& MassTransform,
-	const Chaos::FReal CollisionMarginFraction)
+	const Chaos::FReal CollisionMarginFraction,
+	const float CollisionObjectReduction)
 {
 	using FConvexVec3 = Chaos::FConvex::FVec3Type;
 
@@ -263,9 +264,11 @@ FCollisionStructureManager::NewImplicitConvex(
 		for (auto& Index : ConvexIndices)
 		{
 			TArray<FConvexVec3> ConvexVertices = (*ConvexGeometry)[Index]->GetVertices();
+			FConvexVec3 COM = MassTransform.InverseTransformPosition((*ConvexGeometry)[Index]->GetCenterOfMass());
+			FConvexVec3::FReal ScaleFactor = 1 - CollisionObjectReduction / 100.f;
 			for (int32 Idx = 0; Idx < ConvexVertices.Num(); Idx++)
 			{
-				ConvexVertices[Idx] = MassTransform.InverseTransformPosition(FVector(ConvexVertices[Idx]));
+				ConvexVertices[Idx] = ((FConvexVec3)MassTransform.InverseTransformPosition(FVector(ConvexVertices[Idx])) - COM) * ScaleFactor + COM;
 			}
 
 			Chaos::FReal Margin = (Chaos::FReal)(*ConvexGeometry)[Index]->BoundingBox().Extents().Min() * CollisionMarginFraction;
diff --git a/Engine/Source/Runtime/Experimental/Chaos/Private/PhysicsProxy/GeometryCollectionPhysicsProxy.cpp b/Engine/Source/Runtime/Experimental/Chaos/Private/PhysicsProxy/GeometryCollectionPhysicsProxy.cpp
--- a/Engine/Source/Runtime/Experimental/Chaos/Private/PhysicsProxy/GeometryCollectionPhysicsProxy.cpp
+++ b/Engine/Source/Runtime/Experimental/Chaos/Private/PhysicsProxy/GeometryCollectionPhysicsProxy.cpp
@@ -2216,7 +2216,8 @@ static FGeometryDynamicCollection::FSharedImplicit CreateImplicitGeometry(
 							ConvexGeometry,
 							CollisionTypeData.CollisionType,
 							CollectionMassToLocal[TransformGroupIndex],
-							CollisionTypeData.CollisionMarginFraction
+							CollisionTypeData.CollisionMarginFraction,
+							CollisionTypeData.CollisionObjectReductionPercentage
 						)
 					);
 				}
diff --git a/Engine/Source/Runtime/Experimental/Chaos/Public/GeometryCollection/GeometryCollectionCollisionStructureManager.h b/Engine/Source/Runtime/Experimental/Chaos/Public/GeometryCollection/GeometryCollectionCollisionStructureManager.h
--- a/Engine/Source/Runtime/Experimental/Chaos/Public/GeometryCollection/GeometryCollectionCollisionStructureManager.h
+++ b/Engine/Source/Runtime/Experimental/Chaos/Public/GeometryCollection/GeometryCollectionCollisionStructureManager.h
@@ -116,7 +116,8 @@ public:
 		const TManagedArray<TUniquePtr<Chaos::FConvex>>* ConvexGeometry,
 		const ECollisionTypeEnum CollisionType,
 		const FTransform& MassTransform,
-		const Chaos::FReal CollisionMarginFraction);
+		const Chaos::FReal CollisionMarginFraction,
+		const float CollisionObjectReduction);
 
 	static FImplicit* NewImplicitLevelset(
 		Chaos::FErrorReporter ErrorReporter,
