Add support for the Start Awake flag for geometry collections

diff --git a/Engine/Source/Runtime/Experimental/Chaos/Private/PhysicsProxy/GeometryCollectionPhysicsProxy.cpp b/Engine/Source/Runtime/Experimental/Chaos/Private/PhysicsProxy/GeometryCollectionPhysicsProxy.cpp
--- a/Engine/Source/Runtime/Experimental/Chaos/Private/PhysicsProxy/GeometryCollectionPhysicsProxy.cpp
+++ b/Engine/Source/Runtime/Experimental/Chaos/Private/PhysicsProxy/GeometryCollectionPhysicsProxy.cpp
@@ -750,6 +750,32 @@ void FGeometryCollectionPhysicsProxy::SetGravityEnabled(const Chaos::FPBDRigidsS
 	}
 }
 
+void FGeometryCollectionPhysicsProxy::SetSleepingState(const Chaos::FPBDRigidsSolver& RigidsSolver)
+{
+	for (FClusterHandle* Handle: SolverParticleHandles)
+	{
+		if (Handle && !Handle->Disabled())
+		{
+			// Sleeping Geometry Collections:
+			//   A sleeping geometry collection is dynamic internally, and then the top level
+			//   active clusters are set to sleeping. Sleeping is not propagated up from the 
+			//   leaf nodes like kinematic or dynamic clusters.
+			RigidsSolver.GetEvolution()->SetParticleObjectState(Handle, Chaos::EObjectStateType::Sleeping);
+		}
+	}
+}
+
+void FGeometryCollectionPhysicsProxy::DirtyAllParticles(const Chaos::FPBDRigidsSolver& RigidsSolver)
+{
+	for (FClusterHandle* Handle: SolverParticleHandles)
+	{
+		if (Handle)
+		{
+			RigidsSolver.GetEvolution()->DirtyParticle(*Handle);
+		}
+	}
+}
+
 void FGeometryCollectionPhysicsProxy::InitializeBodiesPT(Chaos::FPBDRigidsSolver* RigidsSolver, typename Chaos::FPBDRigidsSolver::FParticlesType& Particles)
 {
 	const FGeometryCollection* RestCollection = Parameters.RestCollection;
@@ -1066,24 +1092,17 @@ void FGeometryCollectionPhysicsProxy::InitializeBodiesPT(Chaos::FPBDRigidsSolver
 		const bool bEnableGravity = Parameters.EnableGravity && !DisableGeometryCollectionGravity;
 		SetGravityEnabled(*RigidsSolver, bEnableGravity);
 
-		// call DirtyParticle to make sure the acceleration structure is up to date with all the changes happening here
-		for (int32 TransformGroupIndex = 0; TransformGroupIndex < NumTransforms; ++TransformGroupIndex)
+		const bool bStartSleeping =
+				(Parameters.ObjectType == EObjectStateTypeEnum::Chaos_Object_Sleeping
+			 || (Parameters.ObjectType == EObjectStateTypeEnum::Chaos_Object_Dynamic && !Parameters.StartAwake));
+		if (bStartSleeping)
 		{
-			if (FClusterHandle* Handle = SolverParticleHandles[TransformGroupIndex])
-			{
-				// Sleeping Geometry Collections:
-				//   A sleeping geometry collection is dynamic internally, and then the top level
-				//   active clusters are set to sleeping. Sleeping is not propagated up from the 
-				//   leaf nodes like kinematic or dynamic clusters. 
-				if (!Handle->Disabled() && Parameters.ObjectType == EObjectStateTypeEnum::Chaos_Object_Sleeping)
-				{
-				 	RigidsSolver->GetEvolution()->SetParticleObjectState(Handle, Chaos::EObjectStateType::Sleeping);
-				}
-
-				RigidsSolver->GetEvolution()->DirtyParticle(*Handle);
-			}
+			SetSleepingState(*RigidsSolver);
 		}
 
+		// call DirtyParticle to make sure the acceleration structure is up to date with all the changes happening here
+		DirtyAllParticles(*RigidsSolver);
+
 	} // end if simulating...
 
 }
diff --git a/Engine/Source/Runtime/Experimental/Chaos/Public/GeometryCollection/GeometryCollectionSimulationCoreTypes.h b/Engine/Source/Runtime/Experimental/Chaos/Public/GeometryCollection/GeometryCollectionSimulationCoreTypes.h
--- a/Engine/Source/Runtime/Experimental/Chaos/Public/GeometryCollection/GeometryCollectionSimulationCoreTypes.h
+++ b/Engine/Source/Runtime/Experimental/Chaos/Public/GeometryCollection/GeometryCollectionSimulationCoreTypes.h
@@ -188,6 +188,7 @@ struct FSimulationParameters
 		, ReverseCacheBeginTime(0.0f)
 		, bClearCache(false)
 		, ObjectType(EObjectStateTypeEnum::Chaos_NONE)
+		, StartAwake(true)
 		, bGenerateBreakingData(false)
 		, bGenerateCollisionData(false)
 		, bGenerateTrailingData(false)
@@ -224,6 +225,7 @@ struct FSimulationParameters
 		, ReverseCacheBeginTime(Other.ReverseCacheBeginTime)
 		, bClearCache(Other.bClearCache)
 		, ObjectType(Other.ObjectType)
+		, StartAwake(Other.StartAwake)
 		, PhysicalMaterialHandle(Other.PhysicalMaterialHandle)
 		, bGenerateBreakingData(Other.bGenerateBreakingData)
 		, bGenerateCollisionData(Other.bGenerateCollisionData)
@@ -280,6 +282,7 @@ struct FSimulationParameters
 	bool bClearCache;
 
 	EObjectStateTypeEnum ObjectType;
+	bool StartAwake;
 
 	Chaos::FMaterialHandle PhysicalMaterialHandle;
 
diff --git a/Engine/Source/Runtime/Experimental/Chaos/Public/PhysicsProxy/GeometryCollectionPhysicsProxy.h b/Engine/Source/Runtime/Experimental/Chaos/Public/PhysicsProxy/GeometryCollectionPhysicsProxy.h
--- a/Engine/Source/Runtime/Experimental/Chaos/Public/PhysicsProxy/GeometryCollectionPhysicsProxy.h
+++ b/Engine/Source/Runtime/Experimental/Chaos/Public/PhysicsProxy/GeometryCollectionPhysicsProxy.h
@@ -274,6 +274,8 @@ protected:
 		const Chaos::FUniqueIdx* ExistingIndex);
 
 	void SetGravityEnabled(const Chaos::FPBDRigidsSolver& RigidsSolver, bool bEnabled);
+	void SetSleepingState(const Chaos::FPBDRigidsSolver& RigidsSolver);
+	void DirtyAllParticles(const Chaos::FPBDRigidsSolver& RigidsSolver);
 	
 	/** 
 	 * Traverses the parents of \p TransformIndex in \p GeometryCollection, counting
diff --git a/Engine/Source/Runtime/Experimental/GeometryCollectionEngine/Private/GeometryCollection/GeometryCollectionComponent.cpp b/Engine/Source/Runtime/Experimental/GeometryCollectionEngine/Private/GeometryCollection/GeometryCollectionComponent.cpp
--- a/Engine/Source/Runtime/Experimental/GeometryCollectionEngine/Private/GeometryCollection/GeometryCollectionComponent.cpp
+++ b/Engine/Source/Runtime/Experimental/GeometryCollectionEngine/Private/GeometryCollection/GeometryCollectionComponent.cpp
@@ -2116,6 +2116,7 @@ void UGeometryCollectionComponent::OnCreatePhysicsState()
 			// If the Component is set to Dynamic, we look to the RestCollection for initial dynamic state override per transform.
 			TManagedArray<int32> & DynamicState = DynamicCollection->DynamicState;
 
+			// if this code is changed you may need to account for bStartAwake
 			EObjectStateTypeEnum LocalObjectType = (ObjectType != EObjectStateTypeEnum::Chaos_Object_Sleeping) ? ObjectType : EObjectStateTypeEnum::Chaos_Object_Dynamic;
 			if (LocalObjectType != EObjectStateTypeEnum::Chaos_Object_UserDefined)
 			{
@@ -2232,6 +2233,7 @@ void UGeometryCollectionComponent::RegisterAndInitializePhysicsProxy()
 		SimulationParameters.InitialAngularVelocity = InitialAngularVelocity;
 		SimulationParameters.bClearCache = true;
 		SimulationParameters.ObjectType = ObjectType;
+		SimulationParameters.StartAwake = BodyInstance.bStartAwake;
 		SimulationParameters.CacheType = CacheParameters.CacheMode;
 		SimulationParameters.ReverseCacheBeginTime = CacheParameters.ReverseCacheBeginTime;
 		SimulationParameters.bGenerateBreakingData = bNotifyBreaks;
