Chaos : destruction add an cvar to use box bounding overlapping filtering for simple delaunay clustering ( only for union clustering connection graph generation for now )

diff --git a/Engine/Source/Runtime/Experimental/Chaos/Private/Chaos/PBDRigidClustering.cpp b/Engine/Source/Runtime/Experimental/Chaos/Private/Chaos/PBDRigidClustering.cpp
--- a/Engine/Source/Runtime/Experimental/Chaos/Private/Chaos/PBDRigidClustering.cpp
+++ b/Engine/Source/Runtime/Experimental/Chaos/Private/Chaos/PBDRigidClustering.cpp
@@ -42,6 +42,8 @@ namespace Chaos
 	int32 DeactivateClusterChildren = 0;
 	FAutoConsoleVariableRef CVarDeactivateClusterChildren(TEXT("p.DeactivateClusterChildren"), DeactivateClusterChildren, TEXT("If children should be decativated when broken and put into another cluster."));
 
+	int32 UseBoundingBoxForConnectionGraphFiltering = 0;
+	FAutoConsoleVariableRef CVarUseBoundingBoxForConnectionGraphFiltering(TEXT("p.UseBoundingBoxForConnectionGraphFiltering"), UseBoundingBoxForConnectionGraphFiltering, TEXT("when on, use bounding box overlaps to filter connection during the connection graph generation"));
 
 	//==========================================================================
 	// TPBDRigidClustering
@@ -1364,6 +1366,15 @@ namespace Chaos
 		}
 	}
 
+	static bool IsOveralppingConnection(const FPBDRigidParticleHandle* Child1, const FPBDRigidParticleHandle* Child2)
+	{
+		if (ensure(Child1 != nullptr && Child2 != nullptr ))
+		{
+			return Child1->WorldSpaceInflatedBounds().Intersects(Child2->WorldSpaceInflatedBounds());
+		}
+		return false;
+	}
+	
 	DECLARE_CYCLE_STAT(TEXT("TPBDRigidClustering<>::UpdateConnectivityGraphUsingDelaunayTriangulation"), STAT_UpdateConnectivityGraphUsingDelaunayTriangulation, STATGROUP_Chaos);
 	void 
 	FRigidClustering::UpdateConnectivityGraphUsingDelaunayTriangulation(
@@ -1396,9 +1407,17 @@ namespace Chaos
 					bFirstSmaller ? Child2 : Child1);
 				if (!UniqueEdges.Find(SortedPair))
 				{
-					// this does not use ConnectNodes because Neighbors is bi-direction : as in (1,2),(2,1)
-					ConnectNodes(Child1, Child2);
-					UniqueEdges.Add(SortedPair);
+					bool ValidConnection = true;
+					if (UseBoundingBoxForConnectionGraphFiltering)
+					{
+						ValidConnection = IsOveralppingConnection(Child1, Child2);
+					}
+					if (ValidConnection)
+					{
+						// this does not use ConnectNodes because Neighbors is bi-direction : as in (1,2),(2,1)
+						ConnectNodes(Child1, Child2);
+						UniqueEdges.Add(SortedPair);
+					}
 				}
 			}
 		}
