Update B_Creative_MoveTool_BlueprintPreviewer to support previewing selected GeometryCollectionComponents.

diff --git a/Engine/Source/Runtime/Experimental/GeometryCollectionEngine/Public/GeometryCollection/GeometryCollectionComponent.h b/Engine/Source/Runtime/Experimental/GeometryCollectionEngine/Public/GeometryCollection/GeometryCollectionComponent.h
--- a/Engine/Source/Runtime/Experimental/GeometryCollectionEngine/Public/GeometryCollection/GeometryCollectionComponent.h
+++ b/Engine/Source/Runtime/Experimental/GeometryCollectionEngine/Public/GeometryCollection/GeometryCollectionComponent.h
@@ -412,8 +412,12 @@ public:
 	TObjectPtr<AChaosSolverActor> ChaosSolverActor;
 
 	/** RestCollection */
+	UFUNCTION(BlueprintCallable, Category = "ChaosPhysics")
 	void SetRestCollection(const UGeometryCollection * RestCollectionIn);
+
+
 	FORCEINLINE const UGeometryCollection* GetRestCollection() const { return RestCollection; }
+
 	FORCEINLINE FGeometryCollectionEdit EditRestCollection(GeometryCollection::EEditUpdate EditUpdate = GeometryCollection::EEditUpdate::RestPhysics, bool bShapeIsUnchanged = false) { return FGeometryCollectionEdit(this, EditUpdate, bShapeIsUnchanged); }
 #if WITH_EDITOR
 	FORCEINLINE FScopedColorEdit EditBoneSelection(bool bForceUpdate = false) { return FScopedColorEdit(this, bForceUpdate); }
