Add CCDAxisThresholdMode cvar, which has 3 different options for generating ccd axis thresholds.

0 = object bounds
1 = min object bounds for each axis (original behavior... not quite correct)
2 = smallest shape bounds on each axis (correct, except when interacting with terrain or objects which can give in-plane overlaps)
3 = smallest shape bounds on any axis (least performant, most likely to enable CCD)

diff --git a/Engine/Source/Runtime/Experimental/Chaos/Private/Chaos/CCDUtilities.cpp b/Engine/Source/Runtime/Experimental/Chaos/Private/Chaos/CCDUtilities.cpp
--- a/Engine/Source/Runtime/Experimental/Chaos/Private/Chaos/CCDUtilities.cpp
+++ b/Engine/Source/Runtime/Experimental/Chaos/Private/Chaos/CCDUtilities.cpp
@@ -31,6 +31,9 @@ namespace Chaos
 		FRealSingle CCDEnableThresholdBoundsScale = 0.4f;
 		FAutoConsoleVariableRef  CVarCCDEnableThresholdBoundsScale(TEXT("p.Chaos.CCD.EnableThresholdBoundsScale"), CCDEnableThresholdBoundsScale , TEXT("CCD is used when object position is changing > smallest bound's extent * BoundsScale. 0 will always Use CCD. Values < 0 disables CCD."));
 
+		int32 CCDAxisThresholdMode = 1;
+		FAutoConsoleVariableRef  CVarCCDAxisThresholdMode(TEXT("p.Chaos.CCD.AxisThresholdMode"), CCDAxisThresholdMode , TEXT("Change the mode used to generate CCD axis threshold bounds for particle geometries.\n0: Use object bounds\n1: Find the thinnest object bound on any axis and use it for all CCD axes\n2: On each axis, use the thinnest shape bound on that axis\n3: Find the thinnest shape bound on any axis and use this for all axes"));
+
 		extern int32 ChaosSolverDrawCCDInteractions;
 
 #if CHAOS_DEBUG_DRAW
diff --git a/Engine/Source/Runtime/Experimental/Chaos/Public/Chaos/GeometryParticles.h b/Engine/Source/Runtime/Experimental/Chaos/Public/Chaos/GeometryParticles.h
--- a/Engine/Source/Runtime/Experimental/Chaos/Public/Chaos/GeometryParticles.h
+++ b/Engine/Source/Runtime/Experimental/Chaos/Public/Chaos/GeometryParticles.h
@@ -31,6 +31,11 @@ namespace Chaos
 	using FShapesArray = TArray<TUniquePtr<FPerShapeData>, TInlineAllocator<1>>;
 	using FConstraintHandleArray = TArray<FConstraintHandle*>;
 
+	namespace CVars
+	{
+		CHAOS_API extern int32 CCDAxisThresholdMode;
+	}
+
 	/** Data that is associated with geometry. If a union is used an entry is created per internal geometry */
 	class CHAOS_API FPerShapeData
 	{
@@ -638,26 +643,50 @@ namespace Chaos
 			{
 				MLocalBounds[Index] = TAABB<T, d>(InGeometry->BoundingBox());
 
-				// Update CCDAxisThreshold to be the extents of the bounds of the smallest shape on each axis
-				MCCDAxisThreshold[Index] = MLocalBounds[Index].Extents();
-				for (const TUniquePtr<FPerShapeData>& Shape : ShapesArray(Index))
+				if (CVars::CCDAxisThresholdMode == 0)
+				{
+					// Use object extents as CCD axis threshold
+					MCCDAxisThreshold[Index] = MLocalBounds[Index].Extents();
+				}
+				else if (CVars::CCDAxisThresholdMode == 1)
+				{
+					// Use thinnest object extents as all axis CCD thresholds
+					MCCDAxisThreshold[Index] = FVec3(MLocalBounds[Index].Extents().GetMin());
+				}
+				else
 				{
-					// Only sim-enabled shapes should ever be swept with CCD, so make sure the
-					// sim-enabled flag is on for each shape before considering it's min bounds
-					// for CCD extents.
-					if (Shape->GetSimEnabled())
+					// Find minimum shape bounds thickness on each axis
+					FVec3 ThinnestBoundsPerAxis = MLocalBounds[Index].Extents();
+					for (const TUniquePtr<FPerShapeData>& Shape : ShapesArray(Index))
 					{
-						const TSerializablePtr<FImplicitObject> Geometry = Shape->GetGeometry();
-						if (Geometry->HasBoundingBox())
+						// Only sim-enabled shapes should ever be swept with CCD, so make sure the
+						// sim-enabled flag is on for each shape before considering it's min bounds
+						// for CCD extents.
+						if (Shape->GetSimEnabled())
 						{
-							const TVector<T, d> ShapeExtents = Geometry->BoundingBox().Extents();
-							TVector<T, d>& CCDAxisThreshold = MCCDAxisThreshold[Index];
-							for (int32 AxisIndex = 0; AxisIndex < d; ++AxisIndex)
+							const TSerializablePtr<FImplicitObject> Geometry = Shape->GetGeometry();
+							if (Geometry->HasBoundingBox())
 							{
-								MCCDAxisThreshold[Index][AxisIndex] = FMath::Min(ShapeExtents[AxisIndex], MCCDAxisThreshold[Index][AxisIndex]);
+								const TVector<T, d> ShapeExtents = Geometry->BoundingBox().Extents();
+								TVector<T, d>& CCDAxisThreshold = MCCDAxisThreshold[Index];
+								for (int32 AxisIndex = 0; AxisIndex < d; ++AxisIndex)
+								{
+									ThinnestBoundsPerAxis[AxisIndex] = FMath::Min(ShapeExtents[AxisIndex], ThinnestBoundsPerAxis[AxisIndex]);
+								}
 							}
 						}
 					}
+
+					if (CVars::CCDAxisThresholdMode == 2)
+					{
+						// On each axis, use the thinnest shape bound on that axis
+						MCCDAxisThreshold[Index] = ThinnestBoundsPerAxis;
+					}
+					else if (CVars::CCDAxisThresholdMode == 3)
+					{
+						// Find the thinnest shape bound on any axis and use this for all axes
+						MCCDAxisThreshold[Index] = FVec3(ThinnestBoundsPerAxis.GetMin());
+					}
 				}
 
 				// Update the world-space stat of all the shapes - must be called after UpdateShapesArray
