Optimize Gather
Remove inner multithreading in Gather

diff --git a/Engine/Source/Runtime/Experimental/Chaos/Private/Chaos/PBDConstraintRule.cpp b/Engine/Source/Runtime/Experimental/Chaos/Private/Chaos/PBDConstraintRule.cpp
--- a/Engine/Source/Runtime/Experimental/Chaos/Private/Chaos/PBDConstraintRule.cpp
+++ b/Engine/Source/Runtime/Experimental/Chaos/Private/Chaos/PBDConstraintRule.cpp
@@ -202,7 +202,8 @@ namespace Chaos
 								}
 							}
 						}
-						InnerPhysicsParallelFor(IslandConstraints.Num(), [&](const int32 Index)
+
+					        auto Gather = [&](const int32 Index)
 						{
 							FConstraintHandle* ConstraintHandle = IslandConstraints[Index];
 							if (ConstraintHandle->GetContainerId() == GetContainerId())
@@ -217,7 +218,12 @@ namespace Chaos
 									Constraint->GatherInput(Dt, INDEX_NONE, INDEX_NONE, *IslandGroup);
 								}
 							}
-						});
+						};
+						for (int32 Index = 0; Index < IslandConstraints.Num(); ++Index)
+						{
+							Gather(Index);
+						}
+						
 					}
 				}
 			}
diff --git a/Engine/Source/Runtime/Experimental/Chaos/Private/Chaos/PBDRigidsEvolutionGBF.cpp b/Engine/Source/Runtime/Experimental/Chaos/Private/Chaos/PBDRigidsEvolutionGBF.cpp
--- a/Engine/Source/Runtime/Experimental/Chaos/Private/Chaos/PBDRigidsEvolutionGBF.cpp
+++ b/Engine/Source/Runtime/Experimental/Chaos/Private/Chaos/PBDRigidsEvolutionGBF.cpp
@@ -458,7 +458,6 @@ void FPBDRigidsEvolutionGBF::AdvanceOneTimeStepImpl(const FReal Dt, const FSubSt
 	SleepedIslands.SetNum(GetConstraintGraph().NumIslands());
 	TArray<TArray<FPBDRigidParticleHandle*>> DisabledParticles;
 	DisabledParticles.SetNum(GetConstraintGraph().NumIslands());
-	SleepedIslands.SetNum(GetConstraintGraph().NumIslands());
 	if(Dt > 0)
 	{
 		SCOPE_CYCLE_COUNTER(STAT_Evolution_ParallelSolve);
