Fix CIS

diff --git a/Engine/Plugins/Experimental/Dataflow/Source/DataflowEditor/Private/Dataflow/DataflowEditorToolkit.h b/Engine/Plugins/Experimental/Dataflow/Source/DataflowEditor/Private/Dataflow/DataflowEditorToolkit.h
--- a/Engine/Plugins/Experimental/Dataflow/Source/DataflowEditor/Private/Dataflow/DataflowEditorToolkit.h
+++ b/Engine/Plugins/Experimental/Dataflow/Source/DataflowEditor/Private/Dataflow/DataflowEditorToolkit.h
@@ -5,6 +5,7 @@
 #include "CoreMinimal.h"
 #include "Toolkits/SimpleAssetEditor.h"
 #include "Misc/NotifyHook.h"
+#include "GraphEditor.h"
 
 class FEditorViewportTabContent;
 class IDetailsView;
diff --git a/Engine/Plugins/Experimental/Dataflow/Source/DataflowEditor/Public/Dataflow/DataflowEditorPlugin.h b/Engine/Plugins/Experimental/Dataflow/Source/DataflowEditor/Public/Dataflow/DataflowEditorPlugin.h
--- a/Engine/Plugins/Experimental/Dataflow/Source/DataflowEditor/Public/Dataflow/DataflowEditorPlugin.h
+++ b/Engine/Plugins/Experimental/Dataflow/Source/DataflowEditor/Public/Dataflow/DataflowEditorPlugin.h
@@ -6,6 +6,7 @@
 #include "Modules/ModuleInterface.h"
 #include "Modules/ModuleManager.h"
 //#include "Toolkits/IToolkit.h"
+#include "Toolkits/AssetEditorToolkit.h"
 
 //class ISlateStyle;
 //class FSlateStyleSet;
diff --git a/Engine/Plugins/Experimental/Dataflow/Source/DataflowEditor/Public/Dataflow/DataflowSchema.h b/Engine/Plugins/Experimental/Dataflow/Source/DataflowEditor/Public/Dataflow/DataflowSchema.h
--- a/Engine/Plugins/Experimental/Dataflow/Source/DataflowEditor/Public/Dataflow/DataflowSchema.h
+++ b/Engine/Plugins/Experimental/Dataflow/Source/DataflowEditor/Public/Dataflow/DataflowSchema.h
@@ -1,6 +1,7 @@
 // Copyright Epic Games, Inc. All Rights Reserved.
 
 #pragma once
+#include "EdGraph/EdGraphSchema.h"
 
 #include "DataflowSchema.generated.h"
 
diff --git a/Engine/Source/Runtime/Experimental/Dataflow/Core/Public/Dataflow/DataflowNodeParameters.h b/Engine/Source/Runtime/Experimental/Dataflow/Core/Public/Dataflow/DataflowNodeParameters.h
--- a/Engine/Source/Runtime/Experimental/Dataflow/Core/Public/Dataflow/DataflowNodeParameters.h
+++ b/Engine/Source/Runtime/Experimental/Dataflow/Core/Public/Dataflow/DataflowNodeParameters.h
@@ -5,7 +5,6 @@
 #include "CoreMinimal.h"
 
 class  UDataflow;
-class  UEdGraphPin;
 
 namespace Dataflow
 {
diff --git a/Engine/Source/Runtime/Experimental/Dataflow/Engine/Private/Dataflow/DataflowEdNode.cpp b/Engine/Source/Runtime/Experimental/Dataflow/Engine/Private/Dataflow/DataflowEdNode.cpp
--- a/Engine/Source/Runtime/Experimental/Dataflow/Engine/Private/Dataflow/DataflowEdNode.cpp
+++ b/Engine/Source/Runtime/Experimental/Dataflow/Engine/Private/Dataflow/DataflowEdNode.cpp
@@ -6,11 +6,15 @@
 #include "Dataflow/Dataflow.h"
 #include "Logging/LogMacros.h"
 
+#if WITH_EDITOR && !(UE_BUILD_SHIPPING || UE_BUILD_TEST)
+#include "EdGraph/EdGraphPin.h"
+#endif
+
+
 #define LOCTEXT_NAMESPACE "DataflowEdNode"
 
 DEFINE_LOG_CATEGORY_STATIC(DATAFLOWNODE_LOG, Error, All);
 
-
 void UDataflowEdNode::AllocateDefaultPins()
 {
 	UE_LOG(DATAFLOWNODE_LOG, Verbose, TEXT("UDataflowEdNode::AllocateDefaultPins()"));
