Chaos - collision testing stuff

diff --git a/Engine/Source/Runtime/Experimental/Chaos/Private/Chaos/HeightField.cpp b/Engine/Source/Runtime/Experimental/Chaos/Private/Chaos/HeightField.cpp
--- a/Engine/Source/Runtime/Experimental/Chaos/Private/Chaos/HeightField.cpp
+++ b/Engine/Source/Runtime/Experimental/Chaos/Private/Chaos/HeightField.cpp
@@ -16,6 +16,8 @@
 #include "Chaos/Triangle.h"
 #include "Chaos/TriangleRegister.h"
 
+//PRAGMA_DISABLE_OPTIMIZATION
+
 namespace Chaos
 {
 	extern FRealSingle Chaos_Collision_EdgePrunePlaneDistance;
@@ -1778,11 +1780,11 @@ namespace Chaos
 				const FReal ApproximateDistToObject = FVec3::DistSquared(QueryTM.GetLocation(), A);
 				const FReal SweepLength = ApproximateSizeOfObject + ApproximateDistToObject;
 				const FRigidTransform3 QueryStartTM(QueryTM.GetLocation() + TriNormal * SweepLength, QueryTM.GetRotation());
-				if (GJKRaycast2(TriangleConvexReg, QueryGeom, QueryStartTM, -TriNormal, SweepLength, Penetration, ClosestB, Normal, (FReal)0., true))
+				if (GJKRaycast2(TriangleConvexReg, QueryGeom, QueryStartTM, -TriNormal, SweepLength, Penetration, ClosestA, Normal, (FReal)0., true))
 				{
-					LocalContactLocation = ClosestB;
-					LocalContactNormal = TriNormal;
 					LocalContactPhi = Penetration - SweepLength;
+					LocalContactLocation = ClosestA + LocalContactPhi * Normal;
+					LocalContactNormal = TriNormal;
 					return true;
 				}
 				return false;
diff --git a/Engine/Source/Runtime/Experimental/Chaos/Private/Chaos/PBDRigidsEvolutionGBF.cpp b/Engine/Source/Runtime/Experimental/Chaos/Private/Chaos/PBDRigidsEvolutionGBF.cpp
--- a/Engine/Source/Runtime/Experimental/Chaos/Private/Chaos/PBDRigidsEvolutionGBF.cpp
+++ b/Engine/Source/Runtime/Experimental/Chaos/Private/Chaos/PBDRigidsEvolutionGBF.cpp
@@ -142,6 +142,12 @@ namespace Chaos
 
 		bool bChaos_CollisionStore_Enabled = true;
 		FAutoConsoleVariableRef CVarCollisionStoreEnabled(TEXT("p.Chaos.CollisionStore.Enabled"), bChaos_CollisionStore_Enabled, TEXT(""));
+
+		// Put the solver into a mode where it reset particles to their initial positions each frame.
+		// This is used to test collision detection and - it will be removed
+		// @chaos(todo): remove this when no longer needed
+		bool bChaos_Solver_TestMode  = false;
+		FAutoConsoleVariableRef CVarChaosSolverTestMode(TEXT("p.Chaos.Solver.TestMode"), bChaos_Solver_TestMode, TEXT(""));
 	}
 
 	using namespace CVars;
@@ -370,6 +376,12 @@ void FPBDRigidsEvolutionGBF::AdvanceOneTimeStepImpl(const FReal Dt, const FSubSt
 		Integrate(Particles.GetActiveParticlesView(), Dt);
 	}
 
+	if (bChaos_Solver_TestMode)
+	{
+		TestModeResetParticles();
+		TestModeResetCollisions();
+	}
+
 	{
 		SCOPE_CYCLE_COUNTER(STAT_Evolution_KinematicTargets);
 		CSV_SCOPED_TIMING_STAT(PhysicsVerbose, StepSolver_KinematicTargets);
@@ -615,6 +627,11 @@ void FPBDRigidsEvolutionGBF::AdvanceOneTimeStepImpl(const FReal Dt, const FSubSt
 
 	ParticleUpdatePosition(Particles.GetDirtyParticlesView(), Dt);
 
+	if (bChaos_Solver_TestMode)
+	{
+		TestModeResetParticles();
+	}
+
 	if (CVars::DoFinalProbeNarrowPhase)
 	{
 		// Run contact updates on probe constraints
@@ -790,6 +807,44 @@ void FPBDRigidsEvolutionGBF::SetIsDeterministic(const bool bInIsDeterministic)
 	CollisionConstraints.SetIsDeterministic(bInIsDeterministic);
 }
 
+
+void FPBDRigidsEvolutionGBF::TestModeResetParticles()
+{
+	for (auto& Rigid : Particles.GetNonDisabledDynamicView())
+	{
+		FTestModeParticleData* Data = TestModeData.Find(Rigid.Handle());
+		if (Data != nullptr)
+		{
+			Rigid.X() = Data->X;
+			Rigid.P() = Data->P;
+			Rigid.R() = Data->R;
+			Rigid.Q() = Data->Q;
+			Rigid.V() = Data->V;
+			Rigid.W() = Data->W;
+		}
+		if (Data == nullptr)
+		{
+			Data = &TestModeData.Add(Rigid.Handle());
+			Data->X = Rigid.X();
+			Data->P = Rigid.P();
+			Data->R = Rigid.R();
+			Data->Q = Rigid.Q();
+			Data->V = Rigid.V();
+			Data->W = Rigid.W();
+		}
+	}
+}
+
+void FPBDRigidsEvolutionGBF::TestModeResetCollisions()
+{
+	for (FPBDCollisionConstraintHandle* Collision : CollisionConstraints.GetConstraintHandles())
+	{
+		Collision->GetContact().ResetManifold();
+		Collision->GetContact().ResetMaterial();
+		Collision->GetContact().GetGJKWarmStartData().Reset();
+	}
+}
+
 FPBDRigidsEvolutionGBF::FPBDRigidsEvolutionGBF(FPBDRigidsSOAs& InParticles,THandleArray<FChaosPhysicsMaterial>& SolverPhysicsMaterials, const TArray<ISimCallbackObject*>* InCollisionModifiers, bool InIsSingleThreaded)
 	: Base(InParticles, SolverPhysicsMaterials, InIsSingleThreaded)
 	, Clustering(*this, Particles.GetClusteredParticles())
diff --git a/Engine/Source/Runtime/Experimental/Chaos/Public/Chaos/PBDRigidsEvolutionGBF.h b/Engine/Source/Runtime/Experimental/Chaos/Public/Chaos/PBDRigidsEvolutionGBF.h
--- a/Engine/Source/Runtime/Experimental/Chaos/Public/Chaos/PBDRigidsEvolutionGBF.h
+++ b/Engine/Source/Runtime/Experimental/Chaos/Public/Chaos/PBDRigidsEvolutionGBF.h
@@ -361,7 +361,6 @@ namespace Chaos
 		void ScatterSolverOutput(FReal Dt, int32 GroupIndex);
 
 		void UpdateInertiaConditioning();
-		void UpdateParticleInertiaConditioning(FPBDRigidParticleHandle* Rigid, const FReal MaxDistance, const FReal MaxRotationRatio);
 
 		FEvolutionResimCache* GetCurrentStepResimCache()
 		{
@@ -394,6 +393,19 @@ namespace Chaos
 		FCCDManager CCDManager;
 
 		bool bIsDeterministic;
+
+		// Test Mode for Collision issues
+		// @todo(chaos): remove this when no longer needed
+		void TestModeResetParticles();
+		void TestModeResetCollisions();
+
+		struct FTestModeParticleData
+		{
+			FVec3 X, P, V, W;
+			FRotation3 R, Q;
+		};
+		TMap<FPBDRigidParticleHandle*, FTestModeParticleData> TestModeData;
+		// End Test Mode
 	};
 
 }
