fix NoPCH compilation (missing header)

diff --git a/Engine/Source/Runtime/Experimental/Animation/Constraints/Private/ConstraintsManager.cpp b/Engine/Source/Runtime/Experimental/Animation/Constraints/Private/ConstraintsManager.cpp
--- a/Engine/Source/Runtime/Experimental/Animation/Constraints/Private/ConstraintsManager.cpp
+++ b/Engine/Source/Runtime/Experimental/Animation/Constraints/Private/ConstraintsManager.cpp
@@ -7,6 +7,7 @@
 #include "Algo/Copy.h"
 #include "Algo/StableSort.h"
 
+#include "Engine/World.h"
 
 /** 
  * FConstraintTickFunction
