fix NoPCH compilation

diff --git a/Engine/Source/Runtime/Experimental/Animation/Constraints/Private/ConstraintsManager.cpp b/Engine/Source/Runtime/Experimental/Animation/Constraints/Private/ConstraintsManager.cpp
--- a/Engine/Source/Runtime/Experimental/Animation/Constraints/Private/ConstraintsManager.cpp
+++ b/Engine/Source/Runtime/Experimental/Animation/Constraints/Private/ConstraintsManager.cpp
@@ -8,6 +8,7 @@
 #include "Algo/StableSort.h"
 
 #include "Engine/World.h"
+#include "Engine/Level.h"
 
 /** 
  * FConstraintTickFunction
diff --git a/Engine/Source/Runtime/Experimental/Animation/Constraints/Private/TransformableHandle.cpp b/Engine/Source/Runtime/Experimental/Animation/Constraints/Private/TransformableHandle.cpp
--- a/Engine/Source/Runtime/Experimental/Animation/Constraints/Private/TransformableHandle.cpp
+++ b/Engine/Source/Runtime/Experimental/Animation/Constraints/Private/TransformableHandle.cpp
@@ -3,6 +3,9 @@
 
 #include "TransformableHandle.h"
 
+#include "Components/SceneComponent.h"
+#include "GameFramework/Actor.h"
+
 /**
  * UTransformableHandle
  */
diff --git a/Engine/Source/Runtime/Experimental/Animation/Constraints/Private/TransformableRegistry.cpp b/Engine/Source/Runtime/Experimental/Animation/Constraints/Private/TransformableRegistry.cpp
--- a/Engine/Source/Runtime/Experimental/Animation/Constraints/Private/TransformableRegistry.cpp
+++ b/Engine/Source/Runtime/Experimental/Animation/Constraints/Private/TransformableRegistry.cpp
@@ -2,6 +2,8 @@
 
 #include "TransformableRegistry.h"
 
+#include "UObject/Class.h"
+
 FTransformableRegistry::~FTransformableRegistry() = default;
 
 FTransformableRegistry& FTransformableRegistry::Get()
diff --git a/Engine/Source/Runtime/Experimental/Animation/Constraints/Public/ConstraintsManager.h b/Engine/Source/Runtime/Experimental/Animation/Constraints/Public/ConstraintsManager.h
--- a/Engine/Source/Runtime/Experimental/Animation/Constraints/Public/ConstraintsManager.h
+++ b/Engine/Source/Runtime/Experimental/Animation/Constraints/Public/ConstraintsManager.h
@@ -3,10 +3,12 @@
 #pragma once
 
 #include "UObject/Object.h"
+#include "Engine/EngineBaseTypes.h"
 
 #include "ConstraintsManager.generated.h"
 
 class UTickableConstraint;
+class ULevel;
 
 /** 
  * FConstraintTickFunction
diff --git a/Engine/Source/Runtime/Experimental/Animation/Constraints/Public/TransformableRegistry.h b/Engine/Source/Runtime/Experimental/Animation/Constraints/Public/TransformableRegistry.h
--- a/Engine/Source/Runtime/Experimental/Animation/Constraints/Public/TransformableRegistry.h
+++ b/Engine/Source/Runtime/Experimental/Animation/Constraints/Public/TransformableRegistry.h
@@ -2,8 +2,12 @@
 
 #pragma once
 
+#include "Templates/Function.h"
+#include "Containers/Map.h"
+
 class UTransformableHandle;
 class UClass;
+class UObject;
 
 /**
  * TransformableRegistry
