fixed missing include (NoPCH compilation error)

diff --git a/Engine/Source/Runtime/Experimental/Animation/Constraints/Private/TransformConstraint.cpp b/Engine/Source/Runtime/Experimental/Animation/Constraints/Private/TransformConstraint.cpp
--- a/Engine/Source/Runtime/Experimental/Animation/Constraints/Private/TransformConstraint.cpp
+++ b/Engine/Source/Runtime/Experimental/Animation/Constraints/Private/TransformConstraint.cpp
@@ -5,6 +5,7 @@
 #include "Animation/Constraints/Public/TransformableHandle.h"
 #include "ConstraintsManager.h"
 #include "TransformableRegistry.h"
+#include "GameFramework/Actor.h"
 
 /** 
  * UTickableTransformConstraint
@@ -19,10 +20,7 @@ int64 UTickableTransformConstraint::GetType() const
 #if WITH_EDITOR
 FName UTickableTransformConstraint::GetLabel() const
 {
-	const FName ParentName = ParentTRSHandle->IsValid() ? ParentTRSHandle->GetName() : NAME_None;
-	const FName ChildName = ChildTRSHandle->IsValid() ? ChildTRSHandle->GetName() : NAME_None;
-	const FString Label = FString::Printf(TEXT("%s -> %s"), *ParentName.ToString(), *ChildName.ToString() );
-	return FName(*Label);
+	return ParentTRSHandle->IsValid() ? ParentTRSHandle->GetName() : NAME_None;
 }
 #endif
 
