Chaos : add cvar to control boudning box filtering margin for cluster group connection graph generation

diff --git a/Engine/Source/Runtime/Experimental/Chaos/Private/Chaos/PBDRigidClustering.cpp b/Engine/Source/Runtime/Experimental/Chaos/Private/Chaos/PBDRigidClustering.cpp
--- a/Engine/Source/Runtime/Experimental/Chaos/Private/Chaos/PBDRigidClustering.cpp
+++ b/Engine/Source/Runtime/Experimental/Chaos/Private/Chaos/PBDRigidClustering.cpp
@@ -43,7 +43,10 @@ namespace Chaos
 	FAutoConsoleVariableRef CVarDeactivateClusterChildren(TEXT("p.DeactivateClusterChildren"), DeactivateClusterChildren, TEXT("If children should be decativated when broken and put into another cluster."));
 
 	int32 UseBoundingBoxForConnectionGraphFiltering = 0;
-	FAutoConsoleVariableRef CVarUseBoundingBoxForConnectionGraphFiltering(TEXT("p.UseBoundingBoxForConnectionGraphFiltering"), UseBoundingBoxForConnectionGraphFiltering, TEXT("when on, use bounding box overlaps to filter connection during the connection graph generation"));
+	FAutoConsoleVariableRef CVarUseBoundingBoxForConnectionGraphFiltering(TEXT("p.UseBoundingBoxForConnectionGraphFiltering"), UseBoundingBoxForConnectionGraphFiltering, TEXT("when on, use bounding box overlaps to filter connection during the connection graph generation [def: 0]"));
+
+	float BoundingBoxMarginForConnectionGraphFiltering = 0;
+	FAutoConsoleVariableRef CVarBoundingBoxMarginForConnectionGraphFiltering(TEXT("p.BoundingBoxMarginForConnectionGraphFiltering"), BoundingBoxMarginForConnectionGraphFiltering, TEXT("when UseBoundingBoxForConnectionGraphFiltering is on, the margin to use for the oevrlap test [def: 0]"));
 
 	//==========================================================================
 	// TPBDRigidClustering
@@ -1370,7 +1373,9 @@ namespace Chaos
 	{
 		if (ensure(Child1 != nullptr && Child2 != nullptr ))
 		{
-			return Child1->WorldSpaceInflatedBounds().Intersects(Child2->WorldSpaceInflatedBounds());
+			FAABB3 Bounds1 = Child1->WorldSpaceInflatedBounds();
+			Bounds1.Thicken(BoundingBoxMarginForConnectionGraphFiltering);
+			return Bounds1.Intersects(Child2->WorldSpaceInflatedBounds());
 		}
 		return false;
 	}
