Chaos destruction : add cvar for relative size attribute computation

diff --git a/Engine/Source/Runtime/Experimental/Chaos/Private/GeometryCollection/GeometryCollectionConvexUtility.cpp b/Engine/Source/Runtime/Experimental/Chaos/Private/GeometryCollection/GeometryCollectionConvexUtility.cpp
--- a/Engine/Source/Runtime/Experimental/Chaos/Private/GeometryCollection/GeometryCollectionConvexUtility.cpp
+++ b/Engine/Source/Runtime/Experimental/Chaos/Private/GeometryCollection/GeometryCollectionConvexUtility.cpp
@@ -10,6 +10,13 @@
 #include "Templates/Sorting.h"
 #include "Spatial/PointHashGrid3.h"
 
+bool UseVolumeToComputeRelativeSize = false;
+FAutoConsoleVariableRef CVarUseVolumeToComputeRelativeSize(TEXT("p.gc.UseVolumeToComputeRelativeSize"), UseVolumeToComputeRelativeSize, TEXT("Use Volume To Compute RelativeSize instead of the side of the cubic volume (def: false)"));
+
+bool UseLargestClusterToComputeRelativeSize = false;
+FAutoConsoleVariableRef CVarUseMaxClusterToComputeRelativeSize(TEXT("p.gc.UseLargestClusterToComputeRelativeSize"), UseVolumeToComputeRelativeSize, TEXT("Use the largest Cluster as reference for the releative size instead of the largest child (def: false)"));
+
+
 TOptional<FGeometryCollectionConvexUtility::FGeometryCollectionConvexData> FGeometryCollectionConvexUtility::GetConvexHullDataIfPresent(FGeometryCollection* GeometryCollection)
 {
 	check(GeometryCollection);
@@ -1408,18 +1415,19 @@ void FGeometryCollectionConvexUtility::CopyChildConvexes(const FGeometryCollecti
 	}
 }
 
-void FGeometryCollectionConvexUtility::SetVolumeAttributes(FGeometryCollection* Collection)
+static float GetRelativeSizeDimensionFromVolume(const float Volume)
 {
-	if (!Collection->HasAttribute("Volume", FGeometryCollection::TransformGroup))
+	if (UseVolumeToComputeRelativeSize)
 	{
-		Collection->AddAttribute<float>("Volume", FGeometryCollection::TransformGroup);
+		return Volume;
 	}
-	if (!Collection->HasAttribute("Size", FGeometryCollection::TransformGroup))
-	{
-		Collection->AddAttribute<float>("Size", FGeometryCollection::TransformGroup);
-	}
-	TManagedArray<float>& Volumes = Collection->ModifyAttribute<float>("Volume", FTransformCollection::TransformGroup);
-	TManagedArray<float>& Sizes = Collection->ModifyAttribute<float>("Size", FTransformCollection::TransformGroup);
+	return FGenericPlatformMath::Pow(Volume, 1.0f / 3.0f);
+}
+
+void FGeometryCollectionConvexUtility::SetVolumeAttributes(FGeometryCollection* Collection)
+{
+	TManagedArray<float>& Volumes = Collection->AddAttribute<float>("Volume", FTransformCollection::TransformGroup);
+	TManagedArray<float>& Sizes = Collection->AddAttribute<float>("Size", FTransformCollection::TransformGroup);
 
 	const TManagedArray<int32>& SimulationType = Collection->SimulationType;
 	const TManagedArray<int32>& TransformToGeometryIndex = Collection->TransformToGeometryIndex;
@@ -1430,7 +1438,7 @@ void FGeometryCollectionConvexUtility::SetVolumeAttributes(FGeometryCollection*
 	TArray<int32> RecursiveOrder = GeometryCollectionAlgo::ComputeRecursiveOrder(*Collection);
 	float MaxGeoVolume = 0.0f;
 	Volumes.Fill(0.0f);
-	for (int32 Bone : RecursiveOrder)
+	for (const int32 Bone : RecursiveOrder)
 	{
 		if (SimulationType[Bone] == FGeometryCollection::ESimulationTypes::FST_Rigid)
 		{
@@ -1442,22 +1450,33 @@ void FGeometryCollectionConvexUtility::SetVolumeAttributes(FGeometryCollection*
 			}
 			else
 			{
-				float GeoVolume = (float)ComputeGeometryVolume(Collection, GeoIdx, Transforms[Bone], 1.0);
+				const float GeoVolume = (float)ComputeGeometryVolume(Collection, GeoIdx, Transforms[Bone], 1.0);
 				MaxGeoVolume = FMath::Max(MaxGeoVolume, GeoVolume);
 				Volumes[Bone] = GeoVolume;
 			}
 		}
-		int32 ParentIdx = Parent[Bone];
+		const int32 ParentIdx = Parent[Bone];
 		if (ParentIdx != INDEX_NONE)
 		{
 			Volumes[ParentIdx] += Volumes[Bone];
 		}
 	}
-	float SizeScaleFactor = MaxGeoVolume > 0 ? (1.0f / FGenericPlatformMath::Pow(MaxGeoVolume, (1.0f / 3.0f))) : 1.0f;
+
+	if (UseLargestClusterToComputeRelativeSize)
+	{
+		// just go over all the bones as the largest clusters will be naturally larger in volume than any of the children
+		for (int32 BoneIdx = 0; BoneIdx < Volumes.Num(); BoneIdx++)
+		{
+			MaxGeoVolume = FMath::Max(MaxGeoVolume, Volumes[BoneIdx]);
+		}
+	}
+	
+	const float ReferenceSize = GetRelativeSizeDimensionFromVolume(MaxGeoVolume);
+	const float OneOverReferenceSize = MaxGeoVolume > 0 ? (1.0f / ReferenceSize) : 1.0f;
 	for (int32 BoneIdx = 0; BoneIdx < Volumes.Num(); BoneIdx++)
 	{
-		float Rt3Volume = FGenericPlatformMath::Pow(Volumes[BoneIdx], 1.0f / 3.0f);
-		Sizes[BoneIdx] = Rt3Volume * SizeScaleFactor;
+		const float ActualSize = GetRelativeSizeDimensionFromVolume(Volumes[BoneIdx]);
+		Sizes[BoneIdx] = ActualSize * OneOverReferenceSize;
 	}
 }
 
