ChaosFlesh : Dataflow integration
- Added Dataflow support

diff --git a/Engine/Plugins/Experimental/Dataflow/Source/DataflowEditor/Private/Dataflow/DataflowEditorToolkit.cpp b/Engine/Plugins/Experimental/Dataflow/Source/DataflowEditor/Private/Dataflow/DataflowEditorToolkit.cpp
--- a/Engine/Plugins/Experimental/Dataflow/Source/DataflowEditor/Private/Dataflow/DataflowEditorToolkit.cpp
+++ b/Engine/Plugins/Experimental/Dataflow/Source/DataflowEditor/Private/Dataflow/DataflowEditorToolkit.cpp
@@ -7,7 +7,7 @@
 #include "Dataflow/DataflowNodeFactory.h"
 #include "Dataflow/DataflowObject.h"
 #include "Dataflow/DataflowSchema.h"
-#include "Dataflow/Dataflow.h"
+#include "Dataflow/DataflowCore.h"
 #include "EditorStyleSet.h"
 #include "EditorViewportTabContent.h"
 #include "EditorViewportLayout.h"
diff --git a/Engine/Plugins/Experimental/Dataflow/Source/DataflowEditor/Private/Dataflow/DataflowSNode.cpp b/Engine/Plugins/Experimental/Dataflow/Source/DataflowEditor/Private/Dataflow/DataflowSNode.cpp
--- a/Engine/Plugins/Experimental/Dataflow/Source/DataflowEditor/Private/Dataflow/DataflowSNode.cpp
+++ b/Engine/Plugins/Experimental/Dataflow/Source/DataflowEditor/Private/Dataflow/DataflowSNode.cpp
@@ -5,7 +5,7 @@
 #include "Dataflow/DataflowEdNode.h"
 #include "Dataflow/DataflowNodeFactory.h"
 #include "Dataflow/DataflowObject.h"
-#include "Dataflow/Dataflow.h"
+#include "Dataflow/DataflowCore.h"
 #include "Logging/LogMacros.h"
 #include "Widgets/Layout/SBorder.h"
 
diff --git a/Engine/Plugins/Experimental/Dataflow/Source/DataflowEditor/Public/Dataflow/DataflowEditor.h b/Engine/Plugins/Experimental/Dataflow/Source/DataflowEditor/Public/Dataflow/DataflowEditor.h
new file mode 100644
--- /dev/null
+++ b/Engine/Plugins/Experimental/Dataflow/Source/DataflowEditor/Public/Dataflow/DataflowEditor.h
@@ -0,0 +1,12 @@
+// Copyright Epic Games, Inc. All Rights Reserved.
+
+#pragma once
+
+#include "Dataflow/DataflowEngine.h"
+#include "Dataflow/DataflowAssetActions.h"
+#include "Dataflow/DataflowAssetFactory.h"
+#include "Dataflow/DataflowEditorActions.h"
+#include "Dataflow/DataflowEditorPlugin.h"
+#include "Dataflow/DataflowSchema.h"
+#include "Dataflow/DataflowSNode.h"
+#include "Dataflow/DataflowSNodeFactories.h"
diff --git a/Engine/Source/Runtime/Experimental/Dataflow/Core/Public/Dataflow/Dataflow.h b/Engine/Source/Runtime/Experimental/Dataflow/Core/Public/Dataflow/Dataflow.h
deleted file mode 100644
diff --git a/Engine/Source/Runtime/Experimental/Dataflow/Core/Public/Dataflow/DataflowCore.h b/Engine/Source/Runtime/Experimental/Dataflow/Core/Public/Dataflow/DataflowCore.h
new file mode 100644
--- /dev/null
+++ b/Engine/Source/Runtime/Experimental/Dataflow/Core/Public/Dataflow/DataflowCore.h
@@ -0,0 +1,10 @@
+// Copyright Epic Games, Inc. All Rights Reserved.
+
+#pragma once
+
+#include "Dataflow/DataflowGraph.h"
+#include "Dataflow/DataflowConnection.h"
+#include "Dataflow/DataflowInputOutput.h"
+#include "Dataflow/DataflowNode.h"
+#include "Dataflow/DataflowNodeFactory.h"
+#include "Dataflow/DataflowNodeParameters.h"
\ No newline at end of file
diff --git a/Engine/Source/Runtime/Experimental/Dataflow/Core/Public/Dataflow/DataflowNode.h b/Engine/Source/Runtime/Experimental/Dataflow/Core/Public/Dataflow/DataflowNode.h
--- a/Engine/Source/Runtime/Experimental/Dataflow/Core/Public/Dataflow/DataflowNode.h
+++ b/Engine/Source/Runtime/Experimental/Dataflow/Core/Public/Dataflow/DataflowNode.h
@@ -46,7 +46,7 @@ namespace Dataflow
 		FConnection* FindInput(FName Name) const;
 		FConnection* FindOutput(FName Name) const;
 
-		virtual void Evaluate(const FContext& Context, FConnection*) { ensure(false); }
+		virtual void Evaluate(const FContext& Context, FConnection*) const { ensure(false); }
 		void InvalidateOutputs();
 
 		virtual void SerializeInternal(FArchive& Ar) {};
diff --git a/Engine/Source/Runtime/Experimental/Dataflow/Core/Public/Dataflow/DataflowNodeParameters.h b/Engine/Source/Runtime/Experimental/Dataflow/Core/Public/Dataflow/DataflowNodeParameters.h
--- a/Engine/Source/Runtime/Experimental/Dataflow/Core/Public/Dataflow/DataflowNodeParameters.h
+++ b/Engine/Source/Runtime/Experimental/Dataflow/Core/Public/Dataflow/DataflowNodeParameters.h
@@ -8,8 +8,11 @@ class  UDataflow;
 
 namespace Dataflow
 {
-	struct DATAFLOWCORE_API FContext
+	class DATAFLOWCORE_API FContext
 	{
+	public:
+		FContext(float InTime) : Timestamp(InTime) {}
+
 		float Timestamp = 0.f;
 
 		uint32 GetTypeHash() const
diff --git a/Engine/Source/Runtime/Experimental/Dataflow/Engine/DataflowEngine.Build.cs b/Engine/Source/Runtime/Experimental/Dataflow/Engine/DataflowEngine.Build.cs
--- a/Engine/Source/Runtime/Experimental/Dataflow/Engine/DataflowEngine.Build.cs
+++ b/Engine/Source/Runtime/Experimental/Dataflow/Engine/DataflowEngine.Build.cs
@@ -13,9 +13,10 @@ namespace UnrealBuildTool.Rules
 					"Core",
 					"CoreUObject",
 					"Engine",
-					"DataflowCore"
+					"DataflowCore",
+					"Chaos"
 				}
-				);
+			);
 		}
 	}
 }
diff --git a/Engine/Source/Runtime/Experimental/Dataflow/Engine/Private/Dataflow/DataflowEdNode.cpp b/Engine/Source/Runtime/Experimental/Dataflow/Engine/Private/Dataflow/DataflowEdNode.cpp
--- a/Engine/Source/Runtime/Experimental/Dataflow/Engine/Private/Dataflow/DataflowEdNode.cpp
+++ b/Engine/Source/Runtime/Experimental/Dataflow/Engine/Private/Dataflow/DataflowEdNode.cpp
@@ -3,7 +3,7 @@
 #include "Dataflow/DataflowEdNode.h"
 
 #include "Dataflow/DataflowNode.h"
-#include "Dataflow/Dataflow.h"
+#include "Dataflow/DataflowCore.h"
 #include "Logging/LogMacros.h"
 
 #if WITH_EDITOR && !(UE_BUILD_SHIPPING || UE_BUILD_TEST)
diff --git a/Engine/Source/Runtime/Experimental/Dataflow/Engine/Private/Dataflow/DataflowObject.cpp b/Engine/Source/Runtime/Experimental/Dataflow/Engine/Private/Dataflow/DataflowObject.cpp
--- a/Engine/Source/Runtime/Experimental/Dataflow/Engine/Private/Dataflow/DataflowObject.cpp
+++ b/Engine/Source/Runtime/Experimental/Dataflow/Engine/Private/Dataflow/DataflowObject.cpp
@@ -1,17 +1,43 @@
 // Copyright Epic Games, Inc. All Rights Reserved.
 
 #include "Dataflow/DataflowObject.h"
-#include "Dataflow/Dataflow.h"
+#include "Dataflow/DataflowCore.h"
 #include "Dataflow/DataflowEdNode.h"
 
 
-#define LOCTEXT_NAMESPACE "UDataflow")
+#define LOCTEXT_NAMESPACE "UDataflow"
+
+FDataflowAssetEdit::FDataflowAssetEdit(UDataflow* InAsset, FPostEditFunctionCallback InCallback)
+	: PostEditCallback(InCallback)
+	, Asset(InAsset)
+{
+}
+
+FDataflowAssetEdit::~FDataflowAssetEdit()
+{
+	PostEditCallback();
+}
+
+Dataflow::FGraph* FDataflowAssetEdit::GetGraph()
+{
+	if (Asset)
+	{
+		return Asset->Dataflow.Get();
+	}
+	return nullptr;
+}
 
 UDataflow::UDataflow(const FObjectInitializer& ObjectInitializer)
 	: UEdGraph(ObjectInitializer)
 	, Dataflow(new Dataflow::FGraph())
 {}
 
+
+void UDataflow::PostEditCallback()
+{
+	// mark as dirty for the UObject
+}
+
 #if WITH_EDITOR
 
 void UDataflow::PostEditChangeProperty(struct FPropertyChangedEvent& PropertyChangedEvent)
diff --git a/Engine/Source/Runtime/Experimental/Dataflow/Engine/Private/Dataflow/DataflowObjectInterface.cpp b/Engine/Source/Runtime/Experimental/Dataflow/Engine/Private/Dataflow/DataflowObjectInterface.cpp
new file mode 100644
--- /dev/null
+++ b/Engine/Source/Runtime/Experimental/Dataflow/Engine/Private/Dataflow/DataflowObjectInterface.cpp
@@ -0,0 +1,11 @@
+// Copyright Epic Games, Inc. All Rights Reserved.
+
+#include "Dataflow/DataflowObjectInterface.h"
+
+namespace Dataflow
+{
+	FEngineContext::FEngineContext(UDataflow* InGraph, float InTime)
+			: FContext(InTime)
+			, Graph(InGraph) 
+		{}
+}
\ No newline at end of file
diff --git a/Engine/Source/Runtime/Experimental/Dataflow/Engine/Public/Dataflow/DataflowEdNode.h b/Engine/Source/Runtime/Experimental/Dataflow/Engine/Public/Dataflow/DataflowEdNode.h
--- a/Engine/Source/Runtime/Experimental/Dataflow/Engine/Public/Dataflow/DataflowEdNode.h
+++ b/Engine/Source/Runtime/Experimental/Dataflow/Engine/Public/Dataflow/DataflowEdNode.h
@@ -4,7 +4,7 @@
 
 #include "CoreMinimal.h"
 #include "EdGraph/EdGraph.h"
-#include "Dataflow/Dataflow.h"
+#include "Dataflow/DataflowCore.h"
 
 #include "DataflowEdNode.generated.h"
 
diff --git a/Engine/Source/Runtime/Experimental/Dataflow/Engine/Public/Dataflow/DataflowEngine.h b/Engine/Source/Runtime/Experimental/Dataflow/Engine/Public/Dataflow/DataflowEngine.h
new file mode 100644
--- /dev/null
+++ b/Engine/Source/Runtime/Experimental/Dataflow/Engine/Public/Dataflow/DataflowEngine.h
@@ -0,0 +1,8 @@
+// Copyright Epic Games, Inc. All Rights Reserved.
+
+#pragma once
+
+#include "Dataflow/DataflowCore.h"
+#include "Dataflow/DataflowObjectInterface.h"
+#include "Dataflow/DataflowEdNode.h"
+#include "Dataflow/DataflowObject.h"
diff --git a/Engine/Source/Runtime/Experimental/Dataflow/Engine/Public/Dataflow/DataflowObject.h b/Engine/Source/Runtime/Experimental/Dataflow/Engine/Public/Dataflow/DataflowObject.h
--- a/Engine/Source/Runtime/Experimental/Dataflow/Engine/Public/Dataflow/DataflowObject.h
+++ b/Engine/Source/Runtime/Experimental/Dataflow/Engine/Public/Dataflow/DataflowObject.h
@@ -4,10 +4,37 @@
 
 #include "CoreMinimal.h"
 #include "EdGraph/EdGraph.h"
-#include "Dataflow/Dataflow.h"
+#include "Dataflow/DataflowCore.h"
 
 #include "DataflowObject.generated.h"
 
+
+/**
+*	FFleshAssetEdit
+*     Structured RestCollection access where the scope
+*     of the object controls serialization back into the
+*     dynamic collection
+*
+*/
+class DATAFLOWENGINE_API FDataflowAssetEdit
+{
+public:
+	typedef TFunctionRef<void()> FPostEditFunctionCallback;
+	friend UDataflow;
+
+	/**
+	 * @param UDataflow				The "FAsset" to edit
+	 */
+	FDataflowAssetEdit(UDataflow *InAsset, FPostEditFunctionCallback InCallable);
+	~FDataflowAssetEdit();
+
+	Dataflow::FGraph* GetGraph();
+
+private:
+	FPostEditFunctionCallback PostEditCallback;
+	UDataflow* Asset;
+};
+
 /**
 * UDataflow (UObject)
 *
@@ -20,6 +47,7 @@ class DATAFLOWENGINE_API UDataflow : public UEdGraph
 	GENERATED_UCLASS_BODY()
 
 	TSharedPtr<Dataflow::FGraph, ESPMode::ThreadSafe> Dataflow;
+	void PostEditCallback();
 
 public:
 	UDataflow(const FObjectInitializer& ObjectInitializer = FObjectInitializer::Get());
@@ -34,10 +62,13 @@ public:
 	void Serialize(FArchive& Ar);
 
 	/** Accessors for internal geometry collection */
-	void SetDataflow(TSharedPtr<Dataflow::FGraph, ESPMode::ThreadSafe> DataflowIn) { Dataflow = DataflowIn; }
-	TSharedPtr<Dataflow::FGraph, ESPMode::ThreadSafe>       GetDataflow() { return Dataflow; }
 	const TSharedPtr<Dataflow::FGraph, ESPMode::ThreadSafe> GetDataflow() const { return Dataflow; }
 
-
+	/**Editing the collection should only be through the edit object.*/
+	FDataflowAssetEdit EditDataflow() const {
+		//ThisNC is a by-product of editing through the component.
+		UDataflow* ThisNC = const_cast<UDataflow*>(this);
+		return FDataflowAssetEdit(ThisNC, [ThisNC]() {ThisNC->PostEditCallback(); });
+	}
 };
 
diff --git a/Engine/Source/Runtime/Experimental/Dataflow/Engine/Public/Dataflow/DataflowObjectInterface.h b/Engine/Source/Runtime/Experimental/Dataflow/Engine/Public/Dataflow/DataflowObjectInterface.h
new file mode 100644
--- /dev/null
+++ b/Engine/Source/Runtime/Experimental/Dataflow/Engine/Public/Dataflow/DataflowObjectInterface.h
@@ -0,0 +1,19 @@
+// Copyright Epic Games, Inc. All Rights Reserved.
+
+#pragma once
+
+#include "Dataflow/DataflowNodeParameters.h"
+#include "GeometryCollection/ManagedArrayCollection.h"
+#include "Dataflow/DataflowGraph.h"
+
+
+namespace Dataflow
+{
+
+	class DATAFLOWENGINE_API FEngineContext : public FContext
+	{
+	public:
+		FEngineContext(UDataflow* InGraph, float InTime);
+		UDataflow* Graph = nullptr;
+	};
+}
\ No newline at end of file
