Chaos : perf optimization for non moving disabled particle with a internal cluster parent

diff --git a/Engine/Source/Runtime/Experimental/Chaos/Private/PhysicsProxy/GeometryCollectionPhysicsProxy.cpp b/Engine/Source/Runtime/Experimental/Chaos/Private/PhysicsProxy/GeometryCollectionPhysicsProxy.cpp
--- a/Engine/Source/Runtime/Experimental/Chaos/Private/PhysicsProxy/GeometryCollectionPhysicsProxy.cpp
+++ b/Engine/Source/Runtime/Experimental/Chaos/Private/PhysicsProxy/GeometryCollectionPhysicsProxy.cpp
@@ -1806,6 +1806,30 @@ void FGeometryCollectionPhysicsProxy::BufferPhysicsResults_External(Chaos::FDirt
 	// TargetResults.IsObjectLoading;
 }
 
+static void UpdateParticleHandleTransform(Chaos::FPBDRigidsSolver& CurrentSolver, Chaos::FPBDRigidClusteredParticleHandle& Handle, const Chaos::FRigidTransform3& NewTransform)
+{
+	Handle.X() = NewTransform.GetTranslation();
+	Handle.R() = NewTransform.GetRotation();
+	Handle.UpdateWorldSpaceState(NewTransform, Chaos::FVec3{0});
+	CurrentSolver.GetEvolution()->DirtyParticle(Handle);
+}
+
+static void UpdateParticleHandleTransformIfNeeded(Chaos::FPBDRigidsSolver& CurrentSolver, Chaos::FPBDRigidClusteredParticleHandle& Handle, const Chaos::FRigidTransform3& NewTransform)
+{
+	const Chaos::FVec3 NewX = NewTransform.GetTranslation();
+	const Chaos::FVec3 OldX = Handle.X();
+	
+	const Chaos::FRotation3 NewR  = NewTransform.GetRotation();
+	const Chaos::FRotation3 OldR = Handle.R();
+	
+	if (OldX != NewX || OldR != NewR)
+	{
+		UpdateParticleHandleTransform(CurrentSolver, Handle, NewTransform);
+	}
+}
+
+
+
 void FGeometryCollectionPhysicsProxy::BufferPhysicsResults_Internal(Chaos::FPBDRigidsSolver* CurrentSolver, Chaos::FDirtyGeometryCollectionData& BufferData)
 {
 	/**
@@ -1840,7 +1864,7 @@ void FGeometryCollectionPhysicsProxy::BufferPhysicsResults_Internal(Chaos::FPBDR
 			TargetResults.Parent[TransformGroupIndex] = PhysicsThreadCollection.Parent[TransformGroupIndex];
 
 			TargetResults.DisabledStates[TransformGroupIndex] = true;
-			Chaos::TPBDRigidClusteredParticleHandle<Chaos::FReal, 3>* Handle = SolverParticleHandles[TransformGroupIndex];
+			Chaos::FPBDRigidClusteredParticleHandle* Handle = SolverParticleHandles[TransformGroupIndex];
 			if (!Handle)
 			{
 				PhysicsThreadCollection.Active[TransformGroupIndex] = !TargetResults.DisabledStates[TransformGroupIndex];
@@ -1910,7 +1934,7 @@ void FGeometryCollectionPhysicsProxy::BufferPhysicsResults_Internal(Chaos::FPBDR
 				// the positions directly from the solvers cluster particle.
 				if(TPBDRigidParticleHandle<Chaos::FReal, 3>* ClusterParentBase = Handle->ClusterIds().Id)
 				{
-					if(Chaos::TPBDRigidClusteredParticleHandle<Chaos::FReal, 3>* ClusterParent = ClusterParentBase->CastToClustered())
+					if(Chaos::FPBDRigidClusteredParticleHandle* ClusterParent = ClusterParentBase->CastToClustered())
 					{
 						// syncronize parents if it has changed.
 						if(SolverClusterID[TransformGroupIndex] != ClusterParent)
@@ -1936,9 +1960,7 @@ void FGeometryCollectionPhysicsProxy::BufferPhysicsResults_Internal(Chaos::FPBDR
 
 						if(ClusterParent->InternalCluster())
 						{
-							Chaos::TPBDRigidClusteredParticleHandle<Chaos::FReal, 3>* ProxyElementHandle = SolverParticleHandles[TransformGroupIndex];
-
-							const FTransform ParticleToWorld = ProxyElementHandle->ChildToParent() * FRigidTransform3(ClusterParent->X(), ClusterParent->R());    // aka ClusterChildToWorld
+							const FTransform ParticleToWorld = Handle->ChildToParent() * FRigidTransform3(ClusterParent->X(), ClusterParent->R());    // aka ClusterChildToWorld
 							TargetResults.ParticleXs[TransformGroupIndex] = ParticleToWorld.GetTranslation();
 							TargetResults.ParticleRs[TransformGroupIndex] = ParticleToWorld.GetRotation();
 							TargetResults.ParticleVs[TransformGroupIndex] = ClusterParent->V();
@@ -1959,21 +1981,16 @@ void FGeometryCollectionPhysicsProxy::BufferPhysicsResults_Internal(Chaos::FPBDR
 							TargetResults.DisabledStates[TransformGroupIndex] = false;
 							IsObjectDynamic                                   = true;
 
-							ProxyElementHandle->X() = ParticleToWorld.GetTranslation();
-							ProxyElementHandle->R() = ParticleToWorld.GetRotation();
-							CurrentSolver->GetEvolution()->DirtyParticle(*ProxyElementHandle);
+							// as we just transitioned from disabled to non disabled the update is unconditional 
+							UpdateParticleHandleTransform(*CurrentSolver, *Handle, ParticleToWorld);
 						}
 
 						if (bGeometryCollectionEnabledNestedChildTransformUpdates)
 						{
 							if (!ClusterParent->Disabled())
 							{
-								FTransform ChildToWorld = Handle->ChildToParent() * FRigidTransform3(ClusterParent->X(), ClusterParent->R());
-								Handle->X() = Handle->P() = ChildToWorld.GetTranslation();
-								Handle->R() = Handle->Q() = ChildToWorld.GetRotation();
-								// TODO : shoudl we transfer velocity as well ?
-								Handle->UpdateWorldSpaceState(ChildToWorld, FVec3(0));
-								CurrentSolver->GetEvolution()->DirtyParticle(*Handle);
+								const FRigidTransform3 ChildToWorld = Handle->ChildToParent() * FRigidTransform3(ClusterParent->X(), ClusterParent->R());
+								UpdateParticleHandleTransformIfNeeded(*CurrentSolver, *Handle, ChildToWorld);
 							}
 						}
 					}
