Chaos : fix convex generation resulting in non convex faces in some specific configuration

diff --git a/Engine/Source/Runtime/Experimental/Chaos/Public/Chaos/CollisionConvexMesh.h b/Engine/Source/Runtime/Experimental/Chaos/Public/Chaos/CollisionConvexMesh.h
--- a/Engine/Source/Runtime/Experimental/Chaos/Public/Chaos/CollisionConvexMesh.h
+++ b/Engine/Source/Runtime/Experimental/Chaos/Public/Chaos/CollisionConvexMesh.h
@@ -10,6 +10,8 @@
 #include "Containers/ChunkedArray.h"
 #include "CompGeom/ConvexHull3.h"
 
+PRAGMA_DISABLE_OPTIMIZATION
+
 #define DEBUG_HULL_GENERATION 0
 
 // Those flags allow to output geometric data in OBJ compatible format
@@ -22,6 +24,8 @@
 #define DEBUG_HULL_GENERATION_HULLSTEPS_TO_OBJ 0
 #define DEBUG_HULL_GENERATION_BUILDHORIZON_TO_OBJ 0
 
+#define DEBUG_HULL_VERBOSITY VeryVerbose
+
 namespace Chaos
 {
 	// When encountering a triangle or quad in hull generation (3-points or 4 coplanar points) we will instead generate
@@ -283,7 +287,7 @@ namespace Chaos
 				InitialFacesString += FString::Printf(TEXT("(%d %d %d) "), Current->FirstEdge->Vertex, Current->FirstEdge->Next->Vertex, Current->FirstEdge->Prev->Vertex);
 				Current = Current->Next;
 			}
-			UE_LOG(LogChaos, VeryVerbose, TEXT("%s"), *InitialFacesString);
+			UE_LOG(LogChaos, DEBUG_HULL_VERBOSITY, TEXT("%s"), *InitialFacesString);
 #endif
 
 			FConvexFace* DummyFace = Pool.AllocConvexFace(Faces->Plane);
@@ -297,20 +301,20 @@ namespace Chaos
 
 #if DEBUG_HULL_GENERATION
 #if DEBUG_HULL_GENERATION_HULLSTEPS_TO_OBJ
-				UE_LOG(LogChaos, VeryVerbose, TEXT("# ======================================================"));
+				UE_LOG(LogChaos, DEBUG_HULL_VERBOSITY, TEXT("# ======================================================"));
 				const FVec3Type ConflictPos = InVertices[ConflictV->Vertex];
-				UE_LOG(LogChaos, VeryVerbose, TEXT("# GENERATED HULL before adding Vtx %d (%f %f %f)"), ConflictV->Vertex, ConflictPos.X, ConflictPos.Y, ConflictPos.Z);
-				UE_LOG(LogChaos, VeryVerbose, TEXT("# ------------------------------------------------------"));
+				UE_LOG(LogChaos, DEBUG_HULL_VERBOSITY, TEXT("# GENERATED HULL before adding Vtx %d (%f %f %f)"), ConflictV->Vertex, ConflictPos.X, ConflictPos.Y, ConflictPos.Z);
+				UE_LOG(LogChaos, DEBUG_HULL_VERBOSITY, TEXT("# ------------------------------------------------------"));
 				FConvexFace* Face = DummyFace->Next;
 				while (Face)
 				{
 					const FVector P1 = InVertices[Face->FirstEdge->Prev->Vertex];
 					const FVector P2 = InVertices[Face->FirstEdge->Next->Vertex];
 					const FVector P3 = InVertices[Face->FirstEdge->Vertex];
-					UE_LOG(LogChaos, VeryVerbose, TEXT("v %f %f %f"), P1.X, P1.Y, P1.Z);
-					UE_LOG(LogChaos, VeryVerbose, TEXT("v %f %f %f"), P2.X, P2.Y, P2.Z);
-					UE_LOG(LogChaos, VeryVerbose, TEXT("v %f %f %f"), P3.X, P3.Y, P3.Z);
-					UE_LOG(LogChaos, VeryVerbose, TEXT("f -3 -2 -1"));
+					UE_LOG(LogChaos, DEBUG_HULL_VERBOSITY, TEXT("v %f %f %f"), P1.X, P1.Y, P1.Z);
+					UE_LOG(LogChaos, DEBUG_HULL_VERBOSITY, TEXT("v %f %f %f"), P2.X, P2.Y, P2.Z);
+					UE_LOG(LogChaos, DEBUG_HULL_VERBOSITY, TEXT("v %f %f %f"), P3.X, P3.Y, P3.Z);
+					UE_LOG(LogChaos, DEBUG_HULL_VERBOSITY, TEXT("f -3 -2 -1"));
 					Face = Face->Next;
 				}
 #endif
@@ -434,6 +438,52 @@ namespace Chaos
 			check(InOutVertices.Num() > 3);
 		}
 
+		static bool IsFaceOutlineConvex(const FPlaneType& Plane, const TArray<int32>& Indices, const TArray<FVec3Type>& Vertices)
+		{
+			TArray<int8> Signs;
+			Signs.SetNum(Indices.Num());
+
+			if (Indices.Num() < 4)
+			{
+				// degenerated face
+				return true;
+			}
+
+			for (int32 PointIndex = 0; PointIndex < Indices.Num(); PointIndex++)
+			{
+				const int32 Index0 = Indices[PointIndex];
+				const int32 Index1 = Indices[(PointIndex + 1) % Indices.Num()];
+				const int32 Index2 = Indices[(PointIndex + 2) % Indices.Num()];
+
+				const FVec3Type Point0 = Vertices[Index0];
+				const FVec3Type Point1 = Vertices[Index1];
+				const FVec3Type Point2 = Vertices[Index2];
+
+				const FVec3Type Segment0(Point1 - Point0);
+				const FVec3Type Segment1(Point2 - Point1);
+
+				const FVec3Type Cross = FVec3Type::CrossProduct(Segment0, Segment1);
+				const FRealType Dot = FVec3Type::DotProduct(Cross, Plane.Normal());
+
+				Signs[PointIndex] = static_cast<int8>(FMath::Sign(Dot));
+			}
+
+			int8 RefSign = 0;
+			for (const int8 Sign : Signs)
+			{
+				if (RefSign == 0)
+				{
+					RefSign = Sign;
+				}
+				if (Sign != 0 && RefSign != Sign)
+				{
+					return false;
+				}
+			}
+
+			return true;
+		}
+
 		// Convert multi-triangle faces to single n-gons
 		static void MergeFaces(TArray<FPlaneType>& InOutPlanes, TArray<TArray<int32>>& InOutFaceVertexIndices, const TArray<FVec3Type>& Vertices, FRealType DistanceThreshold)
 		{
@@ -500,12 +550,84 @@ namespace Chaos
 			// Re-order the face vertices to form the face half-edges
 			for (int32 PlaneIndex0 = 0; PlaneIndex0 < InOutPlanes.Num(); ++PlaneIndex0)
 			{
-				SortFaceVerticesCCW(InOutPlanes[PlaneIndex0], InOutFaceVertexIndices[PlaneIndex0], Vertices);
+				SortFaceVerticesCCWAndRemoveInsideVertices(InOutPlanes[PlaneIndex0], InOutFaceVertexIndices[PlaneIndex0], Vertices);
+#if DEBUG_HULL_GENERATION
+				ensure(IsFaceOutlineConvex(InOutPlanes[PlaneIndex0], InOutFaceVertexIndices[PlaneIndex0], Vertices));
+#endif
 			}
 		}
 
+		// IMPORTANT : vertices are assumed to be sorted CCW
+		static void RemoveInsideFaceVertices(const FPlaneType& Face, TArray<int32>& InOutFaceVertexIndices, const TArray<FVec3Type>& Vertices, const FVec3Type& Centroid)
+		{
+			if (InOutFaceVertexIndices.Num() <= 3)
+			{
+				return;
+			}
+			// find furthest point from centroid as it is garanteed to be part of the convex hull 
+			int32 StartIndex = 0;
+			FRealType FurthestSquaredDistance = TNumericLimits<FRealType>::Lowest();
+			for (int32 Index = 0; Index < InOutFaceVertexIndices.Num(); ++Index)
+			{
+				const FRealType SquaredDistance = (Vertices[InOutFaceVertexIndices[Index]] - Centroid).SquaredLength();
+				if (SquaredDistance > FurthestSquaredDistance)
+				{
+					StartIndex = Index;
+					FurthestSquaredDistance = SquaredDistance;
+				}
+			}
+
+			const int32 VtxCount = InOutFaceVertexIndices.Num();
+
+			struct PointSegment
+			{
+				int32 VtxIndex;
+				FVec3Type Segment;
+			};
+
+			const int32 VtxStartIndex = InOutFaceVertexIndices[StartIndex];
+			int32 VtxIndex0 = VtxStartIndex;
+			int32 VtxIndex1 = InOutFaceVertexIndices[(StartIndex + 1) % VtxCount];
+
+			TArray<PointSegment> Stack;
+			Stack.Push({ VtxIndex0, FVec3Type{0} });
+			Stack.Push({ VtxIndex1, Vertices[VtxIndex1] - Vertices[VtxIndex0] });
+
+			int32 Step = 2; // we already processed the two first 
+			while (Step <= VtxCount)
+			{
+				const int32 NextIndex = (StartIndex + Step) % VtxCount;
+
+				const PointSegment& LastOnStack = Stack.Last();
+				VtxIndex0 = LastOnStack.VtxIndex;
+				VtxIndex1 = InOutFaceVertexIndices[NextIndex];
+				const FVec3Type Segment = Vertices[VtxIndex1] - Vertices[VtxIndex0];
+
+				const FVec3 Cross = FVec3Type::CrossProduct(LastOnStack.Segment, Segment);
+				if (FVec3Type::DotProduct(Cross, Face.Normal()) >= 0)
+				{
+					if (VtxIndex1 != VtxStartIndex)
+					{
+						Stack.Push({ VtxIndex1, Segment });
+					}
+					Step++;
+				}
+				else
+				{
+					Stack.Pop();
+				}
+			}
+
+			InOutFaceVertexIndices.Reset();
+			for (const PointSegment& PointSegment : Stack)
+			{
+				InOutFaceVertexIndices.Add(PointSegment.VtxIndex);
+			}
+			ensure(InOutFaceVertexIndices.Num() >= 3);
+		}
+
 		// Reorder the vertices to be counter-clockwise about the normal
-		static void SortFaceVerticesCCW(const FPlaneType& Face, TArray<int32>& InOutFaceVertexIndices, const TArray<FVec3Type>& Vertices)
+		static void SortFaceVerticesCCWAndRemoveInsideVertices(const FPlaneType& Face, TArray<int32>& InOutFaceVertexIndices, const TArray<FVec3Type>& Vertices)
 		{
 			FMatrix33 FaceMatrix = (FMatrix44f)FRotationMatrix::MakeFromZ(FVector(Face.Normal()));
 
@@ -532,6 +654,8 @@ namespace Chaos
 			};
 
 			InOutFaceVertexIndices.Sort(VertexSortPredicate);
+
+			RemoveInsideFaceVertices(Face, InOutFaceVertexIndices, Vertices, Centroid);
 		}
 
 		// Generate the vertex indices for all planes in CCW order (used to serialize old data that did not have structure data)
@@ -549,7 +673,7 @@ namespace Chaos
 					}
 				}
 
-				SortFaceVerticesCCW(InPlanes[PlaneIndex], OutFaceVertexIndices[PlaneIndex], Vertices);
+				SortFaceVerticesCCWAndRemoveInsideVertices(InPlanes[PlaneIndex], OutFaceVertexIndices[PlaneIndex], Vertices);
 			}
 		}
 
@@ -825,11 +949,11 @@ namespace Chaos
 
 		static FHalfEdge* FindConflictVertex(const TArray<FVec3Type>& InVertices, FConvexFace* FaceList)
 		{
-			UE_CLOG(DEBUG_HULL_GENERATION, LogChaos, VeryVerbose, TEXT("Finding conflict vertex"));
+			UE_CLOG(DEBUG_HULL_GENERATION, LogChaos, DEBUG_HULL_VERBOSITY, TEXT("Finding conflict vertex"));
 
 			for(FConvexFace* CurFace = FaceList; CurFace; CurFace = CurFace->Next)
 			{
-				UE_CLOG(DEBUG_HULL_GENERATION, LogChaos, VeryVerbose, TEXT("\tTesting Face (%d %d %d)"), CurFace->FirstEdge->Vertex, CurFace->FirstEdge->Next->Vertex, CurFace->FirstEdge->Prev->Vertex);
+				UE_CLOG(DEBUG_HULL_GENERATION, LogChaos, DEBUG_HULL_VERBOSITY, TEXT("\tTesting Face (%d %d %d)"), CurFace->FirstEdge->Vertex, CurFace->FirstEdge->Next->Vertex, CurFace->FirstEdge->Prev->Vertex);
 
 				FRealType MaxD = TNumericLimits<FRealType>::Lowest();
 				FHalfEdge* MaxV = nullptr;
@@ -844,8 +968,8 @@ namespace Chaos
 					}
 				}
 
-				UE_CLOG((DEBUG_HULL_GENERATION && !MaxV), LogChaos, VeryVerbose, TEXT("\t\tNo Conflict List"));
-				UE_CLOG((DEBUG_HULL_GENERATION && MaxV), LogChaos, VeryVerbose, TEXT("\t\tFound %d at distance %f"), MaxV->Vertex, MaxD);
+				UE_CLOG((DEBUG_HULL_GENERATION && !MaxV), LogChaos, DEBUG_HULL_VERBOSITY, TEXT("\t\tNo Conflict List"));
+				UE_CLOG((DEBUG_HULL_GENERATION && MaxV), LogChaos, DEBUG_HULL_VERBOSITY, TEXT("\t\tFound %d at distance %f"), MaxV->Vertex, MaxD);
 
 				check(CurFace->ConflictList == nullptr || MaxV);
 				if(MaxV)
@@ -873,7 +997,7 @@ namespace Chaos
 		static void BuildHorizon(const TArray<FVec3Type>& InVertices, FHalfEdge* ConflictV, TArray<FHalfEdge*>& HorizonEdges, TArray<FConvexFace*>& FacesToDelete, const Params& InParams)
 		{
 #if DEBUG_HULL_GENERATION
-			UE_LOG(LogChaos, VeryVerbose, TEXT("Generate horizon - START"));
+			UE_LOG(LogChaos, DEBUG_HULL_VERBOSITY, TEXT("Generate horizon - START"));
 #endif
 			//We must flood fill from the initial face and mark edges of faces the conflict vertex cannot see
 			//In order to return a CCW ordering we must traverse each face in CCW order from the edge we crossed over
@@ -896,7 +1020,7 @@ namespace Chaos
 					QueueString += FString::Printf(TEXT(" [%d - %d] "), QueuedEdge->Vertex, QueuedEdge->Next->Vertex);
 				}
 				QueueString += TEXT(")");
-				UE_LOG(LogChaos, VeryVerbose, TEXT("%s"), *QueueString);
+				UE_LOG(LogChaos, DEBUG_HULL_VERBOSITY, TEXT("%s"), *QueueString);
 #endif 
 
 				FHalfEdge* Edge = Queue.Pop(/*bAllowShrinking=*/false);
@@ -905,13 +1029,13 @@ namespace Chaos
 				FConvexFace* NextFace = Twin->Face;
 
 #if DEBUG_HULL_GENERATION
-				UE_LOG(LogChaos, VeryVerbose, TEXT("\tPop edge [%d - %d] from queue"), Edge->Vertex, Edge->Next->Vertex);
+				UE_LOG(LogChaos, DEBUG_HULL_VERBOSITY, TEXT("\tPop edge [%d - %d] from queue"), Edge->Vertex, Edge->Next->Vertex);
 #endif 
 
 				if(Processed.Contains(NextFace))
 				{
 #if DEBUG_HULL_GENERATION
-					UE_LOG(LogChaos, VeryVerbose, TEXT("\tTwin Face [%d] already processed - skip"), NextFace);
+					UE_LOG(LogChaos, DEBUG_HULL_VERBOSITY, TEXT("\tTwin Face [%d] already processed - skip"), NextFace);
 #endif
 					continue;
 				}
@@ -919,7 +1043,7 @@ namespace Chaos
 				if(Distance > Epsilon)
 				{
 #if DEBUG_HULL_GENERATION
-					UE_LOG(LogChaos, VeryVerbose, TEXT("\tDistance [%f] > Epsilon [%f] - add to queue"), Distance, Epsilon);
+					UE_LOG(LogChaos, DEBUG_HULL_VERBOSITY, TEXT("\tDistance [%f] > Epsilon [%f] - add to queue"), Distance, Epsilon);
 #endif 
 					Queue.Add(Twin->Prev); //stack pops so reverse order
 					Queue.Add(Twin->Next);
@@ -928,7 +1052,7 @@ namespace Chaos
 				else
 				{
 #if DEBUG_HULL_GENERATION
-					UE_LOG(LogChaos, VeryVerbose, TEXT("\tAdd [%d - %d] to horizon "), Edge->Vertex, Edge->Next->Vertex);
+					UE_LOG(LogChaos, DEBUG_HULL_VERBOSITY, TEXT("\tAdd [%d - %d] to horizon "), Edge->Vertex, Edge->Next->Vertex);
 #endif 
 					// we need to ensure that the horizon is a continous edge loop
 					// this may get the wrong edge path, but way more roibust than not testing this
@@ -939,7 +1063,7 @@ namespace Chaos
 					else
 					{
 #if DEBUG_HULL_GENERATION
-						UE_LOG(LogChaos, VeryVerbose, TEXT("\tNON VALID EDGE LOOP - internal horizon edge detected [%d - %d] - skipping "), Edge->Vertex, Edge->Next->Vertex);
+						UE_LOG(LogChaos, DEBUG_HULL_VERBOSITY, TEXT("\tNON VALID EDGE LOOP - internal horizon edge detected [%d - %d] - skipping "), Edge->Vertex, Edge->Next->Vertex);
 #endif 
 					}
 				}
@@ -947,20 +1071,20 @@ namespace Chaos
 
 #if DEBUG_HULL_GENERATION
 #if DEBUG_HULL_GENERATION_BUILDHORIZON_TO_OBJ
-			UE_LOG(LogChaos, VeryVerbose, TEXT("# ======================================================"));
+			UE_LOG(LogChaos, DEBUG_HULL_VERBOSITY, TEXT("# ======================================================"));
 			const FVec3Type ConflictPos = InVertices[ConflictV->Vertex];
-			UE_LOG(LogChaos, VeryVerbose, TEXT("# BUILD_HORIZON - Conflict Vertex = %d (%f %f %f)"), ConflictV->Vertex, ConflictPos.X, ConflictPos.Y, ConflictPos.Z);
-			UE_LOG(LogChaos, VeryVerbose, TEXT("# ------------------------------------------------------"));
+			UE_LOG(LogChaos, DEBUG_HULL_VERBOSITY, TEXT("# BUILD_HORIZON - Conflict Vertex = %d (%f %f %f)"), ConflictV->Vertex, ConflictPos.X, ConflictPos.Y, ConflictPos.Z);
+			UE_LOG(LogChaos, DEBUG_HULL_VERBOSITY, TEXT("# ------------------------------------------------------"));
 			for (TSet<FConvexFace*>::TConstIterator SetIt(Processed); SetIt; ++SetIt)
 			{
 				const FConvexFace* Face = *SetIt;
 				const FVector P1 = InVertices[Face->FirstEdge->Prev->Vertex];
 				const FVector P2 = InVertices[Face->FirstEdge->Next->Vertex];
 				const FVector P3 = InVertices[Face->FirstEdge->Vertex];
-				UE_LOG(LogChaos, VeryVerbose, TEXT("v %f %f %f"), P1.X, P1.Y, P1.Z);
-				UE_LOG(LogChaos, VeryVerbose, TEXT("v %f %f %f"), P2.X, P2.Y, P2.Z);
-				UE_LOG(LogChaos, VeryVerbose, TEXT("v %f %f %f"), P3.X, P3.Y, P3.Z);
-				UE_LOG(LogChaos, VeryVerbose, TEXT("f -3 -2 -1"));
+				UE_LOG(LogChaos, DEBUG_HULL_VERBOSITY, TEXT("v %f %f %f"), P1.X, P1.Y, P1.Z);
+				UE_LOG(LogChaos, DEBUG_HULL_VERBOSITY, TEXT("v %f %f %f"), P2.X, P2.Y, P2.Z);
+				UE_LOG(LogChaos, DEBUG_HULL_VERBOSITY, TEXT("v %f %f %f"), P3.X, P3.Y, P3.Z);
+				UE_LOG(LogChaos, DEBUG_HULL_VERBOSITY, TEXT("f -3 -2 -1"));
 			}
 #endif
 
@@ -970,9 +1094,9 @@ namespace Chaos
 				HorizonString += FString::Printf(TEXT("%d (%d)"), HorizonEdge->Vertex, HorizonEdge->Next->Vertex);
 			}
 			HorizonString += TEXT(")");
-			UE_LOG(LogChaos, VeryVerbose, TEXT("%s"), *HorizonString);
+			UE_LOG(LogChaos, DEBUG_HULL_VERBOSITY, TEXT("%s"), *HorizonString);
 
-			UE_LOG(LogChaos, VeryVerbose, TEXT("Generate horizon - END"));
+			UE_LOG(LogChaos, DEBUG_HULL_VERBOSITY, TEXT("Generate horizon - END"));
 #endif
 
 		}
@@ -1052,7 +1176,7 @@ namespace Chaos
 
 		static bool AddVertex(FMemPool& Pool, const TArray<FVec3Type>& InVertices, FHalfEdge* ConflictV, const Params& InParams)
 		{
-			UE_CLOG(DEBUG_HULL_GENERATION, LogChaos, VeryVerbose, TEXT("Adding Vertex %d"), ConflictV->Vertex);
+			UE_CLOG(DEBUG_HULL_GENERATION, LogChaos, DEBUG_HULL_VERBOSITY, TEXT("Adding Vertex %d"), ConflictV->Vertex);
 
 			TArray<FHalfEdge*> HorizonEdges;
 			TArray<FConvexFace*> FacesToDelete;
@@ -1070,14 +1194,14 @@ namespace Chaos
 			{
 				NewFaceString += FString::Printf(TEXT("(%d %d %d) "), Face->FirstEdge->Vertex, Face->FirstEdge->Next->Vertex, Face->FirstEdge->Prev->Vertex);
 			}
-			UE_LOG(LogChaos, VeryVerbose, TEXT("%s"), *NewFaceString);
+			UE_LOG(LogChaos, DEBUG_HULL_VERBOSITY, TEXT("%s"), *NewFaceString);
 
 			FString DeleteFaceString(TEXT("\tDelete Faces: "));
 			for(const FConvexFace* Face : FacesToDelete)
 			{
 				DeleteFaceString += FString::Printf(TEXT("(%d %d %d) "), Face->FirstEdge->Vertex, Face->FirstEdge->Next->Vertex, Face->FirstEdge->Prev->Vertex);
 			}
-			UE_LOG(LogChaos, VeryVerbose, TEXT("%s"), *DeleteFaceString);
+			UE_LOG(LogChaos, DEBUG_HULL_VERBOSITY, TEXT("%s"), *DeleteFaceString);
 #endif
 
 			for(FConvexFace* Face : FacesToDelete)
@@ -1109,3 +1233,4 @@ namespace Chaos
 	};
 }
 
+PRAGMA_ENABLE_OPTIMIZATION
diff --git a/Engine/Source/Runtime/Experimental/ChaosCore/Public/Chaos/Core.h b/Engine/Source/Runtime/Experimental/ChaosCore/Public/Chaos/Core.h
--- a/Engine/Source/Runtime/Experimental/ChaosCore/Public/Chaos/Core.h
+++ b/Engine/Source/Runtime/Experimental/ChaosCore/Public/Chaos/Core.h
@@ -51,6 +51,6 @@ namespace Chaos
 	// NOTE: if you get a merge conflict on the GUID, you must replace it with a new GUID - do not accept the source or target
 	// or you will likely get DDC version conflicts resulting in crashes during load.
 	// Core version string for Chaos data. Any DDC builder dependent on Chaos for serialization should depend on this version
-	static const FString ChaosVersionString = TEXT("841FF58C-9DD5-4AF7-BE5B-959F1F2BC4D7");
+	static const FString ChaosVersionString = TEXT("0D2D5C4C-A885-43D3-B53F-7F7C46417E05");
 
 }
