Chaos : GeometryCollection
- Added NewCopy

diff --git a/Engine/Source/Runtime/Experimental/Chaos/Private/GeometryCollection/ManagedArrayCollection.cpp b/Engine/Source/Runtime/Experimental/Chaos/Private/GeometryCollection/ManagedArrayCollection.cpp
--- a/Engine/Source/Runtime/Experimental/Chaos/Private/GeometryCollection/ManagedArrayCollection.cpp
+++ b/Engine/Source/Runtime/Experimental/Chaos/Private/GeometryCollection/ManagedArrayCollection.cpp
@@ -15,6 +15,7 @@ FManagedArrayCollection::FManagedArrayCollection()
 	Version = 9;
 }
 
+
 void FManagedArrayCollection::AddGroup(FName Group)
 {
 	ensure(!GroupInfo.Contains(Group));
@@ -353,8 +354,20 @@ void FManagedArrayCollection::CopyAttribute(const FManagedArrayCollection& InCol
 	SyncGroupSizeFrom(InCollection, Group);
 	FKeyType Key = FManagedArrayCollection::MakeMapKey(Name, Group);
 
+	if (!HasAttribute(Name, Group))
+	{
+		const FValueType& V = InCollection.Map[Key];
+		EArrayType Type = V.ArrayType;
+		FValueType Value(Type, *NewManagedTypedArray(Type));
+		Value.Value->Resize(NumElements(Group));
+		Value.GroupIndexDependency = V.GroupIndexDependency;
+		Value.Saved = V.Saved;
+		Value.bExternalValue = V.bExternalValue;
+		Map.Add(Key, MoveTemp(Value));
+	}
+
 	const FValueType& OriginalValue = InCollection.Map[Key];
-	const FValueType& DestValue = Map[Key];	//todo(ocohen): API assumes an AddAttribute is called before copy is done. It'd be nice to handle the case where AddAttribute was not done first
+	const FValueType& DestValue = Map[Key];	
 	check(OriginalValue.ArrayType == DestValue.ArrayType);
 	DestValue.Value->Init(*OriginalValue.Value);
 }
diff --git a/Engine/Source/Runtime/Experimental/Chaos/Public/GeometryCollection/ManagedArrayCollection.h b/Engine/Source/Runtime/Experimental/Chaos/Public/GeometryCollection/ManagedArrayCollection.h
--- a/Engine/Source/Runtime/Experimental/Chaos/Public/GeometryCollection/ManagedArrayCollection.h
+++ b/Engine/Source/Runtime/Experimental/Chaos/Public/GeometryCollection/ManagedArrayCollection.h
@@ -109,6 +109,33 @@ public:
 		return ModifyAttribute<T>(Name, Group);
 	}
 
+	/**
+	* Duplicate the ManagedArrayCollection as the specified templated type
+	* @return unowned pointer to new collection.
+	*/
+	template<class T>
+	T* NewCopy() const
+	{
+		T* Collection = new T();
+		for (const TTuple<FKeyType, FValueType>& Entry : Map)
+		{
+			if (!Collection->HasGroup(Entry.Key.Get<1>()))
+			{
+				Collection->AddGroup(Entry.Key.Get<1>());
+			}
+
+			if (NumElements(Entry.Key.Get<1>()) != Collection->NumElements(Entry.Key.Get<1>()))
+			{
+				ensure(!Collection->NumElements(Entry.Key.Get<1>()));
+				Collection->AddElements(NumElements(Entry.Key.Get<1>()), Entry.Key.Get<1>());
+			}
+
+			Collection->CopyAttribute(*this, Entry.Key.Get<0>(), Entry.Key.Get<1>());
+		}
+		return Collection;
+	}
+
+
 	/**
 	* Add an external attribute of Type(T) to the group for size management. Lifetime is managed by the caller, must make sure the array is alive when the collection is
 	* @param Name - The name of the attribute
diff --git a/Engine/Source/Runtime/Experimental/Chaos/Public/GeometryCollection/ManagedArrayTypes.h b/Engine/Source/Runtime/Experimental/Chaos/Public/GeometryCollection/ManagedArrayTypes.h
--- a/Engine/Source/Runtime/Experimental/Chaos/Public/GeometryCollection/ManagedArrayTypes.h
+++ b/Engine/Source/Runtime/Experimental/Chaos/Public/GeometryCollection/ManagedArrayTypes.h
@@ -94,19 +94,3 @@ inline FManagedArrayBase* NewManagedTypedArray(EManagedArrayType ArrayType)
 	check(false);
 	return nullptr;
 }
-
-// ---------------------------------------------------------
-//     Makes a copy from one array to another
-//
-inline FManagedArrayBase* CopyManagedTypedArray(EManagedArrayType ArrayType, FManagedArrayBase* Dest, FManagedArrayBase* Src)
-{
-	switch (ArrayType)
-	{
-#define MANAGED_ARRAY_TYPE(a,A)	case EManagedArrayType::F##A##Type:\
-		return new TManagedArray<a>();
-#include "ManagedArrayTypeValues.inl"
-#undef MANAGED_ARRAY_TYPE
-	}
-	check(false);
-	return nullptr;
-}
