Chaos : cleanup and imporvement of the computation of mass properties and collision generation in the GC physics proxy

diff --git a/Engine/Source/Runtime/Experimental/Chaos/Private/Chaos/MassProperties.cpp b/Engine/Source/Runtime/Experimental/Chaos/Private/Chaos/MassProperties.cpp
--- a/Engine/Source/Runtime/Experimental/Chaos/Private/Chaos/MassProperties.cpp
+++ b/Engine/Source/Runtime/Experimental/Chaos/Private/Chaos/MassProperties.cpp
@@ -13,79 +13,80 @@
 
 namespace Chaos
 {
-	FRotation3 TransformToLocalSpace(FMatrix33& Inertia)
+	template<typename T>
+	TRotation<T,3> TransformToLocalSpace(PMatrix<T,3,3>& Inertia)
 	{
-		FRotation3 FinalRotation;
+		TRotation<T,3> FinalRotation;
 
 		// Extract Eigenvalues
-		FReal OffDiagSize = FMath::Square(Inertia.M[1][0]) + FMath::Square(Inertia.M[2][0]) + FMath::Square(Inertia.M[2][1]);
+		T OffDiagSize = FMath::Square(Inertia.M[1][0]) + FMath::Square(Inertia.M[2][0]) + FMath::Square(Inertia.M[2][1]);
 		FRealDouble Trace = (Inertia.M[0][0] + Inertia.M[1][1] + Inertia.M[2][2]) / 3;
 
 		if (Trace <= UE_SMALL_NUMBER)
 		{
 			// Tiny inertia - numerical instability would follow. We should not get this unless we have bad input.
-			return FRotation3::FromIdentity();
+			return TRotation<T,3>::FromIdentity();
 		}
 
 		if ((OffDiagSize / Trace) < UE_SMALL_NUMBER)
 		{
 			// Almost diagonal matrix - we are already in local space.
-			return FRotation3::FromIdentity();
+			return TRotation<T,3>::FromIdentity();
 		}
 
-		FReal Size = static_cast<FReal>(FMath::Sqrt((FMath::Square(Inertia.M[0][0] - Trace) + FMath::Square(Inertia.M[1][1] - Trace) + FMath::Square(Inertia.M[2][2] - Trace) + 2. * OffDiagSize) / 6.));
-		FMatrix33 NewMat = (Inertia - FMatrix::Identity * static_cast<FReal>(Trace)) * static_cast<FReal>((1 / Size));
-		FReal HalfDeterminant = NewMat.Determinant() / 2;
-		FReal Angle = HalfDeterminant <= -1 ? UE_PI / 3 : (HalfDeterminant >= 1 ? 0 : FMath::Acos(HalfDeterminant) / 3);
+		T Size = static_cast<T>(FMath::Sqrt((FMath::Square(Inertia.M[0][0] - Trace) + FMath::Square(Inertia.M[1][1] - Trace) + FMath::Square(Inertia.M[2][2] - Trace) + 2. * OffDiagSize) / 6.));
+		PMatrix<T,3,3> NewMat = (Inertia - PMatrix<T,3,3>::Identity * static_cast<T>(Trace)) * static_cast<T>((1 / Size));
+		T HalfDeterminant = NewMat.Determinant() / 2;
+		T Angle = HalfDeterminant <= -1 ? UE_PI / 3 : (HalfDeterminant >= 1 ? 0 : FMath::Acos(HalfDeterminant) / 3);
 		
-		FReal m00 = static_cast<FReal>(Trace + 2 * Size * FMath::Cos(Angle));
-		FReal m11 = static_cast<FReal>(Trace + 2 * Size * FMath::Cos(Angle + (2 * UE_PI / 3)));
-		FReal m22 = static_cast<FReal>(3 * Trace - m00 - m11);
+		T m00 = static_cast<T>(Trace + 2 * Size * FMath::Cos(Angle));
+		T m11 = static_cast<T>(Trace + 2 * Size * FMath::Cos(Angle + (2 * UE_PI / 3)));
+		T m22 = static_cast<T>(3 * Trace - m00 - m11);
 
 		// Extract Eigenvectors
 		bool DoSwap = ((m00 - m11) > (m11 - m22)) ? false : true;
-		FVec3 Eigenvector0 = (Inertia.SubtractDiagonal(DoSwap ? m22 : m00)).SymmetricCofactorMatrix().LargestColumnNormalized();
-		FVec3 Orthogonal = Eigenvector0.GetOrthogonalVector().GetSafeNormal();
+		TVec3<T> Eigenvector0 = (Inertia.SubtractDiagonal(DoSwap ? m22 : m00)).SymmetricCofactorMatrix().LargestColumnNormalized();
+		TVec3<T> Orthogonal = Eigenvector0.GetOrthogonalVector().GetSafeNormal();
 		
-		PMatrix<FReal, 3, 2> Cofactors(Orthogonal, FVec3::CrossProduct(Eigenvector0, Orthogonal));
-		PMatrix<FReal, 3, 2> CofactorsScaled = Inertia * Cofactors;
+		PMatrix<T,3,2> Cofactors(Orthogonal, TVec3<T>::CrossProduct(Eigenvector0, Orthogonal));
+		PMatrix<T,3,2> CofactorsScaled = Inertia * Cofactors;
 		
-		PMatrix<FReal, 2, 2> IR(
+		PMatrix<T,2,2> IR(
 			CofactorsScaled.M[0] * Cofactors.M[0] + CofactorsScaled.M[1] * Cofactors.M[1] + CofactorsScaled.M[2] * Cofactors.M[2],
 			CofactorsScaled.M[3] * Cofactors.M[0] + CofactorsScaled.M[4] * Cofactors.M[1] + CofactorsScaled.M[5] * Cofactors.M[2],
 			CofactorsScaled.M[3] * Cofactors.M[3] + CofactorsScaled.M[4] * Cofactors.M[4] + CofactorsScaled.M[5] * Cofactors.M[5]);
 
-		PMatrix<FReal, 2, 2> IM1 = IR.SubtractDiagonal(DoSwap ? m00 : m22);
-		FReal OffDiag = IM1.M[1] * IM1.M[1];
-		FReal IM1Scale0 = FMath::Max(FReal(0), IM1.M[3] * IM1.M[3] + OffDiag);
-		FReal IM1Scale1 = FMath::Max(FReal(0), IM1.M[0] * IM1.M[0] + OffDiag);
-		FReal SqrtIM1Scale0 = FMath::Sqrt(IM1Scale0);
-		FReal SqrtIM1Scale1 = FMath::Sqrt(IM1Scale1);
+		PMatrix<T,2,2> IM1 = IR.SubtractDiagonal(DoSwap ? m00 : m22);
+		T OffDiag = IM1.M[1] * IM1.M[1];
+		T IM1Scale0 = FMath::Max(T(0), IM1.M[3] * IM1.M[3] + OffDiag);
+		T IM1Scale1 = FMath::Max(T(0), IM1.M[0] * IM1.M[0] + OffDiag);
+		T SqrtIM1Scale0 = FMath::Sqrt(IM1Scale0);
+		T SqrtIM1Scale1 = FMath::Sqrt(IM1Scale1);
 
-		FVec3 Eigenvector2, Eigenvector1;
+		TVec3<T> Eigenvector2, Eigenvector1;
 		if ((SqrtIM1Scale0 < UE_KINDA_SMALL_NUMBER) && (SqrtIM1Scale1 < UE_KINDA_SMALL_NUMBER))
 		{
 			Eigenvector1 = Orthogonal;
-			Eigenvector2 = FVec3::CrossProduct(Eigenvector0, Orthogonal).GetSafeNormal();
+			Eigenvector2 = TVec3<T>::CrossProduct(Eigenvector0, Orthogonal).GetSafeNormal();
 		}
 		else
 		{
-			FVec2 SmallEigenvector2 = IM1Scale0 > IM1Scale1 
-				? (FVec2(IM1.M[3], -IM1.M[1]) / SqrtIM1Scale0) 
-				: (IM1Scale1 > 0 ? (FVec2(-IM1.M[1], IM1.M[0]) / SqrtIM1Scale1) : FVec2(1, 0));
+			TVec2<T> SmallEigenvector2 = IM1Scale0 > IM1Scale1 
+				? (TVec2<T>(IM1.M[3], -IM1.M[1]) / SqrtIM1Scale0) 
+				: (IM1Scale1 > 0 ? (TVec2<T>(-IM1.M[1], IM1.M[0]) / SqrtIM1Scale1) : TVec2<T>(1, 0));
 			Eigenvector2 = (Cofactors * SmallEigenvector2).GetSafeNormal();
-			Eigenvector1 = FVec3::CrossProduct(Eigenvector2, Eigenvector0).GetSafeNormal();
+			Eigenvector1 = TVec3<T>::CrossProduct(Eigenvector2, Eigenvector0).GetSafeNormal();
 		}
 
 		// Return results
-		Inertia = FMatrix33(m00, 0, 0, m11, 0, m22);
-		FMatrix33 RotationMatrix = DoSwap ? FMatrix33(Eigenvector2, Eigenvector1, -Eigenvector0) : FMatrix33(Eigenvector0, Eigenvector1, Eigenvector2);
+		Inertia = PMatrix<T,3,3>(m00, 0, 0, m11, 0, m22);
+		PMatrix<T,3,3> RotationMatrix = DoSwap ? PMatrix<T,3,3>(Eigenvector2, Eigenvector1, -Eigenvector0) : PMatrix<T,3,3>(Eigenvector0, Eigenvector1, Eigenvector2);
 
 		// NOTE: UE Matrix are column-major, so the PMatrix constructor is not setting eigenvectors - we need to transpose it to get a UE rotation matrix.
-		FinalRotation = FRotation3(RotationMatrix.GetTransposed());
+		FinalRotation = TRotation<T,3>(RotationMatrix.GetTransposed());
 		if (!ensure(FMath::IsNearlyEqual((float)FinalRotation.Size(), 1.0f, UE_KINDA_SMALL_NUMBER)))
 		{
-			return FRotation3::FromIdentity();
+			return TRotation<T,3>::FromIdentity();
 		}
 		
 		return FinalRotation;
@@ -130,14 +131,8 @@ namespace Chaos
    		return MassProperties;
    	}
 	
-	template<typename T, typename TSurfaces>
-	void CHAOS_API CalculateVolumeAndCenterOfMass(const TParticles<T, 3>& Vertices, const TSurfaces& Surfaces, T& OutVolume, TVec3<T>& OutCenterOfMass)
-	{
-		CalculateVolumeAndCenterOfMass(Vertices.AllX(), Surfaces, OutVolume, OutCenterOfMass);
-	}
-
-	template<typename T, typename TSurfaces>
-	void CalculateVolumeAndCenterOfMass(const TArray<TVec3<T>>& Vertices, const TSurfaces& Surfaces, T& OutVolume, TVec3<T>& OutCenterOfMass)
+	template<typename T, typename TVec, typename TSurfaces>
+	void CalculateVolumeAndCenterOfMassImpl(const TArray<TVec>& Vertices, const TSurfaces& Surfaces, T& OutVolume, TVec& OutCenterOfMass)
 	{
 		if (!Surfaces.Num())
 		{
@@ -146,18 +141,18 @@ namespace Chaos
 		}
 		
 		T Volume = 0;
-		TVec3<T> VolumeTimesSum(0);
-		TVec3<T> Center = Vertices[Surfaces[0][0]];
+		TVec VolumeTimesSum(0);
+		TVec Center = Vertices[Surfaces[0][0]];
 		for (const auto& Element : Surfaces)
 		{
 			// For now we only support triangular elements
 			ensure(Element.Num() == 3);
 
 			PMatrix<T,3,3> DeltaMatrix;
-			TVec3<T> PerElementSize;
+			TVec PerElementSize;
 			for (int32 i = 0; i < Element.Num(); ++i)
 			{
-				TVec3<T> DeltaVector = Vertices[Element[i]] - Center;
+				TVec DeltaVector = Vertices[Element[i]] - Center;
 				DeltaMatrix.M[0][i] = DeltaVector[0];
 				DeltaMatrix.M[1][i] = DeltaVector[1];
 				DeltaMatrix.M[2][i] = DeltaVector[2];
@@ -181,36 +176,68 @@ namespace Chaos
 		OutVolume = Volume / 6;
 	}
 
-	template <typename TSurfaces>
-	void CalculateInertiaAndRotationOfMass(const FParticles& Vertices, const TSurfaces& Surfaces, const FReal Density, const FVec3& CenterOfMass,
-		FMatrix33& OutInertiaTensor, FRotation3& OutRotationOfMass)
+	template<typename T, typename TSurfaces>
+	void CHAOS_API CalculateVolumeAndCenterOfMass(const TParticles<T, 3>& Vertices, const TSurfaces& Surfaces, T& OutVolume, TVec3<T>& OutCenterOfMass)
+	{
+		CalculateVolumeAndCenterOfMassImpl(Vertices.AllX(), Surfaces, OutVolume, OutCenterOfMass);
+	}
+
+	template<typename T, typename TSurfaces>
+	void CHAOS_API CalculateVolumeAndCenterOfMass(const TArray<TVec3<T>>& Vertices, const TSurfaces& Surfaces, T& OutVolume, TVec3<T>& OutCenterOfMass)
+	{
+		CalculateVolumeAndCenterOfMassImpl(Vertices, Surfaces, OutVolume, OutCenterOfMass);
+	}
+
+	template<typename T, typename TSurfaces>
+	void CHAOS_API CalculateVolumeAndCenterOfMass(const TArray<UE::Math::TVector<T>>& Vertices, const TSurfaces& Surfaces, T& OutVolume, UE::Math::TVector<T>& OutCenterOfMass)
+	{
+		CalculateVolumeAndCenterOfMassImpl(Vertices, Surfaces, OutVolume, OutCenterOfMass);
+	}
+
+	template <typename T, typename TVec, typename TSurfaces>
+	void CalculateInertiaAndRotationOfMassImpl(const TArray<TVec>& Vertices, const TSurfaces& Surfaces, const T Density, const TVec& CenterOfMass,
+		PMatrix<T, 3, 3>& OutInertiaTensor, TRotation<T, 3>& OutRotationOfMass)
 	{
 		check(Density > 0);
 
-		static const FMatrix33 Standard(2, 1, 1, 2, 1, 2);
-		FMatrix33 Covariance(0);
+		static const PMatrix<T, 3, 3> Standard(2, 1, 1, 2, 1, 2);
+		PMatrix<T, 3, 3> Covariance(0);
 		for (const auto& Element : Surfaces)
 		{
-			FMatrix33 DeltaMatrix(0);
+			PMatrix<T, 3, 3> DeltaMatrix(0);
 			for (int32 i = 0; i < Element.Num(); ++i)
 			{
-				FVec3 DeltaVector = Vertices.X(Element[i]) - CenterOfMass;
+				TVec DeltaVector = Vertices[Element[i]] - CenterOfMass;
 				DeltaMatrix.M[0][i] = DeltaVector[0];
 				DeltaMatrix.M[1][i] = DeltaVector[1];
 				DeltaMatrix.M[2][i] = DeltaVector[2];
 			}
-			FReal Det = DeltaMatrix.M[0][0] * (DeltaMatrix.M[1][1] * DeltaMatrix.M[2][2] - DeltaMatrix.M[1][2] * DeltaMatrix.M[2][1]) -
+			T Det = DeltaMatrix.M[0][0] * (DeltaMatrix.M[1][1] * DeltaMatrix.M[2][2] - DeltaMatrix.M[1][2] * DeltaMatrix.M[2][1]) -
 				DeltaMatrix.M[0][1] * (DeltaMatrix.M[1][0] * DeltaMatrix.M[2][2] - DeltaMatrix.M[1][2] * DeltaMatrix.M[2][0]) +
 				DeltaMatrix.M[0][2] * (DeltaMatrix.M[1][0] * DeltaMatrix.M[2][1] - DeltaMatrix.M[1][1] * DeltaMatrix.M[2][0]);
-			const FMatrix33 ScaledStandard = Standard * Det;
+			const PMatrix<T, 3, 3> ScaledStandard = Standard * Det;
 			Covariance += DeltaMatrix * ScaledStandard * DeltaMatrix.GetTransposed();
 		}
-		FReal Trace = Covariance.M[0][0] + Covariance.M[1][1] + Covariance.M[2][2];
-		FMatrix33 TraceMat(Trace, Trace, Trace);
-		OutInertiaTensor = (TraceMat - Covariance) * (1 / (FReal)120) * Density;
+		T Trace = Covariance.M[0][0] + Covariance.M[1][1] + Covariance.M[2][2];
+		PMatrix<T, 3, 3> TraceMat(Trace, Trace, Trace);
+		OutInertiaTensor = (TraceMat - Covariance) * (1 / (T)120) * Density;
 		OutRotationOfMass = TransformToLocalSpace(OutInertiaTensor);
 	}
 
+	template<typename T, typename TSurfaces>
+	void CHAOS_API CalculateInertiaAndRotationOfMass(const TParticles<T, 3>& Vertices, const TSurfaces& Surfaces, const T Density, const TVec3<T>& CenterOfMass,
+		PMatrix<T, 3, 3>& OutInertiaTensor, TRotation<T, 3>& OutRotationOfMass)
+	{
+		CalculateInertiaAndRotationOfMassImpl(Vertices.AllX(), Surfaces, Density, CenterOfMass, OutInertiaTensor, OutRotationOfMass);
+	}
+
+	template<typename T, typename TSurfaces>
+	void CHAOS_API CalculateInertiaAndRotationOfMass(const TArray<UE::Math::TVector<T>>& Vertices, const TSurfaces& Surfaces, const T Density, const UE::Math::TVector<T>& CenterOfMass,
+		PMatrix<T,3,3>& OutInertiaTensor, TRotation<T, 3>& OutRotationOfMass)
+	{
+		CalculateInertiaAndRotationOfMassImpl(Vertices, Surfaces, Density, CenterOfMass, OutInertiaTensor, OutRotationOfMass);
+	}
+	
 	template<typename TSurfaces>
 	FMassProperties CalculateMassProperties(
 		const FParticles & Vertices,
@@ -286,6 +313,9 @@ namespace Chaos
 		return NewMP;
 	}
 
+	template CHAOS_API TRotation<FRealSingle,3> TransformToLocalSpace(PMatrix<FRealSingle,3,3>& Inertia);
+	template CHAOS_API TRotation<FRealDouble,3> TransformToLocalSpace(PMatrix<FRealDouble,3,3>& Inertia);
+	
 	template CHAOS_API FMassProperties CalculateMassProperties(const FParticles& Vertices, const TArray<TVec3<int32>>& Surfaces, const FReal Mass);
 	template CHAOS_API FMassProperties CalculateMassProperties(const FParticles & Vertices, const TArray<TArray<int32>>& Surfaces, const FReal Mass);
 
@@ -298,9 +328,19 @@ namespace Chaos
 	template CHAOS_API void CalculateVolumeAndCenterOfMass(const TArray<TVec3<FRealSingle>>& Vertices, const TArray<TArray<int32>>& Surfaces, FRealSingle& OutVolume, TVec3<FRealSingle>& OutCenterOfMass);
 	template CHAOS_API void CalculateVolumeAndCenterOfMass(const TArray<TVec3<FRealDouble>>& Vertices, const TArray<TArray<int32>>& Surfaces, FRealDouble& OutVolume, TVec3<FRealDouble>& OutCenterOfMass);
 
+	template CHAOS_API void CalculateVolumeAndCenterOfMass(const TArray<UE::Math::TVector<float>>& Vertices, const TArray<TVec3<int32>>& Surfaces, float& OutVolume, UE::Math::TVector<float>& OutCenterOfMass);
+
+	template void CHAOS_API CalculateInertiaAndRotationOfMass(const TParticles<FRealSingle,3>& Vertices, const TArray<TVec3<int32>>& Surfaces, const FRealSingle Density, const TVec3<FRealSingle>& CenterOfMass,
+		PMatrix<FRealSingle,3,3>& OutInertiaTensor, TRotation<FRealSingle, 3>& OutRotationOfMass);
+	template void CHAOS_API CalculateInertiaAndRotationOfMass(const TParticles<FRealSingle,3>& Vertices, const TArray<TArray<int32>>& Surfaces, const FRealSingle Density, const TVec3<FRealSingle>& CenterOfMass,
+		PMatrix<FRealSingle,3,3>& OutInertiaTensor, TRotation<FRealSingle, 3>& OutRotationOfMass);
 
-	template CHAOS_API void CalculateInertiaAndRotationOfMass(const FParticles& Vertices, const TArray<TVec3<int32>>& Surface, const FReal Density,
-		const FVec3& CenterOfMass, FMatrix33& OutInertiaTensor, FRotation3& OutRotationOfMass);
-	template CHAOS_API void CalculateInertiaAndRotationOfMass(const FParticles& Vertices, const TArray<TArray<int32>>& Surface, const FReal Density,
-		const FVec3& CenterOfMass, FMatrix33& OutInertiaTensor, FRotation3& OutRotationOfMass);
+	template void CHAOS_API CalculateInertiaAndRotationOfMass(const TParticles<FRealDouble,3>& Vertices, const TArray<TVec3<int32>>& Surfaces, const FRealDouble Density, const TVec3<FRealDouble>& CenterOfMass,
+		PMatrix<FRealDouble,3,3>& OutInertiaTensor, TRotation<FRealDouble, 3>& OutRotationOfMass);
+	template void CHAOS_API CalculateInertiaAndRotationOfMass(const TParticles<FRealDouble,3>& Vertices, const TArray<TArray<int32>>& Surfaces, const FRealDouble Density, const TVec3<FRealDouble>& CenterOfMass,
+		PMatrix<FRealDouble,3,3>& OutInertiaTensor, TRotation<FRealDouble, 3>& OutRotationOfMass);
+
+	template void CHAOS_API CalculateInertiaAndRotationOfMass(const TArray<UE::Math::TVector<float>>& Vertices, const TArray<TVec3<int32>>& Surfaces, const float Density, const UE::Math::TVector<float>& CenterOfMass,
+		PMatrix<float,3,3>& OutInertiaTensor, TRotation<float, 3>& OutRotationOfMass);
+	
 }
diff --git a/Engine/Source/Runtime/Experimental/Chaos/Private/PhysicsProxy/GeometryCollectionPhysicsProxy.cpp b/Engine/Source/Runtime/Experimental/Chaos/Private/PhysicsProxy/GeometryCollectionPhysicsProxy.cpp
--- a/Engine/Source/Runtime/Experimental/Chaos/Private/PhysicsProxy/GeometryCollectionPhysicsProxy.cpp
+++ b/Engine/Source/Runtime/Experimental/Chaos/Private/PhysicsProxy/GeometryCollectionPhysicsProxy.cpp
@@ -111,7 +111,7 @@ void FGeometryCollectionResults::Reset()
 //==============================================================================
 
 
-Chaos::FTriangleMesh* CreateTriangleMesh(
+TUniquePtr<Chaos::FTriangleMesh> CreateTriangleMesh(
 	const int32 FaceStart,
 	const int32 FaceCount, 
 	const TManagedArray<bool>& Visible, 
@@ -141,7 +141,7 @@ Chaos::FTriangleMesh* CreateTriangleMesh(
 			}
 		}
 	}
-	return new Chaos::FTriangleMesh(MoveTemp(Faces)); // Culls geometrically degenerate faces
+	return MakeUnique<Chaos::FTriangleMesh>(MoveTemp(Faces)); // Culls geometrically degenerate faces
 }
 
 TArray<int32> ComputeTransformToGeometryMap(const FGeometryCollection& Collection)
@@ -2283,7 +2283,7 @@ static FGeometryDynamicCollection::FSharedImplicit CreateImplicitGeometry(
 }
 
 static void ComputeGeometryVolumeAndCenterOfMass(
-	const Chaos::FParticles& MassSpaceParticles,
+	const TArray<FVector3f> Vertices,
 	const Chaos::FTriangleMesh& TriMesh,
 	const FBox& BoundingBox,
 	const FSharedSimulationParameters& SharedParams, 
@@ -2303,8 +2303,12 @@ static void ComputeGeometryVolumeAndCenterOfMass(
 	
 	if (!OutIsTooSmallGeometry)
 	{
-		Chaos::CalculateVolumeAndCenterOfMass(MassSpaceParticles, TriMesh.GetElements(), OutUncorrectedVolume, OutCenterOfMass);
-		OutVolume = OutUncorrectedVolume;
+		float Volume;
+		FVector3f CenterOfMass;
+		Chaos::CalculateVolumeAndCenterOfMass(Vertices, TriMesh.GetElements(), Volume, CenterOfMass);
+		OutVolume = Volume;
+		OutUncorrectedVolume = Volume;
+		OutCenterOfMass = CenterOfMass;
 		
 		if (OutUncorrectedVolume == 0)
 		{
@@ -2325,6 +2329,78 @@ static void ComputeGeometryVolumeAndCenterOfMass(
 	}
 }
 
+static void ComputeMassPropertyFromMesh(
+	const int32 GeometryIndex,
+	const FGeometryCollection& RestCollection,
+	const FSharedSimulationParameters& SharedParams,
+	TUniquePtr<Chaos::FTriangleMesh>& TriMesh,
+	Chaos::FMassProperties& OutMassProperties,
+	bool& bOutIsTooSmallGeometry
+)
+{
+	// VerticesGroup
+	const TManagedArray<FVector3f>& Vertex = RestCollection.Vertex;
+
+	// FacesGroup
+	const TManagedArray<bool>& Visible = RestCollection.Visible;
+	const TManagedArray<FIntVector>& Indices = RestCollection.Indices;
+
+	// GeometryGroup
+	const TManagedArray<FBox>& BoundingBox = RestCollection.BoundingBox;
+	const TManagedArray<int32>& FaceStart = RestCollection.FaceStart;
+	const TManagedArray<int32>& FaceCount = RestCollection.FaceCount;
+
+	
+	TriMesh = CreateTriangleMesh(
+		FaceStart[GeometryIndex],
+		FaceCount[GeometryIndex],
+		Visible,
+		Indices);
+
+	Chaos::FReal UncorrectedVolume = 0;
+	
+	bOutIsTooSmallGeometry = false;
+	ComputeGeometryVolumeAndCenterOfMass(
+		Vertex.GetConstArray(),
+		*TriMesh,
+		BoundingBox[GeometryIndex],
+		SharedParams,
+		UncorrectedVolume,
+		OutMassProperties.Volume,
+		OutMassProperties.CenterOfMass,
+		bOutIsTooSmallGeometry
+		);
+
+	// now compute unit inertia ( assuming mass of 1 it will be scaled later by the real mass
+	constexpr Chaos::FReal UnitMass = 1000;
+	ensure (OutMassProperties.Volume > UE_SMALL_NUMBER);
+	const Chaos::FReal Density = UnitMass / OutMassProperties.Volume;
+	const bool bUseBoundingBoxInertia = (UncorrectedVolume == 0);  
+	if (bUseBoundingBoxInertia)
+	{
+		OutMassProperties.CenterOfMass = BoundingBox[GeometryIndex].GetCenter();
+		Chaos::CalculateInertiaAndRotationOfMass(BoundingBox[GeometryIndex], Density, OutMassProperties.InertiaTensor, OutMassProperties.RotationOfMass);
+	}
+	else
+	{
+		const FVector3f CenterOfMass{ OutMassProperties.CenterOfMass };
+		Chaos::PMatrix<float,3,3> InertiaTensor;
+		Chaos::TRotation<float,3> RotationOfMass;
+		CalculateInertiaAndRotationOfMass(Vertex.GetConstArray(), TriMesh->GetSurfaceElements(), (float)Density, CenterOfMass, InertiaTensor, RotationOfMass);
+		OutMassProperties.CenterOfMass = CenterOfMass;
+		OutMassProperties.InertiaTensor = InertiaTensor;
+		OutMassProperties.RotationOfMass = RotationOfMass;
+
+		// scale the inertia tensor accordingly to compensate for the volume adjustment
+		if (OutMassProperties.Volume != UncorrectedVolume)
+		{
+			const FVector::FReal InertiaScale = OutMassProperties.Volume / UncorrectedVolume;
+			OutMassProperties.InertiaTensor = OutMassProperties.InertiaTensor.ApplyScale(InertiaScale);
+		}
+	}
+}
+
+
 /** 
 	NOTE - Making any changes to data stored on the rest collection below MUST be accompanied
 	by a rotation of the DDC key in FDerivedDataGeometryCollectionCooker::GetVersionString
@@ -2364,8 +2440,8 @@ void FGeometryCollectionPhysicsProxy::InitializeSharedCollisionStructures(
 	RestCollection.RemoveAttribute(FGeometryDynamicCollection::ImplicitsAttribute, FTransformCollection::TransformGroup);
 	TManagedArray<FGeometryDynamicCollection::FSharedImplicit>& CollectionImplicits = RestCollection.AddAttribute<FGeometryDynamicCollection::FSharedImplicit>(FGeometryDynamicCollection::ImplicitsAttribute, FTransformCollection::TransformGroup);
 
-	TManagedArray<TSet<int32>>* TransformToConvexIndices = RestCollection.FindAttribute<TSet<int32>>("TransformToConvexIndices", FTransformCollection::TransformGroup);
-	TManagedArray<TUniquePtr<Chaos::FConvex>>* ConvexGeometry = RestCollection.FindAttribute<TUniquePtr<Chaos::FConvex>>("ConvexHull", "Convex");
+	//TManagedArray<TSet<int32>>* TransformToConvexIndices = RestCollection.FindAttribute<TSet<int32>>("TransformToConvexIndices", FTransformCollection::TransformGroup);
+	//TManagedArray<TUniquePtr<Chaos::FConvex>>* ConvexGeometry = RestCollection.FindAttribute<TUniquePtr<Chaos::FConvex>>("ConvexHull", "Convex");
 
 	bool bUseRelativeSize = RestCollection.HasAttribute(TEXT("Size"), FTransformCollection::TransformGroup);
 	if (!bUseRelativeSize)
@@ -2396,8 +2472,7 @@ void FGeometryCollectionPhysicsProxy::InitializeSharedCollisionStructures(
 	const TManagedArray<int32>& VertexCount = RestCollection.VertexCount;
 	const TManagedArray<int32>& FaceStart = RestCollection.FaceStart;
 	const TManagedArray<int32>& FaceCount = RestCollection.FaceCount;
-
-
+	
 	TArray<FTransform> CollectionSpaceTransforms;
 	{ // tmp scope
 		const TManagedArray<FTransform>& HierarchyTransform = RestCollection.Transform;
@@ -2410,24 +2485,19 @@ void FGeometryCollectionPhysicsProxy::InitializeSharedCollisionStructures(
 	TArray<TUniquePtr<FTriangleMesh>> TriangleMeshesArray;	//use to union trimeshes in cluster case
 	TriangleMeshesArray.AddDefaulted(NumTransforms);
 
-	FParticles MassSpaceParticles;
-	MassSpaceParticles.AddParticles(Vertex.Num());
-	for (int32 Idx = 0; Idx < Vertex.Num(); ++Idx)
-	{
-		MassSpaceParticles.X(Idx) = Vertex[Idx];	//mass space computation done later down
-	}
+	// FParticles MassSpaceParticles;
+	// MassSpaceParticles.AddParticles(Vertex.Num());
+	// for (int32 Idx = 0; Idx < Vertex.Num(); ++Idx)
+	// {
+	// 	MassSpaceParticles.X(Idx) = Vertex[Idx];	//mass space computation done later down
+	// }
 
 	TArray<FMassProperties> MassPropertiesArray;
 	MassPropertiesArray.AddUninitialized(NumGeometries);
 
-	TArray<FVector::FReal> UncorrectedVolume;
-	UncorrectedVolume.Init(false, NumGeometries);
+	//TArray<FVector::FReal> UncorrectedVolume;
+	//UncorrectedVolume.Init(0, NumGeometries);
 
-	// We skip very small geometry and log as a warning. To avoid log spamming, we wait
-	// until we complete the loop before reporting the skips.
-	bool bSkippedSmallGeometry = false;
-	
-	FReal TotalVolume = 0.f;
 	// The geometry group has a set of transform indices that maps a geometry index
 	// to a transform index, but only in the case where there is a 1-to-1 mapping 
 	// between the two.  In the event where a geometry is instanced for multiple
@@ -2435,66 +2505,41 @@ void FGeometryCollectionPhysicsProxy::InitializeSharedCollisionStructures(
 	// Otherwise, iterating over the geometry group is a convenient way to iterate
 	// over all the leaves of the hierarchy.
 	check(!TransformIndex.Contains(INDEX_NONE)); // TODO: implement support for instanced bodies
-	// @todo(chaos) could this be in a parallel for and having the total volume computed in a tight loop right after? 
-	for (int32 GeometryIndex = 0; GeometryIndex < NumGeometries; GeometryIndex++)
+
+	TArray<bool> IsTooSmallGeometryArray;
+	IsTooSmallGeometryArray.Init(false, NumGeometries);
+
+	// compute mass properties in parallel
+	ParallelFor(NumGeometries, [&](int32 GeometryIndex)
 	{
-		const int32 TransformGroupIndex = TransformIndex[GeometryIndex];
-		
 		// @chaos(todo): is that really needed? 
 		GenerateInnerAndOuterRadiiIfNeeded(RestCollection, GeometryIndex);
 
-		CollectionSimulatableParticles[TransformGroupIndex] = false;
-		if (SimulationType[TransformGroupIndex] > FGeometryCollection::ESimulationTypes::FST_None)
+		const int32 TransformGroupIndex = TransformIndex[GeometryIndex];
+		CollectionSimulatableParticles[TransformGroupIndex] = (SimulationType[TransformGroupIndex] > FGeometryCollection::ESimulationTypes::FST_None);
+		
+		if (CollectionSimulatableParticles[TransformGroupIndex])
 		{
-			TUniquePtr<FTriangleMesh> TriMesh(
-				CreateTriangleMesh(
-					FaceStart[GeometryIndex],
-					FaceCount[GeometryIndex],
-					Visible,
-					Indices));
-
-			// Compute volume and center of mass 
-			FMassProperties& MassProperties = MassPropertiesArray[GeometryIndex];
-			{
-				bool bIsTooSmallGeometry = false;
-				ComputeGeometryVolumeAndCenterOfMass(
-					MassSpaceParticles,
-					*TriMesh,
-					BoundingBox[GeometryIndex],
-					SharedParams,
-					UncorrectedVolume[GeometryIndex],
-					MassProperties.Volume,
-					MassProperties.CenterOfMass,
-					bIsTooSmallGeometry
-					);
-				
-				if (bIsTooSmallGeometry)
-				{
-					bSkippedSmallGeometry = true;
-				}
-				else
-				{
-					CollectionSimulatableParticles[TransformGroupIndex] = true;
-				
-					// shift the particles to the mass local space
-					const FVector MassTranslation = MassProperties.CenterOfMass;
-					if (!FMath::IsNearlyZero(MassTranslation.SizeSquared()))
-					{
-						const int32 IdxStart = VertexStart[GeometryIndex];
-						const int32 IdxEnd = IdxStart + VertexCount[GeometryIndex];
-						for (int32 Idx = IdxStart; Idx < IdxEnd; ++Idx)
-						{
-							MassSpaceParticles.X(Idx) -= MassTranslation;
-						}
-					}
-				}
-
-				// update Mass to local matrix
-				CollectionMassToLocal[TransformGroupIndex] = FTransform(FQuat::Identity, MassProperties.CenterOfMass);
-			}
-
-			TotalVolume += MassProperties.Volume;
-			TriangleMeshesArray[TransformGroupIndex] = MoveTemp(TriMesh);
+			ComputeMassPropertyFromMesh(GeometryIndex, RestCollection, SharedParams, TriangleMeshesArray[GeometryIndex], MassPropertiesArray[GeometryIndex], IsTooSmallGeometryArray[GeometryIndex]);
+		}
+	});
+	
+	// compute total volume and whether we need to skip small geometry ( can't use parallelfors )
+	bool bSkippedSmallGeometry = false;
+	FReal TotalVolume = 0.f;
+	for (int32 GeometryIndex = 0; GeometryIndex < NumGeometries; GeometryIndex++)
+	{
+		const int32 TransformGroupIndex = TransformIndex[GeometryIndex];
+		if (CollectionSimulatableParticles[TransformGroupIndex])
+		{
+			TotalVolume += MassPropertiesArray[GeometryIndex].Volume;
+		}
+		// We skip very small geometry and log as a warning. To avoid log spamming, we wait
+		// until we complete the loop before reporting the skips.
+		if (IsTooSmallGeometryArray[GeometryIndex])
+		{
+			CollectionSimulatableParticles[TransformGroupIndex] = false;
+			bSkippedSmallGeometry = true;
 		}
 	}
 
@@ -2503,7 +2548,7 @@ void FGeometryCollectionPhysicsProxy::InitializeSharedCollisionStructures(
 	{
 		UE_LOG(LogChaos, Warning, TEXT("Some geometry is too small to be simulated and has been skipped."));
 	}
-
+	
 	// User provides us with total mass or density.
 	// Density must be the same for individual parts and the total. Density_i = Density = Mass_i / Volume_i
 	// Total mass must equal sum of individual parts. Mass_i = TotalMass * Volume_i / TotalVolume => Density_i = TotalMass / TotalVolume
@@ -2512,9 +2557,7 @@ void FGeometryCollectionPhysicsProxy::InitializeSharedCollisionStructures(
 	const FReal ClampedTotalMass = FMath::Clamp(DesiredTotalMass, SharedParams.MinimumMassClamp, SharedParams.MaximumMassClamp);
 	const FReal DesiredDensity = ClampedTotalMass / TotalVolume;
 
-	FVec3 MaxChildBounds(1);
 	ParallelFor(NumGeometries, [&](int32 GeometryIndex)
-	//for (int32 GeometryIndex = 0; GeometryIndex < NumGeometries; GeometryIndex++)
 	{
 		// Need a new error reporter for parallel for loop here as it wouldn't be thread-safe to write to the prefix
 		Chaos::FErrorReporter LocalErrorReporter;
@@ -2544,43 +2587,51 @@ void FGeometryCollectionPhysicsProxy::InitializeSharedCollisionStructures(
 			const FReal Density_i = Mass_i / Volume_i;
 			CollectionMass[TransformGroupIndex] = (FRealSingle)Mass_i;
 
-			const bool bUseBoundingBoxInertia = (UncorrectedVolume[GeometryIndex] == 0);  
-			if (bUseBoundingBoxInertia)
-			{
-				CalculateInertiaAndRotationOfMass(BoundingBox[GeometryIndex], Density_i, MassProperties.InertiaTensor, MassProperties.RotationOfMass);
-			}
-			else
-			{
-				// Note: particles already in CoM space, so passing in zero as CoM
-				CalculateInertiaAndRotationOfMass(MassSpaceParticles, TriMesh->GetSurfaceElements(), Density_i, FVec3(0), MassProperties.InertiaTensor, MassProperties.RotationOfMass);
-				// scale the inertia tensor accordingly
-				if (MassProperties.Volume != UncorrectedVolume[GeometryIndex])
-				{
-					const FVector::FReal InertiaScale = MassProperties.Volume / UncorrectedVolume[GeometryIndex];
-					MassProperties.InertiaTensor = MassProperties.InertiaTensor.ApplyScale(InertiaScale);
-				}
-			}
-			
+			// scale the mass properties inertia tensor by the mass
+			MassPropertiesArray[GeometryIndex].InertiaTensor *= Mass_i / 1000.0;
+			const Chaos::FMatrix33& InertiaTensor = MassPropertiesArray[GeometryIndex].InertiaTensor;
+			CollectionInertiaTensor[TransformGroupIndex] = FVector3f((float)InertiaTensor.M[0][0], (float)InertiaTensor.M[1][1], (float)InertiaTensor.M[2][2]);
+
+			CollectionMassToLocal[TransformGroupIndex] = FTransform(MassProperties.RotationOfMass, MassProperties.CenterOfMass);
+		}
+	});
+
+	// we need the vertices in mass space to compute the collision particles
+	// todo(chaos) we shoudl eventually get rid of that that's a waste of memory and performance
+	FParticles MassSpaceParticles;
+	MassSpaceParticles.AddParticles(Vertex.Num());
+	for (int32 Idx = 0; Idx < Vertex.Num(); ++Idx)
+	{
+		MassSpaceParticles.X(Idx) = Vertex[Idx];	//mass space computation done later down
+	}
+	
+	// compute the bounding box in mesh space
+	FVec3 MaxChildBounds(1);
+	ParallelFor(NumGeometries, [&](int32 GeometryIndex)
+	{
+		const int32 TransformGroupIndex = TransformIndex[GeometryIndex];
+		if (CollectionSimulatableParticles[TransformGroupIndex])
+		{
+			const TUniquePtr<Chaos::FTriangleMesh>& TriMesh = TriangleMeshesArray[GeometryIndex];
+			const FTransform& MassToLocalTransform = CollectionMassToLocal[TransformGroupIndex];
+
 			// Orient particle in the local mass space 
-			if (!MassProperties.RotationOfMass.Equals(FQuat::Identity))
+			if (!MassPropertiesArray[TransformGroupIndex].RotationOfMass.Equals(FQuat::Identity) ||
+				!FMath::IsNearlyZero(MassPropertiesArray[TransformGroupIndex].CenterOfMass.SizeSquared()))
 			{
-				FTransform InverseMassRotation = FTransform(MassProperties.RotationOfMass.Inverse());
 				const int32 IdxStart = VertexStart[GeometryIndex];
 				const int32 IdxEnd = IdxStart + VertexCount[GeometryIndex];
 				for (int32 Idx = IdxStart; Idx < IdxEnd; ++Idx)
 				{
-					MassSpaceParticles.X(Idx) = InverseMassRotation.TransformPosition(MassSpaceParticles.X(Idx));
+					MassSpaceParticles.X(Idx) = MassToLocalTransform.InverseTransformPositionNoScale(MassSpaceParticles.X(Idx));
 				}
 			}
 
-			// update Rest collection properties
-			CollectionInertiaTensor[TransformGroupIndex] = FVector3f((float)MassProperties.InertiaTensor.M[0][0], (float)MassProperties.InertiaTensor.M[1][1], (float)MassProperties.InertiaTensor.M[2][2]);
-			CollectionMassToLocal[TransformGroupIndex] = FTransform(MassProperties.RotationOfMass, MassProperties.CenterOfMass);
-
-
 			FBox InstanceBoundingBox(EForceInit::ForceInitToZero);
 			if (TriMesh->GetElements().Num())
 			{
+				// really not optimized as this generate a full set of the vertices with all the allocation that comes with it
+				// the cost of going through all the vertices multiple time  may be faster 
 				const TSet<int32> MeshVertices = TriMesh->GetVertices();
 				for (const int32 Idx : MeshVertices)
 				{
@@ -2598,12 +2649,15 @@ void FGeometryCollectionPhysicsProxy::InitializeSharedCollisionStructures(
 			}
 			else
 			{
-				InstanceBoundingBox = FBox(MassProperties.CenterOfMass, MassProperties.CenterOfMass);
+				const Chaos::FVec3 CenterOfMass = MassPropertiesArray[TransformGroupIndex].CenterOfMass;
+				InstanceBoundingBox = FBox(CenterOfMass, CenterOfMass);
 			}
 
 			const FSharedSimulationSizeSpecificData& SizeSpecificData = GetSizeSpecificData(SharedParams.SizeSpecificData, RestCollection, TransformGroupIndex, InstanceBoundingBox);
 			if (SizeSpecificData.CollisionShapesData.Num())
 			{
+				// Need a new error reporter for parallel for loop here as it wouldn't be thread-safe to write to the prefix
+				Chaos::FErrorReporter LocalErrorReporter;
 				//
 				//  Build the simplicial for the rest collection. This will be used later in the DynamicCollection to 
 				//  populate the collision structures of the simulation. 
diff --git a/Engine/Source/Runtime/Experimental/Chaos/Public/Chaos/MassProperties.h b/Engine/Source/Runtime/Experimental/Chaos/Public/Chaos/MassProperties.h
--- a/Engine/Source/Runtime/Experimental/Chaos/Public/Chaos/MassProperties.h
+++ b/Engine/Source/Runtime/Experimental/Chaos/Public/Chaos/MassProperties.h
@@ -33,7 +33,8 @@ namespace Chaos
 		FMatrix33 InertiaTensor;
 	};
 
-	FRotation3 CHAOS_API TransformToLocalSpace(FMatrix33& Inertia);
+	template <typename T>
+	TRotation<T,3> CHAOS_API TransformToLocalSpace(PMatrix<T,3,3>& Inertia);
 	void CHAOS_API TransformToLocalSpace(FMassProperties& MassProperties);
 
 	template<typename T, typename TSurfaces>
@@ -42,13 +43,20 @@ namespace Chaos
 	template<typename T, typename TSurfaces>
 	void CHAOS_API CalculateVolumeAndCenterOfMass(const TArray<TVec3<T>>& Vertices, const TSurfaces& Surfaces, T& OutVolume, TVec3<T>& OutCenterOfMass);
 
+	template<typename T, typename TSurfaces>
+	void CHAOS_API CalculateVolumeAndCenterOfMass(const TArray<UE::Math::TVector<T>>& Vertices, const TSurfaces& Surfaces, T& OutVolume, UE::Math::TVector<T>& OutCenterOfMass);
+	
 	template<typename TSurfaces>
 	FMassProperties CHAOS_API CalculateMassProperties(const FParticles& Vertices, const TSurfaces& Surfaces, const FReal Mass);
 
-	template<typename TSurfaces>
-	void CHAOS_API CalculateInertiaAndRotationOfMass(const FParticles& Vertices, const TSurfaces& Surfaces, const FReal Density, const FVec3& CenterOfMass,
-		FMatrix33& OutInertiaTensor, FRotation3& OutRotationOfMass);
+	template<typename T, typename TSurfaces>
+	void CHAOS_API CalculateInertiaAndRotationOfMass(const TParticles<T,3>& Vertices, const TSurfaces& Surfaces, const T Density, const TVec3<T>& CenterOfMass,
+		PMatrix<T,3,3>& OutInertiaTensor, TRotation<T, 3>& OutRotationOfMass);
 
+	template<typename T, typename TSurfaces>
+	void CHAOS_API CalculateInertiaAndRotationOfMass(const TArray<UE::Math::TVector<T>>& Vertices, const TSurfaces& Surfaces, const T Density, const UE::Math::TVector<T>& CenterOfMass,
+		PMatrix<T,3,3>& OutInertiaTensor, TRotation<T, 3>& OutRotationOfMass);
+	
 	void CHAOS_API CalculateVolumeAndCenterOfMass(const FBox& BoundingBox, FVector::FReal& OutVolume, FVector& OutCenterOfMass);
 	void CHAOS_API CalculateInertiaAndRotationOfMass(const FBox& BoundingBox, const FVector::FReal Density, FMatrix33& OutInertiaTensor, FRotation3& OutRotationOfMass);
 
diff --git a/Engine/Source/Runtime/Experimental/Chaos/Public/GeometryCollection/ManagedArray.h b/Engine/Source/Runtime/Experimental/Chaos/Public/GeometryCollection/ManagedArray.h
--- a/Engine/Source/Runtime/Experimental/Chaos/Public/GeometryCollection/ManagedArray.h
+++ b/Engine/Source/Runtime/Experimental/Chaos/Public/GeometryCollection/ManagedArray.h
@@ -342,6 +342,21 @@ public:
 		return Array[Index];
 	}
 
+	/**
+	* Helper function for returning the internal const array
+	*
+	* @returns const array of all the elements
+	*/
+	FORCEINLINE const TArray<ElementType>& GetConstArray()
+	{
+		return Array;
+	}
+
+	FORCEINLINE const TArray<ElementType>& GetConstArray() const
+	{
+		return Array;
+	}
+	
 	/**
 	* Helper function for returning a typed pointer to the first array entry.
 	*
diff --git a/Engine/Source/Runtime/Experimental/Chaos/Public/PhysicsProxy/GeometryCollectionPhysicsProxy.h b/Engine/Source/Runtime/Experimental/Chaos/Public/PhysicsProxy/GeometryCollectionPhysicsProxy.h
--- a/Engine/Source/Runtime/Experimental/Chaos/Public/PhysicsProxy/GeometryCollectionPhysicsProxy.h
+++ b/Engine/Source/Runtime/Experimental/Chaos/Public/PhysicsProxy/GeometryCollectionPhysicsProxy.h
@@ -373,5 +373,5 @@ private:
 	FProxyInterpolationData InterpolationData; 
 };
 
-CHAOS_API Chaos::FTriangleMesh* CreateTriangleMesh(const int32 FaceStart,const int32 FaceCount,const TManagedArray<bool>& Visible,const TManagedArray<FIntVector>& Indices, bool bRotateWinding = true);
+CHAOS_API TUniquePtr<Chaos::FTriangleMesh> CreateTriangleMesh(const int32 FaceStart,const int32 FaceCount,const TManagedArray<bool>& Visible,const TManagedArray<FIntVector>& Indices, bool bRotateWinding = true);
 CHAOS_API void BuildSimulationData(Chaos::FErrorReporter& ErrorReporter, FGeometryCollection& GeometryCollection, const FSharedSimulationParameters& SharedParams);
