Chaos
- fix SLerp drive axis calculation
- remove single-constraint SLerp drive fast-path (does not work for high stiffness)

diff --git a/Engine/Source/Runtime/Experimental/Chaos/Private/Chaos/Joint/PBDJointCachedSolverGaussSeidel.cpp b/Engine/Source/Runtime/Experimental/Chaos/Private/Chaos/Joint/PBDJointCachedSolverGaussSeidel.cpp
--- a/Engine/Source/Runtime/Experimental/Chaos/Private/Chaos/Joint/PBDJointCachedSolverGaussSeidel.cpp
+++ b/Engine/Source/Runtime/Experimental/Chaos/Private/Chaos/Joint/PBDJointCachedSolverGaussSeidel.cpp
@@ -1695,58 +1695,32 @@ void FPBDJointCachedSolver::InitSLerpDrive(
 	}
 	RotationDrives.bAccelerationMode = FPBDJointUtilities::GetDriveAccelerationMode(SolverSettings, JointSettings);
 
-	// If damping is enabled, we need to apply the drive about all 3 axes, but without damping we can just drive along the axis of error
-	if (RotationDrives.ConstraintSoftDamping[0]  > 0.0f)
-	{
-		// NOTE: Slerp target velocity only works properly if we have a stiffness of zero.
-		FVec3 Axes[3] = { FVec3(1, 0, 0), FVec3(0, 1, 0), FVec3(0, 0, 1) };
-		if (RotationDrives.ConstraintSoftStiffness[0] > 0.0f)
-		{
-			FPBDJointUtilities::GetLockedRotationAxes(ConnectorRs[0], ConnectorRs[1], Axes[0], Axes[1], Axes[2]);
-			Utilities::NormalizeSafe(Axes[0], UE_KINDA_SMALL_NUMBER);
-			Utilities::NormalizeSafe(Axes[1], UE_KINDA_SMALL_NUMBER);
-			Utilities::NormalizeSafe(Axes[2], UE_KINDA_SMALL_NUMBER);
-		}
-		const FRotation3 R01 = ConnectorRs[0].Inverse() * ConnectorRs[1];
-		FRotation3 TargetAngPos = JointSettings.AngularDrivePositionTarget;
-		TargetAngPos.EnforceShortestArcWith(R01);
-		const FRotation3 R1Error = TargetAngPos.Inverse() * R01;
-		FReal AxisAngles[3] = 
-		{ 
-			2.0f * FMath::Asin(R1Error.X), 
-			2.0f * FMath::Asin(R1Error.Y), 
-			2.0f * FMath::Asin(R1Error.Z) 
-		};
-
-		const FVec3 TargetAngVel = ConnectorRs[0] * JointSettings.AngularDriveVelocityTarget;
+	const FRotation3 R01 = ConnectorRs[0].Inverse() * ConnectorRs[1];
+	FRotation3 TargetAngPos = JointSettings.AngularDrivePositionTarget;
+	TargetAngPos.EnforceShortestArcWith(R01);
+	const FRotation3 R1Error = TargetAngPos.Inverse() * R01;
+
+	// @todo(chaos): try approximation for Asin. Try (x), (x + (1/6) x^3), or (x + (1/2) x^7)
+	const FReal AxisAngles[3] = 
+	{ 
+		2.0f * FMath::Asin(R1Error.X), 
+		2.0f * FMath::Asin(R1Error.Y), 
+		2.0f * FMath::Asin(R1Error.Z) 
+	};
 
-		for (int32 AxisIndex = 0; AxisIndex < 3; ++AxisIndex)
-		{
-			InitRotationConstraintDrive(AxisIndex, Axes[AxisIndex],  Dt, AxisAngles[AxisIndex]);
-			RotationDrives.ConstraintVX[AxisIndex] = FVec3::DotProduct(TargetAngVel, RotationDrives.ConstraintAxis[AxisIndex]);
-		}
-	}
-	else
-	{
-		const FRotation3 TargetR1 = ConnectorRs[0] * JointSettings.AngularDrivePositionTarget;
-		const FRotation3 DR = TargetR1 * ConnectorRs[1].Inverse();
+	const FRotation3& AxesRotation = ConnectorRs[1];
+	const FVec3 Axes[3] = {
+		AxesRotation.GetAxisX(),
+		AxesRotation.GetAxisY(),
+		AxesRotation.GetAxisZ()
+	};
 
-		FVec3 SLerpAxis;
-		FReal SLerpAngle;
-		if (DR.ToAxisAndAngleSafe(SLerpAxis, SLerpAngle, FVec3(1, 0, 0)))
-		{
-			if (SLerpAngle > (FReal)UE_PI)
-			{
-				SLerpAngle = SLerpAngle - (FReal)2 * UE_PI;
-			}
+	const FVec3 TargetAngVel = ConnectorRs[0] * JointSettings.AngularDriveVelocityTarget;
 
-			if (FMath::Abs(SLerpAngle) > AngleTolerance)
-			{
-				FReal AngVelTarget = (JointSettings.AngularDriveDamping > FReal(0)) ? FVec3::DotProduct(SLerpAxis, ConnectorRs[0] * JointSettings.AngularDriveVelocityTarget) : 0.0f;
-				InitRotationConstraintDrive((int32)EJointAngularConstraintIndex::Swing1, SLerpAxis, Dt, -SLerpAngle);
-				RotationDrives.ConstraintVX[(int32)EJointAngularConstraintIndex::Swing1] = AngVelTarget;
-			}
-		}
+	for (int32 AxisIndex = 0; AxisIndex < 3; ++AxisIndex)
+	{
+		InitRotationConstraintDrive(AxisIndex, Axes[AxisIndex],  Dt, AxisAngles[AxisIndex]);
+		RotationDrives.ConstraintVX[AxisIndex] = FVec3::DotProduct(TargetAngVel, RotationDrives.ConstraintAxis[AxisIndex]);
 	}
 }
 
diff --git a/Engine/Source/Runtime/Experimental/Chaos/Private/Chaos/Joint/PBDJointSolverGaussSeidel.cpp b/Engine/Source/Runtime/Experimental/Chaos/Private/Chaos/Joint/PBDJointSolverGaussSeidel.cpp
--- a/Engine/Source/Runtime/Experimental/Chaos/Private/Chaos/Joint/PBDJointSolverGaussSeidel.cpp
+++ b/Engine/Source/Runtime/Experimental/Chaos/Private/Chaos/Joint/PBDJointSolverGaussSeidel.cpp
@@ -1960,64 +1960,33 @@ namespace Chaos
 		const FReal AngularDriveDamping = FPBDJointUtilities::GetAngularSLerpDriveDamping(SolverSettings, JointSettings);
 		const bool bAccelerationMode = FPBDJointUtilities::GetDriveAccelerationMode(SolverSettings, JointSettings);
 
-		// If damping is enabled, we need to apply the drive about all 3 axes, but without damping we can just drive along the axis of error
-		if (AngularDriveDamping > 0.0f)
-		{
-			// NOTE: Slerp target velocity only works properly if we have a stiffness of zero.
-			FVec3 Axes[3] = { FVec3(1, 0, 0), FVec3(0, 1, 0), FVec3(0, 0, 1) };
-			if (AngularDriveStiffness > 0.0f)
-			{
-				FPBDJointUtilities::GetLockedRotationAxes(ConnectorRs[0], ConnectorRs[1], Axes[0], Axes[1], Axes[2]);
-				Utilities::NormalizeSafe(Axes[0], UE_KINDA_SMALL_NUMBER);
-				Utilities::NormalizeSafe(Axes[1], UE_KINDA_SMALL_NUMBER);
-				Utilities::NormalizeSafe(Axes[2], UE_KINDA_SMALL_NUMBER);
-			}
+		// NOTE: Slerp target velocity only works properly if we have a stiffness of zero.
+		const FRotation3 R01 = ConnectorRs[0].Inverse() * ConnectorRs[1];
+		FRotation3 TargetAngPos = JointSettings.AngularDrivePositionTarget;
+		TargetAngPos.EnforceShortestArcWith(R01);
+		const FRotation3 R1Error = TargetAngPos.Inverse() * R01;
+		FReal AxisAngles[3] = 
+		{ 
+			2.0f * FMath::Asin(R1Error.X), 
+			2.0f * FMath::Asin(R1Error.Y), 
+			2.0f * FMath::Asin(R1Error.Z) 
+		};
 
-			const FRotation3 R01 = ConnectorRs[0].Inverse() * ConnectorRs[1];
-			FRotation3 TargetAngPos = JointSettings.AngularDrivePositionTarget;
-			TargetAngPos.EnforceShortestArcWith(R01);
-			const FRotation3 R1Error = TargetAngPos.Inverse() * R01;
-			FReal AxisAngles[3] = 
-			{ 
-				2.0f * FMath::Asin(R1Error.X), 
-				2.0f * FMath::Asin(R1Error.Y), 
-				2.0f * FMath::Asin(R1Error.Z) 
-			};
+		const FRotation3& AxesRotation = ConnectorRs[1];
+		const FVec3 Axes[3] = {
+			AxesRotation.GetAxisX(),
+			AxesRotation.GetAxisY(),
+			AxesRotation.GetAxisZ()
+		};
 
-			const FVec3 TargetAngVel = ConnectorRs[0] * JointSettings.AngularDriveVelocityTarget;
+		const FVec3 TargetAngVel = ConnectorRs[0] * JointSettings.AngularDriveVelocityTarget;
 
-			for (int32 AxisIndex = 0; AxisIndex < 3; ++AxisIndex)
-			{
-				FReal ReturnedLambda = RotationDriveLambdas[AxisIndex];
-				const FReal AxisAngVel = FVec3::DotProduct(TargetAngVel, Axes[AxisIndex]);
-				ApplyRotationConstraintSoft(Dt, AngularDriveStiffness, AngularDriveDamping, bAccelerationMode, Axes[AxisIndex], AxisAngles[AxisIndex], AxisAngVel, ReturnedLambda);
-				RotationDriveLambdas[AxisIndex] = ReturnedLambda;
-			}
-		}
-		else
+		for (int32 AxisIndex = 0; AxisIndex < 3; ++AxisIndex)
 		{
-			const FRotation3 TargetR1 = ConnectorRs[0] * JointSettings.AngularDrivePositionTarget;
-			const FRotation3 DR = TargetR1 * ConnectorRs[1].Inverse();
-
-			FVec3 SLerpAxis;
-			FReal SLerpAngle;
-			if (DR.ToAxisAndAngleSafe(SLerpAxis, SLerpAngle, FVec3(1, 0, 0)))
-			{
-				if (SLerpAngle > (FReal)UE_PI)
-				{
-					SLerpAngle = SLerpAngle - (FReal)2 * UE_PI;
-				}
-
-				UE_LOG(LogChaosJoint, VeryVerbose, TEXT("      SLerpDrive Pos: %f Axis: %f %f %f"), -SLerpAngle, SLerpAxis.X, SLerpAxis.Y, SLerpAxis.Z);
-
-				if (FMath::Abs(SLerpAngle) > AngleTolerance)
-				{
-					FReal ReturnedLambda = RotationDriveLambdas[(int32)EJointAngularConstraintIndex::Swing1];
-					FReal AngVelTarget = (JointSettings.AngularDriveDamping > FReal(0)) ? FVec3::DotProduct(SLerpAxis, ConnectorRs[0] * JointSettings.AngularDriveVelocityTarget) : 0.0f;
-					ApplyRotationConstraintSoft(Dt, AngularDriveStiffness, AngularDriveDamping, bAccelerationMode, SLerpAxis, -SLerpAngle, AngVelTarget, ReturnedLambda);
-					RotationDriveLambdas[(int32)EJointAngularConstraintIndex::Swing1] = ReturnedLambda;
-				}
-			}
+			FReal ReturnedLambda = RotationDriveLambdas[AxisIndex];
+			const FReal AxisAngVel = FVec3::DotProduct(TargetAngVel, Axes[AxisIndex]);
+			ApplyRotationConstraintSoft(Dt, AngularDriveStiffness, AngularDriveDamping, bAccelerationMode, Axes[AxisIndex], AxisAngles[AxisIndex], AxisAngVel, ReturnedLambda);
+			RotationDriveLambdas[AxisIndex] = ReturnedLambda;
 		}
 	}
 
