Chaos : fix wrong indexing of inertia tensors that can cause potential crash

diff --git a/Engine/Source/Runtime/Experimental/Chaos/Private/PhysicsProxy/GeometryCollectionPhysicsProxy.cpp b/Engine/Source/Runtime/Experimental/Chaos/Private/PhysicsProxy/GeometryCollectionPhysicsProxy.cpp
--- a/Engine/Source/Runtime/Experimental/Chaos/Private/PhysicsProxy/GeometryCollectionPhysicsProxy.cpp
+++ b/Engine/Source/Runtime/Experimental/Chaos/Private/PhysicsProxy/GeometryCollectionPhysicsProxy.cpp
@@ -2616,8 +2616,8 @@ void FGeometryCollectionPhysicsProxy::InitializeSharedCollisionStructures(
 			const FTransform& MassToLocalTransform = CollectionMassToLocal[TransformGroupIndex];
 
 			// Orient particle in the local mass space 
-			if (!MassPropertiesArray[TransformGroupIndex].RotationOfMass.Equals(FQuat::Identity) ||
-				!FMath::IsNearlyZero(MassPropertiesArray[TransformGroupIndex].CenterOfMass.SizeSquared()))
+			if (!MassPropertiesArray[GeometryIndex].RotationOfMass.Equals(FQuat::Identity) ||
+				!FMath::IsNearlyZero(MassPropertiesArray[GeometryIndex].CenterOfMass.SizeSquared()))
 			{
 				const int32 IdxStart = VertexStart[GeometryIndex];
 				const int32 IdxEnd = IdxStart + VertexCount[GeometryIndex];
@@ -2649,7 +2649,7 @@ void FGeometryCollectionPhysicsProxy::InitializeSharedCollisionStructures(
 			}
 			else
 			{
-				const Chaos::FVec3 CenterOfMass = MassPropertiesArray[TransformGroupIndex].CenterOfMass;
+				const Chaos::FVec3 CenterOfMass = MassPropertiesArray[GeometryIndex].CenterOfMass;
 				InstanceBoundingBox = FBox(CenterOfMass, CenterOfMass);
 			}
 
