Chaos debug draw : do not draw transforms for disabled particles

diff --git a/Engine/Source/Runtime/Experimental/Chaos/Private/Chaos/ChaosDebugDraw.cpp b/Engine/Source/Runtime/Experimental/Chaos/Private/Chaos/ChaosDebugDraw.cpp
--- a/Engine/Source/Runtime/Experimental/Chaos/Private/Chaos/ChaosDebugDraw.cpp
+++ b/Engine/Source/Runtime/Experimental/Chaos/Private/Chaos/ChaosDebugDraw.cpp
@@ -786,6 +786,12 @@ namespace Chaos
 
 		void DrawParticleTransformImpl(const FRigidTransform3& SpaceTransform, const FGeometryParticleHandle* InParticle, int32 Index, FRealSingle ColorScale, const FChaosDebugDrawSettings& Settings)
 		{
+			const TPBDRigidParticleHandle<FReal, 3>* Rigid = InParticle->CastToRigidParticle();
+			if (Rigid && Rigid->Disabled())
+			{
+				return;
+			}
+			
 			FColor Red = (ColorScale * FColor::Red).ToFColor(false);
 			FColor Green = (ColorScale * FColor::Green).ToFColor(false);
 			FColor Blue = (ColorScale * FColor::Blue).ToFColor(false);
@@ -821,7 +827,7 @@ namespace Chaos
 
 			if (Settings.InertiaScale > 0.0f)
 			{
-				if (const TPBDRigidParticleHandle<FReal, 3>* Rigid = InParticle->CastToRigidParticle())
+				if (Rigid)
 				{
 					if (Rigid->InvM() != 0)
 					{
