Chaos : GeometryCollection
- Default template parameter for NewCopy

diff --git a/Engine/Source/Runtime/Experimental/Chaos/Public/GeometryCollection/ManagedArrayCollection.h b/Engine/Source/Runtime/Experimental/Chaos/Public/GeometryCollection/ManagedArrayCollection.h
--- a/Engine/Source/Runtime/Experimental/Chaos/Public/GeometryCollection/ManagedArrayCollection.h
+++ b/Engine/Source/Runtime/Experimental/Chaos/Public/GeometryCollection/ManagedArrayCollection.h
@@ -113,7 +113,7 @@ public:
 	* Duplicate the ManagedArrayCollection as the specified templated type
 	* @return unowned pointer to new collection.
 	*/
-	template<class T>
+	template<class T = FManagedArrayCollection>
 	T* NewCopy() const
 	{
 		T* Collection = new T();
