Dataflow : DeepCopy

diff --git a/Engine/Source/Runtime/Experimental/Dataflow/Core/Public/Dataflow/DataflowConnection.h b/Engine/Source/Runtime/Experimental/Dataflow/Core/Public/Dataflow/DataflowConnection.h
--- a/Engine/Source/Runtime/Experimental/Dataflow/Core/Public/Dataflow/DataflowConnection.h
+++ b/Engine/Source/Runtime/Experimental/Dataflow/Core/Public/Dataflow/DataflowConnection.h
@@ -10,6 +10,7 @@ namespace Dataflow
 	class FNode;
 
 	template<class T> inline DATAFLOWCORE_API FName GraphConnectionTypeName();
+	template<class T> inline DATAFLOWCORE_API T DeepCopy(const T&);
 
 	struct FPin
 	{
diff --git a/Engine/Source/Runtime/Experimental/Dataflow/Core/Public/Dataflow/DataflowInputOutput.h b/Engine/Source/Runtime/Experimental/Dataflow/Core/Public/Dataflow/DataflowInputOutput.h
--- a/Engine/Source/Runtime/Experimental/Dataflow/Core/Public/Dataflow/DataflowInputOutput.h
+++ b/Engine/Source/Runtime/Experimental/Dataflow/Core/Public/Dataflow/DataflowInputOutput.h
@@ -167,7 +167,7 @@ namespace Dataflow
 		void SetValue(T InVal, const FContext& Context)
 		{
 			CacheKey = Context.GetTypeHash();
-			Cache.Data = InVal;
+			Cache.Data = DeepCopy<T>(InVal);
 		}
 
 		T Evaluate(const FContext& Context)
@@ -201,21 +201,33 @@ namespace Dataflow
 // Used this macros to defined dataflow connection types 
 //
 
-#define DATAFLOW_CONNECTION_TYPE(D, a,A)			\
-	template<> inline FName D GraphConnectionTypeName<a>()	\
-			{ return FName(TEXT(#A)); }				\
-	template struct D TOutputParameters<a>;			\
-	template class D TOutput<a>;					\
-	template struct D TInputParameters<a>;			\
+#define DATAFLOW_CONNECTION_TYPE_PRIMITIVE(D, a,A)					\
+	template<> inline FName D GraphConnectionTypeName<a>()			\
+		{ return FName(TEXT(#A)); }									\
+	template<> inline a D DeepCopy<a>(const a & Val)				\
+		{ return Val; }												\
+	template struct D TOutputParameters<a>;							\
+	template class D TOutput<a>;									\
+	template struct D TInputParameters<a>;							\
 	template class D TInput<a>;
 
+#define DATAFLOW_CONNECTION_TYPE_SHARED_POINTER(D, a,A)				\
+	template<> inline FName D GraphConnectionTypeName<a>()			\
+		{ return FName(TEXT(#A)); }									\
+	template<> inline a D DeepCopy<a>(const a& Val)					\
+		{ return Val ? TSharedPtr<a::ElementType>(Val->NewCopy())	\
+		: TSharedPtr<a::ElementType>(nullptr); }					\
+	template struct D TOutputParameters<a>;							\
+	template class D TOutput<a>;									\
+	template struct D TInputParameters<a>;							\
+	template class D TInput<a>;
 
 
-DATAFLOW_CONNECTION_TYPE(DATAFLOWCORE_API, bool, Bool)
-DATAFLOW_CONNECTION_TYPE(DATAFLOWCORE_API, char, Char)
-DATAFLOW_CONNECTION_TYPE(DATAFLOWCORE_API, int, Integer)
-DATAFLOW_CONNECTION_TYPE(DATAFLOWCORE_API, uint8, UInt8)
-DATAFLOW_CONNECTION_TYPE(DATAFLOWCORE_API, float, Float)
-DATAFLOW_CONNECTION_TYPE(DATAFLOWCORE_API, double, Double)
+DATAFLOW_CONNECTION_TYPE_PRIMITIVE(DATAFLOWCORE_API, bool, Bool)
+DATAFLOW_CONNECTION_TYPE_PRIMITIVE(DATAFLOWCORE_API, char, Char)
+DATAFLOW_CONNECTION_TYPE_PRIMITIVE(DATAFLOWCORE_API, int, Integer)
+DATAFLOW_CONNECTION_TYPE_PRIMITIVE(DATAFLOWCORE_API, uint8, UInt8)
+DATAFLOW_CONNECTION_TYPE_PRIMITIVE(DATAFLOWCORE_API, float, Float)
+DATAFLOW_CONNECTION_TYPE_PRIMITIVE(DATAFLOWCORE_API, double, Double)
 }
 
