Dataflow : const evaluation

diff --git a/Engine/Source/Runtime/Experimental/Dataflow/Core/Public/Dataflow/DataflowInputOutput.h b/Engine/Source/Runtime/Experimental/Dataflow/Core/Public/Dataflow/DataflowInputOutput.h
--- a/Engine/Source/Runtime/Experimental/Dataflow/Core/Public/Dataflow/DataflowInputOutput.h
+++ b/Engine/Source/Runtime/Experimental/Dataflow/Core/Public/Dataflow/DataflowInputOutput.h
@@ -71,11 +71,11 @@ namespace Dataflow
 			return false;
 		}
 
-		T GetValue(const FContext& Context)
+		T GetValue(const FContext& Context) const
 		{
 			if (Connection)
 			{
-				return Connection->Evaluate(Context);
+				return Connection->GetValue(Context);
 			}
 			return Default;
 		}
@@ -170,7 +170,7 @@ namespace Dataflow
 			Cache.Data = DeepCopy<T>(InVal);
 		}
 
-		T Evaluate(const FContext& Context)
+		T GetValue(const FContext& Context)
 		{
 			if (CacheKey != Context.GetTypeHash())
 			{
