Chaos : fix convex generation to properly remove inside faces and trim invalid faces

diff --git a/Engine/Source/Runtime/Experimental/Chaos/Public/Chaos/CollisionConvexMesh.h b/Engine/Source/Runtime/Experimental/Chaos/Public/Chaos/CollisionConvexMesh.h
--- a/Engine/Source/Runtime/Experimental/Chaos/Public/Chaos/CollisionConvexMesh.h
+++ b/Engine/Source/Runtime/Experimental/Chaos/Public/Chaos/CollisionConvexMesh.h
@@ -434,6 +434,69 @@ namespace Chaos
 			check(InOutVertices.Num() > 3);
 		}
 
+		static bool IsFaceOutlineConvex(const FPlaneType& Plane, const TArray<int32>& Indices, const TArray<FVec3Type>& Vertices)
+		{
+			TArray<int8> Signs;
+			Signs.SetNum(Indices.Num());
+
+			if (Indices.Num() < 4)
+			{
+				return true;
+			}
+
+			for (int32 PointIndex = 0; PointIndex < Indices.Num(); PointIndex++)
+			{
+				const int32 Index0 = Indices[PointIndex];
+				const int32 Index1 = Indices[(PointIndex + 1) % Indices.Num()];
+				const int32 Index2 = Indices[(PointIndex + 2) % Indices.Num()];
+
+				const FVec3Type Point0 = Vertices[Index0];
+				const FVec3Type Point1 = Vertices[Index1];
+				const FVec3Type Point2 = Vertices[Index2];
+
+				const FVec3Type Segment0(Point1 - Point0);
+				const FVec3Type Segment1(Point2 - Point1);
+
+				const FVec3Type Cross = FVec3Type::CrossProduct(Segment0, Segment1);
+				const FRealType Dot = FVec3Type::DotProduct(Cross, Plane.Normal());
+
+				Signs[PointIndex] = static_cast<int8>(FMath::Sign(Dot));
+			}
+
+			int8 RefSign = 0;
+			for (const int8 Sign : Signs)
+			{
+				if (RefSign == 0)
+				{
+					RefSign = Sign;
+				}
+				if (Sign != 0 && RefSign != Sign)
+				{
+					return false;
+				}
+			}
+
+			return true;
+		}
+
+		// remove any invalid faces ( less than 3 vertices )
+		static void RemoveInvalidFaces(TArray<FPlaneType>& InOutPlanes, TArray<TArray<int32>>& InOutFaceVertexIndices)
+		{
+			int32 PlaneIndex = 0;
+			while (PlaneIndex < InOutPlanes.Num())
+			{
+				if (InOutFaceVertexIndices[PlaneIndex].Num() < 3)
+				{
+					InOutFaceVertexIndices.RemoveAtSwap(PlaneIndex);
+					InOutPlanes.RemoveAtSwap(PlaneIndex);
+				}
+				else
+				{
+					++PlaneIndex;
+				}
+			}
+		}
+		
 		// Convert multi-triangle faces to single n-gons
 		static void MergeFaces(TArray<FPlaneType>& InOutPlanes, TArray<TArray<int32>>& InOutFaceVertexIndices, const TArray<FVec3Type>& Vertices, FRealType DistanceThreshold)
 		{
@@ -500,21 +563,93 @@ namespace Chaos
 			// Re-order the face vertices to form the face half-edges
 			for (int32 PlaneIndex0 = 0; PlaneIndex0 < InOutPlanes.Num(); ++PlaneIndex0)
 			{
-				SortFaceVerticesCCW(InOutPlanes[PlaneIndex0], InOutFaceVertexIndices[PlaneIndex0], Vertices);
+				FinalizeFaces(InOutPlanes[PlaneIndex0], InOutFaceVertexIndices[PlaneIndex0], Vertices);
+#if DEBUG_HULL_GENERATION
+				ensure(IsFaceOutlineConvex(InOutPlanes[PlaneIndex0], InOutFaceVertexIndices[PlaneIndex0], Vertices));
+#endif
 			}
+			
+			RemoveInvalidFaces(InOutPlanes, InOutFaceVertexIndices);
 		}
 
-		// Reorder the vertices to be counter-clockwise about the normal
-		static void SortFaceVerticesCCW(const FPlaneType& Face, TArray<int32>& InOutFaceVertexIndices, const TArray<FVec3Type>& Vertices)
+		// IMPORTANT : vertices are assumed to be sorted CCW
+		static void RemoveInsideFaceVertices(const FPlaneType& Face, TArray<int32>& InOutFaceVertexIndices, const TArray<FVec3Type>& Vertices, const FVec3Type& Centroid)
 		{
-			FMatrix33 FaceMatrix = (FMatrix44f)FRotationMatrix::MakeFromZ(FVector(Face.Normal()));
+			// we let 3 points faces being processed as there may be colinear cases that will result n invalid faces out of this function?
+			if (InOutFaceVertexIndices.Num() < 3)
+			{
+			 	return;
+			}
+			// find furthest point from centroid as it is garanteed to be part of the convex hull 
+			int32 StartIndex = 0;
+			FRealType FurthestSquaredDistance = TNumericLimits<FRealType>::Lowest();
+			for (int32 Index = 0; Index < InOutFaceVertexIndices.Num(); ++Index)
+			{
+				const FRealType SquaredDistance = (Vertices[InOutFaceVertexIndices[Index]] - Centroid).SquaredLength();
+				if (SquaredDistance > FurthestSquaredDistance)
+				{
+					StartIndex = Index;
+					FurthestSquaredDistance = SquaredDistance;
+				}
+			}
 
-			FVec3Type Centroid(0);
-			for (int32 VertexIndex = 0; VertexIndex < InOutFaceVertexIndices.Num(); ++VertexIndex)
+			const int32 VtxCount = InOutFaceVertexIndices.Num();
+
+			struct FPointSegment
 			{
-				Centroid += Vertices[InOutFaceVertexIndices[VertexIndex]];
+				int32 VtxIndex;
+				FVec3Type Segment;
+			};
+
+			const int32 VtxStartIndex = InOutFaceVertexIndices[StartIndex];
+			int32 VtxIndex0 = VtxStartIndex;
+			int32 VtxIndex1 = InOutFaceVertexIndices[(StartIndex + 1) % VtxCount];
+
+			TArray<FPointSegment> Stack;
+			Stack.Push({ VtxIndex0, FVec3Type{0} });
+			Stack.Push({ VtxIndex1, Vertices[VtxIndex1] - Vertices[VtxIndex0] });
+
+			int32 Step = 2; // we already processed the two first 
+			while (Step <= VtxCount)
+			{
+				const int32 NextIndex = (StartIndex + Step) % VtxCount;
+
+				const FPointSegment& LastOnStack = Stack.Last();
+				VtxIndex0 = LastOnStack.VtxIndex;
+				VtxIndex1 = InOutFaceVertexIndices[NextIndex];
+				const FVec3Type Segment = Vertices[VtxIndex1] - Vertices[VtxIndex0];
+
+				const FVec3 Cross = FVec3Type::CrossProduct(LastOnStack.Segment, Segment);
+				if (FVec3Type::DotProduct(Cross, Face.Normal()) >= 0)
+				{
+					if (VtxIndex1 != VtxStartIndex)
+					{
+						Stack.Push({ VtxIndex1, Segment });
+					}
+					Step++;
+				}
+				else
+				{
+					Stack.Pop();
+				}
 			}
-			Centroid /= FRealType(InOutFaceVertexIndices.Num());
+
+			// faces where all points are on the same line may produce 2 points faces after removal
+			// in that case let's keep the entire face empty as a sign this should be trimmed 
+			InOutFaceVertexIndices.Reset();
+			if (Stack.Num() >= 3)
+			{
+				for (const FPointSegment& PointSegment : Stack)
+				{
+					InOutFaceVertexIndices.Add(PointSegment.VtxIndex);
+				}
+			}
+		}
+
+		// Reorder the vertices to be counter-clockwise about the normal
+		static void SortFaceVerticesCCW(const FPlaneType& Face, TArray<int32>& InOutFaceVertexIndices, const TArray<FVec3Type>& Vertices, const FVec3Type& Centroid)
+		{
+			const FMatrix33 FaceMatrix = (FMatrix44f)FRotationMatrix::MakeFromZ(FVector(Face.Normal()));
 
 			// [2, -2] based on clockwise angle about the normal
 			auto VertexScore = [&Centroid, &FaceMatrix, &Vertices](int32 VertexIndex)
@@ -534,8 +669,24 @@ namespace Chaos
 			InOutFaceVertexIndices.Sort(VertexSortPredicate);
 		}
 
+		// make sure faces have CCW winding, remove inside points 
+		static void FinalizeFaces(const FPlaneType& Face, TArray<int32>& InOutFaceVertexIndices, const TArray<FVec3Type>& Vertices)
+		{
+			// compute centroid as this is needed for both sorting and removing inside points
+			FVec3Type Centroid(0);
+			for (int32 VertexIndex = 0; VertexIndex < InOutFaceVertexIndices.Num(); ++VertexIndex)
+			{
+				Centroid += Vertices[InOutFaceVertexIndices[VertexIndex]];
+			}
+			Centroid /= FRealType(InOutFaceVertexIndices.Num());
+
+			SortFaceVerticesCCW(Face, InOutFaceVertexIndices, Vertices, Centroid);
+			
+			RemoveInsideFaceVertices(Face, InOutFaceVertexIndices, Vertices, Centroid);
+		}
+
 		// Generate the vertex indices for all planes in CCW order (used to serialize old data that did not have structure data)
-		static void BuildPlaneVertexIndices(const TArray<FPlaneType>& InPlanes, const TArray<FVec3Type>& Vertices, TArray<TArray<int32>>& OutFaceVertexIndices, const FRealType DistanceTolerance = 1.e-3f)
+		static void BuildPlaneVertexIndices(TArray<FPlaneType>& InPlanes, const TArray<FVec3Type>& Vertices, TArray<TArray<int32>>& OutFaceVertexIndices, const FRealType DistanceTolerance = 1.e-3f)
 		{
 			OutFaceVertexIndices.Reset(InPlanes.Num());
 			for (int32 PlaneIndex = 0; PlaneIndex < InPlanes.Num(); ++PlaneIndex)
@@ -549,8 +700,9 @@ namespace Chaos
 					}
 				}
 
-				SortFaceVerticesCCW(InPlanes[PlaneIndex], OutFaceVertexIndices[PlaneIndex], Vertices);
+				FinalizeFaces(InPlanes[PlaneIndex], OutFaceVertexIndices[PlaneIndex], Vertices);
 			}
+			RemoveInvalidFaces(InPlanes, OutFaceVertexIndices);
 		}
 
 		// CVars variables for controlling geometry complexity checking and simplification
diff --git a/Engine/Source/Runtime/Experimental/ChaosCore/Public/Chaos/Core.h b/Engine/Source/Runtime/Experimental/ChaosCore/Public/Chaos/Core.h
--- a/Engine/Source/Runtime/Experimental/ChaosCore/Public/Chaos/Core.h
+++ b/Engine/Source/Runtime/Experimental/ChaosCore/Public/Chaos/Core.h
@@ -51,6 +51,6 @@ namespace Chaos
 	// NOTE: if you get a merge conflict on the GUID, you must replace it with a new GUID - do not accept the source or target
 	// or you will likely get DDC version conflicts resulting in crashes during load.
 	// Core version string for Chaos data. Any DDC builder dependent on Chaos for serialization should depend on this version
-	static const FString ChaosVersionString = TEXT("841FF58C-9DD5-4AF7-BE5B-959F1F2BC4D7");
+	static const FString ChaosVersionString = TEXT("5F283595-0103-4FE6-9684-E8B820158BFC");
 
 }
