Chaos : add new connection graph generation method ( delaunay filtered by bounds )

diff --git a/Engine/Source/Runtime/Experimental/Chaos/Private/Chaos/PBDRigidClustering.cpp b/Engine/Source/Runtime/Experimental/Chaos/Private/Chaos/PBDRigidClustering.cpp
--- a/Engine/Source/Runtime/Experimental/Chaos/Private/Chaos/PBDRigidClustering.cpp
+++ b/Engine/Source/Runtime/Experimental/Chaos/Private/Chaos/PBDRigidClustering.cpp
@@ -939,6 +939,11 @@ namespace Chaos
 				UpdateConnectivityGraphUsingDelaunayTriangulation(Parent, Parameters); // not thread safe
 			}
 
+			if (LocalConnectionMethod == FClusterCreationParameters::EConnectionMethod::BoundsOverlapFilteredDelaunayTriangulation)
+			{
+				UpdateConnectivityGraphUsingDelaunayTriangulationWithBoundsOverlaps(Parent, Parameters);
+			}
+
 			if (LocalConnectionMethod == FClusterCreationParameters::EConnectionMethod::PointImplicitAugmentedWithMinimalDelaunay ||
 				LocalConnectionMethod == FClusterCreationParameters::EConnectionMethod::MinimalSpanningSubsetDelaunayTriangulation)
 			{
@@ -1141,6 +1146,33 @@ namespace Chaos
 
 	}
 
+	static void ConnectClusteredNodes(FPBDRigidClusteredParticleHandle* ClusteredChild1, FPBDRigidClusteredParticleHandle* ClusteredChild2)
+	{
+		check(ClusteredChild1 && ClusteredChild2);
+		if (ClusteredChild1 == ClusteredChild2)
+			return;
+		const FReal AvgStrain = (ClusteredChild1->Strains() + ClusteredChild2->Strains()) * (FReal)0.5;
+		TArray<TConnectivityEdge<FReal>>& Edges1 = ClusteredChild1->ConnectivityEdges();
+		TArray<TConnectivityEdge<FReal>>& Edges2 = ClusteredChild2->ConnectivityEdges();
+		if (//Edges1.Num() < Parameters.MaxNumConnections && 
+			!Edges1.FindByKey(ClusteredChild2))
+		{
+			Edges1.Add(TConnectivityEdge<FReal>(ClusteredChild2, AvgStrain));
+		}
+		if (//Edges2.Num() < Parameters.MaxNumConnections && 
+			!Edges2.FindByKey(ClusteredChild1))
+		{
+			Edges2.Add(TConnectivityEdge<FReal>(ClusteredChild1, AvgStrain));
+		}
+	}
+	
+	static void ConnectNodes(FPBDRigidParticleHandle* Child1, FPBDRigidParticleHandle* Child2)
+	{
+		check(Child1 != Child2);
+		FPBDRigidClusteredParticleHandle* ClusteredChild1 = Child1->CastToClustered();
+		FPBDRigidClusteredParticleHandle* ClusteredChild2 = Child2->CastToClustered();
+		ConnectClusteredNodes(ClusteredChild1, ClusteredChild2);
+	}
 
 	DECLARE_CYCLE_STAT(TEXT("TPBDRigidClustering<>::UpdateConnectivityGraphUsingPointImplicit"), STAT_UpdateConnectivityGraphUsingPointImplicit, STATGROUP_Chaos);
 	void 
@@ -1368,28 +1400,28 @@ namespace Chaos
 			}
 		}
 	}
-
-	static bool IsOveralppingConnection(const FPBDRigidParticleHandle* Child1, const FPBDRigidParticleHandle* Child2)
+	
+	// connection filters
+	static bool IsAlwaysValidConnection(const FPBDRigidParticleHandle* Child1, const FPBDRigidParticleHandle* Child2)
+	{
+		return true;
+	}
+	
+	static bool IsOverlappingConnection(const FPBDRigidParticleHandle* Child1, const FPBDRigidParticleHandle* Child2, FReal Margin)
 	{
 		if (ensure(Child1 != nullptr && Child2 != nullptr ))
 		{
 			FAABB3 Bounds1 = Child1->WorldSpaceInflatedBounds();
-			Bounds1.Thicken(BoundingBoxMarginForConnectionGraphFiltering);
+			Bounds1.Thicken(Margin);
 			return Bounds1.Intersects(Child2->WorldSpaceInflatedBounds());
 		}
 		return false;
 	}
-	
-	DECLARE_CYCLE_STAT(TEXT("TPBDRigidClustering<>::UpdateConnectivityGraphUsingDelaunayTriangulation"), STAT_UpdateConnectivityGraphUsingDelaunayTriangulation, STATGROUP_Chaos);
-	void 
-	FRigidClustering::UpdateConnectivityGraphUsingDelaunayTriangulation(
-		Chaos::FPBDRigidClusteredParticleHandle* Parent, 
-		const FClusterCreationParameters& Parameters)
-	{
-		SCOPE_CYCLE_COUNTER(STAT_UpdateConnectivityGraphUsingDelaunayTriangulation);
-
-		const TArray<FPBDRigidParticleHandle*>& Children = MChildren[Parent];
 
+	template <typename FilterLambda>
+	static void UpdateConnectivityGraphUsingDelaunayTriangulationWithFiltering(
+		const TArray<FPBDRigidParticleHandle*>& Children, FilterLambda ShouldKeepConnection)
+	{
 		TArray<FVector> Pts;
 		Pts.AddUninitialized(Children.Num());
 		for (int32 i = 0; i < Children.Num(); ++i)
@@ -1412,12 +1444,7 @@ namespace Chaos
 					bFirstSmaller ? Child2 : Child1);
 				if (!UniqueEdges.Find(SortedPair))
 				{
-					bool ValidConnection = true;
-					if (UseBoundingBoxForConnectionGraphFiltering)
-					{
-						ValidConnection = IsOveralppingConnection(Child1, Child2);
-					}
-					if (ValidConnection)
+					if (ShouldKeepConnection(Child1, Child2))
 					{
 						// this does not use ConnectNodes because Neighbors is bi-direction : as in (1,2),(2,1)
 						ConnectNodes(Child1, Child2);
@@ -1427,41 +1454,45 @@ namespace Chaos
 			}
 		}
 	}
-
-	void 
-	FRigidClustering::ConnectNodes(
-		FPBDRigidParticleHandle* Child1,
-		FPBDRigidParticleHandle* Child2)
+	
+	DECLARE_CYCLE_STAT(TEXT("TPBDRigidClustering<>::UpdateConnectivityGraphUsingDelaunayTriangulation"), STAT_UpdateConnectivityGraphUsingDelaunayTriangulation, STATGROUP_Chaos);
+	void FRigidClustering::UpdateConnectivityGraphUsingDelaunayTriangulation(
+		const Chaos::FPBDRigidClusteredParticleHandle* Parent, 
+		const FClusterCreationParameters& Parameters)
 	{
-		check(Child1 != Child2);
-		FPBDRigidClusteredParticleHandle* ClusteredChild1 = Child1->CastToClustered();
-		FPBDRigidClusteredParticleHandle* ClusteredChild2 = Child2->CastToClustered();
-		ConnectNodes(ClusteredChild1, ClusteredChild2);
-	}
+		SCOPE_CYCLE_COUNTER(STAT_UpdateConnectivityGraphUsingDelaunayTriangulation);
 
-	void 
-	FRigidClustering::ConnectNodes(
-		FPBDRigidClusteredParticleHandle* ClusteredChild1,
-		FPBDRigidClusteredParticleHandle* ClusteredChild2)
-	{
-		check(ClusteredChild1 && ClusteredChild2);
-		if (ClusteredChild1 == ClusteredChild2)
-			return;
-		const FReal AvgStrain = (ClusteredChild1->Strains() + ClusteredChild2->Strains()) * (FReal)0.5;
-		TArray<TConnectivityEdge<FReal>>& Edges1 = ClusteredChild1->ConnectivityEdges();
-		TArray<TConnectivityEdge<FReal>>& Edges2 = ClusteredChild2->ConnectivityEdges();
-		if (//Edges1.Num() < Parameters.MaxNumConnections && 
-			!Edges1.FindByKey(ClusteredChild2))
+		if (UseBoundingBoxForConnectionGraphFiltering)
 		{
-			Edges1.Add(TConnectivityEdge<FReal>(ClusteredChild2, AvgStrain));
+			constexpr auto IsOverlappingConnectionUsingCVarMargin =
+				[] (const FPBDRigidParticleHandle* Child1, const FPBDRigidParticleHandle* Child2)
+				{
+					return IsOverlappingConnection(Child1, Child2, BoundingBoxMarginForConnectionGraphFiltering);
+				};
+			UpdateConnectivityGraphUsingDelaunayTriangulationWithFiltering(MChildren[Parent], IsOverlappingConnectionUsingCVarMargin);
 		}
-		if (//Edges2.Num() < Parameters.MaxNumConnections && 
-			!Edges2.FindByKey(ClusteredChild1))
+		else
 		{
-			Edges2.Add(TConnectivityEdge<FReal>(ClusteredChild1, AvgStrain));
+			UpdateConnectivityGraphUsingDelaunayTriangulationWithFiltering(MChildren[Parent], IsAlwaysValidConnection);
 		}
 	}
 
+	DECLARE_CYCLE_STAT(TEXT("TPBDRigidClustering<>::UpdateConnectivityGraphUsingDelaunayTriangulationWithBoundsOverlaps"), STAT_UpdateConnectivityGraphUsingDelaunayTriangulationWithBoundsOverlaps, STATGROUP_Chaos);
+	void FRigidClustering::UpdateConnectivityGraphUsingDelaunayTriangulationWithBoundsOverlaps(
+		const Chaos::FPBDRigidClusteredParticleHandle* Parent, 
+		const FClusterCreationParameters& Parameters)
+	{
+		SCOPE_CYCLE_COUNTER(STAT_UpdateConnectivityGraphUsingDelaunayTriangulationWithBoundsOverlaps);
+
+		const auto IsOverlappingConnectionUsingCVarMargin =
+			[&Parameters] (const FPBDRigidParticleHandle* Child1, const FPBDRigidParticleHandle* Child2)
+			{
+				return IsOverlappingConnection(Child1, Child2, Parameters.ConnectionGraphBoundsFilteringMargin);
+			};
+		
+		UpdateConnectivityGraphUsingDelaunayTriangulationWithFiltering(MChildren[Parent], IsOverlappingConnectionUsingCVarMargin);
+	}
+
 	void 
 	FRigidClustering::RemoveNodeConnections(
 		FPBDRigidParticleHandle* Child)
diff --git a/Engine/Source/Runtime/Experimental/Chaos/Private/GeometryCollection/GeometryCollectionConnectionGraphUtility.cpp b/Engine/Source/Runtime/Experimental/Chaos/Private/GeometryCollection/GeometryCollectionConnectionGraphUtility.cpp
--- a/Engine/Source/Runtime/Experimental/Chaos/Private/GeometryCollection/GeometryCollectionConnectionGraphUtility.cpp
+++ b/Engine/Source/Runtime/Experimental/Chaos/Private/GeometryCollection/GeometryCollectionConnectionGraphUtility.cpp
@@ -38,6 +38,12 @@ FGeometryCollectionConnectionGraphGenerator::ComputeConnectivityGraph(const FGeo
 		//UpdateConnectivityGraphUsingPointImplicit(Parent, Parameters);
 	}
 
+	if (LocalConnectionMethod == EClusterUnionMethod::PointImplicitAugmentedWithMinimalDelaunay)
+	{
+		// todo : fix later when we have proper implementation for this on the GT side 
+		FixConnectivityGraphUsingDelaunayTriangulation(Graph, Collection, ClusterTransformIndex);
+	}
+	
 	if (LocalConnectionMethod == EClusterUnionMethod::DelaunayTriangulation)
 	{
 		ComputeConnectivityGraphUsingDelaunayTriangulation(Graph, Collection, ClusterTransformIndex);
diff --git a/Engine/Source/Runtime/Experimental/Chaos/Private/PBDRigidsSolver.cpp b/Engine/Source/Runtime/Experimental/Chaos/Private/PBDRigidsSolver.cpp
--- a/Engine/Source/Runtime/Experimental/Chaos/Private/PBDRigidsSolver.cpp
+++ b/Engine/Source/Runtime/Experimental/Chaos/Private/PBDRigidsSolver.cpp
@@ -1690,6 +1690,8 @@ CSV_CUSTOM_STAT(PhysicsCounters, Name, Value, ECsvCustomStatOp::Set);
 			return ETargetEnum::MinimalSpanningSubsetDelaunayTriangulation;
 		case EClusterUnionMethod::PointImplicitAugmentedWithMinimalDelaunay:
 			return ETargetEnum::PointImplicitAugmentedWithMinimalDelaunay;
+		case EClusterUnionMethod::BoundsOverlapFilteredDelaunayTriangulation:
+			return ETargetEnum::BoundsOverlapFilteredDelaunayTriangulation;
 		}
 
 		return ETargetEnum::None;
diff --git a/Engine/Source/Runtime/Experimental/Chaos/Private/PhysicsProxy/GeometryCollectionPhysicsProxy.cpp b/Engine/Source/Runtime/Experimental/Chaos/Private/PhysicsProxy/GeometryCollectionPhysicsProxy.cpp
--- a/Engine/Source/Runtime/Experimental/Chaos/Private/PhysicsProxy/GeometryCollectionPhysicsProxy.cpp
+++ b/Engine/Source/Runtime/Experimental/Chaos/Private/PhysicsProxy/GeometryCollectionPhysicsProxy.cpp
@@ -1040,6 +1040,7 @@ void FGeometryCollectionPhysicsProxy::InitializeBodiesPT(Chaos::FPBDRigidsSolver
 							// #todo: should other parameters be set here?  Previously, there was no parameters being sent, and it is unclear
 							// where some of these parameters are defined (ie: CollisionThicknessPercent)
 							ClusterParams.ConnectionMethod = Parameters.ClusterConnectionMethod;
+							ClusterParams.ConnectionGraphBoundsFilteringMargin = Parameters.ConnectionGraphBoundsFilteringMargin;
 						
 							RigidsSolver->GetEvolution()->GetRigidClustering().GenerateConnectionGraph(ClusteredParticle, ClusterParams);
 						}
@@ -1048,26 +1049,29 @@ void FGeometryCollectionPhysicsProxy::InitializeBodiesPT(Chaos::FPBDRigidsSolver
 			}
 			else
 			{
-				for (int32 TransformGroupIndex = 0; TransformGroupIndex < NumTransforms; ++TransformGroupIndex)
+				if (Connections)
 				{
-					if (FClusterHandle* ClusteredParticle = SolverParticleHandles[TransformGroupIndex]) 
+					for (int32 TransformGroupIndex = 0; TransformGroupIndex < NumTransforms; ++TransformGroupIndex)
 					{
-						const TSet<int32>& Siblings = (*Connections)[TransformGroupIndex];
-						for (const int32 SiblingTransformIndex: Siblings)
+						if (FClusterHandle* ClusteredParticle = SolverParticleHandles[TransformGroupIndex]) 
 						{
-							if (FClusterHandle* OtherClusteredParticle = SolverParticleHandles[SiblingTransformIndex])
+							const TSet<int32>& Siblings = (*Connections)[TransformGroupIndex];
+							for (const int32 SiblingTransformIndex: Siblings)
 							{
-								// we add both connection for integrity so let use the index order as a way to filter out the symmetrical part 
-								if (SiblingTransformIndex > TransformGroupIndex)
+								if (FClusterHandle* OtherClusteredParticle = SolverParticleHandles[SiblingTransformIndex])
 								{
-									Chaos::TConnectivityEdge<Chaos::FReal> Edge;
-									Edge.Strain = (ClusteredParticle->Strains() + OtherClusteredParticle->Strains()) * 0.5f;
+									// we add both connection for integrity so let use the index order as a way to filter out the symmetrical part 
+									if (SiblingTransformIndex > TransformGroupIndex)
+									{
+										Chaos::TConnectivityEdge<Chaos::FReal> Edge;
+										Edge.Strain = (ClusteredParticle->Strains() + OtherClusteredParticle->Strains()) * 0.5f;
 									
-									Edge.Sibling = OtherClusteredParticle;
-									ClusteredParticle->ConnectivityEdges().Add(Edge);
+										Edge.Sibling = OtherClusteredParticle;
+										ClusteredParticle->ConnectivityEdges().Add(Edge);
 									
-									Edge.Sibling = ClusteredParticle;
-									OtherClusteredParticle->ConnectivityEdges().Add(Edge);
+										Edge.Sibling = ClusteredParticle;
+										OtherClusteredParticle->ConnectivityEdges().Add(Edge);
+									}
 								}
 							}
 						}
diff --git a/Engine/Source/Runtime/Experimental/Chaos/Public/Chaos/ClusterCreationParameters.h b/Engine/Source/Runtime/Experimental/Chaos/Public/Chaos/ClusterCreationParameters.h
--- a/Engine/Source/Runtime/Experimental/Chaos/Public/Chaos/ClusterCreationParameters.h
+++ b/Engine/Source/Runtime/Experimental/Chaos/Public/Chaos/ClusterCreationParameters.h
@@ -18,16 +18,18 @@ namespace Chaos
 			DelaunayTriangulation,
 			MinimalSpanningSubsetDelaunayTriangulation,
 			PointImplicitAugmentedWithMinimalDelaunay,
+			BoundsOverlapFilteredDelaunayTriangulation,
 			None
 		};
 
-
 		FClusterCreationParameters(
 			FReal CoillisionThicknessPercentIn = (FReal)0.3
 			, int32 MaxNumConnectionsIn = 100
 			, bool bCleanCollisionParticlesIn = true
 			, bool bCopyCollisionParticlesIn = true
-			, bool bGenerateConnectionGraphIn = true, EConnectionMethod ConnectionMethodIn = EConnectionMethod::MinimalSpanningSubsetDelaunayTriangulation
+			, bool bGenerateConnectionGraphIn = true
+			, EConnectionMethod ConnectionMethodIn = EConnectionMethod::MinimalSpanningSubsetDelaunayTriangulation
+			, FReal ConnectionGraphBoundsFilteringMarginIn = 0
 			, FBVHParticles* CollisionParticlesIn = nullptr
 			, Chaos::TPBDRigidClusteredParticleHandle<Chaos::FReal,3>* ClusterParticleHandleIn = nullptr
 			, const FVec3& ScaleIn = FVec3::OneVector
@@ -38,6 +40,7 @@ namespace Chaos
 			, bCopyCollisionParticles(bCopyCollisionParticlesIn)
 			, bGenerateConnectionGraph(bGenerateConnectionGraphIn)
 			, ConnectionMethod(ConnectionMethodIn)
+			, ConnectionGraphBoundsFilteringMargin(ConnectionGraphBoundsFilteringMarginIn)
 			, CollisionParticles(CollisionParticlesIn)
 			, ClusterParticleHandle(ClusterParticleHandleIn)
 			, Scale(ScaleIn)
@@ -49,6 +52,7 @@ namespace Chaos
 		bool bCopyCollisionParticles;
 		bool bGenerateConnectionGraph;
 		EConnectionMethod ConnectionMethod;
+		FReal ConnectionGraphBoundsFilteringMargin;
 		FBVHParticles* CollisionParticles;
 		Chaos::FPBDRigidClusteredParticleHandle* ClusterParticleHandle;
 		FVec3 Scale;
diff --git a/Engine/Source/Runtime/Experimental/Chaos/Public/Chaos/PBDRigidClustering.h b/Engine/Source/Runtime/Experimental/Chaos/Public/Chaos/PBDRigidClustering.h
--- a/Engine/Source/Runtime/Experimental/Chaos/Public/Chaos/PBDRigidClustering.h
+++ b/Engine/Source/Runtime/Experimental/Chaos/Public/Chaos/PBDRigidClustering.h
@@ -290,16 +290,12 @@ public:
 		Chaos::FPBDRigidClusteredParticleHandle* Parent,
 		const FClusterCreationParameters& Parameters = FClusterCreationParameters());
 	void UpdateConnectivityGraphUsingDelaunayTriangulation(
-		Chaos::FPBDRigidClusteredParticleHandle* Parent,
+		const Chaos::FPBDRigidClusteredParticleHandle* Parent,
 		const FClusterCreationParameters& Parameters = FClusterCreationParameters());
-
-	void ConnectNodes(
-		FPBDRigidParticleHandle* Child1,
-		FPBDRigidParticleHandle* Child2);
-	void ConnectNodes(
-		FPBDRigidClusteredParticleHandle* Child1,
-		FPBDRigidClusteredParticleHandle* Child2);
-
+	void UpdateConnectivityGraphUsingDelaunayTriangulationWithBoundsOverlaps(
+		const Chaos::FPBDRigidClusteredParticleHandle* Parent,
+		const FClusterCreationParameters& Parameters = FClusterCreationParameters());
+	
 	void RemoveNodeConnections(FPBDRigidParticleHandle* Child);
 	void RemoveNodeConnections(FPBDRigidClusteredParticleHandle* Child);
 
diff --git a/Engine/Source/Runtime/Experimental/Chaos/Public/ChaosSolverConfiguration.h b/Engine/Source/Runtime/Experimental/Chaos/Public/ChaosSolverConfiguration.h
--- a/Engine/Source/Runtime/Experimental/Chaos/Public/ChaosSolverConfiguration.h
+++ b/Engine/Source/Runtime/Experimental/Chaos/Public/ChaosSolverConfiguration.h
@@ -16,6 +16,7 @@ enum class EClusterUnionMethod : uint8
 	DelaunayTriangulation,
 	MinimalSpanningSubsetDelaunayTriangulation,
 	PointImplicitAugmentedWithMinimalDelaunay,
+	BoundsOverlapFilteredDelaunayTriangulation,
 	None
 };
 
diff --git a/Engine/Source/Runtime/Experimental/Chaos/Public/GeometryCollection/GeometryCollectionSimulationCoreTypes.h b/Engine/Source/Runtime/Experimental/Chaos/Public/GeometryCollection/GeometryCollectionSimulationCoreTypes.h
--- a/Engine/Source/Runtime/Experimental/Chaos/Public/GeometryCollection/GeometryCollectionSimulationCoreTypes.h
+++ b/Engine/Source/Runtime/Experimental/Chaos/Public/GeometryCollection/GeometryCollectionSimulationCoreTypes.h
@@ -178,6 +178,7 @@ struct FSimulationParameters
 		, DamageThreshold({500000.f, 50000.f, 5000.f})
 		, bUsePerClusterOnlyDamageThreshold(false)
 		, ClusterConnectionMethod(Chaos::FClusterCreationParameters::EConnectionMethod::PointImplicit)
+		, ConnectionGraphBoundsFilteringMargin(0)
 		, CollisionGroup(0)
 		, CollisionSampleFraction(1.0)
 		, InitialVelocityType(EInitialVelocityTypeEnum::Chaos_Initial_Velocity_None)
@@ -215,6 +216,7 @@ struct FSimulationParameters
 		, DamageThreshold(Other.DamageThreshold)
 		, bUsePerClusterOnlyDamageThreshold(Other.bUsePerClusterOnlyDamageThreshold)
 		, ClusterConnectionMethod(Other.ClusterConnectionMethod)
+		, ConnectionGraphBoundsFilteringMargin(Other.ConnectionGraphBoundsFilteringMargin)
 		, CollisionGroup(Other.CollisionGroup)
 		, CollisionSampleFraction(Other.CollisionSampleFraction)
 		, InitialVelocityType(Other.InitialVelocityType)
@@ -268,6 +270,7 @@ struct FSimulationParameters
 	TArray<float> DamageThreshold;
 	bool bUsePerClusterOnlyDamageThreshold;
 	Chaos::FClusterCreationParameters::EConnectionMethod ClusterConnectionMethod;
+	float ConnectionGraphBoundsFilteringMargin;
 
 	int32 CollisionGroup;
 	float CollisionSampleFraction;
diff --git a/Engine/Source/Runtime/Experimental/ChaosSolverEngine/Private/Chaos/ChaosSolverActor.cpp b/Engine/Source/Runtime/Experimental/ChaosSolverEngine/Private/Chaos/ChaosSolverActor.cpp
--- a/Engine/Source/Runtime/Experimental/ChaosSolverEngine/Private/Chaos/ChaosSolverActor.cpp
+++ b/Engine/Source/Runtime/Experimental/ChaosSolverEngine/Private/Chaos/ChaosSolverActor.cpp
@@ -349,6 +349,8 @@ void AChaosSolverActor::PostLoad()
 				return EClusterUnionMethod::MinimalSpanningSubsetDelaunayTriangulation;
 			case EClusterConnectionTypeEnum::Chaos_PointImplicitAugmentedWithMinimalDelaunay:
 				return EClusterUnionMethod::PointImplicitAugmentedWithMinimalDelaunay;
+			case EClusterConnectionTypeEnum::Chaos_BoundsOverlapFilteredDelaunayTriangulation:
+				return EClusterUnionMethod::BoundsOverlapFilteredDelaunayTriangulation;
 			default:
 				break;
 			}
diff --git a/Engine/Source/Runtime/Experimental/ChaosSolverEngine/Public/Chaos/ChaosSolverActor.h b/Engine/Source/Runtime/Experimental/ChaosSolverEngine/Public/Chaos/ChaosSolverActor.h
--- a/Engine/Source/Runtime/Experimental/ChaosSolverEngine/Public/Chaos/ChaosSolverActor.h
+++ b/Engine/Source/Runtime/Experimental/ChaosSolverEngine/Public/Chaos/ChaosSolverActor.h
@@ -33,8 +33,9 @@ enum class EClusterConnectionTypeEnum : uint8
 {
 	Chaos_PointImplicit = Chaos::FClusterCreationParameters::PointImplicit ,
 	Chaos_DelaunayTriangulation = Chaos::FClusterCreationParameters::DelaunayTriangulation UMETA(Hidden),
-	Chaos_MinimalSpanningSubsetDelaunayTriangulation = Chaos::FClusterCreationParameters::MinimalSpanningSubsetDelaunayTriangulation ,
-	Chaos_PointImplicitAugmentedWithMinimalDelaunay = Chaos::FClusterCreationParameters::PointImplicitAugmentedWithMinimalDelaunay ,
+	Chaos_MinimalSpanningSubsetDelaunayTriangulation = Chaos::FClusterCreationParameters::MinimalSpanningSubsetDelaunayTriangulation,
+	Chaos_PointImplicitAugmentedWithMinimalDelaunay = Chaos::FClusterCreationParameters::PointImplicitAugmentedWithMinimalDelaunay,
+	Chaos_BoundsOverlapFilteredDelaunayTriangulation = Chaos::FClusterCreationParameters::BoundsOverlapFilteredDelaunayTriangulation,
 	Chaos_None = Chaos::FClusterCreationParameters::None UMETA(Hidden),
 	//
 	Chaos_EClsuterCreationParameters_Max UMETA(Hidden)
diff --git a/Engine/Source/Runtime/Experimental/GeometryCollectionEngine/Private/GeometryCollection/GeometryCollectionComponent.cpp b/Engine/Source/Runtime/Experimental/GeometryCollectionEngine/Private/GeometryCollection/GeometryCollectionComponent.cpp
--- a/Engine/Source/Runtime/Experimental/GeometryCollectionEngine/Private/GeometryCollection/GeometryCollectionComponent.cpp
+++ b/Engine/Source/Runtime/Experimental/GeometryCollectionEngine/Private/GeometryCollection/GeometryCollectionComponent.cpp
@@ -2212,11 +2212,13 @@ void UGeometryCollectionComponent::RegisterAndInitializePhysicsProxy()
 		SimulationParameters.Name = GetPathName();
 #endif
 		EClusterConnectionTypeEnum ClusterCollectionType = ClusterConnectionType_DEPRECATED;
+		float ConnectionGraphBoundsFilteringMargin = 0;
 		if (RestCollection)
 		{
 			RestCollection->GetSharedSimulationParams(SimulationParameters.Shared);
 			SimulationParameters.RestCollection = RestCollection->GetGeometryCollection().Get();
 			ClusterCollectionType = RestCollection->ClusterConnectionType;
+			ConnectionGraphBoundsFilteringMargin = RestCollection->ConnectionGraphBoundsFilteringMargin;
 		}
 		SimulationParameters.Simulating = BodyInstance.bSimulatePhysics;
 		SimulationParameters.EnableClustering = EnableClustering;
@@ -2226,6 +2228,7 @@ void UGeometryCollectionComponent::RegisterAndInitializePhysicsProxy()
 		SimulationParameters.DamageThreshold = DamageThreshold;
 		SimulationParameters.bUsePerClusterOnlyDamageThreshold = RestCollection? RestCollection->PerClusterOnlyDamageThreshold: false; 
 		SimulationParameters.ClusterConnectionMethod = (Chaos::FClusterCreationParameters::EConnectionMethod)ClusterCollectionType;
+		SimulationParameters.ConnectionGraphBoundsFilteringMargin = ConnectionGraphBoundsFilteringMargin; 
 		SimulationParameters.CollisionGroup = CollisionGroup;
 		SimulationParameters.CollisionSampleFraction = CollisionSampleFraction;
 		SimulationParameters.InitialVelocityType = InitialVelocityType;
diff --git a/Engine/Source/Runtime/Experimental/GeometryCollectionEngine/Private/GeometryCollection/GeometryCollectionObject.cpp b/Engine/Source/Runtime/Experimental/GeometryCollectionEngine/Private/GeometryCollection/GeometryCollectionObject.cpp
--- a/Engine/Source/Runtime/Experimental/GeometryCollectionEngine/Private/GeometryCollection/GeometryCollectionObject.cpp
+++ b/Engine/Source/Runtime/Experimental/GeometryCollectionEngine/Private/GeometryCollection/GeometryCollectionObject.cpp
@@ -75,6 +75,7 @@ UGeometryCollection::UGeometryCollection(const FObjectInitializer& ObjectInitial
 	, DamageThreshold({ 500000.f, 50000.f, 5000.f })
 	, PerClusterOnlyDamageThreshold(false)
 	, ClusterConnectionType(EClusterConnectionTypeEnum::Chaos_MinimalSpanningSubsetDelaunayTriangulation)
+	, ConnectionGraphBoundsFilteringMargin(0)
 	, bUseFullPrecisionUVs(false)
 	, bStripOnCook(false)
 	, EnableNanite(false)
diff --git a/Engine/Source/Runtime/Experimental/GeometryCollectionEngine/Public/GeometryCollection/GeometryCollectionObject.h b/Engine/Source/Runtime/Experimental/GeometryCollectionEngine/Public/GeometryCollection/GeometryCollectionObject.h
--- a/Engine/Source/Runtime/Experimental/GeometryCollectionEngine/Public/GeometryCollection/GeometryCollectionObject.h
+++ b/Engine/Source/Runtime/Experimental/GeometryCollectionEngine/Public/GeometryCollection/GeometryCollectionObject.h
@@ -440,6 +440,9 @@ public:
 	UPROPERTY(EditAnywhere, Category = "Clustering")
 	EClusterConnectionTypeEnum ClusterConnectionType;
 
+	UPROPERTY(EditAnywhere, Category = "Clustering")
+	float ConnectionGraphBoundsFilteringMargin;
+
 #if WITH_EDITORONLY_DATA
 	UPROPERTY(BlueprintReadWrite, EditAnywhere, Category = "GeometrySource")
 	TArray<FGeometryCollectionSource> GeometrySource;
