Fixed erroneous deleter being used for derived types of FProxyTimestampBase

diff --git a/Engine/Source/Runtime/Experimental/Chaos/Private/PhysicsProxy/JointConstraintProxy.cpp b/Engine/Source/Runtime/Experimental/Chaos/Private/PhysicsProxy/JointConstraintProxy.cpp
--- a/Engine/Source/Runtime/Experimental/Chaos/Private/PhysicsProxy/JointConstraintProxy.cpp
+++ b/Engine/Source/Runtime/Experimental/Chaos/Private/PhysicsProxy/JointConstraintProxy.cpp
@@ -19,7 +19,7 @@ namespace Chaos
 {
 
 FJointConstraintPhysicsProxy::FJointConstraintPhysicsProxy(FJointConstraint* InConstraint, FPBDJointConstraintHandle* InHandle, UObject* InOwner)
-	: Base(EPhysicsProxyType::JointConstraintType, InOwner, new FProxyTimestampBase)
+	: Base(EPhysicsProxyType::JointConstraintType, InOwner, MakeShared<FProxyTimestampBase>())
 	, Constraint_GT(InConstraint) // This proxy assumes ownership of the Constraint, and will free it during DestroyOnPhysicsThread
 	, Constraint_PT(InHandle)
 {
diff --git a/Engine/Source/Runtime/Experimental/Chaos/Private/PhysicsProxy/SingleParticlePhysicsProxy.cpp b/Engine/Source/Runtime/Experimental/Chaos/Private/PhysicsProxy/SingleParticlePhysicsProxy.cpp
--- a/Engine/Source/Runtime/Experimental/Chaos/Private/PhysicsProxy/SingleParticlePhysicsProxy.cpp
+++ b/Engine/Source/Runtime/Experimental/Chaos/Private/PhysicsProxy/SingleParticlePhysicsProxy.cpp
@@ -25,7 +25,7 @@ int32 SyncKinematicOnGameThread = 0;
 FAutoConsoleVariableRef CVar_SyncKinematicOnGameThread(TEXT("P.Chaos.SyncKinematicOnGameThread"), SyncKinematicOnGameThread, TEXT("If set to 1, if a kinematic is flagged to send position back to game thread, move component, if 0, do not."));
 
 FSingleParticlePhysicsProxy::FSingleParticlePhysicsProxy(TUniquePtr<PARTICLE_TYPE>&& InParticle, FParticleHandle* InHandle, UObject* InOwner)
-	: IPhysicsProxyBase(EPhysicsProxyType::SingleParticleProxy, InOwner, new FSingleParticleProxyTimestamp)
+	: IPhysicsProxyBase(EPhysicsProxyType::SingleParticleProxy, InOwner, MakeShared<FSingleParticleProxyTimestamp>())
 	, Particle(MoveTemp(InParticle))
 	, Handle(InHandle)
 {
diff --git a/Engine/Source/Runtime/Experimental/Chaos/Private/PhysicsProxy/SuspensionConstraintProxy.cpp b/Engine/Source/Runtime/Experimental/Chaos/Private/PhysicsProxy/SuspensionConstraintProxy.cpp
--- a/Engine/Source/Runtime/Experimental/Chaos/Private/PhysicsProxy/SuspensionConstraintProxy.cpp
+++ b/Engine/Source/Runtime/Experimental/Chaos/Private/PhysicsProxy/SuspensionConstraintProxy.cpp
@@ -18,7 +18,7 @@ namespace Chaos
 {
 
 FSuspensionConstraintPhysicsProxy::FSuspensionConstraintPhysicsProxy(FSuspensionConstraint* InConstraint, FPBDSuspensionConstraintHandle* InHandle, UObject* InOwner)
-	: Base(EPhysicsProxyType::SuspensionConstraintType, InOwner, new FProxyTimestampBase)
+	: Base(EPhysicsProxyType::SuspensionConstraintType, InOwner, MakeShared<FProxyTimestampBase>())
 	, Constraint_GT(InConstraint)
 	, Constraint_PT(InHandle)
 	, bInitialized(false)
diff --git a/Engine/Source/Runtime/Experimental/Chaos/Public/Chaos/Framework/PhysicsProxy.h b/Engine/Source/Runtime/Experimental/Chaos/Public/Chaos/Framework/PhysicsProxy.h
--- a/Engine/Source/Runtime/Experimental/Chaos/Public/Chaos/Framework/PhysicsProxy.h
+++ b/Engine/Source/Runtime/Experimental/Chaos/Public/Chaos/Framework/PhysicsProxy.h
@@ -44,12 +44,12 @@ public:
 	using FIntArray = Chaos::TArrayCollectionArray<int32>;
 
 	TPhysicsProxy()
-		: IPhysicsProxyBase(ConcreteType(), nullptr, new TProxyTimeStamp)
+		: IPhysicsProxyBase(ConcreteType(), nullptr, MakeShared<TProxyTimeStamp>())
 	{
 	}
 
 	explicit TPhysicsProxy(UObject* InOwner)
-		: IPhysicsProxyBase(ConcreteType(), InOwner, new TProxyTimeStamp)
+		: IPhysicsProxyBase(ConcreteType(), InOwner, MakeShared<TProxyTimeStamp>())
 	{
 	}
 
diff --git a/Engine/Source/Runtime/Experimental/Chaos/Public/Chaos/Framework/PhysicsProxyBase.h b/Engine/Source/Runtime/Experimental/Chaos/Public/Chaos/Framework/PhysicsProxyBase.h
--- a/Engine/Source/Runtime/Experimental/Chaos/Public/Chaos/Framework/PhysicsProxyBase.h
+++ b/Engine/Source/Runtime/Experimental/Chaos/Public/Chaos/Framework/PhysicsProxyBase.h
@@ -61,7 +61,7 @@ struct CHAOS_API FGeometryCollectionProxyTimestamp: public FProxyTimestampBase
 class CHAOS_API IPhysicsProxyBase
 {
 public:
-	IPhysicsProxyBase(EPhysicsProxyType InType, UObject* InOwner, FProxyTimestampBase* InProxyTimeStamp)
+	IPhysicsProxyBase(EPhysicsProxyType InType, UObject* InOwner, TSharedPtr<FProxyTimestampBase,ESPMode::ThreadSafe> InProxyTimeStamp)
 		: Solver(nullptr)
 		, Owner(InOwner)
 		, DirtyIdx(INDEX_NONE)
