Separate stat labels for STAT_Collisions_ParticlePairBroadPhase and STAT_Collisions_SpatialBroadPhase

diff --git a/Engine/Source/Runtime/Experimental/Chaos/Public/ChaosStats.h b/Engine/Source/Runtime/Experimental/Chaos/Public/ChaosStats.h
--- a/Engine/Source/Runtime/Experimental/Chaos/Public/ChaosStats.h
+++ b/Engine/Source/Runtime/Experimental/Chaos/Public/ChaosStats.h
@@ -117,8 +117,8 @@ DECLARE_CYCLE_STAT_EXTERN(TEXT("[Field Update] Impulse"), STAT_ForceUpdateField_
 
 // Collision Detection
 DECLARE_CYCLE_STAT_EXTERN(TEXT("Collisions::Detect"), STAT_Collisions_Detect, STATGROUP_ChaosCollision, CHAOS_API);
-DECLARE_CYCLE_STAT_EXTERN(TEXT("Collisions::BroadPhase"), STAT_Collisions_ParticlePairBroadPhase, STATGROUP_ChaosCollision, CHAOS_API);
-DECLARE_CYCLE_STAT_EXTERN(TEXT("Collisions::BroadPhase"), STAT_Collisions_SpatialBroadPhase, STATGROUP_ChaosCollision, CHAOS_API);
+DECLARE_CYCLE_STAT_EXTERN(TEXT("Collisions::BroadPhase::ParticlePair"), STAT_Collisions_ParticlePairBroadPhase, STATGROUP_ChaosCollision, CHAOS_API);
+DECLARE_CYCLE_STAT_EXTERN(TEXT("Collisions::BroadPhase::Spatial"), STAT_Collisions_SpatialBroadPhase, STATGROUP_ChaosCollision, CHAOS_API);
 DECLARE_CYCLE_STAT_EXTERN(TEXT("Collisions::AABBTree"), STAT_Collisions_AABBTree, STATGROUP_ChaosCollision, CHAOS_API);
 DECLARE_CYCLE_STAT_EXTERN(TEXT("Collisions::Filtering"), STAT_Collisions_Filtering, STATGROUP_ChaosCollision, CHAOS_API);
 DECLARE_CYCLE_STAT_EXTERN(TEXT("Collisions::Restore"), STAT_Collisions_Restore, STATGROUP_ChaosCollision, CHAOS_API);
