fix geometry collection hit proxies when bone selection is not enabled

diff --git a/Engine/Source/Runtime/Experimental/GeometryCollectionEngine/Private/GeometryCollection/GeometryCollectionSceneProxy.cpp b/Engine/Source/Runtime/Experimental/GeometryCollectionEngine/Private/GeometryCollection/GeometryCollectionSceneProxy.cpp
--- a/Engine/Source/Runtime/Experimental/GeometryCollectionEngine/Private/GeometryCollection/GeometryCollectionSceneProxy.cpp
+++ b/Engine/Source/Runtime/Experimental/GeometryCollectionEngine/Private/GeometryCollection/GeometryCollectionSceneProxy.cpp
@@ -353,7 +353,7 @@ void FGeometryCollectionSceneProxy::InitResources()
 	VertexBuffers.ColorVertexBuffer.BindColorVertexBuffer(&VertexFactory, Data);
 
 #if GEOMETRYCOLLECTION_EDITOR_SELECTION
-	if (bEnableBoneSelection)
+	// Note: Could skip this when bEnableBoneSelection is false if the hitproxy shader was made to not require per-vertex hit proxy IDs in that case
 	{
 		HitProxyIdBuffer.Init(NumVertices);
 		HitProxyIdBuffer.InitResource();
@@ -508,6 +508,7 @@ void FGeometryCollectionSceneProxy::SetConstantData_RenderThread(FGeometryCollec
 
 	if (GetRequiredVertexCount())
 	{
+		FColor WholeObjectHitProxyColor = GetPrimitiveSceneInfo()->DefaultDynamicHitProxyId.GetColor();
 		ParallelFor(Vertices.Num(), [&](int32 i)
 		{
 			const FDynamicMeshVertex& Vertex = Vertices[i];
@@ -520,14 +521,15 @@ void FGeometryCollectionSceneProxy::SetConstantData_RenderThread(FGeometryCollec
 			}
 			VertexBuffers.ColorVertexBuffer.VertexColor(i) = Vertex.Color;
 #if GEOMETRYCOLLECTION_EDITOR_SELECTION
-			if (bEnableBoneSelection)
+			if (bEnableBoneSelection && PerBoneHitProxies.Num())
 			{
 				// One proxy per bone
 				const int32 ProxyIndex = ConstantData->BoneMap[i];
-				if (PerBoneHitProxies.Num())
-				{ 
-					HitProxyIdBuffer.VertexColor(i) = PerBoneHitProxies[ProxyIndex]->Id.GetColor();
-				}
+				HitProxyIdBuffer.VertexColor(i) = PerBoneHitProxies[ProxyIndex]->Id.GetColor();
+			}
+			else
+			{
+				HitProxyIdBuffer.VertexColor(i) = WholeObjectHitProxyColor;
 			}
 #endif
 		});
@@ -547,7 +549,7 @@ void FGeometryCollectionSceneProxy::SetConstantData_RenderThread(FGeometryCollec
 		}
 
 #if GEOMETRYCOLLECTION_EDITOR_SELECTION
-		if (bEnableBoneSelection)
+		// Note: Could skip this when bEnableBoneSelection is false if the hitproxy shader was made to not require per-vertex hit proxy IDs in that case
 		{
 			auto& VertexBuffer = HitProxyIdBuffer;
 			void* VertexBufferData = RHILockBuffer(VertexBuffer.VertexBufferRHI, 0, VertexBuffer.GetNumVertices() * VertexBuffer.GetStride(), RLM_WriteOnly);
@@ -1207,6 +1209,8 @@ HHitProxy* FGeometryCollectionSceneProxy::CreateHitProxies(UPrimitiveComponent*
 	else if (Component->GetOwner())
 	{
 #if GEOMETRYCOLLECTION_EDITOR_SELECTION
+		// Note the below-created subsection hitproxies are never drawn, because the per-vertex hitproxy 
+		// rendering path is always on for geometry collections (i.e., USE_PER_VERTEX_HITPROXY_ID is 1)
 		const int32 NumTransforms = (Sections.Num() > 0) ? SubSectionHitProxies.Num() / Sections.Num(): 0;
 		for (int32 SectionIndex = 0; SectionIndex < Sections.Num(); ++SectionIndex)
 		{
diff --git a/Engine/Source/Runtime/Experimental/GeometryCollectionEngine/Private/GeometryCollection/GeometryCollectionSceneProxy.h b/Engine/Source/Runtime/Experimental/GeometryCollectionEngine/Private/GeometryCollection/GeometryCollectionSceneProxy.h
--- a/Engine/Source/Runtime/Experimental/GeometryCollectionEngine/Private/GeometryCollection/GeometryCollectionSceneProxy.h
+++ b/Engine/Source/Runtime/Experimental/GeometryCollectionEngine/Private/GeometryCollection/GeometryCollectionSceneProxy.h
@@ -316,7 +316,11 @@ public:
 	// FPrimitiveSceneProxy interface.
 #if WITH_EDITOR
 	virtual HHitProxy* CreateHitProxies(UPrimitiveComponent* Component, TArray<TRefCountPtr<HHitProxy> >& OutHitProxies) override;
-	virtual const FColorVertexBuffer* GetCustomHitProxyIdBuffer() const override { return bEnableBoneSelection ? &HitProxyIdBuffer : nullptr; }
+	virtual const FColorVertexBuffer* GetCustomHitProxyIdBuffer() const override
+	{
+		// Note: Could return nullptr when bEnableBoneSelection is false if the hitproxy shader was made to not require per-vertex hit proxy IDs in that case
+		return &HitProxyIdBuffer;
+	}
 #endif // WITH_EDITOR
 
 #if GEOMETRYCOLLECTION_EDITOR_SELECTION
diff --git a/Engine/Source/Runtime/Experimental/GeometryCollectionEngine/Public/GeometryCollection/GeometryCollectionHitProxy.h b/Engine/Source/Runtime/Experimental/GeometryCollectionEngine/Public/GeometryCollection/GeometryCollectionHitProxy.h
--- a/Engine/Source/Runtime/Experimental/GeometryCollectionEngine/Public/GeometryCollection/GeometryCollectionHitProxy.h
+++ b/Engine/Source/Runtime/Experimental/GeometryCollectionEngine/Public/GeometryCollection/GeometryCollectionHitProxy.h
@@ -8,7 +8,9 @@
 
 class UGeometryCollectionComponent;
 
-/** HitProxy with transform index information. */
+/** HitProxy with transform index information.
+ * @todo: These are never actually drawn because the shader is always in per-vertex mode, where it is passed HGeometryCollectionBone (or the whole-object HActor)
+ */
 struct HGeometryCollection : public HActor
 {
 	DECLARE_HIT_PROXY(GEOMETRYCOLLECTIONENGINE_API)
