Deprecated ChaosVersionString for ChaosVersionGUID to avoid a memory allocation per DLL (Found during Insights trace analysis)

diff --git a/Engine/Source/Runtime/Engine/Private/PhysicsEngine/Experimental/ChaosDerivedData.cpp b/Engine/Source/Runtime/Engine/Private/PhysicsEngine/Experimental/ChaosDerivedData.cpp
--- a/Engine/Source/Runtime/Engine/Private/PhysicsEngine/Experimental/ChaosDerivedData.cpp
+++ b/Engine/Source/Runtime/Engine/Private/PhysicsEngine/Experimental/ChaosDerivedData.cpp
@@ -56,7 +56,7 @@ const TCHAR* FChaosDerivedDataCooker::GetVersionString() const
 	// for their own DDC or serialized data - change Chaos::ChaosVersionString in Chaos/Core.h to bump our
 	// Chaos data version. Callers can also rely on that version in their builders and avoid bad serialization
 	// when basic Chaos data changes
-	return *Chaos::ChaosVersionString;
+	return Chaos::ChaosVersionGUID;
 }
 
 FString FChaosDerivedDataCooker::GetDebugContextString() const
diff --git a/Engine/Source/Runtime/Experimental/ChaosCore/Public/Chaos/Core.h b/Engine/Source/Runtime/Experimental/ChaosCore/Public/Chaos/Core.h
--- a/Engine/Source/Runtime/Experimental/ChaosCore/Public/Chaos/Core.h
+++ b/Engine/Source/Runtime/Experimental/ChaosCore/Public/Chaos/Core.h
@@ -51,6 +51,14 @@ namespace Chaos
 	// NOTE: if you get a merge conflict on the GUID, you must replace it with a new GUID - do not accept the source or target
 	// or you will likely get DDC version conflicts resulting in crashes during load.
 	// Core version string for Chaos data. Any DDC builder dependent on Chaos for serialization should depend on this version
-	static const FString ChaosVersionString = TEXT("5F283595-0103-4FE6-9684-E8B820158BFC");
+	static const TCHAR* const ChaosVersionGUID = TEXT("5F283595-0103-4FE6-9684-E8B820158BFC");
 
+	inline FString GetChaosVersionStringInner()
+	{
+		return FString(ChaosVersionGUID);
+	}
+
+#define ChaosVersionString \
+	UE_DEPRECATED_MACRO(5.1, "ChaosVersionString is deprecated, please use ChaosVersionGUID instead") \
+	GetChaosVersionStringInner()
 }
diff --git a/Engine/Source/Runtime/Experimental/GeometryCollectionEngine/Private/GeometryCollection/DerivedDataGeometryCollectionCooker.cpp b/Engine/Source/Runtime/Experimental/GeometryCollectionEngine/Private/GeometryCollection/DerivedDataGeometryCollectionCooker.cpp
--- a/Engine/Source/Runtime/Experimental/GeometryCollectionEngine/Private/GeometryCollection/DerivedDataGeometryCollectionCooker.cpp
+++ b/Engine/Source/Runtime/Experimental/GeometryCollectionEngine/Private/GeometryCollection/DerivedDataGeometryCollectionCooker.cpp
@@ -87,7 +87,7 @@ FString FDerivedDataGeometryCollectionCooker::GetPluginSpecificCacheKeySuffix()
 {
 	return FString::Printf(
 		TEXT("%s_%s_%s_%d_%d"),
-		*Chaos::ChaosVersionString,
+		Chaos::ChaosVersionGUID,
 		*GeometryCollection.GetIdGuid().ToString(),
 		*GeometryCollection.GetStateGuid().ToString(),
 		FDestructionObjectVersion::Type::LatestVersion,
diff --git a/Engine/Source/Runtime/Landscape/Private/LandscapeCollision.cpp b/Engine/Source/Runtime/Landscape/Private/LandscapeCollision.cpp
--- a/Engine/Source/Runtime/Landscape/Private/LandscapeCollision.cpp
+++ b/Engine/Source/Runtime/Landscape/Private/LandscapeCollision.cpp
@@ -248,7 +248,7 @@ static FString GetHFDDCKeyString(const FName& Format, bool bDefMaterial, const F
 #if PHYSICS_INTERFACE_PHYSX
 	const FString InterfacePrefix = TEXT("PHYSX");
 #elif WITH_CHAOS
-	const FString InterfacePrefix = FString::Printf(TEXT("%s_%s"), TEXT("CHAOS"), *Chaos::ChaosVersionString);
+	const FString InterfacePrefix = FString::Printf(TEXT("%s_%s"), TEXT("CHAOS"), Chaos::ChaosVersionGUID);
 #else
 	const FString InterfacePrefix = TEXT("UNDEFINED");
 #endif
