Chaos destruction : Make sure geometry component inherits damagethreshold settings from rest collection asset

diff --git a/Engine/Source/Runtime/Experimental/GeometryCollectionEngine/Private/GeometryCollection/GeometryCollectionComponent.cpp b/Engine/Source/Runtime/Experimental/GeometryCollectionEngine/Private/GeometryCollection/GeometryCollectionComponent.cpp
--- a/Engine/Source/Runtime/Experimental/GeometryCollectionEngine/Private/GeometryCollection/GeometryCollectionComponent.cpp
+++ b/Engine/Source/Runtime/Experimental/GeometryCollectionEngine/Private/GeometryCollection/GeometryCollectionComponent.cpp
@@ -2382,6 +2382,9 @@ void UGeometryCollectionComponent::SetRestCollection(const UGeometryCollection*
 		{
 			InitializeEmbeddedGeometry();
 		}
+
+		// initialize the component per level damage threshold from the asset defaults 
+		DamageThreshold = RestCollection->DamageThreshold;
 		
 		//ResetDynamicCollection();
 	}
